Implemented on-the-fly debug view shader generation.
The previous path used to parse the world and generate all shaders at once,
which created a massive stall. This happened immediately after choosing the mode.
For example, following "viewmode shadercomplexity"

The new implementation generate only what is currently being used in the rendering.
Added some visual feedbacks to show missing (pending) shader status in shader complexity.
Added "r.ViewMode.ShaderTimeSlice" to control how much time is spent per frame to generate the viewmode shaders.

diff --git a/Engine/Shaders/Private/ShaderComplexityApplyPixelShader.usf b/Engine/Shaders/Private/ShaderComplexityApplyPixelShader.usf
--- a/Engine/Shaders/Private/ShaderComplexityApplyPixelShader.usf
+++ b/Engine/Shaders/Private/ShaderComplexityApplyPixelShader.usf
@@ -28,7 +28,7 @@ float3 ShaderComplexityColors[MAX_NUM_COMPLEXITY_COLORS];
 
 // float4(bLegend, DebugViewShaderMode, ColorSamping, ComplexityScale), can be optimzied as define if needed
 float4 ShaderComplexityParams;
-// float4(NumComplexityColors, 0, UsedQuadBufferSizeX, UsedQuadBufferSizeY);
+// float4(NumComplexityColors, bShowError, UsedQuadBufferSizeX, UsedQuadBufferSizeY);
 float4 ShaderComplexityParams2;
 
 // @param In  0..1
@@ -184,6 +184,14 @@ void Main(
 
 	// PS cost
 	OutColor = float4(ColorizeComplexityFromMode(SceneColor.r), 0);
+
+	// Mark missing data as invalid, this would happen when shaders are compiling
+	if (SceneColor.r == 0 && ShaderComplexityParams2.y != 0)
+	{
+		int2 PixelPosition = SvPosition.xy;
+		float Error = (PixelPosition.x & 0x08) == (PixelPosition.y & 0x08) ? .22 : 0;
+		OutColor.rgb = Error;
+	}
 	
 	if(!bLegend)
 	{
diff --git a/Engine/Source/Editor/UnrealEd/Private/EditorViewportClient.cpp b/Engine/Source/Editor/UnrealEd/Private/EditorViewportClient.cpp
--- a/Engine/Source/Editor/UnrealEd/Private/EditorViewportClient.cpp
+++ b/Engine/Source/Editor/UnrealEd/Private/EditorViewportClient.cpp
@@ -3634,14 +3634,7 @@ void FEditorViewportClient::Draw(FViewport* InViewport, FCanvas* Canvas)
 
 	ViewFamily.EngineShowFlags = EngineShowFlags;
 
-	if (World && ViewFamily.GetDebugViewShaderMode() != DVSM_None && HasMissingDebugViewModeShaders(true))
-	{
-		TSet<UMaterialInterface*> Materials;
-		if (GetUsedMaterialsInWorld(World, Materials, nullptr))
-		{
-			CompileDebugViewModeShaders(ViewFamily.GetDebugViewShaderMode(), GetCachedScalabilityCVars().MaterialQualityLevel, ViewFamily.GetFeatureLevel(), false, false, Materials, nullptr);
-		}
-	}
+	UpdateDebugViewModeShaders();
 
 	if( ModeTools->GetActiveMode( FBuiltinEditorModes::EM_InterpEdit ) == 0 || !AllowsCinematicControl() )
 	{
@@ -5111,7 +5104,8 @@ void FEditorViewportClient::SetViewMode(EViewModeIndex InViewModeIndex)
 		{
 			FEditorBuildUtils::EditorBuildTextureStreaming(GetWorld(), InViewModeIndex);
 		}
-		FEditorBuildUtils::CompileViewModeShaders(GetWorld(), InViewModeIndex);
+		// Uncomment this to generate inital viewmode data.
+		// FEditorBuildUtils::CompileViewModeShaders(GetWorld(), InViewModeIndex);
 			 
 		PerspViewModeIndex = InViewModeIndex;
 		ApplyViewMode(PerspViewModeIndex, true, EngineShowFlags);
diff --git a/Engine/Source/Runtime/Engine/Private/DebugViewModeHelpers.cpp b/Engine/Source/Runtime/Engine/Private/DebugViewModeHelpers.cpp
--- a/Engine/Source/Runtime/Engine/Private/DebugViewModeHelpers.cpp
+++ b/Engine/Source/Runtime/Engine/Private/DebugViewModeHelpers.cpp
@@ -4,7 +4,7 @@
 	DebugViewModeHelpers.cpp: debug view shader helpers.
 =============================================================================*/
 #include "DebugViewModeHelpers.h"
-#include "DebugViewModeMaterialProxy.h"
+#include "DebugViewModeMaterialManager.h"
 #include "ShaderCompiler.h"
 #include "Misc/ScopedSlowTask.h"
 #include "Misc/FeedbackContext.h"
@@ -252,10 +252,9 @@ bool CompileDebugViewModeShaders(EDebugViewShaderMode ShaderMode, EMaterialQuali
 
 		if (bFullRebuild)
 		{
-			FDebugViewModeMaterialProxy::ClearAllShaders(MaterialInterface);
+			GDebugViewModeMaterialManager.RemoveShaders(MaterialInterface);
 		}
 
-
 		const FMaterial* Material = MaterialInterface->GetMaterialResource(FeatureLevel);
 		if (!Material)
 		{
@@ -281,7 +280,9 @@ bool CompileDebugViewModeShaders(EDebugViewShaderMode ShaderMode, EMaterialQuali
 			MaterialInterface->SetTextureStreamingData(TArray<FMaterialTextureInfo>());
 			continue;
 		}
-		FDebugViewModeMaterialProxy::AddShader(MaterialInterface, QualityLevel, FeatureLevel, !bWaitForPreviousShaders, ShaderMode);
+
+		// If we are not waiting for shaders, then the shader needs to be compiled in sync.
+		GDebugViewModeMaterialManager.AddShader(MaterialInterface, ShaderMode, QualityLevel, FeatureLevel, !bWaitForPreviousShaders);
 	}
 
 	for (UMaterialInterface* RemovedMaterial : MaterialsToRemove)
@@ -292,14 +293,14 @@ bool CompileDebugViewModeShaders(EDebugViewShaderMode ShaderMode, EMaterialQuali
 	if (!bWaitForPreviousShaders || WaitForShaderCompilation(LOCTEXT("CompileDebugViewModeShaders", "Compiling Optional Engine Shaders"), ProgressTask))
 	{
 		// Check The validity of all shaders, removing invalid entries
-		FDebugViewModeMaterialProxy::ValidateAllShaders(Materials);
+		GDebugViewModeMaterialManager.ValidateShaders(true);
 
 		UE_LOG(TextureStreamingBuild, Display, TEXT("Compiling optional shaders took %.3f seconds."), FPlatformTime::Seconds() - StartTime);
 		return true;
 	}
 	else
 	{
-		FDebugViewModeMaterialProxy::ClearAllShaders(nullptr);
+		GDebugViewModeMaterialManager.RemoveShaders(nullptr);
 		return false;
 	}
 #else
diff --git a/Engine/Source/Runtime/Engine/Private/DebugViewModeMaterialManager.cpp b/Engine/Source/Runtime/Engine/Private/DebugViewModeMaterialManager.cpp
new file mode 100644
--- /dev/null
+++ b/Engine/Source/Runtime/Engine/Private/DebugViewModeMaterialManager.cpp
@@ -0,0 +1,321 @@
+// Copyright 1998-2019 Epic Games, Inc. All Rights Reserved.
+
+/*=============================================================================
+DebugViewModeMaterialManager.cpp : Implements debug shader management.
+=============================================================================*/
+
+#include "DebugViewModeMaterialManager.h"
+#include "DebugViewModeInterface.h"
+#include "SceneInterface.h"
+#include "EngineModule.h"
+#include "RendererInterface.h"
+#include "UnrealEngine.h"
+
+ENGINE_API bool GetDebugViewMaterial(const UMaterialInterface* InMaterialInterface, EDebugViewShaderMode InDebugViewMode, ERHIFeatureLevel::Type InFeatureLevel, const FMaterialRenderProxy*& OutMaterialRenderProxy, const FMaterial*& OutMaterial)
+{
+#if WITH_EDITORONLY_DATA
+	return GDebugViewModeMaterialManager.GetShader_RenderThread(InMaterialInterface, InDebugViewMode, InFeatureLevel, OutMaterialRenderProxy, OutMaterial);
+#else
+	return false;
+#endif
+}
+
+ENGINE_API void UpdateDebugViewModeShaders()
+{
+#if WITH_EDITORONLY_DATA
+	GDebugViewModeMaterialManager.Update();
+#endif
+}
+
+ENGINE_API void ClearDebugViewMaterials(UMaterialInterface* InMaterialInterface)
+{
+#if WITH_EDITORONLY_DATA
+	GDebugViewModeMaterialManager.RemoveShaders(InMaterialInterface);
+#endif
+}
+
+float GViewModeShaderTimeSlice = 0.02f;
+static FAutoConsoleVariableRef CVarViewModeShaderTimeSlice(
+	TEXT("r.ViewMode.ShaderTimeSlice"),
+	GViewModeShaderTimeSlice,
+	TEXT("Max time to allocate each frame to generate new shaders. 0 disables (default=.02"),
+	ECVF_Default
+);
+
+extern double GViewModeShaderMissingTime;
+extern int32 GNumViewModeShaderMissing;
+double GViewModeShaderMissingTime = 0;
+int32 GNumViewModeShaderMissing = 0;
+
+#if WITH_EDITORONLY_DATA
+
+FDebugViewModeMaterialManager GDebugViewModeMaterialManager;
+
+// Trigger compilation of a specific viewmode shader.
+void FDebugViewModeMaterialManager::AddShader(
+	UMaterialInterface* MaterialInterface, 
+	EDebugViewShaderMode DebugViewMode,
+	EMaterialQualityLevel::Type QualityLevel, 
+	ERHIFeatureLevel::Type FeatureLevel, 
+	bool bSynchronousCompilation)
+{
+	FScopeLock SL(&CS);
+	QUICK_SCOPE_CYCLE_COUNTER(STAT_DebugViewModeMaterialManager_AddShader);
+
+	check(MaterialInterface);
+
+	const FMaterial* Material = MaterialInterface->GetMaterialResource(FeatureLevel);
+	const FDebugViewModeInterface* DebugViewModeInterface = FDebugViewModeInterface::GetInterface(DebugViewMode);
+	if (Material && DebugViewModeInterface)
+	{
+		if (!DebugViewModeInterface->bNeedsMaterialProperties && FDebugViewModeInterface::AllowFallbackToDefaultMaterial(Material))
+		{
+			MaterialInterface = UMaterial::GetDefaultMaterial(MD_Surface);
+		}
+
+		FMaterialKey Key(MaterialInterface, DebugViewMode, FeatureLevel);
+		FMaterialInfo& Info = MaterialInfos.FindOrAdd(Key);
+		if (Info.ShaderState == EShaderState::New)
+		{
+			MaterialKeys.FindOrAdd(MaterialInterface).Add(Key);
+			Info.DebugProxy = new FDebugViewModeMaterialProxy(MaterialInterface, QualityLevel, FeatureLevel, bSynchronousCompilation, DebugViewMode);
+			Info.ShaderState = EShaderState::Compiling;
+		}
+	}
+}
+
+void FDebugViewModeMaterialManager::Update()
+{
+	FScopeLock SL(&CS);
+	QUICK_SCOPE_CYCLE_COUNTER(STAT_DebugViewModeMaterialManager_Update);
+
+	const EMaterialQualityLevel::Type QualityLevel = GetCachedScalabilityCVars().MaterialQualityLevel;
+	int32 NumShaderCompiling = 0;
+
+	if (MissingShaderKeys.Num())
+	{
+		QUICK_SCOPE_CYCLE_COUNTER(STAT_DebugViewModeMaterialManager_AddMissingShaders);
+		const double StartTime = FPlatformTime::Seconds();
+
+		// Update viewport warning
+		GViewModeShaderMissingTime = FApp::GetCurrentTime();
+		GNumViewModeShaderMissing = MissingShaderKeys.Num();
+
+		TSet<FMaterialKey> NextMissingShaderKeys;
+
+		// First create a proxy for every missing shader.
+		for (const FMaterialKey& Key : MissingShaderKeys)
+		{
+			FMaterialInfo& Info = MaterialInfos.FindChecked(Key);
+			check(!Info.DebugProxy && Info.ShaderState == EShaderState::Missing);
+
+			if (GViewModeShaderTimeSlice <= 0 || (float)(FPlatformTime::Seconds() - StartTime) < GViewModeShaderTimeSlice)
+			{
+				Info.DebugProxy = new FDebugViewModeMaterialProxy(const_cast<UMaterialInterface*>(Key.MaterialInterface), QualityLevel, (ERHIFeatureLevel::Type)Key.FeatureLevel, false, (EDebugViewShaderMode)Key.DebugViewMode);
+				Info.ShaderState = EShaderState::Compiling;
+			}
+			else
+			{
+				NextMissingShaderKeys.Add(Key);
+			}
+		}
+		FMemory::Memswap(&NextMissingShaderKeys, &MissingShaderKeys, sizeof(NextMissingShaderKeys));
+	}
+
+	ValidateShaders(false);
+}
+
+void FDebugViewModeMaterialManager::ValidateShaders(bool bAllShadersReady)
+{
+	QUICK_SCOPE_CYCLE_COUNTER(STAT_DebugViewModeMaterialManager_ValidateShaders);
+
+	// Check every shader that could be missing validation. Otherwise shader first have to be requested before their validation is done.
+	if (bAllShadersReady)
+	{
+		FScopeLock SL(&CS);
+
+		for (TMap<FMaterialKey, FMaterialInfo>::TIterator It(MaterialInfos); It; ++It)
+		{
+			FMaterialInfo& Info = It.Value();
+			if (Info.DebugProxy && Info.ShaderState == EShaderState::Compiling)
+			{
+				PendingValidationShaderKeys.Add(It.Key());
+				Info.ShaderState = EShaderState::PendingValidation;
+			}
+		}
+	}
+
+	TArray<FMaterialKey> KeysToUpdateUEC;
+	TArray<FDebugViewModeMaterialProxy*, TInlineAllocator<4> > DebugMaterialsToDelete;
+	{
+		FScopeLock SL(&CS);
+
+		for (const FMaterialKey& Key : PendingValidationShaderKeys)
+		{
+			check(Key.MaterialInterface);
+			FMaterialInfo& Info = MaterialInfos.FindChecked(Key);
+			check(Info.DebugProxy && Info.ShaderState == EShaderState::PendingValidation);
+
+			// It doesn't look relevant anymore to ensure that the uniform expression set are compatible.
+			// const FMaterial* ReferenceMaterial = Key.MaterialInterface->GetMaterialResource((ERHIFeatureLevel::Type)Key.FeatureLevel);
+
+			if (Info.DebugProxy->IsValid() && Info.DebugProxy->GetGameThreadShaderMap())
+				// && ReferenceMaterial
+				// && ReferenceMaterial->GetGameThreadShaderMap() 
+				// && ReferenceMaterial->GetGameThreadShaderMap()->GetUniformExpressionSet() == Info.DebugProxy->GetGameThreadShaderMap()->GetUniformExpressionSet()
+			{
+				KeysToUpdateUEC.Add(Key);
+				Info.ShaderState = EShaderState::PendingUpdateUEC;
+			}
+			else
+			{
+				Info.DebugProxy->MarkAsInvalid();
+				DebugMaterialsToDelete.Add(Info.DebugProxy);
+				Info.DebugProxy = nullptr;
+				Info.ShaderState = EShaderState::Invalid;
+			}
+		}
+		PendingValidationShaderKeys.Empty();
+	}
+
+	if (KeysToUpdateUEC.Num())
+	{
+		ENQUEUE_RENDER_COMMAND(DebugViewModeMaterialsUpdateUEC)([KeysToUpdateUEC, this](FRHICommandList& RHICmdList)
+		{
+			FScopeLock SL(&CS);
+			QUICK_SCOPE_CYCLE_COUNTER(STAT_DebugViewModeMaterialManager_DebugViewModeMaterialsUpdateUEC);
+
+			for (const FMaterialKey& Key : KeysToUpdateUEC)
+			{
+				FMaterialInfo* Info = MaterialInfos.Find(Key);
+
+				// Could be missing if there was a call to RemoveShaders().
+				if (Info)
+				{
+					check(Info->DebugProxy && Info->ShaderState == EShaderState::PendingUpdateUEC);
+					Info->DebugProxy->UpdateUniformExpressionCacheIfNeeded((ERHIFeatureLevel::Type)Key.FeatureLevel);
+					Info->ShaderState = EShaderState::Valid;
+				}
+			}
+		});
+	}
+
+	if (DebugMaterialsToDelete.Num())
+	{
+		ENQUEUE_RENDER_COMMAND(DeleteDebugViewModeMaterials)([DebugMaterialsToDelete](FRHICommandList& RHICmdList)
+		{
+			QUICK_SCOPE_CYCLE_COUNTER(STAT_DebugViewModeMaterialManager_DeleteDebugViewModeMaterials);
+
+			for (FDebugViewModeMaterialProxy* MaterialToDelete : DebugMaterialsToDelete)
+			{
+				check(MaterialToDelete);
+				delete MaterialToDelete;
+			}
+		});
+	}
+
+	if (bAllShadersReady)
+	{
+		FlushRenderingCommands();
+	}
+}
+
+void FDebugViewModeMaterialManager::RemoveShaders(UMaterialInterface* InMaterialInterface)
+{
+	TArray<FDebugViewModeMaterialProxy*, TInlineAllocator<4> > DebugMaterialsToDelete;
+
+	if (InMaterialInterface)
+	{
+		FScopeLock SL(&CS);
+		QUICK_SCOPE_CYCLE_COUNTER(STAT_DebugViewModeMaterialManager_RemoveShaders);
+
+		const TArray<FMaterialKey>* Keys = MaterialKeys.Find(InMaterialInterface);
+		if (Keys)
+		{
+			for (const FMaterialKey& Key : *Keys)
+			{
+				FMaterialInfo Info = MaterialInfos.FindAndRemoveChecked(Key);
+				if (Info.DebugProxy)
+				{
+					DebugMaterialsToDelete.Add(Info.DebugProxy);
+				}
+				if (Info.ShaderState == EShaderState::Missing)
+				{
+					MissingShaderKeys.Remove(Key);
+				}
+				else if (Info.ShaderState == EShaderState::PendingValidation)
+				{
+					PendingValidationShaderKeys.Remove(Key);
+				}
+			}
+			MaterialKeys.Remove(InMaterialInterface);
+		}
+	}
+	else // Otherwise clear everything.
+	{
+		FScopeLock SL(&CS);
+		QUICK_SCOPE_CYCLE_COUNTER(STAT_DebugViewModeMaterialManager_RemoveAllShaders);
+
+		for (TMap<FMaterialKey, FMaterialInfo>::TIterator It(MaterialInfos); It; ++It)
+		{
+			const FMaterialInfo& Info = It.Value();
+			if (Info.DebugProxy)
+			{
+				DebugMaterialsToDelete.Add(Info.DebugProxy);
+			}
+		}
+		MaterialInfos.Empty();
+		MaterialKeys.Empty();
+		MissingShaderKeys.Empty();
+		PendingValidationShaderKeys.Empty();
+	}
+	
+	if (DebugMaterialsToDelete.Num())
+	{
+		ENQUEUE_RENDER_COMMAND(DeleteDebugViewModeMaterials)([DebugMaterialsToDelete](FRHICommandList& RHICmdList)
+		{
+			QUICK_SCOPE_CYCLE_COUNTER(STAT_DebugViewModeMaterialManager_DeleteDebugViewModeMaterials);
+
+			for (FDebugViewModeMaterialProxy* MaterialToDelete : DebugMaterialsToDelete)
+			{
+				check(MaterialToDelete);
+				delete MaterialToDelete;
+			}
+		});
+	}
+}
+
+bool FDebugViewModeMaterialManager::GetShader_RenderThread(
+	const UMaterialInterface* MaterialInterface, 
+	EDebugViewShaderMode DebugViewMode, 
+	ERHIFeatureLevel::Type FeatureLevel, 
+	const FMaterialRenderProxy*& OutMaterialRenderProxy, 
+	const FMaterial*& OutMaterial)
+{
+	FScopeLock SL(&CS);
+
+	FMaterialKey Key(MaterialInterface, DebugViewMode, FeatureLevel);
+	FMaterialInfo& Info = MaterialInfos.FindOrAdd(Key);
+
+	if (Info.ShaderState == EShaderState::Valid)
+	{
+		OutMaterialRenderProxy = Info.DebugProxy;
+		OutMaterial = Info.DebugProxy;
+		return true;
+	}
+	else if (Info.ShaderState == EShaderState::New)
+	{
+		MaterialKeys.FindOrAdd(MaterialInterface).Add(Key);
+		MissingShaderKeys.Add(Key);
+		Info.ShaderState = EShaderState::Missing;
+	}
+	else if (Info.ShaderState == EShaderState::Compiling && Info.DebugProxy->GetRenderingThreadShaderMap())
+	{
+		PendingValidationShaderKeys.Add(Key);
+		Info.ShaderState = EShaderState::PendingValidation;
+	}
+
+	return false;
+}
+
+#endif // WITH_EDITORONLY_DATA
diff --git a/Engine/Source/Runtime/Engine/Private/DebugViewModeMaterialManager.h b/Engine/Source/Runtime/Engine/Private/DebugViewModeMaterialManager.h
new file mode 100644
--- /dev/null
+++ b/Engine/Source/Runtime/Engine/Private/DebugViewModeMaterialManager.h
@@ -0,0 +1,106 @@
+// Copyright 1998-2019 Epic Games, Inc. All Rights Reserved.
+
+/*=============================================================================
+DebugViewModeMaterialManager.h : Implements debug shader management.
+=============================================================================*/
+
+#pragma once
+
+#include "CoreMinimal.h"
+#include "DebugViewModeMaterialProxy.h"
+
+#if WITH_EDITORONLY_DATA
+
+class FDebugViewModeMaterialManager
+{
+public:
+
+	// Trigger compilation of a specific viewmode shader.
+	void AddShader(
+		UMaterialInterface* MaterialInterface, 
+		EDebugViewShaderMode DebugViewMode,
+		EMaterialQualityLevel::Type QualityLevel, 
+		ERHIFeatureLevel::Type FeatureLevel, 
+		bool bSynchronousCompilation);
+
+	// Gamethread update that trigger missing shader compilation and validate shader that are ready to be used.
+	void Update();
+
+	// Validate shaders refering to a specific material (or all shaders that could be pending).
+	void ValidateShaders(bool bAllShadersReady);
+
+	// Clear shaders refering to a specific material (or all shaders if null).
+	void RemoveShaders(UMaterialInterface* InMaterialInterface = nullptr);
+
+	// Get a shader to be used for rendering.
+	bool GetShader_RenderThread(
+		const UMaterialInterface* MaterialInterface, 
+		EDebugViewShaderMode DebugViewMode, 
+		ERHIFeatureLevel::Type FeatureLevel, 
+		const FMaterialRenderProxy*& OutMaterialRenderProxy, 
+		const FMaterial*& OutMaterial);
+
+private:
+
+	struct FMaterialKey
+	{
+		FMaterialKey(const UMaterialInterface* InMaterialInterface, EDebugViewShaderMode InDebugViewMode, ERHIFeatureLevel::Type InFeatureLevel) : MaterialInterface(InMaterialInterface), DebugViewMode((uint32)InDebugViewMode), FeatureLevel((uint32)InFeatureLevel) {}
+
+		const UMaterialInterface* MaterialInterface;
+		uint32 DebugViewMode;
+		uint32 FeatureLevel;
+
+		friend bool operator==(const FMaterialKey& Lhs, const FMaterialKey& Rhs)
+		{
+			return Lhs.MaterialInterface == Rhs.MaterialInterface
+				&& Lhs.DebugViewMode == Rhs.DebugViewMode
+				&& Lhs.FeatureLevel == Rhs.FeatureLevel;
+		}
+
+		friend uint32 GetTypeHash(const FMaterialKey& Key)
+		{
+			const uint32 DebugViewModeShift = 6;
+			static_assert((2 ^ DebugViewModeShift) < DVSM_MAX, "Bit shift too small!");
+
+			return GetTypeHash(Key.MaterialInterface) ^ GetTypeHash((Key.FeatureLevel << DebugViewModeShift) ^ Key.DebugViewMode);
+		}
+	};
+
+	enum class EShaderState
+	{
+		New,				// Entry was just created using default constructor.
+		Missing,			// Entry was just created and we will create the debug proxy in the next call to update (has an entry in MissingShaderKeys).
+		Compiling,			// Debug proxy was created and compilation shader is not finished.
+		PendingValidation,	// Shader compilation is finished but can't be used before validation.
+		PendingUpdateUEC,	// Valid but material needs a call to UpdateUniformExpressionCacheIfNeeded before being used in rendering.
+		Valid,				// Shader can be used.
+		Invalid,			// Shader validation has failed, don't reattempt compilation or use debug proxy.
+	};
+
+	struct FMaterialInfo
+	{
+		// The debug proxy created to generate the shader. If null then shader compilation has not started or been validated.
+		FDebugViewModeMaterialProxy* DebugProxy = nullptr;
+		EShaderState ShaderState = EShaderState::New;
+	};
+
+	// A critical section used to control access of the shader maps.
+	FCriticalSection CS;
+
+	// List of debug shader we have generated. Once the shader is compiled, it needs to be validated before being used.
+	// If the validation fails, we keep an entry in the map to prevent attemption a second recompile.
+	TMap<FMaterialKey, FMaterialInfo> MaterialInfos;
+
+	// List of all keys used for a shader, used to speed up RemoveShaders()
+	TMap<const UMaterialInterface*, TArray<FMaterialKey>> MaterialKeys;
+
+	// List of all entries in EShaderState::Missing state.
+	TSet<FMaterialKey> MissingShaderKeys;
+
+	// List of all entries in EShaderState::PendingValidation state.
+	TSet<FMaterialKey> PendingValidationShaderKeys;
+};
+
+extern FDebugViewModeMaterialManager GDebugViewModeMaterialManager;
+
+#endif // WITH_EDITORONLY_DATA
diff --git a/Engine/Source/Runtime/Engine/Private/DebugViewModeMaterialProxy.cpp b/Engine/Source/Runtime/Engine/Private/DebugViewModeMaterialProxy.cpp
--- a/Engine/Source/Runtime/Engine/Private/DebugViewModeMaterialProxy.cpp
+++ b/Engine/Source/Runtime/Engine/Private/DebugViewModeMaterialProxy.cpp
@@ -10,235 +10,8 @@ DebugViewModeMaterialProxy.cpp : Contains definitions the debug view mode materi
 #include "EngineModule.h"
 #include "RendererInterface.h"
 
-ENGINE_API bool GetDebugViewMaterial(const UMaterialInterface* InMaterialInterface, EDebugViewShaderMode InDebugViewMode, ERHIFeatureLevel::Type InFeatureLevel, const FMaterialRenderProxy*& OutMaterialRenderProxy, const FMaterial*& OutMaterial)
-{
-#if WITH_EDITORONLY_DATA
-	return FDebugViewModeMaterialProxy::GetShader(InMaterialInterface, InDebugViewMode, InFeatureLevel, OutMaterialRenderProxy, OutMaterial);
-#else
-	return false;
-#endif
-}
-
-ENGINE_API bool HasMissingDebugViewModeShaders(bool bClearFlag)
-{
-#if WITH_EDITORONLY_DATA
-	if (FDebugViewModeMaterialProxy::MissingShadersChanged())
-	{
-		if (bClearFlag)
-		{
-			FDebugViewModeMaterialProxy::ClearMissingShadersFlag();
-		}
-		return true;
-	}
-#endif
-	return false;
-}
-
-ENGINE_API void ClearDebugViewMaterials(UMaterialInterface* InMaterialInterface)
-{
-#if WITH_EDITORONLY_DATA
-	FDebugViewModeMaterialProxy::ClearAllShaders(InMaterialInterface);
-#endif
-}
-
 #if WITH_EDITORONLY_DATA
 
-volatile bool FDebugViewModeMaterialProxy::bReentrantCall = false;
-TMap<FDebugViewModeMaterialProxy::FMaterialKey, FDebugViewModeMaterialProxy*> FDebugViewModeMaterialProxy::DebugMaterialShaderMap;
-TSet<FDebugViewModeMaterialProxy::FMaterialKey> FDebugViewModeMaterialProxy::MissingShaders;
-bool FDebugViewModeMaterialProxy::bMissingShadersChanged = false;
-
-void FDebugViewModeMaterialProxy::AddShader(
-	UMaterialInterface* InMaterialInterface, 
-	EMaterialQualityLevel::Type InQualityLevel, 
-	ERHIFeatureLevel::Type InFeatureLevel, 
-	bool bSynchronousCompilation, 
-	EDebugViewShaderMode InDebugViewMode)
-{
-	if (!InMaterialInterface) return;
-
-	const FMaterial* Material = InMaterialInterface->GetMaterialResource(InFeatureLevel);
-	const FDebugViewModeInterface* DebugViewModeInterface = FDebugViewModeInterface::GetInterface(InDebugViewMode);
-	if (!Material || !DebugViewModeInterface) return;
-
-	if (!DebugViewModeInterface->bNeedsMaterialProperties && FDebugViewModeInterface::AllowFallbackToDefaultMaterial(Material))
-	{
-		InMaterialInterface = UMaterial::GetDefaultMaterial(MD_Surface);
-	}
-
-	FMaterialKey ShaderMapKey(InMaterialInterface, InDebugViewMode, InFeatureLevel);
-	if (!DebugMaterialShaderMap.Contains(ShaderMapKey))
-	{
-		DebugMaterialShaderMap.Add(ShaderMapKey, new FDebugViewModeMaterialProxy(InMaterialInterface, InQualityLevel, InFeatureLevel, bSynchronousCompilation, InDebugViewMode));
-	}
-	MissingShaders.Remove(ShaderMapKey);
-}
-
-bool FDebugViewModeMaterialProxy::GetShader(const UMaterialInterface* InMaterialInterface, EDebugViewShaderMode InDebugViewMode, ERHIFeatureLevel::Type InFeatureLevel, const FMaterialRenderProxy*& OutMaterialRenderProxy, const FMaterial*& OutMaterial)
-{
-	FMaterialKey MaterialKey(InMaterialInterface, InDebugViewMode, InFeatureLevel);
-	FDebugViewModeMaterialProxy** BoundMaterial = DebugMaterialShaderMap.Find(MaterialKey);
-	if (BoundMaterial && *BoundMaterial)
-	{
-		if ((*BoundMaterial)->IsValid() && (*BoundMaterial)->GetRenderingThreadShaderMap())
-		{
-			OutMaterialRenderProxy = *BoundMaterial;
-			OutMaterial = *BoundMaterial;
-			return true;
-		}
-		else // This material is not usable for debug view modes.
-		{
-			return false;
-		}
-	}
-	else
-	{
-		if (!MissingShaders.Contains(MaterialKey))
-		{
-			MissingShaders.Add(MaterialKey);
-
-			// Note that a new shader key is missing, so that we can trigger recompilation.
-			// Because it is not garantied that this can be fixed (see ValidateAllShaders()), 
-			// we only kept track of new entries, and don't necessarily try to fix all of them.
-			bMissingShadersChanged = true;
-		}
-		return false;
-	}
-}
-
-void FDebugViewModeMaterialProxy::ClearAllShaders(UMaterialInterface* InMaterialInterface)
-{
-	if (bReentrantCall || DebugMaterialShaderMap.Num() == 0) return;
-
-	FlushRenderingCommands();
-	bReentrantCall = true;
-
-	TArray<FDebugViewModeMaterialProxy*, TInlineAllocator<1> > MaterialsToDelete;
-
-	if (!InMaterialInterface)
-	{
-		for (TMap<FMaterialKey, FDebugViewModeMaterialProxy*>::TIterator It(DebugMaterialShaderMap); It; ++It)
-		{
-			FDebugViewModeMaterialProxy* Proxy = It.Value();
-			MaterialsToDelete.Add(Proxy);
-			It.Value() = nullptr;
-		}
-		DebugMaterialShaderMap.Empty();
-		MissingShaders.Empty();
-	}
-	else
-	{
-		for (TMap<FMaterialKey, FDebugViewModeMaterialProxy*>::TIterator It(DebugMaterialShaderMap); It; ++It)
-		{
-			if (It.Key().MaterialInterface == InMaterialInterface)
-			{
-				FDebugViewModeMaterialProxy* Proxy = It.Value();
-				MaterialsToDelete.Add(Proxy);
-				It.Value() = nullptr;
-
-				It.RemoveCurrent();
-			}
-		}
-
-		for (TSet<FMaterialKey>::TIterator It(MissingShaders); It; ++It)
-		{
-			if (It->MaterialInterface == InMaterialInterface)
-			{
-				It.RemoveCurrent();
-			}
-		}
-	}
-
-	if (MaterialsToDelete.Num())
-	{
-		ENQUEUE_RENDER_COMMAND(DeleteDebugMaterials)([MaterialsToDelete](FRHICommandList& RHICmdList)
-		{
-			for (FDebugViewModeMaterialProxy* MaterialToDelete : MaterialsToDelete)
-			{
-				if (MaterialToDelete)
-				{
-					delete MaterialToDelete;
-				}
-			}
-		});
-
-		FlushRenderingCommands();
-	}
-
-	bReentrantCall = false;
-}
-
-void FDebugViewModeMaterialProxy::ValidateAllShaders(TSet<UMaterialInterface*>& Materials)
-{
-	FlushRenderingCommands();
-
-	TArray<FDebugViewModeMaterialProxy*> MaterialsToUpdate;
-
-	for (TMap<FMaterialKey, FDebugViewModeMaterialProxy*>::TIterator It(DebugMaterialShaderMap); It; ++It)
-	{
-		const UMaterialInterface* OriginalMaterialInterface = It.Key().MaterialInterface;
-		FDebugViewModeMaterialProxy* DebugMaterial = It.Value();
-
-		if (DebugMaterial != nullptr)
-		{
-			const FMaterial* OriginalMaterial = OriginalMaterialInterface->GetMaterialResource(DebugMaterial->FMaterial::GetFeatureLevel());
-
-			if (OriginalMaterial != nullptr && OriginalMaterial->GetGameThreadShaderMap() && DebugMaterial->GetGameThreadShaderMap())
-			{
-				const FUniformExpressionSet& DebugViewUniformExpressionSet = DebugMaterial->GetGameThreadShaderMap()->GetUniformExpressionSet();
-				const FUniformExpressionSet& OrignialUniformExpressionSet = OriginalMaterial->GetGameThreadShaderMap()->GetUniformExpressionSet();
-
-				if (DebugViewUniformExpressionSet == OrignialUniformExpressionSet)
-				{
-					MaterialsToUpdate.Add(DebugMaterial);
-				}
-				else
-				{
-					// This will happen when the debug shader compiled misses logic. Usually caused by custom features in the original shader compilation not implemented in FDebugViewModeMaterialProxy.
-					UE_LOG(TextureStreamingBuild, Verbose, TEXT("Uniform expression set mismatch for %s, skipping shader"), *DebugMaterial->GetMaterialInterface()->GetName());
-
-					// Here we can't destroy the invalid material because it would trigger ClearAllShaders.
-					DebugMaterial->MarkAsInvalid();
-					Materials.Remove(const_cast<UMaterialInterface*>(DebugMaterial->GetMaterialInterface()));
-				}
-			}
-			else
-			{
-				// When using synchronous compilation, it is normal for the original material to not be ready yet.
-				// In this case, we can't validate that the shader will be 100% compatible for overrides, meaning it is risky to use for viewmodes.
-				// This implies that viewmode can't use synchronous compilation.
-				if (!DebugMaterial->GetGameThreadShaderMap() || !DebugMaterial->RequiresSynchronousCompilation())
-				{
-					UE_LOG(TextureStreamingBuild, Verbose, TEXT("Can't get valid shadermap for %s, skipping shader"), *DebugMaterial->GetMaterialInterface()->GetName());
-
-					// Here we can't destroy the invalid material because it would trigger ClearAllShaders.
-					DebugMaterial->MarkAsInvalid();
-					Materials.Remove(const_cast<UMaterialInterface*>(DebugMaterial->GetMaterialInterface()));
-				}
-				else
-				{
-					MaterialsToUpdate.Add(DebugMaterial);
-				}
-				
-			}
-		}
-	}
-
-	if (MaterialsToUpdate.Num())
-	{
-		ENQUEUE_RENDER_COMMAND(UpdateDebugMaterialExpressionCache)([MaterialsToUpdate](FRHICommandList& RHICmdList)
-		{
-			for (FDebugViewModeMaterialProxy* MaterialToUpdate : MaterialsToUpdate)
-			{
-				check(MaterialToUpdate);
-				MaterialToUpdate->UpdateUniformExpressionCacheIfNeeded(MaterialToUpdate->FMaterial::GetFeatureLevel());
-			}
-		});
-
-		FlushRenderingCommands();
-	}
-}
-
 FDebugViewModeMaterialProxy::FDebugViewModeMaterialProxy(
 	UMaterialInterface* InMaterialInterface, 
 	EMaterialQualityLevel::Type QualityLevel, 
@@ -415,4 +188,4 @@ void FDebugViewModeMaterialProxy::GatherExpressionsForCustomInterpolators(TArray
 	}
 }
 
-#endif
\ No newline at end of file
+#endif // WITH_EDITORONLY_DATA
diff --git a/Engine/Source/Runtime/Engine/Private/DebugViewModeMaterialProxy.h b/Engine/Source/Runtime/Engine/Private/DebugViewModeMaterialProxy.h
--- a/Engine/Source/Runtime/Engine/Private/DebugViewModeMaterialProxy.h
+++ b/Engine/Source/Runtime/Engine/Private/DebugViewModeMaterialProxy.h
@@ -164,22 +164,6 @@ public:
 
 	virtual EMaterialShaderMapUsage::Type GetMaterialShaderMapUsage() const { return Usage; }
 
-	////////////////
-
-	static void AddShader(
-		UMaterialInterface* InMaterialInterface, 
-		EMaterialQualityLevel::Type QualityLevel, 
-		ERHIFeatureLevel::Type FeatureLevel, 
-		bool bSynchronousCompilation, 
-		EDebugViewShaderMode InDebugViewMode
-	);
-	static bool GetShader(const UMaterialInterface* InMaterialInterface, EDebugViewShaderMode InDebugViewMode, ERHIFeatureLevel::Type InFeatureLevel, const FMaterialRenderProxy*& OutMaterialRenderProxy, const FMaterial*& OutMaterial);
-	static void ClearAllShaders(UMaterialInterface* InMaterialInterface);
-	static bool HasAnyShaders() { return DebugMaterialShaderMap.Num() > 0; }
-	static void ValidateAllShaders(TSet<UMaterialInterface*>& Materials);
-	static bool MissingShadersChanged() { return bMissingShadersChanged; }
-	static void ClearMissingShadersFlag() { bMissingShadersChanged = false; }
-
 private:
 
 	/** The material interface for this proxy */
@@ -216,35 +200,6 @@ private:
 	bool bValid;
 	bool bIsDefaultMaterial;
 	bool bSynchronousCompilation;
-
-	struct FMaterialKey
-	{
-		FMaterialKey(const UMaterialInterface* InMaterialInterface, EDebugViewShaderMode InDebugViewMode, ERHIFeatureLevel::Type InFeatureLevel) : MaterialInterface(InMaterialInterface), DebugViewMode((uint32)InDebugViewMode), FeatureLevel((uint32)InFeatureLevel) {}
-
-		const UMaterialInterface* MaterialInterface;
-		uint32 DebugViewMode;
-		uint32 FeatureLevel;
-
-		friend bool operator==(const FMaterialKey& Lhs, const FMaterialKey& Rhs)
-		{
-			return Lhs.MaterialInterface == Rhs.MaterialInterface
-				&& Lhs.DebugViewMode == Rhs.DebugViewMode
-				&& Lhs.FeatureLevel == Rhs.FeatureLevel;
-		}
-
-		friend uint32 GetTypeHash(const FMaterialKey& Key)
-		{
-			const uint32 DebugViewModeShift = 6;
-			static_assert((2 ^ DebugViewModeShift) < DVSM_MAX, "Bit shift too small!");
-
-			return GetTypeHash(Key.MaterialInterface) ^ GetTypeHash((Key.FeatureLevel << DebugViewModeShift) ^ Key.DebugViewMode);
-		}
-	};
-
-	static volatile bool bReentrantCall;
-	static TMap<FMaterialKey, FDebugViewModeMaterialProxy*> DebugMaterialShaderMap;
-	static TSet<FMaterialKey> MissingShaders;
-	static bool bMissingShadersChanged;
 };
 
-#endif
+#endif // WITH_EDITORONLY_DATA
diff --git a/Engine/Source/Runtime/Engine/Private/GameViewportClient.cpp b/Engine/Source/Runtime/Engine/Private/GameViewportClient.cpp
--- a/Engine/Source/Runtime/Engine/Private/GameViewportClient.cpp
+++ b/Engine/Source/Runtime/Engine/Private/GameViewportClient.cpp
@@ -1121,14 +1121,8 @@ void UGameViewportClient::Draw(FViewport* InViewport, FCanvas* SceneCanvas)
 		// Force enable view family show flag for HighDPI derived's screen percentage.
 		ViewFamily.EngineShowFlags.ScreenPercentage = true;
 	}
-	if (ViewFamily.GetDebugViewShaderMode() != DVSM_None && HasMissingDebugViewModeShaders(true))
-	{
-		TSet<UMaterialInterface*> Materials;
-		if (GetUsedMaterialsInWorld(MyWorld, Materials, nullptr))
-		{
-			CompileDebugViewModeShaders(ViewFamily.GetDebugViewShaderMode(), GetCachedScalabilityCVars().MaterialQualityLevel, ViewFamily.GetFeatureLevel(), false, false, Materials, nullptr);
-		}
-	}
+
+	UpdateDebugViewModeShaders();
 #endif
 
 	ViewFamily.ViewExtensions = GEngine->ViewExtensions->GatherActiveExtensions(InViewport);
diff --git a/Engine/Source/Runtime/Engine/Private/Materials/MaterialShared.cpp b/Engine/Source/Runtime/Engine/Private/Materials/MaterialShared.cpp
--- a/Engine/Source/Runtime/Engine/Private/Materials/MaterialShared.cpp
+++ b/Engine/Source/Runtime/Engine/Private/Materials/MaterialShared.cpp
@@ -1667,10 +1667,15 @@ bool FMaterial::CacheShaders(const FMaterialShaderMapId& ShaderMapId, EShaderPla
 		// Find the material's cached shader map.
 		GameThreadShaderMap = FMaterialShaderMap::FindId(ShaderMapId, Platform);
 
-		// Attempt to load from the derived data cache if we are uncooked
-		if ((!GameThreadShaderMap || !GameThreadShaderMap->IsComplete(this, true)) && !FPlatformProperties::RequiresCookedData())
+		// On-the-fly view shaders are not using ddc currently, as their shadermap is not persistent.
+		// See FMaterialShaderMap::ProcessCompilationResults().
+		if  (GetMaterialShaderMapUsage() != EMaterialShaderMapUsage::DebugViewMode)
 		{
-			FMaterialShaderMap::LoadFromDerivedDataCache(this, ShaderMapId, Platform, GameThreadShaderMap);
+			// Attempt to load from the derived data cache if we are uncooked
+			if ((!GameThreadShaderMap || !GameThreadShaderMap->IsComplete(this, true)) && !FPlatformProperties::RequiresCookedData())
+			{
+				FMaterialShaderMap::LoadFromDerivedDataCache(this, ShaderMapId, Platform, GameThreadShaderMap);
+			}
 		}
 #endif // WITH_EDITOR
 	}
diff --git a/Engine/Source/Runtime/Engine/Private/ShaderCompiler/ShaderCompiler.cpp b/Engine/Source/Runtime/Engine/Private/ShaderCompiler/ShaderCompiler.cpp
--- a/Engine/Source/Runtime/Engine/Private/ShaderCompiler/ShaderCompiler.cpp
+++ b/Engine/Source/Runtime/Engine/Private/ShaderCompiler/ShaderCompiler.cpp
@@ -1027,7 +1027,10 @@ int32 FShaderCompileThreadRunnable::PullTasksFromQueue()
 
 				if (Manager->CompileQueue.Num() > 0)
 				{
-					UE_LOG(LogShaderCompilers, Display, TEXT("Shaders left to compile %i"), Manager->CompileQueue.Num());
+					if (Manager->CompileQueue.Num() % 10 == 0)
+					{
+						UE_LOG(LogShaderCompilers, Display, TEXT("Shaders left to compile %i"), Manager->CompileQueue.Num());
+					}
 
 					bool bAddedLowLatencyTask = false;
 					int32 JobIndex = 0;
diff --git a/Engine/Source/Runtime/Engine/Private/UnrealEngine.cpp b/Engine/Source/Runtime/Engine/Private/UnrealEngine.cpp
--- a/Engine/Source/Runtime/Engine/Private/UnrealEngine.cpp
+++ b/Engine/Source/Runtime/Engine/Private/UnrealEngine.cpp
@@ -9681,7 +9681,6 @@ float DrawMapWarnings(UWorld* World, FViewport* Viewport, FCanvas* Canvas, UCanv
 			Canvas->DrawItem(SmallTextItem, FVector2D(MessageX, MessageY));
 			MessageY += FontSizeY;
 		}
-
 	}
 
 	if (World->bKismetScriptError)
@@ -9694,8 +9693,19 @@ float DrawMapWarnings(UWorld* World, FViewport* Viewport, FCanvas* Canvas, UCanv
 
 	SmallTextItem.SetColor(FLinearColor::White);
 
+	extern double GViewModeShaderMissingTime;
+	extern int32 GNumViewModeShaderMissing;
+	if (FApp::GetCurrentTime() - GViewModeShaderMissingTime < 1 && GNumViewModeShaderMissing > 0)
+	{
+		SmallTextItem.SetColor(FLinearColor::Yellow);
+		SmallTextItem.Text = FText::Format(LOCTEXT("ViewModeShadersCompilingFmt", "View Mode Shaders Compiling ({0})"), GNumViewModeShaderMissing);
+		Canvas->DrawItem(SmallTextItem, FVector2D(MessageX, MessageY));
+		MessageY += FontSizeY;
+	}
+
 	if (GShaderCompilingManager && GShaderCompilingManager->IsCompiling())
 	{
+		SmallTextItem.SetColor(FLinearColor::White);
 		SmallTextItem.Text = FText::Format(LOCTEXT("ShadersCompilingFmt", "Shaders Compiling ({0})"), GShaderCompilingManager->GetNumRemainingJobs());
 		Canvas->DrawItem(SmallTextItem, FVector2D(MessageX, MessageY));
 		MessageY += FontSizeY;
diff --git a/Engine/Source/Runtime/Engine/Public/DebugViewModeHelpers.h b/Engine/Source/Runtime/Engine/Public/DebugViewModeHelpers.h
--- a/Engine/Source/Runtime/Engine/Public/DebugViewModeHelpers.h
+++ b/Engine/Source/Runtime/Engine/Public/DebugViewModeHelpers.h
@@ -44,7 +44,7 @@ FORCEINLINE bool AllowDebugViewShaderMode(EDebugViewShaderMode ShaderMode, EShad
 ENGINE_API int32 GetNumActorsInWorld(UWorld* InWorld);
 ENGINE_API bool GetUsedMaterialsInWorld(UWorld* InWorld, OUT TSet<UMaterialInterface*>& OutMaterials, FSlowTask* Task);
 ENGINE_API bool CompileDebugViewModeShaders(EDebugViewShaderMode Mode, EMaterialQualityLevel::Type QualityLevel, ERHIFeatureLevel::Type FeatureLevel, bool bFullRebuild, bool bWaitForPreviousShaders, TSet<UMaterialInterface*>& Materials, FSlowTask* ProgressTask);
-ENGINE_API bool HasMissingDebugViewModeShaders(bool bClearFlag);
+ENGINE_API void UpdateDebugViewModeShaders();
 
 ENGINE_API bool WaitForShaderCompilation(const FText& Message, FSlowTask* ProgressTask);
 
diff --git a/Engine/Source/Runtime/Renderer/Private/DebugViewModeRendering.cpp b/Engine/Source/Runtime/Renderer/Private/DebugViewModeRendering.cpp
--- a/Engine/Source/Runtime/Renderer/Private/DebugViewModeRendering.cpp
+++ b/Engine/Source/Runtime/Renderer/Private/DebugViewModeRendering.cpp
@@ -50,7 +50,6 @@ IMPLEMENT_MATERIAL_SHADER_TYPE(,FDebugViewModeHS,TEXT("/Engine/Private/DebugView
 IMPLEMENT_MATERIAL_SHADER_TYPE(,FDebugViewModeDS,TEXT("/Engine/Private/DebugViewModeVertexShader.usf"),TEXT("MainDomain"),SF_Domain);
 
 ENGINE_API bool GetDebugViewMaterial(const UMaterialInterface* InMaterialInterface, EDebugViewShaderMode InDebugViewMode, ERHIFeatureLevel::Type InFeatureLevel,const FMaterialRenderProxy*& OutMaterialRenderProxy, const FMaterial*& OutMaterial);
-ENGINE_API bool HasMissingDebugViewModeShaders(bool bClearFlag);
 
 void FDeferredShadingSceneRenderer::DoDebugViewModePostProcessing(FRHICommandListImmediate& RHICmdList, const FViewInfo& View, TRefCountPtr<IPooledRenderTarget>& VelocityRT)
 {
diff --git a/Engine/Source/Runtime/Renderer/Private/PostProcess/PostProcessVisualizeComplexity.cpp b/Engine/Source/Runtime/Renderer/Private/PostProcess/PostProcessVisualizeComplexity.cpp
--- a/Engine/Source/Runtime/Renderer/Private/PostProcess/PostProcessVisualizeComplexity.cpp
+++ b/Engine/Source/Runtime/Renderer/Private/PostProcess/PostProcessVisualizeComplexity.cpp
@@ -62,7 +62,8 @@ void FVisualizeComplexityApplyPS::SetParameters(
 	const TArray<FLinearColor>& Colors,
 	EColorSampling ColorSampling,
 	float ComplexityScale,
-	bool bLegend
+	bool bLegend,
+	bool bShowError
 	)
 {
 	const FPixelShaderRHIParamRef ShaderRHI = GetPixelShader();
@@ -106,7 +107,7 @@ void FVisualizeComplexityApplyPS::SetParameters(
 
 	FIntPoint UsedQuadBufferSize = (Context.View.ViewRect.Size() + FIntPoint(1, 1)) / 2;
 	SetShaderValue(RHICmdList, ShaderRHI, ShaderComplexityParams, FVector4(bLegend, DebugViewShaderMode, ColorSampling, ComplexityScale));
-	SetShaderValue(RHICmdList, ShaderRHI, ShaderComplexityParams2, FVector4((float)Colors.Num(), 0, (float)UsedQuadBufferSize.X, (float)UsedQuadBufferSize.Y));
+	SetShaderValue(RHICmdList, ShaderRHI, ShaderComplexityParams2, FVector4((float)Colors.Num(), bShowError, (float)UsedQuadBufferSize.X, (float)UsedQuadBufferSize.Y));
 }
 
 IMPLEMENT_SHADER_TYPE(,FVisualizeComplexityApplyPS,TEXT("/Engine/Private/ShaderComplexityApplyPixelShader.usf"),TEXT("Main"),SF_Pixel);
@@ -158,7 +159,7 @@ void FRCPassPostProcessVisualizeComplexity::Process(FRenderingCompositePassConte
 
 		SetGraphicsPipelineState(Context.RHICmdList, GraphicsPSOInit);
 
-		PixelShader->SetParameters(Context.RHICmdList, Context, Colors, ColorSampling, ComplexityScale, bLegend);
+		PixelShader->SetParameters(Context.RHICmdList, Context, Colors, ColorSampling, ComplexityScale, bLegend, ViewFamily.GetDebugViewShaderMode() != DVSM_QuadComplexity);
 
 		DrawRectangle(
 			Context.RHICmdList,
diff --git a/Engine/Source/Runtime/Renderer/Private/PostProcess/PostProcessVisualizeComplexity.h b/Engine/Source/Runtime/Renderer/Private/PostProcess/PostProcessVisualizeComplexity.h
--- a/Engine/Source/Runtime/Renderer/Private/PostProcess/PostProcessVisualizeComplexity.h
+++ b/Engine/Source/Runtime/Renderer/Private/PostProcess/PostProcessVisualizeComplexity.h
@@ -51,7 +51,7 @@ public:
 	FVisualizeComplexityApplyPS() {}
 
 	template <typename TRHICmdList>
-	void SetParameters(TRHICmdList& RHICmdList, const FRenderingCompositePassContext& Context, const TArray<FLinearColor>& Colors, EColorSampling ColorSampling, float ComplexityScale, bool bLegend);
+	void SetParameters(TRHICmdList& RHICmdList, const FRenderingCompositePassContext& Context, const TArray<FLinearColor>& Colors, EColorSampling ColorSampling, float ComplexityScale, bool bLegend, bool bShowError);
 
 	static bool ShouldCompilePermutation(const FGlobalShaderPermutationParameters& Parameters)
 	{
