Disallow ES31 shaders attempting to use SceneDepth in the vertex shader

diff --git a/Engine/Source/Runtime/Engine/Private/Materials/HLSLMaterialTranslator.h b/Engine/Source/Runtime/Engine/Private/Materials/HLSLMaterialTranslator.h
--- a/Engine/Source/Runtime/Engine/Private/Materials/HLSLMaterialTranslator.h
+++ b/Engine/Source/Runtime/Engine/Private/Materials/HLSLMaterialTranslator.h
@@ -3731,8 +3731,7 @@ protected:
 		if (ShaderFrequency == SF_Vertex && FeatureLevel <= ERHIFeatureLevel::ES3_1)
 		{
 			// mobile currently does not support this, we need to read a separate copy of the depth, we must disable framebuffer fetch and force scene texture reads.
-			//return Errorf(TEXT("Cannot read scene depth from the vertex shader with feature level ES3.1 or below."));
-			return Constant(65000.0f);
+			return Errorf(TEXT("Cannot read scene depth from the vertex shader with feature level ES3.1 or below."));
 		}
 
 		if (Offset == INDEX_NONE && bUseOffset)
