Fixed lighting only view mode. Handling EngineShowFlags.Materials for bMaterialModifiesMeshPosition had bad side effects due to how FOverrideSelectionColorMaterialRenderProxy is implemented, as it doesn't replace the color completely, just adds a provided value to the emissive channel.

diff --git a/Engine/Source/Runtime/Engine/Private/PrimitiveDrawingUtils.cpp b/Engine/Source/Runtime/Engine/Private/PrimitiveDrawingUtils.cpp
--- a/Engine/Source/Runtime/Engine/Private/PrimitiveDrawingUtils.cpp
+++ b/Engine/Source/Runtime/Engine/Private/PrimitiveDrawingUtils.cpp
@@ -1493,18 +1493,7 @@ void ApplyViewModeOverrides(
 				}
 			}
 
-			if (bMaterialModifiesMeshPosition)
-			{
-				// If the material is mesh-modifying, we cannot rely on substitution.
-				auto LightingOnlyMaterialInstance = new FOverrideSelectionColorMaterialRenderProxy(
-					Mesh.MaterialRenderProxy, 
-					GEngine->LightingOnlyBrightness
-				);
-
-				Mesh.MaterialRenderProxy = LightingOnlyMaterialInstance;
-				Collector.RegisterOneFrameMaterialProxy(LightingOnlyMaterialInstance);
-			}
-			else if (bTextureMapped == false)
+			if (bTextureMapped == false)
 			{
 				FMaterialRenderProxy* RenderProxy = GEngine->LevelColorationLitMaterial->GetRenderProxy();
 				auto LightingOnlyMaterialInstance = new FColoredMaterialRenderProxy(
