Movable lights are no longer assigned a LightGuid

diff --git a/Engine/Source/Editor/UnrealEd/Private/EditorServer.cpp b/Engine/Source/Editor/UnrealEd/Private/EditorServer.cpp
--- a/Engine/Source/Editor/UnrealEd/Private/EditorServer.cpp
+++ b/Engine/Source/Editor/UnrealEd/Private/EditorServer.cpp
@@ -3848,7 +3848,9 @@ bool UEditorEngine::Map_Check( UWorld* InWorld, const TCHAR* Str, FOutputDevice&
 			if( LightActor )
 			{
 				ULightComponent* LightComponent = LightActor->GetLightComponent();
-				if (LightComponent) // LightComponent component can be null, for example when creating a blueprint deriving from ALight.
+				// LightComponent component can be null, for example when creating a blueprint deriving from ALight.
+				// Movable light components have a light guid of 0, so skip them as well
+				if (LightComponent && !LightComponent->HasStaticShadowing()) 
 				{
 					AActor* ExistingLightActor = LightGuidToActorMap.FindRef( LightComponent->LightGuid );
 					if( ExistingLightActor )
diff --git a/Engine/Source/Runtime/Engine/Classes/Components/LightComponentBase.h b/Engine/Source/Runtime/Engine/Classes/Components/LightComponentBase.h
--- a/Engine/Source/Runtime/Engine/Classes/Components/LightComponentBase.h
+++ b/Engine/Source/Runtime/Engine/Classes/Components/LightComponentBase.h
@@ -200,6 +200,9 @@ class ENGINE_API ULightComponentBase : public USceneComponent
 	bool HasStaticShadowing() const;
 
 #if WITH_EDITOR
+	/** UObject Interface */
+	virtual void PostLoad() override;
+
 	/** UActorComponent Interface */
 	virtual void OnRegister() override;
 	virtual bool CanEditChange(const UProperty* InProperty) const override;
diff --git a/Engine/Source/Runtime/Engine/Private/Components/LightComponent.cpp b/Engine/Source/Runtime/Engine/Private/Components/LightComponent.cpp
--- a/Engine/Source/Runtime/Engine/Private/Components/LightComponent.cpp
+++ b/Engine/Source/Runtime/Engine/Private/Components/LightComponent.cpp
@@ -166,32 +166,38 @@ void ULightComponentBase::PostEditChangeProperty(FPropertyChangedEvent& Property
 void ULightComponentBase::ValidateLightGUIDs()
 {
 	// Validate light guids.
-	if( !LightGuid.IsValid() )
+	if (!LightGuid.IsValid())
 	{
-		LightGuid = FGuid::NewGuid();
+		UpdateLightGUIDs();
 	}
 }
 
 void ULightComponentBase::UpdateLightGUIDs()
 {
-	LightGuid = FGuid::NewGuid();
+	LightGuid = (HasStaticShadowing() ? FGuid::NewGuid() : FGuid());
 }
 
 bool ULightComponentBase::HasStaticLighting() const
 {
-	AActor* Owner = GetOwner();
-
-	return Owner && (Mobility == EComponentMobility::Static);
+	return (Mobility == EComponentMobility::Static) && GetOwner();
 }
 
 bool ULightComponentBase::HasStaticShadowing() const
 {
-	AActor* Owner = GetOwner();
-
-	return Owner && (Mobility != EComponentMobility::Movable);
+	return (Mobility != EComponentMobility::Movable) && GetOwner();
 }
 
 #if WITH_EDITOR
+void ULightComponentBase::PostLoad()
+{
+	Super::PostLoad();
+
+	if (!HasStaticShadowing())
+	{
+		LightGuid.Invalidate();
+	}
+}
+
 void ULightComponentBase::OnRegister()
 {
 	Super::OnRegister();
@@ -1090,7 +1096,7 @@ void ULightComponent::InvalidateLightingCacheDetailed(bool bInvalidateBuildEnque
 /** Invalidates the light's cached lighting with the option to recreate the light Guids. */
 void ULightComponent::InvalidateLightingCacheInner(bool bRecreateLightGuids)
 {
-	if (HasStaticLighting() || HasStaticShadowing())
+	if (HasStaticShadowing())
 	{
 		// Save the light state for transactions.
 		Modify();
@@ -1109,6 +1115,10 @@ void ULightComponent::InvalidateLightingCacheInner(bool bRecreateLightGuids)
 
 		MarkRenderStateDirty();
 	}
+	else
+	{
+		LightGuid.Invalidate();
+	}
 }
 
 TStructOnScope<FActorComponentInstanceData> ULightComponent::GetComponentInstanceData() const
@@ -1121,16 +1131,12 @@ void ULightComponent::ApplyComponentInstanceData(FPrecomputedLightInstanceData*
 {
 	check(LightMapData);
 
-	LightGuid = LightMapData->LightGuid;
-
 	if (!LightMapData->Transform.Equals(GetComponentTransform()))
 	{
-		// Rather than resetting the guid and reallocating one randomly, we'll offset it so that consecutive loads 
-		// behave deterministically  while still disassociating it from any now invalidated baked lighting
-		++LightGuid.D;
 		return;
 	}
 
+	LightGuid = (HasStaticShadowing() ? LightMapData->LightGuid : FGuid());
 	PreviewShadowMapChannel = LightMapData->PreviewShadowMapChannel;
 
 	MarkRenderStateDirty();
diff --git a/Engine/Source/Runtime/Engine/Private/Components/SkyLightComponent.cpp b/Engine/Source/Runtime/Engine/Private/Components/SkyLightComponent.cpp
--- a/Engine/Source/Runtime/Engine/Private/Components/SkyLightComponent.cpp
+++ b/Engine/Source/Runtime/Engine/Private/Components/SkyLightComponent.cpp
@@ -553,7 +553,7 @@ void USkyLightComponent::ApplyComponentInstanceData(FPrecomputedSkyLightInstance
 {
 	check(LightMapData);
 
-	LightGuid = LightMapData->LightGuid;
+	LightGuid = (HasStaticShadowing() ? LightMapData->LightGuid : FGuid());
 	ProcessedSkyTexture = LightMapData->ProcessedSkyTexture;
 	IrradianceEnvironmentMap = LightMapData->IrradianceEnvironmentMap;
 	AverageBrightness = LightMapData->AverageBrightness;
