Undo mem map serialization changes to CompressedTrackOffsets. Save out compressed data in correct order instead

diff --git a/Engine/Source/Runtime/Engine/Private/Animation/AnimEncoding.cpp b/Engine/Source/Runtime/Engine/Private/Animation/AnimEncoding.cpp
--- a/Engine/Source/Runtime/Engine/Private/Animation/AnimEncoding.cpp
+++ b/Engine/Source/Runtime/Engine/Private/Animation/AnimEncoding.cpp
@@ -261,133 +261,6 @@ void AnimEncodingLegacyBase::GetBoneAtom(FTransform& OutAtom, FAnimSequenceDecom
 	}
 }
 
-
-bool AnimEncodingLegacyBase::CanBeMemoryMapped(UAnimSequence& Seq, int32 TotalSize)
-{
-	bool bHasValidScale = Seq.CompressedScaleOffsets.IsValid();
-	if (Seq.CompressedSegments.Num() != 0 || !bHasValidScale)
-	{
-		return true;
-	}
-	const int32 NumTracks = Seq.CompressedTrackOffsets.Num() / 4;
-
-	bool bIsBackwards = false;
-	for (int32 TrackIndex = 0; TrackIndex < NumTracks; ++TrackIndex)
-	{
-		const int32 OffsetScale = Seq.CompressedScaleOffsets.GetOffsetData(TrackIndex, 0);
-		const int32 OffsetTrans = Seq.CompressedTrackOffsets[TrackIndex * 4 + 0];
-		const int32 OffsetRot = Seq.CompressedTrackOffsets[TrackIndex * 4 + 2];
-		if (OffsetTrans > OffsetScale || OffsetRot > OffsetScale)
-		{
-			bIsBackwards = true;
-			check(OffsetRot < 0 || OffsetTrans < OffsetRot); // trans first
-		}
-		else
-		{
-			check(!bIsBackwards || (OffsetTrans < 0 && OffsetRot < 0) || (OffsetScale < 0)); // if one is backwards, they all are (or they are neither backwards nor forwards)
-		}
-	}
-	if (bIsBackwards)
-	{
-		int32 Offset = 0;
-		for (int32 TrackIndex = 0; TrackIndex < NumTracks; ++TrackIndex)
-		{
-			int32 TotalTrackSize = -1;
-			
-			for (int32 Next = TrackIndex + 1; Next < NumTracks; Next++)
-			{
-				const int32 OffsetScaleNext = Seq.CompressedScaleOffsets.GetOffsetData(Next, 0);
-				const int32 OffsetTransNext = Seq.CompressedTrackOffsets[Next * 4 + 0];
-				const int32 OffsetRotNext = Seq.CompressedTrackOffsets[Next * 4 + 2];
-				if (OffsetScaleNext >= 0)
-				{
-					TotalTrackSize = OffsetScaleNext - Offset;
-					break;
-				}
-				if (OffsetTransNext >= 0)
-				{
-					TotalTrackSize = OffsetTransNext - Offset;
-					break;
-				}
-				if (OffsetRotNext >= 0)
-				{
-					TotalTrackSize = OffsetRotNext - Offset;
-					break;
-				}
-			}
-			if (TotalTrackSize == -1)
-			{
-				TotalTrackSize = TotalSize - Offset;
-			}
-			if (TotalTrackSize == 0)
-			{
-				continue; // no data for this track at all, nothing to do
-			}
-
-
-			const int32 OffsetScale = Seq.CompressedScaleOffsets.GetOffsetData(TrackIndex, 0);
-			const int32 OffsetTrans = Seq.CompressedTrackOffsets[TrackIndex * 4 + 0];
-			const int32 OffsetRot = Seq.CompressedTrackOffsets[TrackIndex * 4 + 2];
-
-			int32 ScaleSize = 0;
-			if (OffsetScale >= 0)
-			{
-				if (OffsetTrans >= 0)
-				{
-					ScaleSize = OffsetTrans - OffsetScale;
-				}
-				else if (OffsetRot >= 0)
-				{
-					ScaleSize = OffsetRot - OffsetScale;
-				}
-				else
-				{
-					ScaleSize = TotalTrackSize;
-				}
-				check(ScaleSize > 0);
-			}
-
-			int32 TransSize = 0;
-			if (OffsetTrans >= 0)
-			{
-				if (OffsetRot >= 0)
-				{
-					TransSize = OffsetRot - OffsetTrans;
-				}
-				else
-				{
-					TransSize = TotalTrackSize - ScaleSize;
-				}
-				check(TransSize > 0);
-			}
-			int32 RotSize = 0;
-			if (OffsetRot >= 0)
-			{
-				RotSize = TotalTrackSize - TransSize - ScaleSize;
-				check(RotSize > 0);
-			}
-
-			if (OffsetTrans >= 0)
-			{
-				Seq.CompressedTrackOffsets[TrackIndex * 4 + 0] = Offset;
-				Offset += TransSize;
-			}
-			if (OffsetRot >= 0)
-			{
-				Seq.CompressedTrackOffsets[TrackIndex * 4 + 2] = Offset;
-				Offset += RotSize;
-			}
-			if (OffsetScale >= 0)
-			{
-				Seq.CompressedScaleOffsets.SetOffsetData(TrackIndex, 0, Offset);
-				Offset += ScaleSize;
-			}
-		}
-		check(Offset == TotalSize);
-	}
-	return true;
-}
-
 /**
  * Handles Byte-swapping incoming animation data from a MemoryReader
  *
@@ -483,7 +356,8 @@ PRAGMA_ENABLE_OPTIMIZATION
 void AnimEncodingLegacyBase::ByteSwapOut(
 	UAnimSequence& Seq,
 	TArray<uint8>& SerializedData, 
-	bool ForceByteSwapping)
+	bool ForceByteSwapping,
+	bool bMaintainComponentOrder)
 {
 	FMemoryWriter MemoryWriter( SerializedData, true );
 	MemoryWriter.SetByteSwapping( ForceByteSwapping );
@@ -507,6 +381,10 @@ void AnimEncodingLegacyBase::ByteSwapOut(
 		const int32 OffsetRot		= Seq.CompressedTrackOffsets[TrackIndex*4+2];
 		const int32 NumKeysRot	= Seq.CompressedTrackOffsets[TrackIndex*4+3];
 
+		if (bMaintainComponentOrder)
+		{
+			MemoryWriter.Seek(OffsetTrans);
+		}
 		// Translation data.
 		checkSlow( (OffsetTrans % 4) == 0 && "CompressedByteStream not aligned to four bytes" );
 		uint8* TransTrackData = StreamBase + OffsetTrans;
@@ -522,6 +400,10 @@ void AnimEncodingLegacyBase::ByteSwapOut(
 		// Like the compressed byte stream, pad the serialization stream to four bytes.
 		PadMemoryWriter(&MemoryWriter, TransTrackData, 4);
 
+		if (bMaintainComponentOrder)
+		{
+			MemoryWriter.Seek(OffsetRot);
+		}
 		// Rotation data.
 		checkSlow( (OffsetRot % 4) == 0 && "CompressedByteStream not aligned to four bytes" );
 		uint8* RotTrackData = StreamBase + OffsetRot;
@@ -536,6 +418,11 @@ void AnimEncodingLegacyBase::ByteSwapOut(
 			const int32 OffsetScale	= Seq.CompressedScaleOffsets.GetOffsetData(TrackIndex, 0);
 			const int32 NumKeysScale	= Seq.CompressedScaleOffsets.GetOffsetData(TrackIndex, 1);
 
+			if (bMaintainComponentOrder)
+			{
+				MemoryWriter.Seek(OffsetScale);
+			}
+
 			// Scale data.
 			checkSlow( (OffsetScale % 4) == 0 && "CompressedByteStream not aligned to four bytes" );
 			uint8* ScaleTrackData = StreamBase + OffsetScale;
diff --git a/Engine/Source/Runtime/Engine/Private/Animation/AnimEncoding_PerTrackCompression.cpp b/Engine/Source/Runtime/Engine/Private/Animation/AnimEncoding_PerTrackCompression.cpp
--- a/Engine/Source/Runtime/Engine/Private/Animation/AnimEncoding_PerTrackCompression.cpp
+++ b/Engine/Source/Runtime/Engine/Private/Animation/AnimEncoding_PerTrackCompression.cpp
@@ -353,13 +353,17 @@ static FORCEINLINE_DEBUGGABLE VectorRegister DecompressSingleTrackRotationVector
  * @param	Offset				The starting offset into the compressed byte stream for this track (can be INDEX_NONE to indicate an identity track)
  */
 template<class TArchive>
-void AEFPerTrackCompressionCodec::ByteSwapOneTrack(UAnimSequence& Seq, TArchive& MemoryStream, int32 Offset)
+void AEFPerTrackCompressionCodec::ByteSwapOneTrack(UAnimSequence& Seq, TArchive& MemoryStream, int32 Offset, bool bMaintainComponentOrder)
 {
 	// Translation data.
 	if (Offset != INDEX_NONE)
 	{
 		checkSlow( (Offset % 4) == 0 && "CompressedByteStream not aligned to four bytes" );
 
+		if (bMaintainComponentOrder)
+		{
+			MemoryStream.Seek(Offset);
+		}
 		uint8* TrackData = Seq.CompressedByteStream.GetData() + Offset;
 
 		// Read the header
@@ -412,8 +416,8 @@ void AEFPerTrackCompressionCodec::ByteSwapOneTrack(UAnimSequence& Seq, TArchive&
 	}
 }
 
-template void AEFPerTrackCompressionCodec::ByteSwapOneTrack(UAnimSequence& Seq, FMemoryReader& MemoryStream, int32 Offset);
-template void AEFPerTrackCompressionCodec::ByteSwapOneTrack(UAnimSequence& Seq, FMemoryWriter& MemoryStream, int32 Offset);
+template void AEFPerTrackCompressionCodec::ByteSwapOneTrack(UAnimSequence& Seq, FMemoryReader& MemoryStream, int32 Offset, bool bMaintainComponentOrder);
+template void AEFPerTrackCompressionCodec::ByteSwapOneTrack(UAnimSequence& Seq, FMemoryWriter& MemoryStream, int32 Offset, bool bMaintainComponentOrder);
 
 
 /**
@@ -444,134 +448,6 @@ void AEFPerTrackCompressionCodec::PreservePadding(uint8*& TrackData, FMemoryArch
 	}
 }
 
-bool AEFPerTrackCompressionCodec::CanBeMemoryMapped(UAnimSequence& Seq, int32 TotalSize)
-{
-	bool bHasValidScale = Seq.CompressedScaleOffsets.IsValid();
-	if (Seq.CompressedSegments.Num() != 0 || !bHasValidScale)
-	{
-		return true;
-	}
-	const int32 NumTracks = Seq.CompressedTrackOffsets.Num() / 2;
-	bool bIsBackwards = false;
-	for (int32 TrackIndex = 0; TrackIndex < NumTracks; ++TrackIndex)
-	{
-		const int32 OffsetScale = Seq.CompressedScaleOffsets.GetOffsetData(TrackIndex, 0);
-		const int32 OffsetTrans = Seq.CompressedTrackOffsets[TrackIndex * 2 + 0];
-		const int32 OffsetRot = Seq.CompressedTrackOffsets[TrackIndex * 2 + 1];
-		if (OffsetTrans > OffsetScale || OffsetRot > OffsetScale)
-		{
-			bIsBackwards = true;
-			check(OffsetRot < 0 || OffsetTrans < OffsetRot); // trans first
-		}
-		else
-		{
-			check(!bIsBackwards || (OffsetTrans < 0 && OffsetRot < 0) || (OffsetScale < 0)); // if one is backwards, they all are (or they are neither backwards nor forwards)
-		}
-	}
-	if (bIsBackwards)
-	{
-		int32 Offset = 0;
-		for (int32 TrackIndex = 0; TrackIndex < NumTracks; ++TrackIndex)
-		{
-			int32 TotalTrackSize = -1;
-
-			for (int32 Next = TrackIndex + 1; Next < NumTracks; Next++)
-			{
-				const int32 OffsetScaleNext = Seq.CompressedScaleOffsets.GetOffsetData(Next, 0);
-				const int32 OffsetTransNext = Seq.CompressedTrackOffsets[Next * 2 + 0];
-				const int32 OffsetRotNext = Seq.CompressedTrackOffsets[Next * 2 + 1];
-				if (OffsetScaleNext >= 0)
-				{
-					TotalTrackSize = OffsetScaleNext - Offset;
-					break;
-				}
-				if (OffsetTransNext >= 0)
-				{
-					TotalTrackSize = OffsetTransNext - Offset;
-					break;
-				}
-				if (OffsetRotNext >= 0)
-				{
-					TotalTrackSize = OffsetRotNext - Offset;
-					break;
-				}
-			}
-			if (TotalTrackSize == -1)
-			{
-				TotalTrackSize = TotalSize - Offset;
-			}
-			if (TotalTrackSize == 0)
-			{
-				continue; // no data for this track at all, nothing to do
-			}
-
-
-			const int32 OffsetScale = Seq.CompressedScaleOffsets.GetOffsetData(TrackIndex, 0);
-			const int32 OffsetTrans = Seq.CompressedTrackOffsets[TrackIndex * 2 + 0];
-			const int32 OffsetRot = Seq.CompressedTrackOffsets[TrackIndex * 2 + 1];
-
-
-
-			int32 ScaleSize = 0;
-			if (OffsetScale >= 0)
-			{
-				if (OffsetTrans >= 0)
-				{
-					ScaleSize = OffsetTrans - OffsetScale;
-				}
-				else if (OffsetRot >= 0)
-				{
-					ScaleSize = OffsetRot - OffsetScale;
-				}
-				else
-				{
-					ScaleSize = TotalTrackSize;
-				}
-				check(ScaleSize > 0);
-			}
-
-			int32 TransSize = 0;
-			if (OffsetTrans >= 0)
-			{
-				if (OffsetRot >= 0)
-				{
-					TransSize = OffsetRot - OffsetTrans;
-				}
-				else
-				{
-					TransSize = TotalTrackSize - ScaleSize;
-				}
-				check(TransSize > 0);
-			}
-			int32 RotSize = 0;
-			if (OffsetRot >= 0)
-			{
-				RotSize = TotalTrackSize - TransSize - ScaleSize;
-				check(RotSize > 0);
-			}
-
-			if (OffsetTrans >= 0)
-			{
-				Seq.CompressedTrackOffsets[TrackIndex * 2 + 0] = Offset;
-				Offset += TransSize;
-			}
-			if (OffsetRot >= 0)
-			{
-				Seq.CompressedTrackOffsets[TrackIndex * 2 + 1] = Offset;
-				Offset += RotSize;
-			}
-			if (OffsetScale >= 0)
-			{
-				Seq.CompressedScaleOffsets.SetOffsetData(TrackIndex, 0, Offset);
-				Offset += ScaleSize;
-			}
-
-		}
-		check(Offset == TotalSize);
-	}
-	return true;
-}
-
 /**
  * Handles Byte-swapping incoming animation data from a MemoryReader
  *
@@ -626,7 +502,8 @@ void AEFPerTrackCompressionCodec::ByteSwapIn(
 void AEFPerTrackCompressionCodec::ByteSwapOut(
 	UAnimSequence& Seq,
 	TArray<uint8>& SerializedData, 
-	bool ForceByteSwapping)
+	bool ForceByteSwapping,
+	bool bMaintainComponentOrder)
 {
 	FMemoryWriter MemoryWriter(SerializedData, true);
 	MemoryWriter.SetByteSwapping(ForceByteSwapping);
@@ -643,14 +520,15 @@ void AEFPerTrackCompressionCodec::ByteSwapOut(
 	for ( int32 TrackIndex = 0; TrackIndex < NumTracks; ++TrackIndex )
 	{
 		const int32 OffsetTrans = Seq.CompressedTrackOffsets[TrackIndex*2+0];
-		ByteSwapOneTrack(Seq, MemoryWriter, OffsetTrans);
+		ByteSwapOneTrack(Seq, MemoryWriter, OffsetTrans, bMaintainComponentOrder);
 
 		const int32 OffsetRot = Seq.CompressedTrackOffsets[TrackIndex*2+1];
-		ByteSwapOneTrack(Seq, MemoryWriter, OffsetRot);
+		ByteSwapOneTrack(Seq, MemoryWriter, OffsetRot, bMaintainComponentOrder);
+
 		if (bHasScaleData)
 		{
 			const int32 OffsetScale = Seq.CompressedScaleOffsets.GetOffsetData(TrackIndex, 0);
-			ByteSwapOneTrack(Seq, MemoryWriter, OffsetScale);
+			ByteSwapOneTrack(Seq, MemoryWriter, OffsetScale, bMaintainComponentOrder);
 		}
 	}
 }
diff --git a/Engine/Source/Runtime/Engine/Private/Animation/AnimSequence.cpp b/Engine/Source/Runtime/Engine/Private/Animation/AnimSequence.cpp
--- a/Engine/Source/Runtime/Engine/Private/Animation/AnimSequence.cpp
+++ b/Engine/Source/Runtime/Engine/Private/Animation/AnimSequence.cpp
@@ -2657,18 +2657,19 @@ void UAnimSequence::SerializeCompressedData(FArchive& Ar, bool bDDCData)
 	Ar << RotationCompressionFormat;
 	Ar << ScaleCompressionFormat;
 
+	Ar << CompressedTrackOffsets;
+	Ar << CompressedScaleOffsets;
+	Ar << CompressedSegments;
+
+	Ar << CompressedTrackToSkeletonMapTable;
+	Ar << CompressedCurveNames;
+
+	Ar << CompressedRawDataSize;
+	Ar << CompressedNumFrames;
+
 	if (Ar.IsLoading())
 	{
 		// Serialize the compressed byte stream from the archive to the buffer.
-		Ar << CompressedTrackOffsets;
-		Ar << CompressedScaleOffsets;
-		Ar << CompressedSegments;
-
-		Ar << CompressedTrackToSkeletonMapTable;
-		Ar << CompressedCurveNames;
-
-		Ar << CompressedRawDataSize;
-		Ar << CompressedNumFrames;
 		int32 NumBytes;
 		Ar << NumBytes;
 
@@ -2713,7 +2714,7 @@ void UAnimSequence::SerializeCompressedData(FArchive& Ar, bool bDDCData)
 		}
 		else
 		{
-			if (FPlatformProperties::RequiresCookedData() && ((AnimEncoding*)RotationCodec)->CanBeMemoryMapped(*this, NumBytes))
+			if (FPlatformProperties::RequiresCookedData())
 			{
 				CompressedByteStream.Empty(NumBytes);
 				CompressedByteStream.AddUninitialized(NumBytes);
@@ -2750,6 +2751,8 @@ void UAnimSequence::SerializeCompressedData(FArchive& Ar, bool bDDCData)
 		// Swap the byte stream into a buffer.
 		TArray<uint8> SerializedData;
 
+		const bool bIsCooking = !bDDCData && Ar.IsCooking();
+
 		if (!HasAnyFlags(RF_ClassDefaultObject))
 		{
 			// we must know the proper codecs to use
@@ -2757,7 +2760,7 @@ void UAnimSequence::SerializeCompressedData(FArchive& Ar, bool bDDCData)
 
 			// and then use the codecs to byte swap
 			check(RotationCodec != NULL);
-			((AnimEncoding*)RotationCodec)->ByteSwapOut(*this, SerializedData, Ar.ForceByteSwapping());
+			((AnimEncoding*)RotationCodec)->ByteSwapOut(*this, SerializedData, Ar.ForceByteSwapping(), bIsCooking);
 		}
 
 		// Make sure the entire byte stream was serialized.
@@ -2765,14 +2768,9 @@ void UAnimSequence::SerializeCompressedData(FArchive& Ar, bool bDDCData)
 
 		// Serialize the buffer to archive.
 		int32 Num = SerializedData.Num();
+		Ar << Num;
 
-
-
-		bool bUseBulkDataForSave = !bDDCData && Num && Ar.IsCooking() && Ar.CookingTarget()->SupportsFeature(ETargetPlatformFeatures::MemoryMappedFiles) && Ar.CookingTarget()->SupportsFeature(ETargetPlatformFeatures::MemoryMappedAnimation);
-
-		// repairing the offsets will leave them inconsistent with the sequence itself, so we need to restore them later
-		TArray<int32> SavedCompressedTrackOffsets = CompressedTrackOffsets;
-		FCompressedOffsetData SavedCompressedScaleOffsets = CompressedScaleOffsets;
+		bool bUseBulkDataForSave = Num && bIsCooking && Ar.CookingTarget()->SupportsFeature(ETargetPlatformFeatures::MemoryMappedFiles) && Ar.CookingTarget()->SupportsFeature(ETargetPlatformFeatures::MemoryMappedAnimation);
 
 		bool bSavebUseBulkDataForSave = false;
 		if (!bDDCData)
@@ -2787,21 +2785,7 @@ void UAnimSequence::SerializeCompressedData(FArchive& Ar, bool bDDCData)
 				bSavebUseBulkDataForSave = true;
 			}
 		}
-		if (bUseBulkDataForSave)
-		{
-			bUseBulkDataForSave = ((AnimEncoding*)RotationCodec)->CanBeMemoryMapped(*this, Num);
-		}
-
-		Ar << CompressedTrackOffsets;
-		Ar << CompressedScaleOffsets;
-		Ar << CompressedSegments;
-
-		Ar << CompressedTrackToSkeletonMapTable;
-		Ar << CompressedCurveNames;
 
-		Ar << CompressedRawDataSize;
-		Ar << CompressedNumFrames;
-		Ar << Num;
 		// Count compressed data.
 		Ar.CountBytes(SerializedData.Num(), SerializedData.Num());
 
@@ -2814,24 +2798,18 @@ void UAnimSequence::SerializeCompressedData(FArchive& Ar, bool bDDCData)
 			check(!bUseBulkDataForSave);
 		}
 
-		if (bUseBulkDataForSave)
-		{
-#if WITH_EDITOR
-			OptionalBulk.Lock(LOCK_READ_WRITE);
-			void* Dest = OptionalBulk.Realloc(Num);
-			FMemory::Memcpy(Dest, &(SerializedData[0]), Num);
-			OptionalBulk.Unlock();
-			OptionalBulk.SetBulkDataFlags(BULKDATA_PayloadAtEndOfFile | BULKDATA_PayloadInSeperateFile | BULKDATA_Force_NOT_InlinePayload | BULKDATA_MemoryMappedPayload);
-			OptionalBulk.ClearBulkDataFlags(BULKDATA_ForceInlinePayload);
-			OptionalBulk.Serialize(Ar, this);
-
-
 #define TEST_IS_CORRECTLY_FORMATTED_FOR_MEMORY_MAPPING WITH_EDITOR
 #if TEST_IS_CORRECTLY_FORMATTED_FOR_MEMORY_MAPPING
-			FMemoryReader MemoryReader(SerializedData, true);
+		if(!IsTemplate() && bIsCooking)
+		{
+			TArray<uint8> TempSerialized;
+			check(RotationCodec != NULL);
+			((AnimEncoding*)RotationCodec)->ByteSwapOut(*this, TempSerialized, Ar.ForceByteSwapping());
+
+			FMemoryReader MemoryReader(TempSerialized, true);
 			MemoryReader.SetByteSwapping(Ar.ForceByteSwapping());
-			TArray<uint8> SavedCompressedByteStream = CompressedByteStream;
 
+			TArray<uint8> SavedCompressedByteStream = CompressedByteStream;
 			CompressedByteStream.Empty();
 
 			((AnimEncoding*)RotationCodec)->ByteSwapIn(*this, MemoryReader);
@@ -2841,7 +2819,19 @@ void UAnimSequence::SerializeCompressedData(FArchive& Ar, bool bDDCData)
 			check(FMemory::Memcmp(SerializedData.GetData(), CompressedByteStream.GetData(), Num) == 0);
 
 			CompressedByteStream = SavedCompressedByteStream;
+		}
 #endif
+
+		if (bUseBulkDataForSave)
+		{
+#if WITH_EDITOR
+			OptionalBulk.Lock(LOCK_READ_WRITE);
+			void* Dest = OptionalBulk.Realloc(Num);
+			FMemory::Memcpy(Dest, &(SerializedData[0]), Num);
+			OptionalBulk.Unlock();
+			OptionalBulk.SetBulkDataFlags(BULKDATA_PayloadAtEndOfFile | BULKDATA_PayloadInSeperateFile | BULKDATA_Force_NOT_InlinePayload | BULKDATA_MemoryMappedPayload);
+			OptionalBulk.ClearBulkDataFlags(BULKDATA_ForceInlinePayload);
+			OptionalBulk.Serialize(Ar, this);
 #else
 			UE_LOG(LogAnimation, Fatal, TEXT("Can't save animation as bulk data in non-editor builds!"));
 #endif
@@ -2851,9 +2841,6 @@ void UAnimSequence::SerializeCompressedData(FArchive& Ar, bool bDDCData)
 			Ar.Serialize(SerializedData.GetData(), SerializedData.Num());
 		}
 
-		CompressedTrackOffsets = SavedCompressedTrackOffsets;
-		CompressedScaleOffsets = SavedCompressedScaleOffsets;
-
 		FString CurveCodecPath = CurveCompressionCodec->GetPathName();
 		Ar << CurveCodecPath;
 
diff --git a/Engine/Source/Runtime/Engine/Public/AnimEncoding.h b/Engine/Source/Runtime/Engine/Public/AnimEncoding.h
--- a/Engine/Source/Runtime/Engine/Public/AnimEncoding.h
+++ b/Engine/Source/Runtime/Engine/Public/AnimEncoding.h
@@ -281,15 +281,6 @@ public:
 		UAnimSequence& Seq, 
 		FMemoryReader& MemoryReader) PURE_VIRTUAL(AnimEncoding::ByteSwapIn,);
 
-	/**
-	 * Temp routine to determine if this compressed data can be memory mapped
-	 *
-	 * @param	Seq					An Animation Sequence to contain the read data.
-	 * @return	true if this sequence can be memory mapped
-	 */
-	virtual bool CanBeMemoryMapped(UAnimSequence& Seq, int32 TotalSize)
-		PURE_VIRTUAL(AnimEncoding::CanBeMemoryMapped, return false;);
-
 	/**
 	 * Handles Byte-swapping outgoing animation data to an array of BYTEs
 	 *
@@ -300,7 +291,8 @@ public:
 	virtual void ByteSwapOut(
 		UAnimSequence& Seq,
 		TArray<uint8>& SerializedData, 
-		bool ForceByteSwapping) PURE_VIRTUAL(AnimEncoding::ByteSwapOut,);
+		bool ForceByteSwapping,
+		bool bMaintainComponentOrder=false) PURE_VIRTUAL(AnimEncoding::ByteSwapOut, );
 
 	/**
 	 * Extracts a single BoneAtom from an Animation Sequence.
@@ -476,8 +468,6 @@ public:
 		UAnimSequence& Seq, 
 		FMemoryReader& MemoryReader) override;
 
-	virtual bool CanBeMemoryMapped(UAnimSequence& Seq, int32 TotalSize) override;
-
 	/**
 	 * Handles Byte-swapping outgoing animation data to an array of BYTEs
 	 *
@@ -487,8 +477,9 @@ public:
 	 */
 	virtual void ByteSwapOut(
 		UAnimSequence& Seq,
-		TArray<uint8>& SerializedData, 
-		bool ForceByteSwapping) override;
+		TArray<uint8>& SerializedData,
+		bool ForceByteSwapping,
+		bool bMaintainComponentOrder=false) override;
 
 	/**
 	 * Handles the ByteSwap of compressed animation data on import
diff --git a/Engine/Source/Runtime/Engine/Public/AnimEncoding_PerTrackCompression.h b/Engine/Source/Runtime/Engine/Public/AnimEncoding_PerTrackCompression.h
--- a/Engine/Source/Runtime/Engine/Public/AnimEncoding_PerTrackCompression.h
+++ b/Engine/Source/Runtime/Engine/Public/AnimEncoding_PerTrackCompression.h
@@ -26,8 +26,6 @@ public:
 	 */
 	virtual void ByteSwapIn(UAnimSequence& Seq, FMemoryReader& MemoryReader) override;
 
-	virtual bool CanBeMemoryMapped(UAnimSequence& Seq, int32 TotalSize) override;
-
 	/**
 	 * Handles Byte-swapping outgoing animation data to an array of BYTEs
 	 *
@@ -38,7 +36,8 @@ public:
 	virtual void ByteSwapOut(
 		UAnimSequence& Seq,
 		TArray<uint8>& SerializedData, 
-		bool ForceByteSwapping) override;
+		bool ForceByteSwapping,
+		bool bMaintainComponentOrder=false) override;
 
 	/**
 	 * Extracts a single BoneAtom from an Animation Sequence.
@@ -103,9 +102,10 @@ protected:
 	 * @param	Seq					The Animation Sequence being operated on.
 	 * @param	MemoryStream		The MemoryReader or MemoryWriter object to read from/write to.
 	 * @param	Offset				The starting offset into the compressed byte stream for this track (can be INDEX_NONE to indicate an identity track)
+	 * @param	bMaintainComponentOrder Should we maintain the order of track components (trans/rot/scale) in the byte stream (if false the output will always be T/R/S regardless of input order)
 	 */
 	template<class TArchive>
-	static void ByteSwapOneTrack(UAnimSequence& Seq, TArchive& MemoryStream, int32 Offset);
+	static void ByteSwapOneTrack(UAnimSequence& Seq, TArchive& MemoryStream, int32 Offset, bool bMaintainComponentOrder = false);
 
 	/**
 	 * Preserves 4 byte alignment within a stream
