Stopped FKismetCompilerUtilities::GenerateAssignmentNodes generating redudant UK2Node_PureAssignmentStatement for text pins that matched the CDO value

This was doing a simple string comparison, rather than using the more complex (and correct) logic that compares text values correctly in the editor (as used by UEdGraphSchema::DoesDefaultValueMatchAutogenerated).

This fixes a deterministic cooking issue caused by the unstable keys of the default text value that was intanced into the graph being written into the bytecode.

diff --git a/Engine/Source/Editor/KismetCompiler/Private/KismetCompilerMisc.cpp b/Engine/Source/Editor/KismetCompiler/Private/KismetCompilerMisc.cpp
--- a/Engine/Source/Editor/KismetCompiler/Private/KismetCompilerMisc.cpp
+++ b/Engine/Source/Editor/KismetCompiler/Private/KismetCompilerMisc.cpp
@@ -744,7 +744,7 @@ UEdGraphPin* FKismetCompilerUtilities::GenerateAssignmentNodes(class FKismetComp
 					
 				if (FBlueprintCompilationManager::GetDefaultValue(ForClass, Property, DefaultValueAsString))
 				{
-					if (DefaultValueAsString == OrgPin->GetDefaultAsString())
+					if (Schema->DoesDefaultValueMatch(*OrgPin, DefaultValueAsString))
 					{
 						continue;
 					}
diff --git a/Engine/Source/Runtime/Engine/Classes/EdGraph/EdGraphPin.h b/Engine/Source/Runtime/Engine/Classes/EdGraph/EdGraphPin.h
--- a/Engine/Source/Runtime/Engine/Classes/EdGraph/EdGraphPin.h
+++ b/Engine/Source/Runtime/Engine/Classes/EdGraph/EdGraphPin.h
@@ -494,7 +494,10 @@ public:
 #endif // WITH_EDITORONLY_DATA
 	}
 
-	/** Returns an internal string representation of the string/object/text default value, that is compatible with AutogeneratedDefaultValue and ImportText */
+	/**
+	 * Returns an internal string representation of the string/object/text default value, that is compatible with AutogeneratedDefaultValue and ImportText
+	 * @note Do not directly compare this value with another to test for a default value. Instead use DoesDefaultValueMatch or DoesDefaultValueMatchAutogenerated on the schema.
+	 */
 	ENGINE_API FString GetDefaultAsString() const;
 
 	/** Returns a human readable FText representation of the string/object/text default value */
diff --git a/Engine/Source/Runtime/Engine/Classes/EdGraph/EdGraphSchema.h b/Engine/Source/Runtime/Engine/Classes/EdGraph/EdGraphSchema.h
--- a/Engine/Source/Runtime/Engine/Classes/EdGraph/EdGraphSchema.h
+++ b/Engine/Source/Runtime/Engine/Classes/EdGraph/EdGraphSchema.h
@@ -757,6 +757,9 @@ class ENGINE_API UEdGraphSchema : public UObject
 	/** Sets the text to the specified pin */
 	virtual void TrySetDefaultText(UEdGraphPin& InPin, const FText& InNewDefaultText) const;
 
+	/** Returns if the pin's value matches the given value */
+	virtual bool DoesDefaultValueMatch(const UEdGraphPin& InPin, const FString& InValue) const;
+
 	/** Returns if the pin's value matches what the true (autogenerated) default value for that pin would be */
 	virtual bool DoesDefaultValueMatchAutogenerated(const UEdGraphPin& InPin) const;
 
diff --git a/Engine/Source/Runtime/Engine/Private/EdGraph/EdGraphSchema.cpp b/Engine/Source/Runtime/Engine/Private/EdGraph/EdGraphSchema.cpp
--- a/Engine/Source/Runtime/Engine/Private/EdGraph/EdGraphSchema.cpp
+++ b/Engine/Source/Runtime/Engine/Private/EdGraph/EdGraphSchema.cpp
@@ -497,21 +497,28 @@ void UEdGraphSchema::TrySetDefaultText(UEdGraphPin& InPin, const FText& InNewDef
 #endif	//#if WITH_EDITOR
 }
 
-bool UEdGraphSchema::DoesDefaultValueMatchAutogenerated(const UEdGraphPin& InPin) const
+bool UEdGraphSchema::DoesDefaultValueMatch(const UEdGraphPin& InPin, const FString& InValue) const
 {
 	// Same logic as UEdGraphPin::DoesDefaultValueMatchAutogenerated
 	// Ignoring case on purpose to match default behavior
-	FText AutogeneratedDefaultValueAsText;
-	if (!InPin.DefaultTextValue.IsEmpty() && FTextStringHelper::ReadFromBuffer(*InPin.AutogeneratedDefaultValue, AutogeneratedDefaultValueAsText))
+	FText ValueAsText;
+	if (!InPin.DefaultTextValue.IsEmpty() && FTextStringHelper::ReadFromBuffer(*InValue, ValueAsText))
 	{
-		return UTextProperty::Identical_Implementation(AutogeneratedDefaultValueAsText, InPin.DefaultTextValue, 0);
+		return UTextProperty::Identical_Implementation(ValueAsText, InPin.DefaultTextValue, 0);
 	}
 	else
 	{
-		return InPin.GetDefaultAsString().Equals(InPin.AutogeneratedDefaultValue, ESearchCase::IgnoreCase);
+		return InPin.GetDefaultAsString().Equals(InValue, ESearchCase::IgnoreCase);
 	}
 }
 
+bool UEdGraphSchema::DoesDefaultValueMatchAutogenerated(const UEdGraphPin& InPin) const
+{
+	// Specifically call our version of this function rather than any overrides
+	// This maintains the existing logic of this function while sharing the implementation
+	return UEdGraphSchema::DoesDefaultValueMatch(InPin, InPin.AutogeneratedDefaultValue);
+}
+
 void UEdGraphSchema::BreakNodeLinks(UEdGraphNode& TargetNode) const
 {
 #if WITH_EDITOR
