Delta checkpoint changes to track channel re-use properly

diff --git a/Engine/Source/Runtime/Engine/Classes/Engine/DemoNetConnection.h b/Engine/Source/Runtime/Engine/Classes/Engine/DemoNetConnection.h
--- a/Engine/Source/Runtime/Engine/Classes/Engine/DemoNetConnection.h
+++ b/Engine/Source/Runtime/Engine/Classes/Engine/DemoNetConnection.h
@@ -80,6 +80,8 @@ public:
 	virtual TSharedPtr<FObjectReplicator> CreateReplicatorForNewActorChannel(UObject* Object);
 	virtual FString RemoteAddressToString() override { return TEXT("Demo"); }
 
+	virtual void NotifyActorNetGUID(UActorChannel* Channel) override;
+
 public:
 
 	virtual void Serialize(FArchive& Ar) override;
@@ -99,6 +101,8 @@ public:
 	TArray<FQueuedDemoPacket> QueuedDemoPackets;
 	TArray<FQueuedDemoPacket> QueuedCheckpointPackets;
 
+	TMap<FNetworkGUID, UActorChannel*>& GetOpenChannelMap() { return OpenChannelMap; }
+
 protected:
 	virtual void DestroyIgnoredActor(AActor* Actor) override;
 
@@ -109,4 +113,7 @@ protected:
 
 private:
 	void TrackSendForProfiler(const void* Data, int32 NumBytes);
+
+	// Not a weak object pointer, intended to exist only during checkpoint loading
+	TMap<FNetworkGUID, UActorChannel*> OpenChannelMap;
 };
diff --git a/Engine/Source/Runtime/Engine/Classes/Engine/DemoNetDriver.h b/Engine/Source/Runtime/Engine/Classes/Engine/DemoNetDriver.h
--- a/Engine/Source/Runtime/Engine/Classes/Engine/DemoNetDriver.h
+++ b/Engine/Source/Runtime/Engine/Classes/Engine/DemoNetDriver.h
@@ -24,7 +24,7 @@ class UDemoNetDriver;
 class UDemoNetConnection;
 class FRepState;
 
-DECLARE_LOG_CATEGORY_EXTERN( LogDemo, Log, All );
+DECLARE_LOG_CATEGORY_EXTERN(LogDemo, Log, All);
 
 DECLARE_MULTICAST_DELEGATE(FOnGotoTimeMCDelegate);
 DECLARE_DELEGATE_OneParam(FOnGotoTimeDelegate, const bool /* bWasSuccessful */);
@@ -42,7 +42,7 @@ DECLARE_MULTICAST_DELEGATE(FOnDemoFinishRecordingDelegate);
 class FQueuedReplayTask : public TSharedFromThis<FQueuedReplayTask>
 {
 public:
-	FQueuedReplayTask( UDemoNetDriver* InDriver ) : Driver( InDriver )
+	FQueuedReplayTask(UDemoNetDriver* InDriver) : Driver(InDriver)
 	{
 	}
 
@@ -50,9 +50,9 @@ public:
 	{
 	}
 
-	virtual void	StartTask() = 0;
-	virtual bool	Tick() = 0;
-	virtual FName	GetName() const = 0;
+	virtual void StartTask() = 0;
+	virtual bool Tick() = 0;
+	virtual FName GetName() const = 0;
 
 	TWeakObjectPtr<UDemoNetDriver> Driver;
 };
@@ -60,11 +60,11 @@ public:
 class FReplayExternalData
 {
 public:
-	FReplayExternalData() : TimeSeconds( 0.0f )
+	FReplayExternalData() : TimeSeconds(0.0f)
 	{
 	}
 
-	FReplayExternalData( FBitReader&& InReader, const float InTimeSeconds ) : TimeSeconds( InTimeSeconds )
+	FReplayExternalData(FBitReader&& InReader, const float InTimeSeconds) : TimeSeconds(InTimeSeconds)
 	{
 		Reader = MoveTemp(InReader);
 	}
@@ -79,11 +79,11 @@ public:
 };
 
 // Using an indirect array here since FReplayExternalData stores an FBitReader, and it's not safe to store an FArchive directly in a TArray.
-typedef TIndirectArray< FReplayExternalData > FReplayExternalDataArray;
+typedef TIndirectArray<FReplayExternalData> FReplayExternalDataArray;
 
 struct FPlaybackPacket
 {
-	TArray< uint8 >		Data;
+	TArray<uint8>		Data;
 	float				TimeSeconds;
 	int32				LevelIndex;
 	uint32				SeenLevelIndex;
@@ -161,9 +161,9 @@ struct FLevelNameAndTime
 enum class EReplayHeaderFlags : uint32
 {
 	None				= 0,
-	ClientRecorded		= ( 1 << 0 ),
-	HasStreamingFixes	= ( 1 << 1 ),
-	DeltaCheckpoints	= ( 1 << 2 ),
+	ClientRecorded		= (1 << 0),
+	HasStreamingFixes	= (1 << 1),
+	DeltaCheckpoints	= (1 << 2),
 };
 
 ENUM_CLASS_FLAGS(EReplayHeaderFlags);
@@ -183,25 +183,25 @@ struct FNetworkDemoHeader
 	TArray<FString> GameSpecificData;				// Area for subclasses to write stuff
 
 	FNetworkDemoHeader() :
-		Magic( NETWORK_DEMO_MAGIC ),
-		Version( NETWORK_DEMO_VERSION ),
-		NetworkChecksum( FNetworkVersion::GetLocalNetworkVersion() ),
-		EngineNetworkProtocolVersion( FNetworkVersion::GetEngineNetworkProtocolVersion() ),
-		GameNetworkProtocolVersion( FNetworkVersion::GetGameNetworkProtocolVersion() ),
+		Magic(NETWORK_DEMO_MAGIC),
+		Version(NETWORK_DEMO_VERSION),
+		NetworkChecksum(FNetworkVersion::GetLocalNetworkVersion()),
+		EngineNetworkProtocolVersion(FNetworkVersion::GetEngineNetworkProtocolVersion()),
+		GameNetworkProtocolVersion(FNetworkVersion::GetGameNetworkProtocolVersion()),
 		Guid(),
-		EngineVersion( FEngineVersion::Current() ),
-		HeaderFlags( EReplayHeaderFlags::None )
+		EngineVersion(FEngineVersion::Current()),
+		HeaderFlags(EReplayHeaderFlags::None)
 	{
 	}
 
-	friend FArchive& operator << ( FArchive& Ar, FNetworkDemoHeader& Header )
+	friend FArchive& operator << (FArchive& Ar, FNetworkDemoHeader& Header)
 	{
 		Ar << Header.Magic;
 
 		// Check magic value
-		if ( Header.Magic != NETWORK_DEMO_MAGIC )
+		if (Header.Magic != NETWORK_DEMO_MAGIC)
 		{
-			UE_LOG( LogDemo, Error, TEXT( "Header.Magic != NETWORK_DEMO_MAGIC" ) );
+			UE_LOG(LogDemo, Error, TEXT("Header.Magic != NETWORK_DEMO_MAGIC"));
 			Ar.SetError();
 			return Ar;
 		}
@@ -209,9 +209,9 @@ struct FNetworkDemoHeader
 		Ar << Header.Version;
 
 		// Check version
-		if ( Header.Version < MIN_NETWORK_DEMO_VERSION )
+		if (Header.Version < MIN_NETWORK_DEMO_VERSION)
 		{
-			UE_LOG( LogDemo, Error, TEXT( "Header.Version < MIN_NETWORK_DEMO_VERSION. Header.Version: %i, MIN_NETWORK_DEMO_VERSION: %i" ), Header.Version, MIN_NETWORK_DEMO_VERSION );
+			UE_LOG(LogDemo, Error, TEXT("Header.Version < MIN_NETWORK_DEMO_VERSION. Header.Version: %i, MIN_NETWORK_DEMO_VERSION: %i"), Header.Version, MIN_NETWORK_DEMO_VERSION);
 			Ar.SetError();
 			return Ar;
 		}
@@ -394,46 +394,37 @@ class ENGINE_API UDemoNetDriver : public UNetDriver
 	APlayerController* SpectatorController;
 
 	/** Our network replay streamer */
-	TSharedPtr< class INetworkReplayStreamer >	ReplayStreamer;
+	TSharedPtr<class INetworkReplayStreamer> ReplayStreamer;
 
-	uint32 GetDemoCurrentTimeInMS() { return (uint32)( (double)DemoCurrentTime * 1000 ); }
+	uint32 GetDemoCurrentTimeInMS() const { return (uint32)((double)DemoCurrentTime * 1000); }
 
 	/** Internal debug timing/tracking */
-	double		AccumulatedRecordTime;
-	double		LastRecordAvgFlush;
-	double		MaxRecordTime;
-	int32		RecordCountSinceFlush;
+	double AccumulatedRecordTime;
+	double LastRecordAvgFlush;
+	double MaxRecordTime;
+	int32 RecordCountSinceFlush;
 
 	/** Net startup actors that need to be destroyed after checkpoints are loaded */
-	TSet< FString >									DeletedNetStartupActors;
+	TSet<FString> DeletedNetStartupActors;
 
 	/** Keeps track of NetGUIDs that were deleted, so we can skip them when saving checkpoints. Only used while recording. */
-	TSet< FNetworkGUID >							DeletedNetStartupActorGUIDs;
-
-	/** Delta net startup actors that need to be destroyed after checkpoints are loaded */
-	TSet< FString >									DeltaDeletedNetStartupActors;
-
-	/** GUID list of actors deleted before the checkpoint was recorded, including dynamic actors */
-	TSet<FNetworkGUID>								DeletedActorGuids;
-
-	/** Delta checkpoint list for DeletedActorGuids */
-	TSet<FNetworkGUID>								DeltaDeletedActorGuids;
+	TSet<FNetworkGUID> DeletedNetStartupActorGUIDs;
 
 	/** 
 	 * Net startup actors that need to be rolled back during scrubbing by being destroyed and re-spawned 
 	 * NOTE - DeletedNetStartupActors will take precedence here, and destroy the actor instead
 	 */
 	UPROPERTY(transient)
-	TMap<FString, FRollbackNetStartupActorInfo>		RollbackNetStartupActors;
+	TMap<FString, FRollbackNetStartupActorInfo> RollbackNetStartupActors;
 
-	double				LastCheckpointTime;					// Last time a checkpoint was saved
+	double LastCheckpointTime;					// Last time a checkpoint was saved
 
-	void		RespawnNecessaryNetStartupActors(TArray<AActor*>& SpawnedActors, ULevel* Level = nullptr);
+	void RespawnNecessaryNetStartupActors(TArray<AActor*>& SpawnedActors, ULevel* Level = nullptr);
 
-	virtual bool ShouldSaveCheckpoint();
+	virtual bool ShouldSaveCheckpoint() const;
 
-	void		SaveCheckpoint();
-	void		TickCheckpoint();
+	void SaveCheckpoint();
+	void TickCheckpoint();
 
 	UE_DEPRECATED(4.22, "This method will be made private in future versions.")
 	bool LoadCheckpoint(FArchive* GotoCheckpointArchive, int64 GotoCheckpointSkipExtraTimeInMS)
@@ -449,19 +440,40 @@ class ENGINE_API UDemoNetDriver : public UNetDriver
 private:
 
 	bool LoadCheckpoint(const FGotoResult& GotoResult);
+
+	struct FDeltaCheckpointData
+	{
+		/** Net startup actors that were destroyed */
+		TSet<FString> DestroyedNetStartupActors;
+		/** Destroyed dynamic actors that were active in the previous checkpoint */
+		TSet<FNetworkGUID> DestroyedDynamicActors;	
+		/** Channels closed that were open in the previous checkpoint, and the reason why */
+		TMap<FNetworkGUID, EChannelCloseReason> ChannelsToClose;
+
+		void CountBytes(FArchive& Ar) const
+		{
+			DestroyedNetStartupActors.CountBytes(Ar);
+			DestroyedDynamicActors.CountBytes(Ar);
+			ChannelsToClose.CountBytes(Ar);
+		}
+	};
+	
+	FDeltaCheckpointData RecordingDeltaCheckpointData;
+
+	TArray<TUniquePtr<FDeltaCheckpointData>> PlaybackDeltaCheckpointData;
 	
 public:	
 
 	virtual void Serialize(FArchive& Ar) override;
 
 	/** Returns true if we're in the process of saving a checkpoint. */
-	bool		IsSavingCheckpoint() const;
+	bool IsSavingCheckpoint() const;
 
-	void		SaveExternalData( FArchive& Ar );
-	void		LoadExternalData( FArchive& Ar, const float TimeSeconds );
+	void SaveExternalData(FArchive& Ar);
+	void LoadExternalData(FArchive& Ar, const float TimeSeconds);
 
 	/** Public delegate for external systems to be notified when a replay begins. UDemoNetDriver is passed as a param */
-	static FOnDemoStartedDelegate		OnDemoStarted;
+	static FOnDemoStartedDelegate OnDemoStarted;
 
 	/** Public delegate to be notified when a replay failed to start. UDemoNetDriver and FailureType are passed as params */
 	static FOnDemoFailedToStartDelegate OnDemoFailedToStart;
@@ -475,27 +487,27 @@ public:
 	/** Public Delegate for external systems to be notified when replay recording is about to finish. */
 	FOnDemoFinishRecordingDelegate OnDemoFinishRecordingDelegate;
 
-	bool		IsLoadingCheckpoint() const { return bIsLoadingCheckpoint; }
+	bool IsLoadingCheckpoint() const { return bIsLoadingCheckpoint; }
 	
-	bool		IsPlayingClientReplay() const;
+	bool IsPlayingClientReplay() const;
 
 	/** ExternalDataToObjectMap is used to map a FNetworkGUID to the proper FReplayExternalDataArray */
-	TMap< FNetworkGUID, FReplayExternalDataArray > ExternalDataToObjectMap;
+	TMap<FNetworkGUID, FReplayExternalDataArray> ExternalDataToObjectMap;
 		
 	/** PlaybackPackets are used to buffer packets up when we read a demo frame, which we can then process when the time is right */
-	TArray< FPlaybackPacket > PlaybackPackets;
+	TArray<FPlaybackPacket> PlaybackPackets;
 
 	/**
 	 * During recording, all unique streaming levels since recording started.
 	 * During playback, all streaming level instances we've created.
 	 */
-	TSet< TWeakObjectPtr< UObject > >	UniqueStreamingLevels;
+	TSet<TWeakObjectPtr<UObject>> UniqueStreamingLevels;
 
 	/**
 	 * During recording, streaming levels waiting to be saved next frame.
 	 * During playback, streaming levels that have recently become visible.
 	 */
-	TSet< TWeakObjectPtr< UObject > >	NewStreamingLevelsThisFrame;
+	TSet<TWeakObjectPtr<UObject>> NewStreamingLevelsThisFrame;
 
 	bool bRecordMapChanges;
 
@@ -506,17 +518,17 @@ private:
 		UObject*	Level;
 	};
 
-	bool		bIsFastForwarding;
-	bool		bIsFastForwardingForCheckpoint;
-	bool		bWasStartStreamingSuccessful;
-	bool		bIsLoadingCheckpoint;
+	bool bIsFastForwarding;
+	bool bIsFastForwardingForCheckpoint;
+	bool bWasStartStreamingSuccessful;
+	bool bIsLoadingCheckpoint;
 
 	TArray<FNetworkGUID> NonQueuedGUIDsForScrubbing;
 
 	// Replay tasks
-	TArray< TSharedPtr< FQueuedReplayTask > >	QueuedReplayTasks;
-	TSharedPtr< FQueuedReplayTask >				ActiveReplayTask;
-	TSharedPtr< FQueuedReplayTask >				ActiveScrubReplayTask;
+	TArray<TSharedRef<FQueuedReplayTask>>		QueuedReplayTasks;
+	TSharedPtr<FQueuedReplayTask>				ActiveReplayTask;
+	TSharedPtr<FQueuedReplayTask>				ActiveScrubReplayTask;
 
 	/** Set via GotoTimeInSeconds, only fired once (at most). Called for successful or failed scrub. */
 	FOnGotoTimeDelegate OnGotoTimeDelegate_Transient;
@@ -580,17 +592,17 @@ public:
 
 	// UNetDriver interface.
 
-	virtual bool InitBase( bool bInitAsClient, FNetworkNotify* InNotify, const FURL& URL, bool bReuseAddressAndPort, FString& Error ) override;
+	virtual bool InitBase(bool bInitAsClient, FNetworkNotify* InNotify, const FURL& URL, bool bReuseAddressAndPort, FString& Error) override;
 	virtual void FinishDestroy() override;
 	virtual FString LowLevelGetNetworkNumber() override;
-	virtual bool InitConnect( FNetworkNotify* InNotify, const FURL& ConnectURL, FString& Error ) override;
-	virtual bool InitListen( FNetworkNotify* InNotify, FURL& ListenURL, bool bReuseAddressAndPort, FString& Error ) override;
-	virtual void TickFlush( float DeltaSeconds ) override;
-	virtual void TickDispatch( float DeltaSeconds ) override;
-	virtual void ProcessRemoteFunction( class AActor* Actor, class UFunction* Function, void* Parameters, struct FOutParmRec* OutParms, struct FFrame* Stack, class UObject* SubObject = nullptr ) override;
+	virtual bool InitConnect(FNetworkNotify* InNotify, const FURL& ConnectURL, FString& Error) override;
+	virtual bool InitListen(FNetworkNotify* InNotify, FURL& ListenURL, bool bReuseAddressAndPort, FString& Error) override;
+	virtual void TickFlush(float DeltaSeconds) override;
+	virtual void TickDispatch(float DeltaSeconds) override;
+	virtual void ProcessRemoteFunction(class AActor* Actor, class UFunction* Function, void* Parameters, struct FOutParmRec* OutParms, struct FFrame* Stack, class UObject* SubObject = nullptr) override;
 	virtual bool IsAvailable() const override { return true; }
 	void SkipTime(const float InTimeToSkip);
-	void SkipTimeInternal( const float SecondsToSkip, const bool InFastForward, const bool InIsForCheckpoint );
+	void SkipTimeInternal(const float SecondsToSkip, const bool InFastForward, const bool InIsForCheckpoint);
 	bool InitConnectInternal(FString& Error);
 	virtual bool ShouldClientDestroyTearOffActors() const override;
 	virtual bool ShouldSkipRepNotifies() const override;
@@ -604,6 +616,7 @@ public:
 	virtual bool ShouldReplicateFunction(AActor* Actor, UFunction* Function) const override;
 	virtual bool ShouldReplicateActor(AActor* Actor) const override;
 	virtual void NotifyActorChannelOpen(UActorChannel* Channel, AActor* Actor) override;
+	virtual void NotifyActorChannelCleanedUp(UActorChannel* Channel, EChannelCloseReason CloseReason) override;
 
 	virtual void ProcessLocalServerPackets() override {}
 	virtual void ProcessLocalClientPackets() override {}
@@ -667,9 +680,9 @@ public:
 	/** @return true if the net resource is valid or false if it should not be used */
 	virtual bool IsNetResourceValid(void) override { return true; }
 
-	void TickDemoRecord( float DeltaSeconds );
-	void PauseChannels( const bool bPause );
-	void PauseRecording( const bool bInPauseRecording ) { bPauseRecording = bInPauseRecording; }
+	void TickDemoRecord(float DeltaSeconds);
+	void PauseChannels(const bool bPause);
+	void PauseRecording(const bool bInPauseRecording) { bPauseRecording = bInPauseRecording; }
 	bool IsRecordingPaused() const { return bPauseRecording; }
 
 	bool ConditionallyProcessPlaybackPackets();
@@ -710,23 +723,23 @@ private:
 
 public:
 
-	bool ConditionallyReadDemoFrameIntoPlaybackPackets( FArchive& Ar );
+	bool ConditionallyReadDemoFrameIntoPlaybackPackets(FArchive& Ar);
 
-	bool ProcessPacket( const uint8* Data, int32 Count );
-	bool ProcessPacket( const FPlaybackPacket& PlaybackPacket )
+	bool ProcessPacket(const uint8* Data, int32 Count);
+	bool ProcessPacket(const FPlaybackPacket& PlaybackPacket)
 	{
 		return ShouldSkipPlaybackPacket(PlaybackPacket) ||
 				ProcessPacket(PlaybackPacket.Data.GetData(), PlaybackPacket.Data.Num());
 	}
 
-	void WriteDemoFrameFromQueuedDemoPackets( FArchive& Ar, TArray<FQueuedDemoPacket>& QueuedPackets, float FrameTime );
-	void WritePacket( FArchive& Ar, uint8* Data, int32 Count );
+	void WriteDemoFrameFromQueuedDemoPackets(FArchive& Ar, TArray<FQueuedDemoPacket>& QueuedPackets, float FrameTime);
+	void WritePacket(FArchive& Ar, uint8* Data, int32 Count);
 
-	void TickDemoPlayback( float DeltaSeconds );
+	void TickDemoPlayback(float DeltaSeconds);
 	
-	void FinalizeFastForward( const double StartTime );
+	void FinalizeFastForward(const double StartTime);
 	
-	void SpawnDemoRecSpectator( UNetConnection* Connection, const FURL& ListenURL );
+	void SpawnDemoRecSpectator( UNetConnection* Connection, const FURL& ListenURL);
 
 	/**
 	 * Restores the given player controller so that it properly points to the given NetConnection
@@ -811,7 +824,7 @@ public:
 
 	bool IsFastForwarding() const { return bIsFastForwarding; }
 
-	FReplayExternalDataArray* GetExternalDataArrayForObject( UObject* Object );
+	FReplayExternalDataArray* GetExternalDataArrayForObject(UObject* Object);
 
 	bool ReadDemoFrameIntoPlaybackPackets(FArchive& Ar, TArray<FPlaybackPacket>& Packets, const bool bForLevelFastForward, float* OutTime);
 	bool ReadDemoFrameIntoPlaybackPackets(FArchive& Ar) { return ReadDemoFrameIntoPlaybackPackets(Ar, PlaybackPackets, false, nullptr); }
@@ -827,11 +840,11 @@ public:
 
 	void ReplayStreamingReady(const FStartStreamingResult& Result);
 
-	void AddReplayTask( FQueuedReplayTask* NewTask );
+	void AddReplayTask(FQueuedReplayTask* NewTask);
 	bool IsAnyTaskPending() const;
 	void ClearReplayTasks();
 	bool ProcessReplayTasks();
-	bool IsNamedTaskInQueue( const FName& Name ) const;
+	bool IsNamedTaskInQueue(const FName& Name) const;
 	FName GetNextQueuedTaskName() const;
 
 	/** If a channel is associated with Actor, adds the channel's GUID to the list of GUIDs excluded from queuing bunches during scrubbing. */
@@ -839,7 +852,7 @@ public:
 	/** Adds the channel's GUID to the list of GUIDs excluded from queuing bunches during scrubbing. */
 	void AddNonQueuedGUIDForScrubbing(FNetworkGUID InGUID);
 
-	virtual bool IsLevelInitializedForActor( const AActor* InActor, const UNetConnection* InConnection ) const override;
+	virtual bool IsLevelInitializedForActor(const AActor* InActor, const UNetConnection* InConnection) const override;
 
 	/** Called when a "go to time" operation is completed. */
 	void NotifyGotoTimeFinished(bool bWasSuccessful);
@@ -847,12 +860,12 @@ public:
 	/** Read the streaming level information from the metadata after the level is loaded */
 	void PendingNetGameLoadMapCompleted();
 	
-	virtual void NotifyActorDestroyed( AActor* ThisActor, bool IsSeamlessTravel=false ) override;
-	virtual void NotifyActorLevelUnloaded( AActor* Actor ) override;
-	virtual void NotifyStreamingLevelUnload( ULevel* InLevel ) override;
+	virtual void NotifyActorDestroyed(AActor* ThisActor, bool IsSeamlessTravel=false) override;
+	virtual void NotifyActorLevelUnloaded(AActor* Actor) override;
+	virtual void NotifyStreamingLevelUnload(ULevel* InLevel) override;
 
 	/** Call this function during playback to track net startup actors that need a hard reset when scrubbing, which is done by destroying and then re-spawning */
-	virtual void QueueNetStartupActorForRollbackViaDeletion( AActor* Actor );
+	virtual void QueueNetStartupActorForRollbackViaDeletion(AActor* Actor);
 
 	/** Called when seamless travel begins when recording a replay. */
 	void OnSeamlessTravelStartDuringRecording(const FString& LevelName);
@@ -1033,7 +1046,7 @@ private:
 	{
 		ECheckpointSaveState CheckpointSaveState;						// Current state of checkpoint SaveState
 		FPackageMapAckState CheckpointAckState;							// Current ack state of packagemap for the current checkpoint being saved
-		TArray< FPendingCheckPointActor > PendingCheckpointActors;		// Actors to be serialized by pending checkpoint
+		TArray<FPendingCheckPointActor> PendingCheckpointActors;		// Actors to be serialized by pending checkpoint
 		double				TotalCheckpointSaveTimeSeconds;				// Total time it took to save checkpoint including the finaling part across all frames
 		double				TotalCheckpointReplicationTimeSeconds;		// Total time it took to write all replicated objects across all frames
 		bool				bWriteCheckpointOffset;
@@ -1041,13 +1054,13 @@ private:
 		FArchivePos			CheckpointOffset;
 		uint32				GuidCacheSize;
 
-		TSet<FString>		DeltaDeletedNetStartupActors;
-		TSet<FNetworkGUID>	DeltaDeletedActorGuids;
+		FDeltaCheckpointData DeltaCheckpointData;
 
 		void CountBytes(FArchive& Ar) const
 		{
 			CheckpointAckState.CountBytes(Ar);
 			PendingCheckpointActors.CountBytes(Ar);
+			DeltaCheckpointData.CountBytes(Ar);
 		}
 	};
 
@@ -1192,4 +1205,7 @@ private:
 	float RecordBuildConsiderAndPrioritizeTimeSlice;
 
 	void AdjustConsiderTime(const float ReplicatedPercent);
+
+	bool ProcessFastForwardPackets(TArrayView<FPlaybackPacket> Packets, const TSet<int32>& LevelIndices);
+	void ProcessPlaybackPackets(TArrayView<FPlaybackPacket> Packets);
 };
diff --git a/Engine/Source/Runtime/Engine/Classes/Engine/NetConnection.h b/Engine/Source/Runtime/Engine/Classes/Engine/NetConnection.h
--- a/Engine/Source/Runtime/Engine/Classes/Engine/NetConnection.h
+++ b/Engine/Source/Runtime/Engine/Classes/Engine/NetConnection.h
@@ -628,12 +628,6 @@ public:
 
 	EResendAllDataState ResendAllDataState;
 
-	/** Set of guids we may need to ignore when processing a delta checkpoint */
-	TSet<FNetworkGUID> IgnoredGuids;
-
-	/** Set of channels we may need to ignore when processing a delta checkpoint */
-	TSet<int32> IgnoredChannels;
-
 #if !UE_BUILD_SHIPPING
 	/** Delegate for hooking ReceivedRawPacket */
 	FOnReceivedRawPacket	ReceivedRawPacketDel;
@@ -1054,6 +1048,12 @@ public:
 	 */
 	void SetIgnoreAlreadyOpenedChannels(bool bInIgnoreAlreadyOpenedChannels);
 
+	/**
+	 * Sets whether or not we should ignore bunches for a specific set of NetGUIDs.
+	 * Should only be used with InternalAck.
+	 */
+	void SetIgnoreActorBunches(bool bInIgnoreActorBunches, TSet<FNetworkGUID>&& InIgnoredBunchGuids);
+
 	/** Returns the OutgoingBunches array, only to be used by UChannel::SendBunch */
 	TArray<FOutBunch *>& GetOutgoingBunches() { return OutgoingBunches; }
 
@@ -1070,6 +1070,9 @@ public:
 	 */
 	void FlushPacketOrderCache(bool bFlushWholeCache=false);
 
+	/** Called when an actor channel is open and knows its NetGUID. */
+	ENGINE_API virtual void NotifyActorNetGUID(UActorChannel* Channel) {}
+
 protected:
 
 	ENGINE_API void SetPendingCloseDueToSocketSendFailure();
@@ -1132,6 +1135,14 @@ private:
 	TMap<int32, FNetworkGUID> IgnoringChannels;
 	bool bIgnoreAlreadyOpenedChannels;
 
+	/** Set of guids we may need to ignore when processing a delta checkpoint */
+	TSet<FNetworkGUID> IgnoredBunchGuids;
+
+	/** Set of channels we may need to ignore when processing a delta checkpoint */
+	TSet<int32> IgnoredBunchChannels;
+
+	bool bIgnoreActorBunches;
+
 	/** This is only used in UChannel::SendBunch. It's a member so that we can preserve the allocation between calls, as an optimization, and in a thread-safe way to be compatible with demo.ClientRecordAsyncEndOfFrame */
 	TArray<FOutBunch*> OutgoingBunches;
 
@@ -1155,7 +1166,6 @@ private:
 
 	FNetworkGUID GetActorGUIDFromOpenBunch(FInBunch& Bunch);
 
-
 	/** Out of order packet tracking/correction */
 
 	/** Stat tracking for the total number of out of order packets, for this connection */
diff --git a/Engine/Source/Runtime/Engine/Classes/Engine/NetDriver.h b/Engine/Source/Runtime/Engine/Classes/Engine/NetDriver.h
--- a/Engine/Source/Runtime/Engine/Classes/Engine/NetDriver.h
+++ b/Engine/Source/Runtime/Engine/Classes/Engine/NetDriver.h
@@ -1160,6 +1160,9 @@ public:
 
 	/** Called when an actor channel is remotely opened for an actor. */
 	ENGINE_API virtual void NotifyActorChannelOpen(UActorChannel* Channel, AActor* Actor) {}
+	
+	/** Called when an actor channel is cleaned up foor an actor. */
+	ENGINE_API virtual void NotifyActorChannelCleanedUp(UActorChannel* Channel, EChannelCloseReason CloseReason) {}
 
 protected:
 
diff --git a/Engine/Source/Runtime/Engine/Private/DataChannel.cpp b/Engine/Source/Runtime/Engine/Private/DataChannel.cpp
--- a/Engine/Source/Runtime/Engine/Private/DataChannel.cpp
+++ b/Engine/Source/Runtime/Engine/Private/DataChannel.cpp
@@ -1871,13 +1871,15 @@ void UActorChannel::DestroyActorAndComponents()
 	}
 }
 
-bool UActorChannel::CleanUp( const bool bForDestroy, EChannelCloseReason CloseReason )
+bool UActorChannel::CleanUp(const bool bForDestroy, EChannelCloseReason CloseReason)
 {
 	SCOPE_CYCLE_COUNTER(Stat_ActorChanCleanUp);
 
 	checkf(Connection != nullptr, TEXT("UActorChannel::CleanUp: Connection is null!"));
 	checkf(Connection->Driver != nullptr, TEXT("UActorChannel::CleanUp: Connection->Driver is null!"));
 
+	Connection->Driver->NotifyActorChannelCleanedUp(this, CloseReason);
+
 	UReplicationConnectionDriver* const ConnectionDriver = Connection->GetReplicationConnectionDriver();
 	if (ConnectionDriver)
 	{
@@ -2383,6 +2385,12 @@ void UActorChannel::ReceivedBunch( FInBunch & Bunch )
 			Bunch << ActorNetGUID;
 
 			Mark.Pop( Bunch );
+
+			// we can now map guid to channel, even if all the bunches get queued
+			if (Connection->InternalAck)
+			{
+				Connection->NotifyActorNetGUID(this);
+			}
 		}
 
 		// We need to queue this bunch if any of these are true:
diff --git a/Engine/Source/Runtime/Engine/Private/DemoNetDriver.cpp b/Engine/Source/Runtime/Engine/Private/DemoNetDriver.cpp
--- a/Engine/Source/Runtime/Engine/Private/DemoNetDriver.cpp
+++ b/Engine/Source/Runtime/Engine/Private/DemoNetDriver.cpp
@@ -239,42 +239,64 @@ private:
 	const FArchivePos StartPosition;
 };
 
-class FScopedIgnoreDeletedGuids : public FNoncopyable
+class FScopedIgnoreOpenChannels : public FNoncopyable
 {
 public:
-	FScopedIgnoreDeletedGuids(UNetConnection* InConnection, TSet<FNetworkGUID>& DeletedActorGuids)
+	FScopedIgnoreOpenChannels(UNetConnection* InConnection):
+		Connection(InConnection)
+	{
+		if (Connection.IsValid())
+		{
+			Connection->SetIgnoreAlreadyOpenedChannels(true);
+		}	
+	}
+
+	~FScopedIgnoreOpenChannels()
+	{
+		if (Connection.IsValid())
+		{
+			Connection->SetIgnoreAlreadyOpenedChannels(false);
+		}	
+	}
+
+private:
+	TWeakObjectPtr<UNetConnection> Connection;
+};
+
+class FScopedIgnoreClosedChannels : public FNoncopyable
+{
+public:
+	FScopedIgnoreClosedChannels(UNetConnection* InConnection, TSet<FNetworkGUID>&& IgnoredBunchGuids)
 		: ServerConnection(InConnection)
 	{
-		if (ServerConnection)
+		if (ServerConnection.IsValid())
 		{
-			ServerConnection->IgnoredGuids = MoveTemp(DeletedActorGuids);
-			ServerConnection->IgnoredChannels.Empty();
+			ServerConnection->SetIgnoreActorBunches(true, MoveTemp(IgnoredBunchGuids));
 		}
 	}
 
-	~FScopedIgnoreDeletedGuids()
+	~FScopedIgnoreClosedChannels()
 	{
-		if (ServerConnection)
+		if (ServerConnection.IsValid())
 		{
-			ServerConnection->IgnoredGuids.Empty();
-			ServerConnection->IgnoredChannels.Empty();
+			ServerConnection->SetIgnoreActorBunches(false, TSet<FNetworkGUID>());
 		}
 	}
 
 private:
-	UNetConnection* ServerConnection;
+	TWeakObjectPtr<UNetConnection> ServerConnection;
 };
 
 class FJumpToLiveReplayTask : public FQueuedReplayTask
 {
 public:
-	FJumpToLiveReplayTask( UDemoNetDriver* InDriver ) : FQueuedReplayTask( InDriver )
+	FJumpToLiveReplayTask(UDemoNetDriver* InDriver) : FQueuedReplayTask(InDriver)
 	{
 		if (Driver.IsValid())
 		{
-		InitialTotalDemoTime	= Driver->ReplayStreamer->GetTotalDemoTime();
-		TaskStartTime			= FPlatformTime::Seconds();
-	}
+			InitialTotalDemoTime	= Driver->ReplayStreamer->GetTotalDemoTime();
+			TaskStartTime			= FPlatformTime::Seconds();
+		}
 	}
 
 	virtual void StartTask()
@@ -288,23 +310,23 @@ public:
 			return true;
 		}
 
-		if ( !Driver->ReplayStreamer->IsLive() )
+		if (!Driver->ReplayStreamer->IsLive())
 		{
 			// The replay is no longer live, so don't try to jump to end
 			return true;
 		}
 
 		// Wait for the most recent live time
-		const bool bHasNewReplayTime = ( Driver->ReplayStreamer->GetTotalDemoTime() != InitialTotalDemoTime );
+		const bool bHasNewReplayTime = (Driver->ReplayStreamer->GetTotalDemoTime() != InitialTotalDemoTime);
 
 		// If we haven't gotten a new time from the demo by now, assume it might not be live, and just jump to the end now so we don't hang forever
-		const bool bTimeExpired = ( FPlatformTime::Seconds() - TaskStartTime >= 15 );
+		const bool bTimeExpired = (FPlatformTime::Seconds() - TaskStartTime >= 15);
 
-		if ( bHasNewReplayTime || bTimeExpired )
+		if (bHasNewReplayTime || bTimeExpired)
 		{
-			if ( bTimeExpired )
+			if (bTimeExpired)
 			{
-				UE_LOG( LogDemo, Warning, TEXT( "FJumpToLiveReplayTask::Tick: Too much time since last live update." ) );
+				UE_LOG(LogDemo, Warning, TEXT("FJumpToLiveReplayTask::Tick: Too much time since last live update."));
 			}
 
 			// We're ready to jump to the end now
@@ -421,7 +443,7 @@ public:
 class FSkipTimeInSecondsTask : public FQueuedReplayTask
 {
 public:
-	FSkipTimeInSecondsTask( UDemoNetDriver* InDriver, const float InSecondsToSkip ) : FQueuedReplayTask( InDriver ), SecondsToSkip( InSecondsToSkip )
+	FSkipTimeInSecondsTask(UDemoNetDriver* InDriver, const float InSecondsToSkip) : FQueuedReplayTask(InDriver), SecondsToSkip(InSecondsToSkip)
 	{
 	}
 
@@ -432,13 +454,13 @@ public:
 			return;
 		}
 
-		check( !Driver->IsFastForwarding() );
+		check(!Driver->IsFastForwarding());
 
-		const uint32 TimeInMSToCheck = FMath::Clamp( Driver->GetDemoCurrentTimeInMS() + ( uint32 )( SecondsToSkip * 1000 ), ( uint32 )0, Driver->ReplayStreamer->GetTotalDemoTime() );
+		const uint32 TimeInMSToCheck = FMath::Clamp(Driver->GetDemoCurrentTimeInMS() + (uint32)(SecondsToSkip * 1000), (uint32)0, Driver->ReplayStreamer->GetTotalDemoTime());
 
-		Driver->ReplayStreamer->SetHighPriorityTimeRange( Driver->GetDemoCurrentTimeInMS(), TimeInMSToCheck );
+		Driver->ReplayStreamer->SetHighPriorityTimeRange(Driver->GetDemoCurrentTimeInMS(), TimeInMSToCheck);
 
-		Driver->SkipTimeInternal( SecondsToSkip, true, false );
+		Driver->SkipTimeInternal(SecondsToSkip, true, false);
 	}
 
 	virtual bool Tick() override
@@ -539,7 +561,7 @@ public:
 
 		if (!Result.WasSuccessful())
 		{
-			UE_LOG( LogDemo, Warning, TEXT( "FFastForwardLevelsTask::CheckpointReady: Failed to get checkpoint." ) );
+			UE_LOG(LogDemo, Warning, TEXT("FFastForwardLevelsTask::CheckpointReady: Failed to get checkpoint."));
 		}
 	}
 
@@ -611,14 +633,14 @@ FString UDemoNetDriver::GetLevelPackageName(const ULevel& InLevel)
 	return UWorld::RemovePIEPrefix(PathName);
 }
 
-void UDemoNetDriver::AddReplayTask( FQueuedReplayTask* NewTask )
+void UDemoNetDriver::AddReplayTask(FQueuedReplayTask* NewTask)
 {
-	UE_LOG( LogDemo, Verbose, TEXT( "UDemoNetDriver::AddReplayTask. Name: %s" ), *NewTask->GetName().ToString() );
+	UE_LOG(LogDemo, Verbose, TEXT("UDemoNetDriver::AddReplayTask. Name: %s"), *NewTask->GetName().ToString());
 
-	QueuedReplayTasks.Add( TSharedPtr< FQueuedReplayTask >( NewTask ) );
+	QueuedReplayTasks.Emplace(NewTask);
 
 	// Give this task a chance to immediately start if nothing else is happening
-	if ( !IsAnyTaskPending() )
+	if (!IsAnyTaskPending())
 	{
 		ProcessReplayTasks();	
 	}
@@ -626,7 +648,7 @@ void UDemoNetDriver::AddReplayTask( FQueuedReplayTask* NewTask )
 
 bool UDemoNetDriver::IsAnyTaskPending() const
 {
-	return ( QueuedReplayTasks.Num() > 0 ) || ActiveReplayTask.IsValid();
+	return (QueuedReplayTasks.Num() > 0) || ActiveReplayTask.IsValid();
 }
 
 void UDemoNetDriver::ClearReplayTasks()
@@ -643,23 +665,23 @@ bool UDemoNetDriver::ProcessReplayTasks()
 	// in StartTask() or Tick()), the current task won't be destroyed immediately.
 	TSharedPtr<FQueuedReplayTask> LocalActiveTask;
 
-	if ( !ActiveReplayTask.IsValid() && QueuedReplayTasks.Num() > 0 )
+	if (!ActiveReplayTask.IsValid() && QueuedReplayTasks.Num() > 0)
 	{
 		// If we don't have an active task, pull one off now
 		ActiveReplayTask = QueuedReplayTasks[0];
 		LocalActiveTask = ActiveReplayTask;
-		QueuedReplayTasks.RemoveAt( 0 );
+		QueuedReplayTasks.RemoveAt(0);
 
-		UE_LOG( LogDemo, Verbose, TEXT( "UDemoNetDriver::ProcessReplayTasks. Name: %s" ), *ActiveReplayTask->GetName().ToString() );
+		UE_LOG(LogDemo, Verbose, TEXT("UDemoNetDriver::ProcessReplayTasks. Name: %s"), *ActiveReplayTask->GetName().ToString());
 
 		// Start the task
 		ActiveReplayTask->StartTask();
 	}
 
 	// Tick the currently active task
-	if ( ActiveReplayTask.IsValid() )
+	if (ActiveReplayTask.IsValid())
 	{
-		if ( !ActiveReplayTask->Tick() )
+		if (!ActiveReplayTask->Tick())
 		{
 			// Task isn't done, we can return
 			return false;
@@ -672,16 +694,16 @@ bool UDemoNetDriver::ProcessReplayTasks()
 	return true;	// No tasks to process
 }
 
-bool UDemoNetDriver::IsNamedTaskInQueue( const FName& Name ) const
+bool UDemoNetDriver::IsNamedTaskInQueue(const FName& Name) const
 {
-	if ( ActiveReplayTask.IsValid() && ActiveReplayTask->GetName() == Name )
+	if (ActiveReplayTask.IsValid() && ActiveReplayTask->GetName() == Name)
 	{
 		return true;
 	}
 
-	for ( int32 i = 0; i < QueuedReplayTasks.Num(); i++ )
+	for (const TSharedRef<FQueuedReplayTask>& Task : QueuedReplayTasks)
 	{
-		if ( QueuedReplayTasks[i]->GetName() == Name )
+		if (Task->GetName() == Name)
 		{
 			return true;
 		}
@@ -695,9 +717,9 @@ FName UDemoNetDriver::GetNextQueuedTaskName() const
 	return QueuedReplayTasks.Num() > 0 ? QueuedReplayTasks[0]->GetName() : NAME_None;
 }
 
-bool UDemoNetDriver::InitBase( bool bInitAsClient, FNetworkNotify* InNotify, const FURL& URL, bool bReuseAddressAndPort, FString& Error )
+bool UDemoNetDriver::InitBase(bool bInitAsClient, FNetworkNotify* InNotify, const FURL& URL, bool bReuseAddressAndPort, FString& Error)
 {
-	if ( Super::InitBase( bInitAsClient, InNotify, URL, bReuseAddressAndPort, Error ) )
+	if (Super::InitBase(bInitAsClient, InNotify, URL, bReuseAddressAndPort, Error))
 	{
 		DemoURL							= URL;
 		Time							= 0;
@@ -717,7 +739,7 @@ bool UDemoNetDriver::InitBase( bool bInitAsClient, FNetworkNotify* InNotify, con
 		CheckpointSaveMaxMSPerFrame		= -1.0f;
 		RecordBuildConsiderAndPrioritizeTimeSlice = CVarDemoMaximumRepPrioritizeTime.GetValueOnAnyThread();
 
-		if ( RelevantTimeout == 0.0f )
+		if (RelevantTimeout == 0.0f)
 		{
 			RelevantTimeout = 5.0f;
 		}
@@ -741,10 +763,10 @@ bool UDemoNetDriver::InitBase( bool bInitAsClient, FNetworkNotify* InNotify, con
 
 void UDemoNetDriver::FinishDestroy()
 {
-	if ( !HasAnyFlags( RF_ClassDefaultObject ) )
+	if (!HasAnyFlags(RF_ClassDefaultObject))
 	{
 		// Make sure we stop any recording/playing that might be going on
-		if ( IsRecording() || IsPlaying() )
+		if (IsRecording() || IsPlaying())
 		{
 			StopDemo();
 		}
@@ -757,7 +779,7 @@ void UDemoNetDriver::FinishDestroy()
 
 FString UDemoNetDriver::LowLevelGetNetworkNumber()
 {
-	return FString( TEXT( "" ) );
+	return FString(TEXT(""));
 }
 
 void UDemoNetDriver::ResetDemoState()
@@ -783,51 +805,51 @@ void UDemoNetDriver::ResetDemoState()
 	ClearLevelStreamingState();
 }
 
-bool UDemoNetDriver::InitConnect( FNetworkNotify* InNotify, const FURL& ConnectURL, FString& Error )
+bool UDemoNetDriver::InitConnect(FNetworkNotify* InNotify, const FURL& ConnectURL, FString& Error)
 {
-	if ( GetWorld() == nullptr )
+	if (GetWorld() == nullptr)
 	{
-		UE_LOG( LogDemo, Error, TEXT( "GetWorld() == nullptr" ) );
+		UE_LOG(LogDemo, Error, TEXT("GetWorld() == nullptr"));
 		return false;
 	}
 
-	if ( GetWorld()->GetGameInstance() == nullptr )
+	if (GetWorld()->GetGameInstance() == nullptr)
 	{
-		UE_LOG( LogDemo, Error, TEXT( "GetWorld()->GetGameInstance() == nullptr" ) );
+		UE_LOG(LogDemo, Error, TEXT("GetWorld()->GetGameInstance() == nullptr"));
 		return false;
 	}
 
 	// handle default initialization
-	if ( !InitBase( true, InNotify, ConnectURL, false, Error ) )
+	if (!InitBase(true, InNotify, ConnectURL, false, Error))
 	{
-		GetWorld()->GetGameInstance()->HandleDemoPlaybackFailure( EDemoPlayFailure::InitBase, FString( TEXT( "InitBase FAILED" ) ) );
+		GetWorld()->GetGameInstance()->HandleDemoPlaybackFailure(EDemoPlayFailure::InitBase, FString(TEXT("InitBase FAILED")));
 		return false;
 	}
 
-	GuidCache->SetNetworkChecksumMode( FNetGUIDCache::ENetworkChecksumMode::SaveButIgnore );
+	GuidCache->SetNetworkChecksumMode(FNetGUIDCache::ENetworkChecksumMode::SaveButIgnore);
 
-	if ( CVarForceDisableAsyncPackageMapLoading.GetValueOnGameThread() > 0 )
+	if (CVarForceDisableAsyncPackageMapLoading.GetValueOnGameThread() > 0)
 	{
-		GuidCache->SetAsyncLoadMode( FNetGUIDCache::EAsyncLoadMode::ForceDisable );
+		GuidCache->SetAsyncLoadMode(FNetGUIDCache::EAsyncLoadMode::ForceDisable);
 	}
 	else
 	{
-		GuidCache->SetAsyncLoadMode( FNetGUIDCache::EAsyncLoadMode::UseCVar );
+		GuidCache->SetAsyncLoadMode(FNetGUIDCache::EAsyncLoadMode::UseCVar);
 	}
 
 	// Playback, local machine is a client, and the demo stream acts "as if" it's the server.
 	ServerConnection = NewObject<UNetConnection>(GetTransientPackage(), UDemoNetConnection::StaticClass());
-	ServerConnection->InitConnection( this, USOCK_Pending, ConnectURL, 1000000 );
+	ServerConnection->InitConnection(this, USOCK_Pending, ConnectURL, 1000000);
 
-	TArray< FString > UserNames;
+	TArray<FString> UserNames;
 
-	if ( GetWorld()->GetGameInstance()->GetFirstGamePlayer() != nullptr )
+	if (GetWorld()->GetGameInstance()->GetFirstGamePlayer() != nullptr)
 	{
 		FUniqueNetIdRepl ViewerId = GetWorld()->GetGameInstance()->GetFirstGamePlayer()->GetPreferredUniqueNetId();
 
-		if ( ViewerId.IsValid() )
+		if (ViewerId.IsValid())
 		{ 
-			UserNames.Add( ViewerId.ToString() );
+			UserNames.Add(ViewerId.ToString());
 		}
 	}
 
@@ -837,22 +859,22 @@ bool UDemoNetDriver::InitConnect( FNetworkNotify* InNotify, const FURL& ConnectU
 		SetDuplicateLevelID(FCString::Atoi(LevelPrefixOverrideOption));
 	}
 
-	if ( GetDuplicateLevelID() == -1 )
+	if (GetDuplicateLevelID() == -1)
 	{
 		// Set this driver as the demo net driver for the source level collection.
-		FLevelCollection* const SourceCollection = GetWorld()->FindCollectionByType( ELevelCollectionType::DynamicSourceLevels );
-		if ( SourceCollection )
+		FLevelCollection* const SourceCollection = GetWorld()->FindCollectionByType(ELevelCollectionType::DynamicSourceLevels);
+		if (SourceCollection)
 		{
-			SourceCollection->SetDemoNetDriver( this );
+			SourceCollection->SetDemoNetDriver(this);
 		}
 	}
 	else
 	{
 		// Set this driver as the demo net driver for the duplicate level collection.
-		FLevelCollection* const DuplicateCollection = GetWorld()->FindCollectionByType( ELevelCollectionType::DynamicDuplicatedLevels );
-		if ( DuplicateCollection )
+		FLevelCollection* const DuplicateCollection = GetWorld()->FindCollectionByType(ELevelCollectionType::DynamicDuplicatedLevels);
+		if (DuplicateCollection)
 		{
-			DuplicateCollection->SetDemoNetDriver( this );
+			DuplicateCollection->SetDemoNetDriver(this);
 		}
 	}
 
@@ -866,7 +888,7 @@ bool UDemoNetDriver::InitConnect( FNetworkNotify* InNotify, const FURL& ConnectU
 		UserNames, 
 		false, 
 		FNetworkVersion::GetReplayVersion(), 
-		FStartStreamingCallback::CreateUObject( this, &UDemoNetDriver::ReplayStreamingReady ) );
+		FStartStreamingCallback::CreateUObject(this, &UDemoNetDriver::ReplayStreamingReady));
 
 	return bWasStartStreamingSuccessful;
 }
@@ -879,21 +901,21 @@ bool UDemoNetDriver::ReadPlaybackDemoHeader(FString& Error)
 
 	FArchive* FileAr = ReplayStreamer->GetHeaderArchive();
 
-	if ( !FileAr )
+	if (!FileAr)
 	{
-		Error = FString::Printf( TEXT( "Couldn't open demo file %s for reading" ), *DemoURL.Map );
-		UE_LOG( LogDemo, Error, TEXT( "UDemoNetDriver::ReadPlaybackDemoHeader: %s" ), *Error );
-		GameInstance->HandleDemoPlaybackFailure( EDemoPlayFailure::DemoNotFound, FString( EDemoPlayFailure::ToString( EDemoPlayFailure::DemoNotFound ) ) );
+		Error = FString::Printf(TEXT("Couldn't open demo file %s for reading"), *DemoURL.Map);
+		UE_LOG(LogDemo, Error, TEXT("UDemoNetDriver::ReadPlaybackDemoHeader: %s"), *Error);
+		GameInstance->HandleDemoPlaybackFailure(EDemoPlayFailure::DemoNotFound, FString(EDemoPlayFailure::ToString(EDemoPlayFailure::DemoNotFound)));
 		return false;
 	}
 
 	(*FileAr) << PlaybackDemoHeader;
 
-	if ( FileAr->IsError() )
+	if (FileAr->IsError())
 	{
-		Error = FString( TEXT( "Demo file is corrupt" ) );
-		UE_LOG( LogDemo, Error, TEXT( "UDemoNetDriver::ReadPlaybackDemoHeader: %s" ), *Error );
-		GameInstance->HandleDemoPlaybackFailure( EDemoPlayFailure::Corrupt, Error );
+		Error = FString(TEXT("Demo file is corrupt"));
+		UE_LOG(LogDemo, Error, TEXT("UDemoNetDriver::ReadPlaybackDemoHeader: %s"), *Error);
+		GameInstance->HandleDemoPlaybackFailure(EDemoPlayFailure::Corrupt, Error);
 		return false;
 	}
 	
@@ -923,7 +945,7 @@ bool UDemoNetDriver::ReadPlaybackDemoHeader(FString& Error)
 	return true;
 }
 
-bool UDemoNetDriver::InitConnectInternal( FString& Error )
+bool UDemoNetDriver::InitConnectInternal(FString& Error)
 {
 	ResetDemoState();
 
@@ -950,7 +972,7 @@ bool UDemoNetDriver::InitConnectInternal( FString& Error )
 
 	if (GetDuplicateLevelID() == -1)
 	{
-		if ( bAsyncLoadWorld && GetWorld()->WorldType != EWorldType::PIE ) // Editor doesn't support async map travel
+		if (bAsyncLoadWorld && GetWorld()->WorldType != EWorldType::PIE) // Editor doesn't support async map travel
 		{
 			LevelNamesAndTimes = PlaybackDemoHeader.LevelNamesAndTimes;
 
@@ -965,20 +987,20 @@ bool UDemoNetDriver::InitConnectInternal( FString& Error )
 			FURL LocalDemoURL;
 			LocalDemoURL.Map = PlaybackDemoHeader.LevelNamesAndTimes[0].LevelName;
 
-			FWorldContext * WorldContext = GEngine->GetWorldContextFromWorld( GetWorld() );
+			FWorldContext * WorldContext = GEngine->GetWorldContextFromWorld(GetWorld());
 
-			if ( WorldContext == nullptr )
+			if (WorldContext == nullptr)
 			{
 				UGameInstance* GameInstance = GetWorld()->GetGameInstance();
 
-				Error = FString::Printf( TEXT( "No world context" ) );
-				UE_LOG( LogDemo, Error, TEXT( "UDemoNetDriver::InitConnect: %s" ), *Error );
-				GameInstance->HandleDemoPlaybackFailure( EDemoPlayFailure::Generic, FString( TEXT( "No world context" ) ) );
+				Error = FString::Printf(TEXT("No world context"));
+				UE_LOG(LogDemo, Error, TEXT("UDemoNetDriver::InitConnect: %s"), *Error);
+				GameInstance->HandleDemoPlaybackFailure(EDemoPlayFailure::Generic, FString(TEXT("No world context")));
 				return false;
 			}
 
 			GetWorld()->DemoNetDriver = nullptr;
-			SetWorld( nullptr );
+			SetWorld(nullptr);
 
 			auto NewPendingNetGame = NewObject<UDemoPendingNetGame>();
 
@@ -998,22 +1020,22 @@ bool UDemoNetDriver::InitConnectInternal( FString& Error )
 	return true;
 }
 
-bool UDemoNetDriver::InitListen( FNetworkNotify* InNotify, FURL& ListenURL, bool bReuseAddressAndPort, FString& Error )
+bool UDemoNetDriver::InitListen(FNetworkNotify* InNotify, FURL& ListenURL, bool bReuseAddressAndPort, FString& Error)
 {
-	if ( !InitBase( false, InNotify, ListenURL, bReuseAddressAndPort, Error ) )
+	if (!InitBase(false, InNotify, ListenURL, bReuseAddressAndPort, Error))
 	{
 		return false;
 	}
 
-	GuidCache->SetNetworkChecksumMode( FNetGUIDCache::ENetworkChecksumMode::SaveButIgnore );
+	GuidCache->SetNetworkChecksumMode(FNetGUIDCache::ENetworkChecksumMode::SaveButIgnore);
 
-	check( World != nullptr );
+	check(World != nullptr);
 
 	class AWorldSettings * WorldSettings = World->GetWorldSettings(); 
 
-	if ( !WorldSettings )
+	if (!WorldSettings)
 	{
-		Error = TEXT( "No WorldSettings!!" );
+		Error = TEXT("No WorldSettings!!");
 		return false;
 	}
 
@@ -1024,27 +1046,27 @@ bool UDemoNetDriver::InitListen( FNetworkNotify* InNotify, FURL& ListenURL, bool
 
 	// Recording, local machine is server, demo stream acts "as if" it's a client.
 	UDemoNetConnection* Connection = NewObject<UDemoNetConnection>();
-	Connection->InitConnection( this, USOCK_Open, ListenURL, 1000000 );
+	Connection->InitConnection(this, USOCK_Open, ListenURL, 1000000);
 	Connection->InitSendBuffer();
 
-	AddClientConnection( Connection );
+	AddClientConnection(Connection);
 
-	const TCHAR* FriendlyNameOption = ListenURL.GetOption( TEXT("DemoFriendlyName="), nullptr );
+	const TCHAR* FriendlyNameOption = ListenURL.GetOption(TEXT("DemoFriendlyName="), nullptr);
 
 	bRecordMapChanges = ListenURL.GetOption(TEXT("RecordMapChanges"), nullptr) != nullptr;
 
-	TArray< FString > UserNames;
+	TArray<FString> UserNames;
 	AGameStateBase* GameState = GetWorld()->GetGameState();
 
 	// If a client is recording a replay, GameState may not have replicated yet
-	if (GameState != nullptr )
+	if (GameState != nullptr)
 	{
-		for ( int32 i = 0; i < GameState->PlayerArray.Num(); i++ )
+		for (int32 i = 0; i < GameState->PlayerArray.Num(); i++)
 		{
 			APlayerState* PlayerState = GameState->PlayerArray[i];
-			if ( PlayerState && !PlayerState->bIsABot && !PlayerState->bIsSpectator )
+			if (PlayerState && !PlayerState->bIsABot && !PlayerState->bIsSpectator)
 			{
-				UserNames.Add( PlayerState->UniqueId.ToString() );
+				UserNames.Add(PlayerState->UniqueId.ToString());
 			}
 		}
 	}
@@ -1058,14 +1080,14 @@ bool UDemoNetDriver::InitListen( FNetworkNotify* InNotify, FURL& ListenURL, bool
 		UserNames,
 		true,
 		FNetworkVersion::GetReplayVersion(),
-		FStartStreamingCallback::CreateUObject( this, &UDemoNetDriver::ReplayStreamingReady ) );
+		FStartStreamingCallback::CreateUObject(this, &UDemoNetDriver::ReplayStreamingReady));
 
 	AddNewLevel(World->GetOuter()->GetName());
 
 	bool Result = WriteNetworkDemoHeader(Error);
 
 	// Spawn the demo recording spectator.
-	SpawnDemoRecSpectator( Connection, ListenURL );
+	SpawnDemoRecSpectator(Connection, ListenURL);
 
 	return Result;
 }
@@ -1122,15 +1144,15 @@ void UDemoNetDriver::NotifyStreamingLevelUnload( ULevel* InLevel )
 {
 	if (InLevel && !InLevel->bClientOnlyVisible && HasLevelStreamingFixes() && IsPlaying())
 	{
-			// We can't just iterate over the levels actors, because the ones in the queue will already have been destroyed.
-			for (TMap<FString, FRollbackNetStartupActorInfo>::TIterator It = RollbackNetStartupActors.CreateIterator(); It; ++It)
+		// We can't just iterate over the levels actors, because the ones in the queue will already have been destroyed.
+		for (TMap<FString, FRollbackNetStartupActorInfo>::TIterator It = RollbackNetStartupActors.CreateIterator(); It; ++It)
+		{
+			if (It.Value().Level == InLevel)
 			{
-				if (It.Value().Level == InLevel)
-				{
-					It.RemoveCurrent();
-				}
+				It.RemoveCurrent();
 			}
 		}
+	}
 
 	Super::NotifyStreamingLevelUnload(InLevel);
 }
@@ -1152,7 +1174,7 @@ void UDemoNetDriver::OnPostLoadMapWithWorld(UWorld* InWorld)
 
 TUniquePtr<FScopedPacketManager> UDemoNetDriver::ConditionallyCreatePacketManager(ULevel& Level)
 {
-	if(IsRecording() && HasLevelStreamingFixes())
+	if (IsRecording() && HasLevelStreamingFixes())
 	{
 		// Indices need to be 1 based, so +1.
 		return MakeUnique<FScopedPacketManager>(*(UDemoNetConnection*)ClientConnections[0], FindOrAddLevelStatus(Level).LevelIndex + 1);
@@ -1163,7 +1185,7 @@ TUniquePtr<FScopedPacketManager> UDemoNetDriver::ConditionallyCreatePacketManage
 
 TUniquePtr<FScopedPacketManager> UDemoNetDriver::ConditionallyCreatePacketManager(int32 LevelIndex)
 {
-	if(IsRecording() && HasLevelStreamingFixes())
+	if (IsRecording() && HasLevelStreamingFixes())
 	{
 		// Indices need to be 1 based, so +1.
 		return MakeUnique<FScopedPacketManager>(*(UDemoNetConnection*)ClientConnections[0], LevelIndex);
@@ -1225,9 +1247,9 @@ bool UDemoNetDriver::WriteNetworkDemoHeader(FString& Error)
 {
 	FArchive* FileAr = ReplayStreamer->GetHeaderArchive();
 
-	if( !FileAr )
+	if (!FileAr)
 	{
-		Error = FString::Printf( TEXT("Couldn't open demo file %s for writing"), *DemoURL.Map );//@todo demorec: localize
+		Error = FString::Printf(TEXT("Couldn't open demo file %s for writing"), *DemoURL.Map);	//@todo demorec: localize
 		return false;
 	}
 
@@ -1312,7 +1334,7 @@ void UDemoNetDriver::TickFlushAsyncEndOfFrame(float DeltaSeconds)
 /** Accounts for the network time we spent in the demo driver. */
 double GTickFlushDemoDriverTimeSeconds = 0.0;
 
-void UDemoNetDriver::TickFlushInternal( float DeltaSeconds )
+void UDemoNetDriver::TickFlushInternal(float DeltaSeconds)
 {
 	CSV_SCOPED_TIMING_STAT_EXCLUSIVE(DemoRecording);
 
@@ -1327,7 +1349,7 @@ void UDemoNetDriver::TickFlushInternal( float DeltaSeconds )
 
 	FScopedLevelCollectionContextSwitch LCSwitch(FoundCollectionIndex, GetWorld());
 
-	Super::TickFlush( DeltaSeconds );
+	Super::TickFlush(DeltaSeconds);
 
 	if (!IsRecording() || bIsWaitingForStream)
 	{
@@ -1335,9 +1357,9 @@ void UDemoNetDriver::TickFlushInternal( float DeltaSeconds )
 		return;
 	}
 
-	if ( ReplayStreamer->GetLastError() != ENetworkReplayError::None )
+	if (ReplayStreamer->GetLastError() != ENetworkReplayError::None)
 	{
-		UE_LOG( LogDemo, Error, TEXT( "UDemoNetDriver::TickFlush: ReplayStreamer ERROR: %s" ), ENetworkReplayError::ToString( ReplayStreamer->GetLastError() ) );
+		UE_LOG(LogDemo, Error, TEXT("UDemoNetDriver::TickFlush: ReplayStreamer ERROR: %s"), ENetworkReplayError::ToString(ReplayStreamer->GetLastError()));
 		StopDemo();
 		return;
 	}
@@ -1349,27 +1371,27 @@ void UDemoNetDriver::TickFlushInternal( float DeltaSeconds )
 
 	FArchive* FileAr = ReplayStreamer->GetStreamingArchive();
 
-	if ( FileAr == nullptr )
+	if (FileAr == nullptr)
 	{
-		UE_LOG( LogDemo, Error, TEXT( "UDemoNetDriver::TickFlush: FileAr == nullptr" ) );
+		UE_LOG(LogDemo, Error, TEXT("UDemoNetDriver::TickFlush: FileAr == nullptr"));
 		StopDemo();
 		return;
 	}
 
-	DECLARE_SCOPE_CYCLE_COUNTER( TEXT( "Net replay record time" ), STAT_ReplayRecordTime, STATGROUP_Net );
+	DECLARE_SCOPE_CYCLE_COUNTER(TEXT("Net replay record time"), STAT_ReplayRecordTime, STATGROUP_Net);
 
 	const double StartTime = FPlatformTime::Seconds();
 
-	TickDemoRecord( DeltaSeconds );
+	TickDemoRecord(DeltaSeconds);
 
 	const double EndTime = FPlatformTime::Seconds();
 
-	const double RecordTotalTime = ( EndTime - StartTime );
+	const double RecordTotalTime = (EndTime - StartTime);
 
 	// While recording, the CurrentCL is the same as the recording CL.
 	ConditionallyDisplayBurnInTime(FEngineVersion::Current().GetChangelist(), DemoCurrentTime);
 
-	MaxRecordTime = FMath::Max( MaxRecordTime, RecordTotalTime );
+	MaxRecordTime = FMath::Max(MaxRecordTime, RecordTotalTime);
 
 	AccumulatedRecordTime += RecordTotalTime;
 
@@ -1379,14 +1401,14 @@ void UDemoNetDriver::TickFlushInternal( float DeltaSeconds )
 
 	const double AVG_FLUSH_TIME_IN_SECONDS = 2;
 
-	if ( ElapsedTime > AVG_FLUSH_TIME_IN_SECONDS && RecordCountSinceFlush > 0 )
+	if (ElapsedTime > AVG_FLUSH_TIME_IN_SECONDS && RecordCountSinceFlush > 0)
 	{
-		const float AvgTimeMS = ( AccumulatedRecordTime / RecordCountSinceFlush ) * 1000;
+		const float AvgTimeMS = (AccumulatedRecordTime / RecordCountSinceFlush) * 1000;
 		const float MaxRecordTimeMS = MaxRecordTime * 1000;
 
-		if ( AvgTimeMS > 8.0f )//|| MaxRecordTimeMS > 6.0f )
+		if (AvgTimeMS > 8.0f)//|| MaxRecordTimeMS > 6.0f )
 		{
-			UE_LOG( LogDemo, Verbose, TEXT( "UDemoNetDriver::TickFlush: SLOW FRAME. Avg: %2.2f, Max: %2.2f, Actors: %i" ), AvgTimeMS, MaxRecordTimeMS, GetNetworkObjectList().GetActiveObjects().Num() );
+			UE_LOG(LogDemo, Verbose, TEXT("UDemoNetDriver::TickFlush: SLOW FRAME. Avg: %2.2f, Max: %2.2f, Actors: %i"), AvgTimeMS, MaxRecordTimeMS, GetNetworkObjectList().GetActiveObjects().Num());
 		}
 
 		LastRecordAvgFlush		= EndTime;
@@ -1396,21 +1418,21 @@ void UDemoNetDriver::TickFlushInternal( float DeltaSeconds )
 	}
 }
 
-static float GetClampedDeltaSeconds( UWorld* World, const float DeltaSeconds )
+static float GetClampedDeltaSeconds(UWorld* World, const float DeltaSeconds)
 {
-	check( World != nullptr );
+	check(World != nullptr);
 
 	const float RealDeltaSeconds = DeltaSeconds;
 
 	// Clamp delta seconds
 	AWorldSettings* WorldSettings = World->GetWorldSettings();
-	const float ClampedDeltaSeconds = WorldSettings->FixupDeltaSeconds( DeltaSeconds * WorldSettings->GetEffectiveTimeDilation(), RealDeltaSeconds );
-	check( ClampedDeltaSeconds >= 0.0f );
+	const float ClampedDeltaSeconds = WorldSettings->FixupDeltaSeconds(DeltaSeconds * WorldSettings->GetEffectiveTimeDilation(), RealDeltaSeconds);
+	check(ClampedDeltaSeconds >= 0.0f);
 
 	return ClampedDeltaSeconds;
 }
 
-void UDemoNetDriver::TickDispatch( float DeltaSeconds )
+void UDemoNetDriver::TickDispatch(float DeltaSeconds)
 {
 	LLM_SCOPE(ELLMTag::Networking);
 
@@ -1422,15 +1444,15 @@ void UDemoNetDriver::TickDispatch( float DeltaSeconds )
 
 	FScopedLevelCollectionContextSwitch LCSwitch(FoundCollectionIndex, GetWorld());
 
-	Super::TickDispatch( DeltaSeconds );
+	Super::TickDispatch(DeltaSeconds);
 
-	if ( !IsPlaying() || bIsWaitingForStream )
+	if (!IsPlaying() || bIsWaitingForStream)
 	{
 		// Nothing to do
 		return;
 	}
 
-	if ( ReplayStreamer->GetLastError() != ENetworkReplayError::None )
+	if (ReplayStreamer->GetLastError() != ENetworkReplayError::None)
 	{
 		UE_LOG(LogDemo, Error, TEXT("UDemoNetDriver::TickDispatch: ReplayStreamer ERROR: %s"), ENetworkReplayError::ToString(ReplayStreamer->GetLastError()));
 		NotifyDemoPlaybackFailure(EDemoPlayFailure::ReplayStreamerInternal);
@@ -1439,9 +1461,9 @@ void UDemoNetDriver::TickDispatch( float DeltaSeconds )
 
 	FArchive* FileAr = ReplayStreamer->GetStreamingArchive();
 
-	if ( FileAr == nullptr )
+	if (FileAr == nullptr)
 	{
-		UE_LOG( LogDemo, Error, TEXT( "UDemoNetDriver::TickDispatch: FileAr == nullptr" ) );
+		UE_LOG(LogDemo, Error, TEXT("UDemoNetDriver::TickDispatch: FileAr == nullptr"));
 		NotifyDemoPlaybackFailure(EDemoPlayFailure::ReplayStreamerInternal);
 		return;
 	}
@@ -1451,7 +1473,7 @@ void UDemoNetDriver::TickDispatch( float DeltaSeconds )
 		// Wait until all levels are streamed in
 		for (ULevelStreaming* StreamingLevel : World->GetStreamingLevels())
 		{
-			if ( StreamingLevel && StreamingLevel->ShouldBeLoaded() && (!StreamingLevel->IsLevelLoaded() || !StreamingLevel->GetLoadedLevel()->GetOutermost()->IsFullyLoaded() || !StreamingLevel->IsLevelVisible() ) )
+			if (StreamingLevel && StreamingLevel->ShouldBeLoaded() && (!StreamingLevel->IsLevelLoaded() || !StreamingLevel->GetLoadedLevel()->GetOutermost()->IsFullyLoaded() || !StreamingLevel->IsLevelVisible()))
 			{
 				// Abort, we have more streaming levels to load
 				return;
@@ -1459,7 +1481,7 @@ void UDemoNetDriver::TickDispatch( float DeltaSeconds )
 		}
 	}	
 
-	if ( CVarDemoTimeDilation.GetValueOnGameThread() >= 0.0f )
+	if (CVarDemoTimeDilation.GetValueOnGameThread() >= 0.0f)
 	{
 		World->GetWorldSettings()->DemoPlayTimeDilation = CVarDemoTimeDilation.GetValueOnGameThread();
 	}
@@ -1501,13 +1523,13 @@ void UDemoNetDriver::TickDispatch( float DeltaSeconds )
 		}
 	}
 
-	TickDemoPlayback( DeltaSeconds );
+	TickDemoPlayback(DeltaSeconds);
 
 	// Used LastProcessedPacketTime because it will correlate better with recorded frame time.
 	ConditionallyDisplayBurnInTime(PlaybackDemoHeader.EngineVersion.GetChangelist(), LastProcessedPacketTime);
 }
 
-void UDemoNetDriver::ProcessRemoteFunction( class AActor* Actor, class UFunction* Function, void* Parameters, struct FOutParmRec* OutParms, struct FFrame* Stack, class UObject* SubObject )
+void UDemoNetDriver::ProcessRemoteFunction(AActor* Actor, UFunction* Function, void* Parameters, FOutParmRec* OutParms, FFrame* Stack, UObject* SubObject)
 {
 #if !UE_BUILD_SHIPPING
 	bool bBlockSendRPC = false;
@@ -1517,7 +1539,7 @@ void UDemoNetDriver::ProcessRemoteFunction( class AActor* Actor, class UFunction
 	if (!bBlockSendRPC)
 #endif
 	{
-		if ( IsRecording() )
+		if (IsRecording())
 		{
 			TUniquePtr<FScopedPacketManager> PacketManager(ConditionallyCreatePacketManager(*Actor->GetLevel()));
 
@@ -1534,7 +1556,7 @@ void UDemoNetDriver::ProcessRemoteFunction( class AActor* Actor, class UFunction
 
 bool UDemoNetDriver::ShouldClientDestroyTearOffActors() const
 {
-	if ( CVarDemoFastForwardDestroyTearOffActors.GetValueOnGameThread() != 0 )
+	if (CVarDemoFastForwardDestroyTearOffActors.GetValueOnGameThread() != 0)
 	{
 		return bIsFastForwarding;
 	}
@@ -1544,7 +1566,7 @@ bool UDemoNetDriver::ShouldClientDestroyTearOffActors() const
 
 bool UDemoNetDriver::ShouldSkipRepNotifies() const
 {
-	if ( CVarDemoFastForwardSkipRepNotifies.GetValueOnAnyThread() != 0 )
+	if (CVarDemoFastForwardSkipRepNotifies.GetValueOnAnyThread() != 0)
 	{
 		return bIsFastForwarding;
 	}
@@ -1554,18 +1576,18 @@ bool UDemoNetDriver::ShouldSkipRepNotifies() const
 
 void UDemoNetDriver::StopDemo()
 {
-	if ( !IsRecording() && !IsPlaying() )
+	if (!IsRecording() && !IsPlaying())
 	{
-		UE_LOG( LogDemo, Log, TEXT( "StopDemo: No demo is playing" ) );
+		UE_LOG(LogDemo, Log, TEXT("StopDemo: No demo is playing"));
 		return;
 	}
 	OnDemoFinishRecordingDelegate.Broadcast();
-	UE_LOG( LogDemo, Log, TEXT( "StopDemo: Demo %s stopped at frame %d" ), *DemoURL.Map, DemoFrameNum );
+	UE_LOG(LogDemo, Log, TEXT("StopDemo: Demo %s stopped at frame %d"), *DemoURL.Map, DemoFrameNum);
 
-	if ( !ServerConnection )
+	if (!ServerConnection)
 	{
 		// let GC cleanup the object
-		if ( ClientConnections.Num() > 0 && ClientConnections[0] != nullptr )
+		if (ClientConnections.Num() > 0 && ClientConnections[0] != nullptr)
 		{
 			ClientConnections[0]->Close();
 		}
@@ -1584,7 +1606,7 @@ void UDemoNetDriver::StopDemo()
 	ActiveReplayName = FString();
 	ResetDemoState();
 
-	check( !IsRecording() && !IsPlaying() );
+	check(!IsRecording() && !IsPlaying());
 }
 
 /*-----------------------------------------------------------------------------
@@ -1593,7 +1615,7 @@ Demo Recording tick.
 
 bool UDemoNetDriver::DemoReplicateActor(AActor* Actor, UNetConnection* Connection, bool bMustReplicate)
 {
-	if ( Actor->NetDormancy == DORM_Initial && Actor->IsNetStartupActor() )
+	if (Actor->NetDormancy == DORM_Initial && Actor->IsNetStartupActor())
 	{
 		return false;
 	}
@@ -1605,7 +1627,7 @@ bool UDemoNetDriver::DemoReplicateActor(AActor* Actor, UNetConnection* Connectio
 	// Handle role swapping if this is a client-recorded replay.
 	FScopedActorRoleSwap RoleSwap(Actor);
 
-	if ((Actor->GetRemoteRole() != ROLE_None || Actor->GetTearOff()) && (Actor == Connection->PlayerController || Cast< APlayerController >(Actor) == nullptr))
+	if ((Actor->GetRemoteRole() != ROLE_None || Actor->GetTearOff()) && (Actor == Connection->PlayerController || Cast<APlayerController>(Actor) == nullptr))
 	{
 		const bool bShouldHaveChannel =
 			Actor->bRelevantForNetworkReplays &&
@@ -1640,9 +1662,9 @@ bool UDemoNetDriver::DemoReplicateActor(AActor* Actor, UNetConnection* Connectio
 		}
 	}
 
-	if ( bMustReplicate && !HasDeltaCheckpoints() && Connection->Driver->OutBunches == OriginalOutBunches )
+	if (bMustReplicate && !HasDeltaCheckpoints() && Connection->Driver->OutBunches == OriginalOutBunches)
 	{
-		UE_LOG( LogDemo, Error, TEXT( "DemoReplicateActor: bMustReplicate is true but nothing was sent: %s" ), Actor ? *Actor->GetName() : TEXT( "NULL" ) );
+		UE_LOG(LogDemo, Error, TEXT("DemoReplicateActor: bMustReplicate is true but nothing was sent: %s"), *GetNameSafe(Actor));
 	}
 
 	return bDidReplicateActor;
@@ -1659,7 +1681,7 @@ void UDemoNetDriver::SerializeGuidCache(TSharedPtr<FNetGUIDCache> InGuidCache, F
 
 	const bool bDeltaCheckpoint = HasDeltaCheckpoints();
 
-	for ( auto It = InGuidCache->ObjectLookup.CreateIterator(); It; ++It )
+	for (auto It = InGuidCache->ObjectLookup.CreateIterator(); It; ++It)
 	{
 		FNetworkGUID& NetworkGUID = It.Key();
 		FNetGuidCacheObject& CacheObject = It.Value();
@@ -1714,7 +1736,7 @@ void UDemoNetDriver::SerializeGuidCache(TSharedPtr<FNetGUIDCache> InGuidCache, F
 	*CheckpointArchive << NumValues;
 	CheckpointArchive->Seek(Pos);
 
-	UE_LOG( LogDemo, Verbose, TEXT( "Checkpoint. SerializeGuidCache: %i Unloaded: %i" ), NumValues, UnloadedValues );
+	UE_LOG(LogDemo, Verbose, TEXT("Checkpoint. SerializeGuidCache: %i Unloaded: %i"), NumValues, UnloadedValues);
 }
 
 float UDemoNetDriver::GetCheckpointSaveMaxMSPerFrame() const
@@ -1739,17 +1761,15 @@ void UDemoNetDriver::SaveCheckpoint()
 
 	FArchive* CheckpointArchive = ReplayStreamer->GetCheckpointArchive();
 
-	if ( CheckpointArchive == nullptr )
+	if (CheckpointArchive == nullptr)
 	{
 		// This doesn't mean error, it means the streamer isn't ready to save checkpoints
 		return;
 	}
 
-	check( CheckpointArchive->TotalSize() == 0 );
-
-	check( ClientConnections[0]->SendBuffer.GetNumBits() == 0 );
-
-	check( CheckpointSaveContext.CheckpointSaveState == ECheckpointSaveState_Idle );
+	check(CheckpointArchive->TotalSize() == 0);
+	check(ClientConnections[0]->SendBuffer.GetNumBits() == 0);
+	check(CheckpointSaveContext.CheckpointSaveState == ECheckpointSaveState_Idle);
 	
 	const bool bDeltaCheckpoint = HasDeltaCheckpoints();
 
@@ -1828,7 +1848,7 @@ void UDemoNetDriver::SaveCheckpoint()
 		{
 			AActor* Actor = ObjectInfo.Get()->Actor;
 
-			if ( ClientConnections[0]->FindActorChannelRef( Actor ) )
+			if (ClientConnections[0]->FindActorChannelRef(Actor))
 			{
 				if (!bDeltaCheckpoint || ObjectInfo->bDirtyForReplay)
 				{
@@ -1840,29 +1860,29 @@ void UDemoNetDriver::SaveCheckpoint()
 		}
 	}
 
-	if ( CheckpointSaveContext.PendingCheckpointActors.Num() == 0 )
+	if (CheckpointSaveContext.PendingCheckpointActors.Num() == 0)
 	{
 		return;
 	}
 
-	UPackageMapClient* PackageMapClient = ( ( UPackageMapClient* )ClientConnections[0]->PackageMap );
+	UPackageMapClient* PackageMapClient = Cast<UPackageMapClient>(ClientConnections[0]->PackageMap);
 
-	PackageMapClient->SavePackageMapExportAckStatus( CheckpointSaveContext.CheckpointAckState );
+	PackageMapClient->SavePackageMapExportAckStatus(CheckpointSaveContext.CheckpointAckState);
 
 	// We are now processing checkpoint actors	
 	CheckpointSaveContext.CheckpointSaveState = ECheckpointSaveState_ProcessCheckpointActors;
 	CheckpointSaveContext.TotalCheckpointSaveTimeSeconds = 0;
 	CheckpointSaveContext.TotalCheckpointReplicationTimeSeconds = 0;
 	CheckpointSaveContext.TotalCheckpointSaveFrames = 0;
+
 	LastCheckpointTime = DemoCurrentTime;
 
 	if (bDeltaCheckpoint)
 	{
-		CheckpointSaveContext.DeltaDeletedActorGuids = MoveTemp(DeltaDeletedActorGuids);
-		CheckpointSaveContext.DeltaDeletedNetStartupActors = MoveTemp(DeltaDeletedNetStartupActors);
+		CheckpointSaveContext.DeltaCheckpointData = MoveTemp(RecordingDeltaCheckpointData);
 	}
 
-	UE_LOG( LogDemo, Log, TEXT( "Starting checkpoint. Actors: %i" ), GetNetworkObjectList().GetActiveObjects().Num() );
+	UE_LOG(LogDemo, Log, TEXT("Starting checkpoint. Actors: %i"), GetNetworkObjectList().GetActiveObjects().Num());
 
 	// Do the first checkpoint tick now if we're not amortizing
 	if (GetCheckpointSaveMaxMSPerFrame() <= 0.0f)
@@ -1893,16 +1913,16 @@ static bool inline ShouldExecuteState(const FRepActorsCheckpointParams& Params,
 
 void UDemoNetDriver::TickCheckpoint()
 {
-	if ( CheckpointSaveContext.CheckpointSaveState == ECheckpointSaveState_Idle )
+	if (CheckpointSaveContext.CheckpointSaveState == ECheckpointSaveState_Idle)
 	{
 		return;
 	}
 
-	DECLARE_SCOPE_CYCLE_COUNTER( TEXT( "SaveCheckpoint time" ), STAT_ReplayCheckpointSaveTime, STATGROUP_Net );
+	DECLARE_SCOPE_CYCLE_COUNTER(TEXT("SaveCheckpoint time"), STAT_ReplayCheckpointSaveTime, STATGROUP_Net);
 
 	FArchive* CheckpointArchive = ReplayStreamer->GetCheckpointArchive();
 
-	if ( !ensure( CheckpointArchive != nullptr ) )
+	if (!ensure(CheckpointArchive != nullptr))
 	{
 		return;
 	}
@@ -1925,10 +1945,10 @@ void UDemoNetDriver::TickCheckpoint()
 
 		FlushNetChecked(*ClientConnection);
 
-		UPackageMapClient* PackageMapClient = ( ( UPackageMapClient* )ClientConnection->PackageMap );
+		UPackageMapClient* PackageMapClient = Cast<UPackageMapClient>(ClientConnection->PackageMap);
 
 		// Save package map ack status in case we export stuff during the checkpoint (so we can restore the connection back to what it was before we saved the checkpoint)
-		PackageMapClient->OverridePackageMapExportAckStatus( &CheckpointSaveContext.CheckpointAckState );
+		PackageMapClient->OverridePackageMapExportAckStatus(&CheckpointSaveContext.CheckpointAckState);
 
 		const bool bDeltaCheckpoint = HasDeltaCheckpoints();
 
@@ -1949,7 +1969,7 @@ void UDemoNetDriver::TickCheckpoint()
 
 					// Normally AGameStateBase::ReplicatedWorldTimeSeconds is only updated periodically,
 					// but we want to make sure it's accurate for the checkpoint.
-					if ( GameState )
+					if (GameState)
 					{
 						GameState->UpdateServerTimeSeconds();
 					}
@@ -1982,7 +2002,7 @@ void UDemoNetDriver::TickCheckpoint()
 						}
 						while (--NumActorsToReplicate && bContinue);
 
-						if ( GameState )
+						if (GameState)
 						{
 							// Restore the game state's replicated world time
 							GameState->ReplicatedWorldTimeSeconds = SavedReplicatedServerTimeSeconds;
@@ -1994,10 +2014,10 @@ void UDemoNetDriver::TickCheckpoint()
 					}
 
 					// We are done processing for this frame so  store the TotalCheckpointSave time here to be true to the old behavior which did not account for the	actual saving time of the check point
-					CheckpointSaveContext.TotalCheckpointReplicationTimeSeconds += ( FPlatformTime::Seconds() - Params.StartCheckpointTime );
+					CheckpointSaveContext.TotalCheckpointReplicationTimeSeconds += (FPlatformTime::Seconds() - Params.StartCheckpointTime);
 
 					// if we have replicated all checkpointactors, move on to the next state
-					if ( CheckpointSaveContext.PendingCheckpointActors.Num() == 0 )
+					if (CheckpointSaveContext.PendingCheckpointActors.Num() == 0)
 					{
 						CheckpointSaveContext.CheckpointSaveState = ECheckpointSaveState_SerializeDeletedStartupActors;
 					}
@@ -2029,11 +2049,14 @@ void UDemoNetDriver::TickCheckpoint()
 						// Save deleted startup actors
 						if (bDeltaCheckpoint)
 						{
-							*CheckpointArchive << CheckpointSaveContext.DeltaDeletedNetStartupActors;
-							CheckpointSaveContext.DeltaDeletedNetStartupActors.Empty();
+							*CheckpointArchive << CheckpointSaveContext.DeltaCheckpointData.DestroyedNetStartupActors;
+							CheckpointSaveContext.DeltaCheckpointData.DestroyedNetStartupActors.Empty();
 
-							*CheckpointArchive << CheckpointSaveContext.DeltaDeletedActorGuids;
-							CheckpointSaveContext.DeltaDeletedActorGuids.Empty();
+							*CheckpointArchive << CheckpointSaveContext.DeltaCheckpointData.DestroyedDynamicActors;
+							CheckpointSaveContext.DeltaCheckpointData.DestroyedDynamicActors.Empty();
+
+							*CheckpointArchive << CheckpointSaveContext.DeltaCheckpointData.ChannelsToClose;
+							CheckpointSaveContext.DeltaCheckpointData.ChannelsToClose.Empty();
 						}
 						else
 						{
@@ -2129,7 +2152,7 @@ void UDemoNetDriver::TickCheckpoint()
 	}
 
 	// accumulate time spent over all checkpoint ticks
-	CheckpointSaveContext.TotalCheckpointSaveTimeSeconds += ( CurrentTime - Params.StartCheckpointTime );
+	CheckpointSaveContext.TotalCheckpointSaveTimeSeconds += (CurrentTime - Params.StartCheckpointTime);
 
 	if (CheckpointSaveContext.CheckpointSaveState == ECheckpointSaveState_Finalize)
 	{
@@ -2138,15 +2161,15 @@ void UDemoNetDriver::TickCheckpoint()
 		// Get the total checkpoint size
 		const int32 TotalCheckpointSize = CheckpointArchive->TotalSize();
 
-		if ( CheckpointArchive->TotalSize() > 0 )
+		if (CheckpointArchive->TotalSize() > 0)
 		{
-			ReplayStreamer->FlushCheckpoint( GetLastCheckpointTimeInMS() );
+			ReplayStreamer->FlushCheckpoint(GetLastCheckpointTimeInMS());
 		}
 
 		const float TotalCheckpointTimeInMS = CheckpointSaveContext.TotalCheckpointReplicationTimeSeconds * 1000.0f;
 		const float TotalCheckpointTimeWithOverheadInMS = CheckpointSaveContext.TotalCheckpointSaveTimeSeconds * 1000.0f;
 
-		UE_LOG( LogDemo, Log, TEXT( "Finished checkpoint. Actors: %i, GuidCacheSize: %i, TotalSize: %i, TotalCheckpointSaveFrames: %i, TotalCheckpointTimeInMS: %2.2f, TotalCheckpointTimeWithOverheadInMS: %2.2f" ), GetNetworkObjectList().GetActiveObjects().Num(), CheckpointSaveContext.GuidCacheSize, TotalCheckpointSize, CheckpointSaveContext.TotalCheckpointSaveFrames, TotalCheckpointTimeInMS, TotalCheckpointTimeWithOverheadInMS);
+		UE_LOG(LogDemo, Log, TEXT("Finished checkpoint. Actors: %i, GuidCacheSize: %i, TotalSize: %i, TotalCheckpointSaveFrames: %i, TotalCheckpointTimeInMS: %2.2f, TotalCheckpointTimeWithOverheadInMS: %2.2f"), GetNetworkObjectList().GetActiveObjects().Num(), CheckpointSaveContext.GuidCacheSize, TotalCheckpointSize, CheckpointSaveContext.TotalCheckpointSaveFrames, TotalCheckpointTimeInMS, TotalCheckpointTimeWithOverheadInMS);
 
 		// we are done, out
 		CheckpointSaveContext.CheckpointSaveState = ECheckpointSaveState_Idle;
@@ -2157,7 +2180,9 @@ bool UDemoNetDriver::ReplicateCheckpointActor(AActor* ToReplicate, UDemoNetConne
 {
 	// Early out if the actor has been destroyed or the world is streamed out.
 	if (ToReplicate == nullptr || ToReplicate->GetWorld() == nullptr)
+	{
 		return true;
+	}
 
 	if (UActorChannel* ActorChannel = ClientConnection->FindActorChannelRef(ToReplicate))
 	{
@@ -2216,14 +2241,14 @@ void UDemoNetDriver::LoadExternalData( FArchive& Ar, const float TimeSeconds )
 {
 	DECLARE_SCOPE_CYCLE_COUNTER(TEXT("Demo_LoadExternalData"), Demo_LoadExternalData, STATGROUP_Net);
 
-	while ( true )
+	while (true)
 	{
 		uint32 ExternalDataNumBits;
 
 		// Read payload into payload/guid map
-		Ar.SerializeIntPacked( ExternalDataNumBits );
+		Ar.SerializeIntPacked(ExternalDataNumBits);
 
-		if ( ExternalDataNumBits == 0 )
+		if (ExternalDataNumBits == 0)
 		{
 			return;
 		}
@@ -2233,18 +2258,18 @@ void UDemoNetDriver::LoadExternalData( FArchive& Ar, const float TimeSeconds )
 		// Read net guid this payload belongs to
 		Ar << NetGUID;
 
-		int32 ExternalDataNumBytes = ( ExternalDataNumBits + 7 ) >> 3;
+		int32 ExternalDataNumBytes = (ExternalDataNumBits + 7) >> 3;
 
 		FBitReader Reader(nullptr, ExternalDataNumBits);
 
 		Ar.Serialize(Reader.GetData(), ExternalDataNumBytes);
 
-		Reader.SetEngineNetVer( ServerConnection->EngineNetworkProtocolVersion );
-		Reader.SetGameNetVer( ServerConnection->GameNetworkProtocolVersion );
+		Reader.SetEngineNetVer(ServerConnection->EngineNetworkProtocolVersion);
+		Reader.SetGameNetVer(ServerConnection->GameNetworkProtocolVersion);
 
-		FReplayExternalDataArray& ExternalDataArray = ExternalDataToObjectMap.FindOrAdd( NetGUID );
+		FReplayExternalDataArray& ExternalDataArray = ExternalDataToObjectMap.FindOrAdd(NetGUID);
 
-		ExternalDataArray.Add( new FReplayExternalData( MoveTemp(Reader), TimeSeconds ) );
+		ExternalDataArray.Add(new FReplayExternalData(MoveTemp(Reader), TimeSeconds));
 	}
 }
 
@@ -2319,9 +2344,9 @@ void UDemoNetDriver::RequestEventDataForActiveReplay(const FString& EventID, con
 class FReplayViewer
 {
 public:
-	FReplayViewer( const UNetConnection* Connection ) :
-		Viewer( Connection->PlayerController ? Connection->PlayerController : Connection->OwningActor ), 
-		ViewTarget( Connection->PlayerController ? Connection->PlayerController->GetViewTarget() : Connection->OwningActor )
+	FReplayViewer(const UNetConnection* Connection) :
+		Viewer(Connection->PlayerController ? Connection->PlayerController : Connection->OwningActor), 
+		ViewTarget(Connection->PlayerController ? Connection->PlayerController->GetViewTarget() : Connection->OwningActor)
 	{
 		Location = ViewTarget ? ViewTarget->GetActorLocation() : FVector::ZeroVector;
 	}
@@ -2361,9 +2386,9 @@ public:
 	const double TimeLimitSeconds;
 };
 
-void UDemoNetDriver::TickDemoRecord( float DeltaSeconds )
+void UDemoNetDriver::TickDemoRecord(float DeltaSeconds)
 {
-	if ( !IsRecording() || bPauseRecording )
+	if (!IsRecording() || bPauseRecording)
 	{
 		return;
 	}
@@ -2371,11 +2396,11 @@ void UDemoNetDriver::TickDemoRecord( float DeltaSeconds )
 	CSV_SCOPED_TIMING_STAT(Basic, DemoRecordTime);
 
 	// DeltaSeconds that is padded in, is unclampped and not time dilated
-	DemoCurrentTime += GetClampedDeltaSeconds( World, DeltaSeconds );
+	DemoCurrentTime += GetClampedDeltaSeconds(World, DeltaSeconds);
 
-	ReplayStreamer->UpdateTotalDemoTime( GetDemoCurrentTimeInMS() );
+	ReplayStreamer->UpdateTotalDemoTime(GetDemoCurrentTimeInMS());
 
-	if ( CheckpointSaveContext.CheckpointSaveState != ECheckpointSaveState_Idle )
+	if (CheckpointSaveContext.CheckpointSaveState != ECheckpointSaveState_Idle)
 	{
 		// If we're in the middle of saving a checkpoint, then update that now and return
 		TickCheckpoint();
@@ -2383,12 +2408,12 @@ void UDemoNetDriver::TickDemoRecord( float DeltaSeconds )
 	}
 	else
 	{
-		TickDemoRecordFrame( DeltaSeconds );
+		TickDemoRecordFrame(DeltaSeconds);
 
 		// Save a checkpoint if it's time
 		if (CVarEnableCheckpoints.GetValueOnAnyThread() == 1)
 		{
-			check( CheckpointSaveContext.CheckpointSaveState == ECheckpointSaveState_Idle );		// We early out above, so this shouldn't be possible
+			check(CheckpointSaveContext.CheckpointSaveState == ECheckpointSaveState_Idle);		// We early out above, so this shouldn't be possible
 
 			if (ShouldSaveCheckpoint())
 			{
@@ -2430,11 +2455,11 @@ void UDemoNetDriver::BuildSortedLevelPriorityOnLevels(const TArray<FDemoActorPri
 }
 
 
-void UDemoNetDriver::TickDemoRecordFrame( float DeltaSeconds )
+void UDemoNetDriver::TickDemoRecordFrame(float DeltaSeconds)
 {
 	FArchive* FileAr = ReplayStreamer->GetStreamingArchive();
 
-	if ( FileAr == nullptr )
+	if (FileAr == nullptr)
 	{
 		return;
 	}
@@ -2447,17 +2472,17 @@ void UDemoNetDriver::TickDemoRecordFrame( float DeltaSeconds )
 	{
 		for (ULevelStreaming* StreamingLevel : World->GetStreamingLevels())
 		{
-			if ( StreamingLevel == nullptr || !StreamingLevel->ShouldBeLoaded() || StreamingLevel->ShouldBeAlwaysLoaded() )
+			if (StreamingLevel == nullptr || !StreamingLevel->ShouldBeLoaded() || StreamingLevel->ShouldBeAlwaysLoaded())
 			{
 				continue;
 			}
 
 			TWeakObjectPtr<UObject> WeakStreamingLevel;
 			WeakStreamingLevel = StreamingLevel;
-			if ( !UniqueStreamingLevels.Contains( WeakStreamingLevel ) )
+			if (!UniqueStreamingLevels.Contains(WeakStreamingLevel))
 			{
-				UniqueStreamingLevels.Add( WeakStreamingLevel );
-				NewStreamingLevelsThisFrame.Add( WeakStreamingLevel );
+				UniqueStreamingLevels.Add(WeakStreamingLevel);
+				NewStreamingLevelsThisFrame.Add(WeakStreamingLevel);
 			}
 		}
 	}
@@ -2472,7 +2497,7 @@ void UDemoNetDriver::TickDemoRecordFrame( float DeltaSeconds )
 	FlushNetChecked(*ClientConnection);
 
 	float ServerTickTime = GEngine->GetMaxTickRate( DeltaSeconds );
-	if ( ServerTickTime == 0.0 )
+	if (ServerTickTime == 0.0)
 	{
 		ServerTickTime = DeltaSeconds;
 	}
@@ -2676,7 +2701,7 @@ void UDemoNetDriver::TickDemoRecordFrame( float DeltaSeconds )
 		}
 	}
 
-	if ( HasLevelStreamingFixes() )
+	if (HasLevelStreamingFixes())
 	{
 		SCOPED_NAMED_EVENT(UDemoNetDriver_PrioritizeLevelSort, FColor::Green);
 		DECLARE_SCOPE_CYCLE_COUNTER(TEXT("Replay actor level sorting time."), STAT_ReplayLevelSorting, STATGROUP_Net);
@@ -2694,7 +2719,7 @@ void UDemoNetDriver::TickDemoRecordFrame( float DeltaSeconds )
 		// Intervals are then used to determine the order we write out the replicated objects as we write one packet per level.
 		BuildSortedLevelPriorityOnLevels(PrioritizedActors, LevelIntervals);
 	}
-	else if ( bPrioritizeActors )
+	else if (bPrioritizeActors)
 	{
 		// Sort on priority
 		PrioritizedActors.Sort([](const FDemoActorPriority& A, const FDemoActorPriority& B) { return B.ActorPriority.Priority < A.ActorPriority.Priority; });
@@ -2907,7 +2932,7 @@ bool UDemoNetDriver::ReplicatePrioritizedActors(const FDemoActorPriority* Actors
 	return bTimeRemaining;
 }
 
-bool UDemoNetDriver::ShouldSaveCheckpoint()
+bool UDemoNetDriver::ShouldSaveCheckpoint() const
 {
 	const double CHECKPOINT_DELAY = CVarCheckpointUploadDelayInSeconds.GetValueOnAnyThread();
 
@@ -2919,29 +2944,29 @@ bool UDemoNetDriver::ShouldSaveCheckpoint()
 	return false;
 }
 
-void UDemoNetDriver::PauseChannels( const bool bPause )
+void UDemoNetDriver::PauseChannels(const bool bPause)
 {
-	if ( bPause == bChannelsArePaused )
+	if (bPause == bChannelsArePaused)
 	{
 		return;
 	}
 
 	// Pause all non player controller actors
 	// FIXME: Would love a more elegant way of handling this at a more global level
-	for ( int32 i = ServerConnection->OpenChannels.Num() - 1; i >= 0; i-- )
+	for (int32 i = ServerConnection->OpenChannels.Num() - 1; i >= 0; i--)
 	{
 		UChannel* OpenChannel = ServerConnection->OpenChannels[i];
 
-		UActorChannel* ActorChannel = Cast< UActorChannel >( OpenChannel );
+		UActorChannel* ActorChannel = Cast<UActorChannel>(OpenChannel);
 
-		if ( ActorChannel == nullptr )
+		if (ActorChannel == nullptr)
 		{
 			continue;
 		}
 
 		ActorChannel->CustomTimeDilation = bPause ? 0.0f : 1.0f;
 
-		if (ActorChannel->GetActor() == nullptr || SpectatorControllers.Contains(ActorChannel->GetActor()) )
+		if (ActorChannel->GetActor() == nullptr || SpectatorControllers.Contains(ActorChannel->GetActor()))
 		{
 			continue;
 		}
@@ -2953,7 +2978,7 @@ void UDemoNetDriver::PauseChannels( const bool bPause )
 	bChannelsArePaused = bPause;
 }
 
-bool UDemoNetDriver::ReadDemoFrameIntoPlaybackPackets( FArchive& Ar, TArray<FPlaybackPacket>& InPlaybackPackets, const bool bForLevelFastForward, float* OutTime )
+bool UDemoNetDriver::ReadDemoFrameIntoPlaybackPackets(FArchive& Ar, TArray<FPlaybackPacket>& InPlaybackPackets, const bool bForLevelFastForward, float* OutTime)
 {
 	SCOPED_NAMED_EVENT(UDemoNetDriver_ReadDemoFrameIntoPlaybackPackets, FColor::Purple);
 
@@ -2961,19 +2986,19 @@ bool UDemoNetDriver::ReadDemoFrameIntoPlaybackPackets( FArchive& Ar, TArray<FPla
 
 	check(!bForLevelFastForward || HasLevelStreamingFixes());
 
-	if ( Ar.IsError() )
+	if (Ar.IsError())
 	{
 		UE_LOG(LogDemo, Error, TEXT("UDemoNetDriver::ReadDemoFrameIntoPlaybackPackets: Archive Error"));
 		NotifyDemoPlaybackFailure(EDemoPlayFailure::Serialization);
 		return false;
 	}
 
-	if ( Ar.AtEnd() )
+	if (Ar.AtEnd())
 	{
 		return false;
 	}
 
-	if ( ReplayStreamer->GetLastError() != ENetworkReplayError::None )
+	if (ReplayStreamer->GetLastError() != ENetworkReplayError::None)
 	{
 		UE_LOG(LogDemo, Error, TEXT("UDemoNetDriver::ReadDemoFrameIntoPlaybackPackets: ReplayStreamer ERROR: %s"), ENetworkReplayError::ToString(ReplayStreamer->GetLastError()));
 		NotifyDemoPlaybackFailure(EDemoPlayFailure::ReplayStreamerInternal);
@@ -2983,7 +3008,7 @@ bool UDemoNetDriver::ReadDemoFrameIntoPlaybackPackets( FArchive& Ar, TArray<FPla
 	// Above checks guarantee the Archive is in a valid state, but it's entirely possible that
 	// the ReplayStreamer doesn't have more stream data available (i.e., if we only have checkpoint data).
 	// Therefore, skip this if we know we're only reading in checkpoint data.
-	if ( !bIsLoadingCheckpoint && !ReplayStreamer->IsDataAvailable() )
+	if (!bIsLoadingCheckpoint && !ReplayStreamer->IsDataAvailable())
 	{
 		return false;
 	}
@@ -3265,20 +3290,20 @@ void UDemoNetDriver::OnDownloadHeaderCompletePrivate(const FDownloadHeaderResult
 	}
 }
 
-bool UDemoNetDriver::ConditionallyReadDemoFrameIntoPlaybackPackets( FArchive& Ar )
+bool UDemoNetDriver::ConditionallyReadDemoFrameIntoPlaybackPackets(FArchive& Ar)
 {
-	if ( PlaybackPackets.Num() > 0 )
+	if (PlaybackPackets.Num() > 0)
 	{
 		const float MAX_PLAYBACK_BUFFER_SECONDS = 5.0f;
 
 		const FPlaybackPacket& LastPacket = PlaybackPackets.Last();
-		if ( LastPacket.TimeSeconds > DemoCurrentTime && LastPacket.TimeSeconds - DemoCurrentTime > MAX_PLAYBACK_BUFFER_SECONDS )
+		if (LastPacket.TimeSeconds > DemoCurrentTime && LastPacket.TimeSeconds - DemoCurrentTime > MAX_PLAYBACK_BUFFER_SECONDS)
 		{
 			return false;	// Don't buffer more than MAX_PLAYBACK_BUFFER_SECONDS worth of frames
 		}
 	}
 
-	if ( !ReadDemoFrameIntoPlaybackPackets( Ar ) )
+	if (!ReadDemoFrameIntoPlaybackPackets(Ar))
 	{
 		return false;
 	}
@@ -3287,35 +3312,35 @@ bool UDemoNetDriver::ConditionallyReadDemoFrameIntoPlaybackPackets( FArchive& Ar
 }
 
 // Deprecated, DO NOT USE.
-bool UDemoNetDriver::ReadPacket( FArchive& Archive, uint8* OutReadBuffer, int32& OutBufferSize, const int32 MaxBufferSize )
+bool UDemoNetDriver::ReadPacket(FArchive& Archive, uint8* OutReadBuffer, int32& OutBufferSize, const int32 MaxBufferSize)
 {
 	OutBufferSize = 0;
 
 	Archive << OutBufferSize;
 
-	if ( Archive.IsError() )
+	if (Archive.IsError())
 	{
-		UE_LOG( LogDemo, Error, TEXT( "UDemoNetDriver::ReadPacket: Failed to read demo OutBufferSize" ) );
+		UE_LOG(LogDemo, Error, TEXT("UDemoNetDriver::ReadPacket: Failed to read demo OutBufferSize"));
 		return false;
 	}
 
-	if ( OutBufferSize == 0 )
+	if (OutBufferSize == 0)
 	{
 		return true;		// Done
 	}
 
-	if ( OutBufferSize > MaxBufferSize )
+	if (OutBufferSize > MaxBufferSize)
 	{
-		UE_LOG( LogDemo, Error, TEXT( "UDemoNetDriver::ReadPacket: OutBufferSize > sizeof( ReadBuffer )" ) );
+		UE_LOG(LogDemo, Error, TEXT("UDemoNetDriver::ReadPacket: OutBufferSize > sizeof( ReadBuffer )"));
 		return false;
 	}
 
 	// Read data from file.
-	Archive.Serialize( OutReadBuffer, OutBufferSize );
+	Archive.Serialize(OutReadBuffer, OutBufferSize);
 
-	if ( Archive.IsError() )
+	if (Archive.IsError())
 	{
-		UE_LOG( LogDemo, Error, TEXT( "UDemoNetDriver::ReadPacket: Failed to read demo file packet" ) );
+		UE_LOG(LogDemo, Error, TEXT("UDemoNetDriver::ReadPacket: Failed to read demo file packet"));
 		return false;
 	}
 
@@ -3324,11 +3349,11 @@ bool UDemoNetDriver::ReadPacket( FArchive& Archive, uint8* OutReadBuffer, int32&
 		uint32 ServerChecksum = 0;
 		Archive << ServerChecksum;
 
-		const uint32 Checksum = FCrc::MemCrc32( OutReadBuffer, OutBufferSize, 0 );
+		const uint32 Checksum = FCrc::MemCrc32(OutReadBuffer, OutBufferSize, 0);
 
-		if ( Checksum != ServerChecksum )
+		if (Checksum != ServerChecksum)
 		{
-			UE_LOG( LogDemo, Error, TEXT( "UDemoNetDriver::ReadPacket: Checksum != ServerChecksum" ) );
+			UE_LOG(LogDemo, Error, TEXT("UDemoNetDriver::ReadPacket: Checksum != ServerChecksum"));
 			return false;
 		}
 	}
@@ -3430,14 +3455,14 @@ bool UDemoNetDriver::ShouldSkipPlaybackPacket(const FPlaybackPacket& Packet)
 
 bool UDemoNetDriver::ConditionallyProcessPlaybackPackets()
 {
-	if ( !PlaybackPackets.IsValidIndex(PlaybackPacketIndex) )
+	if (!PlaybackPackets.IsValidIndex(PlaybackPacketIndex))
 	{
 		PauseChannels( true );
 		return false;
 	}
 
 	const FPlaybackPacket& CurPacket = PlaybackPackets[PlaybackPacketIndex];
-	if ( DemoCurrentTime < CurPacket.TimeSeconds )
+	if (DemoCurrentTime < CurPacket.TimeSeconds)
 	{
 		// Not enough time has passed to read another frame
 		return false;
@@ -3457,33 +3482,38 @@ bool UDemoNetDriver::ConditionallyProcessPlaybackPackets()
 
 void UDemoNetDriver::ProcessAllPlaybackPackets()
 {
-	if (PlaybackPackets.Num() > 0)
+	ProcessPlaybackPackets(PlaybackPackets);
+	PlaybackPackets.Empty();
+}
+
+void UDemoNetDriver::ProcessPlaybackPackets(TArrayView<FPlaybackPacket> Packets)
+{
+	if (Packets.Num() > 0)
 	{
-		for (const FPlaybackPacket& PlaybackPacket : PlaybackPackets)
+		for (const FPlaybackPacket& PlaybackPacket : Packets)
 		{
 			ProcessPacket(PlaybackPacket);
 		}
 
-		LastProcessedPacketTime = PlaybackPackets.Last().TimeSeconds;
-		PlaybackPackets.Empty();
+		LastProcessedPacketTime = Packets.Last().TimeSeconds;
 	}
 }
 
-bool UDemoNetDriver::ProcessPacket( const uint8* Data, int32 Count )
+bool UDemoNetDriver::ProcessPacket(const uint8* Data, int32 Count)
 {
-	PauseChannels( false );
+	PauseChannels(false);
 
-	if ( ServerConnection != nullptr )
+	if (ServerConnection != nullptr )
 	{
 		// Process incoming packet.
 		// ReceivedRawPacket shouldn't change any data, so const_cast should be safe.
-		ServerConnection->ReceivedRawPacket( const_cast<uint8*>(Data), Count );
+		ServerConnection->ReceivedRawPacket(const_cast<uint8*>(Data), Count);
 	}
 
-	if ( ServerConnection == nullptr || ServerConnection->State == USOCK_Closed )
+	if (ServerConnection == nullptr || ServerConnection->State == USOCK_Closed)
 	{
 		// Something we received resulted in the demo being stopped
-		UE_LOG( LogDemo, Error, TEXT( "UDemoNetDriver::ProcessPacket: ReceivedRawPacket closed connection" ) );
+		UE_LOG(LogDemo, Error, TEXT("UDemoNetDriver::ProcessPacket: ReceivedRawPacket closed connection"));
 		NotifyDemoPlaybackFailure(EDemoPlayFailure::Generic);
 		return false;
 	}
@@ -3491,14 +3521,14 @@ bool UDemoNetDriver::ProcessPacket( const uint8* Data, int32 Count )
 	return true;
 }
 
-void UDemoNetDriver::WriteDemoFrameFromQueuedDemoPackets( FArchive& Ar, TArray<FQueuedDemoPacket>& QueuedPackets, float FrameTime )
+void UDemoNetDriver::WriteDemoFrameFromQueuedDemoPackets(FArchive& Ar, TArray<FQueuedDemoPacket>& QueuedPackets, float FrameTime)
 {
 	Ar << CurrentLevelIndex;
 	
 	// Save total absolute demo time in seconds
 	Ar << FrameTime;
 
-	((UPackageMapClient*)ClientConnections[0]->PackageMap)->AppendExportData(Ar);
+	Cast<UPackageMapClient>(ClientConnections[0]->PackageMap)->AppendExportData(Ar);
 
 	if (HasLevelStreamingFixes())
 	{
@@ -3536,7 +3566,6 @@ void UDemoNetDriver::WriteDemoFrameFromQueuedDemoPackets( FArchive& Ar, TArray<F
 		NewStreamingLevelsThisFrame.Empty();
 	}
 
-
 	{
 		TUniquePtr<FScopedStoreArchiveOffset> ScopedOffset(HasLevelStreamingFixes() ? new FScopedStoreArchiveOffset(Ar) : nullptr);
 
@@ -3567,36 +3596,36 @@ void UDemoNetDriver::WriteDemoFrameFromQueuedDemoPackets( FArchive& Ar, TArray<F
 	Ar << EndCount;
 }
 
-void UDemoNetDriver::WritePacket( FArchive& Ar, uint8* Data, int32 Count )
+void UDemoNetDriver::WritePacket(FArchive& Ar, uint8* Data, int32 Count)
 {
 	Ar << Count;
-	Ar.Serialize( Data, Count );
+	Ar.Serialize(Data, Count);
 
 #if DEMO_CHECKSUMS == 1
-	uint32 Checksum = FCrc::MemCrc32( Data, Count, 0 );
+	uint32 Checksum = FCrc::MemCrc32(Data, Count, 0);
 	Ar << Checksum;
 #endif
 }
 
 void UDemoNetDriver::SkipTime(const float InTimeToSkip)
 {
-	if ( IsNamedTaskInQueue( ReplayTaskNames::SkipTimeInSecondsTask ) )
+	if (IsNamedTaskInQueue(ReplayTaskNames::SkipTimeInSecondsTask))
 	{
 		return;		// Don't allow time skipping if we already are
 	}
 
-	AddReplayTask( new FSkipTimeInSecondsTask( this, InTimeToSkip ) );
+	AddReplayTask(new FSkipTimeInSecondsTask(this, InTimeToSkip));
 }
 
-void UDemoNetDriver::SkipTimeInternal( const float SecondsToSkip, const bool InFastForward, const bool InIsForCheckpoint )
+void UDemoNetDriver::SkipTimeInternal(const float SecondsToSkip, const bool InFastForward, const bool InIsForCheckpoint)
 {
-	check( !bIsFastForwarding );				// Can only do one of these at a time (use tasks to gate this)
-	check( !bIsFastForwardingForCheckpoint );	// Can only do one of these at a time (use tasks to gate this)
+	check(!bIsFastForwarding);				// Can only do one of these at a time (use tasks to gate this)
+	check(!bIsFastForwardingForCheckpoint);	// Can only do one of these at a time (use tasks to gate this)
 
 	SavedSecondsToSkip = SecondsToSkip;
 	DemoCurrentTime += SecondsToSkip;
 
-	DemoCurrentTime = FMath::Clamp( DemoCurrentTime, 0.0f, DemoTotalTime - 0.01f );
+	DemoCurrentTime = FMath::Clamp(DemoCurrentTime, 0.0f, DemoTotalTime - 0.01f);
 
 	bIsFastForwarding				= InFastForward;
 	bIsFastForwardingForCheckpoint	= InIsForCheckpoint;
@@ -3606,7 +3635,7 @@ void UDemoNetDriver::GotoTimeInSeconds(const float TimeInSeconds, const FOnGotoT
 {
 	OnGotoTimeDelegate_Transient = InOnGotoTimeDelegate;
 
-	if ( IsNamedTaskInQueue( ReplayTaskNames::GotoTimeInSecondsTask ) || bIsFastForwarding )
+	if (IsNamedTaskInQueue(ReplayTaskNames::GotoTimeInSecondsTask) || bIsFastForwarding)
 	{
 		NotifyGotoTimeFinished(false);
 		return;		// Don't allow scrubbing if we already are
@@ -3614,12 +3643,12 @@ void UDemoNetDriver::GotoTimeInSeconds(const float TimeInSeconds, const FOnGotoT
 
 	UE_LOG(LogDemo, Log, TEXT("GotoTimeInSeconds: %2.2f"), TimeInSeconds);
 
-	AddReplayTask( new FGotoTimeInSecondsTask( this, TimeInSeconds ) );
+	AddReplayTask(new FGotoTimeInSecondsTask(this, TimeInSeconds));
 }
 
 void UDemoNetDriver::JumpToEndOfLiveReplay()
 {
-	UE_LOG( LogDemo, Log, TEXT( "UDemoNetConnection::JumpToEndOfLiveReplay." ) );
+	UE_LOG(LogDemo, Log, TEXT("UDemoNetConnection::JumpToEndOfLiveReplay."));
 
 	const uint32 TotalDemoTimeInMS = ReplayStreamer->GetTotalDemoTime();
 
@@ -3627,19 +3656,19 @@ void UDemoNetDriver::JumpToEndOfLiveReplay()
 
 	const uint32 BufferInMS = 5 * 1000;
 
-	const uint32 JoinTimeInMS = FMath::Max( (uint32)0, ReplayStreamer->GetTotalDemoTime() - BufferInMS );
+	const uint32 JoinTimeInMS = FMath::Max((uint32)0, ReplayStreamer->GetTotalDemoTime() - BufferInMS);
 
-	if ( JoinTimeInMS > 0 )
+	if (JoinTimeInMS > 0)
 	{
-		GotoTimeInSeconds( (float)JoinTimeInMS / 1000.0f );
+		GotoTimeInSeconds((float)JoinTimeInMS / 1000.0f);
 	}
 }
 
-void UDemoNetDriver::AddUserToReplay( const FString& UserString )
+void UDemoNetDriver::AddUserToReplay(const FString& UserString)
 {
-	if ( ReplayStreamer.IsValid() )
+	if (ReplayStreamer.IsValid())
 	{
-		ReplayStreamer->AddUserToReplay( UserString );
+		ReplayStreamer->AddUserToReplay(UserString);
 	}
 }
 
@@ -3659,19 +3688,24 @@ static FCsvDemoSettings GetCsvDemoSettings()
 	if (Settings.bCaptureCsv)
 	{
 		if (!FParse::Value(FCommandLine::Get(), TEXT("-csvdemoendtime="), Settings.EndTime))
+		{
 			Settings.EndTime = -1.0;
+		}
+
 		if (!FParse::Value(FCommandLine::Get(), TEXT("-csvdemoframecount="), Settings.FrameCount))
+		{
 			Settings.FrameCount = -1;
+		}
 	}
 
 	return Settings;
 }
 #endif // (CSV_PROFILER && (!UE_BUILD_SHIPPING))
 
-void UDemoNetDriver::TickDemoPlayback( float DeltaSeconds )
+void UDemoNetDriver::TickDemoPlayback(float DeltaSeconds)
 {
 	SCOPED_NAMED_EVENT(UDemoNetDriver_TickDemoPlayback, FColor::Purple);
-	if ( World && World->IsInSeamlessTravel() )
+	if (World && World->IsInSeamlessTravel())
 	{
 		return;
 	}
@@ -3700,31 +3734,31 @@ void UDemoNetDriver::TickDemoPlayback( float DeltaSeconds )
 	}
 #endif // (CSV_PROFILER && (!UE_BUILD_SHIPPING))
 
-	if ( !IsPlaying() )
+	if (!IsPlaying())
 	{
 		return;
 	}
 	
 	// This will be true when watching a live replay and we're grabbing an up to date header.
 	// In that case, we want to pause playback until we can actually travel.
-	if ( bIsWaitingForHeaderDownload )
+	if (bIsWaitingForHeaderDownload)
 	{
 		return;
 	}
 
-	if ( CVarForceDisableAsyncPackageMapLoading.GetValueOnGameThread() > 0 )
+	if (CVarForceDisableAsyncPackageMapLoading.GetValueOnGameThread() > 0)
 	{
-		GuidCache->SetAsyncLoadMode( FNetGUIDCache::EAsyncLoadMode::ForceDisable );
+		GuidCache->SetAsyncLoadMode(FNetGUIDCache::EAsyncLoadMode::ForceDisable);
 	}
 	else
 	{
-		GuidCache->SetAsyncLoadMode( FNetGUIDCache::EAsyncLoadMode::UseCVar );
+		GuidCache->SetAsyncLoadMode(FNetGUIDCache::EAsyncLoadMode::UseCVar);
 	}
 
-	if ( CVarGotoTimeInSeconds.GetValueOnGameThread() >= 0.0f )
+	if (CVarGotoTimeInSeconds.GetValueOnGameThread() >= 0.0f)
 	{
-		GotoTimeInSeconds( CVarGotoTimeInSeconds.GetValueOnGameThread() );
-		CVarGotoTimeInSeconds.AsVariable()->Set( TEXT( "-1" ), ECVF_SetByConsole );
+		GotoTimeInSeconds(CVarGotoTimeInSeconds.GetValueOnGameThread());
+		CVarGotoTimeInSeconds.AsVariable()->Set(TEXT( "-1" ), ECVF_SetByConsole);
 	}
 
 	if (FMath::Abs(CVarDemoSkipTime.GetValueOnGameThread()) > 0.0f)
@@ -3739,12 +3773,12 @@ void UDemoNetDriver::TickDemoPlayback( float DeltaSeconds )
 	PrepFastForwardLevels();
 
 	// Update total demo time
-	if ( ReplayStreamer->GetTotalDemoTime() > 0 )
+	if (ReplayStreamer->GetTotalDemoTime() > 0)
 	{
-		DemoTotalTime = ( float )ReplayStreamer->GetTotalDemoTime() / 1000.0f;
+		DemoTotalTime = (float)ReplayStreamer->GetTotalDemoTime() / 1000.0f;
 	}
 
-	if ( !ProcessReplayTasks() )
+	if (!ProcessReplayTasks())
 	{
 		// We're busy processing tasks, return
 		return;
@@ -3782,13 +3816,13 @@ void UDemoNetDriver::TickDemoPlayback( float DeltaSeconds )
 	}
 
 	// Advance demo time by seconds passed if we're not paused
-	if ( World && World->GetWorldSettings() && World->GetWorldSettings()->GetPauserPlayerState() == nullptr )
+	if (World && World->GetWorldSettings() && World->GetWorldSettings()->GetPauserPlayerState() == nullptr)
 	{
 		DemoCurrentTime += DeltaSeconds;
 	}
 
 	// Clamp time
-	DemoCurrentTime = FMath::Clamp( DemoCurrentTime, 0.0f, DemoTotalTime - 0.01f );
+	DemoCurrentTime = FMath::Clamp(DemoCurrentTime, 0.0f, DemoTotalTime - 0.01f);
 
 	// Make sure there is data available to read
 	// If we're at the end of the demo, just pause channels and return
@@ -3808,18 +3842,18 @@ void UDemoNetDriver::TickDemoPlayback( float DeltaSeconds )
 	}
 
 	// Buffer up demo frames until we have enough time built-up
-	while ( ConditionallyReadDemoFrameIntoPlaybackPackets( *ReplayStreamer->GetStreamingArchive() ) )
+	while (ConditionallyReadDemoFrameIntoPlaybackPackets(*ReplayStreamer->GetStreamingArchive()))
 	{
 	}
 
 	{
 		DECLARE_SCOPE_CYCLE_COUNTER(TEXT("TickDemoPlayback_ProcessPackets"), TickDemoPlayback_ProcessPackets, STATGROUP_Net);
 
-	// Process packets until we are caught up (this implicitly handles fast forward if DemoCurrentTime past many frames)
-	while ( ConditionallyProcessPlaybackPackets() )
-	{
-		DemoFrameNum++;
-	}
+		// Process packets until we are caught up (this implicitly handles fast forward if DemoCurrentTime past many frames)
+		while (ConditionallyProcessPlaybackPackets())
+		{
+			DemoFrameNum++;
+		}
 
 		if (PlaybackPacketIndex > 0)
 		{
@@ -3834,13 +3868,13 @@ void UDemoNetDriver::TickDemoPlayback( float DeltaSeconds )
 	}
 
 	// Finalize any fast forward stuff that needs to happen
-	if ( bIsFastForwarding )
+	if (bIsFastForwarding)
 	{
-		FinalizeFastForward( FastForwardStartSeconds );
+		FinalizeFastForward(FastForwardStartSeconds);
 	}
 }
 
-void UDemoNetDriver::FinalizeFastForward( const double StartTime )
+void UDemoNetDriver::FinalizeFastForward(const double StartTime)
 {
 	DECLARE_SCOPE_CYCLE_COUNTER(TEXT("Demo_FinalizeFastForward"), Demo_FinalizeFastForward, STATGROUP_Net);
 
@@ -3868,7 +3902,7 @@ void UDemoNetDriver::FinalizeFastForward( const double StartTime )
 		GameState->OnRep_ReplicatedWorldTimeSeconds();
 	}
 
-	if ( ServerConnection != nullptr && bIsFastForwardingForCheckpoint )
+	if (ServerConnection != nullptr && bIsFastForwardingForCheckpoint)
 	{
 		// Make a pass at OnReps for startup actors, since they were skipped during checkpoint loading.
 		// At this point the shadow state of these actors should be the actual state from before the checkpoint,
@@ -3902,22 +3936,22 @@ void UDemoNetDriver::FinalizeFastForward( const double StartTime )
 	}
 
 	// Flush all pending RepNotifies that were built up during the fast-forward.
-	if ( ServerConnection != nullptr )
+	if (ServerConnection != nullptr)
 	{
-		for ( auto& ChannelPair : ServerConnection->ActorChannelMap() )
+		for (auto& ChannelPair : ServerConnection->ActorChannelMap())
 		{
-			if ( ChannelPair.Value != nullptr )
+			if (ChannelPair.Value != nullptr)
 			{
-				for ( auto& ReplicatorPair : ChannelPair.Value->ReplicationMap )
+				for (auto& ReplicatorPair : ChannelPair.Value->ReplicationMap)
 				{
-					ReplicatorPair.Value->CallRepNotifies( true );
+					ReplicatorPair.Value->CallRepNotifies(true);
 				}
 			}
 		}
 
 		for (auto& DormantPair : ServerConnection->DormantReplicatorMap)
 		{
-			DormantPair.Value->CallRepNotifies( true );
+			DormantPair.Value->CallRepNotifies(true);
 		}
 	}
 
@@ -3932,7 +3966,7 @@ void UDemoNetDriver::FinalizeFastForward( const double StartTime )
 
 	NotifyGotoTimeFinished(true);
 
-	UE_LOG( LogDemo, Log, TEXT( "Fast forward took %.2f seconds." ), FastForwardTotalSeconds );
+	UE_LOG(LogDemo, Log, TEXT("Fast forward took %.2f seconds."), FastForwardTotalSeconds);
 }
 
 APlayerController* UDemoNetDriver::CreateDemoPlayerController(UNetConnection* Connection, const FURL& ListenURL)
@@ -3959,7 +3993,7 @@ APlayerController* UDemoNetDriver::CreateDemoPlayerController(UNetConnection* Co
 	AGameModeBase* DefaultGameMode = DefaultGameModeClass.GetDefaultObject();
 	UClass* C = DefaultGameMode != nullptr ? DefaultGameMode->ReplaySpectatorPlayerControllerClass : nullptr;
 
-	if ( C == nullptr )
+	if (C == nullptr)
 	{
 		UE_LOG(LogDemo, Error, TEXT("UDemoNetDriver::CreateDemoPlayerController: Failed to load demo spectator class."));
 		return nullptr;
@@ -4011,7 +4045,7 @@ APlayerController* UDemoNetDriver::CreateDemoPlayerController(UNetConnection* Co
 	return NewDemoController;
 }
 
-void UDemoNetDriver::SpawnDemoRecSpectator( UNetConnection* Connection, const FURL& ListenURL )
+void UDemoNetDriver::SpawnDemoRecSpectator(UNetConnection* Connection, const FURL& ListenURL)
 {
 	// Optionally skip spawning the demo spectator if requested via the URL option
 	if (ListenURL.HasOption(TEXT("SkipSpawnSpectatorController")))
@@ -4119,19 +4153,19 @@ int32 UDemoNetDriver::CleanUpSplitscreenConnections(bool bDeleteOwner)
 	return NumSplitscreenConnectionsCleaned;
 }
 
-void UDemoNetDriver::ReplayStreamingReady( const FStartStreamingResult& Result )
+void UDemoNetDriver::ReplayStreamingReady(const FStartStreamingResult& Result)
 {
 	bIsWaitingForStream = false;
 	bWasStartStreamingSuccessful = Result.WasSuccessful();
 
 #if !(UE_BUILD_SHIPPING || UE_BUILD_TEST)
-	if ( CVarDemoForceFailure.GetValueOnGameThread() == 1 )
+	if (CVarDemoForceFailure.GetValueOnGameThread() == 1)
 	{
 		bWasStartStreamingSuccessful = false;
 	}
 #endif
 
-	if ( !bWasStartStreamingSuccessful )
+	if (!bWasStartStreamingSuccessful)
 	{
 		UE_LOG(LogDemo, Warning, TEXT("UDemoNetConnection::ReplayStreamingReady: Failed. %s"), Result.bRecording ? TEXT("") : EDemoPlayFailure::ToString(EDemoPlayFailure::DemoNotFound));
 
@@ -4146,15 +4180,13 @@ void UDemoNetDriver::ReplayStreamingReady( const FStartStreamingResult& Result )
 		return;
 	}
 
-
-
-	if ( !Result.bRecording )
+	if (!Result.bRecording)
 	{
 		FString Error;
 		
 		const double StartTime = FPlatformTime::Seconds();
 
-		if ( !InitConnectInternal( Error ) )
+		if (!InitConnectInternal(Error))
 		{
 			return;
 		}
@@ -4172,18 +4204,18 @@ void UDemoNetDriver::ReplayStreamingReady( const FStartStreamingResult& Result )
 			}
 		}
 
-		if ( ReplayStreamer->IsLive() && ReplayStreamer->GetTotalDemoTime() > 15 * 1000 )
+		if (ReplayStreamer->IsLive() && ReplayStreamer->GetTotalDemoTime() > 15 * 1000)
 		{
 			// If the load time wasn't very long, jump to end now
 			// Otherwise, defer it until we have a more recent replay time
-			if ( FPlatformTime::Seconds() - StartTime < 10 )
+			if (FPlatformTime::Seconds() - StartTime < 10)
 			{
 				JumpToEndOfLiveReplay();
 			}
 			else
 			{
-				UE_LOG( LogDemo, Log, TEXT( "UDemoNetConnection::ReplayStreamingReady: Deferring checkpoint until next available time." ) );
-				AddReplayTask( new FJumpToLiveReplayTask( this ) );
+				UE_LOG(LogDemo, Log, TEXT("UDemoNetConnection::ReplayStreamingReady: Deferring checkpoint until next available time."));
+				AddReplayTask(new FJumpToLiveReplayTask(this));
 			}
 		}
 
@@ -4195,23 +4227,23 @@ void UDemoNetDriver::ReplayStreamingReady( const FStartStreamingResult& Result )
 	}
 }
 
-FReplayExternalDataArray* UDemoNetDriver::GetExternalDataArrayForObject( UObject* Object )
+FReplayExternalDataArray* UDemoNetDriver::GetExternalDataArrayForObject(UObject* Object)
 {
-	FNetworkGUID NetworkGUID = GuidCache->NetGUIDLookup.FindRef( Object );
+	FNetworkGUID NetworkGUID = GuidCache->NetGUIDLookup.FindRef(Object);
 
-	if ( !NetworkGUID.IsValid() )
+	if (!NetworkGUID.IsValid())
 	{
 		return nullptr;
 	}
 
-	return ExternalDataToObjectMap.Find( NetworkGUID );
+	return ExternalDataToObjectMap.Find(NetworkGUID);
 }
 
 void UDemoNetDriver::RespawnNecessaryNetStartupActors(TArray<AActor*>& SpawnedActors, ULevel* Level /* = nullptr */)
 {
-	for ( auto It = RollbackNetStartupActors.CreateIterator(); It; ++It )
+	for (auto It = RollbackNetStartupActors.CreateIterator(); It; ++It)
 	{
-		if ( DeletedNetStartupActors.Contains( It.Key() ) )
+		if (DeletedNetStartupActors.Contains(It.Key()))
 		{
 			// We don't want to re-create these since they should no longer exist after the current checkpoint
 			continue;
@@ -4251,26 +4283,26 @@ void UDemoNetDriver::RespawnNecessaryNetStartupActors(TArray<AActor*>& SpawnedAc
 
 		FActorSpawnParameters SpawnInfo;
 
-		SpawnInfo.Template							= CastChecked<AActor>( RollbackActor.Archetype );
+		SpawnInfo.Template							= CastChecked<AActor>(RollbackActor.Archetype);
 		SpawnInfo.SpawnCollisionHandlingOverride	= ESpawnActorCollisionHandlingMethod::AlwaysSpawn;
 		SpawnInfo.bNoFail							= true;
 		SpawnInfo.Name								= RollbackActor.Name;
 		SpawnInfo.OverrideLevel						= RollbackActor.Level;
 		SpawnInfo.bDeferConstruction				= true;
 
-		const FTransform SpawnTransform = FTransform( RollbackActor.Rotation, RollbackActor.Location );
+		const FTransform SpawnTransform = FTransform(RollbackActor.Rotation, RollbackActor.Location);
 
-		AActor* Actor = GetWorld()->SpawnActorAbsolute( RollbackActor.Archetype->GetClass(), SpawnTransform, SpawnInfo );
+		AActor* Actor = GetWorld()->SpawnActorAbsolute(RollbackActor.Archetype->GetClass(), SpawnTransform, SpawnInfo);
 		if (Actor)
 		{
-			if ( !ensure( Actor->GetFullName() == It.Key() ) )
+			if (!ensure( Actor->GetFullName() == It.Key()))
 			{
-				UE_LOG( LogDemo, Log, TEXT( "RespawnNecessaryNetStartupActors: NetStartupRollbackActor name doesn't match original: %s, %s" ), *Actor->GetFullName(), *It.Key() );
+				UE_LOG(LogDemo, Log, TEXT("RespawnNecessaryNetStartupActors: NetStartupRollbackActor name doesn't match original: %s, %s"), *Actor->GetFullName(), *It.Key());
 			}
 
 			bool bSanityCheckReferences = true;
 
-			for(UObject* ObjRef : RollbackActor.ObjReferences)
+			for (UObject* ObjRef : RollbackActor.ObjReferences)
 			{
 				if (ObjRef == nullptr)
 				{
@@ -4344,7 +4376,7 @@ void UDemoNetDriver::RespawnNecessaryNetStartupActors(TArray<AActor*>& SpawnedAc
 				}
 			}
 
-			check( Actor->GetRemoteRole() == ROLE_Authority );
+			check(Actor->GetRemoteRole() == ROLE_Authority);
 
 			SpawnedActors.Add(Actor);
 		}
@@ -4376,21 +4408,7 @@ void UDemoNetDriver::PrepFastForwardLevels()
 				continue;
 			}
 
-			bool bStartupActors = false;
-
-			for (AActor* Actor : Level->Actors)
-			{
-				if (Actor != nullptr && Actor->IsNetStartupActor())
-				{
-					bStartupActors = true;
-					break;
-				}
-			}
-
-			if (bStartupActors)
-			{
-				LevelsPendingFastForward.Add(Level);
-			}
+			LevelsPendingFastForward.Add(Level);
 		}
 	}
 
@@ -4409,6 +4427,49 @@ void UDemoNetDriver::PrepFastForwardLevels()
 	AddReplayTask(new FFastForwardLevelsTask(this));
 }
 
+bool UDemoNetDriver::ProcessFastForwardPackets(TArrayView<FPlaybackPacket> Packets, const TSet<int32>& LevelIndices)
+{
+	// Process all the packets we need.
+	for (FPlaybackPacket& Packet : Packets)
+	{
+		// Skip packets that aren't associated with levels.
+		if (Packet.SeenLevelIndex == 0)
+		{
+			continue;
+		}
+
+		// Don't attempt to go beyond the current demo time.
+		// These packets should have been already been filtered out while reading.
+		if (!ensureMsgf(Packet.TimeSeconds <= DemoCurrentTime, TEXT("UDemoNetDriver::FastForwardLevels: Read packet beyond DemoCurrentTime DemoTime = %f PacketTime = %f"), DemoCurrentTime, Packet.TimeSeconds))
+		{
+			break;
+		}
+
+		if (SeenLevelStatuses.IsValidIndex(Packet.SeenLevelIndex - 1))
+		{
+			const FLevelStatus& LevelStatus = GetLevelStatus(Packet.SeenLevelIndex);
+			const bool bCareAboutLevel = LevelIndices.Contains(LevelStatus.LevelIndex);
+
+			if (bCareAboutLevel)
+			{
+				// If we tried to process the packet, but failed, then the replay will be in a broken state.
+				// ProcessPacket will have called StopDemo.
+				if (!ProcessPacket(Packet.Data.GetData(), Packet.Data.Num()))
+				{
+					UE_LOG(LogDemo, Warning, TEXT("FastForwardLevel failed to process packet"));
+					return false;
+				}
+			}
+		}
+		else
+		{
+			UE_LOG(LogDemo, Warning, TEXT("FastForwardLevel could not process packet with invalid seen level index"));
+		}
+	}
+
+	return true;
+}
+
 bool UDemoNetDriver::FastForwardLevels(const FGotoResult& GotoResult)
 {
 	DECLARE_SCOPE_CYCLE_COUNTER(TEXT("FastForwardLevels time"), STAT_FastForwardLevelTime, STATGROUP_Net);
@@ -4497,10 +4558,13 @@ bool UDemoNetDriver::FastForwardLevels(const FGotoResult& GotoResult)
 	} ReadPacketsHelper(*this, LastProcessedPacketTime);
 
 	DeletedNetStartupActors.Empty();
-	DeletedActorGuids.Empty();
+	PlaybackDeltaCheckpointData.Empty();
+
+	TArray<TInterval<int32>> DeltaCheckpointPacketIntervals;
+	const bool bDeltaCheckpoint = HasDeltaCheckpoints();
 
 	{
-		auto IgnoreReceivedExportGUIDs = ((UPackageMapClient*)ServerConnection->PackageMap)->ScopedIgnoreReceivedExportGUIDs();
+		auto IgnoreReceivedExportGUIDs = Cast<UPackageMapClient>(ServerConnection->PackageMap)->ScopedIgnoreReceivedExportGUIDs();
 
 		// First, read in the checkpoint data (if any is available);
 		if (CheckpointArchive->TotalSize() != 0)
@@ -4510,8 +4574,6 @@ bool UDemoNetDriver::FastForwardLevels(const FGotoResult& GotoResult)
 
 			TGuardValue<bool> LoadingCheckpointGuard(bIsLoadingCheckpoint, true);
 
-			const bool bDeltaCheckpoint = HasDeltaCheckpoints();
-
 			do 
 			{
 				FArchivePos MaxArchivePos = 0;
@@ -4541,20 +4603,17 @@ bool UDemoNetDriver::FastForwardLevels(const FGotoResult& GotoResult)
 				{
 					if (bDeltaCheckpoint)
 					{
-						TSet<FString> DeltaActors;
-						*CheckpointArchive << DeltaActors;
-
-						DeletedNetStartupActors.Append(DeltaActors);
+						TUniquePtr<FDeltaCheckpointData>& CheckpointData = PlaybackDeltaCheckpointData.Emplace_GetRef(new FDeltaCheckpointData());
 
-						TSet<FNetworkGUID> DeltaGuids;
-						*CheckpointArchive << DeltaGuids;
+						*CheckpointArchive << CheckpointData->DestroyedNetStartupActors;
+						*CheckpointArchive << CheckpointData->DestroyedDynamicActors;
+						*CheckpointArchive << CheckpointData->ChannelsToClose;
 
-						DeletedActorGuids.Append(DeltaGuids);
+						DeletedNetStartupActors.Append(CheckpointData->DestroyedNetStartupActors);
 					}
 					else
 					{
 						DeletedNetStartupActors.Empty();
-						DeletedActorGuids.Empty();
 
 						*CheckpointArchive << DeletedNetStartupActors;
 					}
@@ -4562,12 +4621,23 @@ bool UDemoNetDriver::FastForwardLevels(const FGotoResult& GotoResult)
 
 				CheckpointArchive->Seek(PacketOffset);
 
+				int32 DeltaPacketStartIndex = INDEX_NONE;
+				if (bDeltaCheckpoint)
+				{
+					DeltaPacketStartIndex = ReadPacketsHelper.Packets.Num();
+				}
+
 				if (!ReadPacketsHelper.ReadPackets(*CheckpointArchive) && ReadPacketsHelper.IsError())
 				{
 					UE_LOG(LogDemo, Warning, TEXT("UDemoNetDriver::FastForwardLevels: Failed to read packets from Checkpoint."));
 					NotifyDemoPlaybackFailure(EDemoPlayFailure::Generic);
 					return false;
 				}
+
+				if (bDeltaCheckpoint)
+				{
+					DeltaCheckpointPacketIntervals.Emplace(DeltaPacketStartIndex, ReadPacketsHelper.Packets.Num() - 1);
+				}
 			} 
 			while (!CheckpointArchive->IsError() && (CheckpointArchive->Tell() < CheckpointArchive->TotalSize()));
 		}
@@ -4577,7 +4647,13 @@ bool UDemoNetDriver::FastForwardLevels(const FGotoResult& GotoResult)
 
 		StreamingAr->SetEngineNetVer(PlaybackDemoHeader.EngineNetworkProtocolVersion);
 		StreamingAr->SetGameNetVer(PlaybackDemoHeader.GameNetworkProtocolVersion);
-		
+
+		int32 StreamPacketStartIndex = INDEX_NONE;
+		if (bDeltaCheckpoint)
+		{
+			StreamPacketStartIndex = ReadPacketsHelper.Packets.Num();
+		}
+
 		while (!StreamingAr->AtEnd() && ReplayStreamer->IsDataAvailable() && ReadPacketsHelper.ReadPackets(*StreamingAr));
 
 		if (ReadPacketsHelper.IsError())
@@ -4586,6 +4662,11 @@ bool UDemoNetDriver::FastForwardLevels(const FGotoResult& GotoResult)
 			NotifyDemoPlaybackFailure(EDemoPlayFailure::Serialization);
 			return false;
 		}
+
+		if (bDeltaCheckpoint)
+		{
+			DeltaCheckpointPacketIntervals.Emplace(StreamPacketStartIndex, ReadPacketsHelper.Packets.Num() - 1);
+		}
 	}
 
 	// If we've gotten this far, it means we should have something to process.
@@ -4678,67 +4759,37 @@ bool UDemoNetDriver::FastForwardLevels(const FGotoResult& GotoResult)
 
 	{
 		TGuardValue<bool> FastForward(bIsFastForwarding, true);
-		struct FScopedIgnoreChannels
-		{
-			FScopedIgnoreChannels(UNetConnection* InConnection):
-				Connection(InConnection)
-			{
-				if (Connection.IsValid())
-				{
-					Connection->SetIgnoreAlreadyOpenedChannels(true);
-				}	
-			}
+		FScopedIgnoreOpenChannels ScopedIgnoreChannels(ServerConnection);
 
-			~FScopedIgnoreChannels()
-			{
-				if (Connection.IsValid())
-				{
-					Connection->SetIgnoreAlreadyOpenedChannels(false);
-				}	
-			}
-
-		private:
-			TWeakObjectPtr<UNetConnection> Connection;
-		} ScopedIgnoreChannels(ServerConnection);
-
-		FScopedIgnoreDeletedGuids IgnoreGuids(ServerConnection, DeletedActorGuids);
-
-		// Process all the packets we need.
-		for (FPlaybackPacket& Packet : ReadPacketsHelper.Packets)
+		if (bDeltaCheckpoint)
 		{
-			// Skip packets that aren't associated with levels.
-			if (Packet.SeenLevelIndex == 0)
-			{
-				continue;
-			}
+			UDemoNetConnection* DemoConnection = CastChecked<UDemoNetConnection>(ServerConnection);
 
-			// Don't attempt to go beyond the current demo time.
-			// These packets should have been already been filtered out while reading.
-			if (!ensureMsgf(Packet.TimeSeconds <= DemoCurrentTime, TEXT("UDemoNetDriver::FastForwardLevels: Read packet beyond DemoCurrentTime DemoTime = %f PacketTime = %f"), DemoCurrentTime, Packet.TimeSeconds))
+			for (int32 i = 0; i < DeltaCheckpointPacketIntervals.Num(); ++i)
 			{
-				break;
-			}
-
-			if (SeenLevelStatuses.IsValidIndex(Packet.SeenLevelIndex - 1))
-			{
-				const FLevelStatus& LevelStatus = GetLevelStatus(Packet.SeenLevelIndex);
-				const bool bCareAboutLevel = LevelIndices.Contains(LevelStatus.LevelIndex);
-
-				if (bCareAboutLevel)
+				if (PlaybackDeltaCheckpointData.IsValidIndex(i))
 				{
-					// If we tried to process the packet, but failed, then the replay will be in a broken state.
-					// ProcessPacket will have called StopDemo.
-					if (!ProcessPacket(Packet.Data.GetData(), Packet.Data.Num()))
+					check(PlaybackDeltaCheckpointData[i].IsValid());
+
+					for (auto& ChannelPair : PlaybackDeltaCheckpointData[i]->ChannelsToClose)
 					{
-						UE_LOG(LogDemo, Warning, TEXT("FastForwardLevel failed to process packet"));
-						return false;
+						if (UActorChannel* ActorChannel = DemoConnection->GetOpenChannelMap().FindRef(ChannelPair.Key))
+						{
+							ActorChannel->ConditionalCleanUp(true, ChannelPair.Value);
+						}
 					}
 				}
+
+				check(DeltaCheckpointPacketIntervals[i].IsValid());
+
+				ProcessFastForwardPackets(MakeArrayView<FPlaybackPacket>(&ReadPacketsHelper.Packets[DeltaCheckpointPacketIntervals[i].Min], DeltaCheckpointPacketIntervals[i].Size()), LevelIndices);
 			}
-			else
-			{
-				UE_LOG(LogDemo, Warning, TEXT("FastForwardLevel could not process packet with invalid seen level index"));
-			}
+
+			DemoConnection->GetOpenChannelMap().Empty();
+		}
+		else
+		{
+			ProcessFastForwardPackets(ReadPacketsHelper.Packets, LevelIndices);
 		}
 	}
 
@@ -4804,9 +4855,9 @@ bool UDemoNetDriver::LoadCheckpoint(const FGotoResult& GotoResult)
 
 	FArchive* GotoCheckpointArchive = ReplayStreamer->GetCheckpointArchive();
 
-	check( GotoCheckpointArchive != nullptr );
-	check( !bIsFastForwardingForCheckpoint );
-	check( !bIsFastForwarding );
+	check(GotoCheckpointArchive != nullptr);
+	check(!bIsFastForwardingForCheckpoint);
+	check(!bIsFastForwarding);
 
 	GotoCheckpointArchive->SetEngineNetVer(PlaybackDemoHeader.EngineNetworkProtocolVersion);
 	GotoCheckpointArchive->SetGameNetVer(PlaybackDemoHeader.GameNetworkProtocolVersion);
@@ -4890,26 +4941,26 @@ bool UDemoNetDriver::LoadCheckpoint(const FGotoResult& GotoResult)
 
 	// Save off the current spectator position
 	// Check for nullptr, which can be the case if we haven't played any of the demo yet but want to fast forward (joining live game for example)
-	if ( SpectatorController != nullptr )
+	if (SpectatorController != nullptr)
 	{
 		// Save off the SpectatorController's GUID so that we know not to queue his bunches
-		AddNonQueuedActorForScrubbing( SpectatorController );
+		AddNonQueuedActorForScrubbing(SpectatorController);
 	}
 
 	// Remember the spectator controller's view target so we can restore it
 	FNetworkGUID ViewTargetGUID;
 
-	if ( SpectatorController && SpectatorController->GetViewTarget() )
+	if (SpectatorController && SpectatorController->GetViewTarget())
 	{
-		ViewTargetGUID = GuidCache->NetGUIDLookup.FindRef( SpectatorController->GetViewTarget() );
+		ViewTargetGUID = GuidCache->NetGUIDLookup.FindRef(SpectatorController->GetViewTarget());
 
-		if ( ViewTargetGUID.IsValid() )
+		if (ViewTargetGUID.IsValid())
 		{
-			AddNonQueuedActorForScrubbing( SpectatorController->GetViewTarget() );
+			AddNonQueuedActorForScrubbing(SpectatorController->GetViewTarget());
 		}
 	}
 
-	PauseChannels( false );
+	PauseChannels(false);
 
 	FNetworkReplayDelegates::OnPreScrub.Broadcast(GetWorld());
 
@@ -4917,8 +4968,8 @@ bool UDemoNetDriver::LoadCheckpoint(const FGotoResult& GotoResult)
 
 	struct FPreservedNetworkGUIDEntry
 	{
-		FPreservedNetworkGUIDEntry( const FNetworkGUID InNetGUID, const AActor* const InActor )
-			: NetGUID( InNetGUID ), Actor( InActor ) {}
+		FPreservedNetworkGUIDEntry(const FNetworkGUID InNetGUID, const AActor* const InActor)
+			: NetGUID(InNetGUID), Actor(InActor) {}
 
 		FNetworkGUID NetGUID;
 		const AActor* Actor;
@@ -4952,14 +5003,14 @@ bool UDemoNetDriver::LoadCheckpoint(const FGotoResult& GotoResult)
 
 
 	// Destroy all non startup actors. They will get restored with the checkpoint
-	for ( FActorIterator It( GetWorld() ); It; ++It )
+	for (FActorIterator It(GetWorld()); It; ++It)
 	{
 		// If there are any existing actors that are bAlwaysRelevant, don't queue their bunches.
 		// Actors that do queue their bunches might not appear immediately after the checkpoint is loaded,
 		// and missing bAlwaysRelevant actors are more likely to cause noticeable artifacts.
 		// NOTE - We are adding the actor guid here, under the assumption that the actor will reclaim the same guid when we load the checkpoint
 		// This is normally the case, but could break if actors get destroyed and re-created with different guids during recording
-		if ( It->bAlwaysRelevant )
+		if (It->bAlwaysRelevant)
 		{
 			AddNonQueuedActorForScrubbing(*It);
 		}
@@ -4990,7 +5041,7 @@ bool UDemoNetDriver::LoadCheckpoint(const FGotoResult& GotoResult)
 
 		// Prevent NetStartupActors from being destroyed.
 		// NetStartupActors that can't have properties directly re-applied should use QueueNetStartupActorForRollbackViaDeletion.
-		if ( It->IsNetStartupActor() )
+		if (It->IsNetStartupActor())
 		{
 			// Go ahead and rewind this now, since we won't be destroying it later.
 			if (It->bReplayRewindable)
@@ -5001,40 +5052,40 @@ bool UDemoNetDriver::LoadCheckpoint(const FGotoResult& GotoResult)
 			continue;
 		}
 
-		GetWorld()->DestroyActor( *It, true );
+		GetWorld()->DestroyActor(*It, true);
 	}
 
 	// Destroy all particle FX attached to the WorldSettings (the WorldSettings actor persists but the particle FX spawned at runtime shouldn't)
 	GetWorld()->HandleTimelineScrubbed();
 
 	// Remove references to our KeepAlive actors so that cleaning up the channels won't destroy them.
-	for ( int32 i = ServerConnection->OpenChannels.Num() - 1; i >= 0; i-- )
+	for (int32 i = ServerConnection->OpenChannels.Num() - 1; i >= 0; i--)
 	{
 		UChannel* OpenChannel = ServerConnection->OpenChannels[i];
 		if ( OpenChannel != nullptr )
 		{
-			UActorChannel* ActorChannel = Cast< UActorChannel >( OpenChannel );
-			if ( ActorChannel != nullptr && KeepAliveActors.Contains(ActorChannel->Actor))
+			UActorChannel* ActorChannel = Cast<UActorChannel>(OpenChannel);
+			if (ActorChannel != nullptr && KeepAliveActors.Contains(ActorChannel->Actor))
 			{
 				ActorChannel->Actor = nullptr;
 			}
 		}
 	}
 
-	if ( ServerConnection->OwningActor == SpectatorController )
+	if (ServerConnection->OwningActor == SpectatorController)
 	{
 		ServerConnection->OwningActor = nullptr;
 	}
 #else
-	for ( int32 i = ServerConnection->OpenChannels.Num() - 1; i >= 0; i-- )
+	for (int32 i = ServerConnection->OpenChannels.Num() - 1; i >= 0; i--)
 	{
 		UChannel* OpenChannel = ServerConnection->OpenChannels[i];
-		if ( OpenChannel != nullptr )
+		if (OpenChannel != nullptr)
 		{
-			UActorChannel* ActorChannel = Cast< UActorChannel >( OpenChannel );
-			if ( ActorChannel != nullptr && ActorChannel->GetActor() != nullptr && !ActorChannel->GetActor()->IsNetStartupActor() )
+			UActorChannel* ActorChannel = Cast<UActorChannel>(OpenChannel);
+			if (ActorChannel != nullptr && ActorChannel->GetActor() != nullptr && !ActorChannel->GetActor()->IsNetStartupActor())
 			{
-				GetWorld()->DestroyActor( ActorChannel->GetActor(), true );
+				GetWorld()->DestroyActor(ActorChannel->GetActor(), true);
 			}
 		}
 	}
@@ -5049,11 +5100,11 @@ bool UDemoNetDriver::LoadCheckpoint(const FGotoResult& GotoResult)
 	ServerConnection->CleanUp();
 
 	// Destroy startup actors that need to rollback via being destroyed and re-created
-	for ( FActorIterator It( GetWorld() ); It; ++It )
+	for (FActorIterator It(GetWorld()); It; ++It)
 	{
-		if ( RollbackNetStartupActors.Contains( It->GetFullName() ) )
+		if (RollbackNetStartupActors.Contains(It->GetFullName()))
 		{
-			GetWorld()->DestroyActor( *It, true );
+			GetWorld()->DestroyActor(*It, true);
 		}
 	}
 
@@ -5088,7 +5139,7 @@ bool UDemoNetDriver::LoadCheckpoint(const FGotoResult& GotoResult)
 	// Catch a rare case where the spectator controller is null, but a valid GUID is
 	// found on the GuidCache. The weak pointers in the NetGUIDLookup map are probably
 	// going null, and we want catch these cases and investigate further.
-	if (!ensure( GuidCache->NetGUIDLookup.FindRef( SpectatorController ).IsValid() == ( SpectatorController != nullptr ) ))
+	if (!ensure(GuidCache->NetGUIDLookup.FindRef(SpectatorController).IsValid() == (SpectatorController != nullptr)))
 	{
 		UE_LOG(LogDemo, Log, TEXT("LoadCheckpoint: SpectatorController is null and a valid GUID for null was found in the GuidCache. SpectatorController = %s"),
 			*GetFullNameSafe(SpectatorController));
@@ -5099,23 +5150,22 @@ bool UDemoNetDriver::LoadCheckpoint(const FGotoResult& GotoResult)
 	GuidCache->ResetCacheForDemo();
 
 	// Restore preserved packagemap entries
-	for ( const FPreservedNetworkGUIDEntry& PreservedEntry : NetGUIDsToPreserve )
+	for (const FPreservedNetworkGUIDEntry& PreservedEntry : NetGUIDsToPreserve)
 	{
-		check( PreservedEntry.NetGUID.IsValid() );
+		check(PreservedEntry.NetGUID.IsValid());
 		
-		FNetGuidCacheObject& CacheObject = GuidCache->ObjectLookup.FindOrAdd( PreservedEntry.NetGUID );
+		FNetGuidCacheObject& CacheObject = GuidCache->ObjectLookup.FindOrAdd(PreservedEntry.NetGUID);
 
 		CacheObject.Object = MakeWeakObjectPtr(const_cast<AActor*>(PreservedEntry.Actor));
-		check( CacheObject.Object != nullptr );
+		check(CacheObject.Object != nullptr);
 		CacheObject.bNoLoad = true;
-		GuidCache->NetGUIDLookup.Add( CacheObject.Object, PreservedEntry.NetGUID );
+		GuidCache->NetGUIDLookup.Add(CacheObject.Object, PreservedEntry.NetGUID);
 	}
 
-	if ( GotoCheckpointArchive->TotalSize() == 0 || GotoCheckpointArchive->TotalSize() == INDEX_NONE )
+	if (GotoCheckpointArchive->TotalSize() == 0 || GotoCheckpointArchive->TotalSize() == INDEX_NONE)
 	{
 		// Make sure this is empty so that RespawnNecessaryNetStartupActors will respawn them
 		DeletedNetStartupActors.Empty();
-		DeletedActorGuids.Empty();
 
 		// Re-create all startup actors that were destroyed but should exist beyond this point
 		TArray<AActor*> SpawnedActors;
@@ -5126,9 +5176,9 @@ bool UDemoNetDriver::LoadCheckpoint(const FGotoResult& GotoResult)
 		bDemoPlaybackDone		= false;
 		bIsLoadingCheckpoint	= false;
 
-		if ( GotoResult.ExtraTimeMS != -1 )
+		if (GotoResult.ExtraTimeMS != -1)
 		{
-			SkipTimeInternal( ( float )GotoResult.ExtraTimeMS / 1000.0f, true, true );
+			SkipTimeInternal((float)GotoResult.ExtraTimeMS / 1000.0f, true, true);
 		}
 		else
 		{
@@ -5143,7 +5193,9 @@ bool UDemoNetDriver::LoadCheckpoint(const FGotoResult& GotoResult)
 	GotoCheckpointArchive->Seek(0);
 
 	DeletedNetStartupActors.Empty();
-	DeletedActorGuids.Empty();
+	PlaybackDeltaCheckpointData.Empty();
+
+	TArray<TInterval<int32>> DeltaCheckpointPacketIntervals;
 
 	do
 	{
@@ -5177,19 +5229,20 @@ bool UDemoNetDriver::LoadCheckpoint(const FGotoResult& GotoResult)
 			if (bDeltaCheckpoint)
 			{
 				TSet<FString> DeltaActors;
-				*GotoCheckpointArchive << DeltaActors;
 
 				DeletedNetStartupActors.Append(DeltaActors);
 
-				TSet<FNetworkGUID> DeltaGuids;
-				*GotoCheckpointArchive << DeltaGuids;
+				TUniquePtr<FDeltaCheckpointData>& CheckpointData = PlaybackDeltaCheckpointData.Emplace_GetRef(new FDeltaCheckpointData());
 
-				DeletedActorGuids.Append(DeltaGuids);
+				*GotoCheckpointArchive << CheckpointData->DestroyedNetStartupActors;
+				*GotoCheckpointArchive << CheckpointData->DestroyedDynamicActors;
+				*GotoCheckpointArchive << CheckpointData->ChannelsToClose;
+
+				DeletedNetStartupActors.Append(CheckpointData->DestroyedNetStartupActors);
 			}
 			else
 			{
 				DeletedNetStartupActors.Empty();
-				DeletedActorGuids.Empty();
 
 				*GotoCheckpointArchive << DeletedNetStartupActors;
 			}
@@ -5198,7 +5251,7 @@ bool UDemoNetDriver::LoadCheckpoint(const FGotoResult& GotoResult)
 		int32 NumValues = 0;
 		*GotoCheckpointArchive << NumValues;
 
-		for ( int32 i = 0; i < NumValues; i++ )
+		for (int32 i = 0; i < NumValues; i++)
 		{
 			FNetworkGUID Guid;
 
@@ -5215,36 +5268,45 @@ bool UDemoNetDriver::LoadCheckpoint(const FGotoResult& GotoResult)
 			// Remap the pathname to handle client-recorded replays
 			GEngine->NetworkRemapPath(this, PathName, true);
 
-			CacheObject.PathName = FName( *PathName );
+			CacheObject.PathName = FName(*PathName);
 
 			uint8 Flags = 0;
 			*GotoCheckpointArchive << Flags;
 
-			CacheObject.bNoLoad = ( Flags & ( 1 << 0 ) ) ? true : false;
-			CacheObject.bIgnoreWhenMissing = ( Flags & ( 1 << 1 ) ) ? true : false;		
+			CacheObject.bNoLoad = (Flags & (1 << 0)) ? true : false;
+			CacheObject.bIgnoreWhenMissing = (Flags & (1 << 1)) ? true : false;		
 
-			GuidCache->ObjectLookup.Add( Guid, CacheObject );
+			GuidCache->ObjectLookup.Add(Guid, CacheObject);
 		}
 
+		int32 DeltaPacketStartIndex = INDEX_NONE;
+
 		// Read in the compatible rep layouts in this checkpoint
 		if (bDeltaCheckpoint)
 		{
 			CastChecked<UPackageMapClient>(ServerConnection->PackageMap)->SerializeNetFieldExportDelta(*GotoCheckpointArchive);
+
+			DeltaPacketStartIndex = PlaybackPackets.Num();
 		}
 		else
 		{
 			CastChecked<UPackageMapClient>(ServerConnection->PackageMap)->SerializeNetFieldExportGroupMap(*GotoCheckpointArchive);
 		}
 
-		ReadDemoFrameIntoPlaybackPackets( *GotoCheckpointArchive );
+		ReadDemoFrameIntoPlaybackPackets(*GotoCheckpointArchive);
+
+		if (bDeltaCheckpoint)
+		{
+			DeltaCheckpointPacketIntervals.Emplace(DeltaPacketStartIndex, PlaybackPackets.Num() - 1);
+		}
 	}
 	while (!GotoCheckpointArchive->IsError() && (GotoCheckpointArchive->Tell() < GotoCheckpointArchive->TotalSize()));
 	
 	// Destroy startup actors that shouldn't exist past this checkpoint
-	for ( FActorIterator It( World ); It; ++It )
+	for (FActorIterator It( World ); It; ++It)
 	{
 		const FString FullName = It->GetFullName();
-		if ( DeletedNetStartupActors.Contains( FullName ) )
+		if (DeletedNetStartupActors.Contains(FullName))
 		{
 			if (It->bReplayRewindable)
 			{
@@ -5255,12 +5317,12 @@ bool UDemoNetDriver::LoadCheckpoint(const FGotoResult& GotoResult)
 			}
 
 			// Put this actor on the rollback list so we can undelete it during future scrubbing
-			QueueNetStartupActorForRollbackViaDeletion( *It );
+			QueueNetStartupActorForRollbackViaDeletion(*It);
 
 			UE_LOG(LogDemo, Verbose, TEXT("LoadCheckpoint: deleting startup actor %s"), *FullName);
 
 			// Delete the actor
-			World->DestroyActor( *It, true );
+			World->DestroyActor(*It, true);
 		}
 	}
 
@@ -5268,7 +5330,7 @@ bool UDemoNetDriver::LoadCheckpoint(const FGotoResult& GotoResult)
 	TArray<AActor*> SpawnedActors;
 	RespawnNecessaryNetStartupActors(SpawnedActors);
 		
-	if ( PlaybackPackets.Num() > 0 )
+	if (PlaybackPackets.Num() > 0)
 	{
 		DemoCurrentTime = PlaybackPackets.Last().TimeSeconds;
 	}
@@ -5277,10 +5339,10 @@ bool UDemoNetDriver::LoadCheckpoint(const FGotoResult& GotoResult)
 		DemoCurrentTime = 0.0f;
 	}
 
-	if (GotoResult.ExtraTimeMS != -1 )
+	if (GotoResult.ExtraTimeMS != -1)
 	{
 		// If we need to skip more time for fine scrubbing, set that up now
-		SkipTimeInternal( ( float )GotoResult.ExtraTimeMS / 1000.0f, true, true );
+		SkipTimeInternal((float)GotoResult.ExtraTimeMS / 1000.0f, true, true);
 	}
 	else
 	{
@@ -5290,8 +5352,37 @@ bool UDemoNetDriver::LoadCheckpoint(const FGotoResult& GotoResult)
 	}
 
 	{
-		FScopedIgnoreDeletedGuids IgnoreGuids(ServerConnection, DeletedActorGuids);
-		ProcessAllPlaybackPackets();
+		if (bDeltaCheckpoint)
+		{
+			UDemoNetConnection* DemoConnection = CastChecked<UDemoNetConnection>(ServerConnection);
+
+			for (int32 i = 0; i < DeltaCheckpointPacketIntervals.Num(); ++i)
+			{
+				if (PlaybackDeltaCheckpointData.IsValidIndex(i))
+				{
+					check(PlaybackDeltaCheckpointData[i].IsValid());
+
+					for (auto& ChannelPair : PlaybackDeltaCheckpointData[i]->ChannelsToClose)
+					{
+						if (UActorChannel* ActorChannel = DemoConnection->GetOpenChannelMap().FindRef(ChannelPair.Key))
+						{
+							ActorChannel->ConditionalCleanUp(true, ChannelPair.Value);
+						}
+					}
+				}
+
+				check(DeltaCheckpointPacketIntervals[i].IsValid());
+
+				ProcessPlaybackPackets(MakeArrayView<FPlaybackPacket>(&PlaybackPackets[DeltaCheckpointPacketIntervals[i].Min], DeltaCheckpointPacketIntervals[i].Size()));
+			}
+
+			PlaybackPackets.Empty();
+			DemoConnection->GetOpenChannelMap().Empty();
+		}
+		else
+		{
+			ProcessAllPlaybackPackets();
+		}
 	}
 
 	bDemoPlaybackDone = false;
@@ -5307,13 +5398,13 @@ bool UDemoNetDriver::LoadCheckpoint(const FGotoResult& GotoResult)
 		}
 	}
 
-	if ( SpectatorController && ViewTargetGUID.IsValid() )
+	if (SpectatorController && ViewTargetGUID.IsValid())
 	{
-		AActor* ViewTarget = Cast< AActor >( GuidCache->GetObjectFromNetGUID( ViewTargetGUID, false ) );
+		AActor* ViewTarget = Cast<AActor>(GuidCache->GetObjectFromNetGUID(ViewTargetGUID, false));
 
-		if ( ViewTarget )
+		if (ViewTarget)
 		{
-			SpectatorController->SetViewTarget( ViewTarget );
+			SpectatorController->SetViewTarget(ViewTarget);
 		}
 	}
 
@@ -5336,13 +5427,13 @@ bool UDemoNetDriver::IsSavingCheckpoint() const
 
 bool UDemoNetDriver::ShouldQueueBunchesForActorGUID(FNetworkGUID InGUID) const
 {
-	if ( CVarDemoQueueCheckpointChannels.GetValueOnGameThread() == 0)
+	if (CVarDemoQueueCheckpointChannels.GetValueOnGameThread() == 0)
 	{
 		return false;
 	}
 
 	// While loading a checkpoint, queue most bunches so that we don't process them all on one frame.
-	if ( bIsFastForwardingForCheckpoint )
+	if (bIsFastForwardingForCheckpoint)
 	{
 		return !NonQueuedGUIDsForScrubbing.Contains(InGUID);
 	}
@@ -5359,12 +5450,12 @@ FNetworkGUID UDemoNetDriver::GetGUIDForActor(const AActor* InActor) const
 {
 	UNetConnection* Connection = ServerConnection;
 	
-	if ( ClientConnections.Num() > 0)
+	if (ClientConnections.Num() > 0)
 	{
 		Connection = ClientConnections[0];
 	}
 
-	if ( !Connection )
+	if (!Connection)
 	{
 		return FNetworkGUID();
 	}
@@ -5377,12 +5468,12 @@ AActor* UDemoNetDriver::GetActorForGUID(FNetworkGUID InGUID) const
 {
 	UNetConnection* Connection = ServerConnection;
 	
-	if ( ClientConnections.Num() > 0)
+	if (ClientConnections.Num() > 0)
 	{
 		Connection = ClientConnections[0];
 	}
 
-	if ( !Connection )
+	if (!Connection)
 	{
 		return nullptr;
 	}
@@ -5557,30 +5648,30 @@ void UDemoNetConnection::InitConnection( UNetDriver* InDriver, EConnectionState
 
 FString UDemoNetConnection::LowLevelGetRemoteAddress( bool bAppendPort )
 {
-	return TEXT( "UDemoNetConnection" );
+	return TEXT("UDemoNetConnection");
 }
 
 void UDemoNetConnection::LowLevelSend(void* Data, int32 CountBits, FOutPacketTraits& Traits)
 {
 	uint32 CountBytes = FMath::DivideAndRoundUp(CountBits, 8);
 
-	if ( CountBytes == 0 )
+	if (CountBytes == 0)
 	{
-		UE_LOG( LogDemo, Warning, TEXT( "UDemoNetConnection::LowLevelSend: Ignoring empty packet." ) );
+		UE_LOG(LogDemo, Warning, TEXT("UDemoNetConnection::LowLevelSend: Ignoring empty packet."));
 		return;
 	}
 
-	if ( CountBytes > MAX_DEMO_READ_WRITE_BUFFER )
+	if (CountBytes > MAX_DEMO_READ_WRITE_BUFFER)
 	{
-		UE_LOG( LogDemo, Fatal, TEXT( "UDemoNetConnection::LowLevelSend: CountBytes > MAX_DEMO_READ_WRITE_BUFFER." ) );
+		UE_LOG(LogDemo, Fatal, TEXT("UDemoNetConnection::LowLevelSend: CountBytes > MAX_DEMO_READ_WRITE_BUFFER."));
 	}
 
-	TrackSendForProfiler( Data, CountBytes );
+	TrackSendForProfiler(Data, CountBytes);
 
-	if ( ResendAllDataState != EResendAllDataState::None )
+	if (ResendAllDataState != EResendAllDataState::None)
 	{
 		// This path is only active for a checkpoint saving out, we need to queue in separate list
-		new( QueuedCheckpointPackets )FQueuedDemoPacket( ( uint8* )Data, CountBits, Traits );
+		new(QueuedCheckpointPackets)FQueuedDemoPacket((uint8*)Data, CountBits, Traits);
 		return;
 	}
 
@@ -5597,34 +5688,34 @@ void UDemoNetConnection::TrackSendForProfiler(const void* Data, int32 NumBytes)
 
 FString UDemoNetConnection::LowLevelDescribe()
 {
-	return TEXT( "Demo recording/playback driver connection" );
+	return TEXT("Demo recording/playback driver connection");
 }
 
-int32 UDemoNetConnection::IsNetReady( bool Saturate )
+int32 UDemoNetConnection::IsNetReady(bool Saturate)
 {
 	return 1;
 }
 
-void UDemoNetConnection::FlushNet( bool bIgnoreSimulation )
+void UDemoNetConnection::FlushNet(bool bIgnoreSimulation)
 {
 	// in playback, there is no data to send except
 	// channel closing if an error occurs.
-	if ( GetDriver()->ServerConnection != nullptr )
+	if (GetDriver()->ServerConnection != nullptr)
 	{
 		InitSendBuffer();
 	}
 	else
 	{
-		Super::FlushNet( bIgnoreSimulation );
+		Super::FlushNet(bIgnoreSimulation);
 	}
 }
 
-void UDemoNetConnection::HandleClientPlayer( APlayerController* PC, UNetConnection* NetConnection )
+void UDemoNetConnection::HandleClientPlayer(APlayerController* PC, UNetConnection* NetConnection)
 {
 	UDemoNetDriver* DemoDriver = GetDriver();
 	// If the spectator is the same, assume this is for scrubbing, and we are keeping the old one
 	// (so don't set the position, since we want to persist all that)
-	if (DemoDriver->SpectatorController == PC )
+	if (DemoDriver->SpectatorController == PC)
 	{
 		DemoDriver->RestoreConnectionPostScrub(PC, NetConnection);
 		DemoDriver->SetSpectatorController(PC);
@@ -5642,9 +5733,10 @@ void UDemoNetConnection::HandleClientPlayer( APlayerController* PC, UNetConnecti
 			break;
 		}
 	}
+
 	int32 SavedNetSpeed = LocalPlayer ? LocalPlayer->CurrentNetSpeed : 0;
 
-	Super::HandleClientPlayer( PC, NetConnection );
+	Super::HandleClientPlayer(PC, NetConnection);
 	
 	// Restore the netspeed if we're a local replay
 	if (DemoDriver->bIsLocalReplay && LocalPlayer)
@@ -5655,11 +5747,11 @@ void UDemoNetConnection::HandleClientPlayer( APlayerController* PC, UNetConnecti
 	// This is very likely our main demo controller.
 	DemoDriver->SetSpectatorController(PC);
 
-	for ( FActorIterator It( Driver->World ); It; ++It)
+	for (FActorIterator It(Driver->World); It; ++It)
 	{
-		if ( It->IsA( APlayerStart::StaticClass() ) )
+		if (It->IsA(APlayerStart::StaticClass()))
 		{
-			PC->SetInitialLocationAndRotation( It->GetActorLocation(), It->GetActorRotation() );
+			PC->SetInitialLocationAndRotation(It->GetActorLocation(), It->GetActorRotation());
 			break;
 		}
 	}
@@ -5676,7 +5768,7 @@ bool UDemoNetConnection::ClientHasInitializedLevelFor(const AActor* TestActor) c
 	// We save all currently streamed levels into the demo stream so we can force the demo playback client
 	// to stay in sync with the recording server
 	// This may need to be tweaked or re-evaluated when we start recording demos on the client
-	return ( GetDriver()->DemoFrameNum > 2 || Super::ClientHasInitializedLevelFor( TestActor ) );
+	return (GetDriver()->DemoFrameNum > 2 || Super::ClientHasInitializedLevelFor(TestActor));
 }
 
 TSharedPtr<FObjectReplicator> UDemoNetConnection::CreateReplicatorForNewActorChannel(UObject* Object)
@@ -5716,7 +5808,7 @@ TSharedPtr<FObjectReplicator> UDemoNetConnection::CreateReplicatorForNewActorCha
 }
 
 void UDemoNetConnection::DestroyIgnoredActor(AActor* Actor)
-	{
+{
 	QueueNetStartupActorForRewind(Actor);
 
 	Super::DestroyIgnoredActor(Actor);
@@ -5742,9 +5834,19 @@ void UDemoNetConnection::QueueNetStartupActorForRewind(AActor* Actor)
 	}
 }
 
-bool UDemoNetDriver::IsLevelInitializedForActor( const AActor* InActor, const UNetConnection* InConnection ) const
+void UDemoNetConnection::NotifyActorNetGUID(UActorChannel* Channel)
 {
-	return ( DemoFrameNum > 2 || Super::IsLevelInitializedForActor( InActor, InConnection ) );
+	const UDemoNetDriver* const NetDriver = GetDriver();
+
+	if (Channel && NetDriver && NetDriver->HasDeltaCheckpoints())
+	{
+		GetOpenChannelMap().Add(Channel->ActorNetGUID, Channel);
+	}
+}
+
+bool UDemoNetDriver::IsLevelInitializedForActor(const AActor* InActor, const UNetConnection* InConnection) const
+{
+	return (DemoFrameNum > 2 || Super::IsLevelInitializedForActor(InActor, InConnection));
 }
 
 bool UDemoNetDriver::IsPlayingClientReplay() const
@@ -5789,7 +5891,7 @@ void UDemoNetDriver::InitDestroyedStartupActors()
 	if (World)
 	{
 		check(DeletedNetStartupActors.Num() == 0);
-		check(DeltaDeletedNetStartupActors.Num() == 0);
+		check(RecordingDeltaCheckpointData.DestroyedNetStartupActors.Num() == 0);
 
 		// add startup actors destroyed before the creation of this net driver
 		for (auto LevelIt(World->GetLevelIterator()); LevelIt; ++LevelIt)
@@ -5798,23 +5900,24 @@ void UDemoNetDriver::InitDestroyedStartupActors()
 			if (Level)
 			{
 				const TArray<FReplicatedStaticActorDestructionInfo>& DestroyedReplicatedStaticActors = Level->GetDestroyedReplicatedStaticActors();
-				for(const FReplicatedStaticActorDestructionInfo& Info : DestroyedReplicatedStaticActors)
+				for (const FReplicatedStaticActorDestructionInfo& Info : DestroyedReplicatedStaticActors)
 				{
 					DeletedNetStartupActors.Add(Info.FullName);
-					DeltaDeletedNetStartupActors.Add(Info.FullName);
+					RecordingDeltaCheckpointData.DestroyedNetStartupActors.Add(Info.FullName);
 				}
 			}
 		}
 	}
 }
 
-void UDemoNetDriver::NotifyActorDestroyed( AActor* Actor, bool IsSeamlessTravel )
+void UDemoNetDriver::NotifyActorDestroyed(AActor* Actor, bool IsSeamlessTravel)
 {
-	check( Actor != nullptr );
+	check(Actor != nullptr);
 
 	const bool bIsRecording = IsRecording();
 	const bool bNetStartup = Actor->IsNetStartupActor();
 	const bool bActorRewindable = Actor->bReplayRewindable;
+	const bool bDeltaCheckpoint = HasDeltaCheckpoints();
 
 	if (bActorRewindable && !IsSeamlessTravel && !bIsRecording)
 	{
@@ -5871,8 +5974,12 @@ void UDemoNetDriver::NotifyActorDestroyed( AActor* Actor, bool IsSeamlessTravel
 			UE_CLOG(bActorRewindable, LogDemo, Warning, TEXT("Replay Rewindable Actor destroyed during recording. Replay may show artifacts (%s)"), *FullName);
 
 			UE_LOG(LogDemo, VeryVerbose, TEXT("NotifyActyorDestroyed: adding actor to deleted startup list: %s"), *FullName);
-			DeletedNetStartupActors.Add( FullName );
-			DeltaDeletedNetStartupActors.Add( FullName );
+			DeletedNetStartupActors.Add(FullName);
+
+			if (bDeltaCheckpoint)
+			{
+				RecordingDeltaCheckpointData.DestroyedNetStartupActors.Add(FullName);
+			}
 
 			FNetworkGUID NetGUID = GuidCache->NetGUIDLookup.FindRef(Actor);
 			if (NetGUID.IsValid())
@@ -5882,18 +5989,17 @@ void UDemoNetDriver::NotifyActorDestroyed( AActor* Actor, bool IsSeamlessTravel
 		}
 	}
 
-	if (bIsRecording && !IsSeamlessTravel)
+	if (bIsRecording && !bNetStartup && bDeltaCheckpoint)
 	{
 		FNetworkGUID NetGUID = GuidCache->NetGUIDLookup.FindRef(Actor);
-		if (NetGUID.IsValid() && !NetGUID.IsDefault())
+		if (NetGUID.IsValid())
 		{
-			DeletedActorGuids.Add(NetGUID);
-			DeltaDeletedActorGuids.Add(NetGUID);
+			RecordingDeltaCheckpointData.DestroyedDynamicActors.Add(NetGUID);
 		}
 	}
 
 	TUniquePtr<FScopedPacketManager> PacketManager(ConditionallyCreatePacketManager(*Actor->GetLevel()));
-	Super::NotifyActorDestroyed( Actor, IsSeamlessTravel );
+	Super::NotifyActorDestroyed(Actor, IsSeamlessTravel);
 }
 
 void UDemoNetDriver::CleanupOutstandingRewindActors()
@@ -5939,6 +6045,24 @@ void UDemoNetDriver::NotifyActorChannelOpen(UActorChannel* Channel, AActor* Acto
 	{
 		Actor->RewindForReplay();
 	}	
+
+	if (bValidChannel && bValidActor && IsRecording())
+	{
+		RecordingDeltaCheckpointData.DestroyedDynamicActors.Remove(Channel->ActorNetGUID);
+	}
+}
+
+void UDemoNetDriver::NotifyActorChannelCleanedUp(UActorChannel* Channel, EChannelCloseReason CloseReason)
+{
+	if (IsRecording() && Channel)
+	{
+		if (Channel->bOpenedForCheckpoint)
+		{
+			RecordingDeltaCheckpointData.ChannelsToClose.Add(Channel->ActorNetGUID, CloseReason);
+		}
+	}
+
+	Super::NotifyActorChannelCleanedUp(Channel, CloseReason);
 }
 
 void UDemoNetDriver::NotifyActorLevelUnloaded(AActor* Actor)
@@ -5958,27 +6082,27 @@ void UDemoNetDriver::NotifyActorLevelUnloaded(AActor* Actor)
 	Super::NotifyActorLevelUnloaded(Actor);
 }
 
-void UDemoNetDriver::QueueNetStartupActorForRollbackViaDeletion( AActor* Actor )
+void UDemoNetDriver::QueueNetStartupActorForRollbackViaDeletion(AActor* Actor)
 {
-	check( Actor != nullptr );
+	check(Actor != nullptr);
 
-	if ( !Actor->IsNetStartupActor() )
+	if (!Actor->IsNetStartupActor())
 	{
 		return;		// We only want startup actors
 	}
 
-	if ( !IsPlaying() )
+	if (!IsPlaying())
 	{
 		return;		// We should only be doing this at runtime while playing a replay
 	}
 
 	const FString ActorFullName = Actor->GetFullName();
-	if ( RollbackNetStartupActors.Contains( ActorFullName ) )
+	if (RollbackNetStartupActors.Contains(ActorFullName))
 	{
 		return;		// This actor is already queued up
 	}
 
-	if ( Actor->bReplayRewindable )
+	if (Actor->bReplayRewindable)
 	{
 		UE_LOG(LogDemo, Warning, TEXT("Attempted to queue a Replay Rewindable Actor for Rollback Via Deletion. Replay may have artifacts (%s)"), *ActorFullName);
 		return;
@@ -6043,7 +6167,7 @@ void UDemoNetDriver::ForceNetUpdate(AActor* Actor)
 	}
 	else
 	{
-		if ( FNetworkObjectInfo* NetActor = FindNetworkObjectInfo(Actor) )
+		if (FNetworkObjectInfo* NetActor = FindNetworkObjectInfo(Actor))
 		{
 			// replays use update times relative to DemoCurrentTime and not World->TimeSeconds
 			NetActor->NextUpdateTime = DemoCurrentTime - 0.01f;
@@ -6193,11 +6317,11 @@ void UDemoNetDriver::AdjustConsiderTime(const float ReplicatedPercent)
 	UDemoPendingNetGame.
 -----------------------------------------------------------------------------*/
 
-UDemoPendingNetGame::UDemoPendingNetGame( const FObjectInitializer& ObjectInitializer ) : Super( ObjectInitializer )
+UDemoPendingNetGame::UDemoPendingNetGame(const FObjectInitializer& ObjectInitializer) : Super(ObjectInitializer)
 {
 }
 
-void UDemoPendingNetGame::Tick( float DeltaTime )
+void UDemoPendingNetGame::Tick(float DeltaTime)
 {
 	// Replays don't need to do anything here
 }
@@ -6210,7 +6334,7 @@ void UDemoPendingNetGame::SendJoin()
 void UDemoPendingNetGame::LoadMapCompleted(UEngine* Engine, FWorldContext& Context, bool bLoadedMapSuccessfully, const FString& LoadMapError)
 {
 #if !( UE_BUILD_SHIPPING || UE_BUILD_TEST )
-	if ( CVarDemoForceFailure.GetValueOnGameThread() == 2 )
+	if (CVarDemoForceFailure.GetValueOnGameThread() == 2)
 	{
 		bLoadedMapSuccessfully = false;
 	}
@@ -6219,26 +6343,26 @@ void UDemoPendingNetGame::LoadMapCompleted(UEngine* Engine, FWorldContext& Conte
 	// If we have a demo pending net game we should have a demo net driver
 	check(DemoNetDriver);
 
-	if ( !bLoadedMapSuccessfully )
+	if (!bLoadedMapSuccessfully)
 	{
 		DemoNetDriver->StopDemo();
 
 		// If we don't have a world that means we failed loading the new world.
 		// Since there is no world, we must free the net driver ourselves
 		// Technically the pending net game should handle it, but things aren't quite setup properly to handle that either
-		if ( Context.World() == nullptr )
+		if (Context.World() == nullptr)
 		{
-			GEngine->DestroyNamedNetDriver( Context.PendingNetGame, DemoNetDriver->NetDriverName );
+			GEngine->DestroyNamedNetDriver(Context.PendingNetGame, DemoNetDriver->NetDriverName);
 		}
 
 		Context.PendingNetGame = nullptr;
 
-		GEngine->BrowseToDefaultMap( Context );
+		GEngine->BrowseToDefaultMap(Context);
 
-		UE_LOG( LogDemo, Error, TEXT( "UDemoPendingNetGame::HandlePostLoadMap: LoadMap failed: %s" ), *LoadMapError );
-		if ( Context.OwningGameInstance )
+		UE_LOG(LogDemo, Error, TEXT("UDemoPendingNetGame::HandlePostLoadMap: LoadMap failed: %s"), *LoadMapError);
+		if (Context.OwningGameInstance)
 		{
-			Context.OwningGameInstance->HandleDemoPlaybackFailure( EDemoPlayFailure::LoadMap, FString( TEXT( "LoadMap failed" ) ) );
+			Context.OwningGameInstance->HandleDemoPlaybackFailure(EDemoPlayFailure::LoadMap, FString(TEXT("LoadMap failed")));
 		}
 		return;
 	}
diff --git a/Engine/Source/Runtime/Engine/Private/NetConnection.cpp b/Engine/Source/Runtime/Engine/Private/NetConnection.cpp
--- a/Engine/Source/Runtime/Engine/Private/NetConnection.cpp
+++ b/Engine/Source/Runtime/Engine/Private/NetConnection.cpp
@@ -1623,6 +1623,8 @@ FNetworkGUID UNetConnection::GetActorGUIDFromOpenBunch(FInBunch& Bunch)
 		}
 	}
 
+	NET_CHECKSUM( Bunch );
+
 	FNetworkGUID ActorGUID;
 	Bunch << ActorGUID;
 
@@ -2149,9 +2151,9 @@ PRAGMA_ENABLE_DEPRECATION_WARNINGS
 				}
 
 				// peek for guid
-				if (InternalAck && (Bunch.ChName == NAME_Actor))
+				if (InternalAck && bIgnoreActorBunches)
 				{
-					if (Bunch.bOpen && (!Bunch.bPartial || Bunch.bPartialInitial))
+					if (Bunch.bOpen && (!Bunch.bPartial || Bunch.bPartialInitial) && (Bunch.ChName == NAME_Actor))
 					{
 						FBitReaderMark Mark(Bunch);
 						FNetworkGUID ActorGUID = GetActorGUIDFromOpenBunch(Bunch);
@@ -2159,24 +2161,28 @@ PRAGMA_ENABLE_DEPRECATION_WARNINGS
 
 						if (ActorGUID.IsValid() && !ActorGUID.IsDefault())
 						{
-							if (IgnoredGuids.Contains(ActorGUID))
+							if (IgnoredBunchGuids.Contains(ActorGUID))
 							{
 								UE_LOG(LogNetTraffic, Verbose, TEXT("Adding Channel: %i to ignore list, ignoring guid: %s"), Bunch.ChIndex, *ActorGUID.ToString());
-								IgnoredChannels.Add(Bunch.ChIndex);
+								IgnoredBunchChannels.Add(Bunch.ChIndex);
 								continue;
 							}
 							else
 							{
-								if (IgnoredChannels.Remove(Bunch.ChIndex))
+								if (IgnoredBunchChannels.Remove(Bunch.ChIndex))
 								{
 									UE_LOG(LogNetTraffic, Verbose, TEXT("Removing Channel: %i from ignore list, got new guid: %s"), Bunch.ChIndex, *ActorGUID.ToString());
 								}
 							}
 						}
+						else
+						{
+							UE_LOG(LogNetTraffic, Warning, TEXT("Open bunch with invalid actor guid, Channel: %i"), Bunch.ChIndex);
+						}
 					}
 					else
 					{
-						if (IgnoredChannels.Contains(Bunch.ChIndex))
+						if (IgnoredBunchChannels.Contains(Bunch.ChIndex))
 						{
 							UE_LOG(LogNetTraffic, Verbose, TEXT("Ignoring bunch on channel: %i"), Bunch.ChIndex);
 							continue;
@@ -2286,6 +2292,20 @@ void UNetConnection::SetIgnoreAlreadyOpenedChannels(bool bInIgnoreAlreadyOpenedC
 	IgnoringChannels.Reset();
 }
 
+void UNetConnection::SetIgnoreActorBunches(bool bInIgnoreActorBunches, TSet<FNetworkGUID>&& InIgnoredBunchGuids)
+{
+	check(InternalAck);
+	bIgnoreActorBunches = bInIgnoreActorBunches;
+
+	IgnoredBunchChannels.Empty();
+	InIgnoredBunchGuids.Empty();
+
+	if (bIgnoreActorBunches)
+	{
+		IgnoredBunchGuids = MoveTemp(InIgnoredBunchGuids);
+	}
+}
+
 int32 UNetConnection::WriteBitsToSendBuffer( 
 	const uint8 *	Bits, 
 	const int32		SizeInBits, 
