[NO NON-PII COMMIT MESSAGE]

diff --git a/Engine/Source/Runtime/Engine/Private/VorbisAudioInfo.cpp b/Engine/Source/Runtime/Engine/Private/VorbisAudioInfo.cpp
--- a/Engine/Source/Runtime/Engine/Private/VorbisAudioInfo.cpp
+++ b/Engine/Source/Runtime/Engine/Private/VorbisAudioInfo.cpp
@@ -155,6 +155,8 @@ int FVorbisAudioInfo::SeekMemory( uint32 offset, int whence )
 		break;
 	}
 
+	checkf(BufferOffset < SrcBufferDataSize, TEXT("Seeked out of bounds in Ogg Vorbis file! Offset: %u Total file size: %d. Arguments: Whence %d, Requested Offset %d"), BufferOffset, SrcBufferDataSize, whence, offset);
+	
 	return( BufferOffset );
 }
 
