Swapped orbit hotkeys for debug camera.

diff --git a/Engine/Source/Runtime/Engine/Private/DebugCameraController.cpp b/Engine/Source/Runtime/Engine/Private/DebugCameraController.cpp
--- a/Engine/Source/Runtime/Engine/Private/DebugCameraController.cpp
+++ b/Engine/Source/Runtime/Engine/Private/DebugCameraController.cpp
@@ -72,8 +72,8 @@ void InitializeDebugCameraInputBindings()
 		UPlayerInput::AddEngineDefinedActionMapping(FInputActionKeyMapping("DebugCamera_DecreaseFOV", EKeys::Period));
 		UPlayerInput::AddEngineDefinedActionMapping(FInputActionKeyMapping("DebugCamera_ToggleDisplay", EKeys::BackSpace));
 		UPlayerInput::AddEngineDefinedActionMapping(FInputActionKeyMapping("DebugCamera_FreezeRendering", EKeys::F));
-		UPlayerInput::AddEngineDefinedActionMapping(FInputActionKeyMapping("DebugCamera_OrbitCenter", EKeys::O));
-		UPlayerInput::AddEngineDefinedActionMapping(FInputActionKeyMapping("DebugCamera_OrbitHitPoint", EKeys::O, true));
+		UPlayerInput::AddEngineDefinedActionMapping(FInputActionKeyMapping("DebugCamera_OrbitHitPoint", EKeys::O));
+		UPlayerInput::AddEngineDefinedActionMapping(FInputActionKeyMapping("DebugCamera_OrbitCenter", EKeys::O, true));
 
 		UPlayerInput::AddEngineDefinedActionMapping(FInputActionKeyMapping("DebugCamera_Select", EKeys::Gamepad_RightTrigger));
 		UPlayerInput::AddEngineDefinedActionMapping(FInputActionKeyMapping("DebugCamera_IncreaseSpeed", EKeys::Gamepad_RightShoulder));
@@ -101,8 +101,8 @@ void ADebugCameraController::SetupInputComponent()
 
 	InputComponent->BindAction("DebugCamera_ToggleDisplay", IE_Pressed, this, &ADebugCameraController::ToggleDisplay);
 	InputComponent->BindAction("DebugCamera_FreezeRendering", IE_Pressed, this, &ADebugCameraController::ToggleFreezeRendering);
-	InputComponent->BindAction("DebugCamera_OrbitCenter", IE_Pressed, this, &ADebugCameraController::ToggleOrbitCenter);
 	InputComponent->BindAction("DebugCamera_OrbitHitPoint", IE_Pressed, this, &ADebugCameraController::ToggleOrbitHitPoint);
+	InputComponent->BindAction("DebugCamera_OrbitCenter", IE_Pressed, this, &ADebugCameraController::ToggleOrbitCenter);
 
 	InputComponent->BindTouch(IE_Pressed, this, &ADebugCameraController::OnTouchBegin);
 	InputComponent->BindTouch(IE_Released, this, &ADebugCameraController::OnTouchEnd);
diff --git a/Engine/Source/Runtime/Engine/Private/DebugCameraHUD.cpp b/Engine/Source/Runtime/Engine/Private/DebugCameraHUD.cpp
--- a/Engine/Source/Runtime/Engine/Private/DebugCameraHUD.cpp
+++ b/Engine/Source/Runtime/Engine/Private/DebugCameraHUD.cpp
@@ -214,10 +214,10 @@ void ADebugCameraHUD::PostRender()
 			Canvas->DrawText(RenderFont, TEXT("Toggle Display: BackSpace or XButton"), X, yl, 1.f, 1.f, FontRenderInfo);
 			yl += Y;
 
-			Canvas->DrawText(RenderFont, TEXT("Orbit Selected Actor: O"), X, yl, 1.f, 1.f, FontRenderInfo);
+			Canvas->DrawText(RenderFont, TEXT("Orbit Selected Hitpoint: O"), X, yl, 1.f, 1.f, FontRenderInfo);
 			yl += Y;
 
-			Canvas->DrawText(RenderFont, TEXT("Orbit Selected Hitpoint: Shift-O"), X, yl, 1.f, 1.f, FontRenderInfo);
+			Canvas->DrawText(RenderFont, TEXT("Orbit Selected Actor: Shift-O"), X, yl, 1.f, 1.f, FontRenderInfo);
 			yl += Y;
 		}
 	}
