unblock using editor with some sounds by convernting check to clamp. #fyi charles.anderson #jira none

diff --git a/Engine/Source/Runtime/Engine/Private/VorbisAudioInfo.cpp b/Engine/Source/Runtime/Engine/Private/VorbisAudioInfo.cpp
--- a/Engine/Source/Runtime/Engine/Private/VorbisAudioInfo.cpp
+++ b/Engine/Source/Runtime/Engine/Private/VorbisAudioInfo.cpp
@@ -155,8 +155,7 @@ int FVorbisAudioInfo::SeekMemory( uint32 offset, int whence )
 		break;
 	}
 
-	checkf(BufferOffset < SrcBufferDataSize, TEXT("Seeked out of bounds in Ogg Vorbis file! Offset: %u Total file size: %d. Arguments: Whence %d, Requested Offset %d"), BufferOffset, SrcBufferDataSize, whence, offset);
-	
+	BufferOffset = FMath::Clamp<uint32>(BufferOffset 0, SrcBufferDataSize);
 	return( BufferOffset );
 }
 
