Adding demo.JumpToEndOfLiveReplay to control whether or not to fast forward to the end of a live replay when playback starts
Adding demo.InternalPauseChannels to control using existing pause channels logic, or to let the new delegate / game handle it
Fix clamping of current demo time, precision issues could cause the last packet to never be processed
Deprecating bDemoPlaybackDone, which was no longer used

diff --git a/Engine/Source/Runtime/Engine/Classes/Engine/DemoNetDriver.h b/Engine/Source/Runtime/Engine/Classes/Engine/DemoNetDriver.h
--- a/Engine/Source/Runtime/Engine/Classes/Engine/DemoNetDriver.h
+++ b/Engine/Source/Runtime/Engine/Classes/Engine/DemoNetDriver.h
@@ -18,11 +18,9 @@
 #include "Net/RepLayout.h"
 #include "DemoNetDriver.generated.h"
 
-class Error;
 class FNetworkNotify;
-class UDemoNetDriver;
-class UDemoNetConnection;
 class FRepState;
+class UDemoNetDriver;
 
 DECLARE_LOG_CATEGORY_EXTERN(LogDemo, Log, All);
 
@@ -33,12 +31,10 @@ DECLARE_MULTICAST_DELEGATE_OneParam(FOnDemoStartedDelegate, UDemoNetDriver* /* D
 DECLARE_MULTICAST_DELEGATE_TwoParams(FOnDemoFailedToStartDelegate, UDemoNetDriver* /* DemoNetDriver */, EDemoPlayFailure::Type /* FailureType*/);
 
 DECLARE_MULTICAST_DELEGATE(FOnDemoFinishPlaybackDelegate);
-
-class UDemoNetDriver;
-class UDemoNetConnection;
-
 DECLARE_MULTICAST_DELEGATE(FOnDemoFinishRecordingDelegate);
 
+DECLARE_MULTICAST_DELEGATE_OneParam(FOnPauseChannelsDelegate, const bool /* bPaused */);
+
 class FQueuedReplayTask : public TSharedFromThis<FQueuedReplayTask>
 {
 public:
@@ -128,7 +124,7 @@ static const uint32 NETWORK_DEMO_METADATA_VERSION	= 0;
 USTRUCT()
 struct FLevelNameAndTime
 {
-	GENERATED_USTRUCT_BODY()
+	GENERATED_BODY()
 
 	FLevelNameAndTime()
 		: LevelChangeTimeInMS(0)
@@ -382,6 +378,7 @@ class ENGINE_API UDemoNetDriver : public UNetDriver
 	int32 DemoTotalFrames;
 
 	/** True if we are at the end of playing a demo */
+	UE_DEPRECATED(4.23, "bDemoPlaybackDone is no longer used.")
 	bool bDemoPlaybackDone;
 
 	/** True if as have paused all of the channels */
@@ -487,6 +484,9 @@ public:
 	/** Public Delegate for external systems to be notified when replay recording is about to finish. */
 	FOnDemoFinishRecordingDelegate OnDemoFinishRecordingDelegate;
 
+	/** Delegate for external systems to be notified when channels are paused during playback, usually waiting for data to be available. */
+	FOnPauseChannelsDelegate OnPauseChannelsDelegate;
+
 	bool IsLoadingCheckpoint() const { return bIsLoadingCheckpoint; }
 	
 	bool IsPlayingClientReplay() const;
diff --git a/Engine/Source/Runtime/Engine/Private/DemoNetDriver.cpp b/Engine/Source/Runtime/Engine/Private/DemoNetDriver.cpp
--- a/Engine/Source/Runtime/Engine/Private/DemoNetDriver.cpp
+++ b/Engine/Source/Runtime/Engine/Private/DemoNetDriver.cpp
@@ -66,6 +66,9 @@ static TAutoConsoleVariable<int32> CVarLoopDemo(TEXT("demo.Loop"), 0, TEXT("<1>
 static TAutoConsoleVariable<int32> CVarDemoFastForwardIgnoreRPCs( TEXT( "demo.FastForwardIgnoreRPCs" ), 1, TEXT( "If true, RPCs will be discarded during playback fast forward." ) );
 static TAutoConsoleVariable<int32> CVarDemoLateActorDormancyCheck(TEXT("demo.LateActorDormancyCheck"), 1, TEXT("If true, check if an actor should become dormant as late as possible- when serializing it to the demo archive."));
 
+static TAutoConsoleVariable<int32> CVarDemoJumpToEndOfLiveReplay(TEXT("demo.JumpToEndOfLiveReplay"), 1, TEXT("If true, fast forward to a few seconds before the end when starting playback, if the replay is still being recorded."));
+static TAutoConsoleVariable<int32> CVarDemoInternalPauseChannels(TEXT("demo.InternalPauseChannels"), 1, TEXT("If true, run standard logic for PauseChannels rather than letting the game handle it via FOnPauseChannelsDelegate."));
+
 static int32 GDemoLoopCount = 0;
 static FAutoConsoleVariableRef CVarDemoLoopCount( TEXT( "demo.LoopCount" ), GDemoLoopCount, TEXT( "If > 1, will play the replay that many times before stopping." ) );
 
@@ -723,7 +726,6 @@ bool UDemoNetDriver::InitBase(bool bInitAsClient, FNetworkNotify* InNotify, cons
 	{
 		DemoURL							= URL;
 		Time							= 0;
-		bDemoPlaybackDone				= false;
 		bChannelsArePaused				= false;
 		bIsFastForwarding				= false;
 		bIsFastForwardingForCheckpoint	= false;
@@ -2951,31 +2953,35 @@ void UDemoNetDriver::PauseChannels(const bool bPause)
 		return;
 	}
 
-	// Pause all non player controller actors
-	// FIXME: Would love a more elegant way of handling this at a more global level
-	for (int32 i = ServerConnection->OpenChannels.Num() - 1; i >= 0; i--)
+	if (CVarDemoInternalPauseChannels.GetValueOnAnyThread() > 0)
 	{
-		UChannel* OpenChannel = ServerConnection->OpenChannels[i];
+		// Pause all non player controller actors
+		for (int32 i = ServerConnection->OpenChannels.Num() - 1; i >= 0; i--)
+		{
+			UChannel* OpenChannel = ServerConnection->OpenChannels[i];
 
-		UActorChannel* ActorChannel = Cast<UActorChannel>(OpenChannel);
+			UActorChannel* ActorChannel = Cast<UActorChannel>(OpenChannel);
 
-		if (ActorChannel == nullptr)
-		{
-			continue;
-		}
+			if (ActorChannel == nullptr)
+			{
+				continue;
+			}
 
-		ActorChannel->CustomTimeDilation = bPause ? 0.0f : 1.0f;
+			ActorChannel->CustomTimeDilation = bPause ? 0.0f : 1.0f;
 
-		if (ActorChannel->GetActor() == nullptr || SpectatorControllers.Contains(ActorChannel->GetActor()))
-		{
-			continue;
-		}
+			if (ActorChannel->GetActor() == nullptr || SpectatorControllers.Contains(ActorChannel->GetActor()))
+			{
+				continue;
+			}
 
-		// Better way to pause each actor?
-		ActorChannel->GetActor()->CustomTimeDilation = ActorChannel->CustomTimeDilation;
+			// Better way to pause each actor?
+			ActorChannel->GetActor()->CustomTimeDilation = ActorChannel->CustomTimeDilation;
+		}
 	}
 
 	bChannelsArePaused = bPause;
+
+	OnPauseChannelsDelegate.Broadcast(bChannelsArePaused);
 }
 
 bool UDemoNetDriver::ReadDemoFrameIntoPlaybackPackets(FArchive& Ar, TArray<FPlaybackPacket>& InPlaybackPackets, const bool bForLevelFastForward, float* OutTime)
@@ -3822,11 +3828,11 @@ void UDemoNetDriver::TickDemoPlayback(float DeltaSeconds)
 	}
 
 	// Clamp time
-	DemoCurrentTime = FMath::Clamp(DemoCurrentTime, 0.0f, DemoTotalTime - 0.01f);
+	DemoCurrentTime = FMath::Clamp(DemoCurrentTime, 0.0f, DemoTotalTime + 0.01f);
 
 	// Make sure there is data available to read
 	// If we're at the end of the demo, just pause channels and return
-	if (bDemoPlaybackDone || (!PlaybackPackets.Num() && !ReplayStreamer->IsDataAvailable()))
+	if (!PlaybackPackets.Num() && !ReplayStreamer->IsDataAvailable())
 	{
 		PauseChannels(true);
 		return;
@@ -4204,18 +4210,21 @@ void UDemoNetDriver::ReplayStreamingReady(const FStartStreamingResult& Result)
 			}
 		}
 
-		if (ReplayStreamer->IsLive() && ReplayStreamer->GetTotalDemoTime() > 15 * 1000)
+		if (CVarDemoJumpToEndOfLiveReplay.GetValueOnGameThread() != 0)
 		{
-			// If the load time wasn't very long, jump to end now
-			// Otherwise, defer it until we have a more recent replay time
-			if (FPlatformTime::Seconds() - StartTime < 10)
-			{
-				JumpToEndOfLiveReplay();
-			}
-			else
+			if (ReplayStreamer->IsLive() && ReplayStreamer->GetTotalDemoTime() > 15 * 1000)
 			{
-				UE_LOG(LogDemo, Log, TEXT("UDemoNetConnection::ReplayStreamingReady: Deferring checkpoint until next available time."));
-				AddReplayTask(new FJumpToLiveReplayTask(this));
+				// If the load time wasn't very long, jump to end now
+				// Otherwise, defer it until we have a more recent replay time
+				if (FPlatformTime::Seconds() - StartTime < 10)
+				{
+					JumpToEndOfLiveReplay();
+				}
+				else
+				{
+					UE_LOG(LogDemo, Log, TEXT("UDemoNetConnection::ReplayStreamingReady: Deferring checkpoint until next available time."));
+					AddReplayTask(new FJumpToLiveReplayTask(this));
+				}
 			}
 		}
 
@@ -5173,7 +5182,6 @@ bool UDemoNetDriver::LoadCheckpoint(const FGotoResult& GotoResult)
 
 		// This is the very first checkpoint, we'll read the stream from the very beginning in this case
 		DemoCurrentTime			= 0;
-		bDemoPlaybackDone		= false;
 		bIsLoadingCheckpoint	= false;
 
 		if (GotoResult.ExtraTimeMS != -1)
@@ -5385,7 +5393,6 @@ bool UDemoNetDriver::LoadCheckpoint(const FGotoResult& GotoResult)
 		}
 	}
 
-	bDemoPlaybackDone = false;
 	bIsLoadingCheckpoint = false;
 
 	// Save the replicated server time here
