Added configurable 'Magic' header to game packet protocol.

diff --git a/Engine/Source/Runtime/Engine/Private/PacketHandlers/StatelessConnectHandlerComponent.cpp b/Engine/Source/Runtime/Engine/Private/PacketHandlers/StatelessConnectHandlerComponent.cpp
--- a/Engine/Source/Runtime/Engine/Private/PacketHandlers/StatelessConnectHandlerComponent.cpp
+++ b/Engine/Source/Runtime/Engine/Private/PacketHandlers/StatelessConnectHandlerComponent.cpp
@@ -51,20 +51,20 @@ DEFINE_LOG_CATEGORY(LogHandshake);
  * and the server responding with a unique 'Cookie' value, which the client has to respond with.
  *
  * Client - Initial Connect:
- * [HandshakeBit][RestartHandshakeBit][SecretIdBit][24:PacketSizeFiller][AlignPad]
- *																--->
- *																		Server - Stateless Handshake Challenge:
- *																		[HandshakeBit][RestartHandshakeBit][SecretIdBit][4:Timestamp][20:Cookie][AlignPad]
- *																<---
+ * [?:MagicHeader][HandshakeBit][RestartHandshakeBit][SecretIdBit][24:PacketSizeFiller][AlignPad]
+ *													--->
+ *															Server - Stateless Handshake Challenge:
+ *															[?:MagicHeader][HandshakeBit][RestartHandshakeBit][SecretIdBit][4:Timestamp][20:Cookie][AlignPad]
+ *													<---
  * Client - Stateless Challenge Response:
- * [HandshakeBit][RestartHandshakeBit][SecretIdBit][4:Timestamp][20:Cookie][AlignPad]
- *																--->
- *																		Server:
- *																		Ignore, or create UNetConnection.
- *
- *																		Server - Stateless Handshake Ack
- *																		[HandshakeBit][RestartHandshakeBit][SecretIdBit][4:Timestamp][20:Cookie][AlignPad]
- *																<---
+ * [?:MagicHeader][HandshakeBit][RestartHandshakeBit][SecretIdBit][4:Timestamp][20:Cookie][AlignPad]
+ *													--->
+ *															Server:
+ *															Ignore, or create UNetConnection.
+ *
+ *															Server - Stateless Handshake Ack
+ *															[?:MagicHeader][HandshakeBit][RestartHandshakeBit][SecretIdBit][4:Timestamp][20:Cookie][AlignPad]
+ *													<---
  * Client:
  * Handshake Complete.
  *
@@ -74,26 +74,27 @@ DEFINE_LOG_CATEGORY(LogHandshake);
  * The Restart Handshake process is triggered by receiving a (possibly spoofed) non-handshake packet from an unknown IP,
  * so the protocol has been crafted so the server sends only a minimal (1 byte) response, to minimize DRDoS reflection amplification.
  *
- *																		Server - Restart Handshake Request:
- *																		[HandshakeBit][RestartHandshakeBit][AlignPad]
- *																<--
+ *															Server - Restart Handshake Request:
+ *															[?:MagicHeader][HandshakeBit][RestartHandshakeBit][AlignPad]
+ *													<--
  * Client -  Initial Connect (as above)
- *																-->
- *																		Server -  Stateless Handshake Challenge (as above)
- *																<--
+ *													-->
+ *															Server -  Stateless Handshake Challenge (as above)
+ *													<--
  * Client - Stateless Challenge Response + Original Cookie
- * [HandshakeBit][RestartHandshakeBit][SecretIdBit][4:Timestamp][20:Cookie][20:OriginalCookie][AlignPad]
- *																-->
- *																		Server:
- *																		Ignore, or restore UNetConnection.
+ * [?:MagicHeader][HandshakeBit][RestartHandshakeBit][SecretIdBit][4:Timestamp][20:Cookie][20:OriginalCookie][AlignPad]
+ *													-->
+ *															Server:
+ *															Ignore, or restore UNetConnection.
  *
- *																		Server - Stateless Handshake Ack (as above)
- *																<--
+ *															Server - Stateless Handshake Ack (as above)
+ *													<--
  * Client:
  * Handshake Complete. Connection restored.
  *
  *
  *
+ *	- MagicHeader:			An optional static/predefined header, between 0-32 bits in size. Serves no purpose for the handshake code.
  *	- HandshakeBit:			Bit signifying whether a packet is a handshake packet. Applied to all game packets.
  *	- SecretIdBit:			For handshake packets, specifies which HandshakeSecret array was used to generate Cookie.
  *	- RestartHandshakeBit:  Sent by the server when it detects normal game traffic from an unknown IP/port combination.
@@ -183,6 +184,16 @@ DEFINE_LOG_CATEGORY(LogHandshake);
 #define MIN_COOKIE_LIFETIME			SECRET_UPDATE_TIME
 
 
+/**
+ * CVars
+ */
+
+TAutoConsoleVariable<FString> CVarNetMagicHeader(
+	TEXT("net.MagicHeader"),
+	TEXT(""),
+	TEXT("String representing binary bits which are prepended to every packet sent by the game. Max length: 32 bits."));
+
+
 /**
  * StatelessConnectHandlerComponent
  */
@@ -203,10 +214,43 @@ StatelessConnectHandlerComponent::StatelessConnectHandlerComponent()
 	, LastCookie()
 	, bRestartedHandshake(false)
 	, AuthorisedCookie()
+	, MagicHeader()
 {
 	SetActive(true);
 
 	bRequiresHandshake = true;
+
+	FString MagicHeaderStr = CVarNetMagicHeader.GetValueOnAnyThread();
+
+	if (!MagicHeaderStr.IsEmpty())
+	{
+		int32 HeaderStrLen = MagicHeaderStr.Len();
+
+		if (HeaderStrLen <= 32)
+		{
+			bool bValidBinaryStr = true;
+
+			for (int32 i=0; i<MagicHeaderStr.Len() && bValidBinaryStr; i++)
+			{
+				const TCHAR& CurChar = MagicHeaderStr[i];
+
+				bValidBinaryStr = CurChar == '0' || CurChar == '1';
+
+				MagicHeader.Add(CurChar != '0');
+			}
+
+			if (!bValidBinaryStr)
+			{
+				UE_LOG(LogHandshake, Error, TEXT("CVar net.MagicHeader must be a binary string, containing only 1's and 0's, e.g.: 00010101. Current string: %s"), *MagicHeaderStr);
+
+				MagicHeader.Empty();
+			}
+		}
+		else
+		{
+			UE_LOG(LogHandshake, Error, TEXT("CVar net.MagicHeader is too long (%i), maximum size is 32 bits: %s"), MagicHeaderStr.Len(), *MagicHeaderStr);
+		}
+	}
 }
 
 void StatelessConnectHandlerComponent::CountBytes(FArchive& Ar) const
@@ -231,9 +275,14 @@ void StatelessConnectHandlerComponent::NotifyHandshakeBegin()
 
 		if (ServerConn != nullptr)
 		{
-			FBitWriter InitialPacket(HANDSHAKE_PACKET_SIZE_BITS + 1 /* Termination bit */);
+			FBitWriter InitialPacket(GetAdjustedSizeBits(HANDSHAKE_PACKET_SIZE_BITS) + 1 /* Termination bit */);
 			uint8 bHandshakePacket = 1;
 
+			if (MagicHeader.Num() > 0)
+			{
+				InitialPacket.SerializeBits(MagicHeader.GetData(), MagicHeader.Num());
+			}
+
 			InitialPacket.WriteBit(bHandshakePacket);
 
 
@@ -290,7 +339,7 @@ void StatelessConnectHandlerComponent::SendConnectChallenge(const FString& Clien
 {
 	if (Driver != nullptr)
 	{
-		FBitWriter ChallengePacket(HANDSHAKE_PACKET_SIZE_BITS + 1 /* Termination bit */);
+		FBitWriter ChallengePacket(GetAdjustedSizeBits(HANDSHAKE_PACKET_SIZE_BITS) + 1 /* Termination bit */);
 		uint8 bHandshakePacket = 1;
 		uint8 bRestartHandshake = 0; // Ignored clientside
 		float Timestamp = Driver->Time;
@@ -298,6 +347,11 @@ void StatelessConnectHandlerComponent::SendConnectChallenge(const FString& Clien
 
 		GenerateCookie(ClientAddress, ActiveSecret, Timestamp, Cookie);
 
+		if (MagicHeader.Num() > 0)
+		{
+			ChallengePacket.SerializeBits(MagicHeader.GetData(), MagicHeader.Num());
+		}
+
 		ChallengePacket.WriteBit(bHandshakePacket);
 		ChallengePacket.WriteBit(bRestartHandshake);
 		ChallengePacket.WriteBit(ActiveSecret);
@@ -362,10 +416,16 @@ void StatelessConnectHandlerComponent::SendChallengeResponse(uint8 InSecretId, f
 
 	if (ServerConn != nullptr)
 	{
-		FBitWriter ResponsePacket((bRestartedHandshake ? RESTART_RESPONSE_SIZE_BITS : HANDSHAKE_PACKET_SIZE_BITS) + 1 /* Termination bit */);
+		const int32 BaseSize = GetAdjustedSizeBits(bRestartedHandshake ? RESTART_RESPONSE_SIZE_BITS : HANDSHAKE_PACKET_SIZE_BITS);
+		FBitWriter ResponsePacket(BaseSize + 1 /* Termination bit */);
 		uint8 bHandshakePacket = 1;
 		uint8 bRestartHandshake = (bRestartedHandshake ? 1 : 0);
 
+		if (MagicHeader.Num() > 0)
+		{
+			ResponsePacket.SerializeBits(MagicHeader.GetData(), MagicHeader.Num());
+		}
+
 		ResponsePacket.WriteBit(bHandshakePacket);
 		ResponsePacket.WriteBit(bRestartHandshake);
 		ResponsePacket.WriteBit(InSecretId);
@@ -429,11 +489,16 @@ void StatelessConnectHandlerComponent::SendChallengeAck(const FString& ClientAdd
 {
 	if (Driver != nullptr)
 	{
-		FBitWriter AckPacket(HANDSHAKE_PACKET_SIZE_BITS + 1 /* Termination bit */);
+		FBitWriter AckPacket(GetAdjustedSizeBits(HANDSHAKE_PACKET_SIZE_BITS) + 1 /* Termination bit */);
 		uint8 bHandshakePacket = 1;
 		uint8 bRestartHandshake = 0; // Ignored clientside
 		float Timestamp  = -1.f;
 
+		if (MagicHeader.Num() > 0)
+		{
+			AckPacket.SerializeBits(MagicHeader.GetData(), MagicHeader.Num());
+		}
+
 		AckPacket.WriteBit(bHandshakePacket);
 		AckPacket.WriteBit(bRestartHandshake);
 		AckPacket.WriteBit(bHandshakePacket);	// ActiveSecret
@@ -493,12 +558,17 @@ void StatelessConnectHandlerComponent::SendRestartHandshakeRequest(const FString
 {
 	if (Driver != nullptr)
 	{
-		FBitWriter AckPacket(RESTART_HANDSHAKE_PACKET_SIZE_BITS + 1 /* Termination bit */);
+		FBitWriter RestartPacket(GetAdjustedSizeBits(RESTART_HANDSHAKE_PACKET_SIZE_BITS) + 1 /* Termination bit */);
 		uint8 bHandshakePacket = 1;
 		uint8 bRestartHandshake = 1;
 
-		AckPacket.WriteBit(bHandshakePacket);
-		AckPacket.WriteBit(bRestartHandshake);
+		if (MagicHeader.Num() > 0)
+		{
+			RestartPacket.SerializeBits(MagicHeader.GetData(), MagicHeader.Num());
+		}
+
+		RestartPacket.WriteBit(bHandshakePacket);
+		RestartPacket.WriteBit(bRestartHandshake);
 
 #if !UE_BUILD_SHIPPING
 		FDDoSDetection* DDoS = Handler->GetDDoS();
@@ -506,7 +576,7 @@ void StatelessConnectHandlerComponent::SendRestartHandshakeRequest(const FString
 		UE_CLOG((DDoS == nullptr || !DDoS->CheckLogRestrictions()), LogHandshake, Log, TEXT("SendRestartHandshakeRequest."));
 #endif
 
-		CapHandshakePacket(AckPacket);
+		CapHandshakePacket(RestartPacket);
 
 		
 		// Disable PacketHandler parsing, and send the raw packet
@@ -532,7 +602,7 @@ void StatelessConnectHandlerComponent::SendRestartHandshakeRequest(const FString
 			{
 				FOutPacketTraits Traits;
 
-				Driver->LowLevelSend(ClientAddress, AckPacket.GetData(), AckPacket.GetNumBits(), Traits);
+				Driver->LowLevelSend(ClientAddress, RestartPacket.GetData(), RestartPacket.GetNumBits(), Traits);
 			}
 		}
 
@@ -552,7 +622,7 @@ void StatelessConnectHandlerComponent::SendRestartHandshakeRequest(const FString
 
 void StatelessConnectHandlerComponent::CapHandshakePacket(FBitWriter& HandshakePacket)
 {
-	uint32 NumBits = HandshakePacket.GetNumBits();
+	uint32 NumBits = HandshakePacket.GetNumBits() - GetAdjustedSizeBits(0);
 
 	check(NumBits == HANDSHAKE_PACKET_SIZE_BITS || NumBits == RESTART_HANDSHAKE_PACKET_SIZE_BITS || NumBits == RESTART_RESPONSE_SIZE_BITS);
 
@@ -603,6 +673,13 @@ void StatelessConnectHandlerComponent::InitFromConnectionless(StatelessConnectHa
 
 void StatelessConnectHandlerComponent::Incoming(FBitReader& Packet)
 {
+	if (MagicHeader.Num() > 0)
+	{
+		// Don't bother with the expense of verifying the magic header here.
+		uint32 ReadMagic = 0;
+		Packet.SerializeBits(&ReadMagic, MagicHeader.Num());
+	}
+
 	bool bHandshakePacket = !!Packet.ReadBit() && !Packet.IsError();
 
 	if (bHandshakePacket)
@@ -742,9 +819,14 @@ void StatelessConnectHandlerComponent::Incoming(FBitReader& Packet)
 void StatelessConnectHandlerComponent::Outgoing(FBitWriter& Packet, FOutPacketTraits& Traits)
 {
 	// All UNetConnection packets must specify a zero bHandshakePacket value
-	FBitWriter NewPacket(Packet.GetNumBits()+1, true);
+	FBitWriter NewPacket(GetAdjustedSizeBits(Packet.GetNumBits())+1, true);
 	uint8 bHandshakePacket = 0;
 
+	if (MagicHeader.Num() > 0)
+	{
+		NewPacket.SerializeBits(MagicHeader.GetData(), MagicHeader.Num());
+	}
+
 	NewPacket.WriteBit(bHandshakePacket);
 	NewPacket.SerializeBits(Packet.GetData(), Packet.GetNumBits());
 
@@ -753,6 +835,13 @@ void StatelessConnectHandlerComponent::Outgoing(FBitWriter& Packet, FOutPacketTr
 
 void StatelessConnectHandlerComponent::IncomingConnectionless(const FString& Address, FBitReader& Packet)
 {
+	if (MagicHeader.Num() > 0)
+	{
+		// Don't bother with the expense of verifying the magic header here.
+		uint32 ReadMagic = 0;
+		Packet.SerializeBits(&ReadMagic, MagicHeader.Num());
+	}
+
 	bool bHandshakePacket = !!Packet.ReadBit() && !Packet.IsError();
 
 	LastChallengeSuccessAddress.Empty();
@@ -940,7 +1029,7 @@ void StatelessConnectHandlerComponent::UpdateSecret()
 
 int32 StatelessConnectHandlerComponent::GetReservedPacketBits() const
 {
-	int32 ReturnVal = 1;
+	int32 ReturnVal = MagicHeader.Num() + 1 /* bHandshakePacket */;
 
 #if !UE_BUILD_SHIPPING
 	SET_DWORD_STAT(STAT_PacketReservedHandshake, ReturnVal);
diff --git a/Engine/Source/Runtime/Engine/Public/PacketHandlers/StatelessConnectHandlerComponent.h b/Engine/Source/Runtime/Engine/Public/PacketHandlers/StatelessConnectHandlerComponent.h
--- a/Engine/Source/Runtime/Engine/Public/PacketHandlers/StatelessConnectHandlerComponent.h
+++ b/Engine/Source/Runtime/Engine/Public/PacketHandlers/StatelessConnectHandlerComponent.h
@@ -220,6 +220,14 @@ private:
 	 */
 	void UpdateSecret();
 
+	/**
+	 * Calculate the adjusted predefined packet size, based on MagicHeader
+	 */
+	FORCEINLINE int32 GetAdjustedSizeBits(int32 InSizeBits)
+	{
+		return MagicHeader.Num() + InSizeBits;
+	}
+
 
 private:
 	/** The net driver associated with this handler - for performing connectionless sends */
@@ -273,5 +281,8 @@ private:
 
 	/** The cookie which completed the connection handshake. */
 	uint8 AuthorisedCookie[COOKIE_BYTE_SIZE];
+
+	/** The magic header which is prepended to all packets */
+	TBitArray<> MagicHeader;
 };
 
