Fixed editor assert inside uniform expression cache when changing preview feature levels or material levels.

Flow of the issue based on e.g. UEditorEngine::SetMaterialsFeatureLevel:
* First UMaterial::AllMaterialsCacheResourceShadersForRendering is called and it may change material's shader map
* Then UMaterial::CacheResourceShadersForRendering is calling RecacheUniformExpressions(true) to invalidate it's uniform expression cache as it depends on the material shader map
* Then UMaterialInstances::AllMaterialsCacheResourceShadersForRendering is called
* UMaterialInstances::CacheResourceShadersForRendering doesn't invalidate its cache, leaving a stale UB there if UMaterialInstances material depends on a changed UMaterial
* At the end of AllMaterialsCacheResourceShadersForRendering, FSlowTask::Destroy() is called which ends enqueuing a FSlateRHIRenderer::DrawWindow_RenderThread
* FSlateRHIRenderer::DrawWindow_RenderThread will call FMaterialRenderProxy::UpdateDeferredCachedUniformExpressions(), which will crash on a stale UMaterialInstance's experession cache UB pointer.

Fixes this by forcing UMaterialInstance uniform expression cache UB to be recreated when UMaterialInstance::CacheResourceShadersForRendering is called.

diff --git a/Engine/Source/Runtime/Engine/Private/Materials/MaterialInstance.cpp b/Engine/Source/Runtime/Engine/Private/Materials/MaterialInstance.cpp
--- a/Engine/Source/Runtime/Engine/Private/Materials/MaterialInstance.cpp
+++ b/Engine/Source/Runtime/Engine/Private/Materials/MaterialInstance.cpp
@@ -2583,6 +2583,8 @@ void UMaterialInstance::CacheResourceShadersForRendering()
 			ResourcesToCache.Add(StaticPermutationMaterialResources[LocalActiveQL][FeatureLevel]);
 			CacheShadersForResources(ShaderPlatform, ResourcesToCache, true);
 		}
+
+		RecacheUniformExpressions(true);
 	}
 
 	InitResources();
