Make a new component name that increments the numeric suffix nicely.

diff --git a/Engine/Source/Runtime/Engine/Private/SimpleConstructionScript.cpp b/Engine/Source/Runtime/Engine/Private/SimpleConstructionScript.cpp
--- a/Engine/Source/Runtime/Engine/Private/SimpleConstructionScript.cpp
+++ b/Engine/Source/Runtime/Engine/Private/SimpleConstructionScript.cpp
@@ -1205,9 +1205,37 @@ FName USimpleConstructionScript::GenerateNewComponentName(const UClass* Componen
 		else
 		{
 			FString ComponentName;
+			int32 Counter = 1;
+
+			auto BuildNewName = [&Counter, &ComponentName]()
+			{
+				return FName(*(FString::Printf(TEXT("%s%d"), *ComponentName, Counter++)));
+			};
+
 			if (DesiredName != NAME_None)
 			{
 				ComponentName = DesiredName.ToString();
+
+				// If a desired name is supplied then walk back and find any numeric suffix so we can increment it nicely
+				int32 Index = ComponentName.Len();
+				while (Index > 0 && ComponentName[Index-1] >= '0' && ComponentName[Index-1] <= '9')
+				{
+					--Index;
+				}
+
+				if (Index < ComponentName.Len())
+				{
+					FString NumericSuffix = ComponentName.RightChop(Index);
+					Counter = FCString::Atoi(*NumericSuffix);
+					NumericSuffix = FString::Printf(TEXT("%d"), Counter); // Restringify the counter to account for leading 0s that we don't want to remove
+					ComponentName.RemoveAt(ComponentName.Len() - NumericSuffix.Len(), NumericSuffix.Len(), false);
+					++Counter;
+					NewName = BuildNewName();
+				}
+				else
+				{
+					NewName = DesiredName;
+				}
 			}
 			else
 			{
@@ -1221,13 +1249,12 @@ FName USimpleConstructionScript::GenerateNewComponentName(const UClass* Componen
 				{
 					ComponentName.RemoveFromEnd("_C");
 				}
+				NewName = *ComponentName;
 			}
 
-			NewName = *ComponentName;
-			int32 Counter = 1;
 			while (CurrentNames.Contains(NewName))
 			{
-				NewName = FName(*(FString::Printf(TEXT("%s%d"), *ComponentName, Counter++)));
+				NewName = BuildNewName();
 			}
 		}
 	}
