Fix broken playback of Blueprint timeline tracks in cooked, nativized builds.

diff --git a/Engine/Source/Editor/UnrealEd/Private/Kismet2/BlueprintEditorUtils.cpp b/Engine/Source/Editor/UnrealEd/Private/Kismet2/BlueprintEditorUtils.cpp
--- a/Engine/Source/Editor/UnrealEd/Private/Kismet2/BlueprintEditorUtils.cpp
+++ b/Engine/Source/Editor/UnrealEd/Private/Kismet2/BlueprintEditorUtils.cpp
@@ -1816,6 +1816,9 @@ void FBlueprintEditorUtils::PostDuplicateBlueprint(UBlueprint* Blueprint, bool b
 				if (FBlueprintDuplicationScopeFlags::HasAnyFlag(FBlueprintDuplicationScopeFlags::TheSameTimelineGuid))
 				{
 					NewTimeline->TimelineGuid = OldTimeline->TimelineGuid;
+
+					// Ensure that cached names sync back up with the original GUID.
+					FUpdateTimelineCachedNames::Execute(NewTimeline);
 				}
 
 				NewBPGC->Timelines.Add(NewTimeline);
diff --git a/Engine/Source/Runtime/Engine/Classes/Engine/TimelineTemplate.h b/Engine/Source/Runtime/Engine/Classes/Engine/TimelineTemplate.h
--- a/Engine/Source/Runtime/Engine/Classes/Engine/TimelineTemplate.h
+++ b/Engine/Source/Runtime/Engine/Classes/Engine/TimelineTemplate.h
@@ -278,8 +278,7 @@ public:
 };
 
 /**
- *  Helper class that gives blueprint compilation manager permission to update the cached names for
- *  backwards compatibility purposes as PostLoad will not have occurred yet for compile on load.
+ *  Helper class that gives external implementations permission to update cached names.
  */
 struct ENGINE_API FUpdateTimelineCachedNames
 {
@@ -289,5 +288,9 @@ private:
 		TimelineTemplate->UpdateCachedNames();
 	}
 
+	/** Grant external permission to the compilation manager for backwards compatibility purposes as PostLoad will not have occurred yet for compile on load. */
 	friend struct FBlueprintCompilationManagerImpl;
+
+	/** Grant external permission to Blueprint editor utility methods. For example, duplicating a Blueprint for conversion to C++ requires that cached names be updated. */
+	friend class FBlueprintEditorUtils;
 };
\ No newline at end of file
