Fix compile error on non-MSVC toolchain

diff --git a/Engine/Source/Runtime/Engine/Private/RawIndexBuffer.cpp b/Engine/Source/Runtime/Engine/Private/RawIndexBuffer.cpp
--- a/Engine/Source/Runtime/Engine/Private/RawIndexBuffer.cpp
+++ b/Engine/Source/Runtime/Engine/Private/RawIndexBuffer.cpp
@@ -367,7 +367,6 @@ void FRawStaticIndexBuffer16or32<INDEX_TYPE>::CacheOptimize()
 {
 #if WITH_EDITOR
 	CacheOptimizeIndexBuffer(Indices);
-	CachedNumIndices = Indices.Num();
 #endif
 }
 
