Only create structured buffer resources if running the required feature level.

diff --git a/Engine/Source/Runtime/Engine/Private/PrimitiveUniformShaderParameters.cpp b/Engine/Source/Runtime/Engine/Private/PrimitiveUniformShaderParameters.cpp
--- a/Engine/Source/Runtime/Engine/Private/PrimitiveUniformShaderParameters.cpp
+++ b/Engine/Source/Runtime/Engine/Private/PrimitiveUniformShaderParameters.cpp
@@ -6,6 +6,7 @@
 
 void FSinglePrimitiveStructuredBuffer::InitRHI() 
 {
+	if (IsFeatureLevelSupported(GMaxRHIShaderPlatform, ERHIFeatureLevel::SM5))
 	{
 		FRHIResourceCreateInfo CreateInfo;
 		PrimitiveSceneDataBufferRHI = RHICreateStructuredBuffer(sizeof(FVector4), FPrimitiveSceneShaderData::PrimitiveDataStrideInFloat4s * sizeof(FVector4), BUF_Static | BUF_ShaderResource, CreateInfo);
@@ -19,6 +20,7 @@ void FSinglePrimitiveStructuredBuffer::InitRHI()
 		PrimitiveSceneDataBufferSRV = RHICreateShaderResourceView(PrimitiveSceneDataBufferRHI);
 	}
 
+	if (IsFeatureLevelSupported(GMaxRHIShaderPlatform, ERHIFeatureLevel::SM5))
 	{
 		FRHIResourceCreateInfo CreateInfo;
 		LightmapSceneDataBufferRHI = RHICreateStructuredBuffer(sizeof(FVector4), FLightmapSceneShaderData::LightmapDataStrideInFloat4s * sizeof(FVector4), BUF_Static | BUF_ShaderResource, CreateInfo);
