Include objects in generated class as names that are in use

diff --git a/Engine/Source/Runtime/Engine/Private/SimpleConstructionScript.cpp b/Engine/Source/Runtime/Engine/Private/SimpleConstructionScript.cpp
--- a/Engine/Source/Runtime/Engine/Private/SimpleConstructionScript.cpp
+++ b/Engine/Source/Runtime/Engine/Private/SimpleConstructionScript.cpp
@@ -1155,6 +1155,11 @@ void USimpleConstructionScript::GenerateListOfExistingNames(TSet<FName>& Current
 	// <<< End Backwards Compatibility
 	check(Blueprint);
 
+	ForEachObjectWithOuter(Blueprint->GeneratedClass, [&CurrentNames](UObject* BlueprintClassChild)
+	{
+		CurrentNames.Add(BlueprintClassChild->GetFName());
+	});	
+	
 	UClass* FirstNativeClass = FBlueprintEditorUtils::FindFirstNativeClass(Blueprint->ParentClass);
 
 	ForEachObjectWithOuter(FirstNativeClass->GetDefaultObject(), [&CurrentNames](UObject* NativeCDOChild)
