Merging CL#5485186 and 5497209 along with some additional fixes.

#UE4 When printing out FTimerUnifiedDelegates, non-dynamic delegates will now print their function address as well. Additionally the address will be resolved to a symbol name if symbols are available. Useful when using TimerManager.DumpTimerLogsThreshold or the ListTimers exec command.

#UE4 Add vtbl address to timer debug output for non-dynamic timers, and try to resolve the virtual function address on platforms that use clang. Valuable when determining which specific timers are running and will be used to diagnose performance problems.

diff --git a/Engine/Source/Runtime/Core/Public/Delegates/DelegateBase.h b/Engine/Source/Runtime/Core/Public/Delegates/DelegateBase.h
--- a/Engine/Source/Runtime/Core/Public/Delegates/DelegateBase.h
+++ b/Engine/Source/Runtime/Core/Public/Delegates/DelegateBase.h
@@ -140,6 +140,24 @@ public:
 		return Result;
 	}
 
+	/**
+	 * Returns the address of the method pointer which can be used to learn the address of the function that will be executed.
+	 * Returns nullptr if this delegate type does not directly invoke a function pointer.
+	 *
+	 * Note: Only intended to be used to aid debugging of delegates.
+	 *
+	 * @return The address of the function pointer that would be executed by this delegate
+	 */
+	uint64 GetBoundProgramCounterForTimerManager() const
+	{
+		if (IDelegateInstance* Ptr = GetDelegateInstanceProtected())
+		{
+			return Ptr->GetBoundProgramCounterForTimerManager();
+		}
+
+		return 0;
+	}
+
 	/** 
 	 * Checks to see if this delegate is bound to the given user object.
 	 *
diff --git a/Engine/Source/Runtime/Core/Public/Delegates/DelegateInstancesImpl.h b/Engine/Source/Runtime/Core/Public/Delegates/DelegateInstancesImpl.h
--- a/Engine/Source/Runtime/Core/Public/Delegates/DelegateInstancesImpl.h
+++ b/Engine/Source/Runtime/Core/Public/Delegates/DelegateInstancesImpl.h
@@ -81,6 +81,11 @@ public:
 		return UserObjectPtr.Get();
 	}
 
+	virtual uint64 GetBoundProgramCounterForTimerManager() const override final
+	{
+		return 0;
+	}
+
 	// Deprecated
 	virtual bool HasSameObject( const void* InUserObject ) const override final
 	{
@@ -239,6 +244,15 @@ public:
 		return UserObject.Pin().Get();
 	}
 
+	virtual uint64 GetBoundProgramCounterForTimerManager() const override final
+	{
+#if PLATFORM_64BITS
+		return *((uint64*)&MethodPtr);
+#else
+		return *((uint32*)&MethodPtr);
+#endif
+	}
+
 	// Deprecated
 	virtual bool HasSameObject(const void* InUserObject) const override final
 	{
@@ -420,6 +434,15 @@ public:
 		return UserObject;
 	}
 
+	virtual uint64 GetBoundProgramCounterForTimerManager() const override final
+	{
+#if PLATFORM_64BITS
+		return *((uint64*)&MethodPtr);
+#else
+		return *((uint32*)&MethodPtr);
+#endif
+	}
+
 	// Deprecated
 	virtual bool HasSameObject( const void* InUserObject ) const override final
 	{
@@ -573,6 +596,15 @@ public:
 		return UserObject.Get();
 	}
 
+	virtual uint64 GetBoundProgramCounterForTimerManager() const override final
+	{
+#if PLATFORM_64BITS
+		return *((uint64*)&MethodPtr);
+#else
+		return *((uint32*)&MethodPtr);
+#endif
+	}
+
 	// Deprecated
 	virtual bool HasSameObject( const void* InUserObject ) const override final
 	{
@@ -730,6 +762,15 @@ public:
 		return nullptr;
 	}
 
+	virtual uint64 GetBoundProgramCounterForTimerManager() const override final
+	{
+#if PLATFORM_64BITS
+		return *((uint64*)&StaticFuncPtr);
+#else
+		return *((uint32*)&StaticFuncPtr);
+#endif
+	}
+
 	// Deprecated
 	virtual bool HasSameObject( const void* UserObject ) const override final
 	{
@@ -868,6 +909,11 @@ public:
 		return nullptr;
 	}
 
+	virtual uint64 GetBoundProgramCounterForTimerManager() const override final
+	{
+		return 0;
+	}
+
 	// Deprecated
 	virtual bool HasSameObject(const void* UserObject) const override final
 	{
@@ -1013,6 +1059,11 @@ public:
 		return ContextObject.Get();
 	}
 
+	virtual uint64 GetBoundProgramCounterForTimerManager() const override final
+	{
+		return 0;
+	}
+
 	// Deprecated
 	virtual bool HasSameObject(const void* InContextObject) const override final
 	{
diff --git a/Engine/Source/Runtime/Core/Public/Delegates/IDelegateInstance.h b/Engine/Source/Runtime/Core/Public/Delegates/IDelegateInstance.h
--- a/Engine/Source/Runtime/Core/Public/Delegates/IDelegateInstance.h
+++ b/Engine/Source/Runtime/Core/Public/Delegates/IDelegateInstance.h
@@ -99,6 +99,16 @@ public:
 	 */
 	virtual const void* GetObjectForTimerManager() const = 0;
 
+	/**
+	 * Returns the address of the method pointer which can be used to learn the address of the function that will be executed.
+	 * Returns nullptr if this delegate type does not directly invoke a function pointer.
+	 *
+	 * Note: Only intended to be used to aid debugging of delegates.
+	 *
+	 * @return The address of the function pointer that would be executed by this delegate
+	 */
+	virtual uint64 GetBoundProgramCounterForTimerManager() const = 0;
+
 	/**
 	 * Returns true if this delegate is bound to the specified UserObject,
 	 *
diff --git a/Engine/Source/Runtime/Engine/Private/TimerManager.cpp b/Engine/Source/Runtime/Engine/Private/TimerManager.cpp
--- a/Engine/Source/Runtime/Engine/Private/TimerManager.cpp
+++ b/Engine/Source/Runtime/Engine/Private/TimerManager.cpp
@@ -12,6 +12,7 @@
 #include "Misc/TimeGuard.h"
 #include "ProfilingDebugging/CsvProfiler.h"
 #include "Algo/Transform.h"
+#include "HAL/PlatformStackWalk.h"
 
 DECLARE_CYCLE_STAT(TEXT("SetTimer"), STAT_SetTimer, STATGROUP_Engine);
 DECLARE_CYCLE_STAT(TEXT("SetTimeForNextTick"), STAT_SetTimerForNextTick, STATGROUP_Engine);
@@ -29,7 +30,19 @@ static FAutoConsoleVariableRef CVarDumpTimerLogsThreshold(
 	TEXT("Threshold (in milliseconds) after which we log timer info to try and help track down spikes in the timer code. Disabled when set to 0"),
 	ECVF_Default);
 
-static int32 MaxExpiredTimersToLog = 50;
+static int32 DumpTimerLogResolveVirtualFunctions = 1;
+static FAutoConsoleVariableRef CVarDumpTimerLogResolveVirtualFunctions(
+	TEXT("TimerManager.DumpTimerLogResolveVirtualFunctions"), DumpTimerLogResolveVirtualFunctions,
+	TEXT("When logging timer info virtual functions will be resolved, if possible."),
+	ECVF_Default);
+
+static int32 DumpTimerLogSymbolNames = 1;
+static FAutoConsoleVariableRef CVarDumpTimerLogSymbolNames(
+	TEXT("TimerManager.DumpTimerLogSymbolNames"), DumpTimerLogSymbolNames,
+	TEXT("When logging timer info, symbol names will be included if set to 1."),
+	ECVF_Default);
+
+static int32 MaxExpiredTimersToLog = 30;
 static FAutoConsoleVariableRef CVarMaxExpiredTimersToLog(
 	TEXT("TimerManager.MaxExpiredTimersToLog"), 
 	MaxExpiredTimersToLog,
@@ -145,28 +158,78 @@ void FTimerManager::OnCrash()
 FString FTimerUnifiedDelegate::ToString() const
 {
 	const UObject* Object = nullptr;
-	FName FunctionName = NAME_None;
+	FString FunctionNameStr;
 	bool bDynDelegate = false;
 
 	if (FuncDelegate.IsBound())
 	{
+		FName FunctionName;
 #if USE_DELEGATE_TRYGETBOUNDFUNCTIONNAME
 		FunctionName = FuncDelegate.TryGetBoundFunctionName();
 #endif
+		if (FunctionName.IsNone())
+		{
+			void** VtableAddr = nullptr;
+#if PLATFORM_COMPILER_CLANG || defined(_MSC_VER)
+			// Add the vtable address
+			const void* UserObject = FuncDelegate.GetObjectForTimerManager();
+			if (UserObject)
+			{
+				VtableAddr = *(void***)UserObject;
+				FunctionNameStr = FString::Printf(TEXT("vtbl: %p"), VtableAddr);
+			}
+#endif // PLATFORM_COMPILER_CLANG
+
+			uint64 ProgramCounter = FuncDelegate.GetBoundProgramCounterForTimerManager();
+			if (ProgramCounter != 0)
+			{
+				// Add the function address
+
+#if PLATFORM_COMPILER_CLANG
+				// See if this is a virtual function. Heuristic is that real function addresses are higher than some value, and vtable offsets are lower
+				const uint64 MaxVTableAddressOffset = 32768;
+				if (DumpTimerLogResolveVirtualFunctions && VtableAddr && ProgramCounter > 0 && ProgramCounter < MaxVTableAddressOffset)
+				{
+					// If the ProgramCounter is just an offset to the vtable (virtual member function) then resolve the actual ProgramCounter here.
+					ProgramCounter = (uint64)VtableAddr[ProgramCounter / sizeof(void*)];
+				}
+#endif // PLATFORM_COMPILER_CLANG
+
+				FunctionNameStr += FString::Printf(TEXT(" func: 0x%llx"), ProgramCounter);
+
+				if (DumpTimerLogSymbolNames)
+				{
+					// Try to resolve the function address to a symbol
+					FProgramCounterSymbolInfo SymbolInfo;
+					SymbolInfo.FunctionName[0] = 0;
+					SymbolInfo.Filename[0] = 0;
+					SymbolInfo.LineNumber = 0;
+					FPlatformStackWalk::ProgramCounterToSymbolInfo(ProgramCounter, SymbolInfo);
+					FunctionNameStr += FString::Printf(TEXT(" %s [%s:%d]"), ANSI_TO_TCHAR(SymbolInfo.FunctionName), ANSI_TO_TCHAR(SymbolInfo.Filename), SymbolInfo.LineNumber);
+				}
+			}
+			else
+			{
+				FunctionNameStr = TEXT(" 0x0");
+			}
+		}
+		else
+		{
+			FunctionNameStr = FunctionName.ToString();
+		}
 	}
 	else if (FuncDynDelegate.IsBound())
 	{
 		Object = FuncDynDelegate.GetUObject();
-		FunctionName = FuncDynDelegate.GetFunctionName();
+		FunctionNameStr = FuncDynDelegate.GetFunctionName().ToString();
 		bDynDelegate = true;
 	}
 	else
 	{
-		static FName NotBoundName(TEXT("NotBound!"));
-		FunctionName = NotBoundName;
+		FunctionNameStr = TEXT("NotBound!");
 	}
 
-	return FString::Printf(TEXT("%s,%s,%s"), bDynDelegate ? TEXT("DYN DELEGATE") : TEXT("DELEGATE"), Object == nullptr ? TEXT("NO OBJ") : *Object->GetPathName(), *FunctionName.ToString());
+	return FString::Printf(TEXT("%s,%s,%s"), bDynDelegate ? TEXT("DYN DELEGATE") : TEXT("DELEGATE"), Object == nullptr ? TEXT("NO OBJ") : *Object->GetPathName(), *FunctionNameStr);
 }
 
 // ---------------------------------
