Always clear the need end of frame update to avoid issues when object recycling is occurring

diff --git a/Engine/Source/Runtime/Engine/Private/Components/ActorComponent.cpp b/Engine/Source/Runtime/Engine/Private/Components/ActorComponent.cpp
--- a/Engine/Source/Runtime/Engine/Private/Components/ActorComponent.cpp
+++ b/Engine/Source/Runtime/Engine/Private/Components/ActorComponent.cpp
@@ -1480,13 +1480,9 @@ void UActorComponent::MarkForNeededEndOfFrameUpdate()
 
 void UActorComponent::ClearNeedEndOfFrameUpdate_Internal()
 {
-	// If this is being garbage collected we don't really need to worry about clearing this
-	if (!HasAnyFlags(RF_BeginDestroyed) && !IsUnreachable())
+	if (UWorld* World = GetWorld())
 	{
-		if (UWorld* World = GetWorld())
-		{
-			World->ClearActorComponentEndOfFrameUpdate(this);
-		}
+		World->ClearActorComponentEndOfFrameUpdate(this);
 	}
 }
 
