Added stat to UWorld::FindTeleportSpot

diff --git a/Engine/Source/Runtime/Engine/Private/LevelActor.cpp b/Engine/Source/Runtime/Engine/Private/LevelActor.cpp
--- a/Engine/Source/Runtime/Engine/Private/LevelActor.cpp
+++ b/Engine/Source/Runtime/Engine/Private/LevelActor.cpp
@@ -787,6 +787,8 @@ APlayerController* UWorld::SpawnPlayActor(UPlayer* NewPlayer, ENetRole RemoteRol
 
 bool UWorld::FindTeleportSpot(const AActor* TestActor, FVector& TestLocation, FRotator TestRotation)
 {
+	QUICK_SCOPE_CYCLE_COUNTER(STAT_UWorld_FindTeleportSpot);
+
 	if( !TestActor || !TestActor->GetRootComponent() )
 	{
 		return true;
