Per static mesh LOD group setting for number of optional LODs

diff --git a/Engine/Source/Runtime/Engine/Private/PositionVertexBuffer.cpp b/Engine/Source/Runtime/Engine/Private/PositionVertexBuffer.cpp
--- a/Engine/Source/Runtime/Engine/Private/PositionVertexBuffer.cpp
+++ b/Engine/Source/Runtime/Engine/Private/PositionVertexBuffer.cpp
@@ -174,6 +174,11 @@ void FPositionVertexBuffer::SerializeMetaData(FArchive& Ar)
 	Ar << Stride << NumVertices;
 }
 
+void FPositionVertexBuffer::ClearMetaData()
+{
+	Stride = NumVertices = 0;
+}
+
 /**
 * Specialized assignment operator, only used when importing LOD's.  
 */
diff --git a/Engine/Source/Runtime/Engine/Private/Rendering/ColorVertexBuffer.cpp b/Engine/Source/Runtime/Engine/Private/Rendering/ColorVertexBuffer.cpp
--- a/Engine/Source/Runtime/Engine/Private/Rendering/ColorVertexBuffer.cpp
+++ b/Engine/Source/Runtime/Engine/Private/Rendering/ColorVertexBuffer.cpp
@@ -242,6 +242,11 @@ void FColorVertexBuffer::SerializeMetaData(FArchive& Ar)
 	Ar << Stride << NumVertices;
 }
 
+void FColorVertexBuffer::ClearMetaData()
+{
+	Stride = NumVertices = 0;
+}
+
 
 /** Export the data to a string, used for editor Copy&Paste. */
 void FColorVertexBuffer::ExportText(FString &ValueStr) const
diff --git a/Engine/Source/Runtime/Engine/Private/Rendering/StaticMeshVertexBuffer.cpp b/Engine/Source/Runtime/Engine/Private/Rendering/StaticMeshVertexBuffer.cpp
--- a/Engine/Source/Runtime/Engine/Private/Rendering/StaticMeshVertexBuffer.cpp
+++ b/Engine/Source/Runtime/Engine/Private/Rendering/StaticMeshVertexBuffer.cpp
@@ -245,6 +245,14 @@ void FStaticMeshVertexBuffer::SerializeMetaData(FArchive& Ar)
 	InitTangentAndTexCoordStrides();
 }
 
+void FStaticMeshVertexBuffer::ClearMetaData()
+{
+	NumTexCoords = NumVertices = 0;
+	bUseFullPrecisionUVs = !GVertexElementTypeSupport.IsSupported(VET_Half2);
+	bUseHighPrecisionTangentBasis = false;
+	TangentsStride = TexcoordStride = 0;
+}
+
 
 /**
 * Specialized assignment operator, only used when importing LOD's.
diff --git a/Engine/Source/Runtime/Engine/Private/StaticMesh.cpp b/Engine/Source/Runtime/Engine/Private/StaticMesh.cpp
--- a/Engine/Source/Runtime/Engine/Private/StaticMesh.cpp
+++ b/Engine/Source/Runtime/Engine/Private/StaticMesh.cpp
@@ -189,6 +189,18 @@ FArchive& operator<<(FArchive& Ar, FStaticMeshSection& Section)
 	return Ar;
 }
 
+int32 FStaticMeshLODResources::GetPlatformMinLODIdx(const ITargetPlatform* TargetPlatform, UStaticMesh* StaticMesh)
+{
+#if WITH_EDITOR
+	check(TargetPlatform && StaticMesh);
+	return StaticMesh->MinLOD.GetValueForPlatformIdentifiers(
+		TargetPlatform->GetPlatformInfo().PlatformGroupName,
+		TargetPlatform->GetPlatformInfo().VanillaPlatformName);
+#else
+	return 0;
+#endif
+}
+
 uint8 FStaticMeshLODResources::GenerateClassStripFlags(FArchive& Ar, UStaticMesh* OwnerStaticMesh, int32 Index)
 {
 #if WITH_EDITOR
@@ -202,9 +214,7 @@ uint8 FStaticMeshLODResources::GenerateClassStripFlags(FArchive& Ar, UStaticMesh
 	const bool bWantToStripLOD = Ar.IsCooking()
 		&& (CVarStripMinLodDataDuringCooking.GetValueOnAnyThread() != 0)
 		&& OwnerStaticMesh
-		&& OwnerStaticMesh->MinLOD.GetValueForPlatformIdentifiers(
-			Ar.CookingTarget()->GetPlatformInfo().PlatformGroupName,
-			Ar.CookingTarget()->GetPlatformInfo().VanillaPlatformName) > Index;
+		&& GetPlatformMinLODIdx(Ar.CookingTarget(), OwnerStaticMesh) > Index;
 
 	return (bWantToStripTessellation ? AdjacencyDataStripFlag : 0) |
 		(bWantToStripLOD ? MinLodDataStripFlag : 0);
@@ -279,6 +289,17 @@ bool FStaticMeshLODResources::IsLODInlined(const ITargetPlatform* TargetPlatform
 #endif
 }
 
+int32 FStaticMeshLODResources::GetNumOptionalLODsAllowed(const ITargetPlatform* TargetPlatform, UStaticMesh* StaticMesh)
+{
+#if WITH_EDITOR
+	check(TargetPlatform && StaticMesh);
+	const FStaticMeshLODGroup& LODGroupSettings = TargetPlatform->GetStaticMeshLODSettings().GetLODGroup(StaticMesh->LODGroup);
+	return LODGroupSettings.GetDefaultMaxNumOptionalLODs();
+#else
+	return 0;
+#endif
+}
+
 void FStaticMeshLODResources::AccumVertexBuffersSize(const FStaticMeshVertexBuffers& VertexBuffers, uint32& OutSize)
 {
 #if (WITH_EDITOR || DO_CHECK)
@@ -499,6 +520,22 @@ void FStaticMeshLODResources::SerializeAvailabilityInfo(FArchive& Ar)
 	}
 }
 
+void FStaticMeshLODResources::ClearAvailabilityInfo()
+{
+	DepthOnlyNumTriangles = 0;
+	bHasAdjacencyInfo = false;
+	bHasDepthOnlyIndices = false;
+	bHasReversedIndices = false;
+	bHasReversedDepthOnlyIndices = false;
+	bHasColorVertexData = false;
+	bHasWireframeIndices = false;
+	VertexBuffers.StaticMeshVertexBuffer.ClearMetaData();
+	VertexBuffers.PositionVertexBuffer.ClearMetaData();
+	VertexBuffers.ColorVertexBuffer.ClearMetaData();
+	delete AdditionalIndexBuffers;
+	AdditionalIndexBuffers = nullptr;
+}
+
 void FStaticMeshLODResources::Serialize(FArchive& Ar, UObject* Owner, int32 Index)
 {
 	DECLARE_SCOPE_CYCLE_COUNTER(TEXT("FStaticMeshLODResources::Serialize"), STAT_StaticMeshLODResources_Serialize, STATGROUP_LoadTime);
@@ -533,20 +570,30 @@ void FStaticMeshLODResources::Serialize(FArchive& Ar, UObject* Owner, int32 Inde
 #if WITH_EDITOR
 			if (Ar.IsSaving())
 			{
+				const int32 MaxNumOptionalLODs = GetNumOptionalLODsAllowed(Ar.CookingTarget(), OwnerStaticMesh);
+				const int32 OptionalLODIdx = GetPlatformMinLODIdx(Ar.CookingTarget(), OwnerStaticMesh) - Index;
+				const bool bDiscardBulkData = OptionalLODIdx > MaxNumOptionalLODs;
+
 				TArray<uint8> TmpBuff;
-				FMemoryWriter MemWriter(TmpBuff, true);
-				MemWriter.SetCookingTarget(Ar.CookingTarget());
-				MemWriter.SetByteSwapping(Ar.IsByteSwapping());
-				SerializeBuffers(MemWriter, OwnerStaticMesh, ClassDataStripFlags, TmpBuffersSize);
+				if (!bDiscardBulkData)
+				{
+					FMemoryWriter MemWriter(TmpBuff, true);
+					MemWriter.SetCookingTarget(Ar.CookingTarget());
+					MemWriter.SetByteSwapping(Ar.IsByteSwapping());
+					SerializeBuffers(MemWriter, OwnerStaticMesh, ClassDataStripFlags, TmpBuffersSize);
+				}
 
 				bIsOptionalLOD = bIsBelowMinLOD;
-				const uint32 BulkDataFlags = BULKDATA_Force_NOT_InlinePayload
+				const uint32 BulkDataFlags = (bDiscardBulkData ? 0 : BULKDATA_Force_NOT_InlinePayload)
 					| (bIsOptionalLOD ? BULKDATA_OptionalPayload : 0);
 				BulkData.SetBulkDataFlags(BulkDataFlags);
-				BulkData.Lock(LOCK_READ_WRITE);
-				void* BulkDataMem = BulkData.Realloc(TmpBuff.Num());
-				FMemory::Memcpy(BulkDataMem, TmpBuff.GetData(), TmpBuff.Num());
-				BulkData.Unlock();
+				if (TmpBuff.Num() > 0)
+				{
+					BulkData.Lock(LOCK_READ_WRITE);
+					void* BulkDataMem = BulkData.Realloc(TmpBuff.Num());
+					FMemory::Memcpy(BulkDataMem, TmpBuff.GetData(), TmpBuff.Num());
+					BulkData.Unlock();
+				}
 				BulkData.Serialize(Ar, Owner, Index);
 			}
 			else
@@ -566,6 +613,11 @@ void FStaticMeshLODResources::Serialize(FArchive& Ar, UObject* Owner, int32 Inde
 
 			Ar << TmpBuffersSize;
 			BuffersSize = TmpBuffersSize.CalcBuffersSize();
+
+			if (Ar.IsLoading() && bIsOptionalLOD && !BulkDataSize)
+			{
+				ClearAvailabilityInfo();
+			}
 		}
 	}
 }
@@ -1609,6 +1661,11 @@ void FStaticMeshLODSettings::ReadEntry(FStaticMeshLODGroup& Group, FString Entry
 	{
 		Group.DefaultMaxNumStreamedLODs = FMath::Max(Group.DefaultMaxNumStreamedLODs, 0);
 	}
+
+	if (FParse::Value(*Entry, TEXT("MaxNumOptionalLODs="), Group.DefaultMaxNumOptionalLODs))
+	{
+		Group.DefaultMaxNumOptionalLODs = FMath::Max(Group.DefaultMaxNumOptionalLODs, 0);
+	}
 	
 	int32 LocalSupportLODStreaming = 0;
 	if (FParse::Value(*Entry, TEXT("bSupportLODStreaming="), LocalSupportLODStreaming))
diff --git a/Engine/Source/Runtime/Engine/Public/Rendering/ColorVertexBuffer.h b/Engine/Source/Runtime/Engine/Public/Rendering/ColorVertexBuffer.h
--- a/Engine/Source/Runtime/Engine/Public/Rendering/ColorVertexBuffer.h
+++ b/Engine/Source/Runtime/Engine/Public/Rendering/ColorVertexBuffer.h
@@ -59,6 +59,8 @@ public:
 
 	void SerializeMetaData(FArchive& Ar);
 
+	void ClearMetaData();
+
 	/**
 	* Export the data to a string, used for editor Copy&Paste.
 	* The method must not be called if there is no data.
diff --git a/Engine/Source/Runtime/Engine/Public/Rendering/PositionVertexBuffer.h b/Engine/Source/Runtime/Engine/Public/Rendering/PositionVertexBuffer.h
--- a/Engine/Source/Runtime/Engine/Public/Rendering/PositionVertexBuffer.h
+++ b/Engine/Source/Runtime/Engine/Public/Rendering/PositionVertexBuffer.h
@@ -66,6 +66,8 @@ public:
 
 	void SerializeMetaData(FArchive& Ar);
 
+	void ClearMetaData();
+
 	/**
 	* Specialized assignment operator, only used when importing LOD's.
 	*/
diff --git a/Engine/Source/Runtime/Engine/Public/Rendering/StaticMeshVertexBuffer.h b/Engine/Source/Runtime/Engine/Public/Rendering/StaticMeshVertexBuffer.h
--- a/Engine/Source/Runtime/Engine/Public/Rendering/StaticMeshVertexBuffer.h
+++ b/Engine/Source/Runtime/Engine/Public/Rendering/StaticMeshVertexBuffer.h
@@ -178,6 +178,8 @@ public:
 
 	void SerializeMetaData(FArchive& Ar);
 
+	void ClearMetaData();
+
 	/**
 	* Specialized assignment operator, only used when importing LOD's.
 	*/
diff --git a/Engine/Source/Runtime/Engine/Public/StaticMeshResources.h b/Engine/Source/Runtime/Engine/Public/StaticMeshResources.h
--- a/Engine/Source/Runtime/Engine/Public/StaticMeshResources.h
+++ b/Engine/Source/Runtime/Engine/Public/StaticMeshResources.h
@@ -58,6 +58,7 @@ public:
 	FStaticMeshLODGroup()
 		: DefaultNumLODs(1)
 		, DefaultMaxNumStreamedLODs(0)
+		, DefaultMaxNumOptionalLODs(0)
 		, DefaultLightMapResolution(64)
 		, BasePercentTrianglesMult(1.0f)
 		, bSupportLODStreaming(false)
@@ -79,6 +80,12 @@ public:
 		return DefaultMaxNumStreamedLODs;
 	}
 
+	/** Returns the default maximum of optional LODs */
+	int32 GetDefaultMaxNumOptionalLODs() const
+	{
+		return DefaultMaxNumOptionalLODs;
+	}
+
 	/** Returns the default lightmap resolution. */
 	int32 GetDefaultLightMapResolution() const
 	{
@@ -108,6 +115,8 @@ private:
 	int32 DefaultNumLODs;
 	/** Maximum number of streamed LODs */
 	int32 DefaultMaxNumStreamedLODs;
+	/** Maximum number of optional LODs */
+	int32 DefaultMaxNumOptionalLODs;
 	/** Default lightmap resolution. */
 	int32 DefaultLightMapResolution;
 	/** An additional reduction of base meshes in this group. */
@@ -412,12 +421,16 @@ private:
 		}
 	};
 
+	static int32 GetPlatformMinLODIdx(const ITargetPlatform* TargetPlatform, UStaticMesh* StaticMesh);
+
 	static uint8 GenerateClassStripFlags(FArchive& Ar, UStaticMesh* OwnerStaticMesh, int32 Index);
 
 	static bool IsLODCookedOut(const ITargetPlatform* TargetPlatform, UStaticMesh* StaticMesh, bool bIsBelowMinLOD);
 
 	static bool IsLODInlined(const ITargetPlatform* TargetPlatform, UStaticMesh* StaticMesh, int32 LODIdx, bool bIsBelowMinLOD);
 
+	static int32 GetNumOptionalLODsAllowed(const ITargetPlatform* TargetPlatform, UStaticMesh* StaticMesh);
+
 	/** Compute the size of VertexBuffers and add the result to OutSize */
 	static void AccumVertexBuffersSize(const FStaticMeshVertexBuffers& VertexBuffers, uint32& OutSize);
 
@@ -436,6 +449,8 @@ private:
 	 */
 	void SerializeAvailabilityInfo(FArchive& Ar);
 
+	void ClearAvailabilityInfo();
+
 	template <bool bIncrement>
 	void UpdateIndexMemoryStats();
 
