Fix for subsystem collections keeping their outer alive.

diff --git a/Engine/Source/Runtime/Engine/Private/Subsystems/SubsystemCollection.cpp b/Engine/Source/Runtime/Engine/Private/Subsystems/SubsystemCollection.cpp
--- a/Engine/Source/Runtime/Engine/Private/Subsystems/SubsystemCollection.cpp
+++ b/Engine/Source/Runtime/Engine/Private/Subsystems/SubsystemCollection.cpp
@@ -154,6 +154,7 @@ void FSubsystemCollectionBase::Deinitialize()
 		}
 	}
 	SubsystemMap.Empty();
+	Outer = nullptr;
 }
 
 bool FSubsystemCollectionBase::InitializeDependency(TSubclassOf<USubsystem> SubsystemClass)
@@ -169,7 +170,6 @@ bool FSubsystemCollectionBase::InitializeDependency(TSubclassOf<USubsystem> Subs
 
 void FSubsystemCollectionBase::AddReferencedObjects(FReferenceCollector& Collector)
 {
-	Collector.AddReferencedObject(Outer);
 	Collector.AddReferencedObjects(SubsystemMap);
 }
 
