Fix crash on blendspace when sample is missing

diff --git a/Engine/Source/Runtime/Engine/Private/Animation/BlendSpaceBase.cpp b/Engine/Source/Runtime/Engine/Private/Animation/BlendSpaceBase.cpp
--- a/Engine/Source/Runtime/Engine/Private/Animation/BlendSpaceBase.cpp
+++ b/Engine/Source/Runtime/Engine/Private/Animation/BlendSpaceBase.cpp
@@ -11,6 +11,8 @@
 #include "UObject/FrameworkObjectVersion.h"
 #include "UObject/UObjectIterator.h"
 
+#define LOCTEXT_NAMESPACE "BlendSpaceBase"
+
 DECLARE_CYCLE_STAT(TEXT("BlendSpace GetAnimPose"), STAT_BlendSpace_GetAnimPose, STATGROUP_Anim);
 
 /** Scratch buffers for multithreaded usage */
@@ -37,6 +39,17 @@ void UBlendSpaceBase::PostLoad()
 #if WITH_EDITOR	
 	// Only do this during editor time (could alter the blendspace data during runtime otherwise) 
 	ValidateSampleData();
+#else
+	for (int32 SampleIndex = 0; SampleIndex < SampleData.Num(); ++SampleIndex)
+	{
+		FBlendSample& Sample = SampleData[SampleIndex];
+		if (!Sample.bIsValid)
+		{
+			UE_LOG(LogAnimation, Error, TEXT("[%s : %d] - Missing Sample Animation"), *GetFullName(), SampleIndex + 1);
+			SampleData.RemoveAt(SampleIndex);
+			--SampleIndex;
+		}
+	}
 #endif // WITH_EDITOR
 
 	InitializePerBoneBlend();
@@ -674,6 +687,7 @@ bool UBlendSpaceBase::GetSamplesFromBlendInput(const FVector &BlendInput, TArray
 			const int32 SampleDataIndex = GridElement.Indices[Ind];		
 			if( SampleData.IsValidIndex(SampleDataIndex) 
 #if WITH_EDITOR
+				&& SampleData[SampleDataIndex].bIsValid
 				&& SampleData[SampleDataIndex].Animation->GetSkeleton() == GetSkeleton()
 #endif // WITH_EDITOR
 				)
@@ -828,6 +842,24 @@ void UBlendSpaceBase::ValidateSampleData()
 				}
 			}
 		}		
+		else
+		{
+			if (IsRunningGame())
+			{
+				UE_LOG(LogAnimation, Error, TEXT("[%s : %d] - Missing Sample Animation"), *GetFullName(), SampleIndex + 1);
+			}
+			else
+			{
+				static FName NAME_LoadErrors("LoadErrors");
+				FMessageLog LoadErrors(NAME_LoadErrors);
+
+				TSharedRef<FTokenizedMessage> Message = LoadErrors.Error();
+				Message->AddToken(FTextToken::Create(LOCTEXT("EmptyAnimationData1", "The BlendSpace ")));
+				Message->AddToken(FAssetNameToken::Create(GetPathName(), FText::FromString(GetName())));
+				Message->AddToken(FTextToken::Create(LOCTEXT("EmptyAnimationData2", " has sample with no animation. Recommend to remove sample point or set new animation.")));
+				LoadErrors.Notify();
+			}
+		}
 	}
 
 	// set rotation blend in mesh space
