Fix for null pointer crash in Actor::Tick

#rnx

diff --git a/Engine/Source/Runtime/Engine/Private/Actor.cpp b/Engine/Source/Runtime/Engine/Private/Actor.cpp
--- a/Engine/Source/Runtime/Engine/Private/Actor.cpp
+++ b/Engine/Source/Runtime/Engine/Private/Actor.cpp
@@ -988,7 +988,10 @@ void AActor::Tick( float DeltaSeconds )
 		// DeltaSeconds to make up the frames that they missed (because they wouldn't have missed any).
 		// So pass in the world's DeltaSeconds value rather than our specific DeltaSeconds value.
 		UWorld* MyWorld = GetWorld();
-		MyWorld->GetLatentActionManager().ProcessLatentActions(this, MyWorld->GetDeltaSeconds());
+		if (MyWorld)
+		{
+			MyWorld->GetLatentActionManager().ProcessLatentActions(this, MyWorld->GetDeltaSeconds());
+		}
 	}
 
 	if (bAutoDestroyWhenFinished)
