Fix for replays sometimes getting stuck paused until you scrub again

diff --git a/Engine/Source/Runtime/Engine/Private/DemoNetDriver.cpp b/Engine/Source/Runtime/Engine/Private/DemoNetDriver.cpp
--- a/Engine/Source/Runtime/Engine/Private/DemoNetDriver.cpp
+++ b/Engine/Source/Runtime/Engine/Private/DemoNetDriver.cpp
@@ -3837,6 +3837,11 @@ void UDemoNetDriver::TickDemoPlayback(float DeltaSeconds)
 		PauseChannels(true);
 		return;
 	}
+	else
+	{
+		// we either have packets to process or data available to read
+		PauseChannels(false);
+	}
 
 	// Speculatively grab seconds now in case we need it to get the time it took to fast forward
 	const double FastForwardStartSeconds = FPlatformTime::Seconds();
