fixed: Crash shortly after launch on login screen. Added woarkaround for missing view SRVs on feature level  < SM5

diff --git a/Engine/Source/Runtime/Engine/Private/SceneManagement.cpp b/Engine/Source/Runtime/Engine/Private/SceneManagement.cpp
--- a/Engine/Source/Runtime/Engine/Private/SceneManagement.cpp
+++ b/Engine/Source/Runtime/Engine/Private/SceneManagement.cpp
@@ -733,8 +733,17 @@ FViewUniformShaderParameters::FViewUniformShaderParameters()
 	PreIntegratedBRDF = GWhiteTexture->TextureRHI;
 	PreIntegratedBRDFSampler = TStaticSamplerState<SF_Bilinear, AM_Clamp, AM_Clamp, AM_Clamp>::GetRHI();
 
-	PrimitiveSceneData = GIdentityPrimitiveBuffer.PrimitiveSceneDataBufferSRV;
-	LightmapSceneData = GIdentityPrimitiveBuffer.LightmapSceneDataBufferSRV;
+	if (IsFeatureLevelSupported(GMaxRHIShaderPlatform, ERHIFeatureLevel::SM5))
+	{
+		PrimitiveSceneData = GIdentityPrimitiveBuffer.PrimitiveSceneDataBufferSRV;
+		LightmapSceneData = GIdentityPrimitiveBuffer.LightmapSceneDataBufferSRV;
+	}
+	else
+	{
+		// these are not used on < SM5, set it to something to pass UB validation
+		PrimitiveSceneData = GNullColorVertexBuffer.VertexBufferSRV;
+		LightmapSceneData = GNullColorVertexBuffer.VertexBufferSRV;
+	}
 }
 
 FInstancedViewUniformShaderParameters::FInstancedViewUniformShaderParameters()
