Fix case where a delta checkpoint may have had no packets for a given fast forwarded level

diff --git a/Engine/Source/Runtime/Engine/Private/DemoNetDriver.cpp b/Engine/Source/Runtime/Engine/Private/DemoNetDriver.cpp
--- a/Engine/Source/Runtime/Engine/Private/DemoNetDriver.cpp
+++ b/Engine/Source/Runtime/Engine/Private/DemoNetDriver.cpp
@@ -4650,7 +4650,11 @@ bool UDemoNetDriver::FastForwardLevels(const FGotoResult& GotoResult)
 
 				if (bDeltaCheckpoint)
 				{
-					DeltaCheckpointPacketIntervals.Emplace(DeltaPacketStartIndex, ReadPacketsHelper.Packets.Num() - 1);
+					const int32 DeltaPacketEndIndex = ReadPacketsHelper.Packets.Num() - 1;
+					if (DeltaPacketEndIndex >= DeltaPacketStartIndex)
+					{
+						DeltaCheckpointPacketIntervals.Emplace(DeltaPacketStartIndex, DeltaPacketEndIndex);
+					}
 				}
 			} 
 			while (!CheckpointArchive->IsError() && (CheckpointArchive->Tell() < CheckpointArchive->TotalSize()));
@@ -4679,7 +4683,11 @@ bool UDemoNetDriver::FastForwardLevels(const FGotoResult& GotoResult)
 
 		if (bDeltaCheckpoint)
 		{
-			DeltaCheckpointPacketIntervals.Emplace(StreamPacketStartIndex, ReadPacketsHelper.Packets.Num() - 1);
+			const int32 StreamPacketEndIndex = ReadPacketsHelper.Packets.Num() - 1;
+			if (StreamPacketEndIndex >= StreamPacketStartIndex)
+			{
+				DeltaCheckpointPacketIntervals.Emplace(StreamPacketStartIndex, StreamPacketEndIndex);
+			}
 		}
 	}
 
@@ -5310,7 +5318,11 @@ bool UDemoNetDriver::LoadCheckpoint(const FGotoResult& GotoResult)
 
 		if (bDeltaCheckpoint)
 		{
-			DeltaCheckpointPacketIntervals.Emplace(DeltaPacketStartIndex, PlaybackPackets.Num() - 1);
+			const int32 DeltaPacketEndIndex = PlaybackPackets.Num() - 1;
+			if (DeltaPacketEndIndex >= DeltaPacketStartIndex)
+			{
+				DeltaCheckpointPacketIntervals.Emplace(DeltaPacketStartIndex, DeltaPacketEndIndex);
+			}
 		}
 	}
 	while (!GotoCheckpointArchive->IsError() && (GotoCheckpointArchive->Tell() < GotoCheckpointArchive->TotalSize()));
