Fix compile error

diff --git a/Engine/Source/Runtime/Engine/Private/Animation/BlendSpaceBase.cpp b/Engine/Source/Runtime/Engine/Private/Animation/BlendSpaceBase.cpp
--- a/Engine/Source/Runtime/Engine/Private/Animation/BlendSpaceBase.cpp
+++ b/Engine/Source/Runtime/Engine/Private/Animation/BlendSpaceBase.cpp
@@ -10,6 +10,8 @@
 #include "Animation/BlendSpaceUtilities.h"
 #include "UObject/FrameworkObjectVersion.h"
 #include "UObject/UObjectIterator.h"
+#include "Logging/TokenizedMessage.h"
+#include "Logging/MessageLog.h"
 
 #define LOCTEXT_NAMESPACE "BlendSpaceBase"
 
@@ -39,17 +41,6 @@ void UBlendSpaceBase::PostLoad()
 #if WITH_EDITOR	
 	// Only do this during editor time (could alter the blendspace data during runtime otherwise) 
 	ValidateSampleData();
-#else
-	for (int32 SampleIndex = 0; SampleIndex < SampleData.Num(); ++SampleIndex)
-	{
-		FBlendSample& Sample = SampleData[SampleIndex];
-		if (!Sample.bIsValid)
-		{
-			UE_LOG(LogAnimation, Error, TEXT("[%s : %d] - Missing Sample Animation"), *GetFullName(), SampleIndex + 1);
-			SampleData.RemoveAt(SampleIndex);
-			--SampleIndex;
-		}
-	}
 #endif // WITH_EDITOR
 
 	InitializePerBoneBlend();
