Fix compile error - add undef

diff --git a/Engine/Source/Runtime/Engine/Private/Animation/BlendSpaceBase.cpp b/Engine/Source/Runtime/Engine/Private/Animation/BlendSpaceBase.cpp
--- a/Engine/Source/Runtime/Engine/Private/Animation/BlendSpaceBase.cpp
+++ b/Engine/Source/Runtime/Engine/Private/Animation/BlendSpaceBase.cpp
@@ -1352,3 +1352,5 @@ TArray<FName>* UBlendSpaceBase::GetUniqueMarkerNames()
 {
 	return (SampleIndexWithMarkers != INDEX_NONE && SampleData.Num() > 0) ? SampleData[SampleIndexWithMarkers].Animation->GetUniqueMarkerNames() : nullptr;
 }
+
+#undef LOCTEXT_NAMESPACE 
\ No newline at end of file
