Fix a bug that prevents cooking LOD streaming enabled builds for more than one platform in one session

diff --git a/Engine/Source/Runtime/Engine/Private/StaticMesh.cpp b/Engine/Source/Runtime/Engine/Private/StaticMesh.cpp
--- a/Engine/Source/Runtime/Engine/Private/StaticMesh.cpp
+++ b/Engine/Source/Runtime/Engine/Private/StaticMesh.cpp
@@ -586,6 +586,8 @@ void FStaticMeshLODResources::Serialize(FArchive& Ar, UObject* Owner, int32 Inde
 				bIsOptionalLOD = bIsBelowMinLOD;
 				const uint32 BulkDataFlags = (bDiscardBulkData ? 0 : BULKDATA_Force_NOT_InlinePayload)
 					| (bIsOptionalLOD ? BULKDATA_OptionalPayload : 0);
+				const uint32 OldBulkDataFlags = BulkData.GetBulkDataFlags();
+				BulkData.ClearBulkDataFlags(0xffffffffu);
 				BulkData.SetBulkDataFlags(BulkDataFlags);
 				if (TmpBuff.Num() > 0)
 				{
@@ -595,6 +597,8 @@ void FStaticMeshLODResources::Serialize(FArchive& Ar, UObject* Owner, int32 Inde
 					BulkData.Unlock();
 				}
 				BulkData.Serialize(Ar, Owner, Index);
+				BulkData.ClearBulkDataFlags(0xffffffffu);
+				BulkData.SetBulkDataFlags(OldBulkDataFlags);
 			}
 			else
 #endif
