Ensure instead of crash and log out errors about newly spawned tick functions that were unable to be released

diff --git a/Engine/Source/Runtime/Engine/Private/TickTaskManager.cpp b/Engine/Source/Runtime/Engine/Private/TickTaskManager.cpp
--- a/Engine/Source/Runtime/Engine/Private/TickTaskManager.cpp
+++ b/Engine/Source/Runtime/Engine/Private/TickTaskManager.cpp
@@ -1081,7 +1081,20 @@ public:
 	void EndFrame()
 	{
 		bTickNewlySpawned = false;
-		check(!NewlySpawnedTickFunctions.Num()); // hmmm, this might be ok, but basically anything that was added this late cannot be ticked until the next frame
+#if DO_CHECK
+		// hmmm, this might be ok, but basically anything that was added this late cannot be ticked until the next frame
+		if (NewlySpawnedTickFunctions.Num())
+		{
+			for (TSet<FTickFunction*>::TIterator It(NewlySpawnedTickFunctions); It; ++It)
+			{
+				FTickFunction* TickFunction = *It;
+				UE_LOG(LogTick, Error, TEXT("Newly spawned tick function was added after tick groups have been run. Tick is %s."), *TickFunction->DiagnosticMessage());
+			}
+
+			ensureMsgf(false, TEXT("Newly spawned tick functions were added after tick groups have been run. See log for details."));
+			NewlySpawnedTickFunctions.Empty();
+		}
+#endif
 	}
 	// Interface that is private to FTickFunction
 
