Fix for some oneshot sounds retriggering when distance culled mid-play.

diff --git a/Engine/Source/Runtime/Engine/Private/ActiveSound.cpp b/Engine/Source/Runtime/Engine/Private/ActiveSound.cpp
--- a/Engine/Source/Runtime/Engine/Private/ActiveSound.cpp
+++ b/Engine/Source/Runtime/Engine/Private/ActiveSound.cpp
@@ -381,7 +381,7 @@ void FActiveSound::GetConcurrencyHandles(TArray<FConcurrencyHandle>& OutConcurre
 	}
 }
 
-void FActiveSound::UpdateWaveInstances( TArray<FWaveInstance*> &InWaveInstances, const float DeltaTime )
+void FActiveSound::UpdateWaveInstances(TArray<FWaveInstance*> &InWaveInstances, const float DeltaTime)
 {
 	check(AudioDevice);
 
@@ -414,39 +414,6 @@ void FActiveSound::UpdateWaveInstances( TArray<FWaveInstance*> &InWaveInstances,
 	// Cache the closest listener ptr 
 	ClosestListenerPtr = &Listeners[ClosestListenerIndex];
 
-	bool bPerformDistanceCheckOptimization = true;
-
-	// If we have an attenuation node, we can't know until we evaluate the sound cue if it's audio output going to be audible via a distance check
-	if (Sound->HasAttenuationNode() || 
-		(AudioDevice->VirtualSoundsEnabled() && (Sound->IsAllowedVirtual() || (bHandleSubtitles && bHasExternalSubtitles))) ||
-		!bHasAttenuationSettings ||
-		(bHasAttenuationSettings && (!AttenuationSettings.bAttenuate || AttenuationSettings.FocusDistanceScale != 1.0f || AttenuationSettings.NonFocusDistanceScale != 1.0f)))
-	{
-		bPerformDistanceCheckOptimization = false;
-	}
-	else
-	{
-		// Check the global focus settings... if it's doing a distance scale, we can't optimize due to distance
-		const FGlobalFocusSettings& FocusSettings = AudioDevice->GetGlobalFocusSettings();
-		if (FocusSettings.FocusDistanceScale != 1.0f || FocusSettings.NonFocusDistanceScale != 1.0f)
-		{
-			bPerformDistanceCheckOptimization = false;
-		}
-	}
-
-	// Early out if the sound is further away than we could possibly hear it, but only do this for non-virtualizable sounds.
-	if (bPerformDistanceCheckOptimization)
-	{
-		// The apparent max distance factors the actual max distance of the sound scaled with the distance scale due to focus effects
-		float ApparentMaxDistance = MaxDistance * FocusDistanceScale;
-
-		// Check if we're out of range of being audible, and early return if there's no chance of making sounds
-		if (!Sound->IsVirtualizeWhenSilent() && !AudioDevice->LocationIsAudible(Transform.GetLocation(), ApparentMaxDistance))
-		{
-			return;
-		}
-	}
-
 	FSoundParseParameters ParseParams;
 	ParseParams.Transform = Transform;
 	ParseParams.StartTime = RequestedStartTime;
diff --git a/Engine/Source/Runtime/Engine/Private/AudioDevice.cpp b/Engine/Source/Runtime/Engine/Private/AudioDevice.cpp
--- a/Engine/Source/Runtime/Engine/Private/AudioDevice.cpp
+++ b/Engine/Source/Runtime/Engine/Private/AudioDevice.cpp
@@ -127,6 +127,12 @@ FAutoConsoleVariableRef CVarDisableLegacyReverb(
 	TEXT("0: Enabled, 1: Disabled"),
 	ECVF_Default);
 
+static float SoundDistanceOptimizationLengthCVar = 1.0f;
+FAutoConsoleVariableRef CVarSoundDistanceOptimizationLength(
+	TEXT("au.SoundDistanceOptimizationLength"),
+	SoundDistanceOptimizationLengthCVar,
+	TEXT("The maximum duration a sound must be in order to be a candidate to be culled due to one-shot distance optimization.\n"),
+	ECVF_Default);
 
 /*-----------------------------------------------------------------------------
 FDynamicParameter implementation.
@@ -4112,13 +4118,21 @@ void FAudioDevice::AddNewActiveSound(const FActiveSound& NewActiveSound)
 	// Don't allow buses to try to play if we're not using the audio mixer.
 	if (!IsAudioMixerEnabled())
 	{
-		if (NewActiveSound.Sound)
+		USoundSourceBus* Bus = Cast<USoundSourceBus>(NewActiveSound.Sound);
+		if (Bus)
 		{
-			USoundSourceBus* Bus = Cast<USoundSourceBus>(NewActiveSound.Sound);
-			if (Bus)
-			{
-				return;
-			}
+			return;
+		}
+	}
+
+	if (NewActiveSound.Sound->GetDuration() <= FMath::Max(0.0f, SoundDistanceOptimizationLengthCVar))
+	{
+		// TODO: Determine if this check has already been completed at AudioComponent level and skip if so. Also,
+		// unify code paths determining if sound is audible.
+		if (!SoundIsAudible(NewActiveSound))
+		{
+			UE_LOG(LogAudio, Log, TEXT("New ActiveSound not created for out of range Sound %s"), *NewActiveSound.Sound->GetName());
+			return;
 		}
 	}
 
@@ -4588,6 +4602,51 @@ bool FAudioDevice::SoundIsAudible(USoundBase* Sound, const UWorld* World, const
 	return LocationIsAudible(Location, MaxDistance / DistanceScale);
 }
 
+bool FAudioDevice::SoundIsAudible(const FActiveSound& NewActiveSound)
+{
+	check(NewActiveSound.Sound);
+
+	// If we have an attenuation node, we can't know until we evaluate
+	// the sound cue if it's audio output going to be audible via a
+	// distance check. TODO: Check if this is still the case.
+	if (NewActiveSound.Sound->HasAttenuationNode())
+	{
+		return true;
+	}
+
+	if (VirtualSoundsEnabled() && (NewActiveSound.Sound->IsAllowedVirtual() || (NewActiveSound.bHandleSubtitles && NewActiveSound.bHasExternalSubtitles)))
+	{
+		return true;
+	}
+
+	const float ApparentMaxDistance = NewActiveSound.MaxDistance * NewActiveSound.FocusDistanceScale;
+	if (NewActiveSound.Sound->IsVirtualizeWhenSilent())
+	{
+		return true;
+	}
+
+	const FSoundAttenuationSettings& Attenuation = NewActiveSound.AttenuationSettings;
+	const bool bHasNonFocusDistanceScale = Attenuation.FocusDistanceScale != 1.0f || Attenuation.NonFocusDistanceScale != 1.0f;
+	if (!NewActiveSound.bHasAttenuationSettings ||
+		(NewActiveSound.bHasAttenuationSettings && (!Attenuation.bAttenuate || bHasNonFocusDistanceScale)))
+	{
+		return true;
+	}
+
+	const FGlobalFocusSettings& FocusSettings = GetGlobalFocusSettings();
+	if (FocusSettings.FocusDistanceScale != 1.0f || FocusSettings.NonFocusDistanceScale != 1.0f)
+	{
+		return true;
+	}
+
+	if (LocationIsAudible(NewActiveSound.Transform.GetLocation(), ApparentMaxDistance))
+	{
+		return true;
+	}
+
+	return false;
+}
+
 int32 FAudioDevice::FindClosestListenerIndex(const FTransform& SoundTransform, const TArray<FListener>& InListeners)
 {
 	int32 ClosestListenerIndex = 0;
@@ -4861,7 +4920,7 @@ UAudioComponent* FAudioDevice::CreateComponent(USoundBase* Sound, const FCreateC
 
 			bool bIsAudible = true;
 			// If a sound is a long duration, the position might change before sound finishes so assume it's audible
-			if (Params.bLocationSet && Sound->GetDuration() <= 1.0f)
+			if (Params.bLocationSet && Sound->GetDuration() <= FMath::Max(0.0f, SoundDistanceOptimizationLengthCVar))
 			{
 				float MaxDistance = 0.0f;
 				float FocusFactor = 0.0f;
diff --git a/Engine/Source/Runtime/Engine/Public/AudioDevice.h b/Engine/Source/Runtime/Engine/Public/AudioDevice.h
--- a/Engine/Source/Runtime/Engine/Public/AudioDevice.h
+++ b/Engine/Source/Runtime/Engine/Public/AudioDevice.h
@@ -728,7 +728,6 @@ public:
 		FVector Location;
 
 		void CommonInit();
-		void UpdateProximityRetriggerComponents(float DeltaTime);
 
 		friend class FAudioDevice;
 	};
@@ -1086,7 +1085,8 @@ public:
 	* @param AttenuationSettings	The (optional) attenuation settings the sound is using
 	* @param MaxDistance			The computed max distance of the sound.
 	* @param FocusFactor			The focus factor of the sound.
-	* @param Returns true if the sound is audible, false otherwise.
+	*
+	* @return Returns true if the sound is audible, false otherwise.
 	*/
 	bool SoundIsAudible(USoundBase* Sound, const UWorld* World, const FVector& Location, const FSoundAttenuationSettings* AttenuationSettingsToApply, float MaxDistance, float FocusFactor);
 
@@ -1300,6 +1300,14 @@ private:
 	 */
 	void StopQuietSoundsDueToMaxConcurrency(TArray<FWaveInstance*>& WaveInstances, TArray<FActiveSound*>& ActiveSoundsCopy);
 
+	/**
+	 * Checks if the given sound would be audible.
+	 * @param NewActiveSound	The ActiveSound attempting to be created
+	 * @return True if the sound is audible, false otherwise.
+	 */
+	bool SoundIsAudible(const FActiveSound& NewActiveSound);
+
+
 	/**
 	 * Set the mix for altering sound class properties
 	 *
