Deprecate AudioComponent retrigger system in favor of loop virtualization system
- Moves logic to Audio Thread
- Enable for all loops, not just on opt-in AudioComponents
- Removes distance check when sound is active
- Fixes bugs with calculating max distance for USoundCues that have overrides set

Swarm: https://p4-swarm.epicgames.net/reviews/5824804/

diff --git a/Engine/Plugins/Runtime/SoundMod/Source/SoundMod/Classes/SoundMod.h b/Engine/Plugins/Runtime/SoundMod/Source/SoundMod/Classes/SoundMod.h
--- a/Engine/Plugins/Runtime/SoundMod/Source/SoundMod/Classes/SoundMod.h
+++ b/Engine/Plugins/Runtime/SoundMod/Source/SoundMod/Classes/SoundMod.h
@@ -40,7 +40,6 @@ public:
 	//~ Begin USoundBase Interface.
 	virtual bool IsPlayable() const override;
 	virtual void Parse(class FAudioDevice* AudioDevice, const UPTRINT NodeWaveInstanceHash, FActiveSound& ActiveSound, const FSoundParseParameters& ParseParams, TArray<FWaveInstance*>& WaveInstances) override;
-	virtual float GetMaxDistance() const override;
 	//~ End USoundBase Interface.
 };
 
diff --git a/Engine/Plugins/Runtime/SoundMod/Source/SoundMod/Private/SoundMod.cpp b/Engine/Plugins/Runtime/SoundMod/Source/SoundMod/Private/SoundMod.cpp
--- a/Engine/Plugins/Runtime/SoundMod/Source/SoundMod/Private/SoundMod.cpp
+++ b/Engine/Plugins/Runtime/SoundMod/Source/SoundMod/Private/SoundMod.cpp
@@ -57,10 +57,3 @@ bool USoundMod::IsPlayable() const
 {
 	return true;
 }
-
-float USoundMod::GetMaxDistance() const
-{
-	return (AttenuationSettings ? AttenuationSettings->Attenuation.GetMaxDimension() : WORLD_MAX);
-}
-
-
diff --git a/Engine/Source/Runtime/Engine/Classes/Sound/SoundConcurrency.h b/Engine/Source/Runtime/Engine/Classes/Sound/SoundConcurrency.h
--- a/Engine/Source/Runtime/Engine/Classes/Sound/SoundConcurrency.h
+++ b/Engine/Source/Runtime/Engine/Classes/Sound/SoundConcurrency.h
@@ -213,7 +213,7 @@ public:
 	ENGINE_API ~FSoundConcurrencyManager();
 
 	/** Returns a newly allocated active sound given the input active sound struct. Will return nullptr if the active sound concurrency evaluation doesn't allow for it. */
-	FActiveSound* CreateNewActiveSound(const FActiveSound& NewActiveSound);
+	FActiveSound* CreateNewActiveSound(const FActiveSound& NewActiveSound, bool bIsRetriggering);
 
 	/** Removes the active sound from concurrency tracking when active sound is stopped. */
 	void StopActiveSound(FActiveSound* ActiveSound);
@@ -224,7 +224,7 @@ public:
 private: // Methods
 	/** Evaluates whether or not the sound can play given the concurrency group's rules. Appends permissible
 	sounds to evict in order for sound to play (if required) and returns the desired concurrency group. */
-	FConcurrencyGroup* CanPlaySound(const FActiveSound& NewActiveSound, const FConcurrencyGroupID GroupID, TArray<FActiveSound*>& OutSoundsToEvict);
+	FConcurrencyGroup* CanPlaySound(const FActiveSound& NewActiveSound, const FConcurrencyGroupID GroupID, TArray<FActiveSound*>& OutSoundsToEvict, bool bIsRetriggering);
 
 	/** Creates a new concurrency group and returns pointer to said group */
 	FConcurrencyGroup& CreateNewConcurrencyGroup(const FConcurrencyHandle& ConcurrencyHandle);
@@ -240,10 +240,10 @@ private: // Methods
 	FActiveSound* CreateAndEvictActiveSounds(const FActiveSound& NewActiveSound, const TArray<FConcurrencyGroup*>& GroupsToApply, const TArray<FActiveSound*>& SoundsToEvict);
 
 	/** Finds an active sound able to be evicted based on the provided concurrency settings. */
-	FActiveSound* GetEvictableSound(const FActiveSound& NewActiveSound, const FConcurrencyGroup& ConcurrencyGroup);
+	FActiveSound* GetEvictableSound(const FActiveSound& NewActiveSound, const FConcurrencyGroup& ConcurrencyGroup, bool bIsRetriggering);
 
 	/** Handles concurrency evaluation that happens per USoundConcurrencyObject */
-	FActiveSound* EvaluateConcurrency(const FActiveSound& NewActiveSound, TArray<FConcurrencyHandle>& ConcurrencyHandles);
+	FActiveSound* EvaluateConcurrency(const FActiveSound& NewActiveSound, TArray<FConcurrencyHandle>& ConcurrencyHandles, bool bIsRetriggering);
 
 private: // Data
 	/** Owning audio device ptr for the concurrency manager. */
diff --git a/Engine/Source/Runtime/Engine/Private/ActiveSound.cpp b/Engine/Source/Runtime/Engine/Private/ActiveSound.cpp
--- a/Engine/Source/Runtime/Engine/Private/ActiveSound.cpp
+++ b/Engine/Source/Runtime/Engine/Private/ActiveSound.cpp
@@ -1,6 +1,6 @@
 // Copyright 1998-2019 Epic Games, Inc. All Rights Reserved.
-
 #include "ActiveSound.h"
+
 #include "EngineDefines.h"
 #include "Misc/App.h"
 #include "AudioThread.h"
@@ -165,16 +165,12 @@ void FActiveSound::AddReferencedObjects( FReferenceCollector& Collector)
 
 void FActiveSound::SetWorld(UWorld* InWorld)
 {
-	check(IsInGameThread());
-
 	World = InWorld;
 	WorldID = (InWorld ? InWorld->GetUniqueID() : 0);
 }
 
 void FActiveSound::SetSound(USoundBase* InSound)
 {
-	check(IsInGameThread());
-
 	Sound = InSound;
 	bApplyInteriorVolumes = (SoundClassOverride && SoundClassOverride->Properties.bApplyAmbientVolumes)
 							|| (Sound && Sound->ShouldApplyInteriorVolumes());
@@ -182,16 +178,35 @@ void FActiveSound::SetSound(USoundBase* InSound)
 
 void FActiveSound::SetSoundClass(USoundClass* SoundClass)
 {
-	check(IsInGameThread());
-
 	SoundClassOverride = SoundClass;
 	bApplyInteriorVolumes = (SoundClassOverride && SoundClassOverride->Properties.bApplyAmbientVolumes)
 							|| (Sound && Sound->ShouldApplyInteriorVolumes());
 }
 
+void FActiveSound::ClearAudioComponent()
+{
+	AudioComponentID = 0;
+	AudioComponentUserID = NAME_None;
+	AudioComponentName = NAME_None;
+
+	OwnerID = 0;
+	OwnerName = NAME_None;
+}
+
+void FActiveSound::SetAudioComponent(const FActiveSound& ActiveSound)
+{
+	AudioComponentID = ActiveSound.AudioComponentID;
+	AudioComponentUserID = ActiveSound.AudioComponentUserID;
+	AudioComponentName = ActiveSound.AudioComponentName;
+
+	OwnerID = ActiveSound.OwnerID;
+	OwnerName = ActiveSound.OwnerName;
+}
+
 void FActiveSound::SetAudioComponent(UAudioComponent* Component)
 {
 	check(IsInGameThread());
+	check(Component);
 
 	AActor* Owner = Component->GetOwner();
 
@@ -507,7 +522,7 @@ void FActiveSound::UpdateWaveInstances(TArray<FWaveInstance*> &InWaveInstances,
 
 	if (bFinished)
 	{
-		AudioDevice->StopActiveSound(this);
+		AudioDevice->AddSoundToStop(this);
 	}
 	else if (ThisSoundsWaveInstances.Num() > 0)
 	{
diff --git a/Engine/Source/Runtime/Engine/Private/AudioDevice.cpp b/Engine/Source/Runtime/Engine/Private/AudioDevice.cpp
--- a/Engine/Source/Runtime/Engine/Private/AudioDevice.cpp
+++ b/Engine/Source/Runtime/Engine/Private/AudioDevice.cpp
@@ -1,6 +1,8 @@
 // Copyright 1998-2019 Epic Games, Inc. All Rights Reserved.
 
 #include "AudioDevice.h"
+
+#include "AudioVirtualLoop.h"
 #include "PhysicsEngine/BodyInstance.h"
 #include "Sound/SoundEffectPreset.h"
 #include "Sound/SoundEffectSubmix.h"
@@ -134,61 +136,6 @@ FAutoConsoleVariableRef CVarSoundDistanceOptimizationLength(
 	TEXT("The maximum duration a sound must be in order to be a candidate to be culled due to one-shot distance optimization.\n"),
 	ECVF_Default);
 
-/*-----------------------------------------------------------------------------
-FDynamicParameter implementation.
------------------------------------------------------------------------------*/
-
-FDynamicParameter::FDynamicParameter(float Value)
-	: CurrValue(Value)
-	, StartValue(Value)
-	, DeltaValue(0.0f)
-	, CurrTimeSec(0.0f)
-	, DurationSec(0.0f)
-	, LastTime(0.0f)
-	, TargetValue(Value)
-{}
-
-void FDynamicParameter::Set(float Value, float InDuration)
-{
-	if (TargetValue != Value || DurationSec != InDuration)
-	{
-		TargetValue = Value;
-		if (InDuration > 0.0f)
-		{
-			DeltaValue = Value - CurrValue;
-			StartValue = CurrValue;
-			DurationSec = InDuration;
-			CurrTimeSec = 0.0f;
-		}
-		else
-		{
-			StartValue = Value;
-			DeltaValue = 0.0f;
-			DurationSec = 0.0f;
-			CurrValue = Value;
-		}
-	}
-}
-
-void FDynamicParameter::Update(float DeltaTime)
-{
-	if (DurationSec > 0.0f)
-	{
-		float TimeFraction = CurrTimeSec / DurationSec;
-		if (TimeFraction < 1.0f)
-		{
-			CurrValue = DeltaValue * TimeFraction + StartValue;
-		}
-		else
-		{
-			CurrValue = StartValue + DeltaValue;
-			DurationSec = 0.0f;
-		}
-		CurrTimeSec += DeltaTime;
-	}
-}
-
-
 /*-----------------------------------------------------------------------------
 	FAudioDevice implementation.
 -----------------------------------------------------------------------------*/
@@ -2111,6 +2058,15 @@ void FAudioDevice::StopQuietSoundsDueToMaxConcurrency(TArray<FWaveInstance*>& Wa
 			if (ActiveSound->bShouldStopDueToMaxConcurrency)
 			{
 				AddSoundToStop(ActiveSound);
+				const bool bDoRangeCheck = false;
+				if (FAudioVirtualLoop* VirtualLoop = FAudioVirtualLoop::Virtualize(*this, *ActiveSound, bDoRangeCheck))
+				{
+					VirtualLoop->CalculateUpdateInterval(true);
+
+					// Clear must be called after AddSoundToStop to ensure AudioComponent is properly removed from AudioComponentIDToActiveSoundMap
+					ActiveSound->ClearAudioComponent();
+					AddVirtualLoop(*VirtualLoop);
+				}
 				ActiveSoundsCopy.RemoveAtSwap(i, 1, false);
 			}
 		}
@@ -2633,6 +2589,22 @@ void FAudioDevice::UpdateSoundClassProperties(float DeltaTime)
 	}
 }
 
+void FAudioDevice::VirtualizeInactiveLoops()
+{
+	const bool bDoRangeCheck = true;
+	for (FActiveSound* ActiveSound : ActiveSounds)
+	{
+		if (FAudioVirtualLoop* VirtualLoop = FAudioVirtualLoop::Virtualize(*this, *ActiveSound, bDoRangeCheck))
+		{
+			AddSoundToStop(ActiveSound);
+
+			// Clear must be called after AddSoundToStop to ensure AudioComponent is properly removed from AudioComponentIDToActiveSoundMap
+			ActiveSound->ClearAudioComponent();
+			AddVirtualLoop(*VirtualLoop);
+		}
+	}
+}
+
 float FListener::Interpolate(const double EndTime)
 {
 	if (FApp::GetCurrentTime() < InteriorStartTime)
@@ -3558,11 +3530,12 @@ void FAudioDevice::StopSources(TArray<FWaveInstance*>& WaveInstances, int32 Firs
 
 #if STATS
 	uint32 AudibleInactiveSounds = 0;
+	static const float MinReportVol = Audio::ConvertToDecibels(-60.f);
 	// Count how many sounds are not being played but were audible
 	for (int32 InstanceIndex = 0; InstanceIndex < FirstActiveIndex; InstanceIndex++)
 	{
 		FWaveInstance* WaveInstance = WaveInstances[InstanceIndex];
-		if (WaveInstance->GetVolumeWithDistanceAttenuation() > 0.1f)
+		if (WaveInstance->GetVolumeWithDistanceAttenuation() > MinReportVol)
 		{
 			AudibleInactiveSounds++;
 		}
@@ -3773,6 +3746,13 @@ void FAudioDevice::Update(bool bGameTicking)
 		UpdateAudioClock();
 	}
 
+	{
+		SCOPED_NAMED_EVENT(FAudioDevice_UpdateVirtualLoops, FColor::Blue);
+		// Update which loops should re-trigger due to coming back into proximity
+		// or allowed by concurrency re-evaluating in context of other sounds stopping
+		UpdateVirtualLoops();
+	}
+
 	// update if baked analysis is enabled
 	bIsBakedAnalysisEnabled = (BakedAnalysisEnabledCVar == 1);
 
@@ -3837,6 +3817,9 @@ void FAudioDevice::Update(bool bGameTicking)
 	// Gets the current state of the sound classes accounting for sound mix
 	UpdateSoundClassProperties(GetDeviceDeltaTime());
 
+	// Set looping ActiveSounds that are out-of-range to virtual and add to stop
+	VirtualizeInactiveLoops();
+
 	ProcessingPendingActiveSoundStops();
 
 	// Update listener transform
@@ -4091,6 +4074,11 @@ void FAudioDevice::StopAllSounds(bool bShouldStopUISounds)
 		}
 	}
 
+	for (TPair<FActiveSound*, FAudioVirtualLoop*>& Pair : VirtualLoops)
+	{
+		AddSoundToStop(Pair.Key);
+	}
+
 	// Immediately process stopping sounds
 	ProcessingPendingActiveSoundStops();
 }
@@ -4107,6 +4095,11 @@ void FAudioDevice::InitializePluginListeners(UWorld* World)
 }
 
 void FAudioDevice::AddNewActiveSound(const FActiveSound& NewActiveSound)
+{
+	AddNewActiveSoundInternal(NewActiveSound, nullptr);
+}
+
+void FAudioDevice::AddNewActiveSoundInternal(const FActiveSound& NewActiveSound, FAudioVirtualLoop* VirtualLoopToRetrigger)
 {
 	LLM_SCOPE(ELLMTag::AudioMisc);
 
@@ -4141,9 +4134,9 @@ void FAudioDevice::AddNewActiveSound(const FActiveSound& NewActiveSound)
 		DECLARE_CYCLE_STAT(TEXT("FAudioThreadTask.AddNewActiveSound"), STAT_AudioAddNewActiveSound, STATGROUP_AudioThreadCommands);
 
 		FAudioDevice* AudioDevice = this;
-		FAudioThread::RunCommandOnAudioThread([AudioDevice, NewActiveSound]()
+		FAudioThread::RunCommandOnAudioThread([AudioDevice, NewActiveSound, VirtualLoopToRetrigger]()
 		{
-			AudioDevice->AddNewActiveSound(NewActiveSound);
+			AudioDevice->AddNewActiveSoundInternal(NewActiveSound, VirtualLoopToRetrigger);
 		}, GET_STATID(STAT_AudioAddNewActiveSound));
 
 		return;
@@ -4183,6 +4176,18 @@ void FAudioDevice::AddNewActiveSound(const FActiveSound& NewActiveSound)
 		return;
 	}
 
+	// Determine if sound is loop and eligible for virtualize prior to creating "live" active sound in next Concurrency check step
+	if (!VirtualLoopToRetrigger)
+	{
+		const bool bDoRangeCheck = true;
+		VirtualLoopToRetrigger = FAudioVirtualLoop::Virtualize(*this, NewActiveSound, bDoRangeCheck);
+		if (VirtualLoopToRetrigger)
+		{
+			AddVirtualLoop(*VirtualLoopToRetrigger);
+			return;
+		}
+	}
+
 	// Evaluate concurrency. This will create an ActiveSound ptr which is a copy of NewActiveSound if the sound can play.
 	FActiveSound* ActiveSound = nullptr;
 
@@ -4190,11 +4195,21 @@ void FAudioDevice::AddNewActiveSound(const FActiveSound& NewActiveSound)
 		SCOPE_CYCLE_COUNTER(STAT_AudioEvaluateConcurrency);
 
 		// Try to create a new active sound. This returns nullptr if too many sounds are playing with this sound's concurrency setting
-		ActiveSound = ConcurrencyManager.CreateNewActiveSound(NewActiveSound);
+		ActiveSound = ConcurrencyManager.CreateNewActiveSound(NewActiveSound, VirtualLoopToRetrigger != nullptr);
 	}
 
 	if (!ActiveSound)
 	{
+		// Didn't pass concurrency, so attempt to add to re-trigger if applicable
+		if (!VirtualLoopToRetrigger)
+		{
+			const bool bDoRangeCheck = false;
+			VirtualLoopToRetrigger = FAudioVirtualLoop::Virtualize(*this, NewActiveSound, bDoRangeCheck);
+			if (VirtualLoopToRetrigger)
+			{
+				AddVirtualLoop(*VirtualLoopToRetrigger);
+			}
+		}
 		return;
 	}
 
@@ -4214,7 +4229,6 @@ void FAudioDevice::AddNewActiveSound(const FActiveSound& NewActiveSound)
 #endif // !(UE_BUILD_SHIPPING || UE_BUILD_TEST)
 
 	// Cull one-shot active sounds if we've reached our max limit of one shot active sounds before we attempt to evaluate concurrency
-
 	if (ActiveSound->IsOneShot())
 	{
 		OneShotCount++;
@@ -4242,6 +4256,19 @@ void FAudioDevice::AddNewActiveSound(const FActiveSound& NewActiveSound)
 	}
 #endif
 
+	// Loop has been successfully created, so add to stop before adding 'live' ActiveSound.
+	// Mark to not report playback complete on destruction as responsibility therein has been
+	// passed to newly created ActiveSound added below.  Add as stopping sound prior to adding
+	// new sound to ensure AudioComponentIDToActiveSoundMap is registered with the correct ActiveSound.
+	if (VirtualLoopToRetrigger)
+	{
+		FActiveSound& VirtualActiveSound = VirtualLoopToRetrigger->GetActiveSound();
+		AddSoundToStop(&VirtualActiveSound);
+
+		// Clear must be called after AddSoundToStop to ensure AudioComponent is properly removed from AudioComponentIDToActiveSoundMap
+		VirtualActiveSound.ClearAudioComponent();
+	}
+
 	ActiveSounds.Add(ActiveSound);
 	if (ActiveSound->GetAudioComponentID() > 0)
 	{
@@ -4249,6 +4276,48 @@ void FAudioDevice::AddNewActiveSound(const FActiveSound& NewActiveSound)
 	}
 }
 
+void FAudioDevice::RetriggerVirtualLoop(FAudioVirtualLoop& VirtualLoopToRetrigger)
+{
+	AddNewActiveSoundInternal(VirtualLoopToRetrigger.GetActiveSound(), &VirtualLoopToRetrigger);
+}
+
+void FAudioDevice::AddVirtualLoop(FAudioVirtualLoop& VirtualLoop)
+{
+	FActiveSound& ActiveSound = VirtualLoop.GetActiveSound();
+	check(!VirtualLoops.Contains(&ActiveSound));
+
+	const int64 ComponentID = ActiveSound.GetAudioComponentID();
+	if (ComponentID > 0)
+	{
+		check(!AudioComponentIDToActiveSoundMap.Contains(ComponentID));
+		AudioComponentIDToActiveSoundMap.Add(ComponentID, &ActiveSound);
+	}
+
+	VirtualLoops.Add(&ActiveSound, &VirtualLoop);
+}
+
+bool FAudioDevice::RemoveVirtualLoop(FActiveSound& ActiveSound)
+{
+	check(IsInAudioThread());
+
+	if (FAudioVirtualLoop* VirtualLoop = VirtualLoops.FindRef(&ActiveSound))
+	{
+		const uint64 ComponentID = ActiveSound.GetAudioComponentID();
+		if (ComponentID > 0)
+		{
+			UAudioComponent::PlaybackCompleted(ComponentID, false);
+		}
+
+		check(VirtualLoop);
+		delete VirtualLoop;
+		VirtualLoops.Remove(&ActiveSound);
+
+		return true;
+	}
+
+	return false;
+}
+
 void FAudioDevice::ProcessingPendingActiveSoundStops(bool bForceDelete)
 {
 	SCOPED_NAMED_EVENT(FAudioDevice_PendingActiveSoundStops, FColor::Blue);
@@ -4305,6 +4374,14 @@ void FAudioDevice::ProcessingPendingActiveSoundStops(bool bForceDelete)
 	{
 		check(ActiveSound);
 
+		// If the request was to stop an ActiveSound that
+		// is set to re-trigger but is not playing, remove
+		// and continue
+		if (RemoveVirtualLoop(*ActiveSound))
+		{
+			continue;
+		}
+
 		ActiveSound->Stop(bForceDelete);
 
 		USoundBase* Sound = ActiveSound->GetSound();
@@ -4333,27 +4410,36 @@ void FAudioDevice::ProcessingPendingActiveSoundStops(bool bForceDelete)
 	PendingSoundsToStop.Reset();
 }
 
-
 void FAudioDevice::AddSoundToStop(FActiveSound* SoundToStop)
 {
 	check(IsInAudioThread());
 	check(SoundToStop);
 
-	const uint64 AudioComponentID = SoundToStop->GetAudioComponentID();
-	if (AudioComponentID > 0)
-	{
-		AudioComponentIDToActiveSoundMap.Remove(AudioComponentID);
-	}
-
 	bool bAlreadyPending = false;
 	PendingSoundsToStop.Add(SoundToStop, &bAlreadyPending);
-
 	if (!bAlreadyPending)
 	{
-		ConcurrencyManager.StopActiveSound(SoundToStop);
+		const uint64 AudioComponentID = SoundToStop->GetAudioComponentID();
+		if (AudioComponentID > 0)
+		{
+			AudioComponentIDToActiveSoundMap.Remove(AudioComponentID);
+		}
+
+		if (!VirtualLoops.Contains(SoundToStop))
+		{
+			ConcurrencyManager.StopActiveSound(SoundToStop);
+		}
 	}
 }
 
+bool FAudioDevice::IsPendingStop(FActiveSound* ActiveSound)
+{
+	check(IsInAudioThread());
+	check(ActiveSound);
+
+	return PendingSoundsToStop.Contains(ActiveSound) || PendingSoundsToDelete.Contains(ActiveSound);
+}
+
 void FAudioDevice::StopActiveSound(const uint64 AudioComponentID)
 {
 	if (!IsInAudioThread())
@@ -4372,7 +4458,7 @@ void FAudioDevice::StopActiveSound(const uint64 AudioComponentID)
 	FActiveSound* ActiveSound = FindActiveSound(AudioComponentID);
 	if (ActiveSound)
 	{
-		StopActiveSound(ActiveSound);
+		AddSoundToStop(ActiveSound);
 	}
 }
 
@@ -4415,8 +4501,8 @@ void FAudioDevice::RemoveActiveSound(FActiveSound* ActiveSound)
 {
 	check(IsInAudioThread());
 
-	// Perform the notification
-	const int32 ComponentID = ActiveSound->GetAudioComponentID();
+	// Perform the notification if not sound not set to re-trigger
+	const uint64 ComponentID = ActiveSound->GetAudioComponentID();
 	if (ComponentID > 0)
 	{
 		UAudioComponent::PlaybackCompleted(ComponentID, false);
@@ -4581,7 +4667,7 @@ void FAudioDevice::GetMaxDistanceAndFocusFactor(USoundBase* Sound, const UWorld*
 	}
 }
 
-bool FAudioDevice::SoundIsAudible(USoundBase* Sound, const UWorld* World, const FVector& Location, const FSoundAttenuationSettings* AttenuationSettingsToApply, float MaxDistance, float FocusFactor)
+bool FAudioDevice::SoundIsAudible(USoundBase* Sound, const UWorld* World, const FVector& Location, const FSoundAttenuationSettings* AttenuationSettingsToApply, float MaxDistance, float FocusFactor) const
 {
 	check(IsInGameThread());
 
@@ -5127,6 +5213,11 @@ void FAudioDevice::Flush(UWorld* WorldToFlush, bool bClearActivatedReverb)
 		}
 	}
 
+	for (TPair<FActiveSound*, FAudioVirtualLoop*>& Pair : VirtualLoops)
+	{
+		AddSoundToStop(Pair.Key);
+	}
+
 	// Immediately stop all pending active sounds
 	ProcessingPendingActiveSoundStops(WorldToFlush == nullptr || WorldToFlush->bIsTearingDown);
 
@@ -5178,7 +5269,6 @@ void FAudioDevice::Flush(UWorld* WorldToFlush, bool bClearActivatedReverb)
 		// Make sure any in-flight audio rendering commands get executed.
 		FlushAudioRenderingCommands();
 	}
-
 }
 
 /**
@@ -5659,6 +5749,14 @@ void FAudioDevice::UpdateProximityRetriggerComponents()
 	}
 }
 
+void FAudioDevice::UpdateVirtualLoops()
+{
+	for (TPair<FActiveSound*, FAudioVirtualLoop*>& Pair : VirtualLoops)
+	{
+		check(Pair.Value);
+		Pair.Value->Update();
+	}
+}
 
 #if WITH_EDITOR
 void FAudioDevice::OnBeginPIE(const bool bIsSimulating)
diff --git a/Engine/Source/Runtime/Engine/Private/AudioDynamicParameter.cpp b/Engine/Source/Runtime/Engine/Private/AudioDynamicParameter.cpp
new file mode 100644
--- /dev/null
+++ b/Engine/Source/Runtime/Engine/Private/AudioDynamicParameter.cpp
@@ -0,0 +1,52 @@
+// Copyright 1998-2019 Epic Games, Inc. All Rights Reserved.
+#include "AudioDynamicParameter.h"
+
+FDynamicParameter::FDynamicParameter(float Value)
+	: CurrValue(Value)
+	, StartValue(Value)
+	, DeltaValue(0.0f)
+	, CurrTimeSec(0.0f)
+	, DurationSec(0.0f)
+	, LastTime(0.0f)
+	, TargetValue(Value)
+{}
+
+void FDynamicParameter::Set(float Value, float InDuration)
+{
+	if (TargetValue != Value || DurationSec != InDuration)
+	{
+		TargetValue = Value;
+		if (InDuration > 0.0f)
+		{
+			DeltaValue = Value - CurrValue;
+			StartValue = CurrValue;
+			DurationSec = InDuration;
+			CurrTimeSec = 0.0f;
+		}
+		else
+		{
+			StartValue = Value;
+			DeltaValue = 0.0f;
+			DurationSec = 0.0f;
+			CurrValue = Value;
+		}
+	}
+}
+
+void FDynamicParameter::Update(float DeltaTime)
+{
+	if (DurationSec > 0.0f)
+	{
+		float TimeFraction = CurrTimeSec / DurationSec;
+		if (TimeFraction < 1.0f)
+		{
+			CurrValue = DeltaValue * TimeFraction + StartValue;
+		}
+		else
+		{
+			CurrValue = StartValue + DeltaValue;
+			DurationSec = 0.0f;
+		}
+		CurrTimeSec += DeltaTime;
+	}
+}
\ No newline at end of file
diff --git a/Engine/Source/Runtime/Engine/Private/AudioVirtualLoop.cpp b/Engine/Source/Runtime/Engine/Private/AudioVirtualLoop.cpp
new file mode 100644
--- /dev/null
+++ b/Engine/Source/Runtime/Engine/Private/AudioVirtualLoop.cpp
@@ -0,0 +1,205 @@
+// Copyright 1998-2019 Epic Games, Inc. All Rights Reserved.
+#include "AudioVirtualLoop.h"
+
+#include "ActiveSound.h"
+#include "AudioDevice.h"
+#include "Sound/SoundBase.h"
+
+	static int32 bVirtualLoopsEnabledCVar = 1;
+	FAutoConsoleVariableRef CVarVirtualLoopsEnabled(
+		TEXT("au.VirtualLoops"),
+		bVirtualLoopsEnabledCVar,
+		TEXT("Enables or disables whether virtualizing is supported for audio loops.\n"),
+		ECVF_Default);
+
+	static float VirtualLoopsPerfDistanceCVar = 15000.0f;
+	FAutoConsoleVariableRef CVarVirtualLoopsPerfDistance(
+		TEXT("au.VirtualLoopsPerfDistance"),
+		VirtualLoopsPerfDistanceCVar,
+		TEXT("Sets virtual loop distance to scale update rate between min and max beyond max audible distance of sound.\n"),
+		ECVF_Default);
+
+	static float VirtualUpdateRateMinCVar = 0.1f;
+	FAutoConsoleVariableRef CVarVirtualLoopsUpdateRateMin(
+		TEXT("au.VirtualUpdateRateMin"),
+		VirtualUpdateRateMinCVar,
+		TEXT("Sets minimum rate to check if sound becomes audible again at sound's max audible distance.\n"),
+		ECVF_Default);
+
+	static float VirtualUpdateRateMaxCVar = 3.0f;
+	FAutoConsoleVariableRef CVarVirtualLoopsUpdateRateMax(
+		TEXT("au.VirtualUpdateRateMax"),
+		VirtualUpdateRateMinCVar,
+		TEXT("Sets maximum rate to check if sound becomes audible again (at beyond sound's max audible distance + perf scaling distance).\n"),
+		ECVF_Default);
+
+FAudioVirtualLoop::FAudioVirtualLoop(FAudioDevice& InAudioDevice, const FActiveSound& NewActiveSound)
+	: TimeSinceLastUpdate(0.0f)
+	, UpdateInterval(0.0f)
+	, AudioDevice(&InAudioDevice)
+	, ActiveSound(nullptr)
+{
+	SetActiveSound(NewActiveSound);
+}
+
+FAudioVirtualLoop::~FAudioVirtualLoop()
+{
+	check(ActiveSound);
+	delete(ActiveSound);
+}
+
+FAudioVirtualLoop* FAudioVirtualLoop::Virtualize(FAudioDevice& InAudioDevice, const FActiveSound& InActiveSound, bool bDoRangeCheck)
+{
+	if (!bVirtualLoopsEnabledCVar || InActiveSound.bIsPreviewSound || !InActiveSound.IsLooping())
+	{
+		return nullptr;
+	}
+
+	if (bDoRangeCheck && IsInAudibleRange(InAudioDevice, InActiveSound))
+	{
+		return nullptr;
+	}
+
+	return new FAudioVirtualLoop(InAudioDevice, InActiveSound);
+}
+
+void FAudioVirtualLoop::CalculateUpdateInterval(bool bIsAtMaxConcurrency)
+{
+	// If calculating due to being at max concurrency, set to max rate as
+	// sound will most likely be killed again on next check until concurrency
+	// is no longer full.  This limits starting and stopping of excess sounds
+	// virtualizing.
+	if (bIsAtMaxConcurrency)
+	{
+		UpdateInterval = VirtualUpdateRateMaxCVar;
+	}
+	else
+	{
+		const float DistanceToListener = AudioDevice->GetDistanceToNearestListener(ActiveSound->Transform.GetLocation());
+		const float DistanceRatio = (DistanceToListener - ActiveSound->MaxDistance) / FMath::Max(VirtualLoopsPerfDistanceCVar, 1.0f);
+		const float DistanceRatioClamped = FMath::Clamp(DistanceRatio, 0.0f, 1.0f);
+		UpdateInterval = FMath::Lerp(VirtualUpdateRateMinCVar, VirtualUpdateRateMaxCVar, DistanceRatioClamped);
+	}
+}
+
+FActiveSound& FAudioVirtualLoop::GetActiveSound()
+{
+	check(ActiveSound);
+	return *ActiveSound;
+}
+
+bool FAudioVirtualLoop::IsInAudibleRange(const FAudioDevice& InAudioDevice, const FActiveSound& InActiveSound)
+{
+	const FSoundAttenuationSettings* AttenuationSettingsToApply = InActiveSound.bHasAttenuationSettings ? &InActiveSound.AttenuationSettings : nullptr;
+
+	if (!InActiveSound.bAllowSpatialization)
+	{
+		return true;
+	}
+
+	float DistanceScale = 1.0f;
+	if (InActiveSound.bHasAttenuationSettings)
+	{
+		// If we are not using distance-based attenuation, this sound will be audible regardless of distance.
+		if (!AttenuationSettingsToApply->bAttenuate)
+		{
+			return true;
+		}
+
+		DistanceScale = AttenuationSettingsToApply->GetFocusDistanceScale(InAudioDevice.GetGlobalFocusSettings(), InActiveSound.FocusDistanceScale);
+	}
+
+	DistanceScale = FMath::Max(DistanceScale, 0.0001f);
+	const FVector Location = InActiveSound.Transform.GetLocation();
+	return InAudioDevice.LocationIsAudible(Location, InActiveSound.MaxDistance / DistanceScale);
+}
+
+void FAudioVirtualLoop::SetActiveSound(const FActiveSound& InActiveSound)
+{
+	check(!ActiveSound);
+
+	ActiveSound = new FActiveSound();
+
+	ActiveSound->bIsUISound = InActiveSound.bIsUISound;
+
+	ActiveSound->SetAudioComponent(InActiveSound);
+	ActiveSound->SetAudioDevice(AudioDevice);
+
+	if (GIsEditor)
+	{
+		ActiveSound->SetWorld(InActiveSound.GetWorld());
+	}
+
+	ActiveSound->SetSound(InActiveSound.GetSound());
+	ActiveSound->SetSoundClass(InActiveSound.GetSoundClass());
+
+	ActiveSound->ConcurrencySet = InActiveSound.ConcurrencySet;
+	ActiveSound->VolumeMultiplier = InActiveSound.VolumeMultiplier;
+	ActiveSound->Priority = InActiveSound.Priority;
+	ActiveSound->PitchMultiplier = InActiveSound.PitchMultiplier;
+	ActiveSound->bEnableLowPassFilter = InActiveSound.bEnableLowPassFilter;
+	ActiveSound->LowPassFilterFrequency = InActiveSound.LowPassFilterFrequency;
+	ActiveSound->RequestedStartTime = InActiveSound.RequestedStartTime;
+	ActiveSound->SubtitlePriority = InActiveSound.SubtitlePriority;
+	ActiveSound->bShouldRemainActiveIfDropped = InActiveSound.bShouldRemainActiveIfDropped;
+	ActiveSound->bHandleSubtitles = InActiveSound.bHandleSubtitles;
+	ActiveSound->bIgnoreForFlushing = InActiveSound.bIgnoreForFlushing;
+	ActiveSound->bIsUISound = InActiveSound.bIsUISound;
+	ActiveSound->bIsMusic = InActiveSound.bIsMusic;
+	ActiveSound->bAlwaysPlay = InActiveSound.bAlwaysPlay;
+	ActiveSound->bReverb = InActiveSound.bReverb;
+	ActiveSound->bCenterChannelOnly = InActiveSound.bCenterChannelOnly;
+	ActiveSound->bIsPreviewSound = InActiveSound.bIsPreviewSound;
+	ActiveSound->bLocationDefined = InActiveSound.bLocationDefined;
+	ActiveSound->bIsPaused = InActiveSound.bIsPaused;
+	ActiveSound->Transform = InActiveSound.Transform;
+	ActiveSound->bAllowSpatialization = InActiveSound.bAllowSpatialization;
+	ActiveSound->bHasAttenuationSettings = InActiveSound.bHasAttenuationSettings;
+	ActiveSound->AttenuationSettings = InActiveSound.AttenuationSettings;
+	ActiveSound->FocusPriorityScale = InActiveSound.FocusPriorityScale;
+	ActiveSound->FocusDistanceScale = InActiveSound.FocusDistanceScale;
+	ActiveSound->EnvelopeFollowerAttackTime = InActiveSound.EnvelopeFollowerAttackTime;
+	ActiveSound->EnvelopeFollowerReleaseTime = InActiveSound.EnvelopeFollowerReleaseTime;
+
+	ActiveSound->bUpdatePlayPercentage = InActiveSound.bUpdatePlayPercentage;
+	ActiveSound->bUpdateSingleEnvelopeValue = InActiveSound.bUpdateSingleEnvelopeValue;
+	ActiveSound->bUpdateMultiEnvelopeValue = InActiveSound.bUpdateMultiEnvelopeValue;
+
+	ActiveSound->bUpdatePlaybackTime = InActiveSound.bUpdatePlaybackTime;
+
+	ActiveSound->MaxDistance = InActiveSound.MaxDistance;
+	ActiveSound->InstanceParameters = InActiveSound.InstanceParameters;
+	ActiveSound->TargetAdjustVolumeMultiplier = InActiveSound.TargetAdjustVolumeMultiplier;
+	ActiveSound->CurrentAdjustVolumeMultiplier = InActiveSound.CurrentAdjustVolumeMultiplier;
+	ActiveSound->TargetAdjustVolumeStopTime = InActiveSound.TargetAdjustVolumeStopTime;
+}
+
+void FAudioVirtualLoop::Update()
+{
+	check(AudioDevice);
+	check(ActiveSound);
+
+	const float DeltaTime = AudioDevice->GetDeviceDeltaTime();
+
+	// Keep playback time up-to-date as it may be used to evaluate concurrency
+	ActiveSound->PlaybackTime += DeltaTime * ActiveSound->MinCurrentPitch;
+
+	if (UpdateInterval > 0.0f)
+	{
+		TimeSinceLastUpdate += DeltaTime;
+		if (UpdateInterval > TimeSinceLastUpdate)
+		{
+			return;
+		}
+		TimeSinceLastUpdate = 0.0f;
+	}
+
+	if (IsInAudibleRange(*AudioDevice, *ActiveSound))
+	{
+		AudioDevice->RetriggerVirtualLoop(*this);
+	}
+	else
+	{
+		CalculateUpdateInterval();
+	}
+}
diff --git a/Engine/Source/Runtime/Engine/Private/Components/AudioComponent.cpp b/Engine/Source/Runtime/Engine/Private/Components/AudioComponent.cpp
--- a/Engine/Source/Runtime/Engine/Private/Components/AudioComponent.cpp
+++ b/Engine/Source/Runtime/Engine/Private/Components/AudioComponent.cpp
@@ -21,12 +21,12 @@ FAutoConsoleVariableRef CVarBackedAnalysisTimeShift(
 	TEXT("Value: The time in seconds to shift the timeline."),
 	ECVF_Default);
 
-static int32 bRetriggerLoopOnProximityEnabledCVar = 1;
+static int32 bRetriggerLoopOnProximityEnabledCVar = 0;
 FAutoConsoleVariableRef CVarRetriggerLoopOnProximityEnabled(
 	TEXT("au.RetriggerLoopOnProximity"),
 	bRetriggerLoopOnProximityEnabledCVar,
 	TEXT("Enables or disables whether retriggering is supported on AudioComponent loops.\n")
-	TEXT("Value: The time in seconds to shift the timeline."),
+	TEXT("Value: Whether or not retriggering is supported on AudioComponent level."),
 	ECVF_Default);
 
 static float RetriggerLoopOnProximityDistanceCVar = 15000.0f;
diff --git a/Engine/Source/Runtime/Engine/Private/SoundBase.cpp b/Engine/Source/Runtime/Engine/Private/SoundBase.cpp
--- a/Engine/Source/Runtime/Engine/Private/SoundBase.cpp
+++ b/Engine/Source/Runtime/Engine/Private/SoundBase.cpp
@@ -75,7 +75,16 @@ const FSoundAttenuationSettings* USoundBase::GetAttenuationSettingsToApply() con
 
 float USoundBase::GetMaxDistance() const
 {
-	return (AttenuationSettings ? AttenuationSettings->Attenuation.GetMaxDimension() : WORLD_MAX);
+	if (AttenuationSettings)
+	{
+		FSoundAttenuationSettings& Settings = AttenuationSettings->Attenuation;
+		if (Settings.bAttenuate)
+		{
+			return Settings.GetMaxDimension();
+		}
+	}
+
+	return WORLD_MAX;
 }
 
 float USoundBase::GetDuration()
diff --git a/Engine/Source/Runtime/Engine/Private/SoundConcurrency.cpp b/Engine/Source/Runtime/Engine/Private/SoundConcurrency.cpp
--- a/Engine/Source/Runtime/Engine/Private/SoundConcurrency.cpp
+++ b/Engine/Source/Runtime/Engine/Private/SoundConcurrency.cpp
@@ -4,8 +4,10 @@
 #include "Components/AudioComponent.h"
 #include "ActiveSound.h"
 #include "AudioDevice.h"
+#include "AudioVirtualLoop.h"
 #include "Sound/SoundBase.h"
 
+
 // Forward Declarations
 struct FListener;
 
@@ -252,7 +254,7 @@ void FSoundConcurrencyManager::CreateNewGroupsFromHandles(
 	}
 }
 
-FActiveSound* FSoundConcurrencyManager::CreateNewActiveSound(const FActiveSound& NewActiveSound)
+FActiveSound* FSoundConcurrencyManager::CreateNewActiveSound(const FActiveSound& NewActiveSound, bool bIsRetriggering)
 {
 	check(NewActiveSound.GetSound());
 	check(IsInAudioThread());
@@ -276,7 +278,7 @@ FActiveSound* FSoundConcurrencyManager::CreateNewActiveSound(const FActiveSound&
 	}
 #endif
 
-	return EvaluateConcurrency(NewActiveSound, ConcurrencyHandles);
+	return EvaluateConcurrency(NewActiveSound, ConcurrencyHandles, bIsRetriggering);
 }
 
 FConcurrencyGroup& FSoundConcurrencyManager::CreateNewConcurrencyGroup(const FConcurrencyHandle& ConcurrencyHandle)
@@ -288,7 +290,7 @@ FConcurrencyGroup& FSoundConcurrencyManager::CreateNewConcurrencyGroup(const FCo
 	return *ConcurrencyGroups.FindRef(GroupID);
 }
 
-FConcurrencyGroup* FSoundConcurrencyManager::CanPlaySound(const FActiveSound& NewActiveSound, const FConcurrencyGroupID GroupID, TArray<FActiveSound*>& OutSoundsToEvict)
+FConcurrencyGroup* FSoundConcurrencyManager::CanPlaySound(const FActiveSound& NewActiveSound, const FConcurrencyGroupID GroupID, TArray<FActiveSound*>& OutSoundsToEvict, bool bIsRetriggering)
 {
 	check(GroupID != 0);
 	FConcurrencyGroup* ConcurrencyGroup = ConcurrencyGroups.FindRef(GroupID);
@@ -310,7 +312,7 @@ FConcurrencyGroup* FSoundConcurrencyManager::CanPlaySound(const FActiveSound& Ne
 	if (ConcurrencyGroup->IsFull())
 	{
 		// If no room for new sound, early out
-		if (FActiveSound* SoundToEvict = GetEvictableSound(NewActiveSound, *ConcurrencyGroup))
+		if (FActiveSound* SoundToEvict = GetEvictableSound(NewActiveSound, *ConcurrencyGroup, bIsRetriggering))
 		{
 			OutSoundsToEvict.AddUnique(SoundToEvict);
 		}
@@ -323,7 +325,7 @@ FConcurrencyGroup* FSoundConcurrencyManager::CanPlaySound(const FActiveSound& Ne
 	return ConcurrencyGroup;
 }
 
-FActiveSound* FSoundConcurrencyManager::GetEvictableSound(const FActiveSound& NewActiveSound, const FConcurrencyGroup& ConcurrencyGroup)
+FActiveSound* FSoundConcurrencyManager::GetEvictableSound(const FActiveSound& NewActiveSound, const FConcurrencyGroup& ConcurrencyGroup, bool bIsRetriggering)
 {
 	// Concurrency group isn't full so of course there's room
 	if (!ConcurrencyGroup.IsFull())
@@ -356,6 +358,12 @@ FActiveSound* FSoundConcurrencyManager::GetEvictableSound(const FActiveSound& Ne
 					EvictableSound = ActiveSound;
 				}
 			}
+
+			// Don't evict if attempting to re-trigger an older sound than that which is currently playing
+			if (bIsRetriggering && EvictableSound && NewActiveSound.PlaybackTime > EvictableSound->PlaybackTime)
+			{
+				EvictableSound = nullptr;
+			}
 		}
 		break;
 
@@ -379,8 +387,12 @@ FActiveSound* FSoundConcurrencyManager::GetEvictableSound(const FActiveSound& Ne
 							&& DistanceToActiveSoundSq == DistanceToStopSoundSq
 							&& (EvictableSound == nullptr || ActiveSound->PlaybackTime > EvictableSound->PlaybackTime))
 				{
-					DistanceToStopSoundSq = DistanceToActiveSoundSq;
-					EvictableSound = ActiveSound;
+					// Don't evict if attempting to re-trigger an older sound than that which is currently playing
+					if (!bIsRetriggering || ActiveSound->PlaybackTime > NewActiveSound.PlaybackTime)
+					{
+						DistanceToStopSoundSq = DistanceToActiveSoundSq;
+						EvictableSound = ActiveSound;
+					}
 				}
 			}
 		}
@@ -431,7 +443,7 @@ FActiveSound* FSoundConcurrencyManager::GetEvictableSound(const FActiveSound& Ne
 	return EvictableSound;
 }
 
-FActiveSound* FSoundConcurrencyManager::EvaluateConcurrency(const FActiveSound& NewActiveSound, TArray<FConcurrencyHandle>& ConcurrencyHandles)
+FActiveSound* FSoundConcurrencyManager::EvaluateConcurrency(const FActiveSound& NewActiveSound, TArray<FConcurrencyHandle>& ConcurrencyHandles, bool bIsRetriggering)
 {
 	check(NewActiveSound.GetSound());
 
@@ -446,7 +458,7 @@ FActiveSound* FSoundConcurrencyManager::EvaluateConcurrency(const FActiveSound&
 			{
 				if (FConcurrencyGroupID* ConcurrencyGroupID = ConcurrencyMap.Find(ConcurrencyHandle.ObjectID))
 				{
-					FConcurrencyGroup* ConcurrencyGroup = CanPlaySound(NewActiveSound, *ConcurrencyGroupID, SoundsToEvict);
+					FConcurrencyGroup* ConcurrencyGroup = CanPlaySound(NewActiveSound, *ConcurrencyGroupID, SoundsToEvict, bIsRetriggering);
 					if (!ConcurrencyGroup)
 					{
 						return nullptr;
@@ -462,7 +474,7 @@ FActiveSound* FSoundConcurrencyManager::EvaluateConcurrency(const FActiveSound&
 				{
 					if (FConcurrencyGroupID* ConcurrencyGroupID = ConcurrencyEntry->ConcurrencyObjectToConcurrencyGroup.Find(ConcurrencyHandle.ObjectID))
 					{
-						FConcurrencyGroup* ConcurrencyGroup = CanPlaySound(NewActiveSound, *ConcurrencyGroupID, SoundsToEvict);
+						FConcurrencyGroup* ConcurrencyGroup = CanPlaySound(NewActiveSound, *ConcurrencyGroupID, SoundsToEvict, bIsRetriggering);
 						if (!ConcurrencyGroup)
 						{
 							return nullptr;
@@ -482,7 +494,7 @@ FActiveSound* FSoundConcurrencyManager::EvaluateConcurrency(const FActiveSound&
 					check(Sound);
 					if (FConcurrencyGroupID* ConcurrencyGroupID = InstanceEntry->SoundInstanceToConcurrencyGroup.Find(Sound->GetUniqueID()))
 					{
-						FConcurrencyGroup* ConcurrencyGroup = CanPlaySound(NewActiveSound, *ConcurrencyGroupID, SoundsToEvict);
+						FConcurrencyGroup* ConcurrencyGroup = CanPlaySound(NewActiveSound, *ConcurrencyGroupID, SoundsToEvict, bIsRetriggering);
 						if (!ConcurrencyGroup)
 						{
 							return nullptr;
@@ -498,7 +510,7 @@ FActiveSound* FSoundConcurrencyManager::EvaluateConcurrency(const FActiveSound&
 				const FSoundObjectID SoundObjectID = NewActiveSound.GetSound()->GetUniqueID();
 				if (FConcurrencyGroupID* ConcurrencyGroupID = SoundObjectToConcurrencyGroup.Find(SoundObjectID))
 				{
-					FConcurrencyGroup* ConcurrencyGroup = CanPlaySound(NewActiveSound, *ConcurrencyGroupID, SoundsToEvict);
+					FConcurrencyGroup* ConcurrencyGroup = CanPlaySound(NewActiveSound, *ConcurrencyGroupID, SoundsToEvict, bIsRetriggering);
 					if (!ConcurrencyGroup)
 					{
 						return nullptr;
@@ -571,12 +583,48 @@ FActiveSound* FSoundConcurrencyManager::CreateAndEvictActiveSounds(const FActive
 		// Remove the active sound from the concurrency manager immediately so it doesn't count towards
 		// subsequent concurrency resolution checks (i.e. if sounds are triggered multiple times in this frame)
 		StopActiveSound(SoundToEvict);
+		if (AudioDevice->IsPendingStop(SoundToEvict))
+		{
+			continue;
+		}
+
+		TArray<FConcurrencyHandle> Handles;
+		SoundToEvict->GetConcurrencyHandles(Handles);
+
+		bool bAllowRetrigger = true;
+		for (const FConcurrencyHandle& Handle : Handles)
+		{
+			switch (Handle.Settings.ResolutionRule)
+			{
+				// Stop oldest resolution rules will cause an undesired
+				// cycling through re-triggerable loops from this frame to the next,
+				// so don't revive them.
+				case EMaxConcurrentResolutionRule::StopOldest:
+				case EMaxConcurrentResolutionRule::StopFarthestThenOldest:
+				{
+					bAllowRetrigger = false;
+					break;
+				}
+				default:
+				{
+					continue;
+				}
+			}
+		}
 
-		// Add this sound to list of sounds that need to stop but don't stop it immediately
 		AudioDevice->AddSoundToStop(SoundToEvict);
+		// If using a mode where we support re-triggering, attempt to re-trigger and update boolean state.
+		if (bAllowRetrigger)
+		{
+			const bool bDoRangeCheck = true;
+			if (FAudioVirtualLoop* VirtualLoop = FAudioVirtualLoop::Virtualize(*AudioDevice, *SoundToEvict, bDoRangeCheck))
+			{
+				SoundToEvict->ClearAudioComponent();
+				AudioDevice->AddVirtualLoop(*VirtualLoop);
+			}
+		}
 	}
 
-
 	return ActiveSound;
 }
 
diff --git a/Engine/Source/Runtime/Engine/Private/SoundCue.cpp b/Engine/Source/Runtime/Engine/Private/SoundCue.cpp
--- a/Engine/Source/Runtime/Engine/Private/SoundCue.cpp
+++ b/Engine/Source/Runtime/Engine/Private/SoundCue.cpp
@@ -79,6 +79,15 @@ void USoundCue::CacheAggregateValues()
 		Duration = FirstNode->GetDuration();
 
 		MaxDistance = FirstNode->GetMaxDistance();
+
+		if (const FSoundAttenuationSettings* Settings = GetAttenuationSettingsToApply())
+		{
+			if (Settings->bAttenuate)
+			{
+				MaxDistance = FMath::Max(MaxDistance, Settings->GetMaxDimension());
+			}
+		}
+
 		// If no sound cue nodes overrode the max distance, we need to check the base attenuation
 		if (MaxDistance == 0.0f)
 		{
diff --git a/Engine/Source/Runtime/Engine/Public/ActiveSound.h b/Engine/Source/Runtime/Engine/Public/ActiveSound.h
--- a/Engine/Source/Runtime/Engine/Public/ActiveSound.h
+++ b/Engine/Source/Runtime/Engine/Public/ActiveSound.h
@@ -7,12 +7,13 @@
 #include "Sound/SoundAttenuation.h"
 #include "HAL/ThreadSafeBool.h"
 #include "Audio.h"
+#include "AudioDevice.h"
+#include "AudioDynamicParameter.h"
 #include "Sound/SoundConcurrency.h"
 #include "Components/AudioComponent.h"
 #include "Sound/AudioVolume.h"
 #include "Sound/SoundSubmix.h"
 #include "Sound/SoundSourceBus.h"
-#include "AudioDevice.h"
 
 class FAudioDevice;
 class USoundBase;
@@ -223,6 +224,8 @@ public:
 
 	uint64 GetAudioComponentID() const { return AudioComponentID; }
 	FName GetAudioComponentUserID() const { return AudioComponentUserID; }
+	void ClearAudioComponent();
+	void SetAudioComponent(const FActiveSound& ActiveSound);
 	void SetAudioComponent(UAudioComponent* Component);
 	void SetOwner(AActor* Owner);
 	FString GetAudioComponentName() const;
@@ -231,8 +234,7 @@ public:
 	uint32 GetWorldID() const { return WorldID; }
 	TWeakObjectPtr<UWorld> GetWeakWorld() const { return World; }
 	UWorld* GetWorld() const 
-	{ 
-		check(IsInGameThread()); 
+	{
 		return World.Get();
 	}
 	void SetWorld(UWorld* World);
diff --git a/Engine/Source/Runtime/Engine/Public/AudioDevice.h b/Engine/Source/Runtime/Engine/Public/AudioDevice.h
--- a/Engine/Source/Runtime/Engine/Public/AudioDevice.h
+++ b/Engine/Source/Runtime/Engine/Public/AudioDevice.h
@@ -5,65 +5,54 @@
 #include "CoreMinimal.h"
 #include "Modules/ModuleInterface.h"
 #include "Engine/Engine.h"
-#include "Sound/SoundAttenuation.h"
 #include "Components/AudioComponent.h"
-#include "Sound/SoundClass.h"
 #include "Audio.h"
+#include "AudioDynamicParameter.h"
+#include "Sound/AudioSettings.h"
 #include "Sound/AudioVolume.h"
+#include "Sound/SoundAttenuation.h"
+#include "Sound/SoundClass.h"
 #include "Sound/SoundConcurrency.h"
 #include "Sound/SoundMix.h"
 #include "Sound/SoundSourceBus.h"
-#include "Sound/AudioSettings.h"
 #include "AudioDeviceManager.h"
 #include "DSP/SpectrumAnalyzer.h"
 #include "EngineGlobals.h"
 
+/**
+ * Forward declares
+ */
+
+class FArchive;
 class FAudioEffectsManager;
+class FAudioVirtualLoop;
 class FCanvas;
+class FOutputDevice;
+class FReferenceCollector;
+class FSoundBuffer;
 class FViewport;
 class FViewportClient;
+class IAudioSpatialization;
 class ICompressedAudioInfo;
 class UReverbEffect;
+class USoundAttenuation;
 class USoundBase;
+class USoundClass;
+class USoundConcurrency;
 class USoundEffectSourcePreset;
 class USoundEffectSubmixPreset;
+class USoundMix;
 class USoundSubmix;
 class USoundSourceBus;
 class USoundWave;
-struct FActiveSound;
-struct FAudioQualitySettings;
-
-/**
- * Forward declares
- */
-class USoundClass;
 class UWorld;
-class FOutputDevice;
-class FArchive;
-class FReferenceCollector;
-struct FWaveInstance;
-class USoundWave;
-class FSoundBuffer;
-class USoundBase;
-class USoundAttenuation;
-struct FRotator;
-struct FActiveSound;
-class USoundMix;
-
-class FAudioEffectsManager;
-class FViewportClient;
-class ICompressedAudioInfo;
-class IAudioSpatialization;
-class UReverbEffect;
-class USoundConcurrency;
-class FViewport;
-class FViewportClient;
-class FCanvas;
 
+struct FActiveSound;
 struct FAudioComponentParam;
 struct FAudioQualitySettings;
-class USoundEffectSubmixPreset;
-class USoundEffectSourcePreset;
+struct FRotator;
+struct FWaveInstance;
+
 
 /** 
  * Debug state of the audio system
@@ -142,38 +131,6 @@ namespace ERequestedAudioStats
 	static const uint8 LongSoundNames = 0x20;
 };
 
-/** Simple class that wraps the math involved with interpolating a parameter over time based on audio device update time. */
-class ENGINE_API FDynamicParameter
-{
-public:
-	explicit FDynamicParameter(float Value);
-
-	void Set(float Value, float InDuration);
-	void Update(float DeltaTime);
-
-	bool IsDone() const
-	{
-		return CurrTimeSec >= DurationSec;
-	}
-	float GetValue() const
-	{
-		return CurrValue;
-	}
-	float GetTargetValue() const
-	{
-		return TargetValue;
-	}
-
-private:
-	float CurrValue;
-	float StartValue;
-	float DeltaValue;
-	float CurrTimeSec;
-	float DurationSec;
-	float LastTime;
-	float TargetValue;
-};
-
 /** 
  * Defines the properties of the listener
  */
@@ -770,13 +727,18 @@ public:
 	 */
 	void AddNewActiveSound(const FActiveSound& ActiveSound);
 
+	/**
+	 * Attempts to retrigger a provided loop
+	 */
+	void RetriggerVirtualLoop(FAudioVirtualLoop& VirtualLoop);
+
 	/**
 	 * Removes the active sound for the specified audio component
 	 */
 	void StopActiveSound(uint64 AudioComponentID);
 
 	/**
-	* Stops the active sound
+	* (Deprecated in favor of AddSoundToStop). Stops the active sound
 	*/
 	void StopActiveSound(FActiveSound* ActiveSound);
 
@@ -1088,7 +1050,7 @@ public:
 	*
 	* @return Returns true if the sound is audible, false otherwise.
 	*/
-	bool SoundIsAudible(USoundBase* Sound, const UWorld* World, const FVector& Location, const FSoundAttenuationSettings* AttenuationSettingsToApply, float MaxDistance, float FocusFactor);
+	bool SoundIsAudible(USoundBase* Sound, const UWorld* World, const FVector& Location, const FSoundAttenuationSettings* AttenuationSettingsToApply, float MaxDistance, float FocusFactor) const;
 
 	/** Returns the index of the listener closest to the given sound transform */
 	static int32 FindClosestListenerIndex(const FTransform& SoundTransform, const TArray<FListener>& InListeners);
@@ -1282,6 +1244,12 @@ protected:
 
 private:
 
+	/**
+	 * Adds an active sound to the audio device. Can be a new active sound or one provided by the re-triggering
+	 * loop system.
+	 */
+	void AddNewActiveSoundInternal(const FActiveSound& ActiveSound, FAudioVirtualLoop* VirtualLoop);
+
 	/**
 	* Initializes all plugin listeners belonging to this audio device.
 	* Called in the game thread.
@@ -1382,6 +1350,8 @@ private:
 	*/
 	void UpdateSoundClassProperties(float DeltaTime);
 
+	void VirtualizeInactiveLoops();
+
 	/**
 	 * Recursively apply an adjuster to the passed in sound class and all children of the sound class
 	 * 
@@ -1497,9 +1467,18 @@ public:
 	/** Creates a new platform specific sound source */
 	virtual FSoundSource* CreateSoundSource() = 0;
 
-	void AddSoundToStop(struct FActiveSound* SoundToStop);
+	/**
+	 * Marks a sound to be stopped.  Returns true if added to stop,
+	 * false if already pending stop.
+	 */
+	void AddSoundToStop(FActiveSound* SoundToStop);
+
+	/**
+	 * Whether the provided ActiveSound is currently pending to stop
+	 */
+	bool IsPendingStop(FActiveSound* ActiveSound);
 
-	/**   
+	/**
 	* Gets the direction of the given position vector transformed relative to listener.   
 	* @param Position				Input position vector to transform relative to listener
 	* @param OutDistance			Optional output of distance from position to listener
@@ -1515,6 +1494,8 @@ public:
 
 	void UpdateProximityRetriggerComponents();
 
+	void UpdateVirtualLoops();
+
 	/** Sets the update delta time for the audio frame */
 	virtual void UpdateDeviceDeltaTime()
 	{
@@ -1542,6 +1523,8 @@ private:
 	/** Returns the number of frames to use per precache buffer. */
 	int32 GetNumPrecacheFrames() const;
 
+	bool RemoveVirtualLoop(FActiveSound& ActiveSound);
+
 public:
 
 	/** Query if the editor is in VR Preview for the current play world. Returns false for non-editor builds */
@@ -1550,6 +1533,8 @@ public:
 	/** Returns the audio clock of the audio device. Not supported on all platforms. */
 	double GetAudioClock() const { return AudioClock; }
 
+	void AddVirtualLoop(FAudioVirtualLoop& VirtualLoop);
+
 #if !UE_BUILD_SHIPPING
 	void DumpActiveSounds() const;
 
@@ -1835,6 +1820,11 @@ private:
 	TArray<USoundWave*> PrecachingSoundWaves;
 	TArray<FWaveInstance*> ActiveWaveInstances;
 
+	/** Array of dormant loops stopped due to proximity/applicable concurrency rules
+	  * that can be retriggered.
+	  */
+	TMap<FActiveSound*, FAudioVirtualLoop*> VirtualLoops;
+
 	TArray<FAudioComponentPtr> ProximityRetriggerComponents;
 
 	/** Cached copy of sound class adjusters array. Cached to avoid allocating every frame. */
diff --git a/Engine/Source/Runtime/Engine/Public/AudioDynamicParameter.h b/Engine/Source/Runtime/Engine/Public/AudioDynamicParameter.h
new file mode 100644
--- /dev/null
+++ b/Engine/Source/Runtime/Engine/Public/AudioDynamicParameter.h
@@ -0,0 +1,38 @@
+// Copyright 1998-2019 Epic Games, Inc. All Rights Reserved.
+#pragma once
+
+#include "CoreMinimal.h"
+
+/** Simple class that wraps the math involved with interpolating a parameter
+  * over time based on audio device update time.
+  */
+class ENGINE_API FDynamicParameter
+{
+public:
+	explicit FDynamicParameter(float Value);
+
+	void Set(float Value, float InDuration);
+	void Update(float DeltaTime);
+
+	bool IsDone() const
+	{
+		return CurrTimeSec >= DurationSec;
+	}
+	float GetValue() const
+	{
+		return CurrValue;
+	}
+	float GetTargetValue() const
+	{
+		return TargetValue;
+	}
+
+private:
+	float CurrValue;
+	float StartValue;
+	float DeltaValue;
+	float CurrTimeSec;
+	float DurationSec;
+	float LastTime;
+	float TargetValue;
+};
diff --git a/Engine/Source/Runtime/Engine/Public/AudioVirtualLoop.h b/Engine/Source/Runtime/Engine/Public/AudioVirtualLoop.h
new file mode 100644
--- /dev/null
+++ b/Engine/Source/Runtime/Engine/Public/AudioVirtualLoop.h
@@ -0,0 +1,54 @@
+// Copyright 1998-2019 Epic Games, Inc. All Rights Reserved.
+#pragma once
+
+#include "CoreMinimal.h"
+#include "ActiveSound.h"
+
+class FAudioDevice;
+
+
+/**
+ * Class that tracks virtualized looping active sounds that are eligible to revive re-trigger
+ * as long as no stop request is received from the game thread.
+ */
+class ENGINE_API FAudioVirtualLoop
+{
+	float TimeSinceLastUpdate;
+	float UpdateInterval;
+	FAudioDevice* AudioDevice;
+	FActiveSound* ActiveSound;
+
+	FAudioVirtualLoop(FAudioDevice& AudioDevice, const FActiveSound& NewActiveSound);
+
+public:
+	~FAudioVirtualLoop();
+
+	/**
+	 * Checks if provided active sound is available to be virtualized.  If so, returns new virtual loop ready to be
+	 * added to virtual loop management in audio device.
+	 */
+	static FAudioVirtualLoop* Virtualize(FAudioDevice& InAudioDevice, const FActiveSound& InActiveSound, bool bDoRangeCheck);
+
+	/**
+	 * Returns the internally-managed active sound.
+	 */
+	FActiveSound& GetActiveSound();
+
+	/**
+	 * Check if provided active sound is in audible range.
+	 */
+	static bool IsInAudibleRange(const FAudioDevice& InAudioDevice, const FActiveSound& ActiveSound);
+
+	/**
+	 * Overrides the update interval to the provided length.
+	 */
+	void CalculateUpdateInterval(bool bIsAtMaxConcurrency = false);
+
+	/**
+	  * Updates the loop and checks if ready to play (or 'realize').
+	  */
+	void Update();
+
+private:
+	void SetActiveSound(const FActiveSound& InActiveSound);
+};
