[LLM] tagging for UTexture. Moves about 100k out of UObject, AsyncLoading and MaterialInstance in my test replay.

diff --git a/Engine/Source/Runtime/Engine/Private/Texture.cpp b/Engine/Source/Runtime/Engine/Private/Texture.cpp
--- a/Engine/Source/Runtime/Engine/Private/Texture.cpp
+++ b/Engine/Source/Runtime/Engine/Private/Texture.cpp
@@ -7,6 +7,7 @@
 #include "Materials/Material.h"
 #include "Misc/ConfigCacheIni.h"
 #include "Misc/FeedbackContext.h"
+#include "HAL/LowLevelMemTracker.h"
 #include "UObject/UObjectIterator.h"
 #include "TextureResource.h"
 #include "Engine/Texture2D.h"
@@ -113,6 +114,7 @@ void UTexture::UpdateResource()
 		Resource = CreateResource();
 		if( Resource )
 		{
+			LLM_SCOPE(ELLMTag::Textures);
 			BeginInitResource(Resource);
 		}
 	}
