Adding custom data to the per prim uniform buffer. This data can be accessed from a material through the CustomPrimitiveData material expression and it can be written to on the UStaticMeshComponent from gameplay(not hooked up in this submit).
The custom data is a TArray<float> which supports up to 16 floats. This adds an empty TArray<float> to the UStaticMeshComponent and FStaticMeshSceneProxy by default. Other PrimitiveSceneProxys+Component combos can implement access to the custom data by overriding the GetCustomPrimtiveData() method in FPrimtiveSceneProxy.

Additionally some tighter packing has been done in the FPrimitiveUniformShaderParameters to bubble up the padding to the end.

diff --git a/Engine/Shaders/Private/SceneData.ush b/Engine/Shaders/Private/SceneData.ush
--- a/Engine/Shaders/Private/SceneData.ush
+++ b/Engine/Shaders/Private/SceneData.ush
@@ -9,6 +9,8 @@
 
 #if VF_USE_PRIMITIVE_SCENE_DATA
 
+#define NUM_CUSTOM_PRIMITIVE_DATA 4 // Num float4s used for custom data. Must match FCustomPrimitiveData::NumCustomPrimitiveDataFloat4s in SceneTypes.h
+
 // Must match FPrimitiveUniformShaderParameters in C++
 struct FPrimitiveSceneData
 {
@@ -29,16 +31,18 @@ struct FPrimitiveSceneData
 	float4 ObjectOrientation;
 	float4 NonUniformScale;
 	float3 LocalObjectBoundsMin;
-	float3 LocalObjectBoundsMax;
-	float3 PreSkinnedLocalBounds;
 	uint LightingChannelMask;
+	float3 LocalObjectBoundsMax;
 	uint LightmapDataIndex;
+	float3 PreSkinnedLocalBounds;
 	int SingleCaptureIndex;
     uint OutputVelocity;
+	// <-- 12 byte padding here, feel free to use -->
+	float4 CustomPrimitiveData[NUM_CUSTOM_PRIMITIVE_DATA];
 };
 
 // Stride of a single primitive's data in float4's, must match C++
-#define PRIMITIVE_SCENE_DATA_STRIDE 27
+#define PRIMITIVE_SCENE_DATA_STRIDE 31
 
 // Fetch from scene primitive buffer
 FPrimitiveSceneData GetPrimitiveData(uint PrimitiveId) 
@@ -73,6 +77,7 @@ FPrimitiveSceneData GetPrimitiveData(uint PrimitiveId)
 
 	PrimitiveData.ActorWorldPosition = View.PrimitiveSceneData[PrimitiveBaseOffset + 18].xyz;
 	PrimitiveData.UseSingleSampleShadowFromStationaryLights = View.PrimitiveSceneData[PrimitiveBaseOffset + 18].w;
+
 	PrimitiveData.ObjectBounds = View.PrimitiveSceneData[PrimitiveBaseOffset + 19].xyz;
 	PrimitiveData.LpvBiasMultiplier = View.PrimitiveSceneData[PrimitiveBaseOffset + 19].w;
 
@@ -85,14 +90,21 @@ FPrimitiveSceneData GetPrimitiveData(uint PrimitiveId)
 	PrimitiveData.NonUniformScale = View.PrimitiveSceneData[PrimitiveBaseOffset + 22];
 
 	PrimitiveData.LocalObjectBoundsMin = View.PrimitiveSceneData[PrimitiveBaseOffset + 23].xyz;
-	PrimitiveData.LocalObjectBoundsMax = View.PrimitiveSceneData[PrimitiveBaseOffset + 24].xyz;
 	PrimitiveData.LightingChannelMask = asuint(View.PrimitiveSceneData[PrimitiveBaseOffset + 23].w);
+
+	PrimitiveData.LocalObjectBoundsMax = View.PrimitiveSceneData[PrimitiveBaseOffset + 24].xyz;
 	PrimitiveData.LightmapDataIndex = asuint(View.PrimitiveSceneData[PrimitiveBaseOffset + 24].w);
 
-	PrimitiveData.SingleCaptureIndex = asint(View.PrimitiveSceneData[PrimitiveBaseOffset + 25].x);
-	PrimitiveData.OutputVelocity = asuint(View.PrimitiveSceneData[PrimitiveBaseOffset + 25].y);
+	PrimitiveData.PreSkinnedLocalBounds = View.PrimitiveSceneData[PrimitiveBaseOffset + 25].xyz;
+	PrimitiveData.SingleCaptureIndex = View.PrimitiveSceneData[PrimitiveBaseOffset + 25].w;
+
+	PrimitiveData.OutputVelocity = asuint(View.PrimitiveSceneData[PrimitiveBaseOffset + 26].x);
 
-	PrimitiveData.PreSkinnedLocalBounds = View.PrimitiveSceneData[PrimitiveBaseOffset + 26].xyz;
+	UNROLL
+	for (int i = 0; i < NUM_CUSTOM_PRIMITIVE_DATA; i++)
+	{
+		PrimitiveData.CustomPrimitiveData[i] = View.PrimitiveSceneData[PrimitiveBaseOffset + 27 + i];
+	}
 
 	return PrimitiveData;
 }
diff --git a/Engine/Source/Runtime/Engine/Classes/Components/StaticMeshComponent.h b/Engine/Source/Runtime/Engine/Classes/Components/StaticMeshComponent.h
--- a/Engine/Source/Runtime/Engine/Classes/Components/StaticMeshComponent.h
+++ b/Engine/Source/Runtime/Engine/Classes/Components/StaticMeshComponent.h
@@ -16,6 +16,7 @@
 #include "Templates/UniquePtr.h"
 #include "Runtime/Launch/Resources/Version.h"
 #include "UObject/RenderingObjectVersion.h"
+#include "SceneTypes.h"
 #include "StaticMeshComponent.generated.h"
 
 class FColorVertexBuffer;
@@ -195,6 +196,10 @@ private:
 
 public:
 
+	/** Custom data that can be read by a material through a material expression */
+	UPROPERTY()
+	FCustomPrimitiveData CustomPrimitiveData;
+
 	/** Helper function to get the FName of the private static mesh member */
 	static const FName GetMemberNameChecked_StaticMesh() { return GET_MEMBER_NAME_CHECKED(UStaticMeshComponent, StaticMesh); }
 
diff --git a/Engine/Source/Runtime/Engine/Classes/Materials/MaterialExpressionCustomPrimitiveData.h b/Engine/Source/Runtime/Engine/Classes/Materials/MaterialExpressionCustomPrimitiveData.h
new file mode 100644
--- /dev/null
+++ b/Engine/Source/Runtime/Engine/Classes/Materials/MaterialExpressionCustomPrimitiveData.h
@@ -0,0 +1,26 @@
+// Copyright 1998-2019 Epic Games, Inc. All Rights Reserved.
+
+#pragma once
+
+#include "CoreMinimal.h"
+#include "UObject/ObjectMacros.h"
+#include "Materials/MaterialExpression.h"
+#include "MaterialExpressionCustomPrimitiveData.generated.h"
+
+UCLASS(collapsecategories, hidecategories = Object)
+class UMaterialExpressionCustomPrimitiveData: public UMaterialExpression
+{
+	GENERATED_UCLASS_BODY()
+
+#if WITH_EDITOR
+	//~ Begin UMaterialExpression Interface
+	virtual int32 Compile(class FMaterialCompiler* Compiler, int32 OutputIndex) override;
+	virtual void GetCaption(TArray<FString>& OutCaptions) const override;
+	virtual void PostEditChangeProperty(FPropertyChangedEvent& PropertyChangedEvent) override;
+	//~ End UMaterialExpression Interface
+
+	/** Custom descriptions for each custom data pin, only for making things more readable, don't rely on this data for run-time things */
+	UPROPERTY(EditAnywhere, Category=MaterialExpression)
+	TArray<FString> CustomDescs;
+#endif
+};
\ No newline at end of file
diff --git a/Engine/Source/Runtime/Engine/Private/Materials/HLSLMaterialTranslator.h b/Engine/Source/Runtime/Engine/Private/Materials/HLSLMaterialTranslator.h
--- a/Engine/Source/Runtime/Engine/Private/Materials/HLSLMaterialTranslator.h
+++ b/Engine/Source/Runtime/Engine/Private/Materials/HLSLMaterialTranslator.h
@@ -5470,6 +5470,15 @@ protected:
 
 	}
 
+	virtual int32 CustomPrimitiveData(int32 OutputIndex) override
+	{
+		const int32 CustomDataIndex = OutputIndex / FCustomPrimitiveData::NumCustomPrimitiveDataFloat4s;
+		const int32 ElementIndex = OutputIndex % 4; // Index x, y, z or w
+
+		FString HlslName = FString::Printf(TEXT("CustomPrimitiveData[%d][%d]"), CustomDataIndex, ElementIndex);
+		return GetPrimitiveProperty(MCT_Float, TEXT("CustomPrimitiveData"), *HlslName);
+	}
+
 	virtual int32 CustomExpression( class UMaterialExpressionCustom* Custom, TArray<int32>& CompiledInputs ) override
 	{
 		int32 ResultIdx = INDEX_NONE;
diff --git a/Engine/Source/Runtime/Engine/Private/Materials/MaterialExpressions.cpp b/Engine/Source/Runtime/Engine/Private/Materials/MaterialExpressions.cpp
--- a/Engine/Source/Runtime/Engine/Private/Materials/MaterialExpressions.cpp
+++ b/Engine/Source/Runtime/Engine/Private/Materials/MaterialExpressions.cpp
@@ -72,6 +72,7 @@
 #include "Materials/MaterialExpressionCosine.h"
 #include "Materials/MaterialExpressionCrossProduct.h"
 #include "Materials/MaterialExpressionCustom.h"
+#include "Materials/MaterialExpressionCustomPrimitiveData.h"
 #include "Materials/MaterialExpressionDDX.h"
 #include "Materials/MaterialExpressionDDY.h"
 #include "Materials/MaterialExpressionDecalDerivative.h"
@@ -15203,4 +15204,80 @@ void UMaterialExpressionCurveAtlasRowParameter::PostEditChangeProperty(FProperty
 	Super::PostEditChangeProperty(PropertyChangedEvent);
 }
 #endif
+
+//
+// UMaterialExpressionCustomPrimitiveData
+//
+UMaterialExpressionCustomPrimitiveData::UMaterialExpressionCustomPrimitiveData(const FObjectInitializer& ObjectInitializer)
+	: Super(ObjectInitializer)
+{
+#if WITH_EDITORONLY_DATA
+	// Structure to hold one-time initialization
+	struct FConstructorStatics
+	{
+		FText NAME_Vectors;
+		FConstructorStatics()
+			: NAME_Vectors(LOCTEXT("Custom", "Custom"))
+		{
+		}
+	};
+	static FConstructorStatics ConstructorStatics;
+
+	MenuCategories.Add(ConstructorStatics.NAME_Vectors);
+
+	bShaderInputData = true;
+	bShowOutputNameOnPin = true;
+
+	Outputs.Reset();
+#endif
+}
+
+#if WITH_EDITOR
+int32 UMaterialExpressionCustomPrimitiveData::Compile(class FMaterialCompiler* Compiler, int32 OutputIndex)
+{
+	if (Material && Material->MaterialDomain == MD_DeferredDecal)
+	{
+		return CompilerError(Compiler, TEXT("Expression not available in the deferred decal material domain."));
+	}
+
+	return Compiler->CustomPrimitiveData(OutputIndex);
+}
+
+void UMaterialExpressionCustomPrimitiveData::GetCaption(TArray<FString>& OutCaptions) const
+{
+	OutCaptions.Add(TEXT("Custom Primitive Data"));
+}
+
+void UMaterialExpressionCustomPrimitiveData::PostEditChangeProperty(FPropertyChangedEvent& PropertyChangedEvent)
+{
+	if (PropertyChangedEvent.GetPropertyName() == GET_MEMBER_NAME_CHECKED(UMaterialExpressionCustomPrimitiveData, CustomDescs))
+	{
+		// If more than the supported number of custom floats have been added, remove the overflow
+		if (CustomDescs.Num() > FCustomPrimitiveData::NumCustomPrimitiveDataFloats)
+		{
+			CustomDescs.SetNum(FCustomPrimitiveData::NumCustomPrimitiveDataFloats);
+		}
+
+		// Match the number of pins to the number of descriptions
+		Outputs.SetNumZeroed(CustomDescs.Num());
+
+		// Update the names of the pins
+		for (int i = 0; i < Outputs.Num(); i++)
+		{
+			Outputs[i].OutputName = *CustomDescs[i];
+		}
+
+		// Reconstruct the node to get the new names applied to the pins
+		GraphNode->ReconstructNode();
+
+		// No need to update preview if we only change output descriptions
+		bNeedToUpdatePreview = false;
+	}
+	else
+	{
+		Super::PostEditChangeProperty(PropertyChangedEvent);
+	}
+}
+
+#endif // WITH_EDITOR
 #undef LOCTEXT_NAMESPACE
diff --git a/Engine/Source/Runtime/Engine/Private/PrimitiveSceneProxy.cpp b/Engine/Source/Runtime/Engine/Private/PrimitiveSceneProxy.cpp
--- a/Engine/Source/Runtime/Engine/Private/PrimitiveSceneProxy.cpp
+++ b/Engine/Source/Runtime/Engine/Private/PrimitiveSceneProxy.cpp
@@ -292,7 +292,8 @@ void FPrimitiveSceneProxy::UpdateUniformBuffer()
 				LpvBiasMultiplier,
 				PrimitiveSceneInfo ? PrimitiveSceneInfo->GetLightmapDataOffset() : 0,
 				SingleCaptureIndex, 
-				bOutputVelocity);
+				bOutputVelocity,
+				GetCustomPrimitiveData());
 
 		if (UniformBuffer.GetReference())
 		{
diff --git a/Engine/Source/Runtime/Engine/Private/PrimitiveUniformShaderParameters.cpp b/Engine/Source/Runtime/Engine/Private/PrimitiveUniformShaderParameters.cpp
--- a/Engine/Source/Runtime/Engine/Private/PrimitiveUniformShaderParameters.cpp
+++ b/Engine/Source/Runtime/Engine/Private/PrimitiveUniformShaderParameters.cpp
@@ -66,7 +66,8 @@ FPrimitiveSceneShaderData::FPrimitiveSceneShaderData(const FPrimitiveSceneProxy*
 		Proxy->GetLpvBiasMultiplier(),
 		Proxy->GetPrimitiveSceneInfo()->GetLightmapDataOffset(),
 		SingleCaptureIndex,
-        bOutputVelocity));
+        bOutputVelocity,
+		Proxy->GetCustomPrimitiveData()));
 }
 
 void FPrimitiveSceneShaderData::Setup(const FPrimitiveUniformShaderParameters& PrimitiveUniformShaderParameters)
@@ -113,9 +114,13 @@ void FPrimitiveSceneShaderData::Setup(const FPrimitiveUniformShaderParameters& P
 	Data[24] = FVector4(PrimitiveUniformShaderParameters.LocalObjectBoundsMax, 0.0f);
 	Data[24].W = *(const float*)&PrimitiveUniformShaderParameters.LightmapDataIndex;
 
-	Data[25] = FVector4(0.0f, 0.0f, 0.0f, 0.0f);
-	Data[25].X = *(const float*)&PrimitiveUniformShaderParameters.SingleCaptureIndex;
-	Data[25].Y = *(const float*)&PrimitiveUniformShaderParameters.OutputVelocity;
+	Data[25] = FVector4(PrimitiveUniformShaderParameters.PreSkinnedLocalBounds, *(const float*)&PrimitiveUniformShaderParameters.SingleCaptureIndex); 
 
-	Data[26] = FVector4(PrimitiveUniformShaderParameters.PreSkinnedLocalBounds, 0.0f); // .w unused
-}
\ No newline at end of file
+	Data[26] = FVector4(0.0f, 0.0f, 0.0f, 0.0f);
+	Data[26].X = *(const float*)&PrimitiveUniformShaderParameters.OutputVelocity;
+
+	Data[27] = PrimitiveUniformShaderParameters.CustomPrimitiveData[0];
+	Data[28] = PrimitiveUniformShaderParameters.CustomPrimitiveData[1];
+	Data[29] = PrimitiveUniformShaderParameters.CustomPrimitiveData[2];
+	Data[30] = PrimitiveUniformShaderParameters.CustomPrimitiveData[3];
+}
diff --git a/Engine/Source/Runtime/Engine/Private/SceneManagement.cpp b/Engine/Source/Runtime/Engine/Private/SceneManagement.cpp
--- a/Engine/Source/Runtime/Engine/Private/SceneManagement.cpp
+++ b/Engine/Source/Runtime/Engine/Private/SceneManagement.cpp
@@ -313,7 +313,7 @@ void FDynamicPrimitiveUniformBuffer::Set(
 {
 	check(IsInRenderingThread());
 	UniformBuffer.SetContents(
-		GetPrimitiveUniformShaderParameters(LocalToWorld, PreviousLocalToWorld, WorldBounds.Origin, WorldBounds, LocalBounds, PreSkinnedLocalBounds, bReceivesDecals, false, false, false, bHasPrecomputedVolumetricLightmap, bUseEditorDepthTest, GetDefaultLightingChannelMask(), 1.0f, INDEX_NONE, INDEX_NONE, bOutputVelocity));
+		GetPrimitiveUniformShaderParameters(LocalToWorld, PreviousLocalToWorld, WorldBounds.Origin, WorldBounds, LocalBounds, PreSkinnedLocalBounds, bReceivesDecals, false, false, false, bHasPrecomputedVolumetricLightmap, bUseEditorDepthTest, GetDefaultLightingChannelMask(), 1.0f, INDEX_NONE, INDEX_NONE, bOutputVelocity, nullptr));
 	UniformBuffer.InitResource();
 }
 
diff --git a/Engine/Source/Runtime/Engine/Private/StaticMeshRender.cpp b/Engine/Source/Runtime/Engine/Private/StaticMeshRender.cpp
--- a/Engine/Source/Runtime/Engine/Private/StaticMeshRender.cpp
+++ b/Engine/Source/Runtime/Engine/Private/StaticMeshRender.cpp
@@ -114,6 +114,7 @@ FStaticMeshSceneProxy::FStaticMeshSceneProxy(UStaticMeshComponent* InComponent,
 	, bCastShadow(InComponent->CastShadow)
 	, bReverseCulling(InComponent->bReverseCulling)
 	, MaterialRelevance(InComponent->GetMaterialRelevance(GetScene().GetFeatureLevel()))
+	, CustomPrimitiveData(InComponent->CustomPrimitiveData)
 #if WITH_EDITORONLY_DATA
 	, StreamingDistanceMultiplier(FMath::Max(0.0f, InComponent->StreamingDistanceMultiplier))
 	, StreamingTransformScale(InComponent->GetTextureStreamingTransformScale())
diff --git a/Engine/Source/Runtime/Engine/Public/MaterialCompiler.h b/Engine/Source/Runtime/Engine/Public/MaterialCompiler.h
--- a/Engine/Source/Runtime/Engine/Public/MaterialCompiler.h
+++ b/Engine/Source/Runtime/Engine/Public/MaterialCompiler.h
@@ -293,6 +293,7 @@ public:
 	virtual int32 EyeAdaptation() = 0;
 	virtual int32 AtmosphericLightVector() = 0;
 	virtual int32 AtmosphericLightColor() = 0;
+	virtual int32 CustomPrimitiveData(int32 OutputIndex) = 0;
 	// The compiler can run in a different state and this affects caching of sub expression, Expressions are different (e.g. View.PrevWorldViewOrigin) when using previous frame's values
 	// If possible we should re-factor this to avoid having to deal with compiler state
 	virtual bool IsCurrentlyCompilingForPreviousFrame() const { return false; }
@@ -548,6 +549,11 @@ public:
 	{
 		return Compiler->AtmosphericLightColor();
 	}
+	
+	virtual int32 CustomPrimitiveData(int32 OutputIndex) override
+	{
+		return Compiler->CustomPrimitiveData(OutputIndex);
+	}
 
 	virtual int32 TextureCoordinateOffset() override
 	{
diff --git a/Engine/Source/Runtime/Engine/Public/PrimitiveSceneProxy.h b/Engine/Source/Runtime/Engine/Public/PrimitiveSceneProxy.h
--- a/Engine/Source/Runtime/Engine/Public/PrimitiveSceneProxy.h
+++ b/Engine/Source/Runtime/Engine/Public/PrimitiveSceneProxy.h
@@ -728,6 +728,12 @@ public:
 
 	virtual uint8 GetCurrentFirstLODIdx_RenderThread() const { return 0; }
 
+	/** 
+	 * Get the custom primitive data for this scene proxy. Defaults to zero-data if not implemented.
+	 * @return The payload of custom data that will be set on the primitive and accessible in the material through a material expression.
+	 */
+	ENGINE_API virtual const FCustomPrimitiveData* GetCustomPrimitiveData() const { return nullptr; }
+
 protected:
 
 	/** Allow subclasses to override the primitive name. Used primarily by BSP. */
diff --git a/Engine/Source/Runtime/Engine/Public/PrimitiveUniformShaderParameters.h b/Engine/Source/Runtime/Engine/Public/PrimitiveUniformShaderParameters.h
--- a/Engine/Source/Runtime/Engine/Public/PrimitiveUniformShaderParameters.h
+++ b/Engine/Source/Runtime/Engine/Public/PrimitiveUniformShaderParameters.h
@@ -5,6 +5,7 @@
 
 #include "CoreMinimal.h"
 #include "Engine/EngineTypes.h"
+#include "SceneTypes.h"
 #include "RenderResource.h"
 #include "ShaderParameters.h"
 #include "UniformBuffer.h"
@@ -13,6 +14,7 @@
 /** 
  * The uniform shader parameters associated with a primitive. 
  * Note: Must match FPrimitiveSceneData in shaders.
+ * Note 2: Try to keep this 16 byte aligned. i.e |Matrix4x4|Vector3,float|Vector3,float|Vector4|  _NOT_  |Vector3,(waste padding)|Vector3,(waste padding)|Vector3. Or at least mark out padding if it can't be avoided.
  */
 BEGIN_GLOBAL_SHADER_PARAMETER_STRUCT(FPrimitiveUniformShaderParameters,ENGINE_API)
 	SHADER_PARAMETER(FMatrix,LocalToWorld)		// always needed
@@ -32,12 +34,14 @@ BEGIN_GLOBAL_SHADER_PARAMETER_STRUCT(FPrimitiveUniformShaderParameters,ENGINE_AP
 	SHADER_PARAMETER_EX(FVector4,ObjectOrientation,EShaderPrecisionModifier::Half)
 	SHADER_PARAMETER_EX(FVector4,NonUniformScale,EShaderPrecisionModifier::Half)
 	SHADER_PARAMETER(FVector, LocalObjectBoundsMin)		// This is used in a custom material function (ObjectLocalBounds.uasset)
-	SHADER_PARAMETER(FVector, LocalObjectBoundsMax)		// This is used in a custom material function (ObjectLocalBounds.uasset)
-	SHADER_PARAMETER(FVector, PreSkinnedLocalBounds)	// Local space bounds, pre-skinning
 	SHADER_PARAMETER(uint32,LightingChannelMask)
+	SHADER_PARAMETER(FVector, LocalObjectBoundsMax)		// This is used in a custom material function (ObjectLocalBounds.uasset)
 	SHADER_PARAMETER(uint32,LightmapDataIndex)
+	SHADER_PARAMETER(FVector, PreSkinnedLocalBounds)	// Local space bounds, pre-skinning
 	SHADER_PARAMETER(int32, SingleCaptureIndex)
     SHADER_PARAMETER(uint32, OutputVelocity)
+	// 12 bytes of padding here, feel free to use
+	SHADER_PARAMETER_ARRAY(FVector4, CustomPrimitiveData, [FCustomPrimitiveData::NumCustomPrimitiveDataFloat4s]) // Custom data per primitive that can be accessed through UMaterialExpressionCustomPrimitiveData and modified through UStaticMeshComponent
 
 END_GLOBAL_SHADER_PARAMETER_STRUCT()
 
@@ -59,7 +63,8 @@ inline FPrimitiveUniformShaderParameters GetPrimitiveUniformShaderParameters(
 	float LpvBiasMultiplier,
 	uint32 LightmapDataIndex,
 	int32 SingleCaptureIndex,
-	bool bOutputVelocity
+	bool bOutputVelocity,
+	const FCustomPrimitiveData* CustomPrimitiveData
 )
 {
 	FPrimitiveUniformShaderParameters Result;
@@ -101,6 +106,16 @@ inline FPrimitiveUniformShaderParameters GetPrimitiveUniformShaderParameters(
 	Result.LightmapDataIndex = LightmapDataIndex;
 	Result.SingleCaptureIndex = SingleCaptureIndex;
 	Result.OutputVelocity = (bOutputVelocity) ? 1 : 0;
+
+	// Clear to 0
+	FMemory::Memzero(Result.CustomPrimitiveData);
+
+	// If this primitive has custom primitive data, set it
+	if (CustomPrimitiveData)
+	{
+		// Copy at most up to the max supported number of dwords for safety
+		FMemory::Memcpy(&Result.CustomPrimitiveData, CustomPrimitiveData->Data.GetData(), CustomPrimitiveData->Data.GetTypeSize() * FMath::Min(CustomPrimitiveData->Data.Num(), FCustomPrimitiveData::NumCustomPrimitiveDataFloats));
+	}
 	return Result;
 }
 
@@ -126,7 +141,7 @@ inline FPrimitiveUniformShaderParameters GetPrimitiveUniformShaderParameters(
 {
 	// Pass through call
 	return GetPrimitiveUniformShaderParameters(LocalToWorld, PreviousLocalToWorld, ActorPosition, WorldBounds, LocalBounds, LocalBounds, bReceivesDecals, bHasDistanceFieldRepresentation, bHasCapsuleRepresentation, 
-		bUseSingleSampleShadowFromStationaryLights, bUseVolumetricLightmap, bUseEditorDepthTest, LightingChannelMask, LpvBiasMultiplier, LightmapDataIndex, SingleCaptureIndex, bOutputVelocity);
+		bUseSingleSampleShadowFromStationaryLights, bUseVolumetricLightmap, bUseEditorDepthTest, LightingChannelMask, LpvBiasMultiplier, LightmapDataIndex, SingleCaptureIndex, bOutputVelocity, nullptr);
 }
 
 inline TUniformBufferRef<FPrimitiveUniformShaderParameters> CreatePrimitiveUniformBufferImmediate(
@@ -141,7 +156,7 @@ inline TUniformBufferRef<FPrimitiveUniformShaderParameters> CreatePrimitiveUnifo
 {
 	check(IsInRenderingThread());
 	return TUniformBufferRef<FPrimitiveUniformShaderParameters>::CreateUniformBufferImmediate(
-		GetPrimitiveUniformShaderParameters(LocalToWorld, LocalToWorld, WorldBounds.Origin, WorldBounds, LocalBounds, PreSkinnedLocalBounds, bReceivesDecals, false, false, false, false, bUseEditorDepthTest, GetDefaultLightingChannelMask(), LpvBiasMultiplier, INDEX_NONE, INDEX_NONE, false),
+		GetPrimitiveUniformShaderParameters(LocalToWorld, LocalToWorld, WorldBounds.Origin, WorldBounds, LocalBounds, PreSkinnedLocalBounds, bReceivesDecals, false, false, false, false, bUseEditorDepthTest, GetDefaultLightingChannelMask(), LpvBiasMultiplier, INDEX_NONE, INDEX_NONE, false, nullptr),
 		UniformBuffer_MultiFrame
 		);
 }
@@ -190,7 +205,7 @@ extern ENGINE_API TGlobalResource<FIdentityPrimitiveUniformBuffer> GIdentityPrim
 struct FPrimitiveSceneShaderData
 {
 	// Must match usf
-	enum { PrimitiveDataStrideInFloat4s = 27 };
+	enum { PrimitiveDataStrideInFloat4s = 31 };
 
 	FVector4 Data[PrimitiveDataStrideInFloat4s];
 
diff --git a/Engine/Source/Runtime/Engine/Public/SceneTypes.h b/Engine/Source/Runtime/Engine/Public/SceneTypes.h
--- a/Engine/Source/Runtime/Engine/Public/SceneTypes.h
+++ b/Engine/Source/Runtime/Engine/Public/SceneTypes.h
@@ -19,6 +19,19 @@ typedef TRefCountPtr<FLightMap> FLightMapRef;
 /** A reference to a shadow-map. */
 typedef TRefCountPtr<FShadowMap> FShadowMapRef;
 
+/** Custom primitive data payload. */
+USTRUCT()
+struct FCustomPrimitiveData
+{
+	GENERATED_USTRUCT_BODY()
+
+	static constexpr int32 NumCustomPrimitiveDataFloat4s = 4; // Must match NUM_CUSTOM_PRIMITIVE_DATA in SceneData.ush
+	static constexpr int32 NumCustomPrimitiveDataFloats = NumCustomPrimitiveDataFloat4s * 4;
+
+	UPROPERTY(EditAnywhere, Category=Rendering)
+	TArray<float> Data;
+};
+
 /** 
  * Class used to identify UPrimitiveComponents on the rendering thread without having to pass the pointer around, 
  * Which would make it easy for people to access game thread variables from the rendering thread.
diff --git a/Engine/Source/Runtime/Engine/Public/StaticMeshResources.h b/Engine/Source/Runtime/Engine/Public/StaticMeshResources.h
--- a/Engine/Source/Runtime/Engine/Public/StaticMeshResources.h
+++ b/Engine/Source/Runtime/Engine/Public/StaticMeshResources.h
@@ -772,6 +772,7 @@ public:
 	virtual void GetMeshDescription(int32 LODIndex, TArray<FMeshBatch>& OutMeshElements) const override;
 
 	virtual void GetDynamicMeshElements(const TArray<const FSceneView*>& Views, const FSceneViewFamily& ViewFamily, uint32 VisibilityMap, FMeshElementCollector& Collector) const override;
+	virtual const FCustomPrimitiveData* GetCustomPrimitiveData() const override { return &CustomPrimitiveData; }
 
 #if RHI_RAYTRACING
 	virtual bool IsRayTracingRelevant() const override { return true; }
@@ -883,6 +884,9 @@ protected:
 	/** The view relevance for all the static mesh's materials. */
 	FMaterialRelevance MaterialRelevance;
 
+	/** The custom data for this Scene Proxy */
+	FCustomPrimitiveData CustomPrimitiveData;
+
 #if WITH_EDITORONLY_DATA
 	/** The component streaming distance multiplier */
 	float StreamingDistanceMultiplier;
