- Don't attempt to do virtualize real loop if sound is already pending stop to fix test build
- Minor clean-up including renaming StopActiveSound in ConcurrencyManager to RemoveActiveSound for clarity

diff --git a/Engine/Source/Runtime/Engine/Classes/Sound/SoundConcurrency.h b/Engine/Source/Runtime/Engine/Classes/Sound/SoundConcurrency.h
--- a/Engine/Source/Runtime/Engine/Classes/Sound/SoundConcurrency.h
+++ b/Engine/Source/Runtime/Engine/Classes/Sound/SoundConcurrency.h
@@ -216,7 +216,7 @@ public:
 	FActiveSound* CreateNewActiveSound(const FActiveSound& NewActiveSound, bool bIsRetriggering);
 
 	/** Removes the active sound from concurrency tracking when active sound is stopped. */
-	void StopActiveSound(FActiveSound* ActiveSound);
+	void RemoveActiveSound(FActiveSound* ActiveSound);
 
 	/** Stops any active sounds due to max concurrency quietest sound resolution rule */
 	void UpdateQuietSoundsToStop();
diff --git a/Engine/Source/Runtime/Engine/Private/AudioDevice.cpp b/Engine/Source/Runtime/Engine/Private/AudioDevice.cpp
--- a/Engine/Source/Runtime/Engine/Private/AudioDevice.cpp
+++ b/Engine/Source/Runtime/Engine/Private/AudioDevice.cpp
@@ -2594,13 +2594,17 @@ void FAudioDevice::VirtualizeInactiveLoops()
 	const bool bDoRangeCheck = true;
 	for (FActiveSound* ActiveSound : ActiveSounds)
 	{
-		if (FAudioVirtualLoop* VirtualLoop = FAudioVirtualLoop::Virtualize(*this, *ActiveSound, bDoRangeCheck))
+		// If already pending stop, don't attempt to virtualize
+		if (!IsPendingStop(ActiveSound))
 		{
-			AddSoundToStop(ActiveSound);
+			if (FAudioVirtualLoop* VirtualLoop = FAudioVirtualLoop::Virtualize(*this, *ActiveSound, bDoRangeCheck))
+			{
+				AddSoundToStop(ActiveSound);
 
-			// Clear must be called after AddSoundToStop to ensure AudioComponent is properly removed from AudioComponentIDToActiveSoundMap
-			ActiveSound->ClearAudioComponent();
-			AddVirtualLoop(*VirtualLoop);
+				// Clear must be called after AddSoundToStop to ensure AudioComponent is properly removed from AudioComponentIDToActiveSoundMap
+				ActiveSound->ClearAudioComponent();
+				AddVirtualLoop(*VirtualLoop);
+			}
 		}
 	}
 }
@@ -4289,7 +4293,14 @@ void FAudioDevice::AddVirtualLoop(FAudioVirtualLoop& VirtualLoop)
 	const int64 ComponentID = ActiveSound.GetAudioComponentID();
 	if (ComponentID > 0)
 	{
-		check(!AudioComponentIDToActiveSoundMap.Contains(ComponentID));
+		if (FActiveSound* ExistingSound = AudioComponentIDToActiveSoundMap.FindRef(ComponentID))
+		{
+			UE_LOG(LogAudio, Warning, TEXT("Adding ComponentID for Sound '%s' when map already contains ID for Sound '%s'."),
+				ActiveSound.Sound ? *ActiveSound.Sound->GetName() : TEXT("N/A"),
+				ExistingSound->Sound ? *ExistingSound->Sound->GetName() : TEXT("N/A")
+				);
+			AudioComponentIDToActiveSoundMap.Remove(ComponentID);
+		}
 		AudioComponentIDToActiveSoundMap.Add(ComponentID, &ActiveSound);
 	}
 
@@ -4427,7 +4438,7 @@ void FAudioDevice::AddSoundToStop(FActiveSound* SoundToStop)
 
 		if (!VirtualLoops.Contains(SoundToStop))
 		{
-			ConcurrencyManager.StopActiveSound(SoundToStop);
+			ConcurrencyManager.RemoveActiveSound(SoundToStop);
 		}
 	}
 }
@@ -4705,26 +4716,29 @@ bool FAudioDevice::SoundIsAudible(const FActiveSound& NewActiveSound)
 		return true;
 	}
 
-	const float ApparentMaxDistance = NewActiveSound.MaxDistance * NewActiveSound.FocusDistanceScale;
 	if (NewActiveSound.Sound->IsVirtualizeWhenSilent())
 	{
 		return true;
 	}
 
+	// TODO: bAllowSpatialization is used in other audibility checks but not here.
 	const FSoundAttenuationSettings& Attenuation = NewActiveSound.AttenuationSettings;
-	const bool bHasNonFocusDistanceScale = Attenuation.FocusDistanceScale != 1.0f || Attenuation.NonFocusDistanceScale != 1.0f;
+	const bool bHasFocusScaling = Attenuation.FocusDistanceScale != 1.0f || Attenuation.NonFocusDistanceScale != 1.0f;
 	if (!NewActiveSound.bHasAttenuationSettings ||
-		(NewActiveSound.bHasAttenuationSettings && (!Attenuation.bAttenuate || bHasNonFocusDistanceScale)))
+		(NewActiveSound.bHasAttenuationSettings && (!Attenuation.bAttenuate || bHasFocusScaling)))
 	{
 		return true;
 	}
 
+	// TODO: Check if this is necessary. GetMaxDistanceAndFocusFactor should've solved this and would make this
+	// flavor of SoundIsAudible more accurate.
 	const FGlobalFocusSettings& FocusSettings = GetGlobalFocusSettings();
 	if (FocusSettings.FocusDistanceScale != 1.0f || FocusSettings.NonFocusDistanceScale != 1.0f)
 	{
 		return true;
 	}
 
+	const float ApparentMaxDistance = NewActiveSound.MaxDistance * NewActiveSound.FocusDistanceScale;
 	if (LocationIsAudible(NewActiveSound.Transform.GetLocation(), ApparentMaxDistance))
 	{
 		return true;
diff --git a/Engine/Source/Runtime/Engine/Private/AudioVirtualLoop.cpp b/Engine/Source/Runtime/Engine/Private/AudioVirtualLoop.cpp
--- a/Engine/Source/Runtime/Engine/Private/AudioVirtualLoop.cpp
+++ b/Engine/Source/Runtime/Engine/Private/AudioVirtualLoop.cpp
@@ -7,7 +7,7 @@
 
 	static int32 bVirtualLoopsEnabledCVar = 1;
 	FAutoConsoleVariableRef CVarVirtualLoopsEnabled(
-		TEXT("au.VirtualLoops"),
+		TEXT("au.VirtualLoopsEnabled"),
 		bVirtualLoopsEnabledCVar,
 		TEXT("Enables or disables whether virtualizing is supported for audio loops.\n"),
 		ECVF_Default);
@@ -21,14 +21,14 @@
 
 	static float VirtualUpdateRateMinCVar = 0.1f;
 	FAutoConsoleVariableRef CVarVirtualLoopsUpdateRateMin(
-		TEXT("au.VirtualUpdateRateMin"),
+		TEXT("au.VirtualLoopsUpdateRateMin"),
 		VirtualUpdateRateMinCVar,
 		TEXT("Sets minimum rate to check if sound becomes audible again at sound's max audible distance.\n"),
 		ECVF_Default);
 
 	static float VirtualUpdateRateMaxCVar = 3.0f;
 	FAutoConsoleVariableRef CVarVirtualLoopsUpdateRateMax(
-		TEXT("au.VirtualUpdateRateMax"),
+		TEXT("au.VirtualLoopsUpdateRateMax"),
 		VirtualUpdateRateMinCVar,
 		TEXT("Sets maximum rate to check if sound becomes audible again (at beyond sound's max audible distance + perf scaling distance).\n"),
 		ECVF_Default);
diff --git a/Engine/Source/Runtime/Engine/Private/SoundConcurrency.cpp b/Engine/Source/Runtime/Engine/Private/SoundConcurrency.cpp
--- a/Engine/Source/Runtime/Engine/Private/SoundConcurrency.cpp
+++ b/Engine/Source/Runtime/Engine/Private/SoundConcurrency.cpp
@@ -580,9 +580,10 @@ FActiveSound* FSoundConcurrencyManager::CreateAndEvictActiveSounds(const FActive
 		check(SoundToEvict);
 		check(AudioDevice == SoundToEvict->AudioDevice);
 
+		RemoveActiveSound(SoundToEvict);
+
 		// Remove the active sound from the concurrency manager immediately so it doesn't count towards
 		// subsequent concurrency resolution checks (i.e. if sounds are triggered multiple times in this frame)
-		StopActiveSound(SoundToEvict);
 		if (AudioDevice->IsPendingStop(SoundToEvict))
 		{
 			continue;
@@ -591,18 +592,18 @@ FActiveSound* FSoundConcurrencyManager::CreateAndEvictActiveSounds(const FActive
 		TArray<FConcurrencyHandle> Handles;
 		SoundToEvict->GetConcurrencyHandles(Handles);
 
-		bool bAllowRetrigger = true;
+		bool bAllowVirtual = true;
 		for (const FConcurrencyHandle& Handle : Handles)
 		{
 			switch (Handle.Settings.ResolutionRule)
 			{
 				// Stop oldest resolution rules will cause an undesired
-				// cycling through re-triggerable loops from this frame to the next,
-				// so don't revive them.
+				// cycling through virtual loops from this frame to the next,
+				// so don't permit virtualization in this case.
 				case EMaxConcurrentResolutionRule::StopOldest:
 				case EMaxConcurrentResolutionRule::StopFarthestThenOldest:
 				{
-					bAllowRetrigger = false;
+					bAllowVirtual = false;
 					break;
 				}
 				default:
@@ -614,7 +615,7 @@ FActiveSound* FSoundConcurrencyManager::CreateAndEvictActiveSounds(const FActive
 
 		AudioDevice->AddSoundToStop(SoundToEvict);
 		// If using a mode where we support re-triggering, attempt to re-trigger and update boolean state.
-		if (bAllowRetrigger)
+		if (bAllowVirtual)
 		{
 			const bool bDoRangeCheck = true;
 			if (FAudioVirtualLoop* VirtualLoop = FAudioVirtualLoop::Virtualize(*AudioDevice, *SoundToEvict, bDoRangeCheck))
@@ -628,7 +629,7 @@ FActiveSound* FSoundConcurrencyManager::CreateAndEvictActiveSounds(const FActive
 	return ActiveSound;
 }
 
-void FSoundConcurrencyManager::StopActiveSound(FActiveSound* ActiveSound)
+void FSoundConcurrencyManager::RemoveActiveSound(FActiveSound* ActiveSound)
 {
 	check(IsInAudioThread());
 
