Fix fo non-editor builds. Correctly adding WITH_EDITORONLY_DATA

diff --git a/Engine/Source/Runtime/Engine/Classes/Materials/MaterialExpressionCustomPrimitiveData.h b/Engine/Source/Runtime/Engine/Classes/Materials/MaterialExpressionCustomPrimitiveData.h
--- a/Engine/Source/Runtime/Engine/Classes/Materials/MaterialExpressionCustomPrimitiveData.h
+++ b/Engine/Source/Runtime/Engine/Classes/Materials/MaterialExpressionCustomPrimitiveData.h
@@ -18,9 +18,10 @@ class UMaterialExpressionCustomPrimitiveData: public UMaterialExpression
 	virtual void GetCaption(TArray<FString>& OutCaptions) const override;
 	virtual void PostEditChangeProperty(FPropertyChangedEvent& PropertyChangedEvent) override;
 	//~ End UMaterialExpression Interface
-
+#if WITH_EDITORONLY_DATA
 	/** Custom descriptions for each custom data pin, only for making things more readable, don't rely on this data for run-time things */
 	UPROPERTY(EditAnywhere, Category=MaterialExpression)
 	TArray<FString> CustomDescs;
 #endif
+#endif
 };
\ No newline at end of file
diff --git a/Engine/Source/Runtime/Engine/Private/Materials/MaterialExpressions.cpp b/Engine/Source/Runtime/Engine/Private/Materials/MaterialExpressions.cpp
--- a/Engine/Source/Runtime/Engine/Private/Materials/MaterialExpressions.cpp
+++ b/Engine/Source/Runtime/Engine/Private/Materials/MaterialExpressions.cpp
@@ -15250,6 +15250,7 @@ void UMaterialExpressionCustomPrimitiveData::GetCaption(TArray<FString>& OutCapt
 
 void UMaterialExpressionCustomPrimitiveData::PostEditChangeProperty(FPropertyChangedEvent& PropertyChangedEvent)
 {
+#if WITH_EDITORONLY_DATA
 	if (PropertyChangedEvent.GetPropertyName() == GET_MEMBER_NAME_CHECKED(UMaterialExpressionCustomPrimitiveData, CustomDescs))
 	{
 		// If more than the supported number of custom floats have been added, remove the overflow
@@ -15274,6 +15275,7 @@ void UMaterialExpressionCustomPrimitiveData::PostEditChangeProperty(FPropertyCha
 		bNeedToUpdatePreview = false;
 	}
 	else
+#endif
 	{
 		Super::PostEditChangeProperty(PropertyChangedEvent);
 	}
