Fix for WindowsVideoRecording System. API explicty requests that no file extension is needed, but there is a codepath where .mp4 is not added to the recording system. Will result in error when trying to init the mp4 writer

diff --git a/Engine/Source/Runtime/Engine/Public/VideoRecordingSystem.h b/Engine/Source/Runtime/Engine/Public/VideoRecordingSystem.h
--- a/Engine/Source/Runtime/Engine/Public/VideoRecordingSystem.h
+++ b/Engine/Source/Runtime/Engine/Public/VideoRecordingSystem.h
@@ -124,9 +124,17 @@ public:
 	/** Returns the current state of video recording. */
 	virtual EVideoRecordingState GetRecordingState() const = 0;
 
+	UE_DEPRECATED(4.23, "Use GetOnVideoRecordingFinalizedDelegate")
 	virtual FDelegateHandle RegisterVideoRecordingFinalizedDelegate(const FVideoRecordingFinalized::FDelegate& Delegate) { return OnVideoRecordingFinalized.Add(Delegate); };
+
+	UE_DEPRECATED(4.23, "Use GetOnVideoRecordingFinalizedDelegate")
 	virtual void UnregisterVideoRecordingFinalizedDelegate(FDelegateHandle Handle) { OnVideoRecordingFinalized.Remove(Handle); };
 
+	FVideoRecordingFinalized& GetOnVideoRecordingFinalizedDelegate()
+	{
+		return OnVideoRecordingFinalized;
+	}
+
 protected:
 	FVideoRecordingFinalized OnVideoRecordingFinalized;
 };
diff --git a/Engine/Source/Runtime/Windows/WindowsPlatformFeatures/Private/WindowsVideoRecordingSystem.cpp b/Engine/Source/Runtime/Windows/WindowsPlatformFeatures/Private/WindowsVideoRecordingSystem.cpp
--- a/Engine/Source/Runtime/Windows/WindowsPlatformFeatures/Private/WindowsVideoRecordingSystem.cpp
+++ b/Engine/Source/Runtime/Windows/WindowsPlatformFeatures/Private/WindowsVideoRecordingSystem.cpp
@@ -70,7 +70,7 @@ void FWindowsVideoRecordingSystem::NextRecording()
 	}
 	else
 	{
-		CurrentFilename = Path + BaseFilename;
+		CurrentFilename = Path + BaseFilename + TEXT(".mp4");
 	}
 }
 
