Fixed debug camera orbit jitter issue when simultaneously moving and rotating.

diff --git a/Engine/Source/Runtime/Engine/Classes/Engine/DebugCameraController.h b/Engine/Source/Runtime/Engine/Classes/Engine/DebugCameraController.h
--- a/Engine/Source/Runtime/Engine/Classes/Engine/DebugCameraController.h
+++ b/Engine/Source/Runtime/Engine/Classes/Engine/DebugCameraController.h
@@ -76,12 +76,18 @@ class ENGINE_API ADebugCameraController
 	 */
 	virtual void ToggleFreezeRendering();
 
+	/** Method called prior to processing input */
+	virtual void PreProcessInput(const float DeltaTime, const bool bGamePaused);
+
 	/**
 	 * Updates the rotation of player, based on ControlRotation after RotationInput has been applied.
 	 * This may then be modified by the PlayerCamera, and is passed to Pawn->FaceRotation().
 	 */
 	virtual void UpdateRotation(float DeltaTime) override;
 
+	/** Pre process input when orbiting */
+	void PreProcessInputForOrbit(const float DeltaTime, const bool bGamePaused);
+
 	/** Updates the rotation and location of player when orbiting */
 	void UpdateRotationForOrbit(float DeltaTime);
 
diff --git a/Engine/Source/Runtime/Engine/Private/DebugCameraController.cpp b/Engine/Source/Runtime/Engine/Private/DebugCameraController.cpp
--- a/Engine/Source/Runtime/Engine/Private/DebugCameraController.cpp
+++ b/Engine/Source/Runtime/Engine/Private/DebugCameraController.cpp
@@ -394,6 +394,18 @@ void ADebugCameraController::ToggleFreezeRendering()
 	bIsFrozenRendering = !bIsFrozenRendering;
 }
 
+void ADebugCameraController::PreProcessInput(const float DeltaTime, const bool bGamePaused)
+{
+	if (bIsOrbitingSelectedActor)
+	{
+		PreProcessInputForOrbit(DeltaTime, bGamePaused);
+	}
+	else
+	{
+		Super::PreProcessInput(DeltaTime, bGamePaused);
+	}
+}
+
 void ADebugCameraController::UpdateRotation(float DeltaTime)
 {
 	if (bIsOrbitingSelectedActor)
@@ -406,6 +418,23 @@ void ADebugCameraController::UpdateRotation(float DeltaTime)
 	}
 }
 
+void ADebugCameraController::PreProcessInputForOrbit(const float DeltaTime, const bool bGamePaused)
+{
+	if (bIsOrbitingSelectedActor)
+	{
+		if (APawn* const CurrentPawn = GetPawnOrSpectator())
+		{
+			if (UPawnMovementComponent* MovementComponent = CurrentPawn->GetMovementComponent())
+			{
+				// Reset velocity before processing input when orbiting to limit acceleration which 
+				// can cause overshooting and jittering as orbit attempts to maintain a fixed radius.
+				MovementComponent->Velocity = FVector::ZeroVector;
+				MovementComponent->UpdateComponentVelocity();
+			}
+		}
+	}
+}
+
 void ADebugCameraController::UpdateRotationForOrbit(float DeltaTime)
 {
 	APawn* const CurrentPawn = GetPawnOrSpectator();
@@ -414,9 +443,31 @@ void ADebugCameraController::UpdateRotationForOrbit(float DeltaTime)
 	{
 		bool bUpdatePawn = false;
 		FRotator ViewRotation = GetControlRotation();
+	
+		if (!CurrentPawn->GetLastMovementInputVector().IsZero())
+		{
+			FRotationMatrix ObjectToWorld(ViewRotation);
+			FVector MoveDelta(CurrentPawn->GetActorLocation() - LastOrbitPawnLocation);
+			FVector MoveDeltaObj = ObjectToWorld.GetTransposed().TransformVector(MoveDelta);
 
-		// Handle rotation input
-		if (!RotationInput.IsZero())
+			// Handle either forward or lateral motion but not both, because small forward
+			// motion deltas while moving laterally cause the distance from pivot to drift
+			if (FMath::IsNearlyZero(MoveDeltaObj.Y, 0.01f))
+			{
+				// Clamp delta to avoid flipping to opposite view
+				const float ForwardScale = 3.0f;
+				OrbitRadius = (MoveDeltaObj.X * ForwardScale > OrbitRadius - MIN_ORBIT_RADIUS) ? MIN_ORBIT_RADIUS : OrbitRadius - MoveDeltaObj.X * ForwardScale;
+			}
+			else
+			{
+				// Apply lateral movement component, constraining distance from orbit pivot
+				const float LateralScale = 2.0f;
+				FVector LateralDelta = ObjectToWorld.TransformVector(FVector(0.0f, MoveDeltaObj.Y * LateralScale, 0.0f));
+				ViewRotation = (OrbitPivot - LastOrbitPawnLocation - LateralDelta).ToOrientationRotator();
+			}
+			bUpdatePawn = true;
+		}
+		else if (!RotationInput.IsZero())
 		{
 			ViewRotation += RotationInput;
 
@@ -437,38 +488,6 @@ void ADebugCameraController::UpdateRotationForOrbit(float DeltaTime)
 
 			bUpdatePawn = true;
 		}
-		else // Handle movement input
-		{
-			FVector OrbitPawnLocation = CurrentPawn->GetActorLocation();
-			FVector MoveDelta(OrbitPawnLocation - LastOrbitPawnLocation);
-
-			if (!MoveDelta.IsZero())
-			{
-				FRotationMatrix ObjectToWorld(ViewRotation);
-				FVector MoveDeltaObj = ObjectToWorld.GetTransposed().TransformVector(MoveDelta);
-
-				// Handle either forward or lateral motion but not both, because small forward
-				// motion deltas while moving laterally cause the distance from pivot to drift
-				if (FMath::IsNearlyZero(MoveDeltaObj.Y, 0.01f))
-				{
-					// Clamp delta to avoid flipping to opposite view
-					float MaxDelta = OrbitRadius - MIN_ORBIT_RADIUS;
-					if (MoveDeltaObj.X > MaxDelta)
-					{
-						MoveDeltaObj.X = MaxDelta;
-					}
-
-					OrbitRadius -= MoveDeltaObj.X;
-				}
-				else
-				{
-					// Apply lateral movement component, constraining distance from orbit pivot
-					FVector LateralDelta = ObjectToWorld.TransformVector(FVector(0.0f, MoveDeltaObj.Y, 0.0f));
-					ViewRotation = (OrbitPivot - LastOrbitPawnLocation - LateralDelta).ToOrientationRotator();
-				}
-				bUpdatePawn = true;
-			}
-		}
 
 		if (bUpdatePawn)
 		{
@@ -476,7 +495,7 @@ void ADebugCameraController::UpdateRotationForOrbit(float DeltaTime)
 			CurrentPawn->SetActorLocation(LastOrbitPawnLocation);
 
 			SetControlRotation(ViewRotation);
-			CurrentPawn->FaceRotation(ViewRotation, DeltaTime);
+			CurrentPawn->FaceRotation(ViewRotation);
 		}
 	}
 }
