Fix a bug where some static mesh instances are not added to the streamer when they should be

diff --git a/Engine/Source/Runtime/Engine/Classes/Engine/TextureStreamingTypes.h b/Engine/Source/Runtime/Engine/Classes/Engine/TextureStreamingTypes.h
--- a/Engine/Source/Runtime/Engine/Classes/Engine/TextureStreamingTypes.h
+++ b/Engine/Source/Runtime/Engine/Classes/Engine/TextureStreamingTypes.h
@@ -62,22 +62,39 @@ struct FStreamingRenderAssetPrimitiveInfo
 	UPROPERTY()
 	uint32 PackedRelativeBox;
 
+	/** For mesh components, texel factors are their sphere bound diameters that are 0 when unregistered */
+	UPROPERTY(Transient)
+	uint32 bAllowInvalidTexelFactorWhenUnregistered : 1;
+
 	FStreamingRenderAssetPrimitiveInfo() : 
 		RenderAsset(nullptr),
 		Bounds(ForceInit), 
 		TexelFactor(1.0f),
-		PackedRelativeBox(0)
+		PackedRelativeBox(0),
+		bAllowInvalidTexelFactorWhenUnregistered(false)
 	{
 	}
 
-	FStreamingRenderAssetPrimitiveInfo(UStreamableRenderAsset* InAsset, const FBoxSphereBounds& InBounds, float InTexelFactor, uint32 InPackedRelativeBox = 0) :
+	FStreamingRenderAssetPrimitiveInfo(
+		UStreamableRenderAsset* InAsset,
+		const FBoxSphereBounds& InBounds,
+		float InTexelFactor,
+		uint32 InPackedRelativeBox = 0,
+		bool bInAllowInvalidTexelFactorWhenUnregistered = false) :
 		RenderAsset(InAsset),
 		Bounds(InBounds), 
 		TexelFactor(InTexelFactor),
-		PackedRelativeBox(InPackedRelativeBox)
+		PackedRelativeBox(InPackedRelativeBox),
+		bAllowInvalidTexelFactorWhenUnregistered(bInAllowInvalidTexelFactorWhenUnregistered)
 	{
 	}
 
+	bool CanBeStreamedByDistance(bool bOwningCompRegistered) const
+	{
+		return (TexelFactor > SMALL_NUMBER || (!bOwningCompRegistered && bAllowInvalidTexelFactorWhenUnregistered))
+			&& (Bounds.SphereRadius > SMALL_NUMBER || !bOwningCompRegistered)
+			&& ensure(FMath::IsFinite(TexelFactor));
+	}
 };
 
 /** 
diff --git a/Engine/Source/Runtime/Engine/Private/Components/PrimitiveComponent.cpp b/Engine/Source/Runtime/Engine/Private/Components/PrimitiveComponent.cpp
--- a/Engine/Source/Runtime/Engine/Private/Components/PrimitiveComponent.cpp
+++ b/Engine/Source/Runtime/Engine/Private/Components/PrimitiveComponent.cpp
@@ -391,7 +391,8 @@ void UPrimitiveComponent::GetStreamingRenderAssetInfoWithNULLRemoval(FStreamingT
 #endif
 
 				// Other wise check that everything is setup right. If the component is not yet registered, then the bound data is irrelevant.
-				const bool bCanBeStreamedByDistance = Info.TexelFactor > SMALL_NUMBER && (Info.Bounds.SphereRadius > SMALL_NUMBER || !IsRegistered()) && ensure(FMath::IsFinite(Info.TexelFactor));
+				const bool bCanBeStreamedByDistance = Info.CanBeStreamedByDistance(IsRegistered());
+
 				if (!bForceMipStreaming && !bCanBeStreamedByDistance && Info.TexelFactor >= 0.f)
 				{
 					OutStreamingRenderAssets.RemoveAtSwap(Index--);
diff --git a/Engine/Source/Runtime/Engine/Private/Components/StaticMeshComponent.cpp b/Engine/Source/Runtime/Engine/Private/Components/StaticMeshComponent.cpp
--- a/Engine/Source/Runtime/Engine/Private/Components/StaticMeshComponent.cpp
+++ b/Engine/Source/Runtime/Engine/Private/Components/StaticMeshComponent.cpp
@@ -846,7 +846,7 @@ void UStaticMeshComponent::GetStreamingRenderAssetInfo(FStreamingTextureLevelCon
 	if (IsStreamingRenderAsset(GetStaticMesh()))
 	{
 		const float TexelFactor = ForcedLodModel > 0 ? -(GetStaticMesh()->RenderData->LODResources.Num() - ForcedLodModel + 1) : Bounds.SphereRadius * 2.f;
-		new (OutStreamingRenderAssets) FStreamingRenderAssetPrimitiveInfo(GetStaticMesh(), Bounds, TexelFactor, PackedRelativeBox_Identity);
+		new (OutStreamingRenderAssets) FStreamingRenderAssetPrimitiveInfo(GetStaticMesh(), Bounds, TexelFactor, PackedRelativeBox_Identity, true);
 	}
 }
 
