Fix UVolumeTexture::UpdateSourceFromSourceTexture() that was using the wrong number of bytes per pixel during copy due to using a mismatched enum.

diff --git a/Engine/Source/Runtime/Engine/Private/VolumeTexture.cpp b/Engine/Source/Runtime/Engine/Private/VolumeTexture.cpp
--- a/Engine/Source/Runtime/Engine/Private/VolumeTexture.cpp
+++ b/Engine/Source/Runtime/Engine/Private/VolumeTexture.cpp
@@ -43,8 +43,7 @@ bool UVolumeTexture::UpdateSourceFromSourceTexture()
 		const int32 TileSizeZ = FMath::Min<int32>(Num2DTileX * Num2DTileY, MAX_VOLUME_TEXTURE_DEPTH);
 		if (TileSizeZ > 0)
 		{
-			const FPixelFormatInfo& InitialFormat = GPixelFormats[InitialSource.GetFormat()];
-			const int32 FormatDataSize = InitialFormat.BlockBytes;
+			const int32 FormatDataSize = InitialSource.GetBytesPerPixel();
 			if (FormatDataSize > 0)
 			{
 				TArray<uint8> Ref2DData;
