Safeguard against modifying the VirtualLoop map inline.

diff --git a/Engine/Source/Runtime/Engine/Private/AudioDevice.cpp b/Engine/Source/Runtime/Engine/Private/AudioDevice.cpp
--- a/Engine/Source/Runtime/Engine/Private/AudioDevice.cpp
+++ b/Engine/Source/Runtime/Engine/Private/AudioDevice.cpp
@@ -5765,11 +5765,23 @@ void FAudioDevice::UpdateProximityRetriggerComponents()
 
 void FAudioDevice::UpdateVirtualLoops()
 {
+	static TArray<FAudioVirtualLoop*> VirtualLoopsToRetrigger;
 	for (TPair<FActiveSound*, FAudioVirtualLoop*>& Pair : VirtualLoops)
 	{
-		check(Pair.Value);
-		Pair.Value->Update();
+		FAudioVirtualLoop* VirtualLoop = Pair.Value;
+		check(VirtualLoop);
+
+		if (VirtualLoop->Update())
+		{
+			VirtualLoopsToRetrigger.Add(VirtualLoop);
+		}
+	}
+
+	for (FAudioVirtualLoop* RetriggerLoop : VirtualLoopsToRetrigger)
+	{
+		RetriggerVirtualLoop(*RetriggerLoop);
 	}
+	VirtualLoopsToRetrigger.Reset();
 }
 
 #if WITH_EDITOR
diff --git a/Engine/Source/Runtime/Engine/Private/AudioVirtualLoop.cpp b/Engine/Source/Runtime/Engine/Private/AudioVirtualLoop.cpp
--- a/Engine/Source/Runtime/Engine/Private/AudioVirtualLoop.cpp
+++ b/Engine/Source/Runtime/Engine/Private/AudioVirtualLoop.cpp
@@ -174,7 +174,7 @@ void FAudioVirtualLoop::SetActiveSound(const FActiveSound& InActiveSound)
 	ActiveSound->TargetAdjustVolumeStopTime = InActiveSound.TargetAdjustVolumeStopTime;
 }
 
-void FAudioVirtualLoop::Update()
+bool FAudioVirtualLoop::Update()
 {
 	check(AudioDevice);
 	check(ActiveSound);
@@ -189,17 +189,17 @@ void FAudioVirtualLoop::Update()
 		TimeSinceLastUpdate += DeltaTime;
 		if (UpdateInterval > TimeSinceLastUpdate)
 		{
-			return;
+			return false;
 		}
 		TimeSinceLastUpdate = 0.0f;
 	}
 
-	if (IsInAudibleRange(*AudioDevice, *ActiveSound))
-	{
-		AudioDevice->RetriggerVirtualLoop(*this);
-	}
-	else
+	// If not audible, update when will be checked again and return false
+	if (!IsInAudibleRange(*AudioDevice, *ActiveSound))
 	{
 		CalculateUpdateInterval();
+		return false;
 	}
+
+	return true;
 }
diff --git a/Engine/Source/Runtime/Engine/Public/AudioVirtualLoop.h b/Engine/Source/Runtime/Engine/Public/AudioVirtualLoop.h
--- a/Engine/Source/Runtime/Engine/Public/AudioVirtualLoop.h
+++ b/Engine/Source/Runtime/Engine/Public/AudioVirtualLoop.h
@@ -46,8 +46,9 @@ public:
 
 	/**
 	  * Updates the loop and checks if ready to play (or 'realize').
+	  * Returns whether or not the sound is ready to be realized.
 	  */
-	void Update();
+	bool Update();
 
 private:
 	void SetActiveSound(const FActiveSound& InActiveSound);
