CSV Profiler support for naming and instrumenting event waits + Instrumentation covering the main waits on the game thread

diff --git a/Engine/Source/Runtime/Core/Private/GenericPlatform/GenericPlatformProcess.cpp b/Engine/Source/Runtime/Core/Private/GenericPlatform/GenericPlatformProcess.cpp
--- a/Engine/Source/Runtime/Core/Private/GenericPlatform/GenericPlatformProcess.cpp
+++ b/Engine/Source/Runtime/Core/Private/GenericPlatform/GenericPlatformProcess.cpp
@@ -356,7 +356,7 @@ bool FPThreadEvent::Wait(uint32 WaitTime, const bool bIgnoreThreadIdleStats /*=
 	WaitForStats();
 
 	SCOPE_CYCLE_COUNTER(STAT_EventWait);
-	CSV_SCOPED_TIMING_STAT_EXCLUSIVE_CONDITIONAL(EventWait, IsInGameThread());
+	CSV_SCOPED_WAIT_CONDITIONAL(WaitTime > 0 && IsInGameThread());
 	FThreadIdleStats::FScopeIdle Scope(bIgnoreThreadIdleStats);
 
 	check(bInitialized);
diff --git a/Engine/Source/Runtime/Core/Private/ProfilingDebugging/CsvProfiler.cpp b/Engine/Source/Runtime/Core/Private/ProfilingDebugging/CsvProfiler.cpp
--- a/Engine/Source/Runtime/Core/Private/ProfilingDebugging/CsvProfiler.cpp
+++ b/Engine/Source/Runtime/Core/Private/ProfilingDebugging/CsvProfiler.cpp
@@ -61,6 +61,9 @@ CSV_DEFINE_STAT_GLOBAL(FrameTime);
 
 DEFINE_LOG_CATEGORY_STATIC(LogCsvProfiler, Log, All);
 
+const char * GDefaultWaitStatName = "EventWait";
+const char * GIgnoreWaitStatName =  "[IGNORE]";
+
 TAutoConsoleVariable<int32> CVarCsvBlockOnCaptureEnd(
 	TEXT("csv.BlockOnCaptureEnd"), 
 	0,
@@ -1636,6 +1639,21 @@ public:
 			((uint64)Events.GetAllocatedSize());
 	}
 
+	CSV_PROFILER_INLINE const char * GetWaitStatName() const
+	{
+		return WaitStatNameStack.Num() == 0 ? GDefaultWaitStatName : WaitStatNameStack.Last();
+	}
+
+	CSV_PROFILER_INLINE void PushWaitStatName(const char * WaitStatName)
+	{
+		WaitStatNameStack.Push(WaitStatName);
+	}
+	CSV_PROFILER_INLINE const char* PopWaitStatName()
+	{
+		check(WaitStatNameStack.Num() > 0);
+		return WaitStatNameStack.Pop();
+	}
+
 	// Raw stat data (written from the thread)
 	TSingleProducerSingleConsumerList<FCsvTimingMarker, 256> TimingMarkers;
 	TSingleProducerSingleConsumerList<FCsvCustomStat, 256> CustomStats;
@@ -1645,6 +1663,7 @@ public:
 	const FString ThreadName;
 
 	class FCsvProfilerThreadDataProcessor* DataProcessor;
+	TArray<const char*> WaitStatNameStack;
 };
 
 uint32 FCsvProfilerThreadData::TlsSlot = 0;
@@ -2650,6 +2669,57 @@ void FCsvProfiler::EndExclusiveStat(const char * StatName)
 #endif
 }
 
+
+void FCsvProfiler::BeginSetWaitStat(const char * StatName)
+{
+#if RECORD_TIMESTAMPS
+	if (GCsvProfilerIsCapturing && GCsvCategoriesEnabled[CSV_CATEGORY_INDEX(Exclusive)])
+	{
+		LLM_SCOPE(ELLMTag::CsvProfiler);
+		FCsvProfilerThreadData::Get().PushWaitStatName(StatName == nullptr ? GIgnoreWaitStatName : StatName);
+	}
+#endif
+}
+
+void FCsvProfiler::EndSetWaitStat()
+{
+#if RECORD_TIMESTAMPS
+	if (GCsvProfilerIsCapturing && GCsvCategoriesEnabled[CSV_CATEGORY_INDEX(Exclusive)])
+	{
+		FCsvProfilerThreadData::Get().PopWaitStatName();
+	}
+#endif
+}
+
+void FCsvProfiler::BeginWait()
+{
+#if RECORD_TIMESTAMPS
+	if (GCsvProfilerIsCapturing && GCsvCategoriesEnabled[CSV_CATEGORY_INDEX(Exclusive)])
+	{
+		const char* WaitStatName = FCsvProfilerThreadData::Get().GetWaitStatName();
+		if (WaitStatName != GIgnoreWaitStatName)
+		{
+			FCsvProfilerThreadData::Get().AddTimestampExclusiveBegin(WaitStatName);
+		}
+	}
+#endif
+}
+
+void FCsvProfiler::EndWait()
+{
+#if RECORD_TIMESTAMPS
+	if (GCsvProfilerIsCapturing && GCsvCategoriesEnabled[CSV_CATEGORY_INDEX(Exclusive)])
+	{
+		const char* WaitStatName = FCsvProfilerThreadData::Get().GetWaitStatName();
+		if (WaitStatName != GIgnoreWaitStatName)
+		{
+			FCsvProfilerThreadData::Get().AddTimestampExclusiveEnd(FCsvProfilerThreadData::Get().GetWaitStatName());
+		}
+	}
+#endif
+}
+
+
 void FCsvProfiler::RecordEventfInternal(int32 CategoryIndex, const TCHAR* Fmt, ...)
 {
 	if (GCsvProfilerIsCapturing && GCsvCategoriesEnabled[CategoryIndex])
diff --git a/Engine/Source/Runtime/Core/Private/Windows/WindowsPlatformProcess.cpp b/Engine/Source/Runtime/Core/Private/Windows/WindowsPlatformProcess.cpp
--- a/Engine/Source/Runtime/Core/Private/Windows/WindowsPlatformProcess.cpp
+++ b/Engine/Source/Runtime/Core/Private/Windows/WindowsPlatformProcess.cpp
@@ -1262,8 +1262,8 @@ bool FEventWin::Wait(uint32 WaitTime, const bool bIgnoreThreadIdleStats /*= fals
 	WaitForStats();
 
 	SCOPE_CYCLE_COUNTER( STAT_EventWait );
-	CSV_SCOPED_TIMING_STAT_EXCLUSIVE_CONDITIONAL(EventWait, IsInGameThread());
-	check( Event );
+	CSV_SCOPED_WAIT_CONDITIONAL(WaitTime > 0 && IsInGameThread());
+	check(Event);
 
 	FThreadIdleStats::FScopeIdle Scope( bIgnoreThreadIdleStats );
 	return (WaitForSingleObject( Event, WaitTime ) == WAIT_OBJECT_0);
diff --git a/Engine/Source/Runtime/Core/Public/ProfilingDebugging/CsvProfiler.h b/Engine/Source/Runtime/Core/Public/ProfilingDebugging/CsvProfiler.h
--- a/Engine/Source/Runtime/Core/Public/ProfilingDebugging/CsvProfiler.h
+++ b/Engine/Source/Runtime/Core/Public/ProfilingDebugging/CsvProfiler.h
@@ -49,6 +49,10 @@
 #define CSV_SCOPED_TIMING_STAT_EXCLUSIVE(StatName)				FScopedCsvStatExclusive _ScopedCsvStatExclusive_ ## StatName (#StatName);
 #define CSV_SCOPED_TIMING_STAT_EXCLUSIVE_CONDITIONAL(StatName,Condition) FScopedCsvStatExclusiveConditional _ScopedCsvStatExclusive_ ## StatName (#StatName,Condition);
 
+#define CSV_SCOPED_WAIT_CONDITIONAL(Condition)					FScopedCsvWaitConditional _ScopedCsvWait(Condition);
+#define CSV_SCOPED_SET_WAIT_STAT(StatName)						FScopedCsvSetWaitStat _ScopedCsvSetWaitStat ## StatName("EventWait/"#StatName);
+#define CSV_SCOPED_SET_WAIT_STAT_IGNORE()						FScopedCsvSetWaitStat _ScopedCsvSetWaitStat ## StatName();
+
 #define CSV_CUSTOM_STAT(Category,StatName,Value,Op)				FCsvProfiler::RecordCustomStat(#StatName, CSV_CATEGORY_INDEX(Category), Value, Op)
 #define CSV_CUSTOM_STAT_GLOBAL(StatName,Value,Op) 				FCsvProfiler::RecordCustomStat(#StatName, CSV_CATEGORY_INDEX_GLOBAL, Value, Op)
 
@@ -80,6 +84,9 @@
   #define CSV_SCOPED_TIMING_STAT_GLOBAL(StatName)					
   #define CSV_SCOPED_TIMING_STAT_EXCLUSIVE(StatName)
   #define CSV_SCOPED_TIMING_STAT_EXCLUSIVE_CONDITIONAL(StatName,Condition)
+  #define CSV_SCOPED_WAIT_CONDITIONAL(Condition)
+  #define CSV_SCOPED_SET_WAIT_STAT(StatName)
+  #define CSV_SCOPED_SET_WAIT_STAT_IGNORE()
   #define CSV_CUSTOM_STAT(Category,StatName,Value,Op)				
   #define CSV_CUSTOM_STAT_GLOBAL(StatName,Value,Op) 				
   #define CSV_DEFINE_STAT(Category,StatName)						
@@ -215,6 +222,12 @@ public:
 		RecordEventfInternal(CategoryIndex, Fmt, Args...);
 	}
 
+	CORE_API static void BeginSetWaitStat(const char * StatName);
+	CORE_API static void EndSetWaitStat();
+
+	CORE_API static void BeginWait();
+	CORE_API static void EndWait();
+
 	/** Singleton interface */
 	CORE_API bool IsCapturing();
 	CORE_API bool IsCapturing_Renderthread();
@@ -366,6 +379,46 @@ public:
 	bool bCondition;
 };
 
+class FScopedCsvWaitConditional
+{
+public:
+	FScopedCsvWaitConditional(bool bInCondition)
+		: bCondition(bInCondition)
+	{
+		if (bCondition)
+		{
+			FCsvProfiler::BeginWait();
+#if CSV_EXCLUSIVE_TIMING_STATS_EMIT_NAMED_EVENTS
+			FPlatformMisc::BeginNamedEvent(FColor(255, 128, 128), EventWait);
+#endif
+		}
+	}
+
+	~FScopedCsvWaitConditional()
+	{
+#if CSV_EXCLUSIVE_TIMING_STATS_EMIT_NAMED_EVENTS
+		FPlatformMisc::EndNamedEvent();
+#endif
+		FCsvProfiler::EndWait();
+	}
+	bool bCondition;
+};
+
+class FScopedCsvSetWaitStat
+{
+public:
+	FScopedCsvSetWaitStat(const char * InStatName = nullptr)
+		: StatName(InStatName)
+	{
+		FCsvProfiler::BeginSetWaitStat(StatName);
+	}
+
+	~FScopedCsvSetWaitStat()
+	{
+		FCsvProfiler::EndSetWaitStat();
+	}
+	const char * StatName;
+};
 
 struct FCsvCategory
 {
diff --git a/Engine/Source/Runtime/Engine/Private/LevelTick.cpp b/Engine/Source/Runtime/Engine/Private/LevelTick.cpp
--- a/Engine/Source/Runtime/Engine/Private/LevelTick.cpp
+++ b/Engine/Source/Runtime/Engine/Private/LevelTick.cpp
@@ -981,7 +981,8 @@ void EndSendEndOfFrameUpdatesDrawEvent(FSendAllEndOfFrameUpdates* SendAllEndOfFr
 void UWorld::SendAllEndOfFrameUpdates()
 {
 	SCOPE_CYCLE_COUNTER(STAT_PostTickComponentUpdate);
-	CSV_SCOPED_TIMING_STAT_EXCLUSIVE(PostTickComponentMisc);
+	CSV_SCOPED_TIMING_STAT_EXCLUSIVE(EndOfFrameUpdates);
+	CSV_SCOPED_SET_WAIT_STAT(EndOfFrameUpdates);
 
 	if (!HasEndOfFrameUpdates())
 	{
@@ -1388,6 +1389,7 @@ void UWorld::Tick( ELevelTick TickType, float DeltaSeconds )
 	SCOPE_TIME_GUARD(TEXT("UWorld::Tick"));
 	SCOPED_NAMED_EVENT(UWorld_Tick, FColor::Orange);
 	CSV_SCOPED_TIMING_STAT_EXCLUSIVE(WorldTickMisc);
+	CSV_SCOPED_SET_WAIT_STAT(WorldTickMisc);
 
 	if (GIntraFrameDebuggingGameThread)
 	{
@@ -1441,6 +1443,7 @@ void UWorld::Tick( ELevelTick TickType, float DeltaSeconds )
 	{
 		SCOPE_CYCLE_COUNTER(STAT_NetWorldTickTime);
 		CSV_SCOPED_TIMING_STAT_EXCLUSIVE(NetworkIncoming);
+		CSV_SCOPED_SET_WAIT_STAT(NetworkTick);
 		SCOPE_TIME_GUARD(TEXT("UWorld::Tick - NetTick"));
 		LLM_SCOPE(ELLMTag::Networking);
 		// Update the net code and fetch all incoming packets.
@@ -1490,6 +1493,7 @@ void UWorld::Tick( ELevelTick TickType, float DeltaSeconds )
 	// give the async loading code more time if we're performing a high priority load or are in seamless travel
 	if (Info->bHighPriorityLoading || Info->bHighPriorityLoadingLocal || IsInSeamlessTravel())
 	{
+		CSV_SCOPED_SET_WAIT_STAT(AsyncLoading);
 		// Force it to use the entire time slice, even if blocked on I/O
 		ProcessAsyncLoading(true, true, GPriorityAsyncLoadingExtraTime / 1000.0f);
 	}
@@ -1565,6 +1569,7 @@ void UWorld::Tick( ELevelTick TickType, float DeltaSeconds )
 			{
 				SCOPE_TIME_GUARD_MS(TEXT("UWorld::Tick - TG_PrePhysics"), 10);
 				SCOPE_CYCLE_COUNTER(STAT_TG_PrePhysics);
+				CSV_SCOPED_SET_WAIT_STAT(PrePhysics);
 				RunTickGroup(TG_PrePhysics);
 			}
 			bInTick = false;
@@ -1573,22 +1578,26 @@ void UWorld::Tick( ELevelTick TickType, float DeltaSeconds )
 			{
 				SCOPE_CYCLE_COUNTER(STAT_TG_StartPhysics);
 				SCOPE_TIME_GUARD_MS(TEXT("UWorld::Tick - TG_StartPhysics"), 10);
-				RunTickGroup(TG_StartPhysics); 
+				CSV_SCOPED_SET_WAIT_STAT(StartPhysics);
+				RunTickGroup(TG_StartPhysics);
 			}
 			{
 				SCOPE_CYCLE_COUNTER(STAT_TG_DuringPhysics);
 				SCOPE_TIME_GUARD_MS(TEXT("UWorld::Tick - TG_DuringPhysics"), 10);
+				CSV_SCOPED_SET_WAIT_STAT(DuringPhysics);
 				RunTickGroup(TG_DuringPhysics, false); // No wait here, we should run until idle though. We don't care if all of the async ticks are done before we start running post-phys stuff
 			}
 			TickGroup = TG_EndPhysics; // set this here so the current tick group is correct during collision notifies, though I am not sure it matters. 'cause of the false up there^^^
 			{
 				SCOPE_CYCLE_COUNTER(STAT_TG_EndPhysics);
 				SCOPE_TIME_GUARD_MS(TEXT("UWorld::Tick - TG_EndPhysics"), 10);
+				CSV_SCOPED_SET_WAIT_STAT(EndPhysics);
 				RunTickGroup(TG_EndPhysics);
 			}
 			{
 				SCOPE_CYCLE_COUNTER(STAT_TG_PostPhysics);
 				SCOPE_TIME_GUARD_MS(TEXT("UWorld::Tick - TG_PostPhysics"), 10);
+				CSV_SCOPED_SET_WAIT_STAT(PostPhysics);
 				RunTickGroup(TG_PostPhysics);
 			}
 	
@@ -1683,7 +1692,7 @@ void UWorld::Tick( ELevelTick TickType, float DeltaSeconds )
 				RunTickGroup(TG_LastDemotable);
 			}
 
-			FTickTaskManagerInterface::Get().EndFrame(); 
+			FTickTaskManagerInterface::Get().EndFrame();
 		}
 	}
 
diff --git a/Engine/Source/Runtime/Engine/Private/Streaming/StreamingManagerTexture.cpp b/Engine/Source/Runtime/Engine/Private/Streaming/StreamingManagerTexture.cpp
--- a/Engine/Source/Runtime/Engine/Private/Streaming/StreamingManagerTexture.cpp
+++ b/Engine/Source/Runtime/Engine/Private/Streaming/StreamingManagerTexture.cpp
@@ -1339,6 +1339,7 @@ void FRenderAssetStreamingManager::UpdateResourceStreaming( float DeltaTime, boo
 
 	SCOPE_CYCLE_COUNTER(STAT_RenderAssetStreaming_GameThreadUpdateTime);
 	CSV_SCOPED_TIMING_STAT_EXCLUSIVE(RenderAssetStreaming);
+	CSV_SCOPED_SET_WAIT_STAT(RenderAssetStreaming);
 
 	const bool bUseThreadingForPerf = FApp::ShouldUseThreadingForPerformance();
 
