Fix for NaN ensure triggering when manually assigning uint as float through FVector4 constuctor

#rnx

diff --git a/Engine/Source/Runtime/Engine/Private/PrimitiveUniformShaderParameters.cpp b/Engine/Source/Runtime/Engine/Private/PrimitiveUniformShaderParameters.cpp
--- a/Engine/Source/Runtime/Engine/Private/PrimitiveUniformShaderParameters.cpp
+++ b/Engine/Source/Runtime/Engine/Private/PrimitiveUniformShaderParameters.cpp
@@ -114,7 +114,8 @@ void FPrimitiveSceneShaderData::Setup(const FPrimitiveUniformShaderParameters& P
 	Data[24] = FVector4(PrimitiveUniformShaderParameters.LocalObjectBoundsMax, 0.0f);
 	Data[24].W = *(const float*)&PrimitiveUniformShaderParameters.LightmapDataIndex;
 
-	Data[25] = FVector4(PrimitiveUniformShaderParameters.PreSkinnedLocalBounds, *(const float*)&PrimitiveUniformShaderParameters.SingleCaptureIndex); 
+	Data[25] = FVector4(PrimitiveUniformShaderParameters.PreSkinnedLocalBounds, 0.0f);
+	Data[25].W = *(const float*)&PrimitiveUniformShaderParameters.SingleCaptureIndex;
 
 	Data[26] = FVector4(0.0f, 0.0f, 0.0f, 0.0f);
 	Data[26].X = *(const float*)&PrimitiveUniformShaderParameters.OutputVelocity;
