Fix shader compilation on Linux.

diff --git a/Engine/Shaders/Private/IESLightProfilesCommon.ush b/Engine/Shaders/Private/IESLightProfilesCommon.ush
--- a/Engine/Shaders/Private/IESLightProfilesCommon.ush
+++ b/Engine/Shaders/Private/IESLightProfilesCommon.ush
@@ -36,6 +36,8 @@ float ComputeLightProfileMultiplier(float3 WorldPosition, float3 LightPosition,
 #endif
 }
 
+#if RAYHITGROUPSHADER || RAYGENSHADER
+
 // Apply IES light profile texture. Each row of the texture stores one profile.
 float ComputeRayTracingLightProfileMultiplier(
 	float3 WorldPosition,
@@ -61,3 +63,4 @@ float ComputeRayTracingLightProfileMultiplier(
 
 	return Result;
 }
+#endif // RAYHITGROUPSHADER || RAYGENSHADER
diff --git a/Engine/Source/Runtime/Engine/Private/ShaderCompiler/ShaderCompiler.cpp b/Engine/Source/Runtime/Engine/Private/ShaderCompiler/ShaderCompiler.cpp
--- a/Engine/Source/Runtime/Engine/Private/ShaderCompiler/ShaderCompiler.cpp
+++ b/Engine/Source/Runtime/Engine/Private/ShaderCompiler/ShaderCompiler.cpp
@@ -2989,6 +2989,7 @@ void GlobalBeginCompileShader(
 		Input.Environment.SetDefine(TEXT("COMPUTESHADER"), Target.Frequency == SF_Compute);
 #if RHI_RAYTRACING
 		Input.Environment.SetDefine(TEXT("RAYHITGROUPSHADER"), Target.Frequency == SF_RayHitGroup);
+		Input.Environment.SetDefine(TEXT("RAYGENSHADER"), Target.Frequency == SF_RayGen);
 #endif
 	}
 
