Fix shader compilation in windows with ray tracing enabled.

diff --git a/Engine/Shaders/Private/IESLightProfilesCommon.ush b/Engine/Shaders/Private/IESLightProfilesCommon.ush
--- a/Engine/Shaders/Private/IESLightProfilesCommon.ush
+++ b/Engine/Shaders/Private/IESLightProfilesCommon.ush
@@ -36,7 +36,7 @@ float ComputeLightProfileMultiplier(float3 WorldPosition, float3 LightPosition,
 #endif
 }
 
-#if RAYHITGROUPSHADER || RAYGENSHADER
+#if RAYHITGROUPSHADER || RAYGENSHADER || RAYMISSSHADER
 
 // Apply IES light profile texture. Each row of the texture stores one profile.
 float ComputeRayTracingLightProfileMultiplier(
diff --git a/Engine/Source/Runtime/Engine/Private/ShaderCompiler/ShaderCompiler.cpp b/Engine/Source/Runtime/Engine/Private/ShaderCompiler/ShaderCompiler.cpp
--- a/Engine/Source/Runtime/Engine/Private/ShaderCompiler/ShaderCompiler.cpp
+++ b/Engine/Source/Runtime/Engine/Private/ShaderCompiler/ShaderCompiler.cpp
@@ -2990,6 +2990,7 @@ void GlobalBeginCompileShader(
 #if RHI_RAYTRACING
 		Input.Environment.SetDefine(TEXT("RAYHITGROUPSHADER"), Target.Frequency == SF_RayHitGroup);
 		Input.Environment.SetDefine(TEXT("RAYGENSHADER"), Target.Frequency == SF_RayGen);
+		Input.Environment.SetDefine(TEXT("RAYMISSSHADER"), Target.Frequency == SF_RayMiss);
 #endif
 	}
 
