Fix crash when skeleton with virtual bones has its heirarchy replaced.

diff --git a/Engine/Source/Runtime/Engine/Private/Animation/Skeleton.cpp b/Engine/Source/Runtime/Engine/Private/Animation/Skeleton.cpp
--- a/Engine/Source/Runtime/Engine/Private/Animation/Skeleton.cpp
+++ b/Engine/Source/Runtime/Engine/Private/Animation/Skeleton.cpp
@@ -1019,6 +1019,17 @@ void USkeleton::HandleSkeletonHierarchyChange()
 	// Clear exiting MeshLinkUp tables.
 	ClearCacheData();
 
+	for (int i = VirtualBones.Num() - 1; i >= 0; --i)
+	{
+		FVirtualBone& VB = VirtualBones[i];
+		if (ReferenceSkeleton.FindRawBoneIndex(VB.SourceBoneName) == INDEX_NONE ||
+			ReferenceSkeleton.FindRawBoneIndex(VB.TargetBoneName) == INDEX_NONE)
+		{
+			//Virtual Bone no longer valid
+			VirtualBones.RemoveAtSwap(i, 1, false);
+		}
+	}
+
 	// Fix up loaded animations (any animations that aren't loaded will be fixed on load)
 	int32 NumLoadedAssets = 0;
 	for (TObjectIterator<UAnimationAsset> It; It; ++It)
