Fix for 8.40 client crash with instanced foliage
- Null static mesh check
- Client only, no compat issues

diff --git a/Engine/Source/Runtime/Engine/Private/HierarchicalInstancedStaticMesh.cpp b/Engine/Source/Runtime/Engine/Private/HierarchicalInstancedStaticMesh.cpp
--- a/Engine/Source/Runtime/Engine/Private/HierarchicalInstancedStaticMesh.cpp
+++ b/Engine/Source/Runtime/Engine/Private/HierarchicalInstancedStaticMesh.cpp
@@ -3209,14 +3209,17 @@ TArray<int32> UHierarchicalInstancedStaticMeshComponent::GetInstancesOverlapping
 			WorldSpaceAABB = WorldSpaceAABB.TransformBy(GetComponentTransform());
 		}
 
-		const float StaticMeshBoundsRadius = GetStaticMesh()->GetBounds().SphereRadius;
-		GatherInstancesOverlappingArea(*this, WorldSpaceAABB, 0,
-			[Sphere, StaticMeshBoundsRadius](const FMatrix& InstanceTransform)->bool
-			{
-				FSphere InstanceSphere(InstanceTransform.GetOrigin(), StaticMeshBoundsRadius * InstanceTransform.GetScaleVector().GetMax());
-				return Sphere.Intersects(InstanceSphere);
-			},
-			Result);
+		if (GetStaticMesh() != nullptr)
+		{
+			const float StaticMeshBoundsRadius = GetStaticMesh()->GetBounds().SphereRadius;
+			GatherInstancesOverlappingArea(*this, WorldSpaceAABB, 0,
+				[Sphere, StaticMeshBoundsRadius](const FMatrix& InstanceTransform)->bool
+				{
+					FSphere InstanceSphere(InstanceTransform.GetOrigin(), StaticMeshBoundsRadius * InstanceTransform.GetScaleVector().GetMax());
+					return Sphere.Intersects(InstanceSphere);
+				},
+				Result);
+		}
 		return Result;
 	}
 	else
@@ -3242,14 +3245,18 @@ TArray<int32> UHierarchicalInstancedStaticMeshComponent::GetInstancesOverlapping
 			WorldSpaceBox = WorldSpaceBox.TransformBy(GetComponentTransform());
 		}
 
-		const FBox StaticMeshBox = GetStaticMesh()->GetBounds().GetBox();
-		GatherInstancesOverlappingArea(*this, WorldSpaceBox, 0,
-			[LocalSpaceSpaceBox, StaticMeshBox](const FMatrix& InstanceTransform)->bool
-			{
-				FBox InstanceBox = StaticMeshBox.TransformBy(InstanceTransform);
-				return LocalSpaceSpaceBox.Intersect(InstanceBox);
-			},
-			Result);
+
+		if (GetStaticMesh() != nullptr)
+		{
+			const FBox StaticMeshBox = GetStaticMesh()->GetBounds().GetBox();
+			GatherInstancesOverlappingArea(*this, WorldSpaceBox, 0,
+				[LocalSpaceSpaceBox, StaticMeshBox](const FMatrix& InstanceTransform)->bool
+				{
+					FBox InstanceBox = StaticMeshBox.TransformBy(InstanceTransform);
+					return LocalSpaceSpaceBox.Intersect(InstanceBox);
+				},
+				Result);
+		}
 
 		return Result;
 	}
