Removed useless csv stat from async texture streaming

diff --git a/Engine/Source/Runtime/Engine/Private/Streaming/AsyncTextureStreaming.cpp b/Engine/Source/Runtime/Engine/Private/Streaming/AsyncTextureStreaming.cpp
--- a/Engine/Source/Runtime/Engine/Private/Streaming/AsyncTextureStreaming.cpp
+++ b/Engine/Source/Runtime/Engine/Private/Streaming/AsyncTextureStreaming.cpp
@@ -8,7 +8,6 @@ AsyncTextureStreaming.cpp: Definitions of classes used for texture streaming asy
 #include "Misc/App.h"
 #include "Streaming/StreamingManagerTexture.h"
 #include "Engine/World.h"
-#include "ProfilingDebugging/CsvProfiler.h"
 
 void FAsyncRenderAssetStreamingData::Init(
 	TArray<FStreamingViewInfo> InViewInfos,
@@ -713,7 +712,6 @@ void FRenderAssetStreamingMipCalcTask::DoWork()
 {
 	SCOPED_NAMED_EVENT(FRenderAssetStreamingMipCalcTask_DoWork, FColor::Turquoise);
 	DECLARE_SCOPE_CYCLE_COUNTER(TEXT("FRenderAssetStreamingMipCalcTask::DoWork"), STAT_FRenderAssetStreamingMipCalcTask_DoWork, STATGROUP_StreamingDetails);
-	CSV_SCOPED_TIMING_STAT_GLOBAL(FRenderAssetStreamingMipCalcTask);
 
 	// While the async task is runnning, the StreamingRenderAssets are guarantied not to be reallocated.
 	// 2 things can happen : a texture can be removed, in which case the texture will be set to null
