Adding ENABLE_VERBOSE_ANIM_PERF_TRACKING to help us cull the most frequently encountered perf stats from the general case

diff --git a/Engine/Source/Runtime/Engine/Private/Animation/AnimNodeBase.cpp b/Engine/Source/Runtime/Engine/Private/Animation/AnimNodeBase.cpp
--- a/Engine/Source/Runtime/Engine/Private/Animation/AnimNodeBase.cpp
+++ b/Engine/Source/Runtime/Engine/Private/Animation/AnimNodeBase.cpp
@@ -266,7 +266,10 @@ void FPoseLinkBase::CacheBones(const FAnimationCacheBonesContext& Context)
 
 void FPoseLinkBase::Update(const FAnimationUpdateContext& Context)
 {
+#if ENABLE_VERBOSE_ANIM_PERF_TRACKING
 	QUICK_SCOPE_CYCLE_COUNTER(STAT_FPoseLinkBase_Update);
+#endif // ENABLE_VERBOSE_ANIM_PERF_TRACKING
+
 #if DO_CHECK
 	checkf( !bProcessed, TEXT( "Update already in progress, circular link for AnimInstance [%s] Blueprint [%s]" ), \
 		*Context.AnimInstanceProxy->GetAnimInstanceName(), *GetFullNameSafe(IAnimClassInterface::GetActualAnimClass(Context.AnimInstanceProxy->GetAnimClassInterface())));
diff --git a/Engine/Source/Runtime/Engine/Public/Animation/AnimTypes.h b/Engine/Source/Runtime/Engine/Public/Animation/AnimTypes.h
--- a/Engine/Source/Runtime/Engine/Public/Animation/AnimTypes.h
+++ b/Engine/Source/Runtime/Engine/Public/Animation/AnimTypes.h
@@ -19,6 +19,9 @@ struct FPassedMarker;
 #define DEFAULT_SAMPLERATE			30.f
 #define MINIMUM_ANIMATION_LENGTH	(1/DEFAULT_SAMPLERATE)
 
+// Enable this if you want to locally measure detailed anim perf.  Disabled by default as it is introduces a lot of additional profile markers and associated overhead.
+#define ENABLE_VERBOSE_ANIM_PERF_TRACKING 0
+
 namespace EAnimEventTriggerOffsets
 {
 	enum Type
