Static analysis warning fix. Move underlying logic for to functions to help static analyzer evaluate how SoundToEvict/ActiveSound pointers are valid when operated on.
https://p4-swarm.epicgames.net/changes/5921483

diff --git a/Engine/Source/Runtime/Engine/Classes/Sound/SoundConcurrency.h b/Engine/Source/Runtime/Engine/Classes/Sound/SoundConcurrency.h
--- a/Engine/Source/Runtime/Engine/Classes/Sound/SoundConcurrency.h
+++ b/Engine/Source/Runtime/Engine/Classes/Sound/SoundConcurrency.h
@@ -242,6 +242,11 @@ private: // Methods
 	/** Finds an active sound able to be evicted based on the provided concurrency settings. */
 	FActiveSound* GetEvictableSound(const FActiveSound& NewActiveSound, const FConcurrencyGroup& ConcurrencyGroup, bool bIsRetriggering);
 
+	/** Helper functions for finding evictable sounds for the given concurrency rule */
+	FActiveSound* GetEvictableSoundStopFarthest(const FActiveSound& NewActiveSound, const FConcurrencyGroup& ConcurrencyGroup, bool bIsRetriggering) const;
+	FActiveSound* GetEvictableSoundStopLowestPriority(const FActiveSound& NewActiveSound, const FConcurrencyGroup& ConcurrencyGroup, bool bIsRetriggering) const;
+	FActiveSound* GetEvictableSoundStopOldest(const FActiveSound& NewActiveSound, const FConcurrencyGroup& ConcurrencyGroup, bool bIsRetriggering) const;
+
 	/** Handles concurrency evaluation that happens per USoundConcurrencyObject */
 	FActiveSound* EvaluateConcurrency(const FActiveSound& NewActiveSound, TArray<FConcurrencyHandle>& ConcurrencyHandles, bool bIsRetriggering);
 
diff --git a/Engine/Source/Runtime/Engine/Private/SoundConcurrency.cpp b/Engine/Source/Runtime/Engine/Private/SoundConcurrency.cpp
--- a/Engine/Source/Runtime/Engine/Private/SoundConcurrency.cpp
+++ b/Engine/Source/Runtime/Engine/Private/SoundConcurrency.cpp
@@ -327,119 +327,157 @@ FConcurrencyGroup* FSoundConcurrencyManager::CanPlaySound(const FActiveSound& Ne
 
 FActiveSound* FSoundConcurrencyManager::GetEvictableSound(const FActiveSound& NewActiveSound, const FConcurrencyGroup& ConcurrencyGroup, bool bIsRetriggering)
 {
-	// Concurrency group isn't full so of course there's room
-	if (!ConcurrencyGroup.IsFull())
-	{
-		return nullptr;
-	}
-
-	check(AudioDevice);
-	TArray<FListener>& Listeners = AudioDevice->Listeners;
-
-	const TArray<FActiveSound*>& ActiveSounds = ConcurrencyGroup.GetActiveSounds();
-	FActiveSound* EvictableSound = nullptr;
+	check(ConcurrencyGroup.IsFull());
 
 	const EMaxConcurrentResolutionRule::Type Rule = ConcurrencyGroup.GetSettings().ResolutionRule;
 	switch (Rule)
 	{
 		case EMaxConcurrentResolutionRule::PreventNew:
 		{
-			// Do nothing.  No sound is evictable as we're preventing anything new
+			return nullptr;
 		}
 		break;
 
 		case EMaxConcurrentResolutionRule::StopOldest:
 		{
-			for (int32 SoundIndex = 0; SoundIndex < ActiveSounds.Num(); ++SoundIndex)
-			{
-				FActiveSound* ActiveSound = ActiveSounds[SoundIndex];
-				if (EvictableSound == nullptr || ActiveSound->PlaybackTime > EvictableSound->PlaybackTime)
-				{
-					EvictableSound = ActiveSound;
-				}
-			}
+			return GetEvictableSoundStopOldest(NewActiveSound, ConcurrencyGroup, bIsRetriggering);
 
-			// Don't evict if attempting to re-trigger an older sound than that which is currently playing
-			if (bIsRetriggering && EvictableSound && NewActiveSound.PlaybackTime > EvictableSound->PlaybackTime)
-			{
-				EvictableSound = nullptr;
-			}
 		}
 		break;
 
 		case EMaxConcurrentResolutionRule::StopFarthestThenPreventNew:
 		case EMaxConcurrentResolutionRule::StopFarthestThenOldest:
 		{
-			int32 ClosestListenerIndex = NewActiveSound.FindClosestListener(Listeners);
-			float DistanceToStopSoundSq = FVector::DistSquared(Listeners[ClosestListenerIndex].Transform.GetTranslation(), NewActiveSound.Transform.GetTranslation());
-
-			for (FActiveSound* ActiveSound : ActiveSounds)
-			{
-				ClosestListenerIndex = ActiveSound->FindClosestListener(Listeners);
-				const float DistanceToActiveSoundSq = FVector::DistSquared(Listeners[ClosestListenerIndex].Transform.GetTranslation(), ActiveSound->Transform.GetTranslation());
-
-				if (DistanceToActiveSoundSq > DistanceToStopSoundSq)
-				{
-					DistanceToStopSoundSq = DistanceToActiveSoundSq;
-					EvictableSound = ActiveSound;
-				}
-				else if (Rule == EMaxConcurrentResolutionRule::StopFarthestThenOldest
-							&& DistanceToActiveSoundSq == DistanceToStopSoundSq
-							&& (EvictableSound == nullptr || ActiveSound->PlaybackTime > EvictableSound->PlaybackTime))
-				{
-					// Don't evict if attempting to re-trigger an older sound than that which is currently playing
-					if (!bIsRetriggering || ActiveSound->PlaybackTime > NewActiveSound.PlaybackTime)
-					{
-						DistanceToStopSoundSq = DistanceToActiveSoundSq;
-						EvictableSound = ActiveSound;
-					}
-				}
-			}
+			return GetEvictableSoundStopFarthest(NewActiveSound, ConcurrencyGroup, bIsRetriggering);
 		}
 		break;
 
 		case EMaxConcurrentResolutionRule::StopLowestPriority:
 		case EMaxConcurrentResolutionRule::StopLowestPriorityThenPreventNew:
 		{
-			// Find oldest and oldest lowest priority sound in the group
-			for (FActiveSound* CurrSound : ActiveSounds)
-			{
-				if (EvictableSound == nullptr
-					|| (CurrSound->GetPriority() < EvictableSound->GetPriority())
-					|| (CurrSound->GetPriority() == EvictableSound->GetPriority() && CurrSound->PlaybackTime > EvictableSound->PlaybackTime))
-				{
-					EvictableSound = CurrSound;
-				}
-			}
-
-			if (EvictableSound)
-			{
-				// Drop request as same priority and preventing new
-				if (Rule == EMaxConcurrentResolutionRule::StopLowestPriorityThenPreventNew
-					&& EvictableSound->GetPriority() == NewActiveSound.GetPriority())
-				{
-					EvictableSound = nullptr;
-				}
-
-				// Drop request as NewActiveSound's priority is lower than the lowest priority sound playing
-				else if (EvictableSound->GetPriority() > NewActiveSound.GetPriority())
-				{
-					EvictableSound = nullptr;
-				}
-			}
+			return GetEvictableSoundStopLowestPriority(NewActiveSound, ConcurrencyGroup, bIsRetriggering);
 		}
 		break;
 
 		// Eviction not supported by StopQuietest due to it requiring the sound to be initialized in order to calculate.
 		// Therefore, it is culled later but not evicted.
 		case EMaxConcurrentResolutionRule::StopQuietest:
+		{
+			return nullptr;
+		}
 		break;
 
 		default:
-		checkf(false, TEXT("Unknown EMaxConcurrentResolutionRule enumeration."));
+		{
+			checkf(false, TEXT("Unknown EMaxConcurrentResolutionRule enumeration."));
+		}
 		break;
 	}
 
+	return nullptr;
+}
+
+FActiveSound* FSoundConcurrencyManager::GetEvictableSoundStopOldest(const FActiveSound& NewActiveSound, const FConcurrencyGroup& ConcurrencyGroup, bool bIsRetriggering) const
+{
+	FActiveSound* EvictableSound = nullptr;
+	const TArray<FActiveSound*>& ActiveSounds = ConcurrencyGroup.GetActiveSounds();
+	for (FActiveSound* ActiveSound : ActiveSounds)
+	{
+		if (EvictableSound == nullptr || ActiveSound->PlaybackTime > EvictableSound->PlaybackTime)
+		{
+			EvictableSound = ActiveSound;
+		}
+	}
+
+	// Don't evict if attempting to re-trigger an older sound than that which is currently playing
+	if (bIsRetriggering && EvictableSound && NewActiveSound.PlaybackTime > EvictableSound->PlaybackTime)
+	{
+		return nullptr;
+	}
+
+	return EvictableSound;
+}
+
+FActiveSound* FSoundConcurrencyManager::GetEvictableSoundStopFarthest(const FActiveSound& NewActiveSound, const FConcurrencyGroup& ConcurrencyGroup, bool bIsRetriggering) const
+{
+	const EMaxConcurrentResolutionRule::Type Rule = ConcurrencyGroup.GetSettings().ResolutionRule;
+
+	check(AudioDevice);
+	TArray<FListener>& Listeners = AudioDevice->Listeners;
+
+	int32 ClosestListenerIndex = NewActiveSound.FindClosestListener(Listeners);
+	float DistanceToStopSoundSq = FVector::DistSquared(Listeners[ClosestListenerIndex].Transform.GetTranslation(), NewActiveSound.Transform.GetTranslation());
+
+	FActiveSound* EvictableSound = nullptr;
+	const TArray<FActiveSound*>& ActiveSounds = ConcurrencyGroup.GetActiveSounds();
+	for (FActiveSound* ActiveSound : ActiveSounds)
+	{
+		check(ActiveSound);
+
+		ClosestListenerIndex = ActiveSound->FindClosestListener(Listeners);
+		const float DistanceToActiveSoundSq = FVector::DistSquared(Listeners[ClosestListenerIndex].Transform.GetTranslation(), ActiveSound->Transform.GetTranslation());
+
+		// Stop Farthest
+		if (DistanceToActiveSoundSq > DistanceToStopSoundSq)
+		{
+			DistanceToStopSoundSq = DistanceToActiveSoundSq;
+			EvictableSound = ActiveSound;
+			continue;
+		}
+
+		// Stop Farthest than Oldest
+		if (Rule == EMaxConcurrentResolutionRule::StopFarthestThenOldest && DistanceToActiveSoundSq == DistanceToStopSoundSq)
+		{
+			const bool bIsOlderThanChosen = EvictableSound == nullptr || ActiveSound->PlaybackTime > EvictableSound->PlaybackTime;
+			if (bIsOlderThanChosen)
+			{
+				// Don't evict if attempting to re-trigger an older sound than that which is currently playing
+				if (!bIsRetriggering || ActiveSound->PlaybackTime > NewActiveSound.PlaybackTime)
+				{
+					DistanceToStopSoundSq = DistanceToActiveSoundSq;
+					EvictableSound = ActiveSound;
+				}
+			}
+		}
+	}
+
+	return EvictableSound;
+}
+
+FActiveSound* FSoundConcurrencyManager::GetEvictableSoundStopLowestPriority(const FActiveSound& NewActiveSound, const FConcurrencyGroup& ConcurrencyGroup, bool bIsRetriggering) const
+{
+	// Find oldest and oldest lowest priority sound in the group
+	FActiveSound* EvictableSound = nullptr;
+	const TArray<FActiveSound*>& ActiveSounds = ConcurrencyGroup.GetActiveSounds();
+	for (FActiveSound* ActiveSound : ActiveSounds)
+	{
+		check(ActiveSound);
+
+		if (EvictableSound == nullptr
+			|| (ActiveSound->GetPriority() < EvictableSound->GetPriority())
+			|| (ActiveSound->GetPriority() == EvictableSound->GetPriority() && ActiveSound->PlaybackTime > EvictableSound->PlaybackTime))
+		{
+			EvictableSound = ActiveSound;
+		}
+	}
+
+	if (EvictableSound)
+	{
+		// Drop request as same priority and preventing new
+		const EMaxConcurrentResolutionRule::Type Rule = ConcurrencyGroup.GetSettings().ResolutionRule;
+		if (Rule == EMaxConcurrentResolutionRule::StopLowestPriorityThenPreventNew
+			&& EvictableSound->GetPriority() == NewActiveSound.GetPriority())
+		{
+			return nullptr;
+		}
+
+		// Drop request as NewActiveSound's priority is lower than the lowest priority sound playing
+		else if (EvictableSound->GetPriority() > NewActiveSound.GetPriority())
+		{
+			return nullptr;
+		}
+	}
+
 	return EvictableSound;
 }
 
