Disable spammy animation category which wasn't being used and was adding 200 stats/frame. These can be enabled on the commandline if needed with -csvstatcategories="animation"

diff --git a/Engine/Source/Runtime/Engine/Private/Animation/AnimSequence.cpp b/Engine/Source/Runtime/Engine/Private/Animation/AnimSequence.cpp
--- a/Engine/Source/Runtime/Engine/Private/Animation/AnimSequence.cpp
+++ b/Engine/Source/Runtime/Engine/Private/Animation/AnimSequence.cpp
@@ -43,7 +43,7 @@
 #define USE_SLERP 0
 #define LOCTEXT_NAMESPACE "AnimSequence"
 
-CSV_DEFINE_CATEGORY(Animation, (!UE_BUILD_SHIPPING));
+CSV_DEFINE_CATEGORY(Animation, false);
 
 DECLARE_CYCLE_STAT(TEXT("AnimSeq GetBonePose"), STAT_AnimSeq_GetBonePose, STATGROUP_Anim);
 DECLARE_CYCLE_STAT(TEXT("AnimSeq EvalCurveData"), STAT_AnimSeq_EvalCurveData, STATGROUP_Anim);
