[NO NON-PII COMMIT MESSAGE]

diff --git a/Engine/Source/Runtime/Engine/Private/AudioDecompress.cpp b/Engine/Source/Runtime/Engine/Private/AudioDecompress.cpp
--- a/Engine/Source/Runtime/Engine/Private/AudioDecompress.cpp
+++ b/Engine/Source/Runtime/Engine/Private/AudioDecompress.cpp
@@ -314,10 +314,19 @@ uint32 IStreamedCompressedInfo::IncrementCurrentSampleCount(uint32 NewSamples)
 
 uint32 IStreamedCompressedInfo::WriteFromDecodedPCM(uint8* Destination, uint32 BufferSize)
 {
-	uint32 BytesToCopy = FMath::Min(BufferSize, LastPCMByteSize - LastPCMOffset);
+	if (LastPCMOffset >= LastPCMByteSize)
+	{
+		LastPCMOffset = 0;
+		LastPCMByteSize = 0;
+
+		return 0;
+	}
+
+	uint32 BytesToCopy = FMath::Min<uint32>(BufferSize, LastPCMByteSize - LastPCMOffset);
+
 	if (BytesToCopy > 0)
 	{
-		check(BytesToCopy <= LastDecodedPCM.Num() - LastPCMOffset);
+		checkf(BytesToCopy <= LastDecodedPCM.Num() - LastPCMOffset, TEXT("Wraparound detected in WriteFromDecodedPCM. Current PCM Index %d was greater than PCM size %d, resulting in an invalid BytesToCopy value of %d."), LastPCMOffset, LastDecodedPCM.Num(), BytesToCopy);
 		FMemory::Memcpy(Destination, LastDecodedPCM.GetData() + LastPCMOffset, BytesToCopy);
 		LastPCMOffset += BytesToCopy;
 		if (LastPCMOffset >= LastPCMByteSize)
