- Make sure UObjects tracked by internal ActiveSound in Virtual Loops don't get GC'ed
- Always assign World to VirtualLoop ActiveSounds, not just in editor

diff --git a/Engine/Source/Runtime/Engine/Private/AudioDevice.cpp b/Engine/Source/Runtime/Engine/Private/AudioDevice.cpp
--- a/Engine/Source/Runtime/Engine/Private/AudioDevice.cpp
+++ b/Engine/Source/Runtime/Engine/Private/AudioDevice.cpp
@@ -629,6 +629,11 @@ void FAudioDevice::AddReferencedObjects(FReferenceCollector& Collector)
 		ActiveSound->AddReferencedObjects(Collector);
 	}
 
+	for (TPair<FActiveSound*, FAudioVirtualLoop*>& Pair : VirtualLoops)
+	{
+		Pair.Key->AddReferencedObjects(Collector);
+	}
+
 	// Make sure we don't try to delete any sound waves which may have in-flight decodes
 	Collector.AddReferencedObjects(ReferencedSoundWaves);
 
diff --git a/Engine/Source/Runtime/Engine/Private/AudioVirtualLoop.cpp b/Engine/Source/Runtime/Engine/Private/AudioVirtualLoop.cpp
--- a/Engine/Source/Runtime/Engine/Private/AudioVirtualLoop.cpp
+++ b/Engine/Source/Runtime/Engine/Private/AudioVirtualLoop.cpp
@@ -125,13 +125,9 @@ void FAudioVirtualLoop::SetActiveSound(const FActiveSound& InActiveSound)
 	ActiveSound->SetAudioComponent(InActiveSound);
 	ActiveSound->SetAudioDevice(AudioDevice);
 
-	if (GIsEditor)
-	{
-		ActiveSound->SetWorld(InActiveSound.GetWorld());
-	}
-
 	ActiveSound->SetSound(InActiveSound.GetSound());
 	ActiveSound->SetSoundClass(InActiveSound.GetSoundClass());
+	ActiveSound->SetWorld(InActiveSound.GetWorld());
 
 	ActiveSound->ConcurrencySet = InActiveSound.ConcurrencySet;
 	ActiveSound->VolumeMultiplier = InActiveSound.VolumeMultiplier;
