Added scoped stats to help track down hitches in PostActorConstruction

diff --git a/Engine/Source/Runtime/Engine/Private/Actor.cpp b/Engine/Source/Runtime/Engine/Private/Actor.cpp
--- a/Engine/Source/Runtime/Engine/Private/Actor.cpp
+++ b/Engine/Source/Runtime/Engine/Private/Actor.cpp
@@ -4436,6 +4436,8 @@ void AActor::MarkComponentsRenderStateDirty()
 
 void AActor::InitializeComponents()
 {
+	QUICK_SCOPE_CYCLE_COUNTER(STAT_Actor_InitializeComponents);
+
 	TInlineComponentArray<UActorComponent*> Components;
 	GetComponents(Components);
 
diff --git a/Engine/Source/Runtime/Engine/Private/Components/SkeletalMeshComponent.cpp b/Engine/Source/Runtime/Engine/Private/Components/SkeletalMeshComponent.cpp
--- a/Engine/Source/Runtime/Engine/Private/Components/SkeletalMeshComponent.cpp
+++ b/Engine/Source/Runtime/Engine/Private/Components/SkeletalMeshComponent.cpp
@@ -594,6 +594,8 @@ void USkeletalMeshComponent::OnUnregister()
 
 void USkeletalMeshComponent::InitAnim(bool bForceReinit)
 {
+	QUICK_SCOPE_CYCLE_COUNTER(STAT_SkelMeshComp_InitAnim);
+
 	// a lot of places just call InitAnim without checking Mesh, so 
 	// I'm moving the check here
 	if ( SkeletalMesh != nullptr && IsRegistered() )
