Linix: do not crash the servers on negative deltas, send analytics.

- Negative delta will be substituted by 1/30
- A new analytic event will be sent no more often than once per minute, and a log message will be logged instantly

diff --git a/Engine/Source/Runtime/Engine/Classes/Engine/Engine.h b/Engine/Source/Runtime/Engine/Classes/Engine/Engine.h
--- a/Engine/Source/Runtime/Engine/Classes/Engine/Engine.h
+++ b/Engine/Source/Runtime/Engine/Classes/Engine/Engine.h
@@ -2050,6 +2050,13 @@ public:
 	 */
 	void UpdateTimeAndHandleMaxTickRate();
 
+	/**
+	 * Allows games to correct the negative delta
+	 *
+	 * @return new delta
+	 */
+	virtual double CorrectNegativeTimeDelta(double DeltaRealTime);
+
 	/** Causes the current CustomTimeStep to be shut down and then reinitialized. */
 	void ReinitializeCustomTimeStep();
 
diff --git a/Engine/Source/Runtime/Engine/Private/UnrealEngine.cpp b/Engine/Source/Runtime/Engine/Private/UnrealEngine.cpp
--- a/Engine/Source/Runtime/Engine/Private/UnrealEngine.cpp
+++ b/Engine/Source/Runtime/Engine/Private/UnrealEngine.cpp
@@ -1829,6 +1829,18 @@ static FAutoConsoleCommand SetTimedMemReport(TEXT("TimedMemReport.Delay"), TEXT(
 
 #endif
 
+double UEngine::CorrectNegativeTimeDelta(double DeltaRealTime)
+{
+#if PLATFORM_ANDROID
+	UE_LOG(LogEngine, Warning, TEXT("Detected negative delta time - ignoring"));
+#else
+	// AMD dual-core systems are a known issue that require AMD CPU drivers to be installed. Installer will take care of this for shipping.
+	UE_LOG(LogEngine, Fatal, TEXT("Detected negative delta time - on AMD systems please install http://files.aoaforums.com/I3199-setup.zip.html - DeltaTime:%f"),
+		DeltaRealTime);
+#endif
+	return 0.01;
+}
+
 void UEngine::UpdateTimeAndHandleMaxTickRate()
 {
 	PumpABTest();
@@ -1887,14 +1899,7 @@ void UEngine::UpdateTimeAndHandleMaxTickRate()
 		// Negative delta time means something is wrong with the system. Error out so user can address issue.
 		if(DeltaRealTime < 0 )
 		{
-#if PLATFORM_ANDROID
-			UE_LOG(LogEngine, Warning, TEXT("Detected negative delta time - ignoring"));
-#else
-			// AMD dual-core systems are a known issue that require AMD CPU drivers to be installed. Installer will take care of this for shipping.
-			UE_LOG(LogEngine, Fatal,TEXT("Detected negative delta time - on AMD systems please install http://files.aoaforums.com/I3199-setup.zip.html - DeltaTime:%f, bUseFixedFrameRate:%d, bTimeWasManipulatedDebug:%d, FixedFrameRate:%f"), 
-				DeltaRealTime, bUseFixedFrameRate, bTimeWasManipulatedDebug, FixedFrameRate);
-#endif
-			DeltaRealTime = 0.01;
+			DeltaRealTime = CorrectNegativeTimeDelta(DeltaRealTime);
 		}
 
 		// Give engine chance to update frame rate smoothing
