- Add object name marker to UParticleSystemComponent::TickComponent

diff --git a/Engine/Source/Runtime/Engine/Private/Particles/ParticleComponents.cpp b/Engine/Source/Runtime/Engine/Private/Particles/ParticleComponents.cpp
--- a/Engine/Source/Runtime/Engine/Private/Particles/ParticleComponents.cpp
+++ b/Engine/Source/Runtime/Engine/Private/Particles/ParticleComponents.cpp
@@ -4937,6 +4937,7 @@ void UParticleSystemComponent::TickComponent(float DeltaTime, enum ELevelTick Ti
 	LLM_SCOPE(ELLMTag::Particles);
 	FInGameScopedCycleCounter InGameCycleCounter(GetWorld(), EInGamePerfTrackers::VFXSignificance, EInGamePerfTrackerThreads::GameThread, bIsManagingSignificance);
 	SCOPE_CYCLE_COUNTER(STAT_ParticlesOverview_GT);
+	FScopeCycleCounterUObject AdditionalScope(AdditionalStatObject(), GET_STATID(STAT_ParticlesOverview_GT));
 
 	if (Template == nullptr || Template->Emitters.Num() == 0)
 	{
