Fix build

diff --git a/Engine/Source/Runtime/Engine/Public/MaterialShared.h b/Engine/Source/Runtime/Engine/Public/MaterialShared.h
--- a/Engine/Source/Runtime/Engine/Public/MaterialShared.h
+++ b/Engine/Source/Runtime/Engine/Public/MaterialShared.h
@@ -1,4 +1,4 @@
-﻿// Copyright 1998-2019 Epic Games, Inc. All Rights Reserved.
+// Copyright 1998-2019 Epic Games, Inc. All Rights Reserved.
 
 /*=============================================================================
 	MaterialShared.h: Shared material definitions.
@@ -1783,7 +1783,7 @@ public:
 		return DeletedFlag != 0;
 	}
 
-	bool MarkForGarbageCollection()
+	void MarkForGarbageCollection()
 	{
 		MarkedForGarbageCollection = 1;
 	}
