Improved error handling for corrupt replays

diff --git a/Engine/Source/Runtime/Engine/Private/DemoNetDriver.cpp b/Engine/Source/Runtime/Engine/Private/DemoNetDriver.cpp
--- a/Engine/Source/Runtime/Engine/Private/DemoNetDriver.cpp
+++ b/Engine/Source/Runtime/Engine/Private/DemoNetDriver.cpp
@@ -89,6 +89,7 @@ static TAutoConsoleVariable<float> CVarDemoMinimumRepPrioritizeTime(TEXT("demo.M
 static TAutoConsoleVariable<float> CVarDemoMaximumRepPrioritizeTime(TEXT("demo.MaximumRepPrioritizePercent"), 0.8, TEXT("Maximum percent of time that may be spent prioritizing actors, regardless of throttling."));
 
 static const int32 MAX_DEMO_READ_WRITE_BUFFER = 1024 * 2;
+static const int32 MAX_DEMO_STRING_SERIALIZATION_SIZE = 16 * 1024 * 1024;
 
 namespace ReplayTaskNames
 {
@@ -1581,8 +1582,12 @@ void UDemoNetDriver::StopDemo()
 	if (!IsRecording() && !IsPlaying())
 	{
 		UE_LOG(LogDemo, Log, TEXT("StopDemo: No demo is playing"));
+		ClearReplayTasks();
+		ActiveReplayName = FString();
+		ResetDemoState();
 		return;
 	}
+
 	OnDemoFinishRecordingDelegate.Broadcast();
 	UE_LOG(LogDemo, Log, TEXT("StopDemo: Demo %s stopped at frame %d"), *DemoURL.Map, DemoFrameNum);
 
@@ -4585,6 +4590,7 @@ bool UDemoNetDriver::FastForwardLevels(const FGotoResult& GotoResult)
 		{
 			CheckpointArchive->SetEngineNetVer(PlaybackDemoHeader.EngineNetworkProtocolVersion);
 			CheckpointArchive->SetGameNetVer(PlaybackDemoHeader.GameNetworkProtocolVersion);
+			CheckpointArchive->ArMaxSerializeSize = MAX_DEMO_STRING_SERIALIZATION_SIZE;
 
 			TGuardValue<bool> LoadingCheckpointGuard(bIsLoadingCheckpoint, true);
 
@@ -4883,6 +4889,7 @@ bool UDemoNetDriver::LoadCheckpoint(const FGotoResult& GotoResult)
 
 	GotoCheckpointArchive->SetEngineNetVer(PlaybackDemoHeader.EngineNetworkProtocolVersion);
 	GotoCheckpointArchive->SetGameNetVer(PlaybackDemoHeader.GameNetworkProtocolVersion);
+	GotoCheckpointArchive->ArMaxSerializeSize = MAX_DEMO_STRING_SERIALIZATION_SIZE;
 
 	int32 LevelForCheckpoint = 0;
 
@@ -5298,6 +5305,11 @@ bool UDemoNetDriver::LoadCheckpoint(const FGotoResult& GotoResult)
 			CacheObject.bIgnoreWhenMissing = (Flags & (1 << 1)) ? true : false;		
 
 			GuidCache->ObjectLookup.Add(Guid, CacheObject);
+
+			if (GotoCheckpointArchive->IsError())
+			{
+				break;
+			}
 		}
 
 		int32 DeltaPacketStartIndex = INDEX_NONE;
@@ -5327,33 +5339,36 @@ bool UDemoNetDriver::LoadCheckpoint(const FGotoResult& GotoResult)
 	}
 	while (!GotoCheckpointArchive->IsError() && (GotoCheckpointArchive->Tell() < GotoCheckpointArchive->TotalSize()));
 	
-	// Destroy startup actors that shouldn't exist past this checkpoint
-	for (FActorIterator It( World ); It; ++It)
+	if (World != nullptr)
 	{
-		const FString FullName = It->GetFullName();
-		if (DeletedNetStartupActors.Contains(FullName))
+		// Destroy startup actors that shouldn't exist past this checkpoint
+		for (FActorIterator It( World ); It; ++It)
 		{
-			if (It->bReplayRewindable)
+			const FString FullName = It->GetFullName();
+			if (DeletedNetStartupActors.Contains(FullName))
 			{
-				// Log and skip. We can't queue Rewindable actors and we can't destroy them.
-				// This actor may still get destroyed during cleanup.
-				UE_LOG(LogDemo, Warning, TEXT("Replay Rewindable Actor found in the DeletedNetStartupActors. Replay may show artifacts (%s)"), *FullName);
-				continue;
-			}
+				if (It->bReplayRewindable)
+				{
+					// Log and skip. We can't queue Rewindable actors and we can't destroy them.
+					// This actor may still get destroyed during cleanup.
+					UE_LOG(LogDemo, Warning, TEXT("Replay Rewindable Actor found in the DeletedNetStartupActors. Replay may show artifacts (%s)"), *FullName);
+					continue;
+				}
 
-			// Put this actor on the rollback list so we can undelete it during future scrubbing
-			QueueNetStartupActorForRollbackViaDeletion(*It);
+				// Put this actor on the rollback list so we can undelete it during future scrubbing
+				QueueNetStartupActorForRollbackViaDeletion(*It);
 
-			UE_LOG(LogDemo, Verbose, TEXT("LoadCheckpoint: deleting startup actor %s"), *FullName);
+				UE_LOG(LogDemo, Verbose, TEXT("LoadCheckpoint: deleting startup actor %s"), *FullName);
 
-			// Delete the actor
-			World->DestroyActor(*It, true);
+				// Delete the actor
+				World->DestroyActor(*It, true);
+			}
 		}
-	}
 
-	// Re-create all startup actors that were destroyed but should exist beyond this point
-	TArray<AActor*> SpawnedActors;
-	RespawnNecessaryNetStartupActors(SpawnedActors);
+		// Re-create all startup actors that were destroyed but should exist beyond this point
+		TArray<AActor*> SpawnedActors;
+		RespawnNecessaryNetStartupActors(SpawnedActors);
+	}
 		
 	if (PlaybackPackets.Num() > 0)
 	{
@@ -5379,11 +5394,11 @@ bool UDemoNetDriver::LoadCheckpoint(const FGotoResult& GotoResult)
 	{
 		if (bDeltaCheckpoint)
 		{
-			UDemoNetConnection* DemoConnection = CastChecked<UDemoNetConnection>(ServerConnection);
+			UDemoNetConnection* DemoConnection = Cast<UDemoNetConnection>(ServerConnection);
 
 			for (int32 i = 0; i < DeltaCheckpointPacketIntervals.Num(); ++i)
 			{
-				if (PlaybackDeltaCheckpointData.IsValidIndex(i))
+				if (DemoConnection && PlaybackDeltaCheckpointData.IsValidIndex(i))
 				{
 					check(PlaybackDeltaCheckpointData[i].IsValid());
 
@@ -5402,7 +5417,11 @@ bool UDemoNetDriver::LoadCheckpoint(const FGotoResult& GotoResult)
 			}
 
 			PlaybackPackets.Empty();
-			DemoConnection->GetOpenChannelMap().Empty();
+
+			if (DemoConnection)
+			{
+				DemoConnection->GetOpenChannelMap().Empty();
+			}
 		}
 		else
 		{
@@ -6028,30 +6047,31 @@ void UDemoNetDriver::NotifyActorDestroyed(AActor* Actor, bool IsSeamlessTravel)
 
 void UDemoNetDriver::CleanupOutstandingRewindActors()
 {
-	UWorld* LocalWorld = GetWorld();
-
-	for (const FNetworkGUID& NetGUID : TrackedRewindActorsByGUID)
+	if (UWorld* LocalWorld = GetWorld())
 	{
-		if (FNetGuidCacheObject* CacheObject = GuidCache->ObjectLookup.Find(NetGUID))
+		for (const FNetworkGUID& NetGUID : TrackedRewindActorsByGUID)
 		{
-			if (AActor* Actor = Cast<AActor>(CacheObject->Object))
+			if (FNetGuidCacheObject* CacheObject = GuidCache->ObjectLookup.Find(NetGUID))
 			{
-				// Destroy the actor before removing entries from the GuidCache so its entries are still valid in NotifyActorDestroyed.
-				LocalWorld->DestroyActor(Actor);
+				if (AActor* Actor = Cast<AActor>(CacheObject->Object))
+				{
+					// Destroy the actor before removing entries from the GuidCache so its entries are still valid in NotifyActorDestroyed.
+					LocalWorld->DestroyActor(Actor);
 
-				ensureMsgf(GuidCache->NetGUIDLookup.Remove(CacheObject->Object) > 0, TEXT("CleanupOutstandingRewindActors: No entry found for %d in NetGUIDLookup"), NetGUID.Value);
-				GuidCache->ObjectLookup.Remove(NetGUID);
-				CacheObject->bNoLoad = false;
-			}
+					ensureMsgf(GuidCache->NetGUIDLookup.Remove(CacheObject->Object) > 0, TEXT("CleanupOutstandingRewindActors: No entry found for %d in NetGUIDLookup"), NetGUID.Value);
+					GuidCache->ObjectLookup.Remove(NetGUID);
+					CacheObject->bNoLoad = false;
+				}
+				else
+				{
+					UE_LOG(LogDemo, Warning, TEXT("CleanupOutstandingRewindActors - Invalid object for %d, skipping."), NetGUID.Value);
+					continue;
+				}
+			}	
 			else
 			{
-				UE_LOG(LogDemo, Warning, TEXT("CleanupOutstandingRewindActors - Invalid object for %d, skipping."), NetGUID.Value);
-				continue;
+				UE_LOG(LogDemo, Warning, TEXT("CleanupOutstandingRewindActors - CacheObject not found for %s"), NetGUID.Value);
 			}
-		}	
-		else
-		{
-			UE_LOG(LogDemo, Warning, TEXT("CleanupOutstandingRewindActors - CacheObject not found for %s"), NetGUID.Value);
 		}
 	}
 
