Adding feature to allow audio mixer toggling to only happen once per load.

diff --git a/Engine/Source/Runtime/Engine/Private/AudioDeviceManager.cpp b/Engine/Source/Runtime/Engine/Private/AudioDeviceManager.cpp
--- a/Engine/Source/Runtime/Engine/Private/AudioDeviceManager.cpp
+++ b/Engine/Source/Runtime/Engine/Private/AudioDeviceManager.cpp
@@ -66,6 +66,8 @@ FAudioDeviceManager::FAudioDeviceManager()
 	, bUsingAudioMixer(false)
 	, bPlayAllDeviceAudio(false)
 	, bVisualize3dDebug(false)
+	, bOnlyToggleAudioMixerOnce(false)
+	, bToggledAudioMixer(false)
 {
 	// Check for a command line debug sound argument.
 	FString DebugSound;
@@ -285,6 +287,9 @@ bool FAudioDeviceManager::LoadDefaultAudioDeviceModule()
 		bUsingAudioMixer = false;
 	}
 
+	// Check for config bool that restricts audio mixer toggle to only once. This will allow us to patch audio mixer on or off after initial login.
+	GConfig->GetBool(TEXT("Audio"), TEXT("OnlyToggleAudioMixerOnce"), bOnlyToggleAudioMixerOnce, GEngineIni);
+
 	// Get the audio mixer and non-audio mixer device module names
 	GConfig->GetString(TEXT("Audio"), TEXT("AudioDeviceModuleName"), AudioDeviceModuleName, GEngineIni);
 	GConfig->GetString(TEXT("Audio"), TEXT("AudioMixerModuleName"), AudioMixerModuleName, GEngineIni);
@@ -571,17 +576,23 @@ void FAudioDeviceManager::UpdateActiveAudioDevices(bool bGameTicking)
 		SyncFence.Wait();
 	}
 
-	if (bUsingAudioMixer && !GCvarIsUsingAudioMixer)
-	{
-		ToggleAudioMixer();
-		bUsingAudioMixer = false;
-	}
-	else if (!bUsingAudioMixer && GCvarIsUsingAudioMixer)
+	if (!bOnlyToggleAudioMixerOnce || (bOnlyToggleAudioMixerOnce && !bToggledAudioMixer))
 	{
-		ToggleAudioMixer();
-		bUsingAudioMixer = true;
+		if (bUsingAudioMixer && !GCvarIsUsingAudioMixer)
+		{
+			ToggleAudioMixer();
+			bToggledAudioMixer = true;
+			bUsingAudioMixer = false;
+		}
+		else if (!bUsingAudioMixer && GCvarIsUsingAudioMixer)
+		{
+			ToggleAudioMixer();
+			bToggledAudioMixer = true;
+			bUsingAudioMixer = true;
+		}
 	}
 
+
 	for (FAudioDevice* AudioDevice : Devices)
 	{
 		if (AudioDevice)
diff --git a/Engine/Source/Runtime/Engine/Public/AudioDeviceManager.h b/Engine/Source/Runtime/Engine/Public/AudioDeviceManager.h
--- a/Engine/Source/Runtime/Engine/Public/AudioDeviceManager.h
+++ b/Engine/Source/Runtime/Engine/Public/AudioDeviceManager.h
@@ -321,6 +321,12 @@ private:
 	/** Whether or not 3d debug visualization is enabled. */
 	bool bVisualize3dDebug;
 
+	/** Whether or not we check to toggle audio mixer once. */
+	bool bOnlyToggleAudioMixerOnce;
+
+	/** Whether or not we've toggled the audio mixer. */
+	bool bToggledAudioMixer;
+
 	/** Audio Fence to ensure that we don't allow the audio thread to drift never endingly behind. */
 	FAudioCommandFence SyncFence;
 };
