- Add debug tracking for Virtual Loops
    - Logging around when virtual loops are added and removed when logging verbose audio
    - stat Audio for current number of virtual loops to compare against 'live' ActiveSounds
    - blue trackers with Audio3dVisualize (Added Audio3dVisualize to BaseInput.ini)
- Test and improve enablement/disablement of VirtualLoop feature by actively stopping VirtualLoops on AudioDevice update
- Copy over 'live updated' values when virtualizing loops to new ActiveSound for reference/safety using Copy constructor and omitting data (much cleaner than the mammoth property copy)
- VirtualLoops now no longer created on heap

diff --git a/Engine/Config/BaseInput.ini b/Engine/Config/BaseInput.ini
--- a/Engine/Config/BaseInput.ini
+++ b/Engine/Config/BaseInput.ini
@@ -202,6 +202,7 @@ MaxScrollbackSize=1024
 +ManualAutoCompleteList=(Command="Stat SOUNDWAVES",Desc="Shows active SoundWaves")
 +ManualAutoCompleteList=(Command="Stat SOUNDCUES",Desc="Shows active SoundCues")
 +ManualAutoCompleteList=(Command="Stat SOUNDS",Desc="<?> <sort=distance|class|name|waves|default> <-debug> <off> Shows active SoundCues and SoundWaves")
++ManualAutoCompleteList=(Command="Audio3dVisualize",Desc="Shows locations of sound sources playing (white text) and their left and right channel locations respectively (red and green).  Virtualized loops (if enabled) display in blue.")
 +ManualAutoCompleteList=(Command="StartMovieCapture",Desc=)
 +ManualAutoCompleteList=(Command="StopMovieCapture",Desc=)
 +ManualAutoCompleteList=(Command="TraceTag",Desc="Draw traces that use the specified tag")
diff --git a/Engine/Source/Runtime/Engine/Private/Audio.cpp b/Engine/Source/Runtime/Engine/Private/Audio.cpp
--- a/Engine/Source/Runtime/Engine/Private/Audio.cpp
+++ b/Engine/Source/Runtime/Engine/Private/Audio.cpp
@@ -32,6 +32,7 @@ DEFINE_LOG_CATEGORY(LogAudioDebug);
 DEFINE_STAT(STAT_AudioMemorySize);
 DEFINE_STAT(STAT_ActiveSounds);
 DEFINE_STAT(STAT_AudioSources);
+DEFINE_STAT(STAT_AudioVirtualLoops);
 DEFINE_STAT(STAT_WaveInstances);
 DEFINE_STAT(STAT_WavesDroppedDueToPriority);
 DEFINE_STAT(STAT_AudibleWavesDroppedDueToPriority);
diff --git a/Engine/Source/Runtime/Engine/Private/AudioDevice.cpp b/Engine/Source/Runtime/Engine/Private/AudioDevice.cpp
--- a/Engine/Source/Runtime/Engine/Private/AudioDevice.cpp
+++ b/Engine/Source/Runtime/Engine/Private/AudioDevice.cpp
@@ -2,7 +2,6 @@
 
 #include "AudioDevice.h"
 
-#include "AudioVirtualLoop.h"
 #include "PhysicsEngine/BodyInstance.h"
 #include "Sound/SoundEffectPreset.h"
 #include "Sound/SoundEffectSubmix.h"
@@ -136,6 +135,8 @@ FAutoConsoleVariableRef CVarSoundDistanceOptimizationLength(
 	TEXT("The maximum duration a sound must be in order to be a candidate to be culled due to one-shot distance optimization.\n"),
 	ECVF_Default);
 
+using FVirtualLoopPair = TPair<FActiveSound*, FAudioVirtualLoop>;
+
 /*-----------------------------------------------------------------------------
 	FAudioDevice implementation.
 -----------------------------------------------------------------------------*/
@@ -629,7 +630,7 @@ void FAudioDevice::AddReferencedObjects(FReferenceCollector& Collector)
 		ActiveSound->AddReferencedObjects(Collector);
 	}
 
-	for (TPair<FActiveSound*, FAudioVirtualLoop*>& Pair : VirtualLoops)
+	for (FVirtualLoopPair& Pair : VirtualLoops)
 	{
 		Pair.Key->AddReferencedObjects(Collector);
 	}
@@ -2064,13 +2065,15 @@ void FAudioDevice::StopQuietSoundsDueToMaxConcurrency(TArray<FWaveInstance*>& Wa
 			{
 				AddSoundToStop(ActiveSound);
 				const bool bDoRangeCheck = false;
-				if (FAudioVirtualLoop* VirtualLoop = FAudioVirtualLoop::Virtualize(*this, *ActiveSound, bDoRangeCheck))
+				FAudioVirtualLoop VirtualLoop;
+				if (FAudioVirtualLoop::Virtualize(*ActiveSound, bDoRangeCheck, VirtualLoop))
 				{
-					VirtualLoop->CalculateUpdateInterval(true);
-
 					// Clear must be called after AddSoundToStop to ensure AudioComponent is properly removed from AudioComponentIDToActiveSoundMap
 					ActiveSound->ClearAudioComponent();
-					AddVirtualLoop(*VirtualLoop);
+
+					const bool bIsAtMaxConcurrency = true;
+					VirtualLoop.CalculateUpdateInterval(bIsAtMaxConcurrency);
+					AddVirtualLoop(VirtualLoop);
 				}
 				ActiveSoundsCopy.RemoveAtSwap(i, 1, false);
 			}
@@ -2596,19 +2599,26 @@ void FAudioDevice::UpdateSoundClassProperties(float DeltaTime)
 
 void FAudioDevice::VirtualizeInactiveLoops()
 {
+	// Check if virtual loop system is enabled and don't push to virtual if disabled.
+	if (!FAudioVirtualLoop::IsEnabled())
+	{
+		return;
+	}
+
 	const bool bDoRangeCheck = true;
 	for (FActiveSound* ActiveSound : ActiveSounds)
 	{
 		// If already pending stop, don't attempt to virtualize
 		if (!IsPendingStop(ActiveSound))
 		{
-			if (FAudioVirtualLoop* VirtualLoop = FAudioVirtualLoop::Virtualize(*this, *ActiveSound, bDoRangeCheck))
+			FAudioVirtualLoop VirtualLoop;
+			if (FAudioVirtualLoop::Virtualize(*ActiveSound, bDoRangeCheck, VirtualLoop))
 			{
 				AddSoundToStop(ActiveSound);
 
 				// Clear must be called after AddSoundToStop to ensure AudioComponent is properly removed from AudioComponentIDToActiveSoundMap
 				ActiveSound->ClearAudioComponent();
-				AddVirtualLoop(*VirtualLoop);
+				AddVirtualLoop(VirtualLoop);
 			}
 		}
 	}
@@ -3902,6 +3912,7 @@ void FAudioDevice::Update(bool bGameTicking)
 		INC_DWORD_STAT_BY(STAT_AudioSources, Channels - FreeSources.Num());
 		INC_DWORD_STAT_BY(STAT_WavesDroppedDueToPriority, FMath::Max(ActiveWaveInstances.Num() - Channels, 0));
 		INC_DWORD_STAT_BY(STAT_ActiveSounds, ActiveSounds.Num());
+		INC_DWORD_STAT_BY(STAT_AudioVirtualLoops, VirtualLoops.Num());
 	}
 
 	// now let the platform perform anything it needs to handle
@@ -4083,7 +4094,7 @@ void FAudioDevice::StopAllSounds(bool bShouldStopUISounds)
 		}
 	}
 
-	for (TPair<FActiveSound*, FAudioVirtualLoop*>& Pair : VirtualLoops)
+	for (FVirtualLoopPair& Pair : VirtualLoops)
 	{
 		AddSoundToStop(Pair.Key);
 	}
@@ -4189,10 +4200,10 @@ void FAudioDevice::AddNewActiveSoundInternal(const FActiveSound& NewActiveSound,
 	if (!VirtualLoopToRetrigger)
 	{
 		const bool bDoRangeCheck = true;
-		VirtualLoopToRetrigger = FAudioVirtualLoop::Virtualize(*this, NewActiveSound, bDoRangeCheck);
-		if (VirtualLoopToRetrigger)
+		FAudioVirtualLoop VirtualLoop;
+		if (FAudioVirtualLoop::Virtualize(NewActiveSound, *this, bDoRangeCheck, VirtualLoop))
 		{
-			AddVirtualLoop(*VirtualLoopToRetrigger);
+			AddVirtualLoop(VirtualLoop);
 			return;
 		}
 	}
@@ -4213,15 +4224,18 @@ void FAudioDevice::AddNewActiveSoundInternal(const FActiveSound& NewActiveSound,
 		if (!VirtualLoopToRetrigger)
 		{
 			const bool bDoRangeCheck = false;
-			VirtualLoopToRetrigger = FAudioVirtualLoop::Virtualize(*this, NewActiveSound, bDoRangeCheck);
-			if (VirtualLoopToRetrigger)
+			FAudioVirtualLoop VirtualLoop;
+			if (FAudioVirtualLoop::Virtualize(NewActiveSound, *this, bDoRangeCheck, VirtualLoop))
 			{
-				AddVirtualLoop(*VirtualLoopToRetrigger);
+				AddVirtualLoop(VirtualLoop);
 			}
 		}
 		return;
 	}
 
+	check(ActiveSound->Sound);
+	check(ActiveSound->Sound == NewActiveSound.Sound);
+
 	if (GIsEditor)
 	{
 		// If the sound played on an editor preview world, treat it as a preview sound (unpausable and ignoring the realtime volume slider)
@@ -4231,7 +4245,7 @@ void FAudioDevice::AddNewActiveSoundInternal(const FActiveSound& NewActiveSound,
 		}
 	}
 
-	++NewActiveSound.Sound->CurrentPlayCount;
+	++ActiveSound->Sound->CurrentPlayCount;
 
 #if !(UE_BUILD_SHIPPING || UE_BUILD_TEST)
 	UE_LOG(LogAudio, VeryVerbose, TEXT("New ActiveSound %s Comp: %s Loc: %s"), *NewActiveSound.Sound->GetName(), *NewActiveSound.GetAudioComponentName(), *NewActiveSound.Transform.GetTranslation().ToString());
@@ -4247,23 +4261,20 @@ void FAudioDevice::AddNewActiveSoundInternal(const FActiveSound& NewActiveSound,
 	ActiveSound->bIsPlayingAudio = true;
 
 #if !(UE_BUILD_SHIPPING || UE_BUILD_TEST)
-	if (ActiveSound->Sound)
+	if (!ensureMsgf(ActiveSound->Sound->GetFName() != NAME_None, TEXT("AddNewActiveSound with DESTROYED sound %s. AudioComponent=%s. IsPendingKill=%d. BeginDestroy=%d"),
+		*ActiveSound->Sound->GetPathName(),
+		*ActiveSound->GetAudioComponentName(),
+		(int32)ActiveSound->Sound->IsPendingKill(),
+		(int32)ActiveSound->Sound->HasAnyFlags(RF_BeginDestroyed)))
 	{
-		if (!ensureMsgf(ActiveSound->Sound->GetFName() != NAME_None, TEXT("AddNewActiveSound with DESTROYED sound %s. AudioComponent=%s. IsPendingKill=%d. BeginDestroy=%d"),
-			*ActiveSound->Sound->GetPathName(),
-			*ActiveSound->GetAudioComponentName(),
-			(int32)ActiveSound->Sound->IsPendingKill(),
-			(int32)ActiveSound->Sound->HasAnyFlags(RF_BeginDestroyed)))
-		{
-			static FName InvalidSoundName(TEXT("DESTROYED_Sound"));
-			ActiveSound->DebugOriginalSoundName = InvalidSoundName;
-		}
-		else
-		{
-			ActiveSound->DebugOriginalSoundName = ActiveSound->Sound->GetFName();
-		}
+		static FName InvalidSoundName(TEXT("DESTROYED_Sound"));
+		ActiveSound->DebugOriginalSoundName = InvalidSoundName;
 	}
-#endif
+	else
+	{
+		ActiveSound->DebugOriginalSoundName = ActiveSound->Sound->GetFName();
+	}
+#endif // !(UE_BUILD_SHIPPING || UE_BUILD_TEST)
 
 	// Loop has been successfully created, so add to stop before adding 'live' ActiveSound.
 	// Mark to not report playback complete on destruction as responsibility therein has been
@@ -4290,9 +4301,11 @@ void FAudioDevice::RetriggerVirtualLoop(FAudioVirtualLoop& VirtualLoopToRetrigge
 	AddNewActiveSoundInternal(VirtualLoopToRetrigger.GetActiveSound(), &VirtualLoopToRetrigger);
 }
 
-void FAudioDevice::AddVirtualLoop(FAudioVirtualLoop& VirtualLoop)
+void FAudioDevice::AddVirtualLoop(const FAudioVirtualLoop& InVirtualLoop)
 {
-	FActiveSound& ActiveSound = VirtualLoop.GetActiveSound();
+	FAudioVirtualLoop* VirtualLoop = new FAudioVirtualLoop(InVirtualLoop);
+
+	FActiveSound& ActiveSound = VirtualLoop->GetActiveSound();
 	check(!VirtualLoops.Contains(&ActiveSound));
 
 	const int64 ComponentID = ActiveSound.GetAudioComponentID();
@@ -4309,25 +4322,38 @@ void FAudioDevice::AddVirtualLoop(FAudioVirtualLoop& VirtualLoop)
 		AudioComponentIDToActiveSoundMap.Add(ComponentID, &ActiveSound);
 	}
 
-	VirtualLoops.Add(&ActiveSound, &VirtualLoop);
+#if !(UE_BUILD_SHIPPING || UE_BUILD_TEST)
+	if (ActiveSound.Sound)
+	{
+		const FVector Location = ActiveSound.Transform.GetLocation();
+		UE_LOG(LogAudio, Verbose, TEXT("Adding virtual looping sound '%s' at location %s."), *ActiveSound.Sound->GetName(), *Location.ToCompactString());
+	}
+#endif // !(UE_BUILD_SHIPPING || UE_BUILD_TEST)
+
+	VirtualLoops.Add(&ActiveSound, InVirtualLoop);
 }
 
-bool FAudioDevice::RemoveVirtualLoop(FActiveSound& ActiveSound)
+bool FAudioDevice::RemoveVirtualLoop(FActiveSound& InActiveSound)
 {
 	check(IsInAudioThread());
 
-	if (FAudioVirtualLoop* VirtualLoop = VirtualLoops.FindRef(&ActiveSound))
+	if (FAudioVirtualLoop* VirtualLoop = VirtualLoops.Find(&InActiveSound))
 	{
-		const uint64 ComponentID = ActiveSound.GetAudioComponentID();
+		const uint64 ComponentID = InActiveSound.GetAudioComponentID();
 		if (ComponentID > 0)
 		{
 			UAudioComponent::PlaybackCompleted(ComponentID, false);
 		}
 
-		check(VirtualLoop);
-		delete VirtualLoop;
-		VirtualLoops.Remove(&ActiveSound);
+#if !(UE_BUILD_SHIPPING || UE_BUILD_TEST)
+		if (InActiveSound.Sound)
+		{
+			const FVector Location = InActiveSound.Transform.GetLocation();
+			UE_LOG(LogAudio, Verbose, TEXT("Removing virtual looping sound '%s' at location %s."), *InActiveSound.Sound->GetName(), *Location.ToCompactString());
+		}
+#endif // !(UE_BUILD_SHIPPING || UE_BUILD_TEST)
 
+		VirtualLoops.Remove(&InActiveSound);
 		return true;
 	}
 
@@ -4395,32 +4421,34 @@ void FAudioDevice::ProcessingPendingActiveSoundStops(bool bForceDelete)
 		// and continue
 		if (RemoveVirtualLoop(*ActiveSound))
 		{
-			continue;
+			delete ActiveSound;
 		}
+		else
+		{
+			ActiveSound->Stop(bForceDelete);
 
-		ActiveSound->Stop(bForceDelete);
-
-		USoundBase* Sound = ActiveSound->GetSound();
+			USoundBase* Sound = ActiveSound->GetSound();
 
-		// If the active sound is a one shot, decrement the one shot counter
-		if (Sound && !Sound->IsLooping())
-		{
-			OneShotCount--;
-		}
+			// If the active sound is a one shot, decrement the one shot counter
+			if (Sound && !Sound->IsLooping())
+			{
+				OneShotCount--;
+			}
 
-		const bool bIsStopping = ActiveSound->IsStopping();
+			const bool bIsStopping = ActiveSound->IsStopping();
 
-		// If we can delete the active sound now, then delete it
-		if (bForceDelete || (ActiveSound->CanDelete() && !bIsStopping))
-		{
-			ActiveSound->bAsyncOcclusionPending = false;
+			// If we can delete the active sound now, then delete it
+			if (bForceDelete || (ActiveSound->CanDelete() && !bIsStopping))
+			{
+				ActiveSound->bAsyncOcclusionPending = false;
 
-			delete ActiveSound;
-		}
-		else
-		{
-			// There was an async operation pending or we are stopping (not stopped) so we need to defer deleting this sound
-			PendingSoundsToDelete.AddUnique(ActiveSound);
+				delete ActiveSound;
+			}
+			else
+			{
+				// There was an async operation pending or we are stopping (not stopped) so we need to defer deleting this sound
+				PendingSoundsToDelete.AddUnique(ActiveSound);
+			}
 		}
 	}
 	PendingSoundsToStop.Reset();
@@ -5232,7 +5260,7 @@ void FAudioDevice::Flush(UWorld* WorldToFlush, bool bClearActivatedReverb)
 		}
 	}
 
-	for (TPair<FActiveSound*, FAudioVirtualLoop*>& Pair : VirtualLoops)
+	for (FVirtualLoopPair& Pair : VirtualLoops)
 	{
 		AddSoundToStop(Pair.Key);
 	}
@@ -5770,23 +5798,57 @@ void FAudioDevice::UpdateProximityRetriggerComponents()
 
 void FAudioDevice::UpdateVirtualLoops()
 {
-	static TArray<FAudioVirtualLoop*> VirtualLoopsToRetrigger;
-	for (TPair<FActiveSound*, FAudioVirtualLoop*>& Pair : VirtualLoops)
+	if (FAudioVirtualLoop::IsEnabled())
 	{
-		FAudioVirtualLoop* VirtualLoop = Pair.Value;
-		check(VirtualLoop);
+		TArray<FAudioVirtualLoop> VirtualLoopsToRetrigger;
 
-		if (VirtualLoop->Update())
+		for (FVirtualLoopPair& Pair : VirtualLoops)
 		{
-			VirtualLoopsToRetrigger.Add(VirtualLoop);
+			FAudioVirtualLoop& VirtualLoop = Pair.Value;
+
+			const float DeltaTime = GetDeviceDeltaTime();
+
+			FActiveSound& ActiveSound = Pair.Value.GetActiveSound();
+			// Keep playback time up-to-date as it may be used to evaluate concurrency
+			ActiveSound.PlaybackTime += DeltaTime * ActiveSound.MinCurrentPitch;
+
+			// If the loop is ready to realize, add to array to be re-triggered outside of the loop
+			// to avoid map manipulation while iterating.
+			if (VirtualLoop.CanRealize(DeltaTime))
+			{
+				VirtualLoopsToRetrigger.Add(Pair.Value);
+			}
+		}
+
+		for (FAudioVirtualLoop& RetriggerLoop : VirtualLoopsToRetrigger)
+		{
+			RetriggerVirtualLoop(RetriggerLoop);
 		}
+
+		VirtualLoopsToRetrigger.Reset();
 	}
 
-	for (FAudioVirtualLoop* RetriggerLoop : VirtualLoopsToRetrigger)
+	// if !FAudioVirtualLoop::IsEnabled(), attempt to realize/re-trigger
+	// sounds and remove virtual loops.
+	else
 	{
-		RetriggerVirtualLoop(*RetriggerLoop);
+		// Copies any straggling virtual loops to active sounds and mark them for stop
+		for (FVirtualLoopPair& Pair : VirtualLoops)
+		{
+			FActiveSound* ActiveSound = Pair.Key;
+
+			const uint64 ComponentID = ActiveSound->GetAudioComponentID();
+			if (ComponentID > 0)
+			{
+				AudioComponentIDToActiveSoundMap.Remove(ComponentID);
+			}
+
+			AddNewActiveSound(*ActiveSound);
+
+			ActiveSound->ClearAudioComponent();
+			AddSoundToStop(ActiveSound);
+		}
 	}
-	VirtualLoopsToRetrigger.Reset();
 }
 
 #if WITH_EDITOR
diff --git a/Engine/Source/Runtime/Engine/Private/AudioVirtualLoop.cpp b/Engine/Source/Runtime/Engine/Private/AudioVirtualLoop.cpp
--- a/Engine/Source/Runtime/Engine/Private/AudioVirtualLoop.cpp
+++ b/Engine/Source/Runtime/Engine/Private/AudioVirtualLoop.cpp
@@ -3,64 +3,82 @@
 
 #include "ActiveSound.h"
 #include "AudioDevice.h"
+#include "DrawDebugHelpers.h"
 #include "Sound/SoundBase.h"
 
-	static int32 bVirtualLoopsEnabledCVar = 1;
-	FAutoConsoleVariableRef CVarVirtualLoopsEnabled(
-		TEXT("au.VirtualLoopsEnabled"),
-		bVirtualLoopsEnabledCVar,
-		TEXT("Enables or disables whether virtualizing is supported for audio loops.\n"),
-		ECVF_Default);
-
-	static float VirtualLoopsPerfDistanceCVar = 15000.0f;
-	FAutoConsoleVariableRef CVarVirtualLoopsPerfDistance(
-		TEXT("au.VirtualLoopsPerfDistance"),
-		VirtualLoopsPerfDistanceCVar,
-		TEXT("Sets virtual loop distance to scale update rate between min and max beyond max audible distance of sound.\n"),
-		ECVF_Default);
-
-	static float VirtualUpdateRateMinCVar = 0.1f;
-	FAutoConsoleVariableRef CVarVirtualLoopsUpdateRateMin(
-		TEXT("au.VirtualLoopsUpdateRateMin"),
-		VirtualUpdateRateMinCVar,
-		TEXT("Sets minimum rate to check if sound becomes audible again at sound's max audible distance.\n"),
-		ECVF_Default);
-
-	static float VirtualUpdateRateMaxCVar = 3.0f;
-	FAutoConsoleVariableRef CVarVirtualLoopsUpdateRateMax(
-		TEXT("au.VirtualLoopsUpdateRateMax"),
-		VirtualUpdateRateMinCVar,
-		TEXT("Sets maximum rate to check if sound becomes audible again (at beyond sound's max audible distance + perf scaling distance).\n"),
-		ECVF_Default);
-
-FAudioVirtualLoop::FAudioVirtualLoop(FAudioDevice& InAudioDevice, const FActiveSound& NewActiveSound)
+
+static int32 bVirtualLoopsEnabledCVar = 1;
+FAutoConsoleVariableRef CVarVirtualLoopsEnabled(
+	TEXT("au.VirtualLoops.Enabled"),
+	bVirtualLoopsEnabledCVar,
+	TEXT("Enables or disables whether virtualizing is supported for audio loops.\n"),
+	ECVF_Default);
+
+static float VirtualLoopsPerfDistanceCVar = 15000.0f;
+FAutoConsoleVariableRef CVarVirtualLoopsPerfDistance(
+	TEXT("au.VirtualLoops.PerfDistance"),
+	VirtualLoopsPerfDistanceCVar,
+	TEXT("Sets virtual loop distance to scale update rate between min and max beyond max audible distance of sound.\n"),
+	ECVF_Default);
+
+static float VirtualLoopsUpdateRateMinCVar = 0.1f;
+FAutoConsoleVariableRef CVarVirtualLoopsUpdateRateMin(
+	TEXT("au.VirtualLoops.UpdateRate.Min"),
+	VirtualLoopsUpdateRateMinCVar,
+	TEXT("Sets minimum rate to check if sound becomes audible again at sound's max audible distance.\n"),
+	ECVF_Default);
+
+static float VirtualLoopsUpdateRateMaxCVar = 3.0f;
+FAutoConsoleVariableRef CVarVirtualLoopsUpdateRateMax(
+	TEXT("au.VirtualLoops.UpdateRate.Max"),
+	VirtualLoopsUpdateRateMaxCVar,
+	TEXT("Sets maximum rate to check if sound becomes audible again (at beyond sound's max audible distance + perf scaling distance).\n"),
+	ECVF_Default);
+
+FAudioVirtualLoop::FAudioVirtualLoop()
 	: TimeSinceLastUpdate(0.0f)
 	, UpdateInterval(0.0f)
-	, AudioDevice(&InAudioDevice)
 	, ActiveSound(nullptr)
 {
-	SetActiveSound(NewActiveSound);
 }
 
-FAudioVirtualLoop::~FAudioVirtualLoop()
+bool FAudioVirtualLoop::Virtualize(const FActiveSound& InActiveSound, bool bDoRangeCheck, FAudioVirtualLoop& OutVirtualLoop)
 {
-	check(ActiveSound);
-	delete(ActiveSound);
+	FAudioDevice* AudioDevice = InActiveSound.AudioDevice;
+	check(AudioDevice);
+
+	return Virtualize(InActiveSound, *AudioDevice, bDoRangeCheck, OutVirtualLoop);
 }
 
-FAudioVirtualLoop* FAudioVirtualLoop::Virtualize(FAudioDevice& InAudioDevice, const FActiveSound& InActiveSound, bool bDoRangeCheck)
+bool FAudioVirtualLoop::Virtualize(const FActiveSound& InActiveSound, FAudioDevice& AudioDevice, bool bDoRangeCheck, FAudioVirtualLoop& OutVirtualLoop)
 {
 	if (!bVirtualLoopsEnabledCVar || InActiveSound.bIsPreviewSound || !InActiveSound.IsLooping())
 	{
-		return nullptr;
+		return false;
 	}
 
-	if (bDoRangeCheck && IsInAudibleRange(InAudioDevice, InActiveSound))
+	if (InActiveSound.bFadingOut || InActiveSound.bIsStopping)
 	{
-		return nullptr;
+		return false;
 	}
 
-	return new FAudioVirtualLoop(InAudioDevice, InActiveSound);
+	if (bDoRangeCheck && IsInAudibleRange(InActiveSound, &AudioDevice))
+	{
+		return false;
+	}
+
+	FActiveSound* ActiveSound = new FActiveSound(InActiveSound);
+
+	ActiveSound->bAsyncOcclusionPending = false;
+	ActiveSound->AudioDevice = &AudioDevice;
+	ActiveSound->bIsPlayingAudio = false;
+	ActiveSound->ConcurrencyGroupIDs.Reset();
+	ActiveSound->ConcurrencyGroupVolumeScales.Reset();
+	ActiveSound->VolumeConcurrency = 1.0f;
+	ActiveSound->WaveInstances.Reset();
+
+	OutVirtualLoop.ActiveSound = ActiveSound;
+	return true;
 }
 
 void FAudioVirtualLoop::CalculateUpdateInterval(bool bIsAtMaxConcurrency)
@@ -71,14 +89,18 @@ void FAudioVirtualLoop::CalculateUpdateInterval(bool bIsAtMaxConcurrency)
 	// virtualizing.
 	if (bIsAtMaxConcurrency)
 	{
-		UpdateInterval = VirtualUpdateRateMaxCVar;
+		UpdateInterval = VirtualLoopsUpdateRateMaxCVar;
 	}
 	else
 	{
+		check(ActiveSound);
+		FAudioDevice* AudioDevice = ActiveSound->AudioDevice;
+		check(AudioDevice);
+
 		const float DistanceToListener = AudioDevice->GetDistanceToNearestListener(ActiveSound->Transform.GetLocation());
 		const float DistanceRatio = (DistanceToListener - ActiveSound->MaxDistance) / FMath::Max(VirtualLoopsPerfDistanceCVar, 1.0f);
 		const float DistanceRatioClamped = FMath::Clamp(DistanceRatio, 0.0f, 1.0f);
-		UpdateInterval = FMath::Lerp(VirtualUpdateRateMinCVar, VirtualUpdateRateMaxCVar, DistanceRatioClamped);
+		UpdateInterval = FMath::Lerp(VirtualLoopsUpdateRateMinCVar, VirtualLoopsUpdateRateMaxCVar, DistanceRatioClamped);
 	}
 }
 
@@ -88,98 +110,51 @@ FActiveSound& FAudioVirtualLoop::GetActiveSound()
 	return *ActiveSound;
 }
 
-bool FAudioVirtualLoop::IsInAudibleRange(const FAudioDevice& InAudioDevice, const FActiveSound& InActiveSound)
+const FActiveSound& FAudioVirtualLoop::GetActiveSound() const
 {
-	const FSoundAttenuationSettings* AttenuationSettingsToApply = InActiveSound.bHasAttenuationSettings ? &InActiveSound.AttenuationSettings : nullptr;
+	check(ActiveSound);
+	return *ActiveSound;
+}
 
+bool FAudioVirtualLoop::IsEnabled()
+{
+	return bVirtualLoopsEnabledCVar;
+}
+
+bool FAudioVirtualLoop::IsInAudibleRange(const FActiveSound& InActiveSound, const FAudioDevice* InAudioDevice)
+{
 	if (!InActiveSound.bAllowSpatialization)
 	{
 		return true;
 	}
 
+	const FAudioDevice* AudioDevice = InAudioDevice;
+	if (!AudioDevice)
+	{
+		AudioDevice = InActiveSound.AudioDevice;
+	}
+	check(AudioDevice);
+
 	float DistanceScale = 1.0f;
 	if (InActiveSound.bHasAttenuationSettings)
 	{
 		// If we are not using distance-based attenuation, this sound will be audible regardless of distance.
+		const FSoundAttenuationSettings* AttenuationSettingsToApply = InActiveSound.bHasAttenuationSettings ? &InActiveSound.AttenuationSettings : nullptr;
 		if (!AttenuationSettingsToApply->bAttenuate)
 		{
 			return true;
 		}
 
-		DistanceScale = AttenuationSettingsToApply->GetFocusDistanceScale(InAudioDevice.GetGlobalFocusSettings(), InActiveSound.FocusDistanceScale);
+		DistanceScale = AttenuationSettingsToApply->GetFocusDistanceScale(AudioDevice->GetGlobalFocusSettings(), InActiveSound.FocusDistanceScale);
 	}
 
 	DistanceScale = FMath::Max(DistanceScale, 0.0001f);
 	const FVector Location = InActiveSound.Transform.GetLocation();
-	return InAudioDevice.LocationIsAudible(Location, InActiveSound.MaxDistance / DistanceScale);
-}
-
-void FAudioVirtualLoop::SetActiveSound(const FActiveSound& InActiveSound)
-{
-	check(!ActiveSound);
-
-	ActiveSound = new FActiveSound();
-
-	ActiveSound->bIsUISound = InActiveSound.bIsUISound;
-
-	ActiveSound->SetAudioComponent(InActiveSound);
-	ActiveSound->SetAudioDevice(AudioDevice);
-
-	ActiveSound->SetSound(InActiveSound.GetSound());
-	ActiveSound->SetSoundClass(InActiveSound.GetSoundClass());
-	ActiveSound->SetWorld(InActiveSound.GetWorld());
-
-	ActiveSound->ConcurrencySet = InActiveSound.ConcurrencySet;
-	ActiveSound->VolumeMultiplier = InActiveSound.VolumeMultiplier;
-	ActiveSound->Priority = InActiveSound.Priority;
-	ActiveSound->PitchMultiplier = InActiveSound.PitchMultiplier;
-	ActiveSound->bEnableLowPassFilter = InActiveSound.bEnableLowPassFilter;
-	ActiveSound->LowPassFilterFrequency = InActiveSound.LowPassFilterFrequency;
-	ActiveSound->RequestedStartTime = InActiveSound.RequestedStartTime;
-	ActiveSound->SubtitlePriority = InActiveSound.SubtitlePriority;
-	ActiveSound->bShouldRemainActiveIfDropped = InActiveSound.bShouldRemainActiveIfDropped;
-	ActiveSound->bHandleSubtitles = InActiveSound.bHandleSubtitles;
-	ActiveSound->bIgnoreForFlushing = InActiveSound.bIgnoreForFlushing;
-	ActiveSound->bIsUISound = InActiveSound.bIsUISound;
-	ActiveSound->bIsMusic = InActiveSound.bIsMusic;
-	ActiveSound->bAlwaysPlay = InActiveSound.bAlwaysPlay;
-	ActiveSound->bReverb = InActiveSound.bReverb;
-	ActiveSound->bCenterChannelOnly = InActiveSound.bCenterChannelOnly;
-	ActiveSound->bIsPreviewSound = InActiveSound.bIsPreviewSound;
-	ActiveSound->bLocationDefined = InActiveSound.bLocationDefined;
-	ActiveSound->bIsPaused = InActiveSound.bIsPaused;
-	ActiveSound->Transform = InActiveSound.Transform;
-	ActiveSound->bAllowSpatialization = InActiveSound.bAllowSpatialization;
-	ActiveSound->bHasAttenuationSettings = InActiveSound.bHasAttenuationSettings;
-	ActiveSound->AttenuationSettings = InActiveSound.AttenuationSettings;
-	ActiveSound->FocusPriorityScale = InActiveSound.FocusPriorityScale;
-	ActiveSound->FocusDistanceScale = InActiveSound.FocusDistanceScale;
-	ActiveSound->EnvelopeFollowerAttackTime = InActiveSound.EnvelopeFollowerAttackTime;
-	ActiveSound->EnvelopeFollowerReleaseTime = InActiveSound.EnvelopeFollowerReleaseTime;
-
-	ActiveSound->bUpdatePlayPercentage = InActiveSound.bUpdatePlayPercentage;
-	ActiveSound->bUpdateSingleEnvelopeValue = InActiveSound.bUpdateSingleEnvelopeValue;
-	ActiveSound->bUpdateMultiEnvelopeValue = InActiveSound.bUpdateMultiEnvelopeValue;
-
-	ActiveSound->bUpdatePlaybackTime = InActiveSound.bUpdatePlaybackTime;
-
-	ActiveSound->MaxDistance = InActiveSound.MaxDistance;
-	ActiveSound->InstanceParameters = InActiveSound.InstanceParameters;
-	ActiveSound->TargetAdjustVolumeMultiplier = InActiveSound.TargetAdjustVolumeMultiplier;
-	ActiveSound->CurrentAdjustVolumeMultiplier = InActiveSound.CurrentAdjustVolumeMultiplier;
-	ActiveSound->TargetAdjustVolumeStopTime = InActiveSound.TargetAdjustVolumeStopTime;
+	return AudioDevice->LocationIsAudible(Location, InActiveSound.MaxDistance / DistanceScale);
 }
 
-bool FAudioVirtualLoop::Update()
+bool FAudioVirtualLoop::CanRealize(float DeltaTime)
 {
-	check(AudioDevice);
-	check(ActiveSound);
-
-	const float DeltaTime = AudioDevice->GetDeviceDeltaTime();
-
-	// Keep playback time up-to-date as it may be used to evaluate concurrency
-	ActiveSound->PlaybackTime += DeltaTime * ActiveSound->MinCurrentPitch;
-
 	if (UpdateInterval > 0.0f)
 	{
 		TimeSinceLastUpdate += DeltaTime;
@@ -190,8 +165,10 @@ bool FAudioVirtualLoop::Update()
 		TimeSinceLastUpdate = 0.0f;
 	}
 
+	DrawDebugInfo();
+
 	// If not audible, update when will be checked again and return false
-	if (!IsInAudibleRange(*AudioDevice, *ActiveSound))
+	if (!IsInAudibleRange(*ActiveSound))
 	{
 		CalculateUpdateInterval();
 		return false;
@@ -199,3 +176,34 @@ bool FAudioVirtualLoop::Update()
 
 	return true;
 }
+
+void FAudioVirtualLoop::DrawDebugInfo() const
+{
+#if ENABLE_DRAW_DEBUG
+	// Draw 3d Debug information about this source, if enabled
+	FAudioDeviceManager* DeviceManager = GEngine->GetAudioDeviceManager();
+
+	if (DeviceManager && DeviceManager->IsVisualizeDebug3dEnabled())
+	{
+		DECLARE_CYCLE_STAT(TEXT("FAudioThreadTask.DrawVirtualLoopDebugInfo"), STAT_AudioDrawVirtualLoopDebugInfo, STATGROUP_TaskGraphTasks);
+
+		USoundBase* Sound = ActiveSound->GetSound();
+		check(Sound);
+
+		const FTransform Transform = ActiveSound->Transform;
+		const TWeakObjectPtr<UWorld> World = ActiveSound->GetWeakWorld();
+		const FString Name = Sound->GetName();
+		const float DrawInterval = UpdateInterval;
+		FAudioThread::RunCommandOnGameThread([World, Transform, Name, DrawInterval]()
+		{
+			if (World.IsValid())
+			{
+				FVector Location = Transform.GetLocation();
+				FRotator Rotation = Transform.GetRotation().Rotator();
+				DrawDebugCrosshairs(World.Get(), Location, Rotation, 20.0f, FColor::Blue, false, DrawInterval, SDPG_Foreground);
+				DrawDebugString(World.Get(), Location + FVector(0, 0, 32), *Name, nullptr, FColor::Blue, DrawInterval, false);
+			}
+		}, GET_STATID(STAT_AudioDrawVirtualLoopDebugInfo));
+	}
+#endif // ENABLE_DRAW_DEBUG
+}
diff --git a/Engine/Source/Runtime/Engine/Private/SoundConcurrency.cpp b/Engine/Source/Runtime/Engine/Private/SoundConcurrency.cpp
--- a/Engine/Source/Runtime/Engine/Private/SoundConcurrency.cpp
+++ b/Engine/Source/Runtime/Engine/Private/SoundConcurrency.cpp
@@ -656,10 +656,11 @@ FActiveSound* FSoundConcurrencyManager::CreateAndEvictActiveSounds(const FActive
 		if (bAllowVirtual)
 		{
 			const bool bDoRangeCheck = true;
-			if (FAudioVirtualLoop* VirtualLoop = FAudioVirtualLoop::Virtualize(*AudioDevice, *SoundToEvict, bDoRangeCheck))
+			FAudioVirtualLoop VirtualLoop;
+			if (FAudioVirtualLoop::Virtualize(*SoundToEvict, bDoRangeCheck, VirtualLoop))
 			{
 				SoundToEvict->ClearAudioComponent();
-				AudioDevice->AddVirtualLoop(*VirtualLoop);
+				AudioDevice->AddVirtualLoop(VirtualLoop);
 			}
 		}
 	}
diff --git a/Engine/Source/Runtime/Engine/Public/ActiveSound.h b/Engine/Source/Runtime/Engine/Public/ActiveSound.h
--- a/Engine/Source/Runtime/Engine/Public/ActiveSound.h
+++ b/Engine/Source/Runtime/Engine/Public/ActiveSound.h
@@ -7,7 +7,6 @@
 #include "Sound/SoundAttenuation.h"
 #include "HAL/ThreadSafeBool.h"
 #include "Audio.h"
-#include "AudioDevice.h"
 #include "AudioDynamicParameter.h"
 #include "Sound/SoundConcurrency.h"
 #include "Components/AudioComponent.h"
@@ -23,6 +22,7 @@ struct FSoundSubmixSendInfo;
 struct FSoundSourceBusSendInfo;
 class USoundWave;
 struct FListener;
+struct FAttenuationListenerData;
 
 /**
  *	Struct used for gathering the final parameters to apply to a wave instance
diff --git a/Engine/Source/Runtime/Engine/Public/Audio.h b/Engine/Source/Runtime/Engine/Public/Audio.h
--- a/Engine/Source/Runtime/Engine/Public/Audio.h
+++ b/Engine/Source/Runtime/Engine/Public/Audio.h
@@ -82,6 +82,7 @@ DECLARE_CYCLE_STAT_EXTERN(TEXT("Audio Evaluate Concurrency"), STAT_AudioEvaluate
 DECLARE_DWORD_COUNTER_STAT_EXTERN( TEXT( "Audio Sources" ), STAT_AudioSources, STATGROUP_Audio , );
 DECLARE_DWORD_COUNTER_STAT_EXTERN( TEXT( "Wave Instances" ), STAT_WaveInstances, STATGROUP_Audio , );
 DECLARE_DWORD_COUNTER_STAT_EXTERN( TEXT( "Wave Instances Dropped" ), STAT_WavesDroppedDueToPriority, STATGROUP_Audio , );
+DECLARE_DWORD_COUNTER_STAT_EXTERN(TEXT("Virtualized Loops"), STAT_AudioVirtualLoops, STATGROUP_Audio, );
 DECLARE_DWORD_COUNTER_STAT_EXTERN( TEXT( "Audible Wave Instances Dropped" ), STAT_AudibleWavesDroppedDueToPriority, STATGROUP_Audio , );
 DECLARE_DWORD_COUNTER_STAT_EXTERN( TEXT( "Finished delegates called" ), STAT_AudioFinishedDelegatesCalled, STATGROUP_Audio , );
 DECLARE_CYCLE_STAT_EXTERN( TEXT( "Finished delegates time" ), STAT_AudioFinishedDelegates, STATGROUP_Audio , );
diff --git a/Engine/Source/Runtime/Engine/Public/AudioDevice.h b/Engine/Source/Runtime/Engine/Public/AudioDevice.h
--- a/Engine/Source/Runtime/Engine/Public/AudioDevice.h
+++ b/Engine/Source/Runtime/Engine/Public/AudioDevice.h
@@ -18,6 +18,7 @@
 #include "AudioDeviceManager.h"
 #include "DSP/SpectrumAnalyzer.h"
 #include "EngineGlobals.h"
+#include "AudioVirtualLoop.h"
 
 /**
  * Forward declares
@@ -25,7 +26,6 @@
 
 class FArchive;
 class FAudioEffectsManager;
-class FAudioVirtualLoop;
 class FCanvas;
 class FOutputDevice;
 class FReferenceCollector;
@@ -1533,7 +1533,7 @@ public:
 	/** Returns the audio clock of the audio device. Not supported on all platforms. */
 	double GetAudioClock() const { return AudioClock; }
 
-	void AddVirtualLoop(FAudioVirtualLoop& VirtualLoop);
+	void AddVirtualLoop(const FAudioVirtualLoop& InVirtualLoop);
 
 #if !UE_BUILD_SHIPPING
 	void DumpActiveSounds() const;
@@ -1823,7 +1823,7 @@ private:
 	/** Array of dormant loops stopped due to proximity/applicable concurrency rules
 	  * that can be retriggered.
 	  */
-	TMap<FActiveSound*, FAudioVirtualLoop*> VirtualLoops;
+	TMap<FActiveSound*, FAudioVirtualLoop> VirtualLoops;
 
 	TArray<FAudioComponentPtr> ProximityRetriggerComponents;
 
diff --git a/Engine/Source/Runtime/Engine/Public/AudioVirtualLoop.h b/Engine/Source/Runtime/Engine/Public/AudioVirtualLoop.h
--- a/Engine/Source/Runtime/Engine/Public/AudioVirtualLoop.h
+++ b/Engine/Source/Runtime/Engine/Public/AudioVirtualLoop.h
@@ -4,40 +4,48 @@
 #include "CoreMinimal.h"
 #include "ActiveSound.h"
 
-class FAudioDevice;
-
 
 /**
  * Class that tracks virtualized looping active sounds that are eligible to revive re-trigger
  * as long as no stop request is received from the game thread.
  */
-class ENGINE_API FAudioVirtualLoop
+struct ENGINE_API FAudioVirtualLoop
 {
+private:
 	float TimeSinceLastUpdate;
 	float UpdateInterval;
-	FAudioDevice* AudioDevice;
-	FActiveSound* ActiveSound;
 
-	FAudioVirtualLoop(FAudioDevice& AudioDevice, const FActiveSound& NewActiveSound);
+	FActiveSound* ActiveSound;
 
 public:
-	~FAudioVirtualLoop();
+	FAudioVirtualLoop();
+
+	/**
+	 * Checks if provided active sound is available to be virtualized.  If so, returns new active sound ready to be
+	 * added to virtual loop management by parent audio device.
+	 */
+	static bool Virtualize(const FActiveSound& InActiveSound, bool bDoRangeCheck, FAudioVirtualLoop& OutVirtualLoop);
+	static bool Virtualize(const FActiveSound& InActiveSound, FAudioDevice& AudioDevice, bool bDoRangeCheck, FAudioVirtualLoop& OutVirtualLoop);
 
 	/**
-	 * Checks if provided active sound is available to be virtualized.  If so, returns new virtual loop ready to be
-	 * added to virtual loop management in audio device.
+	 * Whether the virtual loop system is enabled or not
 	 */
-	static FAudioVirtualLoop* Virtualize(FAudioDevice& InAudioDevice, const FActiveSound& InActiveSound, bool bDoRangeCheck);
+	static bool IsEnabled();
 
 	/**
 	 * Returns the internally-managed active sound.
 	 */
 	FActiveSound& GetActiveSound();
 
+	/**
+	 * Returns the internally-managed active sound.
+	 */
+	const FActiveSound& GetActiveSound() const;
+
 	/**
 	 * Check if provided active sound is in audible range.
 	 */
-	static bool IsInAudibleRange(const FAudioDevice& InAudioDevice, const FActiveSound& ActiveSound);
+	static bool IsInAudibleRange(const FActiveSound& InActiveSound, const FAudioDevice* InAudioDevice = nullptr);
 
 	/**
 	 * Overrides the update interval to the provided length.
@@ -48,8 +56,10 @@ public:
 	  * Updates the loop and checks if ready to play (or 'realize').
 	  * Returns whether or not the sound is ready to be realized.
 	  */
-	bool Update();
+	bool CanRealize(float DeltaTime);
 
-private:
-	void SetActiveSound(const FActiveSound& InActiveSound);
+	/**
+	 * Draws debug info for virtual sound
+	 */
+	void DrawDebugInfo() const;
 };
