Add Analytics covering Queued Actors and Async Loads for networking

diff --git a/Engine/Source/Runtime/Engine/Classes/Engine/NetConnection.h b/Engine/Source/Runtime/Engine/Classes/Engine/NetConnection.h
--- a/Engine/Source/Runtime/Engine/Classes/Engine/NetConnection.h
+++ b/Engine/Source/Runtime/Engine/Classes/Engine/NetConnection.h
@@ -24,6 +24,7 @@
 #include "Analytics/EngineNetAnalytics.h"
 #include "PacketTraits.h"
 #include "Net/Util/ResizableCircularQueue.h"
+#include "Net/DelinquencyAnalytics.h"
 
 #include "NetConnection.generated.h"
 
@@ -1073,6 +1074,19 @@ public:
 	/** Called when an actor channel is open and knows its NetGUID. */
 	ENGINE_API virtual void NotifyActorNetGUID(UActorChannel* Channel) {}
 
+	/**
+	 * Returns the current delinquency analytics and resets them.
+	 * This would be similar to calls to Get and Reset separately, except that the caller
+	 * will assume ownership of data in this case.
+	 */
+	ENGINE_API void ConsumeQueuedActorDelinquencyAnalytics(FNetQueuedActorDelinquencyAnalytics& Out);
+
+	/** Returns the current delinquency analytics. */
+	ENGINE_API const FNetQueuedActorDelinquencyAnalytics& GetQueuedActorDelinquencyAnalytics() const;
+
+	/** Resets the current delinquency analytics. */
+	ENGINE_API void ResetQueuedActorDelinquencyAnalytics();
+
 protected:
 
 	ENGINE_API void SetPendingCloseDueToSocketSendFailure();
@@ -1184,8 +1198,6 @@ private:
 	bool bFlushingPacketOrderCache;
 };
 
-
-
 /** Help structs for temporarily setting network settings */
 struct FNetConnectionSettings
 {
diff --git a/Engine/Source/Runtime/Engine/Classes/Engine/NetDriver.h b/Engine/Source/Runtime/Engine/Classes/Engine/NetDriver.h
--- a/Engine/Source/Runtime/Engine/Classes/Engine/NetDriver.h
+++ b/Engine/Source/Runtime/Engine/Classes/Engine/NetDriver.h
@@ -15,6 +15,7 @@
 #include "Channel.h"
 #include "DDoSDetection.h"
 #include "IPAddress.h"
+#include "Net/DelinquencyAnalytics.h"
 
 #include "NetDriver.generated.h"
 
@@ -1164,6 +1165,19 @@ public:
 	/** Called when an actor channel is cleaned up foor an actor. */
 	ENGINE_API virtual void NotifyActorChannelCleanedUp(UActorChannel* Channel, EChannelCloseReason CloseReason) {}
 
+	/**
+	 * Returns the current delinquency analytics and resets them.
+	 * This would be similar to calls to Get and Reset separately, except that the caller
+	 * will assume ownership of data in this case.
+	 */
+	ENGINE_API void ConsumeAsyncLoadDelinquencyAnalytics(FNetAsyncLoadDelinquencyAnalytics& Out);
+
+	/** Returns the current delinquency analytics. */
+	ENGINE_API const FNetAsyncLoadDelinquencyAnalytics& GetAsyncLoadDelinquencyAnalytics() const;
+
+	/** Resets the current delinquency analytics. */
+	ENGINE_API void ResetAsyncLoadDelinquencyAnalytics();
+
 protected:
 
 	/** Register all TickDispatch, TickFlush, PostTickFlush to tick in World */
@@ -1227,7 +1241,6 @@ private:
 
 	void CreateReplicatedStaticActorDestructionInfo(UNetDriver* NetDriver, ULevel* Level, const FReplicatedStaticActorDestructionInfo& Info);
 
-
 	void FlushActorDormancyInternal(AActor *Actor);
 
 	void LoadChannelDefinitions();
@@ -1246,4 +1259,4 @@ private:
 
 	/** NetDriver time to end packet loss burst simulation. */
 	float PacketLossBurstEndTime;
-};
+};
\ No newline at end of file
diff --git a/Engine/Source/Runtime/Engine/Classes/Engine/PackageMapClient.h b/Engine/Source/Runtime/Engine/Classes/Engine/PackageMapClient.h
--- a/Engine/Source/Runtime/Engine/Classes/Engine/PackageMapClient.h
+++ b/Engine/Source/Runtime/Engine/Classes/Engine/PackageMapClient.h
@@ -27,6 +27,7 @@
 #include "Misc/NetworkVersion.h"
 #include "UObject/CoreNet.h"
 #include "Net/DataBunch.h"
+#include "Net/DelinquencyAnalytics.h"
 #include "PackageMapClient.generated.h"
 
 class UNetConnection;
@@ -260,6 +261,10 @@ public:
 
 	void			CountBytes(FArchive& Ar) const;
 
+	void ConsumeAsyncLoadDelinquencyAnalytics(FNetAsyncLoadDelinquencyAnalytics& Out);
+	const FNetAsyncLoadDelinquencyAnalytics& GetAsyncLoadDelinquencyAnalytics() const;
+	void ResetAsyncLoadDelinquencyAnalytics();
+
 	TMap< FNetworkGUID, FNetGuidCacheObject >		ObjectLookup;
 	TMap< TWeakObjectPtr< UObject >, FNetworkGUID >	NetGUIDLookup;
 	int32											UniqueNetIDs[2];
@@ -269,6 +274,7 @@ public:
 
 	UNetDriver *									Driver;
 
+	UE_DEPRECATED(4.23, "This member will be made private in future versions")
 	TMap< FName, FNetworkGUID >						PendingAsyncPackages;
 
 	ENetworkChecksumMode							NetworkChecksumMode;
@@ -301,6 +307,26 @@ public:
 	TMap<FNetworkGUID, FString>						History;
 private:
 #endif
+
+	struct FPendingAsyncLoadRequest
+	{
+		FPendingAsyncLoadRequest(const FNetworkGUID InNetGUID, const float InRequestStartTime):
+			NetGUID(InNetGUID),
+			RequestStartTime(InRequestStartTime)
+		{
+		}
+
+		FNetworkGUID NetGUID;
+		float RequestStartTime;
+	};
+
+	/** Set of packages that are currently pending Async loads, referenced by package name. */
+	TMap<FName, FPendingAsyncLoadRequest> PendingAsyncLoadRequests;
+
+	FNetAsyncLoadDelinquencyAnalytics DelinquentAsyncLoads;
+
+	void StartAsyncLoadingPackage(FNetGuidCacheObject& Object, const FNetworkGUID ObjectGUID, const bool bWasAlreadyAsyncLoading);
+	void ValidateAsyncLoadingPackage(FNetGuidCacheObject& Object, const FNetworkGUID ObjectGUID);
 };
 
 class ENGINE_API FPackageMapAckState
@@ -443,7 +469,9 @@ protected:
 
 	TArray<TArray<uint8>>				ExportGUIDArchives;
 	TSet< FNetworkGUID >				CurrentExportNetGUIDs;				// Current list of NetGUIDs being written to the Export Bunch.
-	TSet< FNetworkGUID >				CurrentQueuedBunchNetGUIDs;			// List of NetGuids with currently queued bunches
+
+	/** Set of Actor NetGUIDs with currently queued bunches and the time they were first queued. */
+	TMap<FNetworkGUID, float> CurrentQueuedBunchNetGUIDs;
 
 	TArray< FNetworkGUID >				PendingAckGUIDs;					// Quick access to all GUID's that haven't been acked
 
@@ -467,4 +495,14 @@ private:
 	void ReceiveNetFieldExportsCompat(FInBunch& InBunch);
 
 	bool bIgnoreReceivedExportGUIDs;
+
+public:
+
+	void ConsumeQueuedActorDelinquencyAnalytics(FNetQueuedActorDelinquencyAnalytics& Out);
+	const FNetQueuedActorDelinquencyAnalytics& GetQueuedActorDelinquencyAnalytics() const;
+	void ResetQueuedActorDelinquencyAnalytics();
+
+private:
+
+	FNetQueuedActorDelinquencyAnalytics DelinquentQueuedActors;
 };
diff --git a/Engine/Source/Runtime/Engine/Private/NetConnection.cpp b/Engine/Source/Runtime/Engine/Private/NetConnection.cpp
--- a/Engine/Source/Runtime/Engine/Private/NetConnection.cpp
+++ b/Engine/Source/Runtime/Engine/Private/NetConnection.cpp
@@ -3460,6 +3460,38 @@ void UNetConnection::SetPendingCloseDueToSocketSendFailure()
 	bConnectionPendingCloseDueToSocketSendFailure = true;
 }
 
+void UNetConnection::ConsumeQueuedActorDelinquencyAnalytics(FNetQueuedActorDelinquencyAnalytics& Out)
+{
+	if (UPackageMapClient* PackageMapClient = Cast<UPackageMapClient>(PackageMap))
+	{
+		return PackageMapClient->ConsumeQueuedActorDelinquencyAnalytics(Out);
+	}
+	else
+	{
+		Out.Reset();
+	}
+}
+
+const FNetQueuedActorDelinquencyAnalytics& UNetConnection::GetQueuedActorDelinquencyAnalytics() const
+{
+	static FNetQueuedActorDelinquencyAnalytics Empty;
+
+	if (UPackageMapClient const * const PackageMapClient = Cast<UPackageMapClient>(PackageMap))
+	{
+		return PackageMapClient->GetQueuedActorDelinquencyAnalytics();
+	}
+	
+	return Empty;
+}
+
+void UNetConnection::ResetQueuedActorDelinquencyAnalytics()
+{
+	if (UPackageMapClient* PackageMapClient = Cast<UPackageMapClient>(PackageMap))
+	{
+		PackageMapClient->ResetQueuedActorDelinquencyAnalytics();
+	}
+}
+
 /*-----------------------------------------------------------------------------
 	USimulatedClientNetConnection.
 -----------------------------------------------------------------------------*/
diff --git a/Engine/Source/Runtime/Engine/Private/NetworkDriver.cpp b/Engine/Source/Runtime/Engine/Private/NetworkDriver.cpp
--- a/Engine/Source/Runtime/Engine/Private/NetworkDriver.cpp
+++ b/Engine/Source/Runtime/Engine/Private/NetworkDriver.cpp
@@ -5450,6 +5450,38 @@ bool UNetDriver::ShouldClientDestroyActor(AActor* Actor) const
 	return (Actor && !Actor->IsA(ALevelScriptActor::StaticClass()));
 }
 
+void UNetDriver::ConsumeAsyncLoadDelinquencyAnalytics(FNetAsyncLoadDelinquencyAnalytics& Out)
+{
+	if (FNetGUIDCache* LocalGuidCache = GuidCache.Get())
+	{
+		LocalGuidCache->ConsumeAsyncLoadDelinquencyAnalytics(Out);
+	}
+	else
+	{
+		Out.Reset();
+	}
+}
+
+const FNetAsyncLoadDelinquencyAnalytics& UNetDriver::GetAsyncLoadDelinquencyAnalytics() const
+{
+	static const FNetAsyncLoadDelinquencyAnalytics Empty(0);
+
+	if (FNetGUIDCache const * const LocalGuidCache = GuidCache.Get())
+	{
+		return LocalGuidCache->GetAsyncLoadDelinquencyAnalytics();
+	}
+
+	return Empty;
+}
+
+void UNetDriver::ResetAsyncLoadDelinquencyAnalytics()
+{
+	if (FNetGUIDCache* LocalGuidCache = GuidCache.Get())
+	{
+		LocalGuidCache->ResetAsyncLoadDelinquencyAnalytics();
+	}
+}
+
 FAutoConsoleCommandWithWorld	DumpRelevantActorsCommand(
 	TEXT("net.DumpRelevantActors"), 
 	TEXT( "Dumps information on relevant actors during next network update" ), 
diff --git a/Engine/Source/Runtime/Engine/Private/PackageMapClient.cpp b/Engine/Source/Runtime/Engine/Private/PackageMapClient.cpp
--- a/Engine/Source/Runtime/Engine/Private/PackageMapClient.cpp
+++ b/Engine/Source/Runtime/Engine/Private/PackageMapClient.cpp
@@ -49,6 +49,32 @@ static TAutoConsoleVariable<int32> CVarIgnoreNetworkChecksumMismatch( TEXT( "net
 extern FAutoConsoleVariableRef CVarEnableMultiplayerWorldOriginRebasing;
 extern TAutoConsoleVariable<int32> CVarFilterGuidRemapping;
 
+static float GGuidCacheTrackAsyncLoadingGUIDTreshold = 0.f;
+static FAutoConsoleVariableRef CVarTrackAsyncLoadingGUIDTreshold(
+	TEXT("net.TrackAsyncLoadingGUIDThreshold"),
+	GGuidCacheTrackAsyncLoadingGUIDTreshold,
+	TEXT("When > 0, any objects that take longer than the threshold to async load will be tracked."
+		" Threshold in seconds, @see FNetGUIDCache::ConsumeDelinquencyAnalytics. Used for Debugging and Analytics")
+);
+
+static float GPackageMapTrackQueuedActorTreshold = 0.f;
+static FAutoConsoleVariableRef CVarTrackQueuedActorTreshold(
+	TEXT("net.TrackQueuedActorThreshold"),
+	GPackageMapTrackQueuedActorTreshold,
+	TEXT("When > 0, any actors that spend longer than the threshold with queued bunches will be tracked."
+		" Threshold in seconds, @see UPackageMap::ConsumeDelinquencyAnalytics. Used for Debugging and Analytics")
+);
+
+static int32 GDelinquencyNumberOfTopOffendersToTrack = 10;
+static FAutoConsoleVariableRef CVarDelinquencyNumberOfTopOffendersToTrack(
+	TEXT("net.DelinquencyNumberOfTopOffendersToTrack"),
+	GDelinquencyNumberOfTopOffendersToTrack,
+	TEXT("When > 0 , this will be the number of 'TopOffenders' that are tracked by the PackageMap and GuidCache for"
+		" Queued Actors and Async Loads respectively."
+		" net.TrackAsyncLoadingGUIDThreshold / net.TrackQueuedActorThreshold still dictate whether or not any of these"
+		" items are tracked.")
+);
+
 void BroadcastNetFailure(UNetDriver* Driver, ENetworkFailure::Type FailureType, const FString& ErrorStr)
 {
 	UWorld* World = Driver->GetWorld();
@@ -89,7 +115,8 @@ void BroadcastNetFailure(UNetDriver* Driver, ENetworkFailure::Type FailureType,
 -----------------------------------------------------------------------------*/
 UPackageMapClient::UPackageMapClient(const FObjectInitializer& ObjectInitializer)
 	: Super(ObjectInitializer)
-  , Connection(nullptr)
+	, Connection(nullptr)
+	, DelinquentQueuedActors(GDelinquencyNumberOfTopOffendersToTrack > 0 ? GDelinquencyNumberOfTopOffendersToTrack : 0)
 {
 }
 
@@ -2093,11 +2120,43 @@ void UPackageMapClient::SetHasQueuedBunches(const FNetworkGUID& NetGUID, bool bH
 {
 	if (bHasQueuedBunches)
 	{
-		CurrentQueuedBunchNetGUIDs.Add(NetGUID);
+		if (GPackageMapTrackQueuedActorTreshold > 0.f)
+		{
+			if (UNetDriver const * const NetDriver = ((Connection) ? Connection->GetDriver() : nullptr))
+			{
+				CurrentQueuedBunchNetGUIDs.Emplace(NetGUID, NetDriver->Time);
+			}
+		}
+
+		DelinquentQueuedActors.MaxConcurrentQueuedActors = FMath::Max<uint32>(DelinquentQueuedActors.MaxConcurrentQueuedActors, CurrentQueuedBunchNetGUIDs.Num());
 	}
 	else
 	{
-		CurrentQueuedBunchNetGUIDs.Remove(NetGUID);
+		float StartTime = 0.f;
+
+		// We try to remove the value regardless of whether or not the CVar is on.
+		// That way if it's toggled on and off, we don't end up wasting resources.
+		// If it is disabled with entries in DelinquentQueuedActors, it will be up
+		// to clients to clear out the map by calling ConsumeDelinquencyAnalytics.
+		if (CurrentQueuedBunchNetGUIDs.RemoveAndCopyValue(NetGUID, StartTime) &&
+			GPackageMapTrackQueuedActorTreshold &&
+			GuidCache)
+		{
+			if (UNetDriver const * const NetDriver = ((Connection) ? Connection->GetDriver() : nullptr))
+			{
+				const float QueuedTime = NetDriver->Time - StartTime;
+				if (QueuedTime > GPackageMapTrackQueuedActorTreshold)
+				{
+					if (FNetGuidCacheObject const * const CacheObject = GuidCache->ObjectLookup.Find(NetGUID))
+					{
+						if (UObject const * const Object = CacheObject->Object.Get())
+						{
+							DelinquentQueuedActors.DelinquentQueuedActors.Emplace(Object->GetClass()->GetFName(), QueuedTime);
+						}
+					}
+				}
+			}
+		}
 	}
 }
 
@@ -2143,6 +2202,8 @@ void UPackageMapClient::Serialize(FArchive& Ar)
 		GRANULAR_NETWORK_MEMORY_TRACKING_TRACK("MustBeMappedGuidsInLastBunch", MustBeMappedGuidsInLastBunch.CountBytes(Ar));
 		GRANULAR_NETWORK_MEMORY_TRACKING_TRACK("NetFieldExports", NetFieldExports.CountBytes(Ar));
 
+		GRANULAR_NETWORK_MEMORY_TRACKING_TRACK("DelinquentQueuedActors", DelinquentQueuedActors.CountBytes(Ar));
+
 		// Don't count the GUID Cache here. Instead, we'll let the UNetDriver count it as
 		// that's the class that constructs it.
 	}
@@ -2157,6 +2218,7 @@ FNetGUIDCache::FNetGUIDCache(UNetDriver* InDriver)
 	, NetworkChecksumMode(ENetworkChecksumMode::SaveAndUse)
 	, AsyncLoadMode(EAsyncLoadMode::UseCVar)
 	, IsExportingNetGUIDBunch(false)
+	, DelinquentAsyncLoads(GDelinquencyNumberOfTopOffendersToTrack > 0 ? GDelinquencyNumberOfTopOffendersToTrack : 0)
 {
 	UniqueNetIDs[0] = UniqueNetIDs[1] = 0;
 	UniqueNetFieldExportGroupPathIndex = 0;
@@ -2637,48 +2699,104 @@ void FNetGUIDCache::RegisterNetGUIDFromPath_Server( const FNetworkGUID& NetGUID,
 	RegisterNetGUID_Internal( NetGUID, CacheObject );
 }
 
-void FNetGUIDCache::AsyncPackageCallback(const FName& PackageName, UPackage * Package, EAsyncLoadingResult::Type Result)
+void FNetGUIDCache::ValidateAsyncLoadingPackage(FNetGuidCacheObject& CacheObject, const FNetworkGUID NetGUID)
 {
-	check( Package == NULL || Package->IsFullyLoaded() );
+	// With level streaming support we may end up trying to load the same package with a different
+	// NetGUID during replay fast-forwarding. This is because if a package was unloaded, and later
+	// re-loaded, it will likely be assigned a new NetGUID (since the TWeakObjectPtr to the old package
+	// in the cache object would have gone stale). During replay fast-forward, it's possible
+	// to see the new NetGUID before the previous one has finished loading, so here we fix up
+	// PendingAsyncPackages to refer to the new NewGUID.
+	FPendingAsyncLoadRequest& PendingLoadRequest = PendingAsyncLoadRequests[CacheObject.PathName];
+	if (PendingLoadRequest.NetGUID != NetGUID)
+	{
+		UE_LOG(LogNetPackageMap, Log, TEXT("ValidateAsyncLoadingPackage: Already async loading package with a different NetGUID. Path: %s, original NetGUID: %s, new NetGUID: %s"),
+			*CacheObject.PathName.ToString(), *PendingLoadRequest.NetGUID.ToString(), *NetGUID.ToString());
 
-	FNetworkGUID NetGUID = PendingAsyncPackages.FindRef(PackageName);
-	
-	PendingAsyncPackages.Remove(PackageName);
+PRAGMA_DISABLE_DEPRECATION_WARNINGS
+		PendingAsyncPackages[CacheObject.PathName] = NetGUID;
+PRAGMA_ENABLE_DEPRECATION_WARNINGS
 
-	if ( !NetGUID.IsValid() )
+		PendingLoadRequest.NetGUID = NetGUID;
+	}
+	else
 	{
-		UE_LOG( LogNetPackageMap, Error, TEXT( "AsyncPackageCallback: Could not find package. Path: %s" ), *PackageName.ToString() );
-		return;
+		UE_LOG(LogNetPackageMap, Log, TEXT("ValidateAsyncLoadingPackage: Already async loading package. Path: %s, NetGUID: %s"), *CacheObject.PathName.ToString(), *NetGUID.ToString());
 	}
+}
 
-	FNetGuidCacheObject * CacheObject = ObjectLookup.Find( NetGUID );
+void FNetGUIDCache::StartAsyncLoadingPackage(FNetGuidCacheObject& CacheObject, const FNetworkGUID NetGUID, const bool bWasAlreadyAsyncLoading)
+{
+	// Something else is already async loading this package, calling load again will add our callback to the existing load request
+PRAGMA_DISABLE_DEPRECATION_WARNINGS
+	PendingAsyncPackages.Add(CacheObject.PathName, NetGUID);
+PRAGMA_ENABLE_DEPRECATION_WARNINGS
 
-	if ( CacheObject == NULL )
-	{
-		UE_LOG( LogNetPackageMap, Error, TEXT( "AsyncPackageCallback: Could not find net guid. Path: %s, NetGUID: %s" ), *PackageName.ToString(), *NetGUID.ToString() );
-		return;
-	}
+	PendingAsyncLoadRequests.Emplace(CacheObject.PathName, FPendingAsyncLoadRequest(NetGUID, Driver->Time));
 
-	if ( !CacheObject->bIsPending )
-	{
-		UE_LOG( LogNetPackageMap, Error, TEXT( "AsyncPackageCallback: Package wasn't pending. Path: %s, NetGUID: %s" ), *PackageName.ToString(), *NetGUID.ToString() );
-	}
+	DelinquentAsyncLoads.MaxConcurrentAsyncLoads = FMath::Max<uint32>(DelinquentAsyncLoads.MaxConcurrentAsyncLoads, PendingAsyncLoadRequests.Num());
 
-	CacheObject->bIsPending = false;
-	
-	if ( Package == NULL )
-	{
-		CacheObject->bIsBroken = true;
-		UE_LOG( LogNetPackageMap, Error, TEXT( "AsyncPackageCallback: Package FAILED to load. Path: %s, NetGUID: %s" ), *PackageName.ToString(), *NetGUID.ToString() );
-	}
+	CacheObject.bIsPending = true;
+	LoadPackageAsync(CacheObject.PathName.ToString(), FLoadPackageAsyncDelegate::CreateRaw(this, &FNetGUIDCache::AsyncPackageCallback));
+}
+
+void FNetGUIDCache::AsyncPackageCallback(const FName& PackageName, UPackage * Package, EAsyncLoadingResult::Type Result)
+{
+	check(Package == nullptr || Package->IsFullyLoaded());
 
-	if (CacheObject->Object.IsValid() && CacheObject->Object->GetWorld())
+	if (FPendingAsyncLoadRequest const * const PendingLoadRequest = PendingAsyncLoadRequests.Find(PackageName))
 	{
-		AGameStateBase* GS = CacheObject->Object->GetWorld()->GetGameState();
-		if (GS)
+		const bool bIsBroken = (Package == nullptr);
+
+		if (FNetGuidCacheObject* CacheObject = ObjectLookup.Find(PendingLoadRequest->NetGUID))
 		{
-			GS->AsyncPackageLoaded(CacheObject->Object.Get());
+			if (!CacheObject->bIsPending)
+			{
+				UE_LOG(LogNetPackageMap, Error, TEXT("AsyncPackageCallback: Package wasn't pending. Path: %s, NetGUID: %s"), *PackageName.ToString(), *PendingLoadRequest->NetGUID.ToString());
+			}
+
+			CacheObject->bIsPending = false;
+
+			if (bIsBroken)
+			{
+				CacheObject->bIsBroken = true;
+				UE_LOG(LogNetPackageMap, Error, TEXT("AsyncPackageCallback: Package FAILED to load. Path: %s, NetGUID: %s"), *PackageName.ToString(), *PendingLoadRequest->NetGUID.ToString());
+			}
+
+			if (UObject* Object = CacheObject->Object.Get())
+			{
+				if (UWorld* World = Object->GetWorld())
+				{
+					if (AGameStateBase* GS = World->GetGameState())
+					{
+						GS->AsyncPackageLoaded(Object);
+					}
+				}
+			}
+		}
+		else
+		{
+			UE_LOG(LogNetPackageMap, Error, TEXT("AsyncPackageCallback: Could not find net guid. Path: %s, NetGUID: %s"), *PackageName.ToString(), *PendingLoadRequest->NetGUID.ToString());
+		}
+
+		// This won't be the exact amount of time that we spent loading the package, but should
+		// give us a close enough estimate (within a frame time).
+		const float LoadTime = (Driver->Time - PendingLoadRequest->RequestStartTime);
+		if (GGuidCacheTrackAsyncLoadingGUIDTreshold > 0.f &&
+			LoadTime >= GGuidCacheTrackAsyncLoadingGUIDTreshold)
+		{
+			DelinquentAsyncLoads.DelinquentAsyncLoads.Emplace(PackageName, LoadTime);
 		}
+
+		PendingAsyncLoadRequests.Remove(PackageName);
+
+PRAGMA_DISABLE_DEPRECATION_WARNINGS
+		PendingAsyncPackages.Remove(PackageName);
+PRAGMA_ENABLE_DEPRECATION_WARNINGS
+	}
+	else
+	{
+		UE_LOG(LogNetPackageMap, Error, TEXT( "AsyncPackageCallback: Could not find package. Path: %s" ), *PackageName.ToString());
 	}
 }
 
@@ -2826,32 +2944,14 @@ UObject* FNetGUIDCache::GetObjectFromNetGUID( const FNetworkGUID& NetGUID, const
 
 			if ( ShouldAsyncLoad() )
 			{
-				if (!PendingAsyncPackages.Contains(CacheObjectPtr->PathName))
+				if (!PendingAsyncLoadRequests.Contains(CacheObjectPtr->PathName))
 				{
-					PendingAsyncPackages.Add(CacheObjectPtr->PathName, NetGUID);
-					CacheObjectPtr->bIsPending = true;
-					LoadPackageAsync(CacheObjectPtr->PathName.ToString(), FLoadPackageAsyncDelegate::CreateRaw(this, &FNetGUIDCache::AsyncPackageCallback));
-
+					StartAsyncLoadingPackage(*CacheObjectPtr, NetGUID, false);
 					UE_LOG(LogNetPackageMap, Log, TEXT("GetObjectFromNetGUID: Async loading package. Path: %s, NetGUID: %s"), *CacheObjectPtr->PathName.ToString(), *NetGUID.ToString());
 				}
 				else
 				{
-					// With level streaming support we may end up trying to load the same package with a different
-					// NetGUID during replay fast-forwarding. This is because if a package was unloaded, and later
-					// re-loaded, it will likely be assigned a new NetGUID (since the TWeakObjectPtr to the old package
-					// in the cache object would have gone stale). During replay fast-forward, it's possible
-					// to see the new NetGUID before the previous one has finished loading, so here we fix up
-					// PendingAsyncPackages to refer to the new NewGUID.
-					const FNetworkGUID ExistingNetGUID = PendingAsyncPackages[CacheObjectPtr->PathName];
-					if (ExistingNetGUID != NetGUID)
-					{
-						UE_LOG(LogNetPackageMap, Log, TEXT("GetObjectFromNetGUID: Already async loading package with a different NetGUID. Path: %s, original NetGUID: %s, new NetGUID: %s"), *CacheObjectPtr->PathName.ToString(), *ExistingNetGUID.ToString(), *NetGUID.ToString());
-						PendingAsyncPackages[CacheObjectPtr->PathName] = NetGUID;
-					}
-					else
-					{
-						UE_LOG(LogNetPackageMap, Log, TEXT("GetObjectFromNetGUID: Already async loading package. Path: %s, NetGUID: %s"), *CacheObjectPtr->PathName.ToString(), *NetGUID.ToString());
-					}
+					ValidateAsyncLoadingPackage(*CacheObjectPtr, NetGUID);
 				}
 
 				// There is nothing else to do except wait on the delegate to tell us this package is done loading
@@ -2907,10 +3007,7 @@ UObject* FNetGUIDCache::GetObjectFromNetGUID( const FNetworkGUID& NetGUID, const
 			if (ShouldAsyncLoad() && Package->HasAnyInternalFlags(EInternalObjectFlags::AsyncLoading))
 			{
 				// Something else is already async loading this package, calling load again will add our callback to the existing load request
-				PendingAsyncPackages.Add(CacheObjectPtr->PathName, NetGUID);
-				CacheObjectPtr->bIsPending = true;
-				LoadPackageAsync(CacheObjectPtr->PathName.ToString(), FLoadPackageAsyncDelegate::CreateRaw(this, &FNetGUIDCache::AsyncPackageCallback));
-
+				StartAsyncLoadingPackage(*CacheObjectPtr, NetGUID, true);
 				UE_LOG(LogNetPackageMap, Log, TEXT("GetObjectFromNetGUID: Listening to existing async load. Path: %s, NetGUID: %s"), *CacheObjectPtr->PathName.ToString(), *NetGUID.ToString());
 			}
 			else
@@ -3274,7 +3371,10 @@ void FNetGUIDCache::CountBytes(FArchive& Ar) const
 		}
 	);
 
+PRAGMA_DISABLE_DEPRECATION_WARNINGS
 	GRANULAR_NETWORK_MEMORY_TRACKING_TRACK("PendingAsyncPackages", PendingAsyncPackages.CountBytes(Ar));
+PRAGMA_ENABLE_DEPRECATION_WARNINGS
+
 	GRANULAR_NETWORK_MEMORY_TRACKING_TRACK("NetFieldExportGroupMap",
 		NetFieldExportGroupMap.CountBytes(Ar);
 
@@ -3310,6 +3410,8 @@ void FNetGUIDCache::CountBytes(FArchive& Ar) const
 		}
 	);
 #endif
+
+	GRANULAR_NETWORK_MEMORY_TRACKING_TRACK("DelinquentAsyncLoads", DelinquentAsyncLoads.CountBytes(Ar));
 }
 
 void FNetFieldExport::CountBytes(FArchive& Ar) const
@@ -3337,6 +3439,38 @@ void FPackageMapAckState::CountBytes(FArchive& Ar) const
 	NetFieldExportAcked.CountBytes(Ar);
 }
 
+void FNetGUIDCache::ConsumeAsyncLoadDelinquencyAnalytics(FNetAsyncLoadDelinquencyAnalytics& Out)
+{
+	Out = MoveTemp(DelinquentAsyncLoads);
+	DelinquentAsyncLoads.MaxConcurrentAsyncLoads = PendingAsyncLoadRequests.Num();
+}
+
+const FNetAsyncLoadDelinquencyAnalytics& FNetGUIDCache::GetAsyncLoadDelinquencyAnalytics() const
+{
+	return DelinquentAsyncLoads;
+}
+
+void FNetGUIDCache::ResetAsyncLoadDelinquencyAnalytics()
+{
+	DelinquentAsyncLoads.Reset();
+}
+
+void UPackageMapClient::ConsumeQueuedActorDelinquencyAnalytics(FNetQueuedActorDelinquencyAnalytics& Out)
+{
+	Out = MoveTemp(DelinquentQueuedActors);
+	DelinquentQueuedActors.MaxConcurrentQueuedActors = CurrentQueuedBunchNetGUIDs.Num();
+}
+
+const FNetQueuedActorDelinquencyAnalytics& UPackageMapClient::GetQueuedActorDelinquencyAnalytics() const
+{
+	return DelinquentQueuedActors;
+}
+
+void UPackageMapClient::ResetQueuedActorDelinquencyAnalytics()
+{
+	DelinquentQueuedActors.Reset();
+}
+
 //------------------------------------------------------
 // Debug command to see how many times we've exported each NetGUID
 // Used for measuring inefficiencies. Some duplication is unavoidable since we cannot garuntee atomicicity across multiple channels.
diff --git a/Engine/Source/Runtime/Engine/Public/Net/DelinquencyAnalytics.h b/Engine/Source/Runtime/Engine/Public/Net/DelinquencyAnalytics.h
new file mode 100644
--- /dev/null
+++ b/Engine/Source/Runtime/Engine/Public/Net/DelinquencyAnalytics.h
@@ -0,0 +1,322 @@
+#pragma once
+
+#include "CoreMinimal.h"
+#include "Algo/BinarySearch.h"
+#include "Serialization/Archive.h"
+
+/**
+ * Tracks an FName ID to a time value. Time will be context dependent, but usually
+ * represents the total amount of time a specific action took (how long a package
+ * took to load, how long an actor had queued bunches, etc.)
+ *
+ * Could have used a TPair, but this will make it more obvious what we're tracking.
+ */
+struct ENGINE_API FDelinquencyNameTimePair
+{
+public:
+
+	FDelinquencyNameTimePair(FName InName, float InTimeSeconds) :
+		Name(InName),
+		TimeSeconds(InTimeSeconds)
+	{
+	}
+
+	FName Name;
+	float TimeSeconds;
+};
+
+struct ENGINE_API FDelinquencyKeyFuncs : public BaseKeyFuncs<FDelinquencyNameTimePair, FDelinquencyNameTimePair, false>
+{
+	static KeyInitType GetSetKey(ElementInitType Element)
+	{
+		return Element;
+	}
+
+	static bool Matches(KeyInitType LHS, KeyInitType RHS)
+	{
+		return LHS.Name == RHS.Name;
+	}
+
+	static uint32 GetKeyHash(KeyInitType Key)
+	{
+		return GetTypeHash(Key.Name);
+	}
+};
+
+/**
+ * Convenience type that can be used to tracks information about things that can result in prolonged
+ * periods of apparent network inactivity, despite actually receiving traffic.
+ *
+ * The overall number of entries is expected to be small, but ultimately is left up to callers.
+ */
+struct ENGINE_API FDelinquencyAnalytics
+{
+public:
+
+	explicit FDelinquencyAnalytics(const uint32 InNumberOfTopOffendersToTrack) :
+		TotalTime(0.f),
+		NumberOfTopOffendersToTrack(InNumberOfTopOffendersToTrack)
+	{
+		TopOffenders.Reserve(NumberOfTopOffendersToTrack);
+	}
+
+	FDelinquencyAnalytics(FDelinquencyAnalytics&& Other):
+		TopOffenders(MoveTemp(Other.TopOffenders)),
+		AllDelinquents(MoveTemp(Other.AllDelinquents)),
+		TotalTime(Other.TotalTime),
+		NumberOfTopOffendersToTrack(Other.NumberOfTopOffendersToTrack)
+	{
+		TopOffenders.Reserve(NumberOfTopOffendersToTrack);
+	}
+
+	FDelinquencyAnalytics(const FDelinquencyAnalytics&) = delete;
+	const FDelinquencyAnalytics& operator=(const FDelinquencyAnalytics&) = delete;
+	FDelinquencyAnalytics& operator=(FDelinquencyAnalytics&&) = default;
+
+	void Emplace(FName Name, float TimeSeconds)
+	{
+		Add(FDelinquencyNameTimePair(Name, TimeSeconds));
+	}
+
+	/**
+	 * Adds the event to the delinquency tracking, by accumulating its time into total time,
+	 * and updating any existing events to choose the one with the highest time.
+	 *
+	 * When NumberOfTopOffendersToTrack == 0, we will just track the set of all events as well as the total time.
+	 *
+	 * When NumberOfTopOffendersToTrack > 0, we will track the set, total time, and also maintain sorted list
+	 * (highest to lowest) of events that occurred.
+	 *
+	 * By setting NumberOfTopOffendersToTrack to 0, users can manage their own lists of "TopOffenders", or
+	 * otherwise avoid the per add overhead of this tracking.
+	 */
+	void Add(FDelinquencyNameTimePair&& ToTrack)
+	{
+		struct FHelper
+		{
+			static bool Compare(const FDelinquencyNameTimePair& LHS, const FDelinquencyNameTimePair& RHS)
+			{
+				return LHS.TimeSeconds > RHS.TimeSeconds;
+			}
+		};
+
+		// Regardless of whether or not this item has been seen before, there was a new entry
+		// so we'll add that time to the total.
+		TotalTime += ToTrack.TimeSeconds;
+
+		// TODO: We might consider tracking the code below as totals instead of the max time for a single event.
+		// For example, an Actor could end up queueing bunches several times within a reporting window, and each
+		// of those events would add to the TotalTime, but below we will only have the time of the longest period
+		// where bunches were queued.
+
+		if (NumberOfTopOffendersToTrack == 0)
+		{
+			AllDelinquents.Emplace(MoveTemp(ToTrack));
+		}
+		else if (TopOffenders.Num() == 0)
+		{
+			TopOffenders.Add(ToTrack);
+			AllDelinquents.Emplace(MoveTemp(ToTrack));
+			return;
+		}
+		else
+		{
+			if (FDelinquencyNameTimePair* AlreadyTracked = AllDelinquents.Find(ToTrack))
+			{
+				// We found an entry, so check its time.
+				if (AlreadyTracked->TimeSeconds >= ToTrack.TimeSeconds)
+				{
+					// We already have tracked a worse offense for this entry, so there's nothing more
+					// we need to do.
+					return;
+				}
+				else if (TopOffenders.Num() > 0)
+				{
+					const float LeastOffensiveTime = TopOffenders.Last().TimeSeconds;
+					if (AlreadyTracked->TimeSeconds >= LeastOffensiveTime)
+					{
+						// Our previous offense would have been tracked in the TopOffenders list, so go ahead and remove it.
+						// We're sorted highest to lowest, but are using greater than to actually compare, so UpperBound should
+						// return the index above which of the entry just before us.
+						// However, it's possible that other entries have the same time, so we need to make sure we remove the
+						// correct one.
+						int32 MaybeOurEntry = Algo::UpperBound<>(TopOffenders, ToTrack, &FHelper::Compare);
+						while(true)
+						{
+							// Sanity check that we have actually found our entry.
+							// If we hit the end of the list, or we see an entry that should be less offensive, then
+							// we've missed our entry. This should never happen.
+							if (MaybeOurEntry == TopOffenders.Num() || TopOffenders[MaybeOurEntry].TimeSeconds > AlreadyTracked->TimeSeconds)
+							{
+								ensureMsgf(false, TEXT("FDelinquencyAnalytics::Add - Unable to find expected entry %s:%f, list may not be sorted!"), *AlreadyTracked->Name.ToString(), AlreadyTracked->TimeSeconds);
+								break;
+							}
+
+							// We found our entry, so we're done.
+							else if (TopOffenders[MaybeOurEntry].Name == AlreadyTracked->Name)
+							{
+								TopOffenders.RemoveAt(MaybeOurEntry, 1, false);
+								break;
+							}
+
+							++MaybeOurEntry;
+						}
+					}
+				}
+			}
+
+			AllDelinquents.Add(ToTrack);
+
+			const int32 LocalNumberOfTopOffendersToTrack = static_cast<int32>(NumberOfTopOffendersToTrack);
+			const int32 InsertAt = Algo::UpperBound<>(TopOffenders, ToTrack, &FHelper::Compare);
+
+			// Check to see if this time was ranked in our top offenders.
+			if (InsertAt < LocalNumberOfTopOffendersToTrack)
+			{
+				// If we're going to displace a previous top offender, remove the least offensive.
+				if (LocalNumberOfTopOffendersToTrack == TopOffenders.Num())
+				{
+					TopOffenders.RemoveAt(TopOffenders.Num() - 1, 1, false);
+				}
+
+				TopOffenders.InsertUninitialized(InsertAt, 1);
+				TopOffenders[InsertAt] = MoveTemp(ToTrack);
+			}
+		}
+	}
+
+	const TArray<FDelinquencyNameTimePair>& GetTopOffenders() const
+	{
+		return TopOffenders;
+	}
+
+	const TSet<FDelinquencyNameTimePair, FDelinquencyKeyFuncs>& GetAllDelinquents() const
+	{
+		return AllDelinquents;
+	}
+
+	const float GetTotalTime() const
+	{
+		return TotalTime;
+	}
+
+	const uint32 GetNumberOfTopOffendersToTrack() const
+	{
+		return NumberOfTopOffendersToTrack;
+	}
+
+	void Reset()
+	{
+		TopOffenders.Reset();
+		AllDelinquents.Reset();
+		TotalTime = 0;
+	}
+
+	void CountBytes(FArchive& Ar) const
+	{
+		TopOffenders.CountBytes(Ar);
+		AllDelinquents.CountBytes(Ar);
+	}
+
+private:
+
+	TArray<FDelinquencyNameTimePair> TopOffenders;
+	TSet<FDelinquencyNameTimePair, FDelinquencyKeyFuncs> AllDelinquents;
+	float TotalTime;
+
+	// This is explicitly non const, as we will be copying / moving these structs around.
+	uint32 NumberOfTopOffendersToTrack;
+};
+
+/**
+ * Tracks data related specific to a NetDriver that can can result in prolonged periods of apparent
+ * network inactivity, despite actually receiving traffic.
+ *
+ * This includes things like Pending Async Loads.
+ *
+ * Also @see FConnectionDelinquencyAnalytics and FDelinquencyAnalytics.
+ */
+struct ENGINE_API FNetAsyncLoadDelinquencyAnalytics
+{
+	FNetAsyncLoadDelinquencyAnalytics() :
+		DelinquentAsyncLoads(0),
+		MaxConcurrentAsyncLoads(0)
+	{
+	}
+
+	FNetAsyncLoadDelinquencyAnalytics(const uint32 NumberOfTopOffendersToTrack) :
+		DelinquentAsyncLoads(NumberOfTopOffendersToTrack),
+		MaxConcurrentAsyncLoads(0)
+	{
+	}
+
+	FNetAsyncLoadDelinquencyAnalytics(FNetAsyncLoadDelinquencyAnalytics&& Other) :
+		DelinquentAsyncLoads(MoveTemp(Other.DelinquentAsyncLoads)),
+		MaxConcurrentAsyncLoads(Other.MaxConcurrentAsyncLoads)
+	{
+	}
+
+	FNetAsyncLoadDelinquencyAnalytics(const FNetAsyncLoadDelinquencyAnalytics&) = delete;
+	const FNetAsyncLoadDelinquencyAnalytics& operator=(const FNetAsyncLoadDelinquencyAnalytics&) = delete;
+	FNetAsyncLoadDelinquencyAnalytics& operator=(FNetAsyncLoadDelinquencyAnalytics&&) = default;
+
+	void CountBytes(FArchive& Ar) const
+	{
+		DelinquentAsyncLoads.CountBytes(Ar);
+	}
+
+	void Reset()
+	{
+		DelinquentAsyncLoads.Reset();
+		MaxConcurrentAsyncLoads = 0;
+	}
+
+	FDelinquencyAnalytics DelinquentAsyncLoads;
+	uint32 MaxConcurrentAsyncLoads;
+};
+
+/**
+ * Tracks data related specific to a NetConnection that can can result in prolonged periods of apparent
+ * network inactivity, despite actually receiving traffic.
+ *
+ * Also @see FDriverDelinquencyAnalytics and FDelinquencyAnalytics.
+ */
+struct ENGINE_API FNetQueuedActorDelinquencyAnalytics
+{
+	FNetQueuedActorDelinquencyAnalytics() :
+		DelinquentQueuedActors(0),
+		MaxConcurrentQueuedActors(0)
+	{
+	}
+
+	FNetQueuedActorDelinquencyAnalytics(const uint32 NumberOfTopOffendersToTrack) :
+		DelinquentQueuedActors(NumberOfTopOffendersToTrack),
+		MaxConcurrentQueuedActors(0)
+	{
+	}
+
+	FNetQueuedActorDelinquencyAnalytics(FNetQueuedActorDelinquencyAnalytics&& Other) :
+		DelinquentQueuedActors(MoveTemp(Other.DelinquentQueuedActors)),
+		MaxConcurrentQueuedActors(Other.MaxConcurrentQueuedActors)
+	{
+	}
+
+	FNetQueuedActorDelinquencyAnalytics(const FNetQueuedActorDelinquencyAnalytics&) = delete;
+	const FNetQueuedActorDelinquencyAnalytics& operator=(const FNetQueuedActorDelinquencyAnalytics&) = delete;
+	FNetQueuedActorDelinquencyAnalytics& operator=(FNetQueuedActorDelinquencyAnalytics&&) = default;
+
+
+	void CountBytes(FArchive& Ar) const
+	{
+		DelinquentQueuedActors.CountBytes(Ar);
+	}
+
+	void Reset()
+	{
+		DelinquentQueuedActors.Reset();
+		MaxConcurrentQueuedActors = 0;
+	}
+
+	FDelinquencyAnalytics DelinquentQueuedActors;
+	uint32 MaxConcurrentQueuedActors;
+};
\ No newline at end of file
