Fix memory leak.

diff --git a/Engine/Source/Runtime/Engine/Private/AudioDevice.cpp b/Engine/Source/Runtime/Engine/Private/AudioDevice.cpp
--- a/Engine/Source/Runtime/Engine/Private/AudioDevice.cpp
+++ b/Engine/Source/Runtime/Engine/Private/AudioDevice.cpp
@@ -4303,9 +4303,9 @@ void FAudioDevice::RetriggerVirtualLoop(FAudioVirtualLoop& VirtualLoopToRetrigge
 
 void FAudioDevice::AddVirtualLoop(const FAudioVirtualLoop& InVirtualLoop)
 {
-	FAudioVirtualLoop* VirtualLoop = new FAudioVirtualLoop(InVirtualLoop);
+	FAudioVirtualLoop VirtualLoop = InVirtualLoop;
 
-	FActiveSound& ActiveSound = VirtualLoop->GetActiveSound();
+	FActiveSound& ActiveSound = VirtualLoop.GetActiveSound();
 	check(!VirtualLoops.Contains(&ActiveSound));
 
 	const int64 ComponentID = ActiveSound.GetAudioComponentID();
@@ -4330,7 +4330,7 @@ void FAudioDevice::AddVirtualLoop(const FAudioVirtualLoop& InVirtualLoop)
 	}
 #endif // !(UE_BUILD_SHIPPING || UE_BUILD_TEST)
 
-	VirtualLoops.Add(&ActiveSound, InVirtualLoop);
+	VirtualLoops.Add(&ActiveSound, MoveTemp(VirtualLoop));
 }
 
 bool FAudioDevice::RemoveVirtualLoop(FActiveSound& InActiveSound)
