Fix for re-playing oneshots in cues that play loops and register as loops if the loop has come back from virtualizing.

diff --git a/Engine/Source/Runtime/Engine/Private/ActiveSound.cpp b/Engine/Source/Runtime/Engine/Private/ActiveSound.cpp
--- a/Engine/Source/Runtime/Engine/Private/ActiveSound.cpp
+++ b/Engine/Source/Runtime/Engine/Private/ActiveSound.cpp
@@ -47,6 +47,7 @@ FActiveSound::FActiveSound()
 	, bFinished(false)
 	, bIsPaused(false)
 	, bShouldStopDueToMaxConcurrency(false)
+	, bHasVirtualized(false)
 	, bRadioFilterSelected(false)
 	, bApplyRadioFilter(false)
 	, bHandleSubtitles(true)
diff --git a/Engine/Source/Runtime/Engine/Private/AudioVirtualLoop.cpp b/Engine/Source/Runtime/Engine/Private/AudioVirtualLoop.cpp
--- a/Engine/Source/Runtime/Engine/Private/AudioVirtualLoop.cpp
+++ b/Engine/Source/Runtime/Engine/Private/AudioVirtualLoop.cpp
@@ -76,6 +76,7 @@ bool FAudioVirtualLoop::Virtualize(const FActiveSound& InActiveSound, FAudioDevi
 	ActiveSound->ConcurrencyGroupVolumeScales.Reset();
 	ActiveSound->VolumeConcurrency = 1.0f;
 	ActiveSound->WaveInstances.Reset();
+	ActiveSound->bHasVirtualized = true;
 
 	OutVirtualLoop.ActiveSound = ActiveSound;
 	return true;
diff --git a/Engine/Source/Runtime/Engine/Private/SoundNodeLooping.cpp b/Engine/Source/Runtime/Engine/Private/SoundNodeLooping.cpp
--- a/Engine/Source/Runtime/Engine/Private/SoundNodeLooping.cpp
+++ b/Engine/Source/Runtime/Engine/Private/SoundNodeLooping.cpp
@@ -37,9 +37,15 @@ void USoundNodeLooping::ParseNodes( FAudioDevice* AudioDevice, const UPTRINT Nod
 	}
 #endif
 
+	// Do not continue parsing if virtualized and this is functioning as a one-shot
+	if (!bLoopIndefinitely && ActiveSound.bHasVirtualized)
+	{
+		return;
+	}
+
 	FSoundParseParameters UpdatedParams = ParseParams;
 	UpdatedParams.NotifyBufferFinishedHooks.AddNotify(this, NodeWaveInstanceHash);
-
+	UpdatedParams.bEnableRetrigger = true;
 	Super::ParseNodes( AudioDevice, NodeWaveInstanceHash, ActiveSound, UpdatedParams, WaveInstances );
 
 	if (ActiveSound.bFinished)
diff --git a/Engine/Source/Runtime/Engine/Private/SoundNodeWavePlayer.cpp b/Engine/Source/Runtime/Engine/Private/SoundNodeWavePlayer.cpp
--- a/Engine/Source/Runtime/Engine/Private/SoundNodeWavePlayer.cpp
+++ b/Engine/Source/Runtime/Engine/Private/SoundNodeWavePlayer.cpp
@@ -133,7 +133,9 @@ void USoundNodeWavePlayer::ParseNodes( FAudioDevice* AudioDevice, const UPTRINT
 			UpdatedParams.bLooping = true;
 			SoundWave->Parse(AudioDevice, NodeWaveInstanceHash, ActiveSound, UpdatedParams, WaveInstances);
 		}
-		else
+		// Don't play non-looping sounds again if this sound has been revived to
+		// avoid re-triggering one shots played adjacent with looping wave instances in cue
+		else if (!ActiveSound.bHasVirtualized || ParseParams.bEnableRetrigger)
 		{
 			SoundWave->Parse(AudioDevice, NodeWaveInstanceHash, ActiveSound, ParseParams, WaveInstances);
 		}
diff --git a/Engine/Source/Runtime/Engine/Public/ActiveSound.h b/Engine/Source/Runtime/Engine/Public/ActiveSound.h
--- a/Engine/Source/Runtime/Engine/Public/ActiveSound.h
+++ b/Engine/Source/Runtime/Engine/Public/ActiveSound.h
@@ -135,26 +135,29 @@ struct FSoundParseParameters
 	// The lowpass filter to apply if the sound is inside an ambient zone
 	float AmbientZoneFilterFrequency;
 
-	// Whether or not to ouput this audio to buses only
-	uint32 bOutputToBusOnly:1;
+	// Whether or not to output this audio to buses only
+	uint8 bOutputToBusOnly:1;
 
 	// Whether the sound should be spatialized
-	uint32 bUseSpatialization:1;
+	uint8 bUseSpatialization:1;
 
 	// Whether the sound should be seamlessly looped
-	uint32 bLooping:1;
+	uint8 bLooping:1;
 	
 	// Whether we have enabled low-pass filtering of this sound
-	uint32 bEnableLowPassFilter:1;
+	uint8 bEnableLowPassFilter:1;
 
 	// Whether this sound is occluded
-	uint32 bIsOccluded:1;
+	uint8 bIsOccluded:1;
 
 	// Whether or not this sound is manually paused (i.e. not by application-wide pause)
-	uint32 bIsPaused:1;
+	uint8 bIsPaused:1;
+
+	// Whether or not this sound can re-trigger
+	uint8 bEnableRetrigger : 1;
 
 	// Whether or not to apply a =6 dB attenuation to stereo spatialization sounds
-	uint32 bApplyNormalizationToStereoSounds:1;
+	uint8 bApplyNormalizationToStereoSounds:1;
 
 	FSoundParseParameters()
 		: SoundClass(nullptr)
@@ -195,6 +198,7 @@ struct FSoundParseParameters
 		, bEnableLowPassFilter(false)
 		, bIsOccluded(false)
 		, bIsPaused(false)
+		, bEnableRetrigger(false)
 		, bApplyNormalizationToStereoSounds(false)
 	{
 	}
@@ -310,6 +314,9 @@ public:
 	/** Whether or not to stop this active sound due to max concurrency */
 	uint8 bShouldStopDueToMaxConcurrency:1;
 
+	/** Whether or not sound has been virtualized and then realized */
+	uint8 bHasVirtualized:1;
+
 	/** If true, the decision on whether to apply the radio filter has been made. */
 	uint8 bRadioFilterSelected:1;
 
