Fix issue with spam of hide bone

diff --git a/Engine/Source/Runtime/Engine/Private/Components/SkeletalMeshComponent.cpp b/Engine/Source/Runtime/Engine/Private/Components/SkeletalMeshComponent.cpp
--- a/Engine/Source/Runtime/Engine/Private/Components/SkeletalMeshComponent.cpp
+++ b/Engine/Source/Runtime/Engine/Private/Components/SkeletalMeshComponent.cpp
@@ -2742,6 +2742,11 @@ void USkeletalMeshComponent::HideBone( int32 BoneIndex, EPhysBodyOp PhysBodyOpti
 		return;
 	}
 
+	if (MasterPoseComponent.IsValid())
+	{
+		return;
+	}
+
 	if (BoneSpaceTransforms.IsValidIndex(BoneIndex))
 	{
 		BoneSpaceTransforms[BoneIndex].SetScale3D(FVector::ZeroVector);
@@ -2758,7 +2763,7 @@ void USkeletalMeshComponent::HideBone( int32 BoneIndex, EPhysBodyOp PhysBodyOpti
 	}
 	else
 	{
-		UE_LOG(LogSkeletalMesh, Warning, TEXT("HideBone: Invalid Body Index has entered. This component doesn't contain buffer for the given body."));
+		UE_LOG(LogSkeletalMesh, Warning, TEXT("HideBone[%s]: Invalid Body Index (%d) has entered. This component doesn't contain buffer for the given body."), *GetNameSafe(SkeletalMesh), BoneIndex);
 	}
 }
 
