Fix missing copyright boilerplate for DelinquencyAnalytics.h

#rnx
#trivial

diff --git a/Engine/Source/Runtime/Engine/Public/Net/DelinquencyAnalytics.h b/Engine/Source/Runtime/Engine/Public/Net/DelinquencyAnalytics.h
--- a/Engine/Source/Runtime/Engine/Public/Net/DelinquencyAnalytics.h
+++ b/Engine/Source/Runtime/Engine/Public/Net/DelinquencyAnalytics.h
@@ -1,3 +1,5 @@
+// Copyright 1998-2019 Epic Games, Inc. All Rights Reserved.
+
 #pragma once
 
 #include "CoreMinimal.h"
