Fix for virtualize when silent and loop virtualization

-- loop virtualization wasn't taking into account "play when silent" option of sounds. this causes emotes (and other sounds) to restart when they get in range.

diff --git a/Engine/Source/Runtime/Engine/Private/AudioVirtualLoop.cpp b/Engine/Source/Runtime/Engine/Private/AudioVirtualLoop.cpp
--- a/Engine/Source/Runtime/Engine/Private/AudioVirtualLoop.cpp
+++ b/Engine/Source/Runtime/Engine/Private/AudioVirtualLoop.cpp
@@ -124,7 +124,7 @@ bool FAudioVirtualLoop::IsEnabled()
 
 bool FAudioVirtualLoop::IsInAudibleRange(const FActiveSound& InActiveSound, const FAudioDevice* InAudioDevice)
 {
-	if (!InActiveSound.bAllowSpatialization)
+	if (!InActiveSound.bAllowSpatialization || (InAudioDevice->VirtualSoundsEnabled() && InActiveSound.IsVirtualizeWhenSilent()))
 	{
 		return true;
 	}
diff --git a/Engine/Source/Runtime/Engine/Public/ActiveSound.h b/Engine/Source/Runtime/Engine/Public/ActiveSound.h
--- a/Engine/Source/Runtime/Engine/Public/ActiveSound.h
+++ b/Engine/Source/Runtime/Engine/Public/ActiveSound.h
@@ -265,6 +265,9 @@ public:
 	/** Whether or not the active sound is currently playing audible sound. */
 	bool IsPlayingAudio() const { return bIsPlayingAudio; }
 
+	/** Whether or not we're set to virtualize when silent. */
+	bool IsVirtualizeWhenSilent() const { return Sound && Sound->IsVirtualizeWhenSilent();  }
+
 	FAudioDevice* AudioDevice;
 
 	/** The concurrent groups that this sound is actively playing in. */
