Fix for null AudioDevice crash introduced in last regression fix

diff --git a/Engine/Source/Runtime/Engine/Private/AudioVirtualLoop.cpp b/Engine/Source/Runtime/Engine/Private/AudioVirtualLoop.cpp
--- a/Engine/Source/Runtime/Engine/Private/AudioVirtualLoop.cpp
+++ b/Engine/Source/Runtime/Engine/Private/AudioVirtualLoop.cpp
@@ -124,7 +124,7 @@ bool FAudioVirtualLoop::IsEnabled()
 
 bool FAudioVirtualLoop::IsInAudibleRange(const FActiveSound& InActiveSound, const FAudioDevice* InAudioDevice)
 {
-	if (!InActiveSound.bAllowSpatialization || (InAudioDevice->VirtualSoundsEnabled() && InActiveSound.IsVirtualizeWhenSilent()))
+	if (!InActiveSound.bAllowSpatialization)
 	{
 		return true;
 	}
@@ -136,6 +136,11 @@ bool FAudioVirtualLoop::IsInAudibleRange(const FActiveSound& InActiveSound, cons
 	}
 	check(AudioDevice);
 
+	if (AudioDevice->VirtualSoundsEnabled() && InActiveSound.IsVirtualizeWhenSilent())
+	{
+		return true;
+	}
+
 	float DistanceScale = 1.0f;
 	if (InActiveSound.bHasAttenuationSettings)
 	{
