Fix implicit cast

diff --git a/Engine/Source/Runtime/Engine/Private/AudioVirtualLoop.cpp b/Engine/Source/Runtime/Engine/Private/AudioVirtualLoop.cpp
--- a/Engine/Source/Runtime/Engine/Private/AudioVirtualLoop.cpp
+++ b/Engine/Source/Runtime/Engine/Private/AudioVirtualLoop.cpp
@@ -119,7 +119,7 @@ const FActiveSound& FAudioVirtualLoop::GetActiveSound() const
 
 bool FAudioVirtualLoop::IsEnabled()
 {
-	return bVirtualLoopsEnabledCVar;
+	return bVirtualLoopsEnabledCVar != 0;
 }
 
 bool FAudioVirtualLoop::IsInAudibleRange(const FActiveSound& InActiveSound, const FAudioDevice* InAudioDevice)
