Merging CL#6034440

Fix an edge case where displaced TopOffenders could cause ensures in delinquency analytics tracking.

diff --git a/Engine/Source/Runtime/Engine/Public/Net/DelinquencyAnalytics.h b/Engine/Source/Runtime/Engine/Public/Net/DelinquencyAnalytics.h
--- a/Engine/Source/Runtime/Engine/Public/Net/DelinquencyAnalytics.h
+++ b/Engine/Source/Runtime/Engine/Public/Net/DelinquencyAnalytics.h
@@ -147,10 +147,14 @@ public:
 						{
 							// Sanity check that we have actually found our entry.
 							// If we hit the end of the list, or we see an entry that should be less offensive, then
-							// we've missed our entry. This should never happen.
+							// we've missed our entry.
 							if (MaybeOurEntry == TopOffenders.Num() || TopOffenders[MaybeOurEntry].TimeSeconds > AlreadyTracked->TimeSeconds)
 							{
-								ensureMsgf(false, TEXT("FDelinquencyAnalytics::Add - Unable to find expected entry %s:%f, list may not be sorted!"), *AlreadyTracked->Name.ToString(), AlreadyTracked->TimeSeconds);
+								// It's possible that multiple entries have the same delinquency time.
+								// If our current entry matches the LeastOffensiveTime, it's possible that we were in the
+								// top offenders list at one point but were pushed out when more offensive entries were added.
+								// If our time doesn't match the least offensive time, we should definitely be in the list.
+								ensureMsgf(AlreadyTracked->TimeSeconds != LeastOffensiveTime, TEXT("FDelinquencyAnalytics::Add - Unable to find expected entry %s:%f, list may not be sorted!"), *AlreadyTracked->Name.ToString(), AlreadyTracked->TimeSeconds);
 								break;
 							}
 
