Merging CL#6036788

#UE4 Fix ensure in FDelinquencyAnalytics::Add

diff --git a/Engine/Source/Runtime/Engine/Public/Net/DelinquencyAnalytics.h b/Engine/Source/Runtime/Engine/Public/Net/DelinquencyAnalytics.h
--- a/Engine/Source/Runtime/Engine/Public/Net/DelinquencyAnalytics.h
+++ b/Engine/Source/Runtime/Engine/Public/Net/DelinquencyAnalytics.h
@@ -154,7 +154,7 @@ public:
 								// If our current entry matches the LeastOffensiveTime, it's possible that we were in the
 								// top offenders list at one point but were pushed out when more offensive entries were added.
 								// If our time doesn't match the least offensive time, we should definitely be in the list.
-								ensureMsgf(AlreadyTracked->TimeSeconds != LeastOffensiveTime, TEXT("FDelinquencyAnalytics::Add - Unable to find expected entry %s:%f, list may not be sorted!"), *AlreadyTracked->Name.ToString(), AlreadyTracked->TimeSeconds);
+								ensureMsgf(AlreadyTracked->TimeSeconds == LeastOffensiveTime, TEXT("FDelinquencyAnalytics::Add - Unable to find expected entry %s:%f, list may not be sorted!"), *AlreadyTracked->Name.ToString(), AlreadyTracked->TimeSeconds);
 								break;
 							}
 
