Copy CL 5822147 but pull RecacheUniformExpressions out of the if statement since MICs still recache their uniform expressions even if they don't have static permutation resources. Also recreates UBs when UMaterial::CacheResourceShadersForRendering is called. This implies that MDCs need to be recached after changing quality levels.

diff --git a/Engine/Source/Runtime/Engine/Private/Materials/Material.cpp b/Engine/Source/Runtime/Engine/Private/Materials/Material.cpp
--- a/Engine/Source/Runtime/Engine/Private/Materials/Material.cpp
+++ b/Engine/Source/Runtime/Engine/Private/Materials/Material.cpp
@@ -3121,9 +3121,7 @@ void UMaterial::CacheResourceShadersForRendering(bool bRegenerateId)
 			}
 		}
 
-		// Material reload can call this function again. Don't recreate since mesh
-		// draw commands cache a pointer to the uniform buffer
-		RecacheUniformExpressions(!FPlatformProperties::RequiresCookedData());
+		RecacheUniformExpressions(true);
 	}
 
 	if (ResourcesToFree.Num())
