Fix preset not getting updated in real-time while auditioning introduced when fixing a null deref crash a while back

diff --git a/Engine/Source/Runtime/Engine/Private/SoundEffectBase.cpp b/Engine/Source/Runtime/Engine/Private/SoundEffectBase.cpp
--- a/Engine/Source/Runtime/Engine/Private/SoundEffectBase.cpp
+++ b/Engine/Source/Runtime/Engine/Private/SoundEffectBase.cpp
@@ -35,9 +35,12 @@ void FSoundEffectBase::SetPreset(USoundEffectPreset* Inpreset)
 		if (Preset)
 		{
 			Preset->AddEffectInstance(this);
-			bChanged = true;
 		}
 	}
+
+	// Anytime notification occurs that the preset has been modified,
+	// flag for update.
+	bChanged = true;
 }
 
 void FSoundEffectBase::ClearPreset()
