[NO NON-PII COMMIT MESSAGE]

diff --git a/Engine/Source/Runtime/AudioMixer/Private/AudioMixerDevice.cpp b/Engine/Source/Runtime/AudioMixer/Private/AudioMixerDevice.cpp
--- a/Engine/Source/Runtime/AudioMixer/Private/AudioMixerDevice.cpp
+++ b/Engine/Source/Runtime/AudioMixer/Private/AudioMixerDevice.cpp
@@ -511,14 +511,6 @@ namespace Audio
 	{
 		LLM_SCOPE(ELLMTag::AudioMixer);
 
-#if WITH_EDITOR
-		// Turn on to only hear PIE audio
-		bool bBypassMainAudioDevice = FParse::Param(FCommandLine::Get(), TEXT("AudioPIEOnly"));
-		if (bBypassMainAudioDevice && IsMainAudioDevice())
-		{
-			return true;
-		}
-#endif
 		// This function could be called in a task manager, which means the thread ID may change between calls.
 		ResetAudioRenderingThreadId();
 
diff --git a/Engine/Source/Runtime/Engine/Private/AudioDeviceManager.cpp b/Engine/Source/Runtime/Engine/Private/AudioDeviceManager.cpp
--- a/Engine/Source/Runtime/Engine/Private/AudioDeviceManager.cpp
+++ b/Engine/Source/Runtime/Engine/Private/AudioDeviceManager.cpp
@@ -376,69 +376,93 @@ bool FAudioDeviceManager::CreateAudioDevice(bool bCreateNewDevice, FCreateAudioD
 		}
 	}
 
-	if (NumActiveAudioDevices < AUDIO_DEVICE_DEFAULT_ALLOWED_DEVICE_COUNT || (bCreateNewDevice && NumActiveAudioDevices < AUDIO_DEVICE_MAX_DEVICE_COUNT))
+	bool bRequiresInit = true;
+
+	// For the first PIE window, we'll just use the main audio device
+	if (NumActiveAudioDevices == 1)
 	{
-		// Create the new audio device and make sure it succeeded
-		OutResults.AudioDevice = AudioDeviceModule->CreateAudioDevice();
-		if (OutResults.AudioDevice == nullptr)
+		FAudioDevice* MainAudioDevice = GEngine->GetMainAudioDevice();
+		if (MainAudioDevice)
+		{
+			++NumWorldsUsingMainAudioDevice;
+			OutResults.Handle = MainAudioDevice->DeviceHandle;
+			OutResults.AudioDevice = MainAudioDevice;
+			bRequiresInit = false;
+		}
+		else
 		{
 			return false;
 		}
+	}
+	else
+	{
+		if (NumActiveAudioDevices < AUDIO_DEVICE_DEFAULT_ALLOWED_DEVICE_COUNT || (bCreateNewDevice && NumActiveAudioDevices < AUDIO_DEVICE_MAX_DEVICE_COUNT))
+		{
+			// Create the new audio device and make sure it succeeded
+			OutResults.AudioDevice = AudioDeviceModule->CreateAudioDevice();
+			if (OutResults.AudioDevice == nullptr)
+			{
+				return false;
+			}
 
-		// Now generation a new audio device handle for the device and store the
-		// ptr to the new device in the array of audio devices.
+			// Now generation a new audio device handle for the device and store the
+			// ptr to the new device in the array of audio devices.
 
-		uint32 AudioDeviceIndex(INDEX_NONE);
+			uint32 AudioDeviceIndex(INDEX_NONE);
 
-		// First check to see if we should start recycling audio device indices, if not
-		// then we add a new entry to the Generation array and generate a new index
-		if (FreeIndicesSize > AUDIO_DEVICE_MINIMUM_FREE_AUDIO_DEVICE_INDICES)
-		{
-			FreeIndices.Dequeue(AudioDeviceIndex);
-			--FreeIndicesSize;
-			check(int32(AudioDeviceIndex) < Devices.Num());
-			check(Devices[AudioDeviceIndex] == nullptr);
-			Devices[AudioDeviceIndex] = OutResults.AudioDevice;
-		}
-		else
-		{
-			// Add a zeroth generation entry in the Generation array, get a brand new
-			// index and append the created device to the end of the Devices array
+			// First check to see if we should start recycling audio device indices, if not
+			// then we add a new entry to the Generation array and generate a new index
+			if (FreeIndicesSize > AUDIO_DEVICE_MINIMUM_FREE_AUDIO_DEVICE_INDICES)
+			{
+				FreeIndices.Dequeue(AudioDeviceIndex);
+				--FreeIndicesSize;
+				check(int32(AudioDeviceIndex) < Devices.Num());
+				check(Devices[AudioDeviceIndex] == nullptr);
+				Devices[AudioDeviceIndex] = OutResults.AudioDevice;
+			}
+			else
+			{
+				// Add a zeroth generation entry in the Generation array, get a brand new
+				// index and append the created device to the end of the Devices array
 
-			Generations.Add(0);
-			AudioDeviceIndex = Generations.Num() - 1;
-			check(AudioDeviceIndex < (1 << AUDIO_DEVICE_HANDLE_INDEX_BITS));
-			Devices.Add(OutResults.AudioDevice);
-		}
+				Generations.Add(0);
+				AudioDeviceIndex = Generations.Num() - 1;
+				check(AudioDeviceIndex < (1 << AUDIO_DEVICE_HANDLE_INDEX_BITS));
+				Devices.Add(OutResults.AudioDevice);
+			}
 
-		OutResults.bNewDevice = true;
-		OutResults.Handle = CreateHandle(AudioDeviceIndex, Generations[AudioDeviceIndex]);
+			OutResults.bNewDevice = true;
+			OutResults.Handle = CreateHandle(AudioDeviceIndex, Generations[AudioDeviceIndex]);
 
-		// Store the handle on the audio device itself
-		OutResults.AudioDevice->DeviceHandle = OutResults.Handle;
-	}
-	else
-	{
-		++NumWorldsUsingMainAudioDevice;
-		FAudioDevice* MainAudioDevice = GEngine->GetMainAudioDevice();
-		if (MainAudioDevice)
+			// Store the handle on the audio device itself
+			OutResults.AudioDevice->DeviceHandle = OutResults.Handle;
+		}
+		else
 		{
-			OutResults.Handle = MainAudioDevice->DeviceHandle;
-			OutResults.AudioDevice = MainAudioDevice;
+			++NumWorldsUsingMainAudioDevice;
+			FAudioDevice* MainAudioDevice = GEngine->GetMainAudioDevice();
+			if (MainAudioDevice)
+			{
+				OutResults.Handle = MainAudioDevice->DeviceHandle;
+				OutResults.AudioDevice = MainAudioDevice;
+			}
 		}
 	}
 
 	++NumActiveAudioDevices;
 
-	const UAudioSettings* AudioSettings = GetDefault<UAudioSettings>();
-	if (OutResults.AudioDevice->Init(AudioSettings->GetHighestMaxChannels())) //-V595
-	{
-		OutResults.AudioDevice->SetMaxChannels(AudioSettings->GetQualityLevelSettings(GEngine->GetGameUserSettings()->GetAudioQualityLevel()).MaxChannels); //-V595
-	}
-	else
+	if (bRequiresInit)
 	{
-		ShutdownAudioDevice(OutResults.Handle);
-		OutResults = FCreateAudioDeviceResults();
+		const UAudioSettings* AudioSettings = GetDefault<UAudioSettings>();
+		if (OutResults.AudioDevice->Init(AudioSettings->GetHighestMaxChannels())) //-V595
+		{
+			OutResults.AudioDevice->SetMaxChannels(AudioSettings->GetQualityLevelSettings(GEngine->GetGameUserSettings()->GetAudioQualityLevel()).MaxChannels); //-V595
+		}
+		else
+		{
+			ShutdownAudioDevice(OutResults.Handle);
+			OutResults = FCreateAudioDeviceResults();
+		}
 	}
 
 	// We need to call fade in, in case we're reusing audio devices
