Issue with drifting during replay paused

diff --git a/Engine/Source/Runtime/Engine/Private/Animation/AnimInterpFilter.cpp b/Engine/Source/Runtime/Engine/Private/Animation/AnimInterpFilter.cpp
--- a/Engine/Source/Runtime/Engine/Private/Animation/AnimInterpFilter.cpp
+++ b/Engine/Source/Runtime/Engine/Private/Animation/AnimInterpFilter.cpp
@@ -144,30 +144,34 @@ void FFIRFilterTimeBased::RefreshValidFilters()
 
 float FFIRFilterTimeBased::GetFilteredData(float Input, float DeltaTime)
 {
-	float Result;
-
-	CurrentTime += DeltaTime;
-
-	if ( IsValid() )
+	if (DeltaTime > KINDA_SMALL_NUMBER)
 	{
-		RefreshValidFilters();
+		float Result;
+		CurrentTime += DeltaTime;
 
-		CurrentStackIndex = GetSafeCurrentStackIndex();
-		FilterWindow[CurrentStackIndex].SetInput(Input, CurrentTime);
-		Result = CalculateFilteredOutput();
-		CurrentStackIndex = CurrentStackIndex+1;
-		if ( CurrentStackIndex > FilterWindow.Num()-1 ) 
+		if (IsValid())
 		{
-			CurrentStackIndex = 0;
+			RefreshValidFilters();
+
+			CurrentStackIndex = GetSafeCurrentStackIndex();
+			FilterWindow[CurrentStackIndex].SetInput(Input, CurrentTime);
+			Result = CalculateFilteredOutput();
+			CurrentStackIndex = CurrentStackIndex + 1;
+			if (CurrentStackIndex > FilterWindow.Num() - 1)
+			{
+				CurrentStackIndex = 0;
+			}
 		}
-	}
-	else
-	{
-		Result = Input;
+		else
+		{
+			Result = Input;
+		}
+
+		LastOutput = Result;
+		return Result;
 	}
 
-	LastOutput = Result;
-	return Result;
+	return LastOutput;
 }
 
 float FFIRFilterTimeBased::GetInterpolationCoefficient (FFilterData & Data)
