Downgrade warning to to verbose log.  Constantly get bugs for this which aren't really bugs as it is almost always called on clean-up/respawn.

diff --git a/Engine/Source/Runtime/Engine/Private/Components/AudioComponent.cpp b/Engine/Source/Runtime/Engine/Private/Components/AudioComponent.cpp
--- a/Engine/Source/Runtime/Engine/Private/Components/AudioComponent.cpp
+++ b/Engine/Source/Runtime/Engine/Private/Components/AudioComponent.cpp
@@ -117,7 +117,7 @@ void UAudioComponent::BeginDestroy()
 
 	if (bIsActive && Sound && Sound->IsLooping())
 	{
-		UE_LOG(LogAudio, Warning, TEXT("Audio Component is being destroyed without stopping looping sound '%s'"), *Sound->GetFullName());
+		UE_LOG(LogAudio, Verbose, TEXT("Audio Component is being destroyed prior to stopping looping sound '%s' directly."), *Sound->GetFullName());
 		Stop();
 	}
 
