Add more logging around loop virtualization for debugging purposes

diff --git a/Engine/Source/Runtime/Engine/Private/AudioDevice.cpp b/Engine/Source/Runtime/Engine/Private/AudioDevice.cpp
--- a/Engine/Source/Runtime/Engine/Private/AudioDevice.cpp
+++ b/Engine/Source/Runtime/Engine/Private/AudioDevice.cpp
@@ -2073,6 +2073,10 @@ void FAudioDevice::StopQuietSoundsDueToMaxConcurrency(TArray<FWaveInstance*>& Wa
 
 					const bool bIsAtMaxConcurrency = true;
 					VirtualLoop.CalculateUpdateInterval(bIsAtMaxConcurrency);
+					if (USoundBase* Sound = ActiveSound->GetSound())
+					{
+						UE_LOG(LogAudio, Verbose, TEXT("Playing ActiveSound %s Virtualizing: Invalidated by 'StopQuiestest' Concurrency Rule."), *Sound->GetName());
+					}
 					AddVirtualLoop(VirtualLoop);
 				}
 				ActiveSoundsCopy.RemoveAtSwap(i, 1, false);
@@ -2618,6 +2622,10 @@ void FAudioDevice::VirtualizeInactiveLoops()
 
 				// Clear must be called after AddSoundToStop to ensure AudioComponent is properly removed from AudioComponentIDToActiveSoundMap
 				ActiveSound->ClearAudioComponent();
+				if (USoundBase* Sound = ActiveSound->GetSound())
+				{
+					UE_LOG(LogAudio, Verbose, TEXT("Playing ActiveSound %s Virtualizing: Out of audible range."), *Sound->GetName());
+				}
 				AddVirtualLoop(VirtualLoop);
 			}
 		}
@@ -4203,6 +4211,7 @@ void FAudioDevice::AddNewActiveSoundInternal(const FActiveSound& NewActiveSound,
 		FAudioVirtualLoop VirtualLoop;
 		if (FAudioVirtualLoop::Virtualize(NewActiveSound, *this, bDoRangeCheck, VirtualLoop))
 		{
+			UE_LOG(LogAudio, Verbose, TEXT("New ActiveSound %s Virtualizing: Failed to pass initial audible range check"), *Sound->GetName());
 			AddVirtualLoop(VirtualLoop);
 			return;
 		}
@@ -4227,6 +4236,7 @@ void FAudioDevice::AddNewActiveSoundInternal(const FActiveSound& NewActiveSound,
 			FAudioVirtualLoop VirtualLoop;
 			if (FAudioVirtualLoop::Virtualize(NewActiveSound, *this, bDoRangeCheck, VirtualLoop))
 			{
+				UE_LOG(LogAudio, Verbose, TEXT("New ActiveSound %s Virtualizing: Failed to pass concurrency"), *Sound->GetName());
 				AddVirtualLoop(VirtualLoop);
 			}
 		}
diff --git a/Engine/Source/Runtime/Engine/Private/SoundConcurrency.cpp b/Engine/Source/Runtime/Engine/Private/SoundConcurrency.cpp
--- a/Engine/Source/Runtime/Engine/Private/SoundConcurrency.cpp
+++ b/Engine/Source/Runtime/Engine/Private/SoundConcurrency.cpp
@@ -660,6 +660,10 @@ FActiveSound* FSoundConcurrencyManager::CreateAndEvictActiveSounds(const FActive
 			if (FAudioVirtualLoop::Virtualize(*SoundToEvict, bDoRangeCheck, VirtualLoop))
 			{
 				SoundToEvict->ClearAudioComponent();
+				if (USoundBase* Sound = SoundToEvict->GetSound())
+				{
+					UE_LOG(LogAudio, Verbose, TEXT("Playing ActiveSound %s Virtualizing: Evicted due to concurrency rules."), *Sound->GetName());
+				}
 				AudioDevice->AddVirtualLoop(VirtualLoop);
 			}
 		}
