Allowing streaming level fast forwards to not pause the rest of the stream while they wait for data (gated behind demo.FastForwardLevelsPausePlayback)
Use IsDataAvailableForTimeRange instead of IsDataAvailable in tick, to avoid fast forward checkpoint tasks blocking playback (also gated behind demo.FastForwardLevelsPausePlayback)

diff --git a/Engine/Source/Runtime/Engine/Classes/Engine/DemoNetDriver.h b/Engine/Source/Runtime/Engine/Classes/Engine/DemoNetDriver.h
--- a/Engine/Source/Runtime/Engine/Classes/Engine/DemoNetDriver.h
+++ b/Engine/Source/Runtime/Engine/Classes/Engine/DemoNetDriver.h
@@ -49,6 +49,7 @@ public:
 	virtual void StartTask() = 0;
 	virtual bool Tick() = 0;
 	virtual FName GetName() const = 0;
+	virtual bool ShouldPausePlayback() const { return true; }
 
 	TWeakObjectPtr<UDemoNetDriver> Driver;
 };
diff --git a/Engine/Source/Runtime/Engine/Private/DemoNetDriver.cpp b/Engine/Source/Runtime/Engine/Private/DemoNetDriver.cpp
--- a/Engine/Source/Runtime/Engine/Private/DemoNetDriver.cpp
+++ b/Engine/Source/Runtime/Engine/Private/DemoNetDriver.cpp
@@ -88,6 +88,8 @@ static TAutoConsoleVariable<float> CVarDemoDecreaseRepPrioritizeThreshold(TEXT("
 static TAutoConsoleVariable<float> CVarDemoMinimumRepPrioritizeTime(TEXT("demo.MinimumRepPrioritizePercent"), 0.3, TEXT("Minimum percent of time that must be spent prioritizing actors, regardless of throttling."));
 static TAutoConsoleVariable<float> CVarDemoMaximumRepPrioritizeTime(TEXT("demo.MaximumRepPrioritizePercent"), 0.8, TEXT("Maximum percent of time that may be spent prioritizing actors, regardless of throttling."));
 
+static TAutoConsoleVariable<int32> CVarFastForwardLevelsPausePlayback(TEXT("demo.FastForwardLevelsPausePlayback"), 0, TEXT("If true, pause channels and playback while fast forward levels task is running."));
+
 static const int32 MAX_DEMO_READ_WRITE_BUFFER = 1024 * 2;
 static const int32 MAX_DEMO_STRING_SERIALIZATION_SIZE = 16 * 1024 * 1024;
 
@@ -380,7 +382,7 @@ public:
 		Driver->ReplayStreamer->GotoTimeInMS(Driver->DemoCurrentTime * 1000, FGotoCallback::CreateSP(this, &FGotoTimeInSecondsTask::CheckpointReady), CheckpointType);
 
 		// Pause channels while we wait (so the world is paused while we wait for the new stream location to load)
-		Driver->PauseChannels( true );
+		Driver->PauseChannels(true);
 	}
 
 	virtual bool Tick() override
@@ -516,8 +518,11 @@ public:
 
 			Driver->ReplayStreamer->GotoTimeInMS(GotoTime, FGotoCallback::CreateSP(this, &FFastForwardLevelsTask::CheckpointReady), CheckpointType);
 
-			// Pause channels while we wait (so the world is paused while we wait for the new stream location to load)
-			Driver->PauseChannels(true);
+			if (CVarFastForwardLevelsPausePlayback.GetValueOnAnyThread() != 0)
+			{
+				// Pause channels while we wait (so the world is paused while we wait for the new stream location to load)
+				Driver->PauseChannels(true);
+			}
 		}
 	}
 
@@ -533,6 +538,9 @@ public:
 		}
 		else if (GotoResult.IsSet())
 		{
+			// if this task is not pausing the rest of the replay stream, make sure there is data available for the current time or we could miss packets
+			const uint32 AvailableDataEndTime = (CVarFastForwardLevelsPausePlayback.GetValueOnAnyThread() != 0) ? GotoTime : Driver->DemoCurrentTime * 1000;
+
 			if (!GotoResult->WasSuccessful())
 			{
 				return true;
@@ -541,7 +549,7 @@ public:
 			// If not all data is available, we could end only partially fast forwarding the levels.
 			// Note, IsDataAvailable may return false even if IsDataAvailableForTimeRange is true.
 			// So, check both to ensure that we don't end up skipping data in FastForwardLevels.
-			else if (GotoResult->ExtraTimeMS > 0 && !(Driver->ReplayStreamer->IsDataAvailable() && Driver->ReplayStreamer->IsDataAvailableForTimeRange(GotoTime - GotoResult->ExtraTimeMS, GotoTime)))
+			else if (GotoResult->ExtraTimeMS > 0 && !(Driver->ReplayStreamer->IsDataAvailable() && Driver->ReplayStreamer->IsDataAvailableForTimeRange(GotoTime - GotoResult->ExtraTimeMS, AvailableDataEndTime)))
 			{
 				return false;
 			}
@@ -557,6 +565,11 @@ public:
 		return ReplayTaskNames::FastForwardLevelsTask;
 	}
 
+	virtual bool ShouldPausePlayback() const override
+	{
+		return (CVarFastForwardLevelsPausePlayback.GetValueOnAnyThread() != 0);
+	}
+
 	void CheckpointReady(const FGotoResult& Result)
 	{
 		check(!GotoResult.IsSet());
@@ -691,7 +704,7 @@ bool UDemoNetDriver::ProcessReplayTasks()
 		if (!ActiveReplayTask->Tick())
 		{
 			// Task isn't done, we can return
-			return false;
+			return !ActiveReplayTask->ShouldPausePlayback();
 		}
 
 		// This task is now done
@@ -2989,6 +3002,8 @@ void UDemoNetDriver::PauseChannels(const bool bPause)
 
 	bChannelsArePaused = bPause;
 
+	UE_LOG(LogDemo, Verbose, TEXT("PauseChannels: %d"), bChannelsArePaused);
+
 	OnPauseChannelsDelegate.Broadcast(bChannelsArePaused);
 }
 
@@ -3471,7 +3486,7 @@ bool UDemoNetDriver::ConditionallyProcessPlaybackPackets()
 {
 	if (!PlaybackPackets.IsValidIndex(PlaybackPacketIndex))
 	{
-		PauseChannels( true );
+		PauseChannels(true);
 		return false;
 	}
 
@@ -3838,17 +3853,25 @@ void UDemoNetDriver::TickDemoPlayback(float DeltaSeconds)
 	// Clamp time
 	DemoCurrentTime = FMath::Clamp(DemoCurrentTime, 0.0f, DemoTotalTime + 0.01f);
 
+	bool bProcessAvailableData = (PlaybackPackets.Num() > 0) || ReplayStreamer->IsDataAvailable();
+	
+	if (CVarFastForwardLevelsPausePlayback.GetValueOnAnyThread() == 0)
+	{
+		const uint32 DemoCurrentTimeInMS = DemoCurrentTime * 1000;
+		bProcessAvailableData = bProcessAvailableData || ReplayStreamer->IsDataAvailableForTimeRange(DemoCurrentTimeInMS, DemoCurrentTimeInMS);
+	}
+
 	// Make sure there is data available to read
 	// If we're at the end of the demo, just pause channels and return
-	if (!PlaybackPackets.Num() && !ReplayStreamer->IsDataAvailable())
+	if (bProcessAvailableData)
 	{
-		PauseChannels(true);
-		return;
+		// we either have packets to process or data available to read
+		PauseChannels(false);
 	}
 	else
 	{
-		// we either have packets to process or data available to read
-		PauseChannels(false);
+		PauseChannels(true);
+		return;
 	}
 
 	// Speculatively grab seconds now in case we need it to get the time it took to fast forward
