Make ALODActor not relevant for calculating level bounds, contained actors are already contributing, and prevents non-center bound pivots impacting the level bound size.

diff --git a/Engine/Source/Runtime/Engine/Classes/Engine/LODActor.h b/Engine/Source/Runtime/Engine/Classes/Engine/LODActor.h
--- a/Engine/Source/Runtime/Engine/Classes/Engine/LODActor.h
+++ b/Engine/Source/Runtime/Engine/Classes/Engine/LODActor.h
@@ -70,6 +70,7 @@ public:
 	virtual FBox GetComponentsBoundingBox(bool bNonColliding = false) const override;
 	virtual void PostRegisterAllComponents() override;
 	virtual void Tick(float DeltaSeconds) override;	
+	virtual bool IsLevelBoundsRelevant() const override { return false; }
 	//~ End AActor Interface
 
 	/** Forces the mesh into view by setting the MinDrawDistance to zero (this pops the mesh into view, no fading)*/
@@ -168,6 +169,9 @@ public:
 
 	virtual void PreSave(const class ITargetPlatform* TargetPlatform) override;
 #endif // WITH_EDITOR	
+
+	
+
 	//~ End UObject Interface.	
 public:
 #if WITH_EDITORONLY_DATA
