Skip invalidating collision if no changes were applied when calling SetCollisionResponse functions.  Can be a significant gain on objects with bGenerateOverlapEvents.

#rnx

diff --git a/Engine/Source/Runtime/Engine/Classes/Engine/EngineTypes.h b/Engine/Source/Runtime/Engine/Classes/Engine/EngineTypes.h
--- a/Engine/Source/Runtime/Engine/Classes/Engine/EngineTypes.h
+++ b/Engine/Source/Runtime/Engine/Classes/Engine/EngineTypes.h
@@ -1042,13 +1042,13 @@ struct ENGINE_API FCollisionResponseContainer
 	FCollisionResponseContainer(ECollisionResponse DefaultResponse);
 
 	/** Set the response of a particular channel in the structure. */
-	void SetResponse(ECollisionChannel Channel, ECollisionResponse NewResponse);
+	bool SetResponse(ECollisionChannel Channel, ECollisionResponse NewResponse);
 
 	/** Set all channels to the specified response */
-	void SetAllChannels(ECollisionResponse NewResponse);
+	bool SetAllChannels(ECollisionResponse NewResponse);
 
 	/** Replace the channels matching the old response with the new response */
-	void ReplaceChannels(ECollisionResponse OldResponse, ECollisionResponse NewResponse);
+	bool ReplaceChannels(ECollisionResponse OldResponse, ECollisionResponse NewResponse);
 
 	/** Returns the response set on the specified channel */
 	FORCEINLINE_DEBUGGABLE ECollisionResponse GetResponse(ECollisionChannel Channel) const { return (ECollisionResponse)EnumArray[Channel]; }
@@ -1063,6 +1063,15 @@ struct ENGINE_API FCollisionResponseContainer
 	/** Returns the game-wide default collision response */
 	static const struct FCollisionResponseContainer& GetDefaultResponseContainer() { return DefaultResponseContainer; }
 
+	bool operator==(const FCollisionResponseContainer& Other) const
+	{
+		return FMemory::Memcmp(EnumArray, Other.EnumArray, sizeof(Other.EnumArray)) == 0;
+	}
+	bool operator!=(const FCollisionResponseContainer& Other) const
+	{
+		return FMemory::Memcmp(EnumArray, Other.EnumArray, sizeof(Other.EnumArray)) != 0;
+	}
+
 private:
 
 	/** static variable for default data to be used without reconstructing everytime **/
diff --git a/Engine/Source/Runtime/Engine/Classes/PhysicsEngine/BodyInstance.h b/Engine/Source/Runtime/Engine/Classes/PhysicsEngine/BodyInstance.h
--- a/Engine/Source/Runtime/Engine/Classes/PhysicsEngine/BodyInstance.h
+++ b/Engine/Source/Runtime/Engine/Classes/PhysicsEngine/BodyInstance.h
@@ -63,20 +63,20 @@ struct ENGINE_API FCollisionResponse
 	FCollisionResponse(ECollisionResponse DefaultResponse);
 
 	/** Set the response of a particular channel in the structure. */
-	void SetResponse(ECollisionChannel Channel, ECollisionResponse NewResponse);
+	bool SetResponse(ECollisionChannel Channel, ECollisionResponse NewResponse);
 
 	/** Set all channels to the specified response */
-	void SetAllChannels(ECollisionResponse NewResponse);
+	bool SetAllChannels(ECollisionResponse NewResponse);
 
 	/** Replace the channels matching the old response with the new response */
-	void ReplaceChannels(ECollisionResponse OldResponse, ECollisionResponse NewResponse);
+	bool ReplaceChannels(ECollisionResponse OldResponse, ECollisionResponse NewResponse);
 
 	/** Returns the response set on the specified channel */
 	FORCEINLINE_DEBUGGABLE ECollisionResponse GetResponse(ECollisionChannel Channel) const { return ResponseToChannels.GetResponse(Channel); }
 	const FCollisionResponseContainer& GetResponseContainer() const { return ResponseToChannels; }
 
 	/** Set all channels from ChannelResponse Array **/
-	void SetCollisionResponseContainer(const FCollisionResponseContainer& InResponseToChannels);
+	bool SetCollisionResponseContainer(const FCollisionResponseContainer& InResponseToChannels);
 	void SetResponsesArray(const TArray<FResponseChannel>& InChannelResponses);
 	void UpdateResponseContainerFromArray();
 
@@ -772,19 +772,19 @@ public:
 	void SetContactReportForceThreshold(float Threshold);
 
 	/** Set the collision response of this body to a particular channel */
-	void SetResponseToChannel(ECollisionChannel Channel, ECollisionResponse NewResponse);
+	bool SetResponseToChannel(ECollisionChannel Channel, ECollisionResponse NewResponse);
 
 	/** Get the collision response of this body to a particular channel */
 	FORCEINLINE_DEBUGGABLE ECollisionResponse GetResponseToChannel(ECollisionChannel Channel) const { return CollisionResponses.GetResponse(Channel); }
 
 	/** Set the response of this body to all channels */
-	void SetResponseToAllChannels(ECollisionResponse NewResponse);
+	bool SetResponseToAllChannels(ECollisionResponse NewResponse);
 
 	/** Replace the channels on this body matching the old response with the new response */
-	void ReplaceResponseToChannels(ECollisionResponse OldResponse, ECollisionResponse NewResponse);
+	bool ReplaceResponseToChannels(ECollisionResponse OldResponse, ECollisionResponse NewResponse);
 
 	/** Set the response of this body to the supplied settings */
-	void SetResponseToChannels(const FCollisionResponseContainer& NewReponses);
+	bool SetResponseToChannels(const FCollisionResponseContainer& NewReponses);
 
 	/** Get Collision ResponseToChannels container for this component **/
 	FORCEINLINE_DEBUGGABLE const FCollisionResponseContainer& GetResponseToChannels() const { return CollisionResponses.GetResponseContainer(); }
@@ -1111,4 +1111,4 @@ FORCEINLINE_DEBUGGABLE bool FBodyInstance::IsInstanceSimulatingPhysics() const
 }
 
 extern template ENGINE_API bool FBodyInstance::OverlapTestForBodiesImpl(const FVector& Position, const FQuat& Rotation, const TArray<FBodyInstance*>& Bodies) const;
-extern template ENGINE_API bool FBodyInstance::OverlapTestForBodiesImpl(const FVector& Position, const FQuat& Rotation, const TArray<FBodyInstance*, TInlineAllocator<1>>& Bodies) const;
\ No newline at end of file
+extern template ENGINE_API bool FBodyInstance::OverlapTestForBodiesImpl(const FVector& Position, const FQuat& Rotation, const TArray<FBodyInstance*, TInlineAllocator<1>>& Bodies) const;
diff --git a/Engine/Source/Runtime/Engine/Private/PhysicsEngine/BodyInstance.cpp b/Engine/Source/Runtime/Engine/Private/PhysicsEngine/BodyInstance.cpp
--- a/Engine/Source/Runtime/Engine/Private/PhysicsEngine/BodyInstance.cpp
+++ b/Engine/Source/Runtime/Engine/Private/PhysicsEngine/BodyInstance.cpp
@@ -104,7 +104,7 @@ FCollisionResponse::FCollisionResponse(ECollisionResponse DefaultResponse)
 }
 
 /** Set the response of a particular channel in the structure. */
-void FCollisionResponse::SetResponse(ECollisionChannel Channel, ECollisionResponse NewResponse)
+bool FCollisionResponse::SetResponse(ECollisionChannel Channel, ECollisionResponse NewResponse)
 {
 #if 1// @hack until PostLoad is disabled for CDO of BP WITH_EDITOR
 	ECollisionResponse DefaultResponse = FCollisionResponseContainer::GetDefaultResponseContainer().GetResponse(Channel);
@@ -118,34 +118,47 @@ void FCollisionResponse::SetResponse(ECollisionChannel Channel, ECollisionRespon
 	}
 #endif
 
-	ResponseToChannels.SetResponse(Channel, NewResponse);
+	return ResponseToChannels.SetResponse(Channel, NewResponse);
 }
 
 /** Set all channels to the specified response */
-void FCollisionResponse::SetAllChannels(ECollisionResponse NewResponse)
+bool FCollisionResponse::SetAllChannels(ECollisionResponse NewResponse)
 {
-	ResponseToChannels.SetAllChannels(NewResponse);
+	if (ResponseToChannels.SetAllChannels(NewResponse))
+	{
 #if 1// @hack until PostLoad is disabled for CDO of BP WITH_EDITOR
-	UpdateArrayFromResponseContainer();
+		UpdateArrayFromResponseContainer();
 #endif
+		return true;
+	}
+	return false;
 }
 
-void FCollisionResponse::ReplaceChannels(ECollisionResponse OldResponse, ECollisionResponse NewResponse)
+bool FCollisionResponse::ReplaceChannels(ECollisionResponse OldResponse, ECollisionResponse NewResponse)
 {
-	ResponseToChannels.ReplaceChannels(OldResponse, NewResponse);
+	if (ResponseToChannels.ReplaceChannels(OldResponse, NewResponse))
+	{
 #if 1// @hack until PostLoad is disabled for CDO of BP WITH_EDITOR
-	UpdateArrayFromResponseContainer();
+		UpdateArrayFromResponseContainer();
 #endif
+		return true;
+	}
+	return false;
 }
 
 /** Set all channels from ChannelResponse Array **/
-void FCollisionResponse::SetCollisionResponseContainer(const FCollisionResponseContainer& InResponseToChannels)
+bool FCollisionResponse::SetCollisionResponseContainer(const FCollisionResponseContainer& InResponseToChannels)
 {
-	ResponseToChannels = InResponseToChannels;
+	if (ResponseToChannels != InResponseToChannels)
+	{
+		ResponseToChannels = InResponseToChannels;
 #if 1// @hack until PostLoad is disabled for CDO of BP WITH_EDITOR
-	// this is only valid case that has to be updated
-	UpdateArrayFromResponseContainer();
+		// this is only valid case that has to be updated
+		UpdateArrayFromResponseContainer();
 #endif
+		return true;
+	}
+	return false;
 }
 
 void FCollisionResponse::SetResponsesArray(const TArray<FResponseChannel>& InChannelResponses)
@@ -460,32 +473,52 @@ void FBodyInstance::InvalidateCollisionProfileName()
 	bPendingCollisionProfileSetup = false;
 }
 
-void FBodyInstance::SetResponseToChannel(ECollisionChannel Channel, ECollisionResponse NewResponse)
+bool FBodyInstance::SetResponseToChannel(ECollisionChannel Channel, ECollisionResponse NewResponse)
 {
-	InvalidateCollisionProfileName();
-	CollisionResponses.SetResponse(Channel, NewResponse);
-	UpdatePhysicsFilterData();
+	if (CollisionResponses.SetResponse(Channel, NewResponse))
+	{
+		InvalidateCollisionProfileName();
+		UpdatePhysicsFilterData();
+		return true;
+	}
+
+	return false;
 }
 
-void FBodyInstance::SetResponseToAllChannels(ECollisionResponse NewResponse)
+bool FBodyInstance::SetResponseToAllChannels(ECollisionResponse NewResponse)
 {
-	InvalidateCollisionProfileName();
-	CollisionResponses.SetAllChannels(NewResponse);
-	UpdatePhysicsFilterData();
+	if (CollisionResponses.SetAllChannels(NewResponse))
+	{
+		InvalidateCollisionProfileName();
+		UpdatePhysicsFilterData();
+		return true;
+	}
+
+	return false;
 }
 	
-void FBodyInstance::ReplaceResponseToChannels(ECollisionResponse OldResponse, ECollisionResponse NewResponse)
+bool FBodyInstance::ReplaceResponseToChannels(ECollisionResponse OldResponse, ECollisionResponse NewResponse)
 {
-	InvalidateCollisionProfileName();
-	CollisionResponses.ReplaceChannels(OldResponse, NewResponse);
-	UpdatePhysicsFilterData();
+	if (CollisionResponses.ReplaceChannels(OldResponse, NewResponse))
+	{
+		InvalidateCollisionProfileName();
+		UpdatePhysicsFilterData();
+		return true;
+	}
+
+	return false;
 }
 
-void FBodyInstance::SetResponseToChannels(const FCollisionResponseContainer& NewReponses)
+bool FBodyInstance::SetResponseToChannels(const FCollisionResponseContainer& NewReponses)
 {
-	InvalidateCollisionProfileName();
-	CollisionResponses.SetCollisionResponseContainer(NewReponses);
-	UpdatePhysicsFilterData();
+	if (CollisionResponses.SetCollisionResponseContainer(NewReponses))
+	{
+		InvalidateCollisionProfileName();
+		UpdatePhysicsFilterData();
+		return true;
+	}
+
+	return false;
 }
 
 void FBodyInstance::SetObjectType(ECollisionChannel Channel)
diff --git a/Engine/Source/Runtime/Engine/Private/PhysicsEngine/PhysUtils.cpp b/Engine/Source/Runtime/Engine/Private/PhysicsEngine/PhysUtils.cpp
--- a/Engine/Source/Runtime/Engine/Private/PhysicsEngine/PhysUtils.cpp
+++ b/Engine/Source/Runtime/Engine/Private/PhysicsEngine/PhysUtils.cpp
@@ -196,32 +196,49 @@ void FRigidBodyContactInfo::SwapOrder()
 // FCollisionResponseContainer
 
 /** Set the status of a particular channel in the structure. */
-void FCollisionResponseContainer::SetResponse(ECollisionChannel Channel, ECollisionResponse NewResponse)
+bool FCollisionResponseContainer::SetResponse(ECollisionChannel Channel, ECollisionResponse NewResponse)
 {
 	if (Channel < ARRAY_COUNT(EnumArray))
 	{
-		EnumArray[Channel] = NewResponse;
+		uint8& CurrentResponse = EnumArray[Channel];
+		if (CurrentResponse != NewResponse)
+		{
+			CurrentResponse = NewResponse;
+			return true;
+		}
 	}
+	return false;
 }
 
 /** Set all channels to the specified state */
-void FCollisionResponseContainer::SetAllChannels(ECollisionResponse NewResponse)
+bool FCollisionResponseContainer::SetAllChannels(ECollisionResponse NewResponse)
 {
+	bool bHasChanged = false;
 	for(int32 i=0; i<ARRAY_COUNT(EnumArray); i++)
 	{
-		EnumArray[i] = NewResponse;
+		uint8& CurrentResponse = EnumArray[i];
+		if (CurrentResponse != NewResponse)
+		{
+			CurrentResponse = NewResponse;
+			bHasChanged = true;
+		}
 	}
+	return bHasChanged;
 }
 
-void FCollisionResponseContainer::ReplaceChannels(ECollisionResponse OldResponse, ECollisionResponse NewResponse)
+bool FCollisionResponseContainer::ReplaceChannels(ECollisionResponse OldResponse, ECollisionResponse NewResponse)
 {
+	bool bHasChanged = false;
 	for (int32 i = 0; i < ARRAY_COUNT(EnumArray); i++)
 	{
-		if(EnumArray[i] == OldResponse)
+		uint8& CurrentResponse = EnumArray[i];
+		if(CurrentResponse == OldResponse)
 		{
-			EnumArray[i] = NewResponse;
+			CurrentResponse = NewResponse;
+			bHasChanged = true;
 		}
 	}
+	return bHasChanged;
 }
 
 FCollisionResponseContainer FCollisionResponseContainer::CreateMinContainer(const FCollisionResponseContainer& A, const FCollisionResponseContainer& B)
diff --git a/Engine/Source/Runtime/Engine/Private/PrimitiveComponentPhysics.cpp b/Engine/Source/Runtime/Engine/Private/PrimitiveComponentPhysics.cpp
--- a/Engine/Source/Runtime/Engine/Private/PrimitiveComponentPhysics.cpp
+++ b/Engine/Source/Runtime/Engine/Private/PrimitiveComponentPhysics.cpp
@@ -867,20 +867,26 @@ void UPrimitiveComponent::SetCollisionObjectType(ECollisionChannel Channel)
 
 void UPrimitiveComponent::SetCollisionResponseToChannel(ECollisionChannel Channel, ECollisionResponse NewResponse)
 {
-	BodyInstance.SetResponseToChannel(Channel, NewResponse);
-	OnComponentCollisionSettingsChanged();
+	if (BodyInstance.SetResponseToChannel(Channel, NewResponse))
+	{ 
+		OnComponentCollisionSettingsChanged();
+	}
 }
 
 void UPrimitiveComponent::SetCollisionResponseToAllChannels(enum ECollisionResponse NewResponse)
 {
-	BodyInstance.SetResponseToAllChannels(NewResponse);
-	OnComponentCollisionSettingsChanged();
+	if (BodyInstance.SetResponseToAllChannels(NewResponse))
+	{
+		OnComponentCollisionSettingsChanged();
+	}
 }
 
 void UPrimitiveComponent::SetCollisionResponseToChannels(const FCollisionResponseContainer& NewReponses)
-{
-	BodyInstance.SetResponseToChannels(NewReponses);
-	OnComponentCollisionSettingsChanged();
+{	
+	if (BodyInstance.SetResponseToChannels(NewReponses))
+	{ 
+		OnComponentCollisionSettingsChanged();
+	}
 }
 
 void UPrimitiveComponent::SetCollisionEnabled(ECollisionEnabled::Type NewType)
@@ -913,7 +919,7 @@ void UPrimitiveComponent::SetCollisionProfileName(FName InCollisionProfileName)
 		// may call SetCollisionProfileName more than once.
 		BodyInstance.SetCollisionProfileNameDeferred(InCollisionProfileName);
 	}
-	else
+	else if (InCollisionProfileName != BodyInstance.GetCollisionProfileName())
 	{
 		ECollisionEnabled::Type OldCollisionEnabled = BodyInstance.GetCollisionEnabled();
 		BodyInstance.SetCollisionProfileName(InCollisionProfileName);
