Fix the inverted check in Delinquency Analytics when searching for an existing top offender.

diff --git a/Engine/Source/Runtime/Engine/Public/Net/DelinquencyAnalytics.h b/Engine/Source/Runtime/Engine/Public/Net/DelinquencyAnalytics.h
--- a/Engine/Source/Runtime/Engine/Public/Net/DelinquencyAnalytics.h
+++ b/Engine/Source/Runtime/Engine/Public/Net/DelinquencyAnalytics.h
@@ -148,7 +148,7 @@ public:
 							// Sanity check that we have actually found our entry.
 							// If we hit the end of the list, or we see an entry that should be less offensive, then
 							// we've missed our entry.
-							if (MaybeOurEntry == TopOffenders.Num() || TopOffenders[MaybeOurEntry].TimeSeconds > AlreadyTracked->TimeSeconds)
+							if (MaybeOurEntry == TopOffenders.Num() || TopOffenders[MaybeOurEntry].TimeSeconds < AlreadyTracked->TimeSeconds)
 							{
 								// It's possible that multiple entries have the same delinquency time.
 								// If our current entry matches the LeastOffensiveTime, it's possible that we were in the
