Add an option to content encryption groups to say that their members should be encrypted even if they are soft referenced by some other asset

diff --git a/Engine/Source/Runtime/Engine/Classes/Engine/ContentEncryptionConfig.h b/Engine/Source/Runtime/Engine/Classes/Engine/ContentEncryptionConfig.h
--- a/Engine/Source/Runtime/Engine/Classes/Engine/ContentEncryptionConfig.h
+++ b/Engine/Source/Runtime/Engine/Classes/Engine/ContentEncryptionConfig.h
@@ -10,6 +10,7 @@ public:
 	{
 		TSet<FName> PackageNames;
 		bool bStageTimeOnly = false;
+		bool bIgnoreSoftReferences = false;
 	};
 
 	typedef TMap<FName, FGroup> TGroupMap;
@@ -24,6 +25,11 @@ public:
 		PackageGroups.FindOrAdd(InGroupName).bStageTimeOnly = bInStageTimeOnly;
 	}
 
+	void SetShouldIgnoreSoftReferences(FName InGroupName, bool bInIgnoreSoftReferences)
+	{
+		PackageGroups.FindOrAdd(InGroupName).bIgnoreSoftReferences = bInIgnoreSoftReferences;
+	}
+
 	void AddReleasedKey(FGuid InKey)
 	{
 		ReleasedKeys.Add(InKey);
