Fix potential crash on invalid shared pointer

diff --git a/Engine/Source/Runtime/Engine/Private/DemoNetDriver.cpp b/Engine/Source/Runtime/Engine/Private/DemoNetDriver.cpp
--- a/Engine/Source/Runtime/Engine/Private/DemoNetDriver.cpp
+++ b/Engine/Source/Runtime/Engine/Private/DemoNetDriver.cpp
@@ -695,16 +695,18 @@ bool UDemoNetDriver::ProcessReplayTasks()
 		UE_LOG(LogDemo, Verbose, TEXT("UDemoNetDriver::ProcessReplayTasks. Name: %s"), *ActiveReplayTask->GetName().ToString());
 
 		// Start the task
-		ActiveReplayTask->StartTask();
+		LocalActiveTask->StartTask();
 	}
 
 	// Tick the currently active task
 	if (ActiveReplayTask.IsValid())
 	{
-		if (!ActiveReplayTask->Tick())
+		LocalActiveTask = ActiveReplayTask;
+
+		if (!LocalActiveTask->Tick())
 		{
 			// Task isn't done, we can return
-			return !ActiveReplayTask->ShouldPausePlayback();
+			return !LocalActiveTask->ShouldPausePlayback();
 		}
 
 		// This task is now done
