Fix log spam of UnHideBone

diff --git a/Engine/Source/Runtime/Engine/Private/Components/SkeletalMeshComponent.cpp b/Engine/Source/Runtime/Engine/Private/Components/SkeletalMeshComponent.cpp
--- a/Engine/Source/Runtime/Engine/Private/Components/SkeletalMeshComponent.cpp
+++ b/Engine/Source/Runtime/Engine/Private/Components/SkeletalMeshComponent.cpp
@@ -2776,6 +2776,11 @@ void USkeletalMeshComponent::UnHideBone( int32 BoneIndex )
 		return;
 	}
 
+	if (MasterPoseComponent.IsValid())
+	{
+		return;
+	}
+
 	if (BoneSpaceTransforms.IsValidIndex(BoneIndex))
 	{
 		BoneSpaceTransforms[BoneIndex].SetScale3D(FVector(1.0f));
@@ -2789,7 +2794,7 @@ void USkeletalMeshComponent::UnHideBone( int32 BoneIndex )
 	}
 	else
 	{
-		UE_LOG(LogSkeletalMesh, Warning, TEXT("UnHideBone: Invalid Body Index has entered. This component doesn't contain buffer for the given body."));
+		UE_LOG(LogSkeletalMesh, Warning, TEXT("UnHideBone[%s]: Invalid Body Index (%d) has entered. This component doesn't contain buffer for the given body."), *GetNameSafe(SkeletalMesh), BoneIndex);
 	}
 }
 
