Don't apply/remove level transform when undoing actions on a streaming level, instead, include the actors whose positions are being modified correctly in a transaction that moves the streaming level
Properly route PostEditUndo to Super in ULevelStreaming

diff --git a/Engine/Source/Runtime/Engine/Classes/Engine/LevelStreaming.h b/Engine/Source/Runtime/Engine/Classes/Engine/LevelStreaming.h
--- a/Engine/Source/Runtime/Engine/Classes/Engine/LevelStreaming.h
+++ b/Engine/Source/Runtime/Engine/Classes/Engine/LevelStreaming.h
@@ -367,7 +367,6 @@ public:
 	
 #if WITH_EDITOR
 	/** Override Pre/PostEditUndo functions to handle editor transform */
-	virtual void PreEditUndo() override;
 	virtual void PostEditUndo() override;
 #endif
 	
diff --git a/Engine/Source/Runtime/Engine/Private/LevelStreaming.cpp b/Engine/Source/Runtime/Engine/Private/LevelStreaming.cpp
--- a/Engine/Source/Runtime/Engine/Private/LevelStreaming.cpp
+++ b/Engine/Source/Runtime/Engine/Private/LevelStreaming.cpp
@@ -1584,14 +1584,10 @@ ALevelScriptActor* ULevelStreaming::GetLevelScriptActor()
 }
 
 #if WITH_EDITOR
-void ULevelStreaming::PreEditUndo()
-{
-	FLevelUtils::RemoveEditorTransform(this, false);
-}
-
 void ULevelStreaming::PostEditUndo()
 {
-	FLevelUtils::ApplyEditorTransform(this, false);
+	Super::PostEditUndo();
+
 	if (UWorld* World = GetWorld())
 	{
 		World->UpdateStreamingLevelShouldBeConsidered(this);
diff --git a/Engine/Source/Runtime/Engine/Private/LevelUtils.cpp b/Engine/Source/Runtime/Engine/Private/LevelUtils.cpp
--- a/Engine/Source/Runtime/Engine/Private/LevelUtils.cpp
+++ b/Engine/Source/Runtime/Engine/Private/LevelUtils.cpp
@@ -333,6 +333,19 @@ void FLevelUtils::SetEditorTransform(ULevelStreaming* StreamingLevel, const FTra
 	const FScopedTransaction LevelOffsetTransaction( LOCTEXT( "ChangeEditorLevelTransform", "Edit Level Transform" ) );
 	StreamingLevel->Modify();
 
+	// Ensure that all Actors are in the transaction so that their location is restored and any construction script behaviors 
+	// based on being at a different location are correctly applied on undo/redo
+	if (ULevel* LoadedLevel = StreamingLevel->GetLoadedLevel())
+	{
+		for (AActor* Actor : LoadedLevel->Actors)
+		{
+			if (Actor)
+			{
+				Actor->Modify();
+			}
+		}
+	}
+
 	// Apply new transform
 	RemoveEditorTransform(StreamingLevel, false );
 	StreamingLevel->LevelTransform = Transform;
