Add new IncludeAssetBundles meta tag to support including referenced object properties in asset bundles for owning object.
#rnx

diff --git a/Engine/Source/Runtime/Engine/Private/AssetManager.cpp b/Engine/Source/Runtime/Engine/Private/AssetManager.cpp
--- a/Engine/Source/Runtime/Engine/Private/AssetManager.cpp
+++ b/Engine/Source/Runtime/Engine/Private/AssetManager.cpp
@@ -3563,6 +3563,7 @@ void UAssetManager::RefreshAssetData(UObject* ChangedObject)
 void UAssetManager::InitializeAssetBundlesFromMetadata(const UStruct* Struct, const void* StructValue, FAssetBundleData& AssetBundle, FName DebugName) const
 {
 	static FName AssetBundlesName = TEXT("AssetBundles");
+	static FName IncludeAssetBundlesName = TEXT("IncludeAssetBundles");
 
 	if (!ensure(Struct && StructValue))
 	{
@@ -3607,7 +3608,7 @@ void UAssetManager::InitializeAssetBundlesFromMetadata(const UStruct* Struct, co
 		}
 		else if (const UObjectProperty* ObjectProperty = Cast<UObjectProperty>(Property))
 		{
-			if (ObjectProperty->PropertyFlags & CPF_InstancedReference)
+			if (ObjectProperty->PropertyFlags & CPF_InstancedReference || ObjectProperty->HasMetaData(IncludeAssetBundlesName))
 			{
 				UObject* const* ObjectPtr = reinterpret_cast<UObject* const*>(PropertyValue);
 				if (ObjectPtr && *ObjectPtr)
