Added yet more scoped stats to PostInitializeComponents to help track down hitches spawning actors on PC

diff --git a/Engine/Source/Runtime/Engine/Private/Actor.cpp b/Engine/Source/Runtime/Engine/Private/Actor.cpp
--- a/Engine/Source/Runtime/Engine/Private/Actor.cpp
+++ b/Engine/Source/Runtime/Engine/Private/Actor.cpp
@@ -4627,6 +4627,8 @@ float AActor::GetLifeSpan() const
 
 void AActor::PostInitializeComponents()
 {
+	QUICK_SCOPE_CYCLE_COUNTER(STAT_Actor_PostInitComponents);
+
 	if( !IsPendingKill() )
 	{
 		bActorInitialized = true;
diff --git a/Engine/Source/Runtime/Engine/Private/Character.cpp b/Engine/Source/Runtime/Engine/Private/Character.cpp
--- a/Engine/Source/Runtime/Engine/Private/Character.cpp
+++ b/Engine/Source/Runtime/Engine/Private/Character.cpp
@@ -110,6 +110,8 @@ ACharacter::ACharacter(const FObjectInitializer& ObjectInitializer)
 
 void ACharacter::PostInitializeComponents()
 {
+	QUICK_SCOPE_CYCLE_COUNTER(STAT_Character_PostInitComponents);
+
 	Super::PostInitializeComponents();
 
 	if (!IsPendingKill())
diff --git a/Engine/Source/Runtime/Engine/Private/Pawn.cpp b/Engine/Source/Runtime/Engine/Private/Pawn.cpp
--- a/Engine/Source/Runtime/Engine/Private/Pawn.cpp
+++ b/Engine/Source/Runtime/Engine/Private/Pawn.cpp
@@ -100,6 +100,8 @@ void APawn::PreInitializeComponents()
 
 void APawn::PostInitializeComponents()
 {
+	QUICK_SCOPE_CYCLE_COUNTER(STAT_Pawn_PostInitComponents);
+
 	Super::PostInitializeComponents();
 	
 	if (!IsPendingKill())
