Reduce memory spike of the ExportGUIDArchives array while demo recording in creative

diff --git a/Engine/Source/Runtime/Engine/Private/PackageMapClient.cpp b/Engine/Source/Runtime/Engine/Private/PackageMapClient.cpp
--- a/Engine/Source/Runtime/Engine/Private/PackageMapClient.cpp
+++ b/Engine/Source/Runtime/Engine/Private/PackageMapClient.cpp
@@ -46,6 +46,7 @@ static const int INTERNAL_LOAD_OBJECT_RECURSION_LIMIT = 16;
 
 static TAutoConsoleVariable<int32> CVarAllowAsyncLoading( TEXT( "net.AllowAsyncLoading" ), 0, TEXT( "Allow async loading" ) );
 static TAutoConsoleVariable<int32> CVarIgnoreNetworkChecksumMismatch( TEXT( "net.IgnoreNetworkChecksumMismatch" ), 0, TEXT( "" ) );
+static TAutoConsoleVariable<int32> CVarReservedNetGuidSize(TEXT("net.ReservedNetGuidSize"), 512, TEXT("Reserved size in bytes for NetGUID serialization"));
 extern FAutoConsoleVariableRef CVarEnableMultiplayerWorldOriginRebasing;
 extern TAutoConsoleVariable<int32> CVarFilterGuidRemapping;
 
@@ -1033,10 +1034,9 @@ UObject* UPackageMapClient::ResolvePathAndAssignNetGUID( const FNetworkGUID& Net
 //
 //--------------------------------------------------------------------
 
-// MAX_BUNCH_SIZE is in bits, so shift to get bytes.
+
 // TODO: This limit might not actually need to be enforced anymore.
-static const int32 MAX_GUID_MEMORY = MAX_BUNCH_SIZE >> 3;
-static const int32 MAX_GUID_COUNT = 2048;
+constexpr int32 MAX_GUID_COUNT = 2048;
 
 bool UPackageMapClient::ExportNetGUIDForReplay(FNetworkGUID& NetGUID, UObject* Object, FString& PathName, UObject* ObjOuter)
 {
@@ -1047,14 +1047,18 @@ bool UPackageMapClient::ExportNetGUIDForReplay(FNetworkGUID& NetGUID, UObject* O
 	{
 		TGuardValue<bool> ExportingGUID(GuidCache->IsExportingNetGUIDBunch, true);
 
+		const int32 MaxReservedSize(CVarReservedNetGuidSize.GetValueOnAnyThread());
+
 		TArray<uint8>& GUIDMemory = ExportGUIDArchives.Emplace_GetRef();
-		GUIDMemory.Reserve(MAX_GUID_MEMORY);
+		GUIDMemory.Reserve(MaxReservedSize);
 
 		FMemoryWriter Writer(GUIDMemory);
 		InternalWriteObject(Writer, NetGUID, Object, PathName, ObjOuter);
 
 		check(!Writer.IsError());
-		ensureMsgf(GUIDMemory.Num() <= MAX_GUID_MEMORY, TEXT("ExportNetGUIDForReplay exceeded MAX_GUID_MEMORY. Max=%l Count=%l"), MAX_GUID_MEMORY, GUIDMemory.Num());
+		ensureMsgf(GUIDMemory.Num() <= MaxReservedSize, TEXT("ExportNetGUIDForReplay exceeded CVarReservedNetGuidSize. Max=%l Count=%l"), MaxReservedSize, GUIDMemory.Num());
+
+		GUIDMemory.Shrink();
 
 		// It's possible InternalWriteObject has modified the NetGUIDAckStatus, so
 		// do a quick sanity check to make sure the ID wasn't removed before updating the status.
