Fixing a crash whenever a mesh is rendered in wireframe

diff --git a/Engine/Source/Runtime/Engine/Private/StaticMeshRender.cpp b/Engine/Source/Runtime/Engine/Private/StaticMeshRender.cpp
--- a/Engine/Source/Runtime/Engine/Private/StaticMeshRender.cpp
+++ b/Engine/Source/Runtime/Engine/Private/StaticMeshRender.cpp
@@ -668,7 +668,7 @@ void FStaticMeshSceneProxy::SetIndexSource(int32 LODIndex, int32 SectionIndex, F
 	const FStaticMeshLODResources& LODModel = RenderData->LODResources[LODIndex];
 	if (bWireframe)
 	{
-		if( LODModel.AdditionalIndexBuffers->WireframeIndexBuffer.IsInitialized()
+		if(LODModel.AdditionalIndexBuffers && LODModel.AdditionalIndexBuffers->WireframeIndexBuffer.IsInitialized()
 			&& !(RHISupportsTessellation(GetScene().GetShaderPlatform()) && OutMeshElement.VertexFactory->GetType()->SupportsTessellationShaders())
 			)
 		{
