Sounds getting virtualized then requesting to fade out aren't being stopped and reporting to Component accordingly until sound potentially moves from virtualization (realizes), potentially never.

diff --git a/Engine/Source/Runtime/Engine/Private/AudioDevice.cpp b/Engine/Source/Runtime/Engine/Private/AudioDevice.cpp
--- a/Engine/Source/Runtime/Engine/Private/AudioDevice.cpp
+++ b/Engine/Source/Runtime/Engine/Private/AudioDevice.cpp
@@ -4002,7 +4002,7 @@ void FAudioDevice::SendUpdateResultsToGameThread(const int32 FirstActiveIndex)
 				FAudioStats::FStatWaveInstanceInfo WaveInstanceInfo;
 				FSoundSource* Source = WaveInstanceSourceMap.FindRef(WaveInstance);
 				WaveInstanceInfo.Description = Source ? Source->Describe((RequestedAudioStats & ERequestedAudioStats::LongSoundNames) != 0) : FString(TEXT("No source"));
-				WaveInstanceInfo.ActualVolume = WaveInstance->GetVolumeWithDistanceAttenuation();
+				WaveInstanceInfo.ActualVolume = WaveInstance->GetVolumeWithDistanceAttenuation() * WaveInstance->GetDynamicVolume();
 				WaveInstanceInfo.InstanceIndex = InstanceIndex;
 				WaveInstanceInfo.WaveInstanceName = *WaveInstance->GetName();
 				StatSoundInfos[*SoundInfoIndex].WaveInstanceInfos.Add(MoveTemp(WaveInstanceInfo));
@@ -5815,10 +5815,23 @@ void FAudioDevice::UpdateVirtualLoops()
 		for (FVirtualLoopPair& Pair : VirtualLoops)
 		{
 			FAudioVirtualLoop& VirtualLoop = Pair.Value;
+			FActiveSound& ActiveSound = VirtualLoop.GetActiveSound();
+
+			// Don't update if stopping.
+			if (ActiveSound.bIsStopping)
+			{
+				continue;
+			}
+
+			// If signaled to fade out and virtualized, add to pending stop list.
+			if (ActiveSound.bFadingOut)
+			{
+				AddSoundToStop(&ActiveSound);
+				continue;
+			}
 
 			const float DeltaTime = GetDeviceDeltaTime();
 
-			FActiveSound& ActiveSound = Pair.Value.GetActiveSound();
 			// Keep playback time up-to-date as it may be used to evaluate concurrency
 			ActiveSound.PlaybackTime += DeltaTime * ActiveSound.MinCurrentPitch;
 
@@ -5826,7 +5839,7 @@ void FAudioDevice::UpdateVirtualLoops()
 			// to avoid map manipulation while iterating.
 			if (VirtualLoop.CanRealize(DeltaTime))
 			{
-				VirtualLoopsToRetrigger.Add(Pair.Value);
+				VirtualLoopsToRetrigger.Add(VirtualLoop);
 			}
 		}
 
