Fix incorrect log

diff --git a/Engine/Source/Runtime/Engine/Private/Components/SkinnedMeshComponent.cpp b/Engine/Source/Runtime/Engine/Private/Components/SkinnedMeshComponent.cpp
--- a/Engine/Source/Runtime/Engine/Private/Components/SkinnedMeshComponent.cpp
+++ b/Engine/Source/Runtime/Engine/Private/Components/SkinnedMeshComponent.cpp
@@ -2441,7 +2441,7 @@ void USkinnedMeshComponent::SetRefPoseOverride(const TArray<FTransform>& NewRefP
 
 	if (NumRealBones != NewRefPoseTransforms.Num())
 	{
-		UE_LOG(LogSkeletalMesh, Warning, TEXT("SetRefPoseOverride (%s) : Expected %d transforms, got %d."), NumRealBones, NewRefPoseTransforms.Num());
+		UE_LOG(LogSkeletalMesh, Warning, TEXT("SetRefPoseOverride (%s) : Expected %d transforms, got %d."), *SkeletalMesh->GetName(), NumRealBones, NewRefPoseTransforms.Num());
 		return;
 	}
 
@@ -3637,4 +3637,4 @@ template FVector GetTypedSkinnedVertexPosition<true, false>(const USkinnedMeshCo
 	const FSkinWeightVertexBuffer& SkinWeightVertexBuffer, const int32 VertIndex, const TArray<FMatrix> & RefToLocals);
 
 template FVector GetTypedSkinnedVertexPosition<false, false>(const USkinnedMeshComponent* SkinnedComp, const FSkelMeshRenderSection& Section, const FPositionVertexBuffer& PositionVertexBuffer,
-	const FSkinWeightVertexBuffer& SkinWeightVertexBuffer, const int32 VertIndex, const TArray<FMatrix> & RefToLocals);
\ No newline at end of file
+	const FSkinWeightVertexBuffer& SkinWeightVertexBuffer, const int32 VertIndex, const TArray<FMatrix> & RefToLocals);
