Add a delegate to UNetConnection that can allow games to remap Level Package Names if the server fails to resolve them.
This is particularly useful when using ULevelStreamingDynamic, as their package names don't have to map to on disk packages.

diff --git a/Engine/Source/Runtime/Engine/Classes/Engine/NetConnection.h b/Engine/Source/Runtime/Engine/Classes/Engine/NetConnection.h
--- a/Engine/Source/Runtime/Engine/Classes/Engine/NetConnection.h
+++ b/Engine/Source/Runtime/Engine/Classes/Engine/NetConnection.h
@@ -165,6 +165,26 @@ DECLARE_DELEGATE_ThreeParams(FOnReceivedRawPacket, void* /*Data*/, int32 /*Count
 DECLARE_DELEGATE_ThreeParams(FOnLowLevelSend, void* /*Data*/, int32 /*Count*/, bool& /*bBlockSend*/);
 #endif
 
+struct FRemapClientLevelPackageNameParams
+{
+	/** Package name that was received. */
+	const FName InPackageName;
+
+	/** [Out] Package name that was remapped.*/
+	FString OutPackageName;
+};
+
+/**
+ * This delegate will be called if a client sends the server a Level Visibility for a package that the
+ * server fails to resolve. This mainly occurs when using Dynamic Level Instances that the server
+ * hasn't loaded / created yet.
+ *
+ * Clients can use this delegate to try and remap the instanced package name back to its base package.
+ * E.G. /Game/Maps/MyMap_Instance1 -> /Games/MyMap.
+ *
+ * @return True if the name was remapped, false otherwise.
+ */
+DECLARE_DELEGATE_RetVal_OneParam(bool, FRemapClientLevelPackageName, FRemapClientLevelPackageNameParams&);
 
 #if DO_ENABLE_NET_TEST
 /**
@@ -1087,6 +1107,18 @@ public:
 	/** Resets the current delinquency analytics. */
 	ENGINE_API void ResetQueuedActorDelinquencyAnalytics();
 
+	/** Clears our RemapLevelPackageName delegate. @see FRemapClientLevelPackageName. */
+	ENGINE_API void ClearRemapLevelPackageName()
+	{
+		RemapLevelPackageName.Unbind();
+	}
+
+	/** Sets (or resets) our RemapLevelPackageName delegate. @see FRemapClientLevelPackageName. */
+	ENGINE_API void SetRemapLevelPackageName(FRemapClientLevelPackageName&& NewDelegate)
+	{
+		RemapLevelPackageName = MoveTemp(NewDelegate);
+	}
+
 protected:
 
 	ENGINE_API void SetPendingCloseDueToSocketSendFailure();
@@ -1196,6 +1228,9 @@ private:
 
 	/** Whether or not PacketOrderCache is presently being flushed */
 	bool bFlushingPacketOrderCache;
+
+	/** @see FRemapClientLevelPackagename. */
+	FRemapClientLevelPackageName RemapLevelPackageName;
 };
 
 /** Help structs for temporarily setting network settings */
diff --git a/Engine/Source/Runtime/Engine/Private/NetConnection.cpp b/Engine/Source/Runtime/Engine/Private/NetConnection.cpp
--- a/Engine/Source/Runtime/Engine/Private/NetConnection.cpp
+++ b/Engine/Source/Runtime/Engine/Private/NetConnection.cpp
@@ -900,7 +900,7 @@ void UNetConnection::UpdateLevelVisibility(const FName& PackageName, bool bIsVis
 	if (bIsVisible)
 	{
 		// verify that we were passed a valid level name
-		FString Filename;
+		FRemapClientLevelPackageNameParams RemapPackageNameParams{ PackageName };
 		UPackage* TempPkg = FindPackage(nullptr, *PackageName.ToString());
 		FLinkerLoad* Linker = FLinkerLoad::FindExistingLinkerForPackage(TempPkg);
 
@@ -923,7 +923,10 @@ void UNetConnection::UpdateLevelVisibility(const FName& PackageName, bool bIsVis
 			}
 		};
 
-		if ( Linker || FPackageName::DoesPackageExist(PackageName.ToString(), nullptr, &Filename ) || Local::IsInLevelList(GetWorld(), PackageName ) )
+		if (Linker ||
+			FPackageName::DoesPackageExist(PackageName.ToString()) ||
+			Local::IsInLevelList(GetWorld(), PackageName) ||
+			(RemapLevelPackageName.IsBound() && RemapLevelPackageName.Execute(RemapPackageNameParams) && FPackageName::DoesPackageExist(RemapPackageNameParams.OutPackageName)))
 		{
 			ClientVisibleLevelNames.Add(PackageName);
 			UE_LOG( LogPlayerController, Verbose, TEXT("ServerUpdateLevelVisibility() Added '%s'"), *PackageName.ToString() );
