Added UStaticMesh::MinLOD to the asset registry so it can be viewed in column view or searched using CB advanced search syntax (once assets are resaved)

diff --git a/Engine/Source/Runtime/Engine/Private/PerPlatformProperties.cpp b/Engine/Source/Runtime/Engine/Private/PerPlatformProperties.cpp
--- a/Engine/Source/Runtime/Engine/Private/PerPlatformProperties.cpp
+++ b/Engine/Source/Runtime/Engine/Private/PerPlatformProperties.cpp
@@ -77,3 +77,21 @@ template ENGINE_API FArchive& operator<<(FArchive&, TPerPlatformProperty<FPerPla
 template ENGINE_API void operator<<(FStructuredArchive::FSlot Slot, TPerPlatformProperty<FPerPlatformInt, int32, NAME_IntProperty>&);
 template ENGINE_API void operator<<(FStructuredArchive::FSlot Slot, TPerPlatformProperty<FPerPlatformFloat, float, NAME_FloatProperty>&);
 template ENGINE_API void operator<<(FStructuredArchive::FSlot Slot, TPerPlatformProperty<FPerPlatformBool, bool, NAME_BoolProperty>&);
+
+FString FPerPlatformInt::ToString() const
+{
+	FString Result = FString::FromInt(Default);
+
+#if WITH_EDITORONLY_DATA
+	TArray<FName> SortedPlatforms;
+	PerPlatform.GetKeys(/*out*/ SortedPlatforms);
+	SortedPlatforms.Sort();
+
+	for (FName Platform : SortedPlatforms)
+	{
+		Result = FString::Printf(TEXT("%s, %s=%d"), *Result, *Platform.ToString(), PerPlatform.FindChecked(Platform));
+	}
+#endif
+
+	return Result;
+}
\ No newline at end of file
diff --git a/Engine/Source/Runtime/Engine/Private/StaticMesh.cpp b/Engine/Source/Runtime/Engine/Private/StaticMesh.cpp
--- a/Engine/Source/Runtime/Engine/Private/StaticMesh.cpp
+++ b/Engine/Source/Runtime/Engine/Private/StaticMesh.cpp
@@ -3201,6 +3201,7 @@ void UStaticMesh::GetAssetRegistryTags(TArray<FAssetRegistryTag>& OutTags) const
 	OutTags.Add( FAssetRegistryTag("ApproxSize", ApproxSizeStr, FAssetRegistryTag::TT_Dimensional) );
 	OutTags.Add( FAssetRegistryTag("CollisionPrims", FString::FromInt(NumCollisionPrims), FAssetRegistryTag::TT_Numerical));
 	OutTags.Add( FAssetRegistryTag("LODs", FString::FromInt(NumLODs), FAssetRegistryTag::TT_Numerical));
+	OutTags.Add( FAssetRegistryTag("MinLOD", MinLOD.ToString(), FAssetRegistryTag::TT_Alphabetical));
 	OutTags.Add( FAssetRegistryTag("SectionsWithCollision", FString::FromInt(NumSectionsWithCollision), FAssetRegistryTag::TT_Numerical));
 	OutTags.Add( FAssetRegistryTag("DefaultCollision", DefaultCollisionName.ToString(), FAssetRegistryTag::TT_Alphabetical));
 	OutTags.Add( FAssetRegistryTag("CollisionComplexity", ComplexityString, FAssetRegistryTag::TT_Alphabetical));
diff --git a/Engine/Source/Runtime/Engine/Public/PerPlatformProperties.h b/Engine/Source/Runtime/Engine/Public/PerPlatformProperties.h
--- a/Engine/Source/Runtime/Engine/Public/PerPlatformProperties.h
+++ b/Engine/Source/Runtime/Engine/Public/PerPlatformProperties.h
@@ -143,6 +143,8 @@ struct ENGINE_API FPerPlatformInt
 	:	Default(InDefaultValue)
 	{
 	}
+
+	FString ToString() const;
 };
 extern template ENGINE_API FArchive& operator<<(FArchive&, TPerPlatformProperty<FPerPlatformInt, int32, NAME_IntProperty>&);
 extern template ENGINE_API void operator<<(FStructuredArchive::FSlot Slot, TPerPlatformProperty<FPerPlatformInt, int32, NAME_IntProperty>&);
