Remove Ocean branch hack for linear color curve warnings
#rnx

diff --git a/Engine/Source/Runtime/Engine/Private/Curves/CurveLinearColor.cpp b/Engine/Source/Runtime/Engine/Private/Curves/CurveLinearColor.cpp
--- a/Engine/Source/Runtime/Engine/Private/Curves/CurveLinearColor.cpp
+++ b/Engine/Source/Runtime/Engine/Private/Curves/CurveLinearColor.cpp
@@ -314,20 +314,13 @@ void UCurveLinearColor::PostLoad()
 {
 	Super::PostLoad();
 
-	// BEGIN OCEAN HACK
-	// This warning will keep coming up each time a new curve asset is saved in the release branch and gets robo-merged back down to main because the engine version is lower in the release branch
-	// Remove this hack once the release branch is brought up to the latest engine version
-	if (false)
-	{
 #if WITH_EDITOR
-		if (GetLinkerCustomVersion(FReleaseObjectVersion::GUID) < FReleaseObjectVersion::UnclampRGBColorCurves)
-		{
-			FMessageLog("LoadErrors").Warning(FText::Format(NSLOCTEXT("CurveEditor", "CurveDataUpdate", "Linear color curves now accurately handle RGB values > 1. If you were relying on HSV clamping, please update {0}"),
-				FText::FromString(GetName())));
-		}
-#endif
+	if (GetLinkerCustomVersion(FReleaseObjectVersion::GUID) < FReleaseObjectVersion::UnclampRGBColorCurves)
+	{
+		FMessageLog("LoadErrors").Warning(FText::Format(NSLOCTEXT("CurveEditor", "CurveDataUpdate", "Linear color curves now accurately handle RGB values > 1. If you were relying on HSV clamping, please update {0}"),
+			FText::FromString(GetName())));
 	}
-	// END OCEAN HACK
+#endif
 }
 
 void UCurveLinearColor::Serialize(FArchive& Ar)
