UE4 - Fix CIS warning
#rnx

diff --git a/Engine/Source/Runtime/Engine/Private/StaticMeshRender.cpp b/Engine/Source/Runtime/Engine/Private/StaticMeshRender.cpp
--- a/Engine/Source/Runtime/Engine/Private/StaticMeshRender.cpp
+++ b/Engine/Source/Runtime/Engine/Private/StaticMeshRender.cpp
@@ -724,6 +724,7 @@ void FStaticMeshSceneProxy::SetIndexSource(int32 LODIndex, int32 SectionIndex, F
 	if ( bRequiresAdjacencyInformation )
 	{
 		check( LODModel.bHasAdjacencyInfo );
+		check(LODModel.AdditionalIndexBuffers);
 		OutElement.IndexBuffer = &LODModel.AdditionalIndexBuffers->AdjacencyIndexBuffer;
 		OutMeshElement.Type = PT_12_ControlPointPatchList;
 		OutElement.FirstIndex *= 4;
