Fixed crash when generating texture streaming data for material instances with no parents.

diff --git a/Engine/Source/Runtime/Engine/Private/DebugViewModeMaterialProxy.cpp b/Engine/Source/Runtime/Engine/Private/DebugViewModeMaterialProxy.cpp
--- a/Engine/Source/Runtime/Engine/Private/DebugViewModeMaterialProxy.cpp
+++ b/Engine/Source/Runtime/Engine/Private/DebugViewModeMaterialProxy.cpp
@@ -35,48 +35,54 @@ FDebugViewModeMaterialProxy::FDebugViewModeMaterialProxy(
 	MaterialInterface->AppendReferencedTextures(ReferencedTextures);
 
 	FMaterialResource* Resource = InMaterialInterface->GetMaterialResource(FeatureLevel);
-	
-	const FDebugViewModeInterface* DebugViewModeInterface = FDebugViewModeInterface::GetInterface(InDebugViewMode);
-	if (DebugViewModeInterface)
+	if (Resource)
 	{
-		PixelShaderName = DebugViewModeInterface->PixelShaderName;
+		const FDebugViewModeInterface* DebugViewModeInterface = FDebugViewModeInterface::GetInterface(InDebugViewMode);
+		if (DebugViewModeInterface)
+		{
+			PixelShaderName = DebugViewModeInterface->PixelShaderName;
+
+			if (!DebugViewModeInterface->bNeedsOnlyLocalVertexFactor)
+			{
+				// Cache material usage.
+				bIsUsedWithSkeletalMesh = Resource->IsUsedWithSkeletalMesh();
+				bIsUsedWithLandscape = Resource->IsUsedWithLandscape();
+				bIsUsedWithParticleSystem = Resource->IsUsedWithParticleSystem();
+				bIsUsedWithParticleSprites = Resource->IsUsedWithParticleSprites();
+				bIsUsedWithBeamTrails = Resource->IsUsedWithBeamTrails();
+				bIsUsedWithMeshParticles = Resource->IsUsedWithMeshParticles();
+				bIsUsedWithNiagaraSprites = Resource->IsUsedWithNiagaraSprites();
+				bIsUsedWithNiagaraRibbons = Resource->IsUsedWithNiagaraRibbons();
+				bIsUsedWithNiagaraMeshParticles = Resource->IsUsedWithNiagaraMeshParticles();
+				bIsUsedWithMorphTargets = Resource->IsUsedWithMorphTargets();
+				bIsUsedWithSplineMeshes = Resource->IsUsedWithSplineMeshes();
+				bIsUsedWithInstancedStaticMeshes = Resource->IsUsedWithInstancedStaticMeshes();
+				bIsUsedWithAPEXCloth = Resource->IsUsedWithAPEXCloth();
+			}
+		}
+
+		FMaterialShaderMapId ResourceId;
+		Resource->GetShaderMapId(GMaxRHIShaderPlatform, ResourceId);
 
-		if (!DebugViewModeInterface->bNeedsOnlyLocalVertexFactor)
 		{
-			// Cache material usage.
-			bIsUsedWithSkeletalMesh = Resource->IsUsedWithSkeletalMesh();
-			bIsUsedWithLandscape = Resource->IsUsedWithLandscape();
-			bIsUsedWithParticleSystem = Resource->IsUsedWithParticleSystem();
-			bIsUsedWithParticleSprites = Resource->IsUsedWithParticleSprites();
-			bIsUsedWithBeamTrails = Resource->IsUsedWithBeamTrails();
-			bIsUsedWithMeshParticles = Resource->IsUsedWithMeshParticles();
-			bIsUsedWithNiagaraSprites = Resource->IsUsedWithNiagaraSprites();
-			bIsUsedWithNiagaraRibbons = Resource->IsUsedWithNiagaraRibbons();
-			bIsUsedWithNiagaraMeshParticles = Resource->IsUsedWithNiagaraMeshParticles();
-			bIsUsedWithMorphTargets = Resource->IsUsedWithMorphTargets();
-			bIsUsedWithSplineMeshes = Resource->IsUsedWithSplineMeshes();
-			bIsUsedWithInstancedStaticMeshes = Resource->IsUsedWithInstancedStaticMeshes();
-			bIsUsedWithAPEXCloth = Resource->IsUsedWithAPEXCloth();
+			TArray<FShaderType*> ShaderTypes;
+			TArray<FVertexFactoryType*> VFTypes;
+			TArray<const FShaderPipelineType*> ShaderPipelineTypes;
+			GetDependentShaderAndVFTypes(GMaxRHIShaderPlatform, ShaderTypes, ShaderPipelineTypes, VFTypes);
+
+			// Overwrite the shader map Id's dependencies with ones that came from the FMaterial actually being compiled (this)
+			// This is necessary as we change FMaterial attributes like GetShadingModel(), which factor into the ShouldCache functions that determine dependent shader types
+			ResourceId.SetShaderDependencies(ShaderTypes, ShaderPipelineTypes, VFTypes, GMaxRHIShaderPlatform);
 		}
-	}
 
-	FMaterialShaderMapId ResourceId;
-	Resource->GetShaderMapId(GMaxRHIShaderPlatform, ResourceId);
+		ResourceId.Usage = Usage;
 
+		CacheShaders(ResourceId, GMaxRHIShaderPlatform, true);
+	}
+	else
 	{
-		TArray<FShaderType*> ShaderTypes;
-		TArray<FVertexFactoryType*> VFTypes;
-		TArray<const FShaderPipelineType*> ShaderPipelineTypes;
-		GetDependentShaderAndVFTypes(GMaxRHIShaderPlatform, ShaderTypes, ShaderPipelineTypes, VFTypes);
-
-		// Overwrite the shader map Id's dependencies with ones that came from the FMaterial actually being compiled (this)
-		// This is necessary as we change FMaterial attributes like GetShadingModel(), which factor into the ShouldCache functions that determine dependent shader types
-		ResourceId.SetShaderDependencies(ShaderTypes, ShaderPipelineTypes, VFTypes, GMaxRHIShaderPlatform);
+		bValid = false;
 	}
-
-	ResourceId.Usage = Usage;
-
-	CacheShaders(ResourceId, GMaxRHIShaderPlatform, true);
 }
 
 bool FDebugViewModeMaterialProxy::RequiresSynchronousCompilation() const
