Using last processed packet time instead of current time when checking for available fast forward data

#rnx

diff --git a/Engine/Source/Runtime/Engine/Private/DemoNetDriver.cpp b/Engine/Source/Runtime/Engine/Private/DemoNetDriver.cpp
--- a/Engine/Source/Runtime/Engine/Private/DemoNetDriver.cpp
+++ b/Engine/Source/Runtime/Engine/Private/DemoNetDriver.cpp
@@ -539,7 +539,8 @@ public:
 		else if (GotoResult.IsSet())
 		{
 			// if this task is not pausing the rest of the replay stream, make sure there is data available for the current time or we could miss packets
-			const uint32 AvailableDataEndTime = (CVarFastForwardLevelsPausePlayback.GetValueOnAnyThread() != 0) ? GotoTime : Driver->DemoCurrentTime * 1000;
+			const float LastProcessedPacketTime = FPendingTaskHelper::GetLastProcessedPacketTime(Driver.Get());
+			const uint32 AvailableDataEndTime = (CVarFastForwardLevelsPausePlayback.GetValueOnAnyThread() != 0) ? GotoTime : LastProcessedPacketTime * 1000;
 
 			if (!GotoResult->WasSuccessful())
 			{
