Add net.SkipReplicatorForDestructionInfos to avoid creation of FObjectReplicator for destruction infos

diff --git a/Engine/Plugins/Runtime/ReplicationGraph/Source/Private/ReplicationGraph.cpp b/Engine/Plugins/Runtime/ReplicationGraph/Source/Private/ReplicationGraph.cpp
--- a/Engine/Plugins/Runtime/ReplicationGraph/Source/Private/ReplicationGraph.cpp
+++ b/Engine/Plugins/Runtime/ReplicationGraph/Source/Private/ReplicationGraph.cpp
@@ -1699,7 +1699,7 @@ int64 UReplicationGraph::ReplicateSingleActor(AActor* Actor, FConnectionReplicat
 					
 		// This will unfortunately cause a callback to this  UNetReplicationGraphConnection and will relook up the ActorInfoMap and set the channel that we already have set.
 		// This is currently unavoidable because channels are created from different code paths (some outside of this loop)
-		ActorInfo.Channel->SetChannelActor( Actor );
+		ActorInfo.Channel->SetChannelActor(Actor, ESetChannelActorFlags::None);
 	}
 
 	if (UNLIKELY(bWantsToGoDormant))
@@ -1992,7 +1992,7 @@ bool UReplicationGraph::ProcessRemoteFunction(class AActor* Actor, UFunction* Fu
 					{
 						// We are within range, we will open a channel now for this actor and call the RPC on it
 						ConnectionActorInfo.Channel = (UActorChannel *)NetConnection->CreateChannelByName( NAME_Actor, EChannelCreateFlags::OpenedLocally );
-						ConnectionActorInfo.Channel->SetChannelActor(Actor);
+						ConnectionActorInfo.Channel->SetChannelActor(Actor, ESetChannelActorFlags::None);
 
 						// Update timeout frame name. We would run into problems if we open the channel, queue a bunch, and then it timeouts before RepGraph replicates properties.
 						UpdateActorChannelCloseFrameNum(Actor, ConnectionActorInfo, GlobalInfo, ReplicationGraphFrame+1 /** Plus one to error on safe side. RepFrame num will be incremented in the next tick */, NetConnection );
@@ -2075,7 +2075,7 @@ bool UReplicationGraph::ProcessRemoteFunction(class AActor* Actor, UFunction* Fu
 			}
 
 			Ch = (UActorChannel *)Connection->CreateChannelByName( NAME_Actor, EChannelCreateFlags::OpenedLocally );
-			Ch->SetChannelActor(Actor);
+			Ch->SetChannelActor(Actor, ESetChannelActorFlags::None);
 		}
 
 		NetDriver->ProcessRemoteFunctionForChannel(Ch, ClassCache, FieldCache, TargetObj, Connection, Function, Parameters, OutParms, Stack, true);
diff --git a/Engine/Source/Runtime/Engine/Classes/Engine/ActorChannel.h b/Engine/Source/Runtime/Engine/Classes/Engine/ActorChannel.h
--- a/Engine/Source/Runtime/Engine/Classes/Engine/ActorChannel.h
+++ b/Engine/Source/Runtime/Engine/Classes/Engine/ActorChannel.h
@@ -16,6 +16,14 @@ class FNetFieldExportGroup;
 class FOutBunch;
 class UNetConnection;
 
+enum class ESetChannelActorFlags : uint32
+{
+	None					= 0,
+	SkipReplicatorCreation	= (1 << 0),
+};
+
+ENUM_CLASS_FLAGS(ESetChannelActorFlags);
+
 /**
  * A channel for exchanging actor and its subobject's properties and RPCs.
  *
@@ -154,7 +162,9 @@ public:
 	 * It's expected that InActor is either null (releasing the channel's reference) or
 	 * a valid actor that is not PendingKill or PendingKillPending.
 	 */
-	void SetChannelActor( AActor* InActor );
+	UE_DEPRECATED(4.23, "SetChannelActor has an additional parameter.")
+	void SetChannelActor(AActor* InActor) { SetChannelActor(InActor, ESetChannelActorFlags::None); }
+	void SetChannelActor(AActor* InActor, ESetChannelActorFlags Flags);
 
 	virtual void NotifyActorChannelOpen(AActor* InActor, FInBunch& InBunch);
 
diff --git a/Engine/Source/Runtime/Engine/Private/DataChannel.cpp b/Engine/Source/Runtime/Engine/Private/DataChannel.cpp
--- a/Engine/Source/Runtime/Engine/Private/DataChannel.cpp
+++ b/Engine/Source/Runtime/Engine/Private/DataChannel.cpp
@@ -1,3 +1,4 @@
+
 // Copyright 1998-2019 Epic Games, Inc. All Rights Reserved.
 
 /*=============================================================================
@@ -66,6 +67,11 @@ FAutoConsoleVariableRef CVarNetPartialBunchReliableThreshold(
 	GCVarNetPartialBunchReliableThreshold,
 	TEXT("If a bunch is broken up into this many partial bunches are more, we will send it reliable even if the original bunch was not reliable. Partial bunches are atonmic and must all make it over to be used"));
 
+int32 GSkipReplicatorForDestructionInfos = 1;
+FAutoConsoleVariableRef CVarNetSkipReplicatorForDestructionInfos(
+	TEXT("net.SkipReplicatorForDestructionInfos"),
+	GSkipReplicatorForDestructionInfos,
+	TEXT("If enabled, skip creation of object replicator in SetChannelActor when we know there is no content payload and we're going to immediately destroy the actor."));
 
 extern TAutoConsoleVariable<int32> CVarFilterGuidRemapping;
 
@@ -2042,10 +2048,10 @@ void UActorChannel::ReceivedNak( int32 NakPacketId )
 	}
 }
 
-void UActorChannel::SetChannelActor( AActor* InActor )
+void UActorChannel::SetChannelActor(AActor* InActor, ESetChannelActorFlags Flags)
 {
 	check(!Closing);
-	check(Actor==NULL);
+	check(Actor == nullptr);
 
 	// Sanity check that the actor is in the same level collection as the channel's driver.
 	const UWorld* const World = Connection->Driver ? Connection->Driver->GetWorld() : nullptr;
@@ -2070,22 +2076,22 @@ void UActorChannel::SetChannelActor( AActor* InActor )
 
 	UE_LOG(LogNetTraffic, VeryVerbose, TEXT("SetChannelActor: ChIndex: %i, Actor: %s, NetGUID: %s"), ChIndex, Actor ? *Actor->GetFullName() : TEXT("NULL"), *ActorNetGUID.ToString() );
 
-	if ( ChIndex >= 0 && Connection->PendingOutRec[ChIndex] > 0 )
+	if (ChIndex >= 0 && Connection->PendingOutRec[ChIndex] > 0)
 	{
 		// send empty reliable bunches to synchronize both sides
 		// UE_LOG(LogNetTraffic, Log, TEXT("%i Actor %s WILL BE sending %i vs first %i"), ChIndex, *Actor->GetName(), Connection->PendingOutRec[ChIndex],Connection->OutReliable[ChIndex]);
 		int32 RealOutReliable = Connection->OutReliable[ChIndex];
 		Connection->OutReliable[ChIndex] = Connection->PendingOutRec[ChIndex] - 1;
-		while ( Connection->PendingOutRec[ChIndex] <= RealOutReliable )
+		while (Connection->PendingOutRec[ChIndex] <= RealOutReliable)
 		{
 			// UE_LOG(LogNetTraffic, Log, TEXT("%i SYNCHRONIZING by sending %i"), ChIndex, Connection->PendingOutRec[ChIndex]);
 
-			FOutBunch Bunch( this, 0 );
+			FOutBunch Bunch(this, 0);
 
 			if (!Bunch.IsError())
 			{
 				Bunch.bReliable = true;
-				SendBunch( &Bunch, 0 );
+				SendBunch(&Bunch, 0);
 				Connection->PendingOutRec[ChIndex]++;
 			}
 			else
@@ -2103,12 +2109,15 @@ void UActorChannel::SetChannelActor( AActor* InActor )
 	if (Actor)
 	{
 		// Add to map.
-		Connection->AddActorChannel( Actor, this );
+		Connection->AddActorChannel(Actor, this);
 
-		check( !ReplicationMap.Contains( Actor ) );
+		check(!ReplicationMap.Contains(Actor));
 
 		// Create the actor replicator, and store a quick access pointer to it
-		ActorReplicator = FindOrCreateReplicator( Actor );
+		if (!EnumHasAnyFlags(Flags, ESetChannelActorFlags::SkipReplicatorCreation))
+		{
+			ActorReplicator = FindOrCreateReplicator(Actor);
+		}
 
 		// Remove from connection's dormancy lists
 		Connection->Driver->GetNetworkObjectList().MarkActive(Actor, Connection, Connection->Driver);
@@ -2471,7 +2480,7 @@ void UActorChannel::ProcessBunch( FInBunch & Bunch )
 					UE_LOG(LogNet, Verbose, TEXT("UActorChannel::ProcessBunch: SerializeNewActor received close bunch for destroyed actor. Actor: %s, Channel: %i"), *GetFullNameSafe(NewChannelActor), ChIndex);
 				}
 
-				SetChannelActor(nullptr);
+				SetChannelActor(nullptr, ESetChannelActorFlags::None);
 				return;
 			}
 
@@ -2490,8 +2499,14 @@ void UActorChannel::ProcessBunch( FInBunch & Bunch )
 			return;
 		}
 
-		UE_LOG(LogNetTraffic, Log, TEXT("      Channel Actor %s:"), *NewChannelActor->GetFullName() );
-		SetChannelActor( NewChannelActor );
+		ESetChannelActorFlags Flags = ESetChannelActorFlags::None;
+		if (GSkipReplicatorForDestructionInfos != 0 && Bunch.bClose && Bunch.AtEnd())
+		{
+			Flags |= ESetChannelActorFlags::SkipReplicatorCreation;
+		}
+
+		UE_LOG(LogNetTraffic, Log, TEXT("      Channel Actor %s:"), *NewChannelActor->GetFullName());
+		SetChannelActor(NewChannelActor, Flags);
 
 		NotifyActorChannelOpen(Actor, Bunch);
 
diff --git a/Engine/Source/Runtime/Engine/Private/DemoNetDriver.cpp b/Engine/Source/Runtime/Engine/Private/DemoNetDriver.cpp
--- a/Engine/Source/Runtime/Engine/Private/DemoNetDriver.cpp
+++ b/Engine/Source/Runtime/Engine/Private/DemoNetDriver.cpp
@@ -1668,7 +1668,7 @@ bool UDemoNetDriver::DemoReplicateActor(AActor* Actor, UNetConnection* Connectio
 			Channel = (UActorChannel*)Connection->CreateChannelByName(NAME_Actor, EChannelCreateFlags::OpenedLocally);
 			if (Channel != nullptr)
 			{
-				Channel->SetChannelActor(Actor);
+				Channel->SetChannelActor(Actor, ESetChannelActorFlags::None);
 			}
 		}
 
diff --git a/Engine/Source/Runtime/Engine/Private/NetworkDriver.cpp b/Engine/Source/Runtime/Engine/Private/NetworkDriver.cpp
--- a/Engine/Source/Runtime/Engine/Private/NetworkDriver.cpp
+++ b/Engine/Source/Runtime/Engine/Private/NetworkDriver.cpp
@@ -1835,7 +1835,7 @@ void UNetDriver::InternalProcessRemoteFunction
 		}
 		if (IsServer)
 		{
-			Ch->SetChannelActor(Actor);
+			Ch->SetChannelActor(Actor, ESetChannelActorFlags::None);
 		}	
 	}
 
@@ -4007,7 +4007,7 @@ int32 UNetDriver::ServerReplicateActors_ProcessPrioritizedActors( UNetConnection
 						Channel = (UActorChannel*)Connection->CreateChannelByName( NAME_Actor, EChannelCreateFlags::OpenedLocally );
 						if ( Channel )
 						{
-							Channel->SetChannelActor( Actor );
+							Channel->SetChannelActor(Actor, ESetChannelActorFlags::None);
 						}
 					}
 					// if we couldn't replicate it for a reason that should be temporary, and this Actor is updated very infrequently, make sure we update it again soon
