[LLM] another pass over EngineMisc. moves about 15MB out in my test replay, mostly moved to LoadMapMisc and StreamingManager.

diff --git a/Engine/Source/Runtime/ClothingSystemRuntime/Private/ClothingSimulation.cpp b/Engine/Source/Runtime/ClothingSystemRuntime/Private/ClothingSimulation.cpp
--- a/Engine/Source/Runtime/ClothingSystemRuntime/Private/ClothingSimulation.cpp
+++ b/Engine/Source/Runtime/ClothingSystemRuntime/Private/ClothingSimulation.cpp
@@ -7,6 +7,7 @@
 #include "ClothingSimulationInterface.h"
 #include "ClothingSystemRuntimeModule.h"
 #include "Assets/ClothingAsset.h"
+#include "HAL/LowLevelMemTracker.h"
 
 static TAutoConsoleVariable<float> GClothMaxDeltaTimeTeleportMultiplier(
 	TEXT("p.Cloth.MaxDeltaTimeTeleportMultiplier"),
@@ -135,6 +136,8 @@ void FClothingSimulationBase::SkinPhysicsMesh(UClothingAsset* InAsset, const FCl
 
 void FClothingSimulationBase::FillContext(USkeletalMeshComponent* InComponent, float InDeltaTime, IClothingSimulationContext* InOutContext)
 {
+	LLM_SCOPE(ELLMTag::SkeletalMesh);
+
 	FClothingSimulationContextBase* BaseContext = static_cast<FClothingSimulationContextBase*>(InOutContext);
 	BaseContext->ComponentToWorld = InComponent->GetComponentTransform();
 	BaseContext->PredictedLod = InComponent->PredictedLODLevel;
diff --git a/Engine/Source/Runtime/CoreUObject/Private/UObject/Class.cpp b/Engine/Source/Runtime/CoreUObject/Private/UObject/Class.cpp
--- a/Engine/Source/Runtime/CoreUObject/Private/UObject/Class.cpp
+++ b/Engine/Source/Runtime/CoreUObject/Private/UObject/Class.cpp
@@ -6,6 +6,7 @@
 
 #include "UObject/Class.h"
 #include "HAL/ThreadSafeBool.h"
+#include "HAL/LowLevelMemTracker.h"
 #include "Misc/ScopeLock.h"
 #include "Serialization/MemoryWriter.h"
 #include "Misc/ConfigCacheIni.h"
@@ -4327,6 +4328,7 @@ void UClass::ClearFunctionMapsCaches()
 
 UFunction* UClass::FindFunctionByName(FName InName, EIncludeSuperFlag::Type IncludeSuper) const
 {
+	LLM_SCOPE(ELLMTag::UObject);
 	UFunction* Result = FuncMap.FindRef(InName);
 	if (Result == nullptr && IncludeSuper == EIncludeSuperFlag::IncludeSuper)
 	{
diff --git a/Engine/Source/Runtime/Engine/Classes/Components/StaticMeshComponent.h b/Engine/Source/Runtime/Engine/Classes/Components/StaticMeshComponent.h
--- a/Engine/Source/Runtime/Engine/Classes/Components/StaticMeshComponent.h
+++ b/Engine/Source/Runtime/Engine/Classes/Components/StaticMeshComponent.h
@@ -440,6 +440,7 @@ public:
 protected: 
 	virtual void OnRegister() override;
 	virtual void OnUnregister() override;
+	virtual void CreateRenderState_Concurrent() override;
 	virtual void OnCreatePhysicsState() override;
 	virtual void OnDestroyPhysicsState() override;
 public:
@@ -680,6 +681,8 @@ public:
 private:
 	FOnStaticMeshChanged OnStaticMeshChangedEvent;
 #endif
+
+	friend class FStaticMeshComponentRecreateRenderStateContext;
 };
 
 /** Vertex data stored per-LOD */
diff --git a/Engine/Source/Runtime/Engine/Private/Animation/AnimInstance.cpp b/Engine/Source/Runtime/Engine/Private/Animation/AnimInstance.cpp
--- a/Engine/Source/Runtime/Engine/Private/Animation/AnimInstance.cpp
+++ b/Engine/Source/Runtime/Engine/Private/Animation/AnimInstance.cpp
@@ -1757,6 +1757,8 @@ bool UAnimInstance::IsPlayingSlotAnimation(const UAnimSequenceBase* Asset, FName
 /** Play a Montage. Returns Length of Montage in seconds. Returns 0.f if failed to play. */
 float UAnimInstance::Montage_Play(UAnimMontage* MontageToPlay, float InPlayRate/*= 1.f*/, EMontagePlayReturnType ReturnValueType, float InTimeToStartMontageAt, bool bStopAllMontages /*= true*/)
 {
+	LLM_SCOPE(ELLMTag::Animation);
+
 	if (MontageToPlay && (MontageToPlay->SequenceLength > 0.f) && MontageToPlay->HasValidSlotSetup())
 	{
 		if (CurrentSkeleton && CurrentSkeleton->IsCompatible(MontageToPlay->GetSkeleton()))
diff --git a/Engine/Source/Runtime/Engine/Private/Components/ActorComponent.cpp b/Engine/Source/Runtime/Engine/Private/Components/ActorComponent.cpp
--- a/Engine/Source/Runtime/Engine/Private/Components/ActorComponent.cpp
+++ b/Engine/Source/Runtime/Engine/Private/Components/ActorComponent.cpp
@@ -29,6 +29,7 @@
 #include "Engine/SimpleConstructionScript.h"
 #include "ComponentUtils.h"
 #include "Engine/Engine.h"
+#include "HAL/LowLevelMemTracker.h"
 
 #if WITH_EDITOR
 #include "Kismet2/ComponentEditorUtils.h"
@@ -1241,6 +1242,7 @@ void UActorComponent::OnDestroyPhysicsState()
 
 void UActorComponent::CreatePhysicsState()
 {
+	LLM_SCOPE(ELLMTag::PhysX);
 	SCOPE_CYCLE_COUNTER(STAT_ComponentCreatePhysicsState);
 
 	if (!bPhysicsStateCreated && WorldPrivate->GetPhysicsScene() && ShouldCreatePhysicsState())
@@ -1286,6 +1288,7 @@ void UActorComponent::ExecuteRegisterEvents()
 	if(FApp::CanEverRender() && !bRenderStateCreated && WorldPrivate->Scene && ShouldCreateRenderState())
 	{
 		SCOPE_CYCLE_COUNTER(STAT_ComponentCreateRenderState);
+		LLM_SCOPE(ELLMTag::SceneRender);
 		CreateRenderState_Concurrent();
 		checkf(bRenderStateCreated, TEXT("Failed to route CreateRenderState_Concurrent (%s)"), *GetFullName());
 	}
@@ -1391,6 +1394,8 @@ void UActorComponent::RemoveTickPrerequisiteComponent(UActorComponent* Prerequis
 
 void UActorComponent::DoDeferredRenderUpdates_Concurrent()
 {
+	LLM_SCOPE(ELLMTag::SceneRender);
+
 	checkf(!IsUnreachable(), TEXT("%s"), *GetFullName());
 	checkf(!IsTemplate(), TEXT("%s"), *GetFullName());
 	checkf(!IsPendingKill(), TEXT("%s"), *GetFullName());
diff --git a/Engine/Source/Runtime/Engine/Private/Components/DecalComponent.cpp b/Engine/Source/Runtime/Engine/Private/Components/DecalComponent.cpp
--- a/Engine/Source/Runtime/Engine/Private/Components/DecalComponent.cpp
+++ b/Engine/Source/Runtime/Engine/Private/Components/DecalComponent.cpp
@@ -9,6 +9,7 @@
 #include "TimerManager.h"
 #include "SceneManagement.h"
 #include "Materials/MaterialInstanceDynamic.h"
+#include "HAL/LowLevelMemTracker.h"
 
 static TAutoConsoleVariable<float> CVarDecalFadeDurationScale(
 	TEXT("r.Decal.FadeDurationScale"),
@@ -263,6 +264,7 @@ void UDecalComponent::GetUsedMaterials( TArray<UMaterialInterface*>& OutMaterial
 
 FDeferredDecalProxy* UDecalComponent::CreateSceneProxy()
 {
+	LLM_SCOPE(ELLMTag::SceneRender);
 	return new FDeferredDecalProxy(this);
 }
 
diff --git a/Engine/Source/Runtime/Engine/Private/Components/SkeletalMeshComponent.cpp b/Engine/Source/Runtime/Engine/Private/Components/SkeletalMeshComponent.cpp
--- a/Engine/Source/Runtime/Engine/Private/Components/SkeletalMeshComponent.cpp
+++ b/Engine/Source/Runtime/Engine/Private/Components/SkeletalMeshComponent.cpp
@@ -595,6 +595,7 @@ void USkeletalMeshComponent::OnUnregister()
 void USkeletalMeshComponent::InitAnim(bool bForceReinit)
 {
 	QUICK_SCOPE_CYCLE_COUNTER(STAT_SkelMeshComp_InitAnim);
+	LLM_SCOPE(ELLMTag::Animation);
 
 	// a lot of places just call InitAnim without checking Mesh, so 
 	// I'm moving the check here
@@ -2588,6 +2589,8 @@ void USkeletalMeshComponent::SetSkeletalMeshWithoutResettingAnimation(USkeletalM
 
 bool USkeletalMeshComponent::AllocateTransformData()
 {
+	LLM_SCOPE(ELLMTag::SkeletalMesh);
+
 	// Allocate transforms if not present.
 	if ( Super::AllocateTransformData() )
 	{
diff --git a/Engine/Source/Runtime/Engine/Private/Components/SkinnedMeshComponent.cpp b/Engine/Source/Runtime/Engine/Private/Components/SkinnedMeshComponent.cpp
--- a/Engine/Source/Runtime/Engine/Private/Components/SkinnedMeshComponent.cpp
+++ b/Engine/Source/Runtime/Engine/Private/Components/SkinnedMeshComponent.cpp
@@ -458,6 +458,8 @@ void USkinnedMeshComponent::Serialize(FArchive& Ar)
 
 void USkinnedMeshComponent::OnRegister()
 {
+	LLM_SCOPE(ELLMTag::SkeletalMesh);
+
 	// The reason this happens before register
 	// is so that any transform update (or children transform update)
 	// won't result in any issues of accessing SpaceBases
@@ -1417,6 +1419,8 @@ FSkeletalMeshRenderData* USkinnedMeshComponent::GetSkeletalMeshRenderData() cons
 
 bool USkinnedMeshComponent::AllocateTransformData()
 {
+	LLM_SCOPE(ELLMTag::SkeletalMesh);
+
 	// Allocate transforms if not present.
 	if ( SkeletalMesh != NULL && MasterPoseComponent == NULL )
 	{
diff --git a/Engine/Source/Runtime/Engine/Private/Components/StaticMeshComponent.cpp b/Engine/Source/Runtime/Engine/Private/Components/StaticMeshComponent.cpp
--- a/Engine/Source/Runtime/Engine/Private/Components/StaticMeshComponent.cpp
+++ b/Engine/Source/Runtime/Engine/Private/Components/StaticMeshComponent.cpp
@@ -602,6 +602,12 @@ void UStaticMeshComponent::OnUnregister()
 	Super::OnUnregister();
 }
 
+void UStaticMeshComponent::CreateRenderState_Concurrent()
+{
+	LLM_SCOPE(ELLMTag::StaticMesh);
+	Super::CreateRenderState_Concurrent();
+}
+
 void UStaticMeshComponent::OnCreatePhysicsState()
 {
 	Super::OnCreatePhysicsState();
diff --git a/Engine/Source/Runtime/Engine/Private/DemoNetDriver.cpp b/Engine/Source/Runtime/Engine/Private/DemoNetDriver.cpp
--- a/Engine/Source/Runtime/Engine/Private/DemoNetDriver.cpp
+++ b/Engine/Source/Runtime/Engine/Private/DemoNetDriver.cpp
@@ -1116,6 +1116,8 @@ bool UDemoNetDriver::InitListen(FNetworkNotify* InNotify, FURL& ListenURL, bool
 
 void UDemoNetDriver::OnLevelAddedToWorld(ULevel* InLevel, UWorld* InWorld)
 {
+	LLM_SCOPE(ELLMTag::Networking);
+
 	Super::OnLevelAddedToWorld(InLevel, InWorld);
 
 	if (InLevel && !InLevel->bClientOnlyVisible && GetWorld() == InWorld && HasLevelStreamingFixes() && IsPlaying())
@@ -1358,6 +1360,7 @@ double GTickFlushDemoDriverTimeSeconds = 0.0;
 
 void UDemoNetDriver::TickFlushInternal(float DeltaSeconds)
 {
+	LLM_SCOPE(ELLMTag::Networking);
 	CSV_SCOPED_TIMING_STAT_EXCLUSIVE(DemoRecording);
 
 	GTickFlushDemoDriverTimeSeconds = 0.0;
@@ -3736,6 +3739,7 @@ static FCsvDemoSettings GetCsvDemoSettings()
 
 void UDemoNetDriver::TickDemoPlayback(float DeltaSeconds)
 {
+	LLM_SCOPE(ELLMTag::Networking);
 	SCOPED_NAMED_EVENT(UDemoNetDriver_TickDemoPlayback, FColor::Purple);
 	if (World && World->IsInSeamlessTravel())
 	{
diff --git a/Engine/Source/Runtime/Engine/Private/Level.cpp b/Engine/Source/Runtime/Engine/Private/Level.cpp
--- a/Engine/Source/Runtime/Engine/Private/Level.cpp
+++ b/Engine/Source/Runtime/Engine/Private/Level.cpp
@@ -58,6 +58,7 @@ Level.cpp: Level-related functions
 #include "Engine/StaticMeshActor.h"
 #include "ComponentRecreateRenderStateContext.h"
 #include "Algo/Copy.h"
+#include "HAL/LowLevelMemTracker.h"
 
 DEFINE_LOG_CATEGORY(LogLevel);
 
@@ -469,6 +470,7 @@ void ULevel::SortActorList()
 		// No need to sort an empty list
 		return;
 	}
+	LLM_REALLOC_SCOPE(Actors.GetData());
 
 	TArray<AActor*> NewActors;
 	TArray<AActor*> NewNetActors;
diff --git a/Engine/Source/Runtime/Engine/Private/Materials/MaterialInstance.cpp b/Engine/Source/Runtime/Engine/Private/Materials/MaterialInstance.cpp
--- a/Engine/Source/Runtime/Engine/Private/Materials/MaterialInstance.cpp
+++ b/Engine/Source/Runtime/Engine/Private/Materials/MaterialInstance.cpp
@@ -1544,6 +1544,7 @@ EMaterialShadingModel UMaterialInstanceDynamic::GetShadingModel() const
 
 void UMaterialInstance::CopyMaterialInstanceParameters(UMaterialInterface* Source)
 {
+	LLM_SCOPE(ELLMTag::MaterialInstance);
 	SCOPE_CYCLE_COUNTER(STAT_MaterialInstance_CopyMatInstParams);
 
 	if ((Source != nullptr) && (Source != this))
@@ -3375,6 +3376,8 @@ bool UMaterialInstance::SetVectorParameterByIndexInternal(int32 ParameterIndex,
 
 void UMaterialInstance::SetVectorParameterValueInternal(const FMaterialParameterInfo& ParameterInfo, FLinearColor Value)
 {
+	LLM_SCOPE(ELLMTag::MaterialInstance);
+
 	FVectorParameterValue* ParameterValue = GameThread_FindParameterByName(VectorParameterValues, ParameterInfo);
 
 	if(!ParameterValue)
@@ -3418,6 +3421,8 @@ bool UMaterialInstance::SetScalarParameterByIndexInternal(int32 ParameterIndex,
 
 void UMaterialInstance::SetScalarParameterValueInternal(const FMaterialParameterInfo& ParameterInfo, float Value)
 {
+	LLM_SCOPE(ELLMTag::MaterialInstance);
+
 	FScalarParameterValue* ParameterValue = GameThread_FindParameterByName(ScalarParameterValues, ParameterInfo);
 
 	if(!ParameterValue)
@@ -3476,6 +3481,8 @@ void UMaterialInstance::SetScalarParameterAtlasInternal(const FMaterialParameter
 
 void UMaterialInstance::SetTextureParameterValueInternal(const FMaterialParameterInfo& ParameterInfo, UTexture* Value)
 {
+	LLM_SCOPE(ELLMTag::MaterialInstance);
+
 	FTextureParameterValue* ParameterValue = GameThread_FindParameterByName(TextureParameterValues, ParameterInfo);
 
 	if(!ParameterValue)
@@ -3504,6 +3511,8 @@ void UMaterialInstance::SetTextureParameterValueInternal(const FMaterialParamete
 
 void UMaterialInstance::SetFontParameterValueInternal(const FMaterialParameterInfo& ParameterInfo,class UFont* FontValue,int32 FontPage)
 {
+	LLM_SCOPE(ELLMTag::MaterialInstance);
+
 	FFontParameterValue* ParameterValue = GameThread_FindParameterByName(FontParameterValues, ParameterInfo);
 
 	if(!ParameterValue)
@@ -4264,6 +4273,7 @@ void UMaterialInstance::SaveShaderStableKeysInner(const class ITargetPlatform* T
 
 void UMaterialInstance::CopyMaterialUniformParametersInternal(UMaterialInterface* Source)
 {
+	LLM_SCOPE(ELLMTag::MaterialInstance);
 	SCOPE_CYCLE_COUNTER(STAT_MaterialInstance_CopyUniformParamsInternal)
 
 	if ((Source == nullptr) || (Source == this))
diff --git a/Engine/Source/Runtime/Engine/Private/Materials/MaterialInstanceDynamic.cpp b/Engine/Source/Runtime/Engine/Private/Materials/MaterialInstanceDynamic.cpp
--- a/Engine/Source/Runtime/Engine/Private/Materials/MaterialInstanceDynamic.cpp
+++ b/Engine/Source/Runtime/Engine/Private/Materials/MaterialInstanceDynamic.cpp
@@ -12,6 +12,7 @@
 #include "UnrealEngine.h"
 #include "Materials/MaterialUniformExpressions.h"
 #include "Stats/StatsMisc.h"
+#include "HAL/LowLevelMemTracker.h"
 
 UMaterialInstanceDynamic::UMaterialInstanceDynamic(const FObjectInitializer& ObjectInitializer)
 	: Super(ObjectInitializer)
@@ -20,6 +21,7 @@ UMaterialInstanceDynamic::UMaterialInstanceDynamic(const FObjectInitializer& Obj
 
 UMaterialInstanceDynamic* UMaterialInstanceDynamic::Create(UMaterialInterface* ParentMaterial, UObject* InOuter)
 {
+	LLM_SCOPE(ELLMTag::MaterialInstance);
 	UObject* Outer = InOuter ? InOuter : GetTransientPackage();
 	UMaterialInstanceDynamic* MID = NewObject<UMaterialInstanceDynamic>(Outer);
 	MID->SetParentInternal(ParentMaterial, false);
@@ -28,6 +30,7 @@ UMaterialInstanceDynamic* UMaterialInstanceDynamic::Create(UMaterialInterface* P
 
 UMaterialInstanceDynamic* UMaterialInstanceDynamic::Create(UMaterialInterface* ParentMaterial, UObject* InOuter, FName Name)
 {
+	LLM_SCOPE(ELLMTag::MaterialInstance);
 	UObject* Outer = InOuter ? InOuter : GetTransientPackage();
 	UMaterialInstanceDynamic* MID = NewObject<UMaterialInstanceDynamic>(Outer, Name);
 	MID->SetParentInternal(ParentMaterial, false);
@@ -324,6 +327,8 @@ void UMaterialInstanceDynamic::CopyInterpParameters(UMaterialInstance* Source)
 
 void UMaterialInstanceDynamic::CopyParameterOverrides(UMaterialInstance* MaterialInstance)
 {
+	LLM_SCOPE(ELLMTag::MaterialInstance);
+
 	ClearParameterValues();
 	if (ensureAsRuntimeWarning(MaterialInstance != nullptr))
 	{
diff --git a/Engine/Source/Runtime/Engine/Private/Materials/MaterialInstanceSupport.h b/Engine/Source/Runtime/Engine/Private/Materials/MaterialInstanceSupport.h
--- a/Engine/Source/Runtime/Engine/Private/Materials/MaterialInstanceSupport.h
+++ b/Engine/Source/Runtime/Engine/Private/Materials/MaterialInstanceSupport.h
@@ -10,6 +10,7 @@
 #include "RenderingThread.h"
 #include "MaterialShared.h"
 #include "Materials/MaterialInstance.h"
+#include "HAL/LowLevelMemTracker.h"
 
 class UTexture;
 
@@ -117,6 +118,8 @@ public:
 	template <typename ValueType>
 	void RenderThread_UpdateParameter(const FMaterialParameterInfo& ParameterInfo, const ValueType& Value )
 	{
+		LLM_SCOPE(ELLMTag::MaterialInstance);
+
 		InvalidateUniformExpressionCache(false);
 		TArray<TNamedParameter<ValueType> >& ValueArray = GetValueArray<ValueType>();
 		const int32 ParameterCount = ValueArray.Num();
diff --git a/Engine/Source/Runtime/Engine/Private/NetworkDriver.cpp b/Engine/Source/Runtime/Engine/Private/NetworkDriver.cpp
--- a/Engine/Source/Runtime/Engine/Private/NetworkDriver.cpp
+++ b/Engine/Source/Runtime/Engine/Private/NetworkDriver.cpp
@@ -11,6 +11,7 @@
 #include "Stats/Stats.h"
 #include "Misc/MemStack.h"
 #include "HAL/IConsoleManager.h"
+#include "HAL/LowLevelMemTracker.h"
 #include "UObject/ObjectMacros.h"
 #include "UObject/Object.h"
 #include "UObject/Class.h"
@@ -392,6 +393,7 @@ void UNetDriver::AssertValid()
 
 void UNetDriver::AddNetworkActor(AActor* Actor)
 {
+	LLM_SCOPE(ELLMTag::Networking);
 	GetNetworkObjectList().FindOrAdd(Actor, this);
 	if (ReplicationDriver)
 	{
diff --git a/Engine/Source/Runtime/Engine/Private/Particles/ParticleComponents.cpp b/Engine/Source/Runtime/Engine/Private/Particles/ParticleComponents.cpp
--- a/Engine/Source/Runtime/Engine/Private/Particles/ParticleComponents.cpp
+++ b/Engine/Source/Runtime/Engine/Private/Particles/ParticleComponents.cpp
@@ -382,6 +382,8 @@ void UParticleLODLevel::PostLoad()
 
 void UParticleLODLevel::UpdateModuleLists()
 {
+	LLM_SCOPE(ELLMTag::Particles);
+
 	SpawningModules.Empty();
 	SpawnModules.Empty();
 	UpdateModules.Empty();
@@ -6973,6 +6975,8 @@ int32 UParticleSystemComponent::GetLODLevel()
  */
 void UParticleSystemComponent::SetFloatParameter(FName Name, float Param)
 {
+	LLM_SCOPE(ELLMTag::Particles);
+
 	if(Name == NAME_None)
 	{
 		return;
@@ -7000,6 +7004,8 @@ void UParticleSystemComponent::SetFloatParameter(FName Name, float Param)
 
 void UParticleSystemComponent::SetFloatRandParameter(FName ParameterName,float Param,float ParamLow)
 {
+	LLM_SCOPE(ELLMTag::Particles);
+
 	if (ParameterName == NAME_None)
 	{
 		return;
@@ -7029,6 +7035,8 @@ void UParticleSystemComponent::SetFloatRandParameter(FName ParameterName,float P
 
 void UParticleSystemComponent::SetVectorParameter(FName Name, FVector Param)
 {
+	LLM_SCOPE(ELLMTag::Particles);
+
 	if (Name == NAME_None)
 	{
 		return;
@@ -7056,6 +7064,8 @@ void UParticleSystemComponent::SetVectorParameter(FName Name, FVector Param)
 
 void UParticleSystemComponent::SetVectorRandParameter(FName ParameterName,const FVector& Param,const FVector& ParamLow)
 {
+	LLM_SCOPE(ELLMTag::Particles);
+
 	if (ParameterName == NAME_None)
 	{
 		return;
@@ -7085,6 +7095,8 @@ void UParticleSystemComponent::SetVectorRandParameter(FName ParameterName,const
 
 void UParticleSystemComponent::SetColorParameter(FName Name, FLinearColor Param)
 {
+	LLM_SCOPE(ELLMTag::Particles);
+
 	if(Name == NAME_None)
 	{
 		return;
@@ -7114,6 +7126,8 @@ void UParticleSystemComponent::SetColorParameter(FName Name, FLinearColor Param)
 
 void UParticleSystemComponent::SetActorParameter(FName Name, AActor* Param)
 {
+	LLM_SCOPE(ELLMTag::Particles);
+
 	if(Name == NAME_None)
 	{
 		return;
@@ -7141,6 +7155,8 @@ void UParticleSystemComponent::SetActorParameter(FName Name, AActor* Param)
 
 void UParticleSystemComponent::SetMaterialParameter(FName Name, UMaterialInterface* Param)
 {
+	LLM_SCOPE(ELLMTag::Particles);
+
 	if(Name == NAME_None)
 	{
 		return;
diff --git a/Engine/Source/Runtime/Engine/Private/StreamableManager.cpp b/Engine/Source/Runtime/Engine/Private/StreamableManager.cpp
--- a/Engine/Source/Runtime/Engine/Private/StreamableManager.cpp
+++ b/Engine/Source/Runtime/Engine/Private/StreamableManager.cpp
@@ -903,6 +903,8 @@ UObject* FStreamableManager::SynchronousLoad(FSoftObjectPath const& Target)
 
 TSharedPtr<FStreamableHandle> FStreamableManager::RequestAsyncLoad(const TArray<FSoftObjectPath>& TargetsToStream, FStreamableDelegate DelegateToCall, TAsyncLoadPriority Priority, bool bManageActiveHandle, bool bStartStalled, const FString& DebugName)
 {
+	LLM_SCOPE(ELLMTag::StreamingManager);
+
 	// Schedule a new callback, this will get called when all related async loads are completed
 	TSharedRef<FStreamableHandle> NewRequest = MakeShareable(new FStreamableHandle());
 	NewRequest->CompleteDelegate = DelegateToCall;
diff --git a/Engine/Source/Runtime/Engine/Private/World.cpp b/Engine/Source/Runtime/Engine/Private/World.cpp
--- a/Engine/Source/Runtime/Engine/Private/World.cpp
+++ b/Engine/Source/Runtime/Engine/Private/World.cpp
@@ -3072,6 +3072,7 @@ void UWorld::UpdateLevelStreaming()
 {
 	SCOPE_CYCLE_COUNTER(STAT_UpdateLevelStreamingTime);
 	CSV_SCOPED_TIMING_STAT_EXCLUSIVE(UpdateLevelStreaming);
+	LLM_SCOPE(ELLMTag::LoadMapMisc);
 
 	// do nothing if level streaming is frozen
 	if (bIsLevelStreamingFrozen)
diff --git a/Engine/Source/Runtime/Renderer/Private/RendererScene.cpp b/Engine/Source/Runtime/Renderer/Private/RendererScene.cpp
--- a/Engine/Source/Runtime/Renderer/Private/RendererScene.cpp
+++ b/Engine/Source/Runtime/Renderer/Private/RendererScene.cpp
@@ -53,6 +53,7 @@
 #include "GPUSkinCache.h"
 #include "DynamicShadowMapChannelBindingHelper.h"
 #include "GPUScene.h"
+#include "HAL/LowLevelMemTracker.h"
 #if RHI_RAYTRACING
 #include "RayTracingDynamicGeometryCollection.h"
 #include "RHIGPUReadback.h"
@@ -1722,6 +1723,8 @@ void FScene::AddLightSceneInfo_RenderThread(FLightSceneInfo* LightSceneInfo)
 
 void FScene::AddLight(ULightComponent* Light)
 {
+	LLM_SCOPE(ELLMTag::SceneRender);
+
 	// Create the light's scene proxy.
 	FLightSceneProxy* Proxy = Light->CreateSceneProxy();
 	if(Proxy)
diff --git a/Engine/Source/Runtime/Renderer/Private/SceneCore.cpp b/Engine/Source/Runtime/Renderer/Private/SceneCore.cpp
--- a/Engine/Source/Runtime/Renderer/Private/SceneCore.cpp
+++ b/Engine/Source/Runtime/Renderer/Private/SceneCore.cpp
@@ -18,6 +18,7 @@
 #include "ScenePrivate.h"
 #include "Containers/AllocatorFixedSizeFreeList.h"
 #include "MaterialShared.h"
+#include "HAL/LowLevelMemTracker.h"
 
 int32 GUnbuiltPreviewShadowsInGame = 1;
 FAutoConsoleVariableRef CVarUnbuiltPreviewShadowsInGame(
@@ -112,6 +113,8 @@ uint32 FLightPrimitiveInteraction::GetMemoryPoolSize()
 
 void FLightPrimitiveInteraction::Create(FLightSceneInfo* LightSceneInfo,FPrimitiveSceneInfo* PrimitiveSceneInfo)
 {
+	LLM_SCOPE(ELLMTag::SceneRender);
+
 	// Attach the light to the primitive's static meshes.
 	bool bDynamic = true;
 	bool bRelevant = false;
