Remove re-trigger functionality from AudioComponent level as it's no longer needed due to VirtualLoop system added to audio thread.

diff --git a/Engine/Source/Runtime/Engine/Classes/Components/AudioComponent.h b/Engine/Source/Runtime/Engine/Classes/Components/AudioComponent.h
--- a/Engine/Source/Runtime/Engine/Classes/Components/AudioComponent.h
+++ b/Engine/Source/Runtime/Engine/Classes/Components/AudioComponent.h
@@ -162,11 +162,6 @@ class ENGINE_API UAudioComponent : public USceneComponent
 	UPROPERTY(EditAnywhere, BlueprintReadWrite, Category = Sound)
 	uint8 bSuppressSubtitles:1;
 
-	/** If true, re-triggers a loop if a listener comes back within range.
-	  * Only works for attenuated, non-virtualized sounds. */
-	UPROPERTY(EditAnywhere, BlueprintReadWrite, Category = Attenuation)
-	uint8 bRetriggerLoopOnProximity : 1;
-
 	/** Whether this audio component is previewing a sound */
 	uint8 bPreviewComponent:1;
 
@@ -194,23 +189,6 @@ class ENGINE_API UAudioComponent : public USceneComponent
 	/** Whether or not this audio component has been paused */
 	uint8 bIsPaused:1;
 
-	/** How playback management is based on whether a listener is within
-	  * proximity (sound's MaxDistance). */
-	enum class ERetriggerWhenInAudibleRange : uint8
-	{
-		/** Sound lifetime is not managed by listener proximity. */
-		Disabled,
-
-		/** Loop re-trigger is enabled and disable has been requested. */
-		DisableRequested,
-
-		/** Sound lifetime is managed by listener proximity (in case of sound being looping and non-virtualized) */
-		Enabled,
-	};
-
-	/** Whether or not the audio component can restart when listener is in proximity */
-	ERetriggerWhenInAudibleRange RetriggerWhenInAudibleRange;
-
 	/**
 	* True if we should automatically attach to AutoAttachParent when Played, and detach from our parent when playback is completed.
 	* This overrides any current attachment that may be present at the time of activation (deferring initial attachment until activation, if AutoAttachParent is null).
@@ -504,9 +482,6 @@ private:
 	/** Whether or not the sound is audible. */
 	bool IsInAudibleRange(float* OutMaxDistance) const;
 
-	/** Returns true if in re-trigger range, false if not. */
-	bool GetRetriggerRate(float* OutRetriggerRate) const;
-
 public:
 
 	/** Sets the sound instance parameter. */
@@ -544,8 +519,6 @@ public:
 	/** Collects the various attenuation shapes that may be applied to the sound played by the audio component for visualization in the editor or via the in game debug visualization. */
 	void CollectAttenuationShapesForVisualization(TMultiMap<EAttenuationShape::Type, FBaseAttenuationSettings::AttenuationShapeDetails>& ShapeDetailsMap) const;
 
-	void OnUpdateProximityRetrigger(float DeltaTime);
-
 	/** Returns the active audio device to use for this component based on whether or not the component is playing in a world. */
 	FAudioDevice* GetAudioDevice() const;
 
diff --git a/Engine/Source/Runtime/Engine/Private/AudioDevice.cpp b/Engine/Source/Runtime/Engine/Private/AudioDevice.cpp
--- a/Engine/Source/Runtime/Engine/Private/AudioDevice.cpp
+++ b/Engine/Source/Runtime/Engine/Private/AudioDevice.cpp
@@ -563,8 +563,6 @@ void FAudioDevice::Teardown()
 	}
 
 	PluginListeners.Reset();
-
-	ProximityRetriggerComponents.Reset();
 }
 
 void FAudioDevice::Suspend(bool bGameTicking)
@@ -642,24 +640,6 @@ void FAudioDevice::AddReferencedObjects(FReferenceCollector& Collector)
 	Collector.AddReferencedObjects(PluginSettingsObjects);
 }
 
-void FAudioDevice::RegisterProximityRetriggeringAudioComponent(UAudioComponent& Component)
-{
-	check(IsInGameThread());
-	ProximityRetriggerComponents.AddUnique(&Component);
-}
-
-const TArray<FAudioComponentPtr>& FAudioDevice::GetProximityRetriggerComponents() const
-{
-	check(IsInGameThread());
-	return ProximityRetriggerComponents;
-}
-
-void FAudioDevice::UnregisterProximityRetriggeringAudioComponent(UAudioComponent& Component)
-{
-	check(IsInGameThread());
-	ProximityRetriggerComponents.RemoveSwap(&Component);
-}
-
 void FAudioDevice::ResetInterpolation()
 {
 	check(IsInAudioThread());
@@ -3700,12 +3680,6 @@ void FAudioDevice::Update(bool bGameTicking)
 {
 	LLM_SCOPE(ELLMTag::AudioMisc);
 
-	// Must be separate from subsequent call as editor is considered game thread
-	if (IsInGameThread())
-	{
-		UpdateProximityRetriggerComponents();
-	}
-
 	if (!IsInAudioThread())
 	{
 		check(IsInGameThread());
@@ -5789,23 +5763,6 @@ float FAudioDevice::GetGameDeltaTime() const
 	return FMath::Min(DeltaTime, 0.5f);
 }
 
-void FAudioDevice::UpdateProximityRetriggerComponents()
-{
-	for (int32 i = ProximityRetriggerComponents.Num() - 1; i >= 0; --i)
-	{
-		FAudioComponentPtr& ComponentPtr = ProximityRetriggerComponents[i];
-		if (ComponentPtr.IsStale())
-		{
-			UE_LOG(LogAudio, Warning, TEXT("AudioComponent destroyed but not removed from retrigger proximity."));
-			ProximityRetriggerComponents.RemoveAtSwap(i, 1, false);
-		}
-		else if (ComponentPtr.IsValid())
-		{
-			ComponentPtr->OnUpdateProximityRetrigger(GetDeviceDeltaTime());
-		}
-	}
-}
-
 void FAudioDevice::UpdateVirtualLoops()
 {
 	if (FAudioVirtualLoop::IsEnabled())
diff --git a/Engine/Source/Runtime/Engine/Private/AudioDeviceManager.cpp b/Engine/Source/Runtime/Engine/Private/AudioDeviceManager.cpp
--- a/Engine/Source/Runtime/Engine/Private/AudioDeviceManager.cpp
+++ b/Engine/Source/Runtime/Engine/Private/AudioDeviceManager.cpp
@@ -168,8 +168,6 @@ void FAudioDeviceManager::ToggleAudioMixer()
 						}
 					}
 
-					TArray<FAudioComponentPtr> RetriggerComponents = AudioDevice->GetProximityRetriggerComponents();
-
 					// Tear it down and delete the old audio device. This does a bunch of cleanup.
 					AudioDevice->Teardown();
 					delete AudioDevice;
@@ -186,12 +184,6 @@ void FAudioDeviceManager::ToggleAudioMixer()
 						AudioDevice->SetMaxChannels(QualityLevelMaxChannels);
 					}
 
-					// Re-register the new audio device's re-trigger components
-					for (FAudioComponentPtr Component : RetriggerComponents)
-					{
-						AudioDevice->RegisterProximityRetriggeringAudioComponent(*Component);
-					}
-
 					// Transfer the sound mix modifiers to the new audio engine
 					AudioDevice->SetSoundMixModifiers(MixModifiers, PrevPassiveSoundMixModifiers, BaseSoundMix);
 					// Setup the mute state of the audio device to be the same that it was
diff --git a/Engine/Source/Runtime/Engine/Private/Components/AudioComponent.cpp b/Engine/Source/Runtime/Engine/Private/Components/AudioComponent.cpp
--- a/Engine/Source/Runtime/Engine/Private/Components/AudioComponent.cpp
+++ b/Engine/Source/Runtime/Engine/Private/Components/AudioComponent.cpp
@@ -21,28 +21,6 @@ FAutoConsoleVariableRef CVarBackedAnalysisTimeShift(
 	TEXT("Value: The time in seconds to shift the timeline."),
 	ECVF_Default);
 
-static int32 bRetriggerLoopOnProximityEnabledCVar = 0;
-FAutoConsoleVariableRef CVarRetriggerLoopOnProximityEnabled(
-	TEXT("au.RetriggerLoopOnProximity"),
-	bRetriggerLoopOnProximityEnabledCVar,
-	TEXT("Enables or disables whether retriggering is supported on AudioComponent loops.\n")
-	TEXT("Value: Whether or not retriggering is supported on AudioComponent level."),
-	ECVF_Default);
-
-static float RetriggerLoopOnProximityDistanceCVar = 15000.0f;
-FAutoConsoleVariableRef CVarRetriggerLoopOnProximityDistance(
-	TEXT("au.RetriggerLoopOnProximityDistance"),
-	RetriggerLoopOnProximityDistanceCVar,
-	TEXT("Sets retrigger loop proximity distance.\n")
-	TEXT("Value: The distance to scale update rate based on distance to listener."),
-	ECVF_Default);
-
-namespace
-{
-	float AudioComponentRetriggerUpdateRateMin = 0.1f;
-	float AudioComponentRetriggerUpdateRateMax = 3.0f;
-} // namespace <>
-
 
 /*-----------------------------------------------------------------------------
 UAudioComponent implementation.
@@ -72,11 +50,6 @@ UAudioComponent::UAudioComponent(const FObjectInitializer& ObjectInitializer)
 	bIsPreviewSound = false;
 	bIsPaused = false;
 
-	bRetriggerLoopOnProximity = false;
-	RetriggerWhenInAudibleRange = ERetriggerWhenInAudibleRange::Disabled;
-	RetriggerTimeSinceLastUpdate = 0.0f;
-	RetriggerUpdateInterval = 0.1f;
-
 	Priority = 1.f;
 	SubtitlePriority = DEFAULT_SUBTITLE_PRIORITY;
 	PitchMultiplier = 1.f;
@@ -121,14 +94,6 @@ void UAudioComponent::BeginDestroy()
 		Stop();
 	}
 
-	if (RetriggerWhenInAudibleRange == ERetriggerWhenInAudibleRange::Enabled)
-	{
-		if (FAudioDevice* AudioDevice = GetAudioDevice())
-		{
-			AudioDevice->UnregisterProximityRetriggeringAudioComponent(*this);
-		}
-	}
-
 	FScopeLock Lock(&AudioIDToComponentMapLock);
 	AudioIDToComponentMap.Remove(AudioComponentID);
 }
@@ -291,35 +256,6 @@ bool UAudioComponent::IsReadyForOwnerToAutoDestroy() const
 	return !IsPlaying();
 }
 
-void UAudioComponent::OnUpdateProximityRetrigger(float DeltaTime)
-{
-	if (RetriggerUpdateInterval > 0.0f)
-	{
-		RetriggerTimeSinceLastUpdate += DeltaTime;
-		if (RetriggerUpdateInterval > RetriggerTimeSinceLastUpdate)
-		{
-			return;
-		}
-		RetriggerTimeSinceLastUpdate = 0.0f;
-	}
-
-	if (RetriggerWhenInAudibleRange == ERetriggerWhenInAudibleRange::Enabled)
-	{
-		bool bIsInAudibleRange = GetRetriggerRate(&RetriggerUpdateInterval);
-		if (bIsActive)
-		{
-			if (!bIsInAudibleRange)
-			{
-				StopInternal();
-			}
-		}
-		else if (bIsInAudibleRange)
-		{
-			Play(); // Play refreshes RetriggerWhenInAudibleRange internally
-		}
-	}
-}
-
 void UAudioComponent::OnUpdateTransform(EUpdateTransformFlags UpdateTransformFlags, ETeleportType Teleport)
 {
 	Super::OnUpdateTransform(UpdateTransformFlags, Teleport);
@@ -393,31 +329,6 @@ bool UAudioComponent::IsInAudibleRange(float* OutMaxDistance) const
 	return AudioDevice->SoundIsAudible(Sound, GetWorld(), Location, AttenuationSettingsToApply, MaxDistance, FocusFactor);
 }
 
-bool UAudioComponent::GetRetriggerRate(float* OutRetriggerRate) const
-{
-	check(OutRetriggerRate);
-
-	FAudioDevice* AudioDevice = GetAudioDevice();
-	if (!AudioDevice)
-	{
-		*OutRetriggerRate = AudioComponentRetriggerUpdateRateMin;
-		return false;
-	}
-
-	float MaxDistance = 0.0f;
-	if (IsInAudibleRange(&MaxDistance))
-	{
-		*OutRetriggerRate = AudioComponentRetriggerUpdateRateMin;
-		return true;
-	}
-
-	const FVector Location = GetComponentTransform().GetLocation();
-	float DistanceToListener = AudioDevice->GetDistanceToNearestListener(Location);
-	const float DistanceRatio = (DistanceToListener - MaxDistance) / FMath::Max(RetriggerLoopOnProximityDistanceCVar, 1.0f);
-	*OutRetriggerRate = FMath::Lerp(AudioComponentRetriggerUpdateRateMin, AudioComponentRetriggerUpdateRateMax, FMath::Min(DistanceRatio, 1.0f));
-	return false;
-}
-
 void UAudioComponent::Play(float StartTime)
 {
 	PlayInternal(StartTime);
@@ -577,39 +488,7 @@ void UAudioComponent::PlayInternal(const float StartTime, const float FadeInDura
 
 			// Bump ActiveCount... this is used to determine if an audio component is still active after "finishing"
 			++ActiveCount;
-
 			AudioDevice->AddNewActiveSound(NewActiveSound);
-
-			if (World && World->bAllowAudioPlayback)
-			{
-				if (bRetriggerLoopOnProximity && bRetriggerLoopOnProximityEnabledCVar)
-				{
-					static FString RetriggerFailedStr = TEXT("Re-trigger loop in listener proximity failed: Sound ");
-					if (bAutoDestroy)
-					{
-						UE_LOG(LogAudio, Warning, TEXT("%s'%s' set to AutoDestroy."), *RetriggerFailedStr, Sound ? *Sound->GetName() : TEXT("N/A"));
-					}
-					else if (Sound && Sound->IsVirtualizeWhenSilent() && AudioDevice->VirtualSoundsEnabled())
-					{
-						UE_LOG(LogAudio, Warning, TEXT("%s'%s' set to 'PlayWhenSilent'."), *RetriggerFailedStr, Sound ? *Sound->GetName() : TEXT("N/A"));
-					}
-					else if (!NewActiveSound.bHasAttenuationSettings)
-					{
-						UE_LOG(LogAudio, Warning, TEXT("%s'%s' not using attenuation."), *RetriggerFailedStr, Sound ? *Sound->GetName() : TEXT("N/A"));
-					}
-					else if (!bIsPreviewSound)
-					{
-						const bool bLooping = NewActiveSound.IsLooping();
-						RetriggerWhenInAudibleRange = bLooping ? ERetriggerWhenInAudibleRange::Enabled : ERetriggerWhenInAudibleRange::Disabled;
-
-						if (bLooping)
-						{
-							AudioDevice->RegisterProximityRetriggeringAudioComponent(*this);
-						}
-					}
-				}
-			}
-
 			bIsActive = true;
 		}
 	}
@@ -653,12 +532,6 @@ void UAudioComponent::FadeOut( float FadeOutDuration, float FadeVolumeLevel )
 			{
 				DECLARE_CYCLE_STAT(TEXT("FAudioThreadTask.FadeOut"), STAT_AudioFadeOut, STATGROUP_AudioThreadCommands);
 
-				if (RetriggerWhenInAudibleRange == ERetriggerWhenInAudibleRange::Enabled && FadeVolumeLevel == 0.0f)
-				{
-					RetriggerWhenInAudibleRange = ERetriggerWhenInAudibleRange::DisableRequested;
-					AudioDevice->UnregisterProximityRetriggeringAudioComponent(*this);
-				}
-
 				const uint64 MyAudioComponentID = AudioComponentID;
 				FAudioThread::RunCommandOnAudioThread([AudioDevice, MyAudioComponentID, FadeOutDuration, FadeVolumeLevel]()
 				{
@@ -677,15 +550,6 @@ void UAudioComponent::FadeOut( float FadeOutDuration, float FadeVolumeLevel )
 			Stop();
 		}
 	}
-	// Calling FadeOut to 0.0f is effectively the same as calling stop if sound is out-of-bounds and is set to retrigger, so kill loop
-	else if (RetriggerWhenInAudibleRange == ERetriggerWhenInAudibleRange::Enabled && FadeVolumeLevel == 0.0f)
-	{
-		if (FAudioDevice* AudioDevice = GetAudioDevice())
-		{
-			RetriggerWhenInAudibleRange = ERetriggerWhenInAudibleRange::Disabled;
-			AudioDevice->UnregisterProximityRetriggeringAudioComponent(*this);
-		}
-	}
 }
 
 void UAudioComponent::AdjustVolume( float AdjustVolumeDuration, float AdjustVolumeLevel )
@@ -694,12 +558,6 @@ void UAudioComponent::AdjustVolume( float AdjustVolumeDuration, float AdjustVolu
 	{
 		if (FAudioDevice* AudioDevice = GetAudioDevice())
 		{
-			if (AdjustVolumeLevel > 0.0f && RetriggerWhenInAudibleRange == ERetriggerWhenInAudibleRange::DisableRequested)
-			{
-				RetriggerWhenInAudibleRange = ERetriggerWhenInAudibleRange::Enabled;
-				AudioDevice->RegisterProximityRetriggeringAudioComponent(*this);
-			}
-
 			DECLARE_CYCLE_STAT(TEXT("FAudioThreadTask.AdjustVolume"), STAT_AudioAdjustVolume, STATGROUP_AudioThreadCommands);
 
 			const uint64 MyAudioComponentID = AudioComponentID;
@@ -728,15 +586,6 @@ void UAudioComponent::AdjustVolume( float AdjustVolumeDuration, float AdjustVolu
 
 void UAudioComponent::Stop()
 {
-	if (FAudioDevice* AudioDevice = GetAudioDevice())
-	{
-		if (RetriggerWhenInAudibleRange == ERetriggerWhenInAudibleRange::Enabled)
-		{
-			RetriggerWhenInAudibleRange = ERetriggerWhenInAudibleRange::Disabled;
-			AudioDevice->UnregisterProximityRetriggeringAudioComponent(*this);
-		}
-	}
-
 	StopInternal();
 }
 
@@ -811,24 +660,6 @@ void UAudioComponent::PlaybackCompleted(bool bFailedToStart)
 		bIsActive = false;
 	}
 
-	if (bFailedToStart)
-	{
-		RetriggerWhenInAudibleRange = ERetriggerWhenInAudibleRange::Disabled;
-		if (FAudioDevice* AudioDevice = GetAudioDevice())
-		{
-			AudioDevice->UnregisterProximityRetriggeringAudioComponent(*this);
-		}
-	}
-
-	if (RetriggerWhenInAudibleRange == ERetriggerWhenInAudibleRange::DisableRequested)
-	{
-		RetriggerWhenInAudibleRange = ERetriggerWhenInAudibleRange::Disabled;
-		if (FAudioDevice* AudioDevice = GetAudioDevice())
-		{
-			AudioDevice->UnregisterProximityRetriggeringAudioComponent(*this);
-		}
-	}
-
 	if (!bFailedToStart && GetWorld() != nullptr && (OnAudioFinished.IsBound() || OnAudioFinishedNative.IsBound()))
 	{
 		INC_DWORD_STAT(STAT_AudioFinishedDelegatesCalled);
diff --git a/Engine/Source/Runtime/Engine/Public/AudioDevice.h b/Engine/Source/Runtime/Engine/Public/AudioDevice.h
--- a/Engine/Source/Runtime/Engine/Public/AudioDevice.h
+++ b/Engine/Source/Runtime/Engine/Public/AudioDevice.h
@@ -532,21 +532,6 @@ public:
 	 */
 	void AddReferencedObjects(FReferenceCollector& Collector);
 
-	/**
-	 * Register an audio component to re-trigger based on re-entering proximity (MaxDistance as provided by owned sound instance).
-	 */
-	void RegisterProximityRetriggeringAudioComponent(UAudioComponent& Component);
-
-	/**
-	 * Returns set of AudioComponents currently set to re-trigger on listener proximity
-	 */
-	const TArray<FAudioComponentPtr>& GetProximityRetriggerComponents() const;
-
-	/**
-	 * Unregister an audio component to re-trigger based on re-entering proximity (MaxDistance as provided by owned sound instance).
-	 */
-	void UnregisterProximityRetriggeringAudioComponent(UAudioComponent& Component);
-
 	/**
 	 * Iterate over the active AudioComponents for wave instances that could be playing.
 	 *
@@ -1492,8 +1477,6 @@ public:
 	/** Returns the game's delta time */
 	float GetGameDeltaTime() const;
 
-	void UpdateProximityRetriggerComponents();
-
 	void UpdateVirtualLoops();
 
 	/** Sets the update delta time for the audio frame */
@@ -1825,8 +1808,6 @@ private:
 	  */
 	TMap<FActiveSound*, FAudioVirtualLoop> VirtualLoops;
 
-	TArray<FAudioComponentPtr> ProximityRetriggerComponents;
-
 	/** Cached copy of sound class adjusters array. Cached to avoid allocating every frame. */
 	TArray<FSoundClassAdjuster> SoundClassAdjustersCopy;
 
