Remove OneShotPriorityCullThreshold. Causes missing sounds in times of transient high activity for trivial perf gains (effectively is an over optimization).
https://p4-swarm.epicgames.net/reviews/6100363#files

diff --git a/Engine/Source/Runtime/Engine/Private/AudioDevice.cpp b/Engine/Source/Runtime/Engine/Private/AudioDevice.cpp
--- a/Engine/Source/Runtime/Engine/Private/AudioDevice.cpp
+++ b/Engine/Source/Runtime/Engine/Private/AudioDevice.cpp
@@ -203,7 +203,6 @@ FAudioDevice::FAudioDevice()
 #endif
 	, ConcurrencyManager(this)
 	, OneShotCount(0)
-	, OneShotPriorityCullThreshold(-1.0f)
 	, GlobalMinPitch(0.4f)
 	, GlobalMaxPitch(2.0f)
 {
@@ -3875,11 +3874,6 @@ void FAudioDevice::Update(bool bGameTicking)
 
 			// Sort by priority (lowest priority first).
 			ActiveWaveInstances.Sort(FCompareFWaveInstanceByPriority());
-			OneShotPriorityCullThreshold = ActiveWaveInstances[MaxWaveInstances]->Priority;
-		}
-		else
-		{
-			OneShotPriorityCullThreshold = -1.0f;
 		}
 
 		// If not paused, update the playback time of the active sounds after we've processed passive mix modifiers
@@ -4161,15 +4155,10 @@ void FAudioDevice::AddNewActiveSoundInternal(const FActiveSound& NewActiveSound,
 	}
 #endif
 
-	// Cull one-shot active sounds if we've reached our max limit of one shot active sounds before we attempt to evaluate concurrency
-	USoundBase* Sound = NewActiveSound.GetSound();
+	const USoundBase* Sound = NewActiveSound.GetSound();
 	check(Sound);
-	if (!Sound->IsLooping() && Sound->Priority < OneShotPriorityCullThreshold)
-	{
-		return;
-	}
 
-	USoundWave* SoundWave = Cast<USoundWave>(Sound);
+	const USoundWave* SoundWave = Cast<USoundWave>(Sound);
 	if (SoundWave && SoundWave->bProcedural && SoundWave->IsGeneratingAudio())
 	{
 		FString SoundWaveName;
@@ -4218,7 +4207,7 @@ void FAudioDevice::AddNewActiveSoundInternal(const FActiveSound& NewActiveSound,
 	}
 
 	check(ActiveSound->Sound);
-	check(ActiveSound->Sound == NewActiveSound.Sound);
+	check(ActiveSound->Sound == Sound);
 
 	if (GIsEditor)
 	{
@@ -4232,7 +4221,7 @@ void FAudioDevice::AddNewActiveSoundInternal(const FActiveSound& NewActiveSound,
 	++ActiveSound->Sound->CurrentPlayCount;
 
 #if !(UE_BUILD_SHIPPING || UE_BUILD_TEST)
-	UE_LOG(LogAudio, VeryVerbose, TEXT("New ActiveSound %s Comp: %s Loc: %s"), *NewActiveSound.Sound->GetName(), *NewActiveSound.GetAudioComponentName(), *NewActiveSound.Transform.GetTranslation().ToString());
+	UE_LOG(LogAudio, VeryVerbose, TEXT("New ActiveSound %s Comp: %s Loc: %s"), *Sound->GetName(), *NewActiveSound.GetAudioComponentName(), *NewActiveSound.Transform.GetTranslation().ToString());
 #endif // !(UE_BUILD_SHIPPING || UE_BUILD_TEST)
 
 	// Cull one-shot active sounds if we've reached our max limit of one shot active sounds before we attempt to evaluate concurrency
diff --git a/Engine/Source/Runtime/Engine/Public/AudioDevice.h b/Engine/Source/Runtime/Engine/Public/AudioDevice.h
--- a/Engine/Source/Runtime/Engine/Public/AudioDevice.h
+++ b/Engine/Source/Runtime/Engine/Public/AudioDevice.h
@@ -1841,9 +1841,6 @@ private:
 	/** A count of the number of one-shot active sounds. */
 	uint32 OneShotCount;
 
-	/** Threshold priority for allowing oneshot active sounds through the max oneshot active sound limit. */
-	float OneShotPriorityCullThreshold;
-
 	// Global min and max pitch scale, derived from audio settings
 	float GlobalMinPitch;
 	float GlobalMaxPitch;
