Adjusting test for array terminator in backwards compatible receive path

diff --git a/Engine/Source/Runtime/Engine/Private/RepLayout.cpp b/Engine/Source/Runtime/Engine/Private/RepLayout.cpp
--- a/Engine/Source/Runtime/Engine/Private/RepLayout.cpp
+++ b/Engine/Source/Runtime/Engine/Private/RepLayout.cpp
@@ -3110,8 +3110,6 @@ bool FRepLayout::ReceiveProperties_BackwardsCompatible_r(
 			FScriptArray* DataArray = (FScriptArray*)(Data + Cmd).Data;
 			FScriptArray* ShadowArray = ShadowData ? (FScriptArray*)(ShadowData + Cmd).Data : nullptr;
 
-			const int32 ShadowArrayNum = ShadowArray ? ShadowArray->Num() : INDEX_NONE;
-
 			FRepObjectDataBuffer LocalData = Data;
 			FRepShadowDataBuffer LocalShadowData = ShadowData;
 
@@ -3150,14 +3148,13 @@ bool FRepLayout::ReceiveProperties_BackwardsCompatible_r(
 					if (TempReader.GetBitsLeft() == 8)
 					{
 						// We have bits left over, so see if its the Array Terminator.
-						// This should be 0, and we should be able to verify that the new number
-						// of elements in the array is smaller than the previous number.
+						// This should be 0
 						uint32 Terminator;
 						TempReader.SerializeIntPacked(Terminator);
 
-						if (Terminator != 0 || (int32)ArrayNum >= ShadowArrayNum)
+						if (Terminator != 0)
 						{
-							UE_LOG(LogRep, Warning, TEXT("ReceiveProperties_BackwardsCompatible_r: Invalid array terminator on shrink. NetFieldExportHandle: %d, OldArrayNum=%d, NewArrayNum=%d"), Terminator, ShadowArrayNum, ArrayNum);
+							UE_LOG(LogRep, Warning, TEXT("ReceiveProperties_BackwardsCompatible_r: Invalid array terminator. Owner: %s, Name: %s, NetFieldExportHandle: %i, Terminator: %d"), *Owner->GetName(), *NetFieldExportGroup->NetFieldExports[NetFieldExportHandle].ExportName.ToString(), NetFieldExportHandle, Terminator);
 							return false;
 						}
 					}
