- Concurrency Generation Volume Recovery
- Add ActiveSound volume to Audio3dVisualize.  Realize this is a bit confusing with multiple soures stacking on the screen.  Will be moving 3dvis to active sound layer to avoid double-drawing as it was before.

diff --git a/Engine/Source/Runtime/Engine/Classes/Sound/SoundConcurrency.h b/Engine/Source/Runtime/Engine/Classes/Sound/SoundConcurrency.h
--- a/Engine/Source/Runtime/Engine/Classes/Sound/SoundConcurrency.h
+++ b/Engine/Source/Runtime/Engine/Classes/Sound/SoundConcurrency.h
@@ -71,20 +71,49 @@ struct ENGINE_API FSoundConcurrencySettings
 	UPROPERTY(EditAnywhere, BlueprintReadWrite, Category = Concurrency)
 	TEnumAsByte<EMaxConcurrentResolutionRule::Type> ResolutionRule;
 
+private:
 	/**
-	 * The amount of attenuation to apply to older voice instances in this concurrency group. This reduces volume of older voices in a concurrency group as new voices play.
+	 * Ducking factor to apply per older voice instance (generation), which compounds as new voices play
+	 * and (optionally) revives them as they stop according to the provided attack/release times.
 	 *
 	 * AppliedVolumeScale = Math.Pow(DuckingScale, VoiceGeneration)
 	 */
-	UPROPERTY(EditAnywhere, BlueprintReadWrite, Category = Concurrency, meta = (UIMin = "0", UIMax = "1", ClampMin = "0", ClampMax = "1"))
+	UPROPERTY(EditAnywhere, Category = "Volume Scaling", meta = (UIMin = "0.0", UIMax = "1.0", ClampMin = "0.0", ClampMax = "1.0"))
 	float VolumeScale;
 
+public:
+	/**
+	 * Time taken to apply duck using volume scalar.
+	 */
+	UPROPERTY(EditAnywhere, BlueprintReadWrite, Category = "Volume Scaling", meta = (DisplayName = "Attack Time", UIMin = "0.0", ClampMin = "0.0", UIMax = "10.0", ClampMax = "1000000.0"))
+	float VolumeScaleAttackTime;
+
+	/**
+	 * Whether or not volume scaling can release (i.e. "recover") volume ducking behavior when concurrency group sounds stop.
+	 */
+	UPROPERTY(EditAnywhere, BlueprintReadWrite, Category = "Volume Scaling", meta = (DisplayName = "Can Release"))
+	uint32 bVolumeScaleCanRelease:1;
+
+	/**
+	 * Time taken to remove duck using volume scalar.
+	 */
+	UPROPERTY(EditAnywhere, BlueprintReadWrite, Category = "Volume Scaling", meta = (DisplayName = "Release Time", EditCondition = "bVolumeScaleCanRelease", UIMin = "0.0", ClampMin = "0.0", UIMax = "10.0", ClampMax="1000000.0"))
+	float VolumeScaleReleaseTime;
+
 	FSoundConcurrencySettings()
 		: MaxCount(16)
-		, bLimitToOwner(false)
+		, bLimitToOwner(0)
 		, ResolutionRule(EMaxConcurrentResolutionRule::StopFarthestThenOldest)
 		, VolumeScale(1.0f)
+		, VolumeScaleAttackTime(0.01f)
+		, bVolumeScaleCanRelease(0)
+		, VolumeScaleReleaseTime(0.0f)
 	{}
+
+	/**
+	 * Retrieves the volume scale
+	 */
+	float GetVolumeScale() const;
 };
 
 UCLASS(BlueprintType, hidecategories=Object, editinlinenew, MinimalAPI)
@@ -123,6 +152,36 @@ struct FConcurrencyHandle
 };
 
 
+/** Sound instance data pertaining to concurrency tracking */
+struct FConcurrencySoundData
+{
+	int32 Generation;
+	float LerpTime;
+
+private:
+	float Elapsed;
+	float DbTargetVolume;
+	float DbStartVolume;
+
+public:
+	FConcurrencySoundData()
+		: Generation(0)
+		, LerpTime(0.0f)
+		, Elapsed(0.0f)
+		, DbTargetVolume(0.0f)
+		, DbStartVolume(0.0f)
+	{
+	}
+
+	void Update(float InElapsed);
+
+	float GetVolume(bool bInDecibels = false) const;
+	float GetTargetVolume(bool bInDecibles = false) const;
+
+	void SetTarget(float InTargetVolume, float InLerpTime);
+};
+
+
 /** Class which tracks array of active sound pointers for concurrency management */
 class FConcurrencyGroup
 {
@@ -132,7 +191,6 @@ class FConcurrencyGroup
 	FConcurrencyGroupID GroupID;
 	FConcurrencyObjectID ObjectID;
 	FSoundConcurrencySettings Settings;
-	int32 Generation;
 
 public:
 	/** Constructor for the max concurrency active sound entry. */
@@ -146,10 +204,10 @@ public:
 	/** Returns the id of the concurrency group */
 	FConcurrencyGroupID GetGroupID() const { return GroupID; }
 
-	/** Returns the current generation */
-	const int32 GetGeneration() const { return Generation; }
+	/** Returns the current generation (effectively, the number of concurrency sound instances active) */
+	const int32 GetNextGeneration() const { return ActiveSounds.Num(); }
 
-	/** Returns the current generation */
+	/** Returns the settings associated with the group */
 	const FSoundConcurrencySettings& GetSettings() const { return Settings; }
 
 	/** Returns the parent object ID */
@@ -162,10 +220,10 @@ public:
 	bool IsFull() const { return Settings.MaxCount <= ActiveSounds.Num(); }
 
 	/** Adds an active sound to the active sound array. */
-	void AddActiveSound(struct FActiveSound* ActiveSound);
+	void AddActiveSound(FActiveSound& ActiveSound);
 
 	/** Removes an active sound from the active sound array. */
-	void RemoveActiveSound(FActiveSound* ActiveSound);
+	void RemoveActiveSound(FActiveSound& ActiveSound);
 
 	/** Sorts the active sound array by volume */
 	void StopQuietSoundsDueToMaxConcurrency();
@@ -216,7 +274,7 @@ public:
 	FActiveSound* CreateNewActiveSound(const FActiveSound& NewActiveSound, bool bIsRetriggering);
 
 	/** Removes the active sound from concurrency tracking when active sound is stopped. */
-	void RemoveActiveSound(FActiveSound* ActiveSound);
+	void RemoveActiveSound(FActiveSound& ActiveSound);
 
 	/** Stops any active sounds due to max concurrency quietest sound resolution rule */
 	void UpdateQuietSoundsToStop();
diff --git a/Engine/Source/Runtime/Engine/Private/ActiveSound.cpp b/Engine/Source/Runtime/Engine/Private/ActiveSound.cpp
--- a/Engine/Source/Runtime/Engine/Private/ActiveSound.cpp
+++ b/Engine/Source/Runtime/Engine/Private/ActiveSound.cpp
@@ -36,7 +36,6 @@ FActiveSound::FActiveSound()
 	, AudioComponentID(0)
 	, OwnerID(0)
 	, AudioDevice(nullptr)
-	, ConcurrencyGeneration(0)
 	, SoundClassOverride(nullptr)
 	, SoundSubmixOverride(nullptr)
 	, bHasCheckedOcclusion(false)
@@ -440,14 +439,17 @@ void FActiveSound::UpdateWaveInstances(TArray<FWaveInstance*> &InWaveInstances,
 	// final value that is correct
 	UpdateAdjustVolumeMultiplier(DeltaTime);
 
+	// Update concurrency volume scalars. Must be done prior to getting collective volume and applying below
+	// in parse params.
+	UpdateConcurrencyVolumeScalars(DeltaTime);
+
 	// If the sound is a preview sound, then ignore the transient master volume and application volume
 	if (!bIsPreviewSound)
 	{
 		ParseParams.VolumeApp = AudioDevice->GetMasterVolume();
 	}
 
-	const float TotalConcurrencyVolumeScale = GetTotalConcurrencyVolumeScale();
-	ParseParams.VolumeMultiplier = VolumeMultiplier * Sound->GetVolumeMultiplier() * CurrentAdjustVolumeMultiplier * TotalConcurrencyVolumeScale;
+	ParseParams.VolumeMultiplier = GetVolume();
 
 	ParseParams.Priority = Priority;
 	ParseParams.Pitch *= PitchMultiplier * Sound->GetPitchMultiplier();
@@ -1046,18 +1048,32 @@ bool FActiveSound::GetFloatParameter( const FName InName, float& OutFloat ) cons
 	return false;
 }
 
+float FActiveSound::GetVolume() const
+{
+	const float Volume = VolumeMultiplier * CurrentAdjustVolumeMultiplier * GetTotalConcurrencyVolumeScale();
+	return Sound ? Volume * Sound->GetVolumeMultiplier() : Volume;
+}
+
 float FActiveSound::GetTotalConcurrencyVolumeScale() const
 {
 	float OutVolume = 1.0f;
 
-	for (const TPair<FConcurrencyGroupID, float>& ConcurrencyGroupVolumePair : ConcurrencyGroupVolumeScales)
+	for (const TPair<FConcurrencyGroupID, FConcurrencySoundData>& ConcurrencyPair : ConcurrencyGroupData)
 	{
-		OutVolume *= ConcurrencyGroupVolumePair.Value;
+		OutVolume *= ConcurrencyPair.Value.GetVolume();
 	}
 
 	return OutVolume;
 }
 
+void FActiveSound::UpdateConcurrencyVolumeScalars(const float DeltaTime)
+{
+	for (TPair<FConcurrencyGroupID, FConcurrencySoundData>& ConcurrencyPair : ConcurrencyGroupData)
+	{
+		ConcurrencyPair.Value.Update(DeltaTime);
+	}
+}
+
 void FActiveSound::SetFloatParameter( const FName InName, const float InFloat )
 {
 	if( InName != NAME_None )
diff --git a/Engine/Source/Runtime/Engine/Private/Audio.cpp b/Engine/Source/Runtime/Engine/Private/Audio.cpp
--- a/Engine/Source/Runtime/Engine/Private/Audio.cpp
+++ b/Engine/Source/Runtime/Engine/Private/Audio.cpp
@@ -429,8 +429,8 @@ void FSoundSource::DrawDebugInfo()
 			const bool bSpatialized = Buffer->NumChannels == 2 && WaveInstance->GetUseSpatialization();
 			const FVector LeftChannelSourceLoc = LeftChannelSourceLocation;
 			const FVector RightChannelSourceLoc = RightChannelSourceLocation;
-
-			FAudioThread::RunCommandOnGameThread([AudioComponentID, Sound, bSpatialized, Location, LeftChannelSourceLoc, RightChannelSourceLoc]()
+			const float Volume = WaveInstance->GetActualVolume();
+			FAudioThread::RunCommandOnGameThread([AudioComponentID, Sound, bSpatialized, Location, LeftChannelSourceLoc, RightChannelSourceLoc, Volume]()
 			{
 				UAudioComponent* AudioComponent = UAudioComponent::GetAudioComponentFromID(AudioComponentID);
 				if (AudioComponent)
@@ -448,7 +448,8 @@ void FSoundSource::DrawDebugInfo()
 						}
 
 						const FString Name = Sound->GetName();
-						DrawDebugString(SoundWorld, AudioComponent->GetComponentLocation() + FVector(0, 0, 32), *Name, nullptr, FColor::White, 0.033, false);
+						const FString Description = FString::Printf(TEXT("%s (%.3f)"), *Name, Volume);
+						DrawDebugString(SoundWorld, AudioComponent->GetComponentLocation() + FVector(0, 0, 32), *Description, nullptr, FColor::White, 0.033, false);
 					}
 				}
 			}, GET_STATID(STAT_AudioDrawSourceDebugInfo));
diff --git a/Engine/Source/Runtime/Engine/Private/AudioDevice.cpp b/Engine/Source/Runtime/Engine/Private/AudioDevice.cpp
--- a/Engine/Source/Runtime/Engine/Private/AudioDevice.cpp
+++ b/Engine/Source/Runtime/Engine/Private/AudioDevice.cpp
@@ -4444,7 +4444,7 @@ void FAudioDevice::AddSoundToStop(FActiveSound* SoundToStop)
 
 		if (!VirtualLoops.Contains(SoundToStop))
 		{
-			ConcurrencyManager.RemoveActiveSound(SoundToStop);
+			ConcurrencyManager.RemoveActiveSound(*SoundToStop);
 		}
 	}
 }
@@ -5787,6 +5787,12 @@ void FAudioDevice::UpdateVirtualLoops()
 			{
 				VirtualLoopsToRetrigger.Add(VirtualLoop);
 			}
+
+			// If signaled to fade out and virtualized, and to pending stop list.
+			if (ActiveSound.bFadingOut)
+			{
+				AddSoundToStop(&ActiveSound);
+			}
 		}
 
 		for (FAudioVirtualLoop& RetriggerLoop : VirtualLoopsToRetrigger)
diff --git a/Engine/Source/Runtime/Engine/Private/AudioVirtualLoop.cpp b/Engine/Source/Runtime/Engine/Private/AudioVirtualLoop.cpp
--- a/Engine/Source/Runtime/Engine/Private/AudioVirtualLoop.cpp
+++ b/Engine/Source/Runtime/Engine/Private/AudioVirtualLoop.cpp
@@ -72,8 +72,7 @@ bool FAudioVirtualLoop::Virtualize(const FActiveSound& InActiveSound, FAudioDevi
 	ActiveSound->bAsyncOcclusionPending = false;
 	ActiveSound->AudioDevice = &AudioDevice;
 	ActiveSound->bIsPlayingAudio = false;
-	ActiveSound->ConcurrencyGroupIDs.Reset();
-	ActiveSound->ConcurrencyGroupVolumeScales.Reset();
+	ActiveSound->ConcurrencyGroupData.Reset();
 	ActiveSound->VolumeConcurrency = 1.0f;
 	ActiveSound->WaveInstances.Reset();
 	ActiveSound->bHasVirtualized = true;
@@ -200,14 +199,16 @@ void FAudioVirtualLoop::DrawDebugInfo() const
 		const TWeakObjectPtr<UWorld> World = ActiveSound->GetWeakWorld();
 		const FString Name = Sound->GetName();
 		const float DrawInterval = UpdateInterval;
-		FAudioThread::RunCommandOnGameThread([World, Transform, Name, DrawInterval]()
+		const float Volume = ActiveSound->GetVolume();
+		FAudioThread::RunCommandOnGameThread([World, Transform, Name, DrawInterval, Volume]()
 		{
 			if (World.IsValid())
 			{
+				const FString Description = FString::Printf(TEXT("%s (%.3f)"), *Name, Volume);
 				FVector Location = Transform.GetLocation();
 				FRotator Rotation = Transform.GetRotation().Rotator();
 				DrawDebugCrosshairs(World.Get(), Location, Rotation, 20.0f, FColor::Blue, false, DrawInterval, SDPG_Foreground);
-				DrawDebugString(World.Get(), Location + FVector(0, 0, 32), *Name, nullptr, FColor::Blue, DrawInterval, false);
+				DrawDebugString(World.Get(), Location + FVector(0, 0, 32), *Description, nullptr, FColor::Blue, DrawInterval, false);
 			}
 		}, GET_STATID(STAT_AudioDrawVirtualLoopDebugInfo));
 	}
diff --git a/Engine/Source/Runtime/Engine/Private/SoundBase.cpp b/Engine/Source/Runtime/Engine/Private/SoundBase.cpp
--- a/Engine/Source/Runtime/Engine/Private/SoundBase.cpp
+++ b/Engine/Source/Runtime/Engine/Private/SoundBase.cpp
@@ -196,7 +196,6 @@ void USoundBase::PostLoad()
 		ConcurrencyOverrides.bLimitToOwner = false;
 		ConcurrencyOverrides.MaxCount = FMath::Max(MaxConcurrentPlayCount_DEPRECATED, 1);
 		ConcurrencyOverrides.ResolutionRule = MaxConcurrentResolutionRule_DEPRECATED;
-		ConcurrencyOverrides.VolumeScale = 1.0f;
 	}
 #endif
 }
diff --git a/Engine/Source/Runtime/Engine/Private/SoundConcurrency.cpp b/Engine/Source/Runtime/Engine/Private/SoundConcurrency.cpp
--- a/Engine/Source/Runtime/Engine/Private/SoundConcurrency.cpp
+++ b/Engine/Source/Runtime/Engine/Private/SoundConcurrency.cpp
@@ -11,6 +11,33 @@
 // Forward Declarations
 struct FListener;
 
+namespace
+{
+	void SetSoundDataTarget(const FActiveSound& ActiveSound, FConcurrencySoundData& SoundData, float InTargetVolume, float InLerpTime)
+	{
+#if UE_BUILD_SHIPPING
+		SoundData.SetTarget(InTargetValue, InLerpTime);
+#else
+		const float LastTargetVolume = SoundData.GetTargetVolume();
+
+		SoundData.SetTarget(InTargetVolume, InLerpTime);
+		if (!FMath::IsNearlyEqual(LastTargetVolume, InTargetVolume))
+		{
+			if (const USoundBase* Sound = ActiveSound.GetSound())
+			{
+				UE_LOG(LogAudio, Verbose,
+					TEXT("Sound '%s' concurrency generation '%i' target volume update: %.3f to %.3f."),
+					*Sound->GetName(),
+					SoundData.Generation,
+					LastTargetVolume,
+					InTargetVolume);
+			}
+		}
+#endif // UE_BUILD_SHIPPING
+	}
+} // namespace <>
+
+
 /************************************************************************/
 /* USoundConcurrency													*/
 /************************************************************************/
@@ -20,6 +47,15 @@ USoundConcurrency::USoundConcurrency(const FObjectInitializer& ObjectInitializer
 {
 }
 
+/************************************************************************/
+/* FSoundConcurrencySettings											*/
+/************************************************************************/
+float FSoundConcurrencySettings::GetVolumeScale() const
+{
+	return FMath::Clamp(VolumeScale, 0.0f, 1.0f);;
+}
+
+
 /************************************************************************/
 /* USoundConcurrency													*/
 /************************************************************************/
@@ -50,6 +86,38 @@ EConcurrencyMode FConcurrencyHandle::GetMode(const FActiveSound& ActiveSound) co
 	return ObjectID == 0 ? EConcurrencyMode::Sound : EConcurrencyMode::Group;
 }
 
+void FConcurrencySoundData::Update(float InElapsed)
+{
+	Elapsed += InElapsed;
+}
+
+float FConcurrencySoundData::GetVolume(bool bInDecibels) const
+{
+	if (FMath::IsNearlyZero(LerpTime) || Elapsed > LerpTime || FMath::IsNearlyEqual(DbTargetVolume, DbStartVolume))
+	{
+		return bInDecibels ? DbTargetVolume : Audio::ConvertToLinear(DbTargetVolume);
+	}
+
+	float Alpha = Elapsed / LerpTime;
+	const float DbCurrentVolume = FMath::Lerp(DbStartVolume, DbTargetVolume, Alpha);
+	return bInDecibels ? DbCurrentVolume : Audio::ConvertToLinear(DbCurrentVolume);
+}
+
+float FConcurrencySoundData::GetTargetVolume(bool bInDecibels) const
+{
+	return bInDecibels ? DbTargetVolume : Audio::ConvertToLinear(DbTargetVolume);
+}
+
+void FConcurrencySoundData::SetTarget(float InTargetVolume, float InLerpTime)
+{
+	DbStartVolume = GetVolume(true);
+
+	LerpTime = FMath::Max(InLerpTime, 0.0f);
+	Elapsed = 0.0f;
+
+	DbTargetVolume = Audio::ConvertToDecibels(InTargetVolume);
+}
+
 /************************************************************************/
 /* FConcurrencyGroup												*/
 /************************************************************************/
@@ -58,7 +126,6 @@ FConcurrencyGroup::FConcurrencyGroup(FConcurrencyGroupID InGroupID, const FConcu
 	: GroupID(InGroupID)
 	, ObjectID(ConcurrencyHandle.ObjectID)
 	, Settings(ConcurrencyHandle.Settings)
-	, Generation(0)
 {
 }
 
@@ -68,55 +135,55 @@ FConcurrencyGroupID FConcurrencyGroup::GenerateNewID()
 	return ++ConcurrencyGroupIDs;
 }
 
-void FConcurrencyGroup::AddActiveSound(FActiveSound* ActiveSound)
+void FConcurrencyGroup::AddActiveSound(FActiveSound& ActiveSound)
 {
 	check(GroupID != 0);
 
-	int32 LastIndex = ActiveSound->ConcurrencyGroupIDs.Num();
-	if (ActiveSound->ConcurrencyGroupIDs.AddUnique(GroupID) == LastIndex)
-	{
-		LastIndex = ActiveSounds.Num();
-		if (ActiveSounds.AddUnique(ActiveSound) == LastIndex)
-		{
-			ActiveSound->ConcurrencyGeneration = Generation++;
-		}
-		else
-		{
-			UE_LOG(LogAudio, Fatal, TEXT("Attempting to increment active sound ConcurrencyGeneration by adding active sound '%s' to group multiple times."), *ActiveSound->GetOwnerName());
-		}
-	}
-	else
+	if (ActiveSound.ConcurrencyGroupData.Contains(GroupID))
 	{
-		UE_LOG(LogAudio, Fatal, TEXT("Attempting to add active sound '%s' to concurrency group multiple times."), *ActiveSound->GetOwnerName());
+		UE_LOG(LogAudio, Fatal, TEXT("Attempting to add active sound '%s' to concurrency group multiple times."), *ActiveSound.GetOwnerName());
+		return;
 	}
-}
 
-void FConcurrencyGroup::RemoveActiveSound(FActiveSound* ActiveSound)
-{
-	check(ActiveSound);
+	FConcurrencySoundData SoundData;
+	SoundData.Generation = ActiveSounds.Num();
+	SetSoundDataTarget(ActiveSound, SoundData, 1.0f, 0.0f);
 
-	// Cache generation being removed
-	const int32 RemovedGeneration = ActiveSound->ConcurrencyGeneration;
+	ActiveSounds.Add(&ActiveSound);
+	ActiveSound.ConcurrencyGroupData.Add(GroupID, MoveTemp(SoundData));
+}
 
+void FConcurrencyGroup::RemoveActiveSound(FActiveSound& ActiveSound)
+{
 	// Remove from array
-	const int32 NumRemoved = ActiveSounds.RemoveSwap(ActiveSound);
+	const int32 NumRemoved = ActiveSounds.RemoveSwap(&ActiveSound);
 	if (NumRemoved == 0)
 	{
 		return;
 	}
-
 	check(NumRemoved == 1);
 
+	FConcurrencySoundData& SoundData = ActiveSound.ConcurrencyGroupData.FindChecked(GroupID);
+
 	// Rebase generations due to removal of a member
 	for (FActiveSound* OtherSound : ActiveSounds)
 	{
-		if (OtherSound->ConcurrencyGeneration > RemovedGeneration)
+		check(OtherSound);
+		FConcurrencySoundData& OtherSoundData = OtherSound->ConcurrencyGroupData.FindChecked(GroupID);
+		if (OtherSoundData.Generation > SoundData.Generation)
 		{
-			OtherSound->ConcurrencyGeneration--;
+			check(OtherSoundData.Generation > 0);
+			check(OtherSoundData.Generation <= ActiveSounds.Num());
+			OtherSoundData.Generation--;
 		}
-	}
 
-	Generation--;
+		if (Settings.bVolumeScaleCanRelease)
+		{
+			const float GenerationDelta = static_cast<float>(ActiveSounds.Num() - OtherSoundData.Generation) - 1.0f;
+			const float NewTargetVolume = FMath::Clamp(FMath::Pow(Settings.GetVolumeScale(), GenerationDelta), 0.0f, 1.0f);
+			SetSoundDataTarget(*OtherSound, OtherSoundData, NewTargetVolume, Settings.VolumeScaleReleaseTime);
+		}
+	}
 }
 
 void FConcurrencyGroup::StopQuietSoundsDueToMaxConcurrency()
@@ -577,23 +644,31 @@ FActiveSound* FSoundConcurrencyManager::CreateAndEvictActiveSounds(const FActive
 		check(ConcurrencyGroup);
 
 		const FSoundConcurrencySettings& Settings = ConcurrencyGroup->GetSettings();
-		const float Volume = Settings.VolumeScale;
-		if (Volume < 1.0f)
+		const float Volume = Settings.GetVolumeScale();
+		if (!FMath::IsNearlyEqual(Volume, 1.0f))
 		{
 			check(Volume >= 0.0f);
-			const int32 NextGeneration = ConcurrencyGroup->GetGeneration() + 1;
+			const int32 NextGeneration = ConcurrencyGroup->GetNextGeneration();
 
 			// If we're ducking older sounds in the concurrency group, then loop through each sound in the concurrency group
 			// and update their duck amount based on each sound's generation and the next generation count. The older the sound, the more ducking.
 			const TArray<FActiveSound*>& ActiveSounds = ConcurrencyGroup->GetActiveSounds();
-			for (FActiveSound* CurActiveSound : ActiveSounds)
+			for (FActiveSound* OtherSound : ActiveSounds)
 			{
-				check(CurActiveSound);
+				check(OtherSound);
 
-				const float ActiveSoundGeneration = static_cast<float>(CurActiveSound->ConcurrencyGeneration);
+				FConcurrencySoundData& OtherSoundData = OtherSound->ConcurrencyGroupData.FindChecked(ConcurrencyGroup->GetGroupID());
+				const float ActiveSoundGeneration = static_cast<float>(OtherSoundData.Generation);
 				const float GenerationDelta = NextGeneration - ActiveSoundGeneration;
-				float& VolumeScale = CurActiveSound->ConcurrencyGroupVolumeScales.FindOrAdd(ConcurrencyGroup->GetGroupID());
-				VolumeScale = FMath::Pow(Volume, GenerationDelta);
+				const float NewTargetVolume = FMath::Pow(Volume, GenerationDelta);
+
+				// Don't allow volume to recover if release time is negative
+				if (Settings.bVolumeScaleCanRelease && OtherSoundData.GetTargetVolume() < NewTargetVolume)
+				{
+					continue;
+				}
+
+				SetSoundDataTarget(*OtherSound, OtherSoundData, NewTargetVolume, Settings.VolumeScaleAttackTime);
 			}
 		}
 
@@ -603,8 +678,8 @@ FActiveSound* FSoundConcurrencyManager::CreateAndEvictActiveSounds(const FActive
 			bTrackConcurrencyVolume = true;
 		}
 
-		// And add it to to the concurrency group. This automatically updates generation counts.
-		ConcurrencyGroup->AddActiveSound(ActiveSound);
+		// And add it to to the concurrency group.
+		ConcurrencyGroup->AddActiveSound(*ActiveSound);
 	}
 
 	if (!bTrackConcurrencyVolume)
@@ -618,7 +693,7 @@ FActiveSound* FSoundConcurrencyManager::CreateAndEvictActiveSounds(const FActive
 		check(SoundToEvict);
 		check(AudioDevice == SoundToEvict->AudioDevice);
 
-		RemoveActiveSound(SoundToEvict);
+		RemoveActiveSound(*SoundToEvict);
 
 		// Remove the active sound from the concurrency manager immediately so it doesn't count towards
 		// subsequent concurrency resolution checks (i.e. if sounds are triggered multiple times in this frame)
@@ -652,7 +727,7 @@ FActiveSound* FSoundConcurrencyManager::CreateAndEvictActiveSounds(const FActive
 		}
 
 		AudioDevice->AddSoundToStop(SoundToEvict);
-		// If using a mode where we support re-triggering, attempt to re-trigger and update boolean state.
+		// If using a mode where we support loop virtualization, attempt to re-trigger and update boolean state.
 		if (bAllowVirtual)
 		{
 			const bool bDoRangeCheck = true;
@@ -672,18 +747,19 @@ FActiveSound* FSoundConcurrencyManager::CreateAndEvictActiveSounds(const FActive
 	return ActiveSound;
 }
 
-void FSoundConcurrencyManager::RemoveActiveSound(FActiveSound* ActiveSound)
+void FSoundConcurrencyManager::RemoveActiveSound(FActiveSound& ActiveSound)
 {
 	check(IsInAudioThread());
 
 	// Remove this sound from it's concurrency list
-	for (const FConcurrencyGroupID ConcurrencyGroupID : ActiveSound->ConcurrencyGroupIDs)
+	for (const TPair<FConcurrencyGroupID, FConcurrencySoundData>& SoundDataPair : ActiveSound.ConcurrencyGroupData)
 	{
+		const FConcurrencyGroupID ConcurrencyGroupID = SoundDataPair.Key;
 		FConcurrencyGroup* ConcurrencyGroup = ConcurrencyGroups.FindRef(ConcurrencyGroupID);
 		if (!ConcurrencyGroup)
 		{
 			UE_LOG(LogAudio, Error, TEXT("Attempting to remove stopped sound '%s' from inactive concurrency group."),
-				ActiveSound->GetSound() ? *ActiveSound->GetSound()->GetName(): TEXT("Unset"));
+				ActiveSound.GetSound() ? *ActiveSound.GetSound()->GetName(): TEXT("Unset"));
 			continue;
 		}
 
@@ -703,14 +779,14 @@ void FSoundConcurrencyManager::RemoveActiveSound(FActiveSound* ActiveSound)
 			ConcurrencyMap.Remove(ConcurrencyObjectID);
 
 			// Remove from sound object map if present.
-			if (USoundBase* Sound = ActiveSound->GetSound())
+			if (USoundBase* Sound = ActiveSound.GetSound())
 			{
 				const FSoundOwnerObjectID ObjectID = Sound->GetUniqueID();
 				SoundObjectToConcurrencyGroup.Remove(ObjectID);
 			}
 
 			// Remove from owner map if present.
-			const uint32 OwnerID = ActiveSound->GetOwnerID();
+			const uint32 OwnerID = ActiveSound.GetOwnerID();
 			if (FOwnerConcurrencyMapEntry* OwnerEntry = OwnerConcurrencyMap.Find(OwnerID))
 			{
 				if (OwnerEntry->ConcurrencyObjectToConcurrencyGroup.Remove(ConcurrencyObjectID))
@@ -725,7 +801,7 @@ void FSoundConcurrencyManager::RemoveActiveSound(FActiveSound* ActiveSound)
 			// Remove from owner per sound map if present.
 			if (FSoundInstanceEntry* InstanceEntry = OwnerPerSoundConcurrencyMap.Find(OwnerID))
 			{
-				if (USoundBase* Sound = ActiveSound->GetSound())
+				if (USoundBase* Sound = ActiveSound.GetSound())
 				{
 					if (InstanceEntry->SoundInstanceToConcurrencyGroup.Remove(Sound->GetUniqueID()))
 					{
@@ -740,7 +816,8 @@ void FSoundConcurrencyManager::RemoveActiveSound(FActiveSound* ActiveSound)
 			delete ConcurrencyGroup;
 		}
 	}
-	ActiveSound->ConcurrencyGroupIDs.Reset();
+
+	ActiveSound.ConcurrencyGroupData.Reset();
 }
 
 void FSoundConcurrencyManager::UpdateQuietSoundsToStop()
diff --git a/Engine/Source/Runtime/Engine/Public/ActiveSound.h b/Engine/Source/Runtime/Engine/Public/ActiveSound.h
--- a/Engine/Source/Runtime/Engine/Public/ActiveSound.h
+++ b/Engine/Source/Runtime/Engine/Public/ActiveSound.h
@@ -271,10 +271,7 @@ public:
 	FAudioDevice* AudioDevice;
 
 	/** The concurrent groups that this sound is actively playing in. */
-	TArray<FConcurrencyGroupID> ConcurrencyGroupIDs;
-
-	/** The generation of this sound in the concurrency group. */
-	int32 ConcurrencyGeneration;
+	TMap<FConcurrencyGroupID, FConcurrencySoundData> ConcurrencyGroupData;
 
 	/** Optional USoundConcurrency to override for the sound. */
 	TSet<USoundConcurrency*> ConcurrencySet;
@@ -420,9 +417,6 @@ public:
 	/** The interpolated parameter for the volume attenuation due to occlusion. */
 	FDynamicParameter CurrentOcclusionVolumeAttenuation;
 
-	/** Volume scale factors to apply to a sound based on the concurrency count of the active sound when it started. Will reduce volume of new sounds if many sounds are playing in concurrency group. */
-	TMap<FConcurrencyGroupID, float> ConcurrencyGroupVolumeScales;
-
 	float SubtitlePriority;
 
 	/** The product of the component priority and the USoundBase priority */
@@ -544,6 +538,9 @@ public:
 
 	void CollectAttenuationShapesForVisualization(TMultiMap<EAttenuationShape::Type, FBaseAttenuationSettings::AttenuationShapeDetails>& ShapeDetailsMap) const;
 
+	/** Gets volume product all gain stages pertaining to active sound */
+	float GetVolume() const;
+
 	/**
 	 * Friend archive function used for serialization.
 	 */
@@ -621,6 +618,9 @@ private:
 	/** Sets the target volume multiplier to achieve over the specified time period */
 	void UpdateAdjustVolumeMultiplier(const float DeltaTime);
 
+	/** Updates ramping concurrency volume scalars */
+	void UpdateConcurrencyVolumeScalars(const float DeltaTime);
+
 	/** if OcclusionCheckInterval > 0.0, checks if the sound has become (un)occluded during playback
 	 * and calls eventOcclusionChanged() if so
 	 * primarily used for gameplay-relevant ambient sounds
@@ -628,7 +628,7 @@ private:
 	 * @param ListenerLocation location of the closest listener to the sound
 	 */
 	void CheckOcclusion(const FVector ListenerLocation, const FVector SoundLocation, const FSoundAttenuationSettings* AttenuationSettingsPtr);
-	 
+
 	/** Apply the interior settings to the ambient sound as appropriate */
 	void HandleInteriorVolumes( const FListener& Listener, struct FSoundParseParameters& ParseParams );
 
