Fixed GetValueOnGameThread ensure when marking components dirty from parallel tasks.  bAllowConcurrentUpdates only needs to be set once.

#rnx

diff --git a/Engine/Source/Runtime/Engine/Private/LevelTick.cpp b/Engine/Source/Runtime/Engine/Private/LevelTick.cpp
--- a/Engine/Source/Runtime/Engine/Private/LevelTick.cpp
+++ b/Engine/Source/Runtime/Engine/Private/LevelTick.cpp
@@ -792,17 +792,20 @@ void UWorld::RunTickGroup(ETickingGroup Group, bool bBlockTillComplete = true)
 static TAutoConsoleVariable<int32> CVarAllowAsyncRenderThreadUpdates(
 	TEXT("AllowAsyncRenderThreadUpdates"),
 	1,
-	TEXT("Used to control async renderthread updates. Also gated on FApp::ShouldUseThreadingForPerformance()."));
+	TEXT("Used to control async renderthread updates. Also gated on FApp::ShouldUseThreadingForPerformance()."),
+	ECVF_ReadOnly);
 
 static TAutoConsoleVariable<int32> CVarAllowAsyncRenderThreadUpdatesDuringGamethreadUpdates(
 	TEXT("AllowAsyncRenderThreadUpdatesDuringGamethreadUpdates"),
 	1,
-	TEXT("If > 0 then we do the gamethread updates _while_ doing parallel updates."));
+	TEXT("If > 0 then we do the gamethread updates _while_ doing parallel updates."),
+	ECVF_ReadOnly);
 
 static TAutoConsoleVariable<int32> CVarAllowAsyncRenderThreadUpdatesEditor(
 	TEXT("AllowAsyncRenderThreadUpdatesEditor"),
 	0,
-	TEXT("Used to control async renderthread updates in the editor."));
+	TEXT("Used to control async renderthread updates in the editor."),
+	ECVF_ReadOnly);
 
 namespace EComponentMarkedForEndOfFrameUpdateState
 {
@@ -901,8 +904,8 @@ void UWorld::MarkActorComponentForNeededEndOfFrameUpdate(UActorComponent* Compon
 	{
 		if (!bForceGameThread)
 		{
-			bool bAllowConcurrentUpdates = FApp::ShouldUseThreadingForPerformance() && 
-				(GIsEditor ? !!CVarAllowAsyncRenderThreadUpdatesEditor.GetValueOnGameThread() : !!CVarAllowAsyncRenderThreadUpdates.GetValueOnGameThread());
+			static bool bAllowConcurrentUpdates = FApp::ShouldUseThreadingForPerformance() && 
+				(GIsEditor ? !!CVarAllowAsyncRenderThreadUpdatesEditor.GetValueOnAnyThread() : !!CVarAllowAsyncRenderThreadUpdates.GetValueOnAnyThread());
 			bForceGameThread = !bAllowConcurrentUpdates 
 								// When there is no rendering thread force all updates on game thread,
 								// to avoid modifying scene structures from multiple task threads
