Added ShouldPerformAirControlForPathFollowing to control whether we should do the air control while following a path.

diff --git a/Engine/Source/Runtime/Engine/Classes/GameFramework/CharacterMovementComponent.h b/Engine/Source/Runtime/Engine/Classes/GameFramework/CharacterMovementComponent.h
--- a/Engine/Source/Runtime/Engine/Classes/GameFramework/CharacterMovementComponent.h
+++ b/Engine/Source/Runtime/Engine/Classes/GameFramework/CharacterMovementComponent.h
@@ -1254,13 +1254,19 @@ public:
 	virtual bool HasValidData() const;
 
 	/**
-	 * Update Velocity and Acceleration to air control in the desired Direction for character using path following.
+	 * If ShouldPerformAirControlForPathFollowing() returns true, it will update Velocity and Acceleration to air control in the desired Direction for character using path following.
 	 * @param Direction is the desired direction of movement
 	 * @param ZDiff is the height difference between the destination and the Pawn's current position
 	 * @see RequestDirectMove()
 	*/
 	virtual void PerformAirControlForPathFollowing(FVector Direction, float ZDiff);
 
+	/**
+	 * Whether Character should perform air control via PerformAirControlForPathFollowing when falling and following a path at the same time
+	 * Default implementation always returns true during MOVE_Falling.
+	 */
+	virtual bool ShouldPerformAirControlForPathFollowing() const;
+
 	/** Transition from walking to falling */
 	virtual void StartFalling(int32 Iterations, float remainingTime, float timeTick, const FVector& Delta, const FVector& subLoc);
 
diff --git a/Engine/Source/Runtime/Engine/Private/Components/CharacterMovementComponent.cpp b/Engine/Source/Runtime/Engine/Private/Components/CharacterMovementComponent.cpp
--- a/Engine/Source/Runtime/Engine/Private/Components/CharacterMovementComponent.cpp
+++ b/Engine/Source/Runtime/Engine/Private/Components/CharacterMovementComponent.cpp
@@ -3293,7 +3293,7 @@ void UCharacterMovementComponent::RequestDirectMove(const FVector& MoveVelocity,
 		return;
 	}
 
-	if (IsFalling())
+	if (ShouldPerformAirControlForPathFollowing())
 	{
 		const FVector FallVelocity = MoveVelocity.GetClampedToMaxSize(GetMaxSpeed());
 		PerformAirControlForPathFollowing(FallVelocity, FallVelocity.Z);
@@ -3310,6 +3310,11 @@ void UCharacterMovementComponent::RequestDirectMove(const FVector& MoveVelocity,
 	}
 }
 
+bool UCharacterMovementComponent::ShouldPerformAirControlForPathFollowing() const
+{
+	return IsFalling();
+}
+
 void UCharacterMovementComponent::RequestPathMove(const FVector& MoveInput)
 {
 	FVector AdjustedMoveInput(MoveInput);
