Fix NetProfiler reporting double Custom Delta Property size.

diff --git a/Engine/Source/Runtime/Engine/Private/DataReplication.cpp b/Engine/Source/Runtime/Engine/Private/DataReplication.cpp
--- a/Engine/Source/Runtime/Engine/Private/DataReplication.cpp
+++ b/Engine/Source/Runtime/Engine/Private/DataReplication.cpp
@@ -2235,7 +2235,9 @@ void FObjectReplicator::WritePropertyHeaderAndPayload(
 	// Send property name and optional array index.
 	check( FieldCache->FieldNetIndex <= ClassCache->GetMaxIndex() );
 
-	const int32 HeaderBits = OwningChannel->WriteFieldHeaderAndPayload( Bunch, ClassCache, FieldCache, NetFieldExportGroup, Payload );
+	// WriteFieldHeaderAndPayload will return the total number of bits written.
+	// So, we subtract out the Payload size to get the actual number of header bits.
+	const int32 HeaderBits = static_cast<int64>(OwningChannel->WriteFieldHeaderAndPayload( Bunch, ClassCache, FieldCache, NetFieldExportGroup, Payload )) - Payload.GetNumBits();
 
 	NETWORK_PROFILER( GNetworkProfiler.TrackWritePropertyHeader( Property, HeaderBits, nullptr ) );
 }
