Fix a bug where some ELLMTag::StaticMesh memory is incorrectly counted to ELLMTag::FileSystem on XB1

diff --git a/Engine/Source/Runtime/Engine/Private/Streaming/StaticMeshUpdate.cpp b/Engine/Source/Runtime/Engine/Private/Streaming/StaticMeshUpdate.cpp
--- a/Engine/Source/Runtime/Engine/Private/Streaming/StaticMeshUpdate.cpp
+++ b/Engine/Source/Runtime/Engine/Private/Streaming/StaticMeshUpdate.cpp
@@ -154,6 +154,8 @@ FStaticMeshStreamIn::~FStaticMeshStreamIn()
 template <bool bRenderThread>
 void FStaticMeshStreamIn::CreateBuffers_Internal(const FContext& Context)
 {
+	LLM_SCOPE(ELLMTag::StaticMesh);
+	
 	UStaticMesh* Mesh = Context.Mesh;
 	FStaticMeshRenderData* RenderData = Context.RenderData;
 	if (!IsCancelled() && Mesh && RenderData)
