- Move check to avoid Building tree on Archetype/CDO inside the BuildTreeIfOutdated method (the crash was because of UHierarchicalInstancedStaticMeshComponent::PreSave)
- Add checks in BuildTree and BuildTree async that object is not an Archetype/CDO

#rnx

diff --git a/Engine/Source/Runtime/Engine/Private/HierarchicalInstancedStaticMesh.cpp b/Engine/Source/Runtime/Engine/Private/HierarchicalInstancedStaticMesh.cpp
--- a/Engine/Source/Runtime/Engine/Private/HierarchicalInstancedStaticMesh.cpp
+++ b/Engine/Source/Runtime/Engine/Private/HierarchicalInstancedStaticMesh.cpp
@@ -1960,7 +1960,7 @@ void UHierarchicalInstancedStaticMeshComponent::PostEditChangeChainProperty(FPro
 		(PropertyChangedEvent.Property != NULL && PropertyChangedEvent.Property->GetFName() == "Transform") ||
 		(PropertyChangedEvent.Property != NULL && PropertyChangedEvent.Property->GetFName() == "StaticMesh"))
 	{
-		if (FApp::CanEverRender() && !HasAnyFlags(RF_ClassDefaultObject | RF_ArchetypeObject))
+		if (FApp::CanEverRender())
 		{
 			BuildTreeIfOutdated(false, false);
 		}
@@ -2049,7 +2049,7 @@ void UHierarchicalInstancedStaticMeshComponent::PostDuplicate(bool bDuplicateFor
 {
 	Super::PostDuplicate(bDuplicateForPIE);
 
-	if (!HasAnyFlags(RF_ClassDefaultObject | RF_ArchetypeObject) && bDuplicateForPIE)
+	if (bDuplicateForPIE)
 	{
 		BuildTreeIfOutdated(false, false);
 	}
@@ -2415,6 +2415,7 @@ void UHierarchicalInstancedStaticMeshComponent::PostBuildStats()
 
 void UHierarchicalInstancedStaticMeshComponent::BuildTree()
 {
+	check(!HasAnyFlags(RF_ClassDefaultObject | RF_ArchetypeObject));
 	checkSlow(IsInGameThread());
 
 	// If we try to build the tree with the static mesh not fully loaded, we can end up in an inconsistent state which ends in a crash later
@@ -2666,20 +2667,22 @@ void UHierarchicalInstancedStaticMeshComponent::OnComponentCreated()
 {
 	Super::OnComponentCreated();
 
-	if (FApp::CanEverRender() && !HasAnyFlags(RF_ClassDefaultObject | RF_ArchetypeObject))
+	// if we are pasting/duplicating this component, it may be created with some instances already in place
+	// in this case, need to ensure that the tree is properly created
+	if (FApp::CanEverRender() && PerInstanceSMData.Num() > 0 && ClusterTreePtr->Num() == 0)
 	{
-		// if we are pasting/duplicating this component, it may be created with some instances already in place
-		// in this case, need to ensure that the tree is properly created
-		if (PerInstanceSMData.Num() > 0 && ClusterTreePtr->Num() == 0)
-		{
-			const bool bForceUpdate = true;
-			BuildTreeIfOutdated(false, bForceUpdate);
-		}
+		const bool bForceUpdate = true;
+		BuildTreeIfOutdated(false, bForceUpdate);
 	}
 }
 
 bool UHierarchicalInstancedStaticMeshComponent::BuildTreeIfOutdated(bool Async, bool ForceUpdate)
 {
+	if (HasAnyFlags(RF_ClassDefaultObject | RF_ArchetypeObject))
+	{
+		return false;
+	}
+
 	if (ForceUpdate 
 		|| InstanceUpdateCmdBuffer.NumTotalCommands() != 0
 		|| InstanceReorderTable.Num() != PerInstanceSMData.Num()
@@ -2720,6 +2723,7 @@ bool UHierarchicalInstancedStaticMeshComponent::BuildTreeIfOutdated(bool Async,
 
 void UHierarchicalInstancedStaticMeshComponent::BuildTreeAsync()
 {
+	check(!HasAnyFlags(RF_ClassDefaultObject | RF_ArchetypeObject));
 	check(IsInGameThread());
 
 	// If we try to build the tree with the static mesh not fully loaded, we can end up in an inconsistent state which ends in a crash later
