Export the CalcView function from SceneCapture2D.

Marc.Audy

diff --git a/Engine/Source/Runtime/Engine/Classes/Engine/SceneCapture2D.h b/Engine/Source/Runtime/Engine/Classes/Engine/SceneCapture2D.h
--- a/Engine/Source/Runtime/Engine/Classes/Engine/SceneCapture2D.h
+++ b/Engine/Source/Runtime/Engine/Classes/Engine/SceneCapture2D.h
@@ -26,7 +26,7 @@ public:
 	UFUNCTION(BlueprintCallable, Category="Rendering")
 	void OnInterpToggle(bool bEnable);
 
-	virtual void CalcCamera(float DeltaTime, struct FMinimalViewInfo& OutMinimalViewInfo) override;
+	ENGINE_API virtual void CalcCamera(float DeltaTime, struct FMinimalViewInfo& OutMinimalViewInfo) override;
 
 	/** Returns CaptureComponent2D subobject **/
 	ENGINE_API class USceneCaptureComponent2D* GetCaptureComponent2D() const { return CaptureComponent2D; }
