Back out changelist 6183954

MarkActorComponentForNeededEndOfFrameUpdate is not thread safe, so using GetValueOnGameThread is fine.

diff --git a/Engine/Source/Runtime/Engine/Private/LevelTick.cpp b/Engine/Source/Runtime/Engine/Private/LevelTick.cpp
--- a/Engine/Source/Runtime/Engine/Private/LevelTick.cpp
+++ b/Engine/Source/Runtime/Engine/Private/LevelTick.cpp
@@ -792,20 +792,17 @@ void UWorld::RunTickGroup(ETickingGroup Group, bool bBlockTillComplete = true)
 static TAutoConsoleVariable<int32> CVarAllowAsyncRenderThreadUpdates(
 	TEXT("AllowAsyncRenderThreadUpdates"),
 	1,
-	TEXT("Used to control async renderthread updates. Also gated on FApp::ShouldUseThreadingForPerformance()."),
-	ECVF_ReadOnly);
+	TEXT("Used to control async renderthread updates. Also gated on FApp::ShouldUseThreadingForPerformance()."));
 
 static TAutoConsoleVariable<int32> CVarAllowAsyncRenderThreadUpdatesDuringGamethreadUpdates(
 	TEXT("AllowAsyncRenderThreadUpdatesDuringGamethreadUpdates"),
 	1,
-	TEXT("If > 0 then we do the gamethread updates _while_ doing parallel updates."),
-	ECVF_ReadOnly);
+	TEXT("If > 0 then we do the gamethread updates _while_ doing parallel updates."));
 
 static TAutoConsoleVariable<int32> CVarAllowAsyncRenderThreadUpdatesEditor(
 	TEXT("AllowAsyncRenderThreadUpdatesEditor"),
 	0,
-	TEXT("Used to control async renderthread updates in the editor."),
-	ECVF_ReadOnly);
+	TEXT("Used to control async renderthread updates in the editor."));
 
 namespace EComponentMarkedForEndOfFrameUpdateState
 {
@@ -904,8 +901,8 @@ void UWorld::MarkActorComponentForNeededEndOfFrameUpdate(UActorComponent* Compon
 	{
 		if (!bForceGameThread)
 		{
-			static bool bAllowConcurrentUpdates = FApp::ShouldUseThreadingForPerformance() && 
-				(GIsEditor ? !!CVarAllowAsyncRenderThreadUpdatesEditor.GetValueOnAnyThread() : !!CVarAllowAsyncRenderThreadUpdates.GetValueOnAnyThread());
+			bool bAllowConcurrentUpdates = FApp::ShouldUseThreadingForPerformance() && 
+				(GIsEditor ? !!CVarAllowAsyncRenderThreadUpdatesEditor.GetValueOnGameThread() : !!CVarAllowAsyncRenderThreadUpdates.GetValueOnGameThread());
 			bForceGameThread = !bAllowConcurrentUpdates 
 								// When there is no rendering thread force all updates on game thread,
 								// to avoid modifying scene structures from multiple task threads
