Fixed crash when shutting down PIE in some rare cases due to the Scene being nullptr.  (Changed code to ensure instead of crashing outright.)

#UE4 #ReleaseNote

diff --git a/Engine/Source/Runtime/Engine/Private/Components/LightComponent.cpp b/Engine/Source/Runtime/Engine/Private/Components/LightComponent.cpp
--- a/Engine/Source/Runtime/Engine/Private/Components/LightComponent.cpp
+++ b/Engine/Source/Runtime/Engine/Private/Components/LightComponent.cpp
@@ -782,7 +782,12 @@ void ULightComponent::SendRenderTransform_Concurrent()
 void ULightComponent::DestroyRenderState_Concurrent()
 {
 	Super::DestroyRenderState_Concurrent();
-	GetWorld()->Scene->RemoveLight(this);
+	UWorld* MyWorld = GetWorld();
+	check(MyWorld != nullptr);
+	if (ensure(MyWorld->Scene != nullptr))
+	{
+		MyWorld->Scene->RemoveLight(this);
+	}
 	bAddedToSceneVisible = false;
 }
 
