- Add validation for UParticleModuleLocationBoneSocket to ensure that all LOD levels have a matching count for SourceLocations if they will tick on update.  The index is stored on the particle on spawn, so if an update module has less entires we would index out of bounds.
- The fix is applied in editor / cooking and not in cooked builds

diff --git a/Engine/Source/Runtime/Engine/Private/Particles/ParticleComponents.cpp b/Engine/Source/Runtime/Engine/Private/Particles/ParticleComponents.cpp
--- a/Engine/Source/Runtime/Engine/Private/Particles/ParticleComponents.cpp
+++ b/Engine/Source/Runtime/Engine/Private/Particles/ParticleComponents.cpp
@@ -65,6 +65,7 @@
 #include "Particles/Lifetime/ParticleModuleLifetimeBase.h"
 #include "Particles/Lifetime/ParticleModuleLifetime.h"
 #include "Particles/Light/ParticleModuleLight.h"
+#include "Particles/Location/ParticleModuleLocationBoneSocket.h"
 #include "Particles/Material/ParticleModuleMeshMaterial.h"
 #include "Particles/Modules/Location/ParticleModulePivotOffset.h"
 #include "Particles/Orbit/ParticleModuleOrbit.h"
@@ -1818,6 +1819,15 @@ void UParticleEmitter::CacheEmitterModuleInfo()
 				? ModuleSubUVAnimation
 				: NULL;
 		}
+		// Perform validation / fixup on some modules that can cause crashes if LODs / Modules are out of sync
+		// This should only be applied on uncooked builds to avoid wasting cycles
+		else if ( !FPlatformProperties::RequiresCookedData() )
+		{
+			if (ParticleModule->IsA(UParticleModuleLocationBoneSocket::StaticClass()))
+			{
+				UParticleModuleLocationBoneSocket::ValidateLODLevels(this, ModuleIdx);
+			}
+		}
 
 		// Set bMeshRotationActive if module says so
 		if(!bMeshRotationActive && ParticleModule->TouchesMeshRotation())
diff --git a/Engine/Source/Runtime/Engine/Private/Particles/ParticleModules_Location.cpp b/Engine/Source/Runtime/Engine/Private/Particles/ParticleModules_Location.cpp
--- a/Engine/Source/Runtime/Engine/Private/Particles/ParticleModules_Location.cpp
+++ b/Engine/Source/Runtime/Engine/Private/Particles/ParticleModules_Location.cpp
@@ -1585,6 +1585,43 @@ void UParticleModuleLocationBoneSocket::SetSourceIndexMode()
 	}
 }
 
+void UParticleModuleLocationBoneSocket::ValidateLODLevels(UParticleEmitter* Emitter, int32 iModule)
+{
+	const int32 NumLODLevels = Emitter->LODLevels.Num();
+	if ( NumLODLevels <= 1 )
+		return;
+
+	bool bRequiresValidate = false;
+	for (int32 iLOD=0; iLOD < NumLODLevels; ++iLOD)
+	{
+		UParticleLODLevel* LODLevel = Emitter->LODLevels[iLOD];
+		UParticleModuleLocationBoneSocket* LODModule = CastChecked<UParticleModuleLocationBoneSocket>(LODLevel->Modules[iModule]);
+		if (!LODModule->bInheritBoneVelocity && !LODModule->bUpdatePositionEachFrame)
+			continue;
+
+		bRequiresValidate = true;
+	}
+
+	if (bRequiresValidate)
+	{
+		UParticleModuleLocationBoneSocket* HighestLODModule = CastChecked<UParticleModuleLocationBoneSocket>(Emitter->LODLevels[0]->Modules[iModule]);
+
+		for (int32 iLOD=1; iLOD < NumLODLevels; ++iLOD)
+		{
+			UParticleLODLevel* LODLevel = Emitter->LODLevels[iLOD];
+			UParticleModuleLocationBoneSocket* LODModule = CastChecked<UParticleModuleLocationBoneSocket>(LODLevel->Modules[iModule]);
+			if (LODModule == HighestLODModule)
+				continue;
+
+			if (LODModule->SourceLocations.Num() != HighestLODModule->SourceLocations.Num())
+			{
+				UE_LOG(LogParticles, Warning, TEXT("UParticleModuleLocationBoneSocket: ParticleSystem '%s' contains an Emitter '%s' LOD '%d' with inconsistent SourceLocations that could lead to a crash or unexpected behaviour.  Fixing by copying the Highest LOD SourceLocations to LOD."), *Emitter->GetOuter()->GetFullName(), *Emitter->GetEmitterName().ToString(), iLOD);
+				LODModule->SourceLocations = HighestLODModule->SourceLocations;
+			}
+		}
+	}
+}
+
 void UParticleModuleLocationBoneSocket::PostLoad()
 {
 	Super::PostLoad();
