- Fix missing header from checkin 6245967

diff --git a/Engine/Source/Runtime/Engine/Classes/Particles/Location/ParticleModuleLocationBoneSocket.h b/Engine/Source/Runtime/Engine/Classes/Particles/Location/ParticleModuleLocationBoneSocket.h
--- a/Engine/Source/Runtime/Engine/Classes/Particles/Location/ParticleModuleLocationBoneSocket.h
+++ b/Engine/Source/Runtime/Engine/Classes/Particles/Location/ParticleModuleLocationBoneSocket.h
@@ -188,6 +188,9 @@ class ENGINE_API UParticleModuleLocationBoneSocket : public UParticleModuleLocat
 	FORCEINLINE bool InheritingBoneVelocity()const { return bInheritBoneVelocity && !bUpdatePositionEachFrame; }
 
 	void SetSourceIndexMode();
+
+	/** Validate that all LOD levels will not have issues with indexing into SourceLocations */
+	static void ValidateLODLevels(UParticleEmitter* Emitter, int32 iModule);
 };
 
 /** ModuleLocationBoneSocket instance payload */
