Added an option to disable 'Pixel Depth Offset' on mobile platforms (r.Mobile.AllowPixelDepthOffset  enabled by default)

diff --git a/Engine/Source/Runtime/Core/Private/HAL/ConsoleManager.cpp b/Engine/Source/Runtime/Core/Private/HAL/ConsoleManager.cpp
--- a/Engine/Source/Runtime/Core/Private/HAL/ConsoleManager.cpp
+++ b/Engine/Source/Runtime/Core/Private/HAL/ConsoleManager.cpp
@@ -1998,6 +1998,13 @@ static TAutoConsoleVariable<int32> CVarMobileAllowDitheredLODTransition(
 	ECVF_ReadOnly | ECVF_RenderThreadSafe
 );
 
+static TAutoConsoleVariable<int32> CVarMobileAllowPixelDepthOffset(
+	TEXT("r.Mobile.AllowPixelDepthOffset"),
+	1,
+	TEXT("Whether to allow 'Pixel Depth Offset' in materials for ES3.1 feature level. Depth modification in pixel shaders may reduce GPU performance"),
+	ECVF_ReadOnly | ECVF_RenderThreadSafe
+);
+
 static TAutoConsoleVariable<int32> CVarSetClearSceneMethod(
 	TEXT("r.ClearSceneMethod"),
 	1,
diff --git a/Engine/Source/Runtime/Engine/Private/Materials/HLSLMaterialTranslator.h b/Engine/Source/Runtime/Engine/Private/Materials/HLSLMaterialTranslator.h
--- a/Engine/Source/Runtime/Engine/Private/Materials/HLSLMaterialTranslator.h
+++ b/Engine/Source/Runtime/Engine/Private/Materials/HLSLMaterialTranslator.h
@@ -681,7 +681,7 @@ public:
 			}
 
 			bUsesEmissiveColor = IsMaterialPropertyUsed(MP_EmissiveColor, Chunk[MP_EmissiveColor], FLinearColor(0, 0, 0, 0), 3);
-			bUsesPixelDepthOffset = IsMaterialPropertyUsed(MP_PixelDepthOffset, Chunk[MP_PixelDepthOffset], FLinearColor(0, 0, 0, 0), 1)
+			bUsesPixelDepthOffset = (AllowPixelDepthOffset(Platform) && IsMaterialPropertyUsed(MP_PixelDepthOffset, Chunk[MP_PixelDepthOffset], FLinearColor(0, 0, 0, 0), 1))
 				|| (Domain == MD_DeferredDecal && Material->GetDecalBlendMode() == DBM_Volumetric_DistanceFunction);
 
 			bUsesWorldPositionOffset = IsMaterialPropertyUsed(MP_WorldPositionOffset, Chunk[MP_WorldPositionOffset], FLinearColor(0, 0, 0, 0), 3);
diff --git a/Engine/Source/Runtime/RenderCore/Private/RenderUtils.cpp b/Engine/Source/Runtime/RenderCore/Private/RenderUtils.cpp
--- a/Engine/Source/Runtime/RenderCore/Private/RenderUtils.cpp
+++ b/Engine/Source/Runtime/RenderCore/Private/RenderUtils.cpp
@@ -898,6 +898,16 @@ RENDERCORE_API bool IsSimpleForwardShadingEnabled(EShaderPlatform Platform)
 	return CVar->GetValueOnAnyThread() != 0 && PlatformSupportsSimpleForwardShading(Platform);
 }
 
+RENDERCORE_API bool AllowPixelDepthOffset(EShaderPlatform Platform)
+{
+	if (IsMobilePlatform(Platform))
+	{
+		static const auto CVar = IConsoleManager::Get().FindTConsoleVariableDataInt(TEXT("r.Mobile.AllowPixelDepthOffset"));
+		return CVar->GetValueOnAnyThread() != 0;
+	}
+	return true;
+}
+
 RENDERCORE_API int32 GUseForwardShading = 0;
 static FAutoConsoleVariableRef CVarForwardShading(
 	TEXT("r.ForwardShading"),
diff --git a/Engine/Source/Runtime/RenderCore/Private/Shader.cpp b/Engine/Source/Runtime/RenderCore/Private/Shader.cpp
--- a/Engine/Source/Runtime/RenderCore/Private/Shader.cpp
+++ b/Engine/Source/Runtime/RenderCore/Private/Shader.cpp
@@ -2157,6 +2157,11 @@ void ShaderMapAppendKeyString(EShaderPlatform Platform, FString& KeyString)
 		}
 	}
 
+	if (!AllowPixelDepthOffset(Platform))
+	{
+		KeyString += TEXT("_NoPDO");
+	}
+	
 	if (IsD3DPlatform(Platform, false))
 	{
 		static const auto CVar = IConsoleManager::Get().FindConsoleVariable(TEXT("r.D3D.RemoveUnusedInterpolators"));
diff --git a/Engine/Source/Runtime/RenderCore/Public/RenderUtils.h b/Engine/Source/Runtime/RenderCore/Public/RenderUtils.h
--- a/Engine/Source/Runtime/RenderCore/Public/RenderUtils.h
+++ b/Engine/Source/Runtime/RenderCore/Public/RenderUtils.h
@@ -428,6 +428,8 @@ RENDERCORE_API bool PlatformSupportsSimpleForwardShading(EShaderPlatform Platfor
 
 RENDERCORE_API bool IsSimpleForwardShadingEnabled(EShaderPlatform Platform);
 
+RENDERCORE_API bool AllowPixelDepthOffset(EShaderPlatform Platform);
+
 /** Returns if ForwardShading is enabled. Only valid for the current platform (otherwise call ITargetPlatform::UsesForwardShading()). */
 inline bool IsForwardShadingEnabled(EShaderPlatform Platform)
 {
