Early out of UActorChannel::ReplicateSubobject if the object is null or pending kill to avoid a subsequent crash

diff --git a/Engine/Source/Runtime/Engine/Private/DataChannel.cpp b/Engine/Source/Runtime/Engine/Private/DataChannel.cpp
--- a/Engine/Source/Runtime/Engine/Private/DataChannel.cpp
+++ b/Engine/Source/Runtime/Engine/Private/DataChannel.cpp
@@ -3818,6 +3818,11 @@ bool UActorChannel::ReplicateSubobject(UObject *Obj, FOutBunch &Bunch, const FRe
 {
 	SCOPE_CYCLE_UOBJECT(ActorChannelRepSubObj, Obj);
 
+	if (!Obj || Obj->IsPendingKill())
+	{
+		return false;
+	}
+
 	TWeakObjectPtr<UObject> WeakObj(Obj);
 
 	// Hack for now: subobjects are SupportsObject==false until they are replicated via ::ReplicateSUbobject, and then we make them supported
