Adding versions of Set Timer by Event and Set Timer by [REDACTED_PERSON] that take a Variance parameter to more easily enable content creators to have some random distribution in their timers ... this is useful when we have lots of a specific type of object that we don't necessarily want to have all fire timers at the exact same moment (i.e. hitch prevention).

diff --git a/Engine/Source/Runtime/Engine/Classes/Kismet/KismetSystemLibrary.h b/Engine/Source/Runtime/Engine/Classes/Kismet/KismetSystemLibrary.h
--- a/Engine/Source/Runtime/Engine/Classes/Kismet/KismetSystemLibrary.h
+++ b/Engine/Source/Runtime/Engine/Classes/Kismet/KismetSystemLibrary.h
@@ -468,6 +468,17 @@ class ENGINE_API UKismetSystemLibrary : public UBlueprintFunctionLibrary
 	 */
 	UFUNCTION(BlueprintCallable, meta=(DisplayName = "Set Timer by Event", ScriptName = "SetTimerDelegate"), Category="Utilities|Time")
 	static FTimerHandle K2_SetTimerDelegate(UPARAM(DisplayName="Event") FTimerDynamicDelegate Delegate, float Time, bool bLooping);
+	
+	/**
+	 * Set a timer to execute delegate. Setting an existing timer will reset that timer with updated parameters.
+	 * @param Event			Event. Can be a K2 function or a Custom Event.
+	 * @param Time			How long to wait before executing the delegate, in seconds. Setting a timer to <= 0 seconds will clear it if it is set.
+	 * @param Variance		We add a RandomRange[-Variance, Variance] to Time.
+	 * @param bLooping		True to keep executing the delegate every Time seconds, false to execute delegate only once.
+	 * @return				The timer handle to pass to other timer functions to manipulate this timer.
+	 */
+	UFUNCTION(BlueprintCallable, meta=(DisplayName = "Set Timer by Event with Variance ", ScriptName = "SetTimerDelegateWithVariance"), Category="Utilities|Time")
+	static FTimerHandle K2_SetTimerDelegateWithVariance(UPARAM(DisplayName="Event") FTimerDynamicDelegate Delegate, float Time, float Variance, bool bLooping);
 
 	/**
 	 * Clears a set timer.
@@ -628,6 +639,19 @@ class ENGINE_API UKismetSystemLibrary : public UBlueprintFunctionLibrary
 	 */
 	UFUNCTION(BlueprintCallable, meta=(DisplayName = "Set Timer by Function Name", ScriptName = "SetTimer", DefaultToSelf = "Object"), Category="Utilities|Time")
 	static FTimerHandle K2_SetTimer(UObject* Object, FString FunctionName, float Time, bool bLooping);
+	
+	/**
+	 * Set a timer to execute delegate. Setting an existing timer will reset that timer with updated parameters.  This version takes a variance param that randomizes the Time within a range +/- the variance.
+	 * Useful in situations where you have a lot of objects using a timer but you don't want all of the object timers to execute on the exact same frame (i.e. helpful in avoiding hitches)
+	 * @param Object		Object that implements the delegate function. Defaults to self (this blueprint)
+	 * @param FunctionName	Delegate function name. Can be a K2 function or a Custom Event.
+	 * @param Time			How long to wait before executing the delegate, in seconds. Setting a timer to <= 0 seconds will clear it if it is set.
+	 * @param Variance		We add a RandomRange[-Variance, Variance] to Time.
+	 * @param bLooping		true to keep executing the delegate every Time seconds, false to execute delegate only once.
+	 * @return				The timer handle to pass to other timer functions to manipulate this timer.
+	 */
+	UFUNCTION(BlueprintCallable, meta=(DisplayName = "Set Timer by Function Name with Variance", ScriptName = "SetTimerWithVariance", DefaultToSelf = "Object"), Category="Utilities|Time")
+	static FTimerHandle K2_SetTimerWithVariance(UObject* Object, FString FunctionName, float Time, float Variance, bool bLooping);
 
 	/**
 	 * Clears a set timer.
diff --git a/Engine/Source/Runtime/Engine/Private/KismetSystemLibrary.cpp b/Engine/Source/Runtime/Engine/Private/KismetSystemLibrary.cpp
--- a/Engine/Source/Runtime/Engine/Private/KismetSystemLibrary.cpp
+++ b/Engine/Source/Runtime/Engine/Private/KismetSystemLibrary.cpp
@@ -409,6 +409,29 @@ FTimerHandle UKismetSystemLibrary::K2_SetTimer(UObject* Object, FString Function
 	return K2_SetTimerDelegate(Delegate, Time, bLooping);
 }
 
+FTimerHandle UKismetSystemLibrary::K2_SetTimerWithVariance(UObject* Object, FString FunctionName, float Time, float Variance, bool bLooping)
+{
+	FName const FunctionFName(*FunctionName);
+
+	if (Object)
+	{
+		UFunction* const Func = Object->FindFunction(FunctionFName);
+		if (Func && (Func->ParmsSize > 0))
+		{
+			// User passed in a valid function, but one that takes parameters
+			// FTimerDynamicDelegate expects zero parameters and will choke on execution if it tries
+			// to execute a mismatched function
+			UE_LOG(LogBlueprintUserMessages, Warning, TEXT("SetTimer passed a function (%s) that expects parameters."), *FunctionName);
+			return FTimerHandle();
+		}
+	}
+
+	FTimerDynamicDelegate Delegate;
+	Delegate.BindUFunction(Object, FunctionFName);
+	Time = Time + FMath::RandRange(-Variance, Variance);
+	return K2_SetTimerDelegate(Delegate, Time, bLooping);
+}
+
 FTimerHandle UKismetSystemLibrary::K2_SetTimerDelegate(FTimerDynamicDelegate Delegate, float Time, bool bLooping)
 {
 	FTimerHandle Handle;
@@ -432,6 +455,30 @@ FTimerHandle UKismetSystemLibrary::K2_SetTimerDelegate(FTimerDynamicDelegate Del
 	return Handle;
 }
 
+FTimerHandle UKismetSystemLibrary::K2_SetTimerDelegateWithVariance(FTimerDynamicDelegate Delegate, float Time, float Variance, bool bLooping)
+{
+	FTimerHandle Handle;
+	if (Delegate.IsBound())
+	{
+		const UWorld* const World = GEngine->GetWorldFromContextObject(Delegate.GetUObject(), EGetWorldErrorMode::LogAndReturnNull);
+		if (World)
+		{
+			FTimerManager& TimerManager = World->GetTimerManager();
+			Handle = TimerManager.K2_FindDynamicTimerHandle(Delegate);
+			Time = Time + FMath::RandRange(-Variance, Variance);
+			TimerManager.SetTimer(Handle, Delegate, Time, bLooping);
+		}
+	}
+	else
+	{
+		UE_LOG(LogBlueprintUserMessages, Warning,
+			TEXT("SetTimer passed a bad function (%s) or object (%s)"),
+			*Delegate.GetFunctionName().ToString(), *GetNameSafe(Delegate.GetUObject()));
+	}
+
+	return Handle;
+}
+
 void UKismetSystemLibrary::K2_ClearTimer(UObject* Object, FString FunctionName)
 {
 	FTimerDynamicDelegate Delegate;
