Reworking the timer variance submitted [at] to instead leverage the pre-existing InFirstDelay parameter exposed in the timer manager's API.  Also using the existing Set Timer functions and adding the initial start delay as an optional advanced parameter.

diff --git a/Engine/Source/Runtime/Engine/Classes/Kismet/KismetSystemLibrary.h b/Engine/Source/Runtime/Engine/Classes/Kismet/KismetSystemLibrary.h
--- a/Engine/Source/Runtime/Engine/Classes/Kismet/KismetSystemLibrary.h
+++ b/Engine/Source/Runtime/Engine/Classes/Kismet/KismetSystemLibrary.h
@@ -461,24 +461,14 @@ class ENGINE_API UKismetSystemLibrary : public UBlueprintFunctionLibrary
 
 	/**
 	 * Set a timer to execute delegate. Setting an existing timer will reset that timer with updated parameters.
-	 * @param Event			Event. Can be a K2 function or a Custom Event.
-	 * @param Time			How long to wait before executing the delegate, in seconds. Setting a timer to <= 0 seconds will clear it if it is set.
-	 * @param bLooping		True to keep executing the delegate every Time seconds, false to execute delegate only once.
-	 * @return				The timer handle to pass to other timer functions to manipulate this timer.
+	 * @param Event				Event. Can be a K2 function or a Custom Event.
+	 * @param Time				How long to wait before executing the delegate, in seconds. Setting a timer to <= 0 seconds will clear it if it is set.
+	 * @param bLooping			True to keep executing the delegate every Time seconds, false to execute delegate only once.
+	 * @param InitialStartDelay Initial delay passed to the timer manager to allow some variance in when the timer starts.
+	 * @return					The timer handle to pass to other timer functions to manipulate this timer.
 	 */
-	UFUNCTION(BlueprintCallable, meta=(DisplayName = "Set Timer by Event", ScriptName = "SetTimerDelegate"), Category="Utilities|Time")
-	static FTimerHandle K2_SetTimerDelegate(UPARAM(DisplayName="Event") FTimerDynamicDelegate Delegate, float Time, bool bLooping);
-	
-	/**
-	 * Set a timer to execute delegate. Setting an existing timer will reset that timer with updated parameters.
-	 * @param Event			Event. Can be a K2 function or a Custom Event.
-	 * @param Time			How long to wait before executing the delegate, in seconds. Setting a timer to <= 0 seconds will clear it if it is set.
-	 * @param Variance		We add a RandomRange[-Variance, Variance] to Time.
-	 * @param bLooping		True to keep executing the delegate every Time seconds, false to execute delegate only once.
-	 * @return				The timer handle to pass to other timer functions to manipulate this timer.
-	 */
-	UFUNCTION(BlueprintCallable, meta=(DisplayName = "Set Timer by Event with Variance ", ScriptName = "SetTimerDelegateWithVariance"), Category="Utilities|Time")
-	static FTimerHandle K2_SetTimerDelegateWithVariance(UPARAM(DisplayName="Event") FTimerDynamicDelegate Delegate, float Time, float Variance, bool bLooping);
+	UFUNCTION(BlueprintCallable, meta=(DisplayName = "Set Timer by Event", ScriptName = "SetTimerDelegate", AdvancedDisplay="InitialStartDelay"), Category="Utilities|Time")
+	static FTimerHandle K2_SetTimerDelegate(UPARAM(DisplayName="Event") FTimerDynamicDelegate Delegate, float Time, bool bLooping, float InitialStartDelay = 0.f);
 
 	/**
 	 * Clears a set timer.
@@ -631,27 +621,15 @@ class ENGINE_API UKismetSystemLibrary : public UBlueprintFunctionLibrary
 
 	/**
 	 * Set a timer to execute delegate. Setting an existing timer will reset that timer with updated parameters.
-	 * @param Object		Object that implements the delegate function. Defaults to self (this blueprint)
-	 * @param FunctionName	Delegate function name. Can be a K2 function or a Custom Event.
-	 * @param Time			How long to wait before executing the delegate, in seconds. Setting a timer to <= 0 seconds will clear it if it is set.
-	 * @param bLooping		true to keep executing the delegate every Time seconds, false to execute delegate only once.
-	 * @return				The timer handle to pass to other timer functions to manipulate this timer.
-	 */
-	UFUNCTION(BlueprintCallable, meta=(DisplayName = "Set Timer by Function Name", ScriptName = "SetTimer", DefaultToSelf = "Object"), Category="Utilities|Time")
-	static FTimerHandle K2_SetTimer(UObject* Object, FString FunctionName, float Time, bool bLooping);
-	
-	/**
-	 * Set a timer to execute delegate. Setting an existing timer will reset that timer with updated parameters.  This version takes a variance param that randomizes the Time within a range +/- the variance.
-	 * Useful in situations where you have a lot of objects using a timer but you don't want all of the object timers to execute on the exact same frame (i.e. helpful in avoiding hitches)
-	 * @param Object		Object that implements the delegate function. Defaults to self (this blueprint)
-	 * @param FunctionName	Delegate function name. Can be a K2 function or a Custom Event.
-	 * @param Time			How long to wait before executing the delegate, in seconds. Setting a timer to <= 0 seconds will clear it if it is set.
-	 * @param Variance		We add a RandomRange[-Variance, Variance] to Time.
-	 * @param bLooping		true to keep executing the delegate every Time seconds, false to execute delegate only once.
-	 * @return				The timer handle to pass to other timer functions to manipulate this timer.
-	 */
-	UFUNCTION(BlueprintCallable, meta=(DisplayName = "Set Timer by Function Name with Variance", ScriptName = "SetTimerWithVariance", DefaultToSelf = "Object"), Category="Utilities|Time")
-	static FTimerHandle K2_SetTimerWithVariance(UObject* Object, FString FunctionName, float Time, float Variance, bool bLooping);
+	 * @param Object			Object that implements the delegate function. Defaults to self (this blueprint)
+	 * @param FunctionName		Delegate function name. Can be a K2 function or a Custom Event.
+	 * @param Time				How long to wait before executing the delegate, in seconds. Setting a timer to <= 0 seconds will clear it if it is set.
+	 * @param bLooping			true to keep executing the delegate every Time seconds, false to execute delegate only once.
+	 * @param InitialStartDelay Initial delay passed to the timer manager to allow some variance in when the timer starts.
+	 * @return					The timer handle to pass to other timer functions to manipulate this timer.
+	 */
+	UFUNCTION(BlueprintCallable, meta=(DisplayName = "Set Timer by Function Name", ScriptName = "SetTimer", DefaultToSelf = "Object", AdvancedDisplay="InitialStartDelay"), Category="Utilities|Time")
+	static FTimerHandle K2_SetTimer(UObject* Object, FString FunctionName, float Time, bool bLooping, float InitialStartDelay = 0.f);
 
 	/**
 	 * Clears a set timer.
diff --git a/Engine/Source/Runtime/Engine/Private/KismetSystemLibrary.cpp b/Engine/Source/Runtime/Engine/Private/KismetSystemLibrary.cpp
--- a/Engine/Source/Runtime/Engine/Private/KismetSystemLibrary.cpp
+++ b/Engine/Source/Runtime/Engine/Private/KismetSystemLibrary.cpp
@@ -387,7 +387,7 @@ FTimerHandle UKismetSystemLibrary::K2_InvalidateTimerHandle(FTimerHandle& TimerH
 	return TimerHandle;
 }
 
-FTimerHandle UKismetSystemLibrary::K2_SetTimer(UObject* Object, FString FunctionName, float Time, bool bLooping)
+FTimerHandle UKismetSystemLibrary::K2_SetTimer(UObject* Object, FString FunctionName, float Time, bool bLooping, float InitialStartDelay)
 {
 	FName const FunctionFName(*FunctionName);
 
@@ -404,35 +404,17 @@ FTimerHandle UKismetSystemLibrary::K2_SetTimer(UObject* Object, FString Function
 		}
 	}
 
-	FTimerDynamicDelegate Delegate;
-	Delegate.BindUFunction(Object, FunctionFName);
-	return K2_SetTimerDelegate(Delegate, Time, bLooping);
-}
-
-FTimerHandle UKismetSystemLibrary::K2_SetTimerWithVariance(UObject* Object, FString FunctionName, float Time, float Variance, bool bLooping)
-{
-	FName const FunctionFName(*FunctionName);
-
-	if (Object)
+	if (Time <= 0.f || InitialStartDelay < 0.f)
 	{
-		UFunction* const Func = Object->FindFunction(FunctionFName);
-		if (Func && (Func->ParmsSize > 0))
-		{
-			// User passed in a valid function, but one that takes parameters
-			// FTimerDynamicDelegate expects zero parameters and will choke on execution if it tries
-			// to execute a mismatched function
-			UE_LOG(LogBlueprintUserMessages, Warning, TEXT("SetTimer passed a function (%s) that expects parameters."), *FunctionName);
-			return FTimerHandle();
-		}
+		FFrame::KismetExecutionMessage(TEXT("SetTimer passed a negative time or initial start delay.  The associated timer may fail to fire!"), ELogVerbosity::Warning);
 	}
 
 	FTimerDynamicDelegate Delegate;
 	Delegate.BindUFunction(Object, FunctionFName);
-	Time = Time + FMath::RandRange(-Variance, Variance);
-	return K2_SetTimerDelegate(Delegate, Time, bLooping);
+	return K2_SetTimerDelegate(Delegate, Time, bLooping, InitialStartDelay);
 }
 
-FTimerHandle UKismetSystemLibrary::K2_SetTimerDelegate(FTimerDynamicDelegate Delegate, float Time, bool bLooping)
+FTimerHandle UKismetSystemLibrary::K2_SetTimerDelegate(FTimerDynamicDelegate Delegate, float Time, bool bLooping, float InitialStartDelay)
 {
 	FTimerHandle Handle;
 	if (Delegate.IsBound())
@@ -440,38 +422,19 @@ FTimerHandle UKismetSystemLibrary::K2_SetTimerDelegate(FTimerDynamicDelegate Del
 		const UWorld* const World = GEngine->GetWorldFromContextObject(Delegate.GetUObject(), EGetWorldErrorMode::LogAndReturnNull);
 		if(World)
 		{
-			FTimerManager& TimerManager = World->GetTimerManager();
-			Handle = TimerManager.K2_FindDynamicTimerHandle(Delegate);
-			TimerManager.SetTimer(Handle, Delegate, Time, bLooping);
-		}
-	}
-	else
-	{
-		UE_LOG(LogBlueprintUserMessages, Warning, 
-			TEXT("SetTimer passed a bad function (%s) or object (%s)"),
-			*Delegate.GetFunctionName().ToString(), *GetNameSafe(Delegate.GetUObject()));
-	}
-
-	return Handle;
-}
+			if (Time <= 0.f || InitialStartDelay < 0.f)
+			{
+				FFrame::KismetExecutionMessage(TEXT("SetTimer passed a negative time or initial start delay.  The associated timer may fail to fire!"), ELogVerbosity::Warning);
+			}
 
-FTimerHandle UKismetSystemLibrary::K2_SetTimerDelegateWithVariance(FTimerDynamicDelegate Delegate, float Time, float Variance, bool bLooping)
-{
-	FTimerHandle Handle;
-	if (Delegate.IsBound())
-	{
-		const UWorld* const World = GEngine->GetWorldFromContextObject(Delegate.GetUObject(), EGetWorldErrorMode::LogAndReturnNull);
-		if (World)
-		{
 			FTimerManager& TimerManager = World->GetTimerManager();
 			Handle = TimerManager.K2_FindDynamicTimerHandle(Delegate);
-			Time = Time + FMath::RandRange(-Variance, Variance);
-			TimerManager.SetTimer(Handle, Delegate, Time, bLooping);
+			TimerManager.SetTimer(Handle, Delegate, Time, bLooping, InitialStartDelay);
 		}
 	}
 	else
 	{
-		UE_LOG(LogBlueprintUserMessages, Warning,
+		UE_LOG(LogBlueprintUserMessages, Warning, 
 			TEXT("SetTimer passed a bad function (%s) or object (%s)"),
 			*Delegate.GetFunctionName().ToString(), *GetNameSafe(Delegate.GetUObject()));
 	}
