#UE4 Fix a stack overflow infinite recursion crash that can happen if two objects reference eachother with instanced object properties and InitializeAssetBundlesFromMetadata is called on one of them.

diff --git a/Engine/Source/Runtime/Engine/Classes/Engine/AssetManager.h b/Engine/Source/Runtime/Engine/Classes/Engine/AssetManager.h
--- a/Engine/Source/Runtime/Engine/Classes/Engine/AssetManager.h
+++ b/Engine/Source/Runtime/Engine/Classes/Engine/AssetManager.h
@@ -697,6 +697,10 @@ protected:
 	FDelegateHandle ChunkInstallDelegateHandle;
 
 private:
+
+	/** Recursive handler for InitializeAssetBundlesFromMetadata */
+	virtual void InitializeAssetBundlesFromMetadata_Recursive(const UStruct* Struct, const void* StructValue, FAssetBundleData& AssetBundle, FName DebugName, TSet<const void*>& AllVisitedStructValues) const;
+
 	/** Per-type asset information, cannot be accessed by children as it is defined in CPP file */
 	TMap<FName, TSharedRef<FPrimaryAssetTypeData>> AssetTypeMap;
 
diff --git a/Engine/Source/Runtime/Engine/Private/AssetManager.cpp b/Engine/Source/Runtime/Engine/Private/AssetManager.cpp
--- a/Engine/Source/Runtime/Engine/Private/AssetManager.cpp
+++ b/Engine/Source/Runtime/Engine/Private/AssetManager.cpp
@@ -3561,6 +3561,12 @@ void UAssetManager::RefreshAssetData(UObject* ChangedObject)
 }
 
 void UAssetManager::InitializeAssetBundlesFromMetadata(const UStruct* Struct, const void* StructValue, FAssetBundleData& AssetBundle, FName DebugName) const
+{
+	TSet<const void*> AllVisitedStructValues;
+	InitializeAssetBundlesFromMetadata_Recursive(Struct, StructValue, AssetBundle, DebugName, AllVisitedStructValues);
+}
+
+void UAssetManager::InitializeAssetBundlesFromMetadata_Recursive(const UStruct* Struct, const void* StructValue, FAssetBundleData& AssetBundle, FName DebugName, TSet<const void*>& AllVisitedStructValues) const
 {
 	static FName AssetBundlesName = TEXT("AssetBundles");
 	static FName IncludeAssetBundlesName = TEXT("IncludeAssetBundles");
@@ -3570,6 +3576,13 @@ void UAssetManager::InitializeAssetBundlesFromMetadata(const UStruct* Struct, co
 		return;
 	}
 
+	if (AllVisitedStructValues.Contains(StructValue))
+	{
+		return;
+	}
+
+	AllVisitedStructValues.Add(StructValue);
+
 	for (TPropertyValueIterator<const UProperty> It(Struct, StructValue); It; ++It)
 	{
 		const UProperty* Property = It.Key();
@@ -3613,7 +3626,8 @@ void UAssetManager::InitializeAssetBundlesFromMetadata(const UStruct* Struct, co
 				UObject* const* ObjectPtr = reinterpret_cast<UObject* const*>(PropertyValue);
 				if (ObjectPtr && *ObjectPtr)
 				{
-					UAssetManager::Get().InitializeAssetBundlesFromMetadata(*ObjectPtr, AssetBundle);
+					const UObject* Object = *ObjectPtr;
+					InitializeAssetBundlesFromMetadata_Recursive(Object->GetClass(), Object, AssetBundle, Object->GetFName(), AllVisitedStructValues);
 				}
 			}
 		}
