#UE4 CIS fix for non-editor builds

diff --git a/Engine/Source/Runtime/Engine/Classes/Engine/AssetManager.h b/Engine/Source/Runtime/Engine/Classes/Engine/AssetManager.h
--- a/Engine/Source/Runtime/Engine/Classes/Engine/AssetManager.h
+++ b/Engine/Source/Runtime/Engine/Classes/Engine/AssetManager.h
@@ -698,8 +698,10 @@ protected:
 
 private:
 
+#if WITH_EDITOR
 	/** Recursive handler for InitializeAssetBundlesFromMetadata */
 	virtual void InitializeAssetBundlesFromMetadata_Recursive(const UStruct* Struct, const void* StructValue, FAssetBundleData& AssetBundle, FName DebugName, TSet<const void*>& AllVisitedStructValues) const;
+#endif
 
 	/** Per-type asset information, cannot be accessed by children as it is defined in CPP file */
 	TMap<FName, TSharedRef<FPrimaryAssetTypeData>> AssetTypeMap;
