Audio Debug improvements
- Standardize call with au.3dVisualize (Enabled, ActiveSounds, VirtualLoops, etc.) but keep Audio3dVisualize for backward compat
- Add enablement calls for various types of visualization (source spatialization, virtual loops, active sounds)
- Add colorization modes for active sounds based on distance, volume, and at random
- Clean-up stat Audio stats to be correct in cases of source numbers, etc. and add stat for MaxChannels

https://p4-swarm.epicgames.net/reviews/6256389

diff --git a/Engine/Source/Runtime/Engine/Private/ActiveSound.cpp b/Engine/Source/Runtime/Engine/Private/ActiveSound.cpp
--- a/Engine/Source/Runtime/Engine/Private/ActiveSound.cpp
+++ b/Engine/Source/Runtime/Engine/Private/ActiveSound.cpp
@@ -1,6 +1,7 @@
 // Copyright 1998-2019 Epic Games, Inc. All Rights Reserved.
 #include "ActiveSound.h"
 
+#include "DrawDebugHelpers.h"
 #include "EngineDefines.h"
 #include "Misc/App.h"
 #include "AudioThread.h"
@@ -25,6 +26,21 @@ FAutoConsoleVariableRef CVarAudioDisableConcurrencyStopSilentForLoops(
 	TEXT("Disables (1) or enables (0) audio concurrency for loops subscribing to stop silent.\n"),
 	ECVF_Default);
 
+static int32 ActiveSoundVisualizeModeCVar = 1;
+FAutoConsoleVariableRef CVarAudioVisualizeActiveSoundsMode(
+	TEXT("au.3dVisualize.ActiveSounds"),
+	ActiveSoundVisualizeModeCVar,
+	TEXT("Visualization mode for active sounds. \n")
+	TEXT("0: Not Enabled, 1: Volume (Lin), 2: Volume (dB), 3: Distance, 4: Random color"),
+	ECVF_Default);
+
+static int32 ActiveSoundVisualizeTypeCVar = 0;
+FAutoConsoleVariableRef CVarAudioVisualizeActiveSounds(
+	TEXT("au.3dVisualize.ActiveSounds.Type"),
+	ActiveSoundVisualizeTypeCVar,
+	TEXT("Whether to show all sounds, on AudioComponents (Components Only), or off of AudioComponents (Non-Component Only). \n")
+	TEXT("0: All, 1: Components Only, 2: Non-Component Only"),
+	ECVF_Default);
 
 FTraceDelegate FActiveSound::ActiveSoundTraceDelegate;
 TMap<FTraceHandle, FActiveSound::FAsyncTraceDetails> FActiveSound::TraceToActiveSoundMap;
@@ -105,6 +121,9 @@ FActiveSound::FActiveSound()
 	, CurrentInteriorLPF(MAX_FILTER_FREQUENCY)
 	, EnvelopeFollowerAttackTime(10)
 	, EnvelopeFollowerReleaseTime(100)
+#if !(UE_BUILD_SHIPPING || UE_BUILD_TEST)
+	, DebugColor(FColor::Black)
+#endif // !(UE_BUILD_SHIPPING || UE_BUILD_TEST)
 	, ClosestListenerPtr(nullptr)
 	, InternalFocusFactor(1.0f)
 {
@@ -626,9 +645,116 @@ void FActiveSound::UpdateWaveInstances(TArray<FWaveInstance*> &InWaveInstances,
 
 	}
 
+	if (ActiveSoundVisualizeModeCVar != 0)
+	{
+		DrawDebugInfo(&ThisSoundsWaveInstances);
+	}
+
 	InWaveInstances.Append(ThisSoundsWaveInstances);
 }
 
+void FActiveSound::DrawDebugInfo(const TArray<FWaveInstance*>* ThisSoundsWaveInstances)
+{
+#if ENABLE_DRAW_DEBUG
+	// Only draw spatialized sounds
+	if (!Sound || !bAllowSpatialization)
+	{
+		return;
+	}
+
+	if (ActiveSoundVisualizeTypeCVar > 0)
+	{
+		if (ActiveSoundVisualizeTypeCVar == 1 && AudioComponentID == 0)
+		{
+			return;
+		}
+
+		if (ActiveSoundVisualizeTypeCVar == 2 && AudioComponentID > 0)
+		{
+			return;
+		}
+	}
+
+	FAudioDeviceManager* DeviceManager = GEngine->GetAudioDeviceManager();
+	if (DeviceManager && DeviceManager->IsVisualizeDebug3dEnabled())
+	{
+		DECLARE_CYCLE_STAT(TEXT("FAudioThreadTask.DrawActiveSoundDebugInfo"), STAT_AudioDrawActiveSoundDebugInfo, STATGROUP_TaskGraphTasks);
+
+		const FString Name = Sound->GetName();
+		const FTransform CurTransform = Transform;
+		FColor TextColor = FColor::White;
+		const float CurMaxDistance = MaxDistance;
+		float DisplayValue = 0.0f;
+		if (ActiveSoundVisualizeModeCVar == 1 || ActiveSoundVisualizeModeCVar == 2)
+		{
+			DisplayValue = ThisSoundsWaveInstances && ThisSoundsWaveInstances->Num() == 1
+				? (*ThisSoundsWaveInstances)[0]->GetVolumeWithDistanceAttenuation()
+				: GetVolume();
+		}
+		else if (ActiveSoundVisualizeModeCVar == 3)
+		{
+			if (AudioDevice)
+			{
+				DisplayValue = AudioDevice->GetDistanceToNearestListener(Transform.GetLocation()) / CurMaxDistance;
+			}
+		}
+#if !(UE_BUILD_SHIPPING || UE_BUILD_TEST)
+		else if (ActiveSoundVisualizeModeCVar == 4)
+		{
+			if (DebugColor == FColor::Black)
+			{
+				DebugColor = FColor::MakeRandomColor();
+			}
+
+			TextColor = DebugColor;
+		}
+#endif // !(UE_BUILD_SHIPPING || UE_BUILD_TEST)
+
+		TWeakObjectPtr<UWorld> WorldPtr = GetWeakWorld();
+		FAudioThread::RunCommandOnGameThread([Name, TextColor, CurTransform, DisplayValue, WorldPtr, CurMaxDistance]()
+		{
+			if (WorldPtr.IsValid())
+			{
+				static const float ColorRedHue = 0.0f;
+				static const float ColorGreenHue = 85.0f;
+
+				const FVector Location = CurTransform.GetLocation();
+				UWorld* World = WorldPtr.Get();
+				DrawDebugSphere(World, Location, 10.0f, 8, FColor::White, false, -1.0f, SDPG_Foreground);
+				FColor Color = TextColor;
+
+				FString Descriptor;
+				if (ActiveSoundVisualizeModeCVar == 1 || ActiveSoundVisualizeModeCVar == 2)
+				{
+					const float DisplayDbVolume = Audio::ConvertToDecibels(DisplayValue);
+					if (ActiveSoundVisualizeModeCVar == 1)
+					{
+						Descriptor = FString::Printf(TEXT(" (Vol: %.3f)"), DisplayValue);
+					}
+					else
+					{
+						Descriptor = FString::Printf(TEXT(" (Vol: %.3f dB)"), DisplayDbVolume);
+					}
+					static const float DbColorMinVol = -30.0f;
+					const float DbVolume = FMath::Clamp(DisplayDbVolume, DbColorMinVol, 0.0f);
+					const float Hue = FMath::Lerp(ColorRedHue, ColorGreenHue, (-1.0f * DbVolume / DbColorMinVol) + 1.0f);
+					Color = FLinearColor::MakeFromHSV8(static_cast<uint8>(Hue), 255u, 255u).ToFColor(true);
+				}
+				else if (ActiveSoundVisualizeModeCVar == 3)
+				{
+					Descriptor = FString::Printf(TEXT(" (Dist: %.3f, Max: %.3f)"), DisplayValue * CurMaxDistance, CurMaxDistance);
+					const float Hue = FMath::Lerp(ColorGreenHue, ColorRedHue, DisplayValue);
+					Color = FLinearColor::MakeFromHSV8(static_cast<uint8>(FMath::Clamp(Hue, 0.0f, 255.f)), 255u, 255u).ToFColor(true);
+				}
+
+				const FString Description = FString::Printf(TEXT("%s%s"), *Name, *Descriptor);
+				DrawDebugString(World, Location + FVector(0, 0, 32), *Description, nullptr, Color, 0.03f, false);
+			}
+		}, GET_STATID(STAT_AudioDrawActiveSoundDebugInfo));
+	}
+#endif // ENABLE_DRAW_DEBUG
+}
+
 void FActiveSound::Stop(bool bStopNow)
 {
 	check(AudioDevice);
diff --git a/Engine/Source/Runtime/Engine/Private/Audio.cpp b/Engine/Source/Runtime/Engine/Private/Audio.cpp
--- a/Engine/Source/Runtime/Engine/Private/Audio.cpp
+++ b/Engine/Source/Runtime/Engine/Private/Audio.cpp
@@ -35,6 +35,8 @@ DEFINE_STAT(STAT_AudioSources);
 DEFINE_STAT(STAT_AudioVirtualLoops);
 DEFINE_STAT(STAT_WaveInstances);
 DEFINE_STAT(STAT_WavesDroppedDueToPriority);
+DEFINE_STAT(STAT_AudioMaxChannels);
+DEFINE_STAT(STAT_AudioMaxStoppingSources);
 DEFINE_STAT(STAT_AudibleWavesDroppedDueToPriority);
 DEFINE_STAT(STAT_AudioFinishedDelegatesCalled);
 DEFINE_STAT(STAT_AudioFinishedDelegates);
@@ -75,6 +77,14 @@ FAutoConsoleVariableRef CVarAllowAudioSpatializationCVar(
 	TEXT("0: Disable, >0: Enable"),
 	ECVF_Default);
 
+static int32 SpatialSourceVisualizeEnabledCVar = 1;
+FAutoConsoleVariableRef CVarAudioVisualizeSpatialSourceEnabled(
+	TEXT("au.3dVisualize.SpatialSources"),
+	SpatialSourceVisualizeEnabledCVar,
+	TEXT("Whether or not audio spatialized sources are visible when 3d visualize is enabled. \n")
+	TEXT("0: Not Enabled, 1: Enabled"),
+	ECVF_Default);
+
 bool IsAudioPluginEnabled(EAudioPlugin PluginType)
 {
 	switch (PluginType)
@@ -412,45 +422,42 @@ void FSoundSource::UpdateStereoEmitterPositions()
 void FSoundSource::DrawDebugInfo()
 {
 #if ENABLE_DRAW_DEBUG
-	// Draw 3d Debug information about this source, if enabled
-	FAudioDeviceManager* DeviceManager = GEngine->GetAudioDeviceManager();
+	if (!WaveInstance)
+	{
+		return;
+	}
+
+	const FActiveSound* ActiveSound = WaveInstance->ActiveSound;
+	if (!ActiveSound)
+	{
+		return;
+	}
+
+	if (!SpatialSourceVisualizeEnabledCVar)
+	{
+		return;
+	}
 
+	FAudioDeviceManager* DeviceManager = GEngine->GetAudioDeviceManager();
 	if (DeviceManager && DeviceManager->IsVisualizeDebug3dEnabled())
 	{
-		const uint32 AudioComponentID = WaveInstance->ActiveSound->GetAudioComponentID();
+		DECLARE_CYCLE_STAT(TEXT("FAudioThreadTask.DrawSourceDebugInfo"), STAT_AudioDrawSourceDebugInfo, STATGROUP_TaskGraphTasks);
 
-		if (AudioComponentID > 0)
+		const bool bSpatialized = Buffer->NumChannels == 2 && WaveInstance->GetUseSpatialization();
+		if (bSpatialized)
 		{
-			DECLARE_CYCLE_STAT(TEXT("FAudioThreadTask.DrawSourceDebugInfo"), STAT_AudioDrawSourceDebugInfo, STATGROUP_TaskGraphTasks);
+			const FRotator Rotator = ActiveSound->Transform.GetRotation().Rotator();
 
-			USoundBase* Sound = WaveInstance->ActiveSound->GetSound();
-			const FVector Location = WaveInstance->Location;
-
-			const bool bSpatialized = Buffer->NumChannels == 2 && WaveInstance->GetUseSpatialization();
+			TWeakObjectPtr<UWorld> WorldPtr = WaveInstance->ActiveSound->GetWeakWorld();
 			const FVector LeftChannelSourceLoc = LeftChannelSourceLocation;
 			const FVector RightChannelSourceLoc = RightChannelSourceLocation;
-			const float Volume = WaveInstance->GetActualVolume();
-			FAudioThread::RunCommandOnGameThread([AudioComponentID, Sound, bSpatialized, Location, LeftChannelSourceLoc, RightChannelSourceLoc, Volume]()
+			FAudioThread::RunCommandOnGameThread([LeftChannelSourceLoc, RightChannelSourceLoc, Rotator, WorldPtr]()
 			{
-				UAudioComponent* AudioComponent = UAudioComponent::GetAudioComponentFromID(AudioComponentID);
-				if (AudioComponent)
+				if (WorldPtr.IsValid())
 				{
-					UWorld* SoundWorld = AudioComponent->GetWorld();
-					if (SoundWorld)
-					{
-						FRotator SoundRotation = AudioComponent->GetComponentRotation();
-						DrawDebugCrosshairs(SoundWorld, Location, SoundRotation, 20.0f, FColor::White, false, -1.0f, SDPG_Foreground);
-
-						if (bSpatialized)
-						{
-							DrawDebugCrosshairs(SoundWorld, LeftChannelSourceLoc, SoundRotation, 20.0f, FColor::Red, false, -1.0f, SDPG_Foreground);
-							DrawDebugCrosshairs(SoundWorld, RightChannelSourceLoc, SoundRotation, 20.0f, FColor::Green, false, -1.0f, SDPG_Foreground);
-						}
-
-						const FString Name = Sound->GetName();
-						const FString Description = FString::Printf(TEXT("%s (%.3f)"), *Name, Volume);
-						DrawDebugString(SoundWorld, AudioComponent->GetComponentLocation() + FVector(0, 0, 32), *Description, nullptr, FColor::White, 0.033, false);
-					}
+					UWorld* World = WorldPtr.Get();
+					DrawDebugCrosshairs(World, LeftChannelSourceLoc, Rotator, 20.0f, FColor::Red, false, -1.0f, SDPG_Foreground);
+					DrawDebugCrosshairs(World, RightChannelSourceLoc, Rotator, 20.0f, FColor::Green, false, -1.0f, SDPG_Foreground);
 				}
 			}, GET_STATID(STAT_AudioDrawSourceDebugInfo));
 		}
diff --git a/Engine/Source/Runtime/Engine/Private/AudioDevice.cpp b/Engine/Source/Runtime/Engine/Private/AudioDevice.cpp
--- a/Engine/Source/Runtime/Engine/Private/AudioDevice.cpp
+++ b/Engine/Source/Runtime/Engine/Private/AudioDevice.cpp
@@ -3884,11 +3884,13 @@ void FAudioDevice::Update(bool bGameTicking)
 
 
 		const int32 Channels = GetMaxChannels();
-		INC_DWORD_STAT_BY(STAT_WaveInstances, ActiveWaveInstances.Num());
-		INC_DWORD_STAT_BY(STAT_AudioSources, Channels - FreeSources.Num());
-		INC_DWORD_STAT_BY(STAT_WavesDroppedDueToPriority, FMath::Max(ActiveWaveInstances.Num() - Channels, 0));
-		INC_DWORD_STAT_BY(STAT_ActiveSounds, ActiveSounds.Num());
-		INC_DWORD_STAT_BY(STAT_AudioVirtualLoops, VirtualLoops.Num());
+		SET_DWORD_STAT(STAT_WaveInstances, ActiveWaveInstances.Num());
+		SET_DWORD_STAT(STAT_AudioSources, Sources.Num() - FreeSources.Num());
+		SET_DWORD_STAT(STAT_WavesDroppedDueToPriority, FMath::Max(ActiveWaveInstances.Num() - Sources.Num(), 0));
+		SET_DWORD_STAT(STAT_ActiveSounds, ActiveSounds.Num());
+		SET_DWORD_STAT(STAT_AudioVirtualLoops, VirtualLoops.Num());
+		SET_DWORD_STAT(STAT_AudioMaxChannels, Channels);
+		SET_DWORD_STAT(STAT_AudioMaxStoppingSources, NumStoppingVoices);
 	}
 
 	// now let the platform perform anything it needs to handle
diff --git a/Engine/Source/Runtime/Engine/Private/AudioDeviceManager.cpp b/Engine/Source/Runtime/Engine/Private/AudioDeviceManager.cpp
--- a/Engine/Source/Runtime/Engine/Private/AudioDeviceManager.cpp
+++ b/Engine/Source/Runtime/Engine/Private/AudioDeviceManager.cpp
@@ -44,6 +44,15 @@ FAutoConsoleVariableRef CVarIsUsingAudioMixer(
 	ECVF_Default);
 
 
+static int32 CVarIsVisualizeEnabled = 0;
+FAutoConsoleVariableRef CVarAudioVisualizeEnabled(
+	TEXT("au.3dVisualize.Enabled"),
+	CVarIsVisualizeEnabled,
+	TEXT("Whether or not audio visualization is enabled. \n")
+	TEXT("0: Not Enabled, 1: Enabled"),
+	ECVF_Default);
+
+
 FAudioDeviceManager::FCreateAudioDeviceResults::FCreateAudioDeviceResults()
 	: Handle(INDEX_NONE)
 	, bNewDevice(false)
@@ -915,6 +924,11 @@ void FAudioDeviceManager::TogglePlayAllDeviceAudio()
 	bPlayAllDeviceAudio = !bPlayAllDeviceAudio;
 }
 
+bool FAudioDeviceManager::IsVisualizeDebug3dEnabled() const
+{
+	return bVisualize3dDebug || CVarIsVisualizeEnabled;
+}
+
 void FAudioDeviceManager::ToggleVisualize3dDebug()
 {
 	if (!IsInAudioThread())
diff --git a/Engine/Source/Runtime/Engine/Private/AudioVirtualLoop.cpp b/Engine/Source/Runtime/Engine/Private/AudioVirtualLoop.cpp
--- a/Engine/Source/Runtime/Engine/Private/AudioVirtualLoop.cpp
+++ b/Engine/Source/Runtime/Engine/Private/AudioVirtualLoop.cpp
@@ -35,6 +35,14 @@ FAutoConsoleVariableRef CVarVirtualLoopsUpdateRateMax(
 	TEXT("Sets maximum rate to check if sound becomes audible again (at beyond sound's max audible distance + perf scaling distance).\n"),
 	ECVF_Default);
 
+static int32 VirtualLoopsVisualizeEnabledCVar = 1;
+FAutoConsoleVariableRef CVarAudioVisualizeVirtualLoopsEnabled(
+	TEXT("au.3dVisualize.VirtualLoops"),
+	VirtualLoopsVisualizeEnabledCVar,
+	TEXT("Whether or not virtualized loops are visible when 3d visualize is enabled. \n")
+	TEXT("0: Not Enabled, 1: Enabled"),
+	ECVF_Default);
+
 FAudioVirtualLoop::FAudioVirtualLoop()
 	: TimeSinceLastUpdate(0.0f)
 	, UpdateInterval(0.0f)
@@ -170,7 +178,10 @@ bool FAudioVirtualLoop::CanRealize(float DeltaTime)
 		TimeSinceLastUpdate = 0.0f;
 	}
 
-	DrawDebugInfo();
+	if (VirtualLoopsVisualizeEnabledCVar)
+	{
+		DrawDebugInfo();
+	}
 
 	// If not audible, update when will be checked again and return false
 	if (!IsInAudibleRange(*ActiveSound))
@@ -199,12 +210,11 @@ void FAudioVirtualLoop::DrawDebugInfo() const
 		const TWeakObjectPtr<UWorld> World = ActiveSound->GetWeakWorld();
 		const FString Name = Sound->GetName();
 		const float DrawInterval = UpdateInterval;
-		const float Volume = ActiveSound->GetVolume();
-		FAudioThread::RunCommandOnGameThread([World, Transform, Name, DrawInterval, Volume]()
+		FAudioThread::RunCommandOnGameThread([World, Transform, Name, DrawInterval]()
 		{
 			if (World.IsValid())
 			{
-				const FString Description = FString::Printf(TEXT("%s (%.3f)"), *Name, Volume);
+				const FString Description = FString::Printf(TEXT("%s [V]"), *Name);
 				FVector Location = Transform.GetLocation();
 				FRotator Rotation = Transform.GetRotation().Rotator();
 				DrawDebugCrosshairs(World.Get(), Location, Rotation, 20.0f, FColor::Blue, false, DrawInterval, SDPG_Foreground);
diff --git a/Engine/Source/Runtime/Engine/Public/ActiveSound.h b/Engine/Source/Runtime/Engine/Public/ActiveSound.h
--- a/Engine/Source/Runtime/Engine/Public/ActiveSound.h
+++ b/Engine/Source/Runtime/Engine/Public/ActiveSound.h
@@ -20,6 +20,7 @@ class USoundSubmix;
 class USoundSourceBus;
 struct FSoundSubmixSendInfo;
 struct FSoundSourceBusSendInfo;
+struct FWaveInstance;
 class USoundWave;
 struct FListener;
 struct FAttenuationListenerData;
@@ -224,6 +225,7 @@ private:
 	FName AudioComponentName;
 	FName OwnerName;
 
+
 public:
 
 	uint64 GetAudioComponentID() const { return AudioComponentID; }
@@ -470,7 +472,7 @@ public:
 	int32 EnvelopeFollowerAttackTime;
 	int32 EnvelopeFollowerReleaseTime;
 
-	TMap<UPTRINT, struct FWaveInstance*> WaveInstances;
+	TMap<UPTRINT, FWaveInstance*> WaveInstances;
 
 	TMap<UPTRINT,uint32> SoundNodeOffsetMap;
 	TArray<uint8> SoundNodeData;
@@ -479,6 +481,7 @@ public:
 
 #if !(UE_BUILD_SHIPPING || UE_BUILD_TEST)
 	FName DebugOriginalSoundName;
+	FColor DebugColor;
 #endif
 
 	// Updates the wave instances to be played.
@@ -494,6 +497,11 @@ public:
 	 */
 	void ApplyRadioFilter(const struct FSoundParseParameters& ParseParams);
 
+	/**
+	 * Draws debug info for builds supporting debug drawing
+	 */
+	void DrawDebugInfo(const TArray<FWaveInstance*>* WaveInstances);
+
 	/** Gets total concurrency gain stage based on all concurrency memberships of sound */
 	float GetTotalConcurrencyVolumeScale() const;
 
diff --git a/Engine/Source/Runtime/Engine/Public/Audio.h b/Engine/Source/Runtime/Engine/Public/Audio.h
--- a/Engine/Source/Runtime/Engine/Public/Audio.h
+++ b/Engine/Source/Runtime/Engine/Public/Audio.h
@@ -84,6 +84,8 @@ DECLARE_DWORD_COUNTER_STAT_EXTERN( TEXT( "Wave Instances" ), STAT_WaveInstances,
 DECLARE_DWORD_COUNTER_STAT_EXTERN( TEXT( "Wave Instances Dropped" ), STAT_WavesDroppedDueToPriority, STATGROUP_Audio , );
 DECLARE_DWORD_COUNTER_STAT_EXTERN(TEXT("Virtualized Loops"), STAT_AudioVirtualLoops, STATGROUP_Audio, );
 DECLARE_DWORD_COUNTER_STAT_EXTERN( TEXT( "Audible Wave Instances Dropped" ), STAT_AudibleWavesDroppedDueToPriority, STATGROUP_Audio , );
+DECLARE_DWORD_COUNTER_STAT_EXTERN(TEXT("Max Channels"), STAT_AudioMaxChannels, STATGROUP_Audio, );
+DECLARE_DWORD_COUNTER_STAT_EXTERN(TEXT("Max Stopping Sources"), STAT_AudioMaxStoppingSources, STATGROUP_Audio, );
 DECLARE_DWORD_COUNTER_STAT_EXTERN( TEXT( "Finished delegates called" ), STAT_AudioFinishedDelegatesCalled, STATGROUP_Audio , );
 DECLARE_CYCLE_STAT_EXTERN( TEXT( "Finished delegates time" ), STAT_AudioFinishedDelegates, STATGROUP_Audio , );
 DECLARE_MEMORY_STAT_EXTERN( TEXT( "Audio Memory Used" ), STAT_AudioMemorySize, STATGROUP_Audio , );
diff --git a/Engine/Source/Runtime/Engine/Public/AudioDeviceManager.h b/Engine/Source/Runtime/Engine/Public/AudioDeviceManager.h
--- a/Engine/Source/Runtime/Engine/Public/AudioDeviceManager.h
+++ b/Engine/Source/Runtime/Engine/Public/AudioDeviceManager.h
@@ -172,7 +172,7 @@ public:
 	bool IsPlayAllDeviceAudio() const { return bPlayAllDeviceAudio; }
 
 	/** Is debug visualization of 3d sounds enabled */
-	bool IsVisualizeDebug3dEnabled() const { return bVisualize3dDebug; }
+	bool IsVisualizeDebug3dEnabled() const;
 
 	/** Toggles 3d visualization of 3d sounds on/off */
 	void ToggleVisualize3dDebug();
