Fix variable shadowing from prior CL

diff --git a/Engine/Source/Runtime/Engine/Private/ActiveSound.cpp b/Engine/Source/Runtime/Engine/Private/ActiveSound.cpp
--- a/Engine/Source/Runtime/Engine/Private/ActiveSound.cpp
+++ b/Engine/Source/Runtime/Engine/Private/ActiveSound.cpp
@@ -719,8 +719,8 @@ void FActiveSound::DrawDebugInfo(const TArray<FWaveInstance*>* ThisSoundsWaveIns
 				static const float ColorGreenHue = 85.0f;
 
 				const FVector Location = CurTransform.GetLocation();
-				UWorld* World = WorldPtr.Get();
-				DrawDebugSphere(World, Location, 10.0f, 8, FColor::White, false, -1.0f, SDPG_Foreground);
+				UWorld* DebugWorld = WorldPtr.Get();
+				DrawDebugSphere(DebugWorld, Location, 10.0f, 8, FColor::White, false, -1.0f, SDPG_Foreground);
 				FColor Color = TextColor;
 
 				FString Descriptor;
@@ -748,7 +748,7 @@ void FActiveSound::DrawDebugInfo(const TArray<FWaveInstance*>* ThisSoundsWaveIns
 				}
 
 				const FString Description = FString::Printf(TEXT("%s%s"), *Name, *Descriptor);
-				DrawDebugString(World, Location + FVector(0, 0, 32), *Description, nullptr, Color, 0.03f, false);
+				DrawDebugString(DebugWorld, Location + FVector(0, 0, 32), *Description, nullptr, Color, 0.03f, false);
 			}
 		}, GET_STATID(STAT_AudioDrawActiveSoundDebugInfo));
 	}
