Don't finish destroy a UStaticMesh if it has a pending update

diff --git a/Engine/Source/Runtime/Engine/Private/StaticMesh.cpp b/Engine/Source/Runtime/Engine/Private/StaticMesh.cpp
--- a/Engine/Source/Runtime/Engine/Private/StaticMesh.cpp
+++ b/Engine/Source/Runtime/Engine/Private/StaticMesh.cpp
@@ -3122,7 +3122,7 @@ void UStaticMesh::BeginDestroy()
 
 bool UStaticMesh::IsReadyForFinishDestroy()
 {
-	return ReleaseResourcesFence.IsFenceComplete();
+	return ReleaseResourcesFence.IsFenceComplete() && !UpdateStreamingStatus();
 }
 
 int32 UStaticMesh::GetNumSectionsWithCollision() const
