Use WasRecentlyRendered which has cached values instead of rolling own checks against GetLastRenderTime
#rnx

diff --git a/Engine/Source/Runtime/Engine/Private/Interpolation.cpp b/Engine/Source/Runtime/Engine/Private/Interpolation.cpp
--- a/Engine/Source/Runtime/Engine/Private/Interpolation.cpp
+++ b/Engine/Source/Runtime/Engine/Private/Interpolation.cpp
@@ -1155,7 +1155,7 @@ void AMatineeActor::StepInterp( float DeltaTime, bool bPreview )
 			for (int32 ActorIndex = 0; ActorIndex < Info.Actors.Num() && bSkipUpdate; ++ActorIndex)
 			{
 				AActor* Actor = Info.Actors[ ActorIndex ];
-				if (Actor != NULL && Actor->GetLastRenderTime() > Actor->GetWorld()->TimeSeconds - 1.f)
+				if (Actor && Actor->WasRecentlyRendered(1.f))
 				{
 					bSkipUpdate = false;
 				}
diff --git a/Engine/Source/Runtime/Engine/Private/UnrealEngine.cpp b/Engine/Source/Runtime/Engine/Private/UnrealEngine.cpp
--- a/Engine/Source/Runtime/Engine/Private/UnrealEngine.cpp
+++ b/Engine/Source/Runtime/Engine/Private/UnrealEngine.cpp
@@ -16142,7 +16142,7 @@ int32 UEngine::RenderStatAI(UWorld* World, FViewport* Viewport, FCanvas* Canvas,
 		if (Controller && !Cast<APlayerController>(Controller))
 		{
 			++NumAI;
-			if (Controller->GetPawn() != NULL && World->GetTimeSeconds() - Controller->GetPawn()->GetLastRenderTime() < 0.08f)
+			if (Controller->GetPawn() && Controller->GetPawn()->WasRecentlyRendered(0.08f))
 			{
 				++NumAIRendered;
 			}
diff --git a/Engine/Source/Runtime/Engine/Private/World.cpp b/Engine/Source/Runtime/Engine/Private/World.cpp
--- a/Engine/Source/Runtime/Engine/Private/World.cpp
+++ b/Engine/Source/Runtime/Engine/Private/World.cpp
@@ -7268,7 +7268,7 @@ static void DumpVisibleActors(UWorld* InWorld)
 	for (FActorIterator ActorIterator(InWorld); ActorIterator; ++ActorIterator)
 	{
 		AActor* Actor = *ActorIterator;
-		if (Actor && Actor->GetLastRenderTime() > (InWorld->GetTimeSeconds() - 0.05f))
+		if (Actor && Actor->WasRecentlyRendered(0.05f))
 		{
 			UE_LOG(LogWorld, Log, TEXT("Visible Actor : %s"), *Actor->GetFullName());
 		}
