Add UNetConnection Saturation Analytics

diff --git a/Engine/Plugins/Runtime/ReplicationGraph/Source/Private/ReplicationGraph.cpp b/Engine/Plugins/Runtime/ReplicationGraph/Source/Private/ReplicationGraph.cpp
--- a/Engine/Plugins/Runtime/ReplicationGraph/Source/Private/ReplicationGraph.cpp
+++ b/Engine/Plugins/Runtime/ReplicationGraph/Source/Private/ReplicationGraph.cpp
@@ -62,6 +62,10 @@
 #include "Misc/ScopeExit.h"
 #include "Net/NetworkGranularMemoryLogging.h"
 
+#if USE_SERVER_PERF_COUNTERS
+#include "PerfCountersModule.h"
+#endif
+
 int32 CVar_RepGraph_Pause = 0;
 static FAutoConsoleVariableRef CVarRepGraphPause(TEXT("Net.RepGraph.Pause"), CVar_RepGraph_Pause, TEXT("Pauses actor replication in the Replication Graph."), ECVF_Default );
 
@@ -800,6 +804,8 @@ int32 UReplicationGraph::ServerReplicateActors(float DeltaSeconds)
 	++NetDriver->ReplicationFrame;	// This counter is used by RepLayout to utilize CL/serialization sharing. We must increment it ourselves, but other places can increment it too, in order to invalidate the shared state.
 	const uint32 FrameNum = ReplicationGraphFrame; // This counter is used internally and drives all frame based replication logic.
 
+	bWasConnectionSaturated = false;
+
 	ON_SCOPE_EXIT
 	{
 		// We increment this after our replication has happened. If we increment at the beginning of this function, then we rep with FrameNum X, then start the next game frame with the same FrameNum X. If at the top of that frame,
@@ -867,6 +873,12 @@ int32 UReplicationGraph::ServerReplicateActors(float DeltaSeconds)
 
 		NumChildrenConnectionsProcessed += NetConnection->Children.Num();
 
+		ON_SCOPE_EXIT
+		{
+			NetConnection->TrackReplicationForAnalytics(bWasConnectionSaturated);
+			bWasConnectionSaturated = false;
+		};
+
 		FBitWriter& ConnectionSendBuffer = NetConnection->SendBuffer; // unused
 		ConnectionManager->QueuedBitsForActorDiscovery = 0;
 
@@ -1298,7 +1310,7 @@ void UReplicationGraph::ReplicateActorListsForConnections_Default(UNetReplicatio
 				// We've exceeded the budget for this category of replication list.
 				RG_QUICK_SCOPE_CYCLE_COUNTER(NET_ReplicateActors_PartialStarvedActorList);
 				HandleStarvedActorList(PrioritizedReplicationList, ActorIdx + 1, ConnectionActorInfoMap, FrameNum);
-				GNumSaturatedConnections++;
+				NotifyConnectionSaturated(*ConnectionManager);
 				break;
 			}
 		}
@@ -1495,7 +1507,7 @@ void UReplicationGraph::ReplicateActorListsForConnections_FastShared(UNetReplica
 #endif
 			if (TotalBitsWritten > MaxBits)
 			{
-				GNumSaturatedConnections++;
+				NotifyConnectionSaturated(*ConnectionManager);
 				return;
 			}
 		}
@@ -2154,6 +2166,12 @@ void UReplicationGraph::SetActorDiscoveryBudget(int32 ActorDiscoveryBudgetInKByt
 	UE_LOG(LogReplicationGraph, Display, TEXT("SetActorDiscoveryBudget set to %d kBps (%d bits per network tick)."), ActorDiscoveryBudgetInKBytesPerSec, ActorDiscoveryMaxBitsPerFrame);
 }
 
+void UReplicationGraph::NotifyConnectionSaturated(UNetReplicationGraphConnection& Connection)
+{
+	bWasConnectionSaturated = true;
+	++GNumSaturatedConnections;
+}
+
 // --------------------------------------------------------------------------------------------------------------------------------------------
 // --------------------------------------------------------------------------------------------------------------------------------------------
 // --------------------------------------------------------------------------------------------------------------------------------------------
@@ -3222,7 +3240,7 @@ void UReplicationGraphNode_DynamicSpatialFrequency::GatherActorListsForConnectio
 			// Bandwidth Cap
 			if (BitsWritten > MaxBits && CVar_RepGraph_DynamicSpatialFrequency_UncapBandwidth == 0)
 			{
-				GNumSaturatedConnections++;
+				RepGraph->NotifyConnectionSaturated(Params.ConnectionManager);
 				break;
 			}
 		}
diff --git a/Engine/Source/Runtime/Engine/Classes/Engine/NetConnection.h b/Engine/Source/Runtime/Engine/Classes/Engine/NetConnection.h
--- a/Engine/Source/Runtime/Engine/Classes/Engine/NetConnection.h
+++ b/Engine/Source/Runtime/Engine/Classes/Engine/NetConnection.h
@@ -24,7 +24,7 @@
 #include "Analytics/EngineNetAnalytics.h"
 #include "PacketTraits.h"
 #include "Net/Util/ResizableCircularQueue.h"
-#include "Net/DelinquencyAnalytics.h"
+#include "Net/NetAnalyticsTypes.h"
 
 #include "NetConnection.generated.h"
 
@@ -1087,6 +1087,27 @@ public:
 	/** Resets the current delinquency analytics. */
 	ENGINE_API void ResetQueuedActorDelinquencyAnalytics();
 
+	/**
+	 * Returns the current saturation analytics and resets them.
+	 * This would be similar to calls to Get and Reset separately, except that the caller
+	 * will assume ownership of data in this case.
+	 */
+	ENGINE_API void ConsumeSaturationAnalytics(FNetConnectionSaturationAnalytics& Out);
+
+	/** Returns the current saturation analytics. */
+	ENGINE_API const FNetConnectionSaturationAnalytics& GetSaturationAnalytics() const;
+
+	/** Resets the current saturation analytics. */
+	ENGINE_API void ResetSaturationAnalytics();
+
+	/**
+	 * Called to notify the connection that we attempted to replicate its actors this frame.
+	 * This is primarily for analytics book keeping.
+	 *
+	 * @param bWasSaturated		True if we failed to replicate all data because we were saturated.
+	 */
+	ENGINE_API void TrackReplicationForAnalytics(const bool bWasSaturated);
+	
 protected:
 
 	ENGINE_API void SetPendingCloseDueToSocketSendFailure();
@@ -1194,6 +1215,8 @@ private:
 	/** The current number of valid packets in PacketOrderCache */
 	int32 PacketOrderCacheCount;
 
+	FNetConnectionSaturationAnalytics SaturationAnalytics;
+
 	/** Whether or not PacketOrderCache is presently being flushed */
 	bool bFlushingPacketOrderCache;
 };
diff --git a/Engine/Source/Runtime/Engine/Classes/Engine/NetDriver.h b/Engine/Source/Runtime/Engine/Classes/Engine/NetDriver.h
--- a/Engine/Source/Runtime/Engine/Classes/Engine/NetDriver.h
+++ b/Engine/Source/Runtime/Engine/Classes/Engine/NetDriver.h
@@ -15,7 +15,7 @@
 #include "Channel.h"
 #include "DDoSDetection.h"
 #include "IPAddress.h"
-#include "Net/DelinquencyAnalytics.h"
+#include "Net/NetAnalyticsTypes.h"
 
 #include "NetDriver.generated.h"
 
diff --git a/Engine/Source/Runtime/Engine/Classes/Engine/PackageMapClient.h b/Engine/Source/Runtime/Engine/Classes/Engine/PackageMapClient.h
--- a/Engine/Source/Runtime/Engine/Classes/Engine/PackageMapClient.h
+++ b/Engine/Source/Runtime/Engine/Classes/Engine/PackageMapClient.h
@@ -27,7 +27,7 @@
 #include "Misc/NetworkVersion.h"
 #include "UObject/CoreNet.h"
 #include "Net/DataBunch.h"
-#include "Net/DelinquencyAnalytics.h"
+#include "Net/NetAnalyticsTypes.h"
 #include "PackageMapClient.generated.h"
 
 class UNetConnection;
diff --git a/Engine/Source/Runtime/Engine/Private/Net/NetAnalyticsTypes.cpp b/Engine/Source/Runtime/Engine/Private/Net/NetAnalyticsTypes.cpp
new file mode 100644
--- /dev/null
+++ b/Engine/Source/Runtime/Engine/Private/Net/NetAnalyticsTypes.cpp
@@ -0,0 +1,179 @@
+// Copyright 1998-2019 Epic Games, Inc. All Rights Reserved.
+
+#include "Net/NetAnalyticsTypes.h"
+#include "Algo/BinarySearch.h"
+#include "Serialization/Archive.h"
+
+FDelinquencyAnalytics::FDelinquencyAnalytics(const uint32 InNumberOfTopOffendersToTrack) :
+	TotalTime(0.f),
+	NumberOfTopOffendersToTrack(InNumberOfTopOffendersToTrack)
+{
+	TopOffenders.Reserve(NumberOfTopOffendersToTrack);
+}
+
+FDelinquencyAnalytics::FDelinquencyAnalytics(FDelinquencyAnalytics&& Other) :
+	TopOffenders(MoveTemp(Other.TopOffenders)),
+	AllDelinquents(MoveTemp(Other.AllDelinquents)),
+	TotalTime(Other.TotalTime),
+	NumberOfTopOffendersToTrack(Other.NumberOfTopOffendersToTrack)
+{
+	TopOffenders.Reserve(NumberOfTopOffendersToTrack);
+}
+
+void FDelinquencyAnalytics::Add(FDelinquencyNameTimePair&& ToTrack)
+{
+	struct FHelper
+	{
+		static bool Compare(const FDelinquencyNameTimePair& LHS, const FDelinquencyNameTimePair& RHS)
+		{
+			return LHS.TimeSeconds > RHS.TimeSeconds;
+		}
+	};
+
+	// Regardless of whether or not this item has been seen before, there was a new entry
+	// so we'll add that time to the total.
+	TotalTime += ToTrack.TimeSeconds;
+
+	// TODO: We might consider tracking the code below as totals instead of the max time for a single event.
+	// For example, an Actor could end up queueing bunches several times within a reporting window, and each
+	// of those events would add to the TotalTime, but below we will only have the time of the longest period
+	// where bunches were queued.
+
+	if (NumberOfTopOffendersToTrack == 0)
+	{
+		AllDelinquents.Emplace(MoveTemp(ToTrack));
+	}
+	else if (TopOffenders.Num() == 0)
+	{
+		TopOffenders.Add(ToTrack);
+		AllDelinquents.Emplace(MoveTemp(ToTrack));
+		return;
+	}
+	else
+	{
+		if (FDelinquencyNameTimePair* AlreadyTracked = AllDelinquents.Find(ToTrack))
+		{
+			// We found an entry, so check its time.
+			if (AlreadyTracked->TimeSeconds >= ToTrack.TimeSeconds)
+			{
+				// We already have tracked a worse offense for this entry, so there's nothing more
+				// we need to do.
+				return;
+			}
+			else if (TopOffenders.Num() > 0)
+			{
+				const float LeastOffensiveTime = TopOffenders.Last().TimeSeconds;
+				if (AlreadyTracked->TimeSeconds >= LeastOffensiveTime)
+				{
+					// Our previous offense would have been tracked in the TopOffenders list, so go ahead and remove it.
+					// We're sorted highest to lowest, but are using greater than to actually compare, so UpperBound should
+					// return the index above which of the entry just before us.
+					// However, it's possible that other entries have the same time, so we need to make sure we remove the
+					// correct one.
+					int32 MaybeOurEntry = Algo::UpperBound<>(TopOffenders, ToTrack, &FHelper::Compare);
+					while (true)
+					{
+						// Sanity check that we have actually found our entry.
+						// If we hit the end of the list, or we see an entry that should be less offensive, then
+						// we've missed our entry.
+						if (MaybeOurEntry == TopOffenders.Num() || TopOffenders[MaybeOurEntry].TimeSeconds < AlreadyTracked->TimeSeconds)
+						{
+							// It's possible that multiple entries have the same delinquency time.
+							// If our current entry matches the LeastOffensiveTime, it's possible that we were in the
+							// top offenders list at one point but were pushed out when more offensive entries were added.
+							// If our time doesn't match the least offensive time, we should definitely be in the list.
+							ensureMsgf(AlreadyTracked->TimeSeconds == LeastOffensiveTime, TEXT("FDelinquencyAnalytics::Add - Unable to find expected entry %s:%f, list may not be sorted!"), *AlreadyTracked->Name.ToString(), AlreadyTracked->TimeSeconds);
+							break;
+						}
+
+						// We found our entry, so we're done.
+						else if (TopOffenders[MaybeOurEntry].Name == AlreadyTracked->Name)
+						{
+							TopOffenders.RemoveAt(MaybeOurEntry, 1, false);
+							break;
+						}
+
+						++MaybeOurEntry;
+					}
+				}
+			}
+		}
+
+		AllDelinquents.Add(ToTrack);
+
+		const int32 LocalNumberOfTopOffendersToTrack = static_cast<int32>(NumberOfTopOffendersToTrack);
+		const int32 InsertAt = Algo::UpperBound<>(TopOffenders, ToTrack, &FHelper::Compare);
+
+		// Check to see if this time was ranked in our top offenders.
+		if (InsertAt < LocalNumberOfTopOffendersToTrack)
+		{
+			// If we're going to displace a previous top offender, remove the least offensive.
+			if (LocalNumberOfTopOffendersToTrack == TopOffenders.Num())
+			{
+				TopOffenders.RemoveAt(TopOffenders.Num() - 1, 1, false);
+			}
+
+			TopOffenders.InsertUninitialized(InsertAt, 1);
+			TopOffenders[InsertAt] = MoveTemp(ToTrack);
+		}
+	}
+}
+
+void FDelinquencyAnalytics::Reset()
+{
+	TopOffenders.Reset();
+	AllDelinquents.Reset();
+	TotalTime = 0;
+}
+
+void FDelinquencyAnalytics::CountBytes(FArchive& Ar) const
+{
+	TopOffenders.CountBytes(Ar);
+	AllDelinquents.CountBytes(Ar);
+}
+
+void FNetConnectionSaturationAnalytics::TrackFrame(const bool bIsSaturated)
+{
+	++NumberOfTrackedFrames;
+
+	if (bIsSaturated)
+	{
+		++NumberOfSaturatedFrames;
+		++CurrentRunOfSaturatedFrames;
+	}
+	else
+	{
+		LongestRunOfSaturatedFrames = FMath::Max<uint32>(CurrentRunOfSaturatedFrames, LongestRunOfSaturatedFrames);
+		CurrentRunOfSaturatedFrames = 0;
+	}
+}
+
+void FNetConnectionSaturationAnalytics::TrackReplication(const bool bIsSaturated)
+{
+	++NumberOfReplications;
+
+	if (bIsSaturated)
+	{
+		++NumberOfSaturatedReplications;
+		++CurrentRunOfSaturatedReplications;
+	}
+	else
+	{
+		LongestRunOfSaturatedReplications = FMath::Max<uint32>(CurrentRunOfSaturatedReplications, LongestRunOfSaturatedReplications);
+		CurrentRunOfSaturatedReplications = 0;
+	}
+}
+
+void FNetConnectionSaturationAnalytics::Reset()
+{
+	NumberOfSaturatedFrames = 0;
+	NumberOfTrackedFrames = 0;
+	LongestRunOfSaturatedFrames = 0;
+
+	NumberOfReplications = 0;
+	NumberOfSaturatedReplications = 0;
+	LongestRunOfSaturatedReplications = 0;
+
+	CurrentRunOfSaturatedFrames = 0;
+	CurrentRunOfSaturatedReplications = 0;
+}
\ No newline at end of file
diff --git a/Engine/Source/Runtime/Engine/Private/NetConnection.cpp b/Engine/Source/Runtime/Engine/Private/NetConnection.cpp
--- a/Engine/Source/Runtime/Engine/Private/NetConnection.cpp
+++ b/Engine/Source/Runtime/Engine/Private/NetConnection.cpp
@@ -2959,6 +2959,8 @@ void UNetConnection::Tick()
 	// Update queued byte count.
 	// this should be at the end so that the cap is applied *after* sending (and adjusting QueuedBytes for) any remaining data for this tick
 
+	SaturationAnalytics.TrackFrame(!IsNetReady(false));
+
 	// Clamp DeltaTime for bandwidth limiting so that if there is a hitch, we don't try to send
 	// a large burst on the next frame, which can cause another hitch if a lot of additional replication occurs.
 	const float DesiredTickRate = GEngine->GetMaxTickRate(0.0f, false);
@@ -3492,6 +3494,28 @@ void UNetConnection::ResetQueuedActorDelinquencyAnalytics()
 	}
 }
 
+void UNetConnection::ConsumeSaturationAnalytics(FNetConnectionSaturationAnalytics& Out)
+{
+	Out = MoveTemp(SaturationAnalytics);
+	SaturationAnalytics.Reset();
+}
+
+const FNetConnectionSaturationAnalytics& UNetConnection::GetSaturationAnalytics() const
+{
+	return SaturationAnalytics;
+}
+
+void UNetConnection::ResetSaturationAnalytics()
+{
+	SaturationAnalytics.Reset();
+}
+
+void UNetConnection::TrackReplicationForAnalytics(const bool bWasSaturated)
+{
+	++TickCount;
+	SaturationAnalytics.TrackReplication(bWasSaturated);
+}
+
 /*-----------------------------------------------------------------------------
 	USimulatedClientNetConnection.
 -----------------------------------------------------------------------------*/
diff --git a/Engine/Source/Runtime/Engine/Private/NetworkDriver.cpp b/Engine/Source/Runtime/Engine/Private/NetworkDriver.cpp
--- a/Engine/Source/Runtime/Engine/Private/NetworkDriver.cpp
+++ b/Engine/Source/Runtime/Engine/Private/NetworkDriver.cpp
@@ -3768,7 +3768,6 @@ int32 UNetDriver::ServerReplicateActors_PrioritizeActors( UNetConnection* Connec
 	// Get list of visible/relevant actors.
 
 	NetTag++;
-	Connection->TickCount++;
 
 	// Set up to skip all sent temporary actors
 	for ( int32 j = 0; j < Connection->SentTemporaries.Num(); j++ )
@@ -4359,6 +4358,9 @@ int32 UNetDriver::ServerReplicateActors(float DeltaSeconds)
 		}
 		else if (Connection->ViewTarget)
 		{		
+
+			const int32 LocalNumSaturated = GNumSaturatedConnections;
+
 			// Make a list of viewers this connection should consider (this connection and children of this connection)
 			TArray<FNetViewer>& ConnectionViewers = WorldSettings->ReplicationViewers;
 
@@ -4431,6 +4433,9 @@ int32 UNetDriver::ServerReplicateActors(float DeltaSeconds)
 			RelevantActorMark.Pop();
 
 			ConnectionViewers.Reset();
+
+			const bool bWasSaturated = GNumSaturatedConnections > LocalNumSaturated;
+			Connection->TrackReplicationForAnalytics(bWasSaturated);
 		}
 	}
 
diff --git a/Engine/Source/Runtime/Engine/Public/Net/DelinquencyAnalytics.h b/Engine/Source/Runtime/Engine/Public/Net/DelinquencyAnalytics.h
deleted file mode 100644
diff --git a/Engine/Source/Runtime/Engine/Public/Net/NetAnalyticsTypes.h b/Engine/Source/Runtime/Engine/Public/Net/NetAnalyticsTypes.h
new file mode 100644
--- /dev/null
+++ b/Engine/Source/Runtime/Engine/Public/Net/NetAnalyticsTypes.h
@@ -0,0 +1,285 @@
+// Copyright 1998-2019 Epic Games, Inc. All Rights Reserved.
+
+#pragma once
+
+#include "CoreMinimal.h"
+
+/**
+ * Tracks an FName ID to a time value. Time will be context dependent, but usually
+ * represents the total amount of time a specific action took (how long a package
+ * took to load, how long an actor had queued bunches, etc.)
+ *
+ * Could have used a TPair, but this will make it more obvious what we're tracking.
+ */
+struct ENGINE_API FDelinquencyNameTimePair
+{
+public:
+
+	FDelinquencyNameTimePair(FName InName, float InTimeSeconds) :
+		Name(InName),
+		TimeSeconds(InTimeSeconds)
+	{
+	}
+
+	FName Name;
+	float TimeSeconds;
+};
+
+struct ENGINE_API FDelinquencyKeyFuncs : public BaseKeyFuncs<FDelinquencyNameTimePair, FDelinquencyNameTimePair, false>
+{
+	static KeyInitType GetSetKey(ElementInitType Element)
+	{
+		return Element;
+	}
+
+	static bool Matches(KeyInitType LHS, KeyInitType RHS)
+	{
+		return LHS.Name == RHS.Name;
+	}
+
+	static uint32 GetKeyHash(KeyInitType Key)
+	{
+		return GetTypeHash(Key.Name);
+	}
+};
+
+/**
+ * Convenience type that can be used to tracks information about things that can result in prolonged
+ * periods of apparent network inactivity, despite actually receiving traffic.
+ *
+ * The overall number of entries is expected to be small, but ultimately is left up to callers.
+ */
+struct ENGINE_API FDelinquencyAnalytics
+{
+public:
+
+	explicit FDelinquencyAnalytics(const uint32 InNumberOfTopOffendersToTrack);
+
+	FDelinquencyAnalytics(FDelinquencyAnalytics&& Other);
+
+	FDelinquencyAnalytics(const FDelinquencyAnalytics&) = delete;
+	const FDelinquencyAnalytics& operator=(const FDelinquencyAnalytics&) = delete;
+	FDelinquencyAnalytics& operator=(FDelinquencyAnalytics&&) = default;
+
+	void Emplace(FName Name, float TimeSeconds)
+	{
+		Add(FDelinquencyNameTimePair(Name, TimeSeconds));
+	}
+
+	/**
+	 * Adds the event to the delinquency tracking, by accumulating its time into total time,
+	 * and updating any existing events to choose the one with the highest time.
+	 *
+	 * When NumberOfTopOffendersToTrack == 0, we will just track the set of all events as well as the total time.
+	 *
+	 * When NumberOfTopOffendersToTrack > 0, we will track the set, total time, and also maintain sorted list
+	 * (highest to lowest) of events that occurred.
+	 *
+	 * By setting NumberOfTopOffendersToTrack to 0, users can manage their own lists of "TopOffenders", or
+	 * otherwise avoid the per add overhead of this tracking.
+	 */
+	void Add(FDelinquencyNameTimePair&& ToTrack);
+
+	const TArray<FDelinquencyNameTimePair>& GetTopOffenders() const
+	{
+		return TopOffenders;
+	}
+
+	const TSet<FDelinquencyNameTimePair, FDelinquencyKeyFuncs>& GetAllDelinquents() const
+	{
+		return AllDelinquents;
+	}
+
+	const float GetTotalTime() const
+	{
+		return TotalTime;
+	}
+
+	const uint32 GetNumberOfTopOffendersToTrack() const
+	{
+		return NumberOfTopOffendersToTrack;
+	}
+
+	void Reset();
+
+	void CountBytes(class FArchive& Ar) const;
+
+private:
+
+	TArray<FDelinquencyNameTimePair> TopOffenders;
+	TSet<FDelinquencyNameTimePair, FDelinquencyKeyFuncs> AllDelinquents;
+	float TotalTime;
+
+	// This is explicitly non const, as we will be copying / moving these structs around.
+	uint32 NumberOfTopOffendersToTrack;
+};
+
+/**
+ * Tracks data related specific to a NetDriver that can can result in prolonged periods of apparent
+ * network inactivity, despite actually receiving traffic.
+ *
+ * This includes things like Pending Async Loads.
+ *
+ * Also @see FConnectionDelinquencyAnalytics and FDelinquencyAnalytics.
+ */
+struct ENGINE_API FNetAsyncLoadDelinquencyAnalytics
+{
+	FNetAsyncLoadDelinquencyAnalytics() :
+		DelinquentAsyncLoads(0),
+		MaxConcurrentAsyncLoads(0)
+	{
+	}
+
+	FNetAsyncLoadDelinquencyAnalytics(const uint32 NumberOfTopOffendersToTrack) :
+		DelinquentAsyncLoads(NumberOfTopOffendersToTrack),
+		MaxConcurrentAsyncLoads(0)
+	{
+	}
+
+	FNetAsyncLoadDelinquencyAnalytics(FNetAsyncLoadDelinquencyAnalytics&& Other) :
+		DelinquentAsyncLoads(MoveTemp(Other.DelinquentAsyncLoads)),
+		MaxConcurrentAsyncLoads(Other.MaxConcurrentAsyncLoads)
+	{
+	}
+
+	FNetAsyncLoadDelinquencyAnalytics(const FNetAsyncLoadDelinquencyAnalytics&) = delete;
+	const FNetAsyncLoadDelinquencyAnalytics& operator=(const FNetAsyncLoadDelinquencyAnalytics&) = delete;
+	FNetAsyncLoadDelinquencyAnalytics& operator=(FNetAsyncLoadDelinquencyAnalytics&&) = default;
+
+	void CountBytes(FArchive& Ar) const
+	{
+		DelinquentAsyncLoads.CountBytes(Ar);
+	}
+
+	void Reset()
+	{
+		DelinquentAsyncLoads.Reset();
+		MaxConcurrentAsyncLoads = 0;
+	}
+
+	FDelinquencyAnalytics DelinquentAsyncLoads;
+	uint32 MaxConcurrentAsyncLoads;
+};
+
+/**
+ * Tracks data related specific to a NetConnection that can can result in prolonged periods of apparent
+ * network inactivity, despite actually receiving traffic.
+ *
+ * Also @see FDriverDelinquencyAnalytics and FDelinquencyAnalytics.
+ */
+struct ENGINE_API FNetQueuedActorDelinquencyAnalytics
+{
+	FNetQueuedActorDelinquencyAnalytics() :
+		DelinquentQueuedActors(0),
+		MaxConcurrentQueuedActors(0)
+	{
+	}
+
+	FNetQueuedActorDelinquencyAnalytics(const uint32 NumberOfTopOffendersToTrack) :
+		DelinquentQueuedActors(NumberOfTopOffendersToTrack),
+		MaxConcurrentQueuedActors(0)
+	{
+	}
+
+	FNetQueuedActorDelinquencyAnalytics(FNetQueuedActorDelinquencyAnalytics&& Other) :
+		DelinquentQueuedActors(MoveTemp(Other.DelinquentQueuedActors)),
+		MaxConcurrentQueuedActors(Other.MaxConcurrentQueuedActors)
+	{
+	}
+
+	FNetQueuedActorDelinquencyAnalytics(const FNetQueuedActorDelinquencyAnalytics&) = delete;
+	const FNetQueuedActorDelinquencyAnalytics& operator=(const FNetQueuedActorDelinquencyAnalytics&) = delete;
+	FNetQueuedActorDelinquencyAnalytics& operator=(FNetQueuedActorDelinquencyAnalytics&&) = default;
+
+
+	void CountBytes(FArchive& Ar) const
+	{
+		DelinquentQueuedActors.CountBytes(Ar);
+	}
+
+	void Reset()
+	{
+		DelinquentQueuedActors.Reset();
+		MaxConcurrentQueuedActors = 0;
+	}
+
+	FDelinquencyAnalytics DelinquentQueuedActors;
+	uint32 MaxConcurrentQueuedActors;
+};
+
+/** Struct wrapping Per Net Connection saturation analytics. */
+struct ENGINE_API FNetConnectionSaturationAnalytics
+{
+	FNetConnectionSaturationAnalytics() :
+		NumberOfTrackedFrames(0),
+		NumberOfSaturatedFrames(0),
+		LongestRunOfSaturatedFrames(0),
+		NumberOfReplications(0),
+		NumberOfSaturatedReplications(0),
+		LongestRunOfSaturatedReplications(0),
+		CurrentRunOfSaturatedFrames(0),
+		CurrentRunOfSaturatedReplications(0)
+	{
+	}
+
+	/** The total number of frames that we have currently tracked. */
+	const uint32 GetNumberOfTrackedFrames() const
+	{
+		return NumberOfTrackedFrames;
+	}
+
+	/** The number of frames we have reported as saturated.*/
+	const uint32 GetNumberOfSaturatedFrames() const
+	{
+		return NumberOfSaturatedFrames;
+	}
+
+	/** The longest number of consecutive frames that we have been saturated. */
+	const uint32 GetLongestRunOfSaturatedFrames() const
+	{
+		return LongestRunOfSaturatedFrames;
+	}
+
+	/**
+	 * The number of times we have tried to replicate data on this connection
+	 * (UNetDriver::ServerReplicateActors / UReplicationGraph::ServerReplicateActors)
+	 */
+	const uint32 GetNumberOfReplications() const
+	{
+		return NumberOfReplications;
+	}
+
+	/** The number of times we have been pre-empted from replicating all data, due to saturation. */
+	const uint32 GetNumberOfSaturatedReplications() const
+	{
+		return NumberOfSaturatedReplications;
+	}
+
+	/** The longest number of consecutive replication attempts where we were pre-empted due to saturation. */
+	const uint32 GetLongestRunOfSaturatedReplications() const
+	{
+		return LongestRunOfSaturatedReplications;
+	}
+
+	/** Resets the state of tracking. */
+	void Reset();
+
+private:
+
+	friend class UNetConnection;
+
+	void TrackFrame(const bool bIsSaturated);
+
+	void TrackReplication(const bool bIsSaturated);
+
+	uint32 NumberOfTrackedFrames;
+	uint32 NumberOfSaturatedFrames;
+	uint32 LongestRunOfSaturatedFrames;
+
+	uint32 NumberOfReplications;
+	uint32 NumberOfSaturatedReplications;
+	uint32 LongestRunOfSaturatedReplications;
+
+	uint32 CurrentRunOfSaturatedFrames;
+	uint32 CurrentRunOfSaturatedReplications;
+};
\ No newline at end of file
