#ocn - Fix for touch ended events with (0,0) position causing missed taps. Waiting to hear back from Engine if this is actually a good fix, but checking in for now because we need those missed taps fixed.

diff --git a/Engine/Source/Runtime/Engine/Private/Slate/SceneViewport.cpp b/Engine/Source/Runtime/Engine/Private/Slate/SceneViewport.cpp
--- a/Engine/Source/Runtime/Engine/Private/Slate/SceneViewport.cpp
+++ b/Engine/Source/Runtime/Engine/Private/Slate/SceneViewport.cpp
@@ -791,6 +791,7 @@ FReply FSceneViewport::OnTouchEnded( const FGeometry& MyGeometry, const FPointer
 	}
 	else
 	{
+		UpdateCachedCursorPos(MyGeometry, TouchEvent);
 		CurCursorPos = CachedCursorPos;
 		CachedCursorPos = FIntPoint(-1, -1);
 	}
