Adding support to set a custom index into the custom primitive data per pin

#rnx

diff --git a/Engine/Source/Runtime/Engine/Classes/Materials/MaterialExpressionCustomPrimitiveData.h b/Engine/Source/Runtime/Engine/Classes/Materials/MaterialExpressionCustomPrimitiveData.h
--- a/Engine/Source/Runtime/Engine/Classes/Materials/MaterialExpressionCustomPrimitiveData.h
+++ b/Engine/Source/Runtime/Engine/Classes/Materials/MaterialExpressionCustomPrimitiveData.h
@@ -7,6 +7,20 @@
 #include "Materials/MaterialExpression.h"
 #include "MaterialExpressionCustomPrimitiveData.generated.h"
 
+USTRUCT()
+struct FPrimitiveDataIndex
+{
+	GENERATED_USTRUCT_BODY()
+
+#if WITH_EDITORONLY_DATA
+	UPROPERTY(EditAnywhere, Category=MaterialExpression)
+	FString CustomDesc;
+#endif
+
+	UPROPERTY(EditAnywhere, Category=MaterialExpression, meta=(ClampMin="0"))
+	uint8 PrimitiveDataIndex;
+};
+
 UCLASS(collapsecategories, hidecategories = Object)
 class UMaterialExpressionCustomPrimitiveData: public UMaterialExpression
 {
@@ -18,10 +32,9 @@ class UMaterialExpressionCustomPrimitiveData: public UMaterialExpression
 	virtual void GetCaption(TArray<FString>& OutCaptions) const override;
 	virtual void PostEditChangeProperty(FPropertyChangedEvent& PropertyChangedEvent) override;
 	//~ End UMaterialExpression Interface
-#if WITH_EDITORONLY_DATA
-	/** Custom descriptions for each custom data pin, only for making things more readable, don't rely on this data for run-time things */
-	UPROPERTY(EditAnywhere, Category=MaterialExpression)
-	TArray<FString> CustomDescs;
-#endif
 #endif
+
+	/** Custom index into the custom data for each pin, also includes custom descriptions for each custom data pin only for making things more readable */
+	UPROPERTY(EditAnywhere, Category=MaterialExpression)
+	TArray<FPrimitiveDataIndex> CustomIndices;
 };
\ No newline at end of file
diff --git a/Engine/Source/Runtime/Engine/Private/Materials/MaterialExpressions.cpp b/Engine/Source/Runtime/Engine/Private/Materials/MaterialExpressions.cpp
--- a/Engine/Source/Runtime/Engine/Private/Materials/MaterialExpressions.cpp
+++ b/Engine/Source/Runtime/Engine/Private/Materials/MaterialExpressions.cpp
@@ -15318,7 +15318,7 @@ int32 UMaterialExpressionCustomPrimitiveData::Compile(class FMaterialCompiler* C
 		return CompilerError(Compiler, TEXT("Expression not available in the deferred decal material domain."));
 	}
 
-	return Compiler->CustomPrimitiveData(OutputIndex);
+	return Compiler->CustomPrimitiveData(CustomIndices[OutputIndex].PrimitiveDataIndex);
 }
 
 void UMaterialExpressionCustomPrimitiveData::GetCaption(TArray<FString>& OutCaptions) const
@@ -15329,21 +15329,23 @@ void UMaterialExpressionCustomPrimitiveData::GetCaption(TArray<FString>& OutCapt
 void UMaterialExpressionCustomPrimitiveData::PostEditChangeProperty(FPropertyChangedEvent& PropertyChangedEvent)
 {
 #if WITH_EDITORONLY_DATA
-	if (PropertyChangedEvent.GetPropertyName() == GET_MEMBER_NAME_CHECKED(UMaterialExpressionCustomPrimitiveData, CustomDescs))
+	
+	if (PropertyChangedEvent.GetPropertyName() == GET_MEMBER_NAME_CHECKED(FPrimitiveDataIndex, CustomDesc) ||
+		PropertyChangedEvent.GetPropertyName() == GET_MEMBER_NAME_CHECKED(UMaterialExpressionCustomPrimitiveData, CustomIndices))
 	{
 		// If more than the supported number of custom floats have been added, remove the overflow
-		if (CustomDescs.Num() > FCustomPrimitiveData::NumCustomPrimitiveDataFloats)
+		if (CustomIndices.Num() > FCustomPrimitiveData::NumCustomPrimitiveDataFloats)
 		{
-			CustomDescs.SetNum(FCustomPrimitiveData::NumCustomPrimitiveDataFloats);
+			CustomIndices.SetNum(FCustomPrimitiveData::NumCustomPrimitiveDataFloats);
 		}
 
 		// Match the number of pins to the number of descriptions
-		Outputs.SetNumZeroed(CustomDescs.Num());
+		Outputs.SetNumZeroed(CustomIndices.Num());
 
 		// Update the names of the pins
 		for (int i = 0; i < Outputs.Num(); i++)
 		{
-			Outputs[i].OutputName = *CustomDescs[i];
+			Outputs[i].OutputName = *(CustomIndices[i].CustomDesc);
 		}
 
 		// Reconstruct the node to get the new names applied to the pins
@@ -15351,12 +15353,22 @@ void UMaterialExpressionCustomPrimitiveData::PostEditChangeProperty(FPropertyCha
 
 		// No need to update preview if we only change output descriptions
 		bNeedToUpdatePreview = false;
+
+		// Done here (skip calling base class)
+		return;
 	}
-	else
-#endif
+	else if (PropertyChangedEvent.GetPropertyName() == GET_MEMBER_NAME_CHECKED(FPrimitiveDataIndex, PrimitiveDataIndex))
 	{
-		Super::PostEditChangeProperty(PropertyChangedEvent);
+		// Clamp indices
+		for (int i = 0; i < CustomIndices.Num(); i++)
+		{
+			const int32 PrimDataIndex = CustomIndices[i].PrimitiveDataIndex;
+			CustomIndices[i].PrimitiveDataIndex = FMath::Clamp(PrimDataIndex, 0, FCustomPrimitiveData::NumCustomPrimitiveDataFloats-1);
+		}
 	}
+#endif
+
+	Super::PostEditChangeProperty(PropertyChangedEvent);
 }
 #endif // WITH_EDITOR
 
