Adding default value to the shading models field on material instance to not have it be invalid on newly created material instances. Also protect from overriding using the From Material Expression index.

#rnx

diff --git a/Engine/Source/Runtime/Engine/Private/Materials/MaterialInstance.cpp b/Engine/Source/Runtime/Engine/Private/Materials/MaterialInstance.cpp
--- a/Engine/Source/Runtime/Engine/Private/Materials/MaterialInstance.cpp
+++ b/Engine/Source/Runtime/Engine/Private/Materials/MaterialInstance.cpp
@@ -402,6 +402,7 @@ UMaterialInstance::UMaterialInstance(const FObjectInitializer& ObjectInitializer
 	bHasStaticPermutationResource = false;
 	ReentrantFlag[0] = false;
 	ReentrantFlag[1] = false;
+	ShadingModels = MSM_Unlit;
 }
 
 void UMaterialInstance::PostInitProperties()	
@@ -2434,9 +2435,13 @@ void UMaterialInstance::UpdateOverridableBaseProperties()
 
 	if (BasePropertyOverrides.bOverride_ShadingModel)
 	{
-		// It's only possible to override using a single shading model
-		ShadingModels = FMaterialShadingModelField(BasePropertyOverrides.ShadingModel);
-		bIsShadingModelFromMaterialExpression = 0;
+		// Can only override using one of the actual shading models
+		if (BasePropertyOverrides.ShadingModel != MSM_FromMaterialExpression)
+		{
+			// It's only possible to override using a single shading model
+			ShadingModels = FMaterialShadingModelField(BasePropertyOverrides.ShadingModel);
+			bIsShadingModelFromMaterialExpression = 0;
+		}
 	}
 	else
 	{
