SetTimer and SetTimerDelegate now take an optional param, InitialStartDelayVariance that will allow users to add some randomization to the initial start delay given to a timer.  This is useful in situations where we have lots of the same object in a level that does some work on a timer when we don't want all of the timers to execute callbacks on the exact same frame.

diff --git a/Engine/Source/Runtime/Engine/Classes/Kismet/KismetSystemLibrary.h b/Engine/Source/Runtime/Engine/Classes/Kismet/KismetSystemLibrary.h
--- a/Engine/Source/Runtime/Engine/Classes/Kismet/KismetSystemLibrary.h
+++ b/Engine/Source/Runtime/Engine/Classes/Kismet/KismetSystemLibrary.h
@@ -461,14 +461,15 @@ class ENGINE_API UKismetSystemLibrary : public UBlueprintFunctionLibrary
 
 	/**
 	 * Set a timer to execute delegate. Setting an existing timer will reset that timer with updated parameters.
-	 * @param Event				Event. Can be a K2 function or a Custom Event.
-	 * @param Time				How long to wait before executing the delegate, in seconds. Setting a timer to <= 0 seconds will clear it if it is set.
-	 * @param bLooping			True to keep executing the delegate every Time seconds, false to execute delegate only once.
-	 * @param InitialStartDelay Initial delay passed to the timer manager to allow some variance in when the timer starts.
-	 * @return					The timer handle to pass to other timer functions to manipulate this timer.
+	 * @param Event						Event. Can be a K2 function or a Custom Event.
+	 * @param Time						How long to wait before executing the delegate, in seconds. Setting a timer to <= 0 seconds will clear it if it is set.
+	 * @param bLooping					True to keep executing the delegate every Time seconds, false to execute delegate only once.
+	 * @param InitialStartDelay			Initial delay passed to the timer manager, in seconds.
+	 * @param InitialStartDelayVariance	Use this to add some variance to when the timer starts in lieu of doing a random range on the InitialStartDelay input, in seconds. 
+	 * @return							The timer handle to pass to other timer functions to manipulate this timer.
 	 */
-	UFUNCTION(BlueprintCallable, meta=(DisplayName = "Set Timer by Event", ScriptName = "SetTimerDelegate", AdvancedDisplay="InitialStartDelay"), Category="Utilities|Time")
-	static FTimerHandle K2_SetTimerDelegate(UPARAM(DisplayName="Event") FTimerDynamicDelegate Delegate, float Time, bool bLooping, float InitialStartDelay = 0.f);
+	UFUNCTION(BlueprintCallable, meta=(DisplayName = "Set Timer by Event", ScriptName = "SetTimerDelegate", AdvancedDisplay="InitialStartDelay, InitialStartDelayVariance"), Category="Utilities|Time")
+	static FTimerHandle K2_SetTimerDelegate(UPARAM(DisplayName="Event") FTimerDynamicDelegate Delegate, float Time, bool bLooping, float InitialStartDelay = 0.f, float InitialStartDelayVariance = 0.f);
 
 	/**
 	 * Clears a set timer.
@@ -621,15 +622,16 @@ class ENGINE_API UKismetSystemLibrary : public UBlueprintFunctionLibrary
 
 	/**
 	 * Set a timer to execute delegate. Setting an existing timer will reset that timer with updated parameters.
-	 * @param Object			Object that implements the delegate function. Defaults to self (this blueprint)
-	 * @param FunctionName		Delegate function name. Can be a K2 function or a Custom Event.
-	 * @param Time				How long to wait before executing the delegate, in seconds. Setting a timer to <= 0 seconds will clear it if it is set.
-	 * @param bLooping			true to keep executing the delegate every Time seconds, false to execute delegate only once.
-	 * @param InitialStartDelay Initial delay passed to the timer manager to allow some variance in when the timer starts.
-	 * @return					The timer handle to pass to other timer functions to manipulate this timer.
-	 */
-	UFUNCTION(BlueprintCallable, meta=(DisplayName = "Set Timer by Function Name", ScriptName = "SetTimer", DefaultToSelf = "Object", AdvancedDisplay="InitialStartDelay"), Category="Utilities|Time")
-	static FTimerHandle K2_SetTimer(UObject* Object, FString FunctionName, float Time, bool bLooping, float InitialStartDelay = 0.f);
+	 * @param Object					Object that implements the delegate function. Defaults to self (this blueprint)
+	 * @param FunctionName				Delegate function name. Can be a K2 function or a Custom Event.
+	 * @param Time						How long to wait before executing the delegate, in seconds. Setting a timer to <= 0 seconds will clear it if it is set.
+	 * @param bLooping					true to keep executing the delegate every Time seconds, false to execute delegate only once.
+	 * @param InitialStartDelay			Initial delay passed to the timer manager to allow some variance in when the timer starts, in seconds.
+	 * @param InitialStartDelayVariance	Use this to add some variance to when the timer starts in lieu of doing a random range on the InitialStartDelay input, in seconds.
+	 * @return							The timer handle to pass to other timer functions to manipulate this timer.
+	 */
+	UFUNCTION(BlueprintCallable, meta=(DisplayName = "Set Timer by Function Name", ScriptName = "SetTimer", DefaultToSelf = "Object", AdvancedDisplay="InitialStartDelay, InitialStartDelayVariance"), Category="Utilities|Time")
+	static FTimerHandle K2_SetTimer(UObject* Object, FString FunctionName, float Time, bool bLooping, float InitialStartDelay = 0.f, float InitialStartDelayVariance = 0.f);
 
 	/**
 	 * Clears a set timer.
diff --git a/Engine/Source/Runtime/Engine/Private/KismetSystemLibrary.cpp b/Engine/Source/Runtime/Engine/Private/KismetSystemLibrary.cpp
--- a/Engine/Source/Runtime/Engine/Private/KismetSystemLibrary.cpp
+++ b/Engine/Source/Runtime/Engine/Private/KismetSystemLibrary.cpp
@@ -387,7 +387,7 @@ FTimerHandle UKismetSystemLibrary::K2_InvalidateTimerHandle(FTimerHandle& TimerH
 	return TimerHandle;
 }
 
-FTimerHandle UKismetSystemLibrary::K2_SetTimer(UObject* Object, FString FunctionName, float Time, bool bLooping, float InitialStartDelay)
+FTimerHandle UKismetSystemLibrary::K2_SetTimer(UObject* Object, FString FunctionName, float Time, bool bLooping, float InitialStartDelay, float InitialStartDelayVariance)
 {
 	FName const FunctionFName(*FunctionName);
 
@@ -404,9 +404,10 @@ FTimerHandle UKismetSystemLibrary::K2_SetTimer(UObject* Object, FString Function
 		}
 	}
 
-	if (Time <= 0.f || InitialStartDelay < 0.f)
+	InitialStartDelay += FMath::RandRange(-InitialStartDelayVariance, InitialStartDelayVariance);
+	if (Time <= 0.f || (Time - InitialStartDelayVariance) < 0.f)
 	{
-		FFrame::KismetExecutionMessage(TEXT("SetTimer passed a negative time or initial start delay.  The associated timer may fail to fire!"), ELogVerbosity::Warning);
+		FFrame::KismetExecutionMessage(TEXT("SetTimer passed a negative time.  The associated timer may fail to fire!  If using InitialStartDelayVariance, ensure it is smaller than (Time + InitialStartDelay)."), ELogVerbosity::Warning);
 	}
 
 	FTimerDynamicDelegate Delegate;
@@ -414,7 +415,7 @@ FTimerHandle UKismetSystemLibrary::K2_SetTimer(UObject* Object, FString Function
 	return K2_SetTimerDelegate(Delegate, Time, bLooping, InitialStartDelay);
 }
 
-FTimerHandle UKismetSystemLibrary::K2_SetTimerDelegate(FTimerDynamicDelegate Delegate, float Time, bool bLooping, float InitialStartDelay)
+FTimerHandle UKismetSystemLibrary::K2_SetTimerDelegate(FTimerDynamicDelegate Delegate, float Time, bool bLooping, float InitialStartDelay, float InitialStartDelayVariance)
 {
 	FTimerHandle Handle;
 	if (Delegate.IsBound())
@@ -422,9 +423,10 @@ FTimerHandle UKismetSystemLibrary::K2_SetTimerDelegate(FTimerDynamicDelegate Del
 		const UWorld* const World = GEngine->GetWorldFromContextObject(Delegate.GetUObject(), EGetWorldErrorMode::LogAndReturnNull);
 		if(World)
 		{
-			if (Time <= 0.f || InitialStartDelay < 0.f)
+			InitialStartDelay += FMath::RandRange(-InitialStartDelayVariance, InitialStartDelayVariance);
+			if (Time <= 0.f || (Time - InitialStartDelayVariance) < 0.f)
 			{
-				FFrame::KismetExecutionMessage(TEXT("SetTimer passed a negative time or initial start delay.  The associated timer may fail to fire!"), ELogVerbosity::Warning);
+				FFrame::KismetExecutionMessage(TEXT("SetTimer passed a negative time or initial start delay.  The associated timer may fail to fire!  If using InitialStartDelayVariance, ensure it is smaller than (Time + InitialStartDelay)."), ELogVerbosity::Warning);
 			}
 
 			FTimerManager& TimerManager = World->GetTimerManager();
