Fix crash when applying component instance data to ucs created components that are not marked pending kill. Should also fix similar issues with components that are not auto-registered.

diff --git a/Engine/Source/Runtime/Engine/Private/ComponentInstanceDataCache.cpp b/Engine/Source/Runtime/Engine/Private/ComponentInstanceDataCache.cpp
--- a/Engine/Source/Runtime/Engine/Private/ComponentInstanceDataCache.cpp
+++ b/Engine/Source/Runtime/Engine/Private/ComponentInstanceDataCache.cpp
@@ -579,8 +579,14 @@ void FComponentInstanceDataCache::ApplyToActor(AActor* Actor, const ECacheApplyP
 		{
 			if (USceneComponent* SceneComponent = Cast<USceneComponent>(Component))
 			{
+				// Scene components that aren't attached to the root component hierarchy won't already have been processed so process them now.
+				// * If there is an unattached scene component 
+				// * If there is a scene component attached to another Actor's hierarchy
+				// * If the scene is not registered (likely because bAutoRegister is false or component is marked pending kill), then we may not have successfully attached to our parent and properly been handled
 				USceneComponent* ParentComponent = SceneComponent->GetAttachParent();
-				if ((ParentComponent == nullptr && SceneComponent != Actor->GetRootComponent()) || (ParentComponent && ParentComponent->GetOwner() != Actor))
+				if (   (ParentComponent == nullptr && SceneComponent != Actor->GetRootComponent()) 			
+					|| (ParentComponent && ParentComponent->GetOwner() != Actor) 					
+					|| (ParentComponent && !SceneComponent->IsRegistered() && !ParentComponent->GetAttachChildren().Contains(SceneComponent)))
 				{
 					AddComponentHierarchy(SceneComponent);
 				}
