Updating copyright comment to 2019

#rnx

diff --git a/Engine/Source/Runtime/Engine/Classes/Materials/MaterialExpressionShadingModel.h b/Engine/Source/Runtime/Engine/Classes/Materials/MaterialExpressionShadingModel.h
--- a/Engine/Source/Runtime/Engine/Classes/Materials/MaterialExpressionShadingModel.h
+++ b/Engine/Source/Runtime/Engine/Classes/Materials/MaterialExpressionShadingModel.h
@@ -1,4 +1,4 @@
-// Copyright 1998-2018 Epic Games, Inc. All Rights Reserved.
+// Copyright 1998-2019 Epic Games, Inc. All Rights Reserved.
 
 #pragma once
 
