Move rebuilding of shading model field to another location to include a code path which was missed.

#rnx

diff --git a/Engine/Source/Runtime/Engine/Private/Materials/Material.cpp b/Engine/Source/Runtime/Engine/Private/Materials/Material.cpp
--- a/Engine/Source/Runtime/Engine/Private/Materials/Material.cpp
+++ b/Engine/Source/Runtime/Engine/Private/Materials/Material.cpp
@@ -3058,6 +3058,12 @@ void UMaterial::UpdateMaterialShaderCacheAndTextureReferences()
 
 void UMaterial::CacheResourceShadersForRendering(bool bRegenerateId)
 {
+	
+#if WITH_EDITOR
+	// Always rebuild the shading model field on recompile
+	RebuildShadingModelField();
+#endif //WITH_EDITOR
+
 	if (bRegenerateId)
 	{
 		// Regenerate this material's Id if requested
@@ -4371,9 +4377,6 @@ void UMaterial::PostEditChangeProperty(FPropertyChangedEvent& PropertyChangedEve
 	
 	if (bRequiresCompilation)
 	{
-		// Always rebuild the shading model field on recompile
-		RebuildShadingModelField();
-
 		// When redirecting an object pointer, we trust that the DDC hash will detect the change and that we don't need to force a recompile.
 		const bool bRegenerateId = PropertyChangedEvent.ChangeType != EPropertyChangeType::Redirected;
 		CacheResourceShadersForRendering(bRegenerateId);
