Added UFXSystemComponent as base class for UParticleSystemComponent and UNiagaraComponent.
Added UFXSystemAsset as base class for UParticleSystem and UNiagaraSystem

diff --git a/Engine/Plugins/FX/Niagara/Source/Niagara/Classes/NiagaraSystem.h b/Engine/Plugins/FX/Niagara/Source/Niagara/Classes/NiagaraSystem.h
--- a/Engine/Plugins/FX/Niagara/Source/Niagara/Classes/NiagaraSystem.h
+++ b/Engine/Plugins/FX/Niagara/Source/Niagara/Classes/NiagaraSystem.h
@@ -5,6 +5,7 @@
 #include "CoreMinimal.h"
 #include "UObject/ObjectMacros.h"
 #include "UObject/Object.h"
+#include "Particles/ParticleSystem.h"
 #include "NiagaraCommon.h"
 #include "NiagaraDataSet.h"
 #include "NiagaraEmitterInstance.h"
@@ -58,7 +59,7 @@ struct FNiagaraSystemCompileRequest
 
 /** Container for multiple emitters that combine together to create a particle system effect.*/
 UCLASS(BlueprintType)
-class NIAGARA_API UNiagaraSystem : public UObject
+class NIAGARA_API UNiagaraSystem : public UFXSystemAsset
 {
 	GENERATED_UCLASS_BODY()
 
diff --git a/Engine/Plugins/FX/Niagara/Source/Niagara/Private/NiagaraComponent.cpp b/Engine/Plugins/FX/Niagara/Source/Niagara/Private/NiagaraComponent.cpp
--- a/Engine/Plugins/FX/Niagara/Source/Niagara/Private/NiagaraComponent.cpp
+++ b/Engine/Plugins/FX/Niagara/Source/Niagara/Private/NiagaraComponent.cpp
@@ -358,6 +358,28 @@ UNiagaraComponent::UNiagaraComponent(const FObjectInitializer& ObjectInitializer
 	SavedAutoAttachRelativeScale3D = FVector(1.f, 1.f, 1.f);
 }
 
+/********* UFXSystemComponent *********/
+void UNiagaraComponent::SetFloatParameter(FName ParameterName, float Param)
+{
+	SetNiagaraVariableFloat(ParameterName.ToString(), Param);
+}
+
+void UNiagaraComponent::SetVectorParameter(FName ParameterName, FVector Param)
+{
+	SetNiagaraVariableVec3(ParameterName.ToString(), Param);
+}
+
+void UNiagaraComponent::SetColorParameter(FName ParameterName, FLinearColor Param)
+{
+	SetNiagaraVariableLinearColor(ParameterName.ToString(), Param);
+}
+
+void UNiagaraComponent::SetActorParameter(FName ParameterName, class AActor* Param)
+{
+	SetNiagaraVariableActor(ParameterName.ToString(), Param);
+}
+/********* UFXSystemComponent *********/
+
 
 void UNiagaraComponent::TickComponent(float DeltaSeconds, enum ELevelTick TickType, FActorComponentTickFunction* ThisTickFunction)
 {
diff --git a/Engine/Plugins/FX/Niagara/Source/Niagara/Public/NiagaraComponent.h b/Engine/Plugins/FX/Niagara/Source/Niagara/Public/NiagaraComponent.h
--- a/Engine/Plugins/FX/Niagara/Source/Niagara/Public/NiagaraComponent.h
+++ b/Engine/Plugins/FX/Niagara/Source/Niagara/Public/NiagaraComponent.h
@@ -7,7 +7,7 @@
 #include "NiagaraCommon.h"
 #include "PrimitiveViewRelevance.h"
 #include "PrimitiveSceneProxy.h"
-#include "Components/PrimitiveComponent.h"
+#include "Particles/ParticleSystemComponent.h"
 #include "NiagaraUserRedirectionParameterStore.h"
 #include "NiagaraSystemInstance.h"
 
@@ -31,7 +31,7 @@ DECLARE_DYNAMIC_MULTICAST_DELEGATE_OneParam(FOnNiagaraSystemFinished, class UNia
 * @see UNiagaraSystem
 */
 UCLASS(ClassGroup = (Rendering, Common), hidecategories = Object, hidecategories = Physics, hidecategories = Collision, showcategories = Trigger, editinlinenew, meta = (BlueprintSpawnableComponent, DisplayName = "Niagara Particle System"))
-class NIAGARA_API UNiagaraComponent : public UPrimitiveComponent
+class NIAGARA_API UNiagaraComponent : public UFXSystemComponent
 {
 	GENERATED_UCLASS_BODY()
 
@@ -40,6 +40,15 @@ class NIAGARA_API UNiagaraComponent : public UPrimitiveComponent
 	DECLARE_MULTICAST_DELEGATE(FOnSynchronizedWithAssetParameters);
 #endif
 
+public:
+
+	/********* UFXSystemComponent *********/
+	void SetFloatParameter(FName ParameterName, float Param) override;
+	void SetVectorParameter(FName ParameterName, FVector Param) override;
+	void SetColorParameter(FName ParameterName, FLinearColor Param) override;
+	void SetActorParameter(FName ParameterName, class AActor* Param) override;
+	/********* UFXSystemComponent *********/
+
 private:
 	UPROPERTY(EditAnywhere, Category="Niagara", meta = (DisplayName = "Niagara System Asset"))
 	UNiagaraSystem* Asset;
diff --git a/Engine/Source/Runtime/Engine/Classes/Particles/ParticleSystem.h b/Engine/Source/Runtime/Engine/Classes/Particles/ParticleSystem.h
--- a/Engine/Source/Runtime/Engine/Classes/Particles/ParticleSystem.h
+++ b/Engine/Source/Runtime/Engine/Classes/Particles/ParticleSystem.h
@@ -103,13 +103,21 @@ struct FNamedEmitterMaterial
 	class UMaterialInterface* Material;
 };
 
+UCLASS(Abstract, MinimalAPI, BlueprintType)
+class UFXSystemAsset : public UObject
+{
+	GENERATED_UCLASS_BODY()
+public:
+	UFXSystemAsset() {}
+};
+
 /**
  * A ParticleSystem is a complete particle effect that contains any number of ParticleEmitters. By allowing multiple emitters
  * in a system, the designer can create elaborate particle effects that are held in a single system. Once created using
  * Cascade, a ParticleSystem can then be inserted into a level or created in script.
  */
 UCLASS(hidecategories=Object, MinimalAPI, BlueprintType)
-class UParticleSystem : public UObject
+class UParticleSystem : public UFXSystemAsset
 {
 	GENERATED_UCLASS_BODY()
 
diff --git a/Engine/Source/Runtime/Engine/Classes/Particles/ParticleSystemComponent.h b/Engine/Source/Runtime/Engine/Classes/Particles/ParticleSystemComponent.h
--- a/Engine/Source/Runtime/Engine/Classes/Particles/ParticleSystemComponent.h
+++ b/Engine/Source/Runtime/Engine/Classes/Particles/ParticleSystemComponent.h
@@ -324,11 +324,44 @@ struct FParticleEventKismetData : public FParticleEventData
 {
 };
 
+UCLASS(Abstract)
+class ENGINE_API UFXSystemComponent : public UPrimitiveComponent
+{
+	GENERATED_UCLASS_BODY()
+public:
+
+	/** Change a named float parameter */
+	UFUNCTION(BlueprintCallable, Category="Effects|Components|ParticleSystem")
+	virtual void SetFloatParameter(FName ParameterName, float Param) {}
+
+	/** 
+	 *	Set a named vector instance parameter on this ParticleSystemComponent. 
+	 *	Updates the parameter if it already exists, or creates a new entry if not. 
+	 */
+	UFUNCTION(BlueprintCallable, Category="Effects|Components|ParticleSystem")
+	virtual void SetVectorParameter(FName ParameterName, FVector Param) {}
+
+	/** 
+	 *	Set a named color instance parameter on this ParticleSystemComponent. 
+	 *	Updates the parameter if it already exists, or creates a new entry if not. 
+	 */
+	UFUNCTION(BlueprintCallable, Category="Effects|Components|ParticleSystem")
+	virtual void SetColorParameter(FName ParameterName, FLinearColor Param) {}
+
+	/** 
+	 *	Set a named actor instance parameter on this ParticleSystemComponent. 
+	 *	Updates the parameter if it already exists, or creates a new entry if not. 
+	 */
+	UFUNCTION(BlueprintCallable, Category="Effects|Components|ParticleSystem")
+	virtual void SetActorParameter(FName ParameterName, class AActor* Param) {}
+};
+
+
 /** 
  * A particle emitter.
  */
 UCLASS(ClassGroup=(Rendering, Common), hidecategories=Object, hidecategories=Physics, hidecategories=Collision, showcategories=Trigger, editinlinenew, meta=(BlueprintSpawnableComponent))
-class ENGINE_API UParticleSystemComponent : public UPrimitiveComponent
+class ENGINE_API UParticleSystemComponent : public UFXSystemComponent
 {
 	friend class FParticleSystemWorldManager;
 
@@ -967,29 +1000,25 @@ public:
 	virtual void SetEmitterEnable(FName EmitterName, bool bNewEnableState);
 
 	/** Change a named float parameter */
-	UFUNCTION(BlueprintCallable, Category="Effects|Components|ParticleSystem")
-	void SetFloatParameter(FName ParameterName, float Param);
+	void SetFloatParameter(FName ParameterName, float Param) override;
 
 	/** 
 	 *	Set a named vector instance parameter on this ParticleSystemComponent. 
 	 *	Updates the parameter if it already exists, or creates a new entry if not. 
 	 */
-	UFUNCTION(BlueprintCallable, Category="Effects|Components|ParticleSystem")
-	void SetVectorParameter(FName ParameterName, FVector Param);
+	void SetVectorParameter(FName ParameterName, FVector Param) override;
 
 	/** 
 	 *	Set a named color instance parameter on this ParticleSystemComponent. 
 	 *	Updates the parameter if it already exists, or creates a new entry if not. 
 	 */
-	UFUNCTION(BlueprintCallable, Category="Effects|Components|ParticleSystem")
-	void SetColorParameter(FName ParameterName, FLinearColor Param);
+	void SetColorParameter(FName ParameterName, FLinearColor Param) override;
 
 	/** 
 	 *	Set a named actor instance parameter on this ParticleSystemComponent. 
 	 *	Updates the parameter if it already exists, or creates a new entry if not. 
 	 */
-	UFUNCTION(BlueprintCallable, Category="Effects|Components|ParticleSystem")
-	void SetActorParameter(FName ParameterName, class AActor* Param);
+	void SetActorParameter(FName ParameterName, class AActor* Param) override;
 
 	/** 
 	 *	Set a named material instance parameter on this ParticleSystemComponent. 
diff --git a/Engine/Source/Runtime/Engine/Private/Particles/ParticleComponents.cpp b/Engine/Source/Runtime/Engine/Private/Particles/ParticleComponents.cpp
--- a/Engine/Source/Runtime/Engine/Private/Particles/ParticleComponents.cpp
+++ b/Engine/Source/Runtime/Engine/Private/Particles/ParticleComponents.cpp
@@ -2048,6 +2048,15 @@ void UParticleSpriteEmitter::SetToSensibleDefaults()
 #endif // WITH_EDITOR
 }
 
+/*-----------------------------------------------------------------------------
+	UFXSystemAsset implementation.
+-----------------------------------------------------------------------------*/
+
+UFXSystemAsset::UFXSystemAsset(const FObjectInitializer& ObjectInitializer)
+	: Super(ObjectInitializer)
+{
+}
+
 /*-----------------------------------------------------------------------------
 	UParticleSystem implementation.
 -----------------------------------------------------------------------------*/
@@ -3309,6 +3318,10 @@ bool UParticleSystem::HasGPUEmitter() const
 	return false;
 }
 
+UFXSystemComponent::UFXSystemComponent(const FObjectInitializer& ObjectInitializer)
+	: Super(ObjectInitializer)
+{}
+
 FOnSystemPreActivationChange UParticleSystemComponent::OnSystemPreActivationChange;
 
 UParticleSystemComponent::UParticleSystemComponent(const FObjectInitializer& ObjectInitializer)
