Fix an ID Mismatch for Fast Array Delta Serialization when Array Items are Net Serialize

diff --git a/Engine/Source/Runtime/Engine/Private/RepLayout.cpp b/Engine/Source/Runtime/Engine/Private/RepLayout.cpp
--- a/Engine/Source/Runtime/Engine/Private/RepLayout.cpp
+++ b/Engine/Source/Runtime/Engine/Private/RepLayout.cpp
@@ -6309,8 +6309,8 @@ bool FRepLayout::DeltaSerializeFastArrayProperty(FFastArrayDeltaSerializeParams&
 								//			some assumptions laid out in the algorithm above.
 								int32 AppendedShadowItems = 0;
 
-								UE_LOG(LogRep, VeryVerbose, TEXT("DeltaSerializeFastArrayProperty: Fixup Shadow State. Owner=%s, Property=%s, bInitial=%d, ObjecyArrayNum=%d, ShadowArrayNum=%d"),
-									*Owner->GetName(), *Parent.CachedPropertyName.ToString(), !!bIsInitial, ObjectArrayNum, ShadowArrayHelper.Num());
+								UE_LOG(LogRepProperties, VeryVerbose, TEXT("DeltaSerializeFastArrayProperty: Fixup Shadow State. Owner=%s, Object=%s, Property=%s, bInitial=%d, ObjectArrayNum=%d, ShadowArrayNum=%d"),
+									*Owner->GetName(), *Object->GetPathName(), *Parent.CachedPropertyName.ToString(), !!bIsInitial, ObjectArrayNum, ShadowArrayHelper.Num());
 
 								for (int32 Index = 0; Index < ObjectArrayNum && Index < ShadowArrayNum; ++Index)
 								{
@@ -6319,7 +6319,7 @@ bool FRepLayout::DeltaSerializeFastArrayProperty(FFastArrayDeltaSerializeParams&
 
 									FastArrayState.IDToIndexMap.Emplace(ObjectReplicationID, Index);
 
-									UE_LOG(LogRep, VeryVerbose, TEXT("DeltaSerializeFastArrayProperty: Handling Item. ID=%d, Index=%d, ShadowID=%d"), ObjectReplicationID, Index, ShadowReplicationID);
+									UE_LOG(LogRepProperties, VeryVerbose, TEXT("DeltaSerializeFastArrayProperty: Handling Item. ID=%d, Index=%d, ShadowID=%d"), ObjectReplicationID, Index, ShadowReplicationID);
 
 									// If our IDs match, there's nothing to do.
 									if (ObjectReplicationID != ShadowReplicationID)
@@ -6365,6 +6365,9 @@ bool FRepLayout::DeltaSerializeFastArrayProperty(FFastArrayDeltaSerializeParams&
 							// Deleted elements will have been chopped off by the resize.
 							if (bIsInitial || (ShadowArrayNum < ObjectArrayNum))
 							{
+								UE_CLOG(bIsInitial, LogRepProperties, VeryVerbose, TEXT("DeltaSerializeFastArrayProperty: Adding initial properties. Owner=%s, Object=%s, Property=%s, bInitial=%d, ObjectArrayNum=%d, ShadowArrayNum=%d"),
+									*Owner->GetName(), *Object->GetPathName(), *Parent.CachedPropertyName.ToString(), !!bIsInitial, ObjectArrayNum, ShadowArrayHelper.Num());
+
 								const int32 StartIndex = bIsInitial ? 0 : ShadowArrayNum;
 								for (int32 Index = StartIndex; Index < ObjectArrayNum; ++Index)
 								{
@@ -6407,6 +6410,13 @@ bool FRepLayout::DeltaSerializeFastArrayProperty(FFastArrayDeltaSerializeParams&
 							{
 								NewChangelist.Add(0);
 								HistoryItem.ChangelistByID.Emplace(IDIndexPair.ID, MoveTemp(NewChangelist));
+
+								// If our FastArraySerializerItems are NetSerialize, then their ID may be reset to INDEX_NONE due
+								// to copying them into the shadow state (see FFastArraySerializerItem::operator=).
+								// In that case, we need to make make sure we reset our ID so they can be found the next
+								// time we try to replicate them.
+								int32& ShadowReplicationID = CustomDeltaProperty.GetFastArrayItemReplicationIDMutable(ElementShadowData.Data);
+								ShadowReplicationID = IDIndexPair.ID;
 							}
 						}
 					}
