Add setting to allow editor to create separate audio device for PIE #jira none #rb Aaron.McLeran

diff --git a/Engine/Source/Editor/UnrealEd/Classes/Settings/LevelEditorMiscSettings.h b/Engine/Source/Editor/UnrealEd/Classes/Settings/LevelEditorMiscSettings.h
--- a/Engine/Source/Editor/UnrealEd/Classes/Settings/LevelEditorMiscSettings.h
+++ b/Engine/Source/Editor/UnrealEd/Classes/Settings/LevelEditorMiscSettings.h
@@ -48,6 +48,10 @@ public:
 	UPROPERTY(EditAnywhere, config, Category=Sound)
 	uint32 bAllowBackgroundAudio:1;
 
+	/** If checked, will create a new audio engine instance for the play-in-editor window. Otherwise, will re-use the audio device used in content browser. */
+	UPROPERTY(EditAnywhere, config, Category = Sound)
+	uint32 bCreateNewAudioDeviceForPlayInEditor: 1;
+
 	/** If true audio will be enabled in the editor. Does not affect PIE **/
 	UPROPERTY(config)
 	uint32 bEnableRealTimeAudio:1;
diff --git a/Engine/Source/Editor/UnrealEd/Private/Settings/SettingsClasses.cpp b/Engine/Source/Editor/UnrealEd/Private/Settings/SettingsClasses.cpp
--- a/Engine/Source/Editor/UnrealEd/Private/Settings/SettingsClasses.cpp
+++ b/Engine/Source/Editor/UnrealEd/Private/Settings/SettingsClasses.cpp
@@ -348,9 +348,9 @@ ULevelEditorMiscSettings::ULevelEditorMiscSettings( const FObjectInitializer& Ob
 	bPromptWhenAddingToLevelBeforeCheckout = true;
 	bPromptWhenAddingToLevelOutsideBounds = true;
 	PercentageThresholdForPrompt = 20.0f;
+	bCreateNewAudioDeviceForPlayInEditor = true;
 }
 
-
 void ULevelEditorMiscSettings::PostEditChangeProperty( struct FPropertyChangedEvent& PropertyChangedEvent )
 {
 	Super::PostEditChangeProperty(PropertyChangedEvent);
diff --git a/Engine/Source/Runtime/Engine/Private/AudioDeviceManager.cpp b/Engine/Source/Runtime/Engine/Private/AudioDeviceManager.cpp
--- a/Engine/Source/Runtime/Engine/Private/AudioDeviceManager.cpp
+++ b/Engine/Source/Runtime/Engine/Private/AudioDeviceManager.cpp
@@ -8,6 +8,7 @@
 #include "Misc/CommandLine.h"
 #include "Misc/ConfigCacheIni.h"
 #include "UObject/UObjectIterator.h"
+#include "Settings/LevelEditorMiscSettings.h"
 
 #if WITH_EDITOR
 #include "AudioEditorModule.h"
@@ -380,7 +381,7 @@ bool FAudioDeviceManager::CreateAudioDevice(bool bCreateNewDevice, FCreateAudioD
 	bool bRequiresInit = true;
 
 	// For the first PIE window, we'll just use the main audio device
-	if (NumActiveAudioDevices == 1)
+	if (NumActiveAudioDevices == 1 && !GetDefault<ULevelEditorMiscSettings>()->bCreateNewAudioDeviceForPlayInEditor)
 	{
 		FAudioDevice* MainAudioDevice = GEngine->GetMainAudioDevice();
 		if (MainAudioDevice)
