Removed incorrect early out when compiling the ShadingModel material expression. Unlit is a valid shading model for the compiler to process.

#rnx

diff --git a/Engine/Source/Runtime/Engine/Private/Materials/HLSLMaterialTranslator.h b/Engine/Source/Runtime/Engine/Private/Materials/HLSLMaterialTranslator.h
--- a/Engine/Source/Runtime/Engine/Private/Materials/HLSLMaterialTranslator.h
+++ b/Engine/Source/Runtime/Engine/Private/Materials/HLSLMaterialTranslator.h
@@ -5491,14 +5491,8 @@ protected:
 		return GetPrimitiveProperty(MCT_Float, TEXT("CustomPrimitiveData"), *HlslName);
 	}
 
-virtual int32 ShadingModel(EMaterialShadingModel InSelectedShadingModel) override
+	virtual int32 ShadingModel(EMaterialShadingModel InSelectedShadingModel) override
 	{
-		// Don't allow unlit shading model, since it has to be mutually exclusive with the other shading models due to #ifdef logic in BasePassPixelShader.usf
-		if (InSelectedShadingModel == MSM_Unlit)
-		{
-			return INDEX_NONE;
-		}
-
 		return AddInlinedCodeChunk(MCT_ShadingModel, TEXT("%d"), InSelectedShadingModel);
 	}
 
