- Simplify max channel count code by removing GEngineIni look-up for value set in Engine-level ini and just directly setting in UAudioSettings constructor
- Fix bug with IOS & XAudio2 old backends not respecting platform settings code/ini settings
- Move AudioDevice to override quality settings MaxChannels by ANY platform setting, on init, not just the minimal value.  This is okay on init as the source array has yet to be initialized.

https://p4-swarm.epicgames.net/reviews/6260062

diff --git a/Engine/Config/BaseEngine.ini b/Engine/Config/BaseEngine.ini
--- a/Engine/Config/BaseEngine.ini
+++ b/Engine/Config/BaseEngine.ini
@@ -806,7 +806,6 @@ gc.UseDisregardForGCOnDedicatedServers=False
 +LocalizationPaths=../../../Engine/Content/Localization/Engine
 
 [Audio]
-MaxChannels=32
 CommonAudioPoolSize=0
 UnfocusedVolumeMultiplier=0.0
 UseAudioThread=true
diff --git a/Engine/Source/Runtime/Android/AndroidAudio/Private/AndroidAudioDevice.cpp b/Engine/Source/Runtime/Android/AndroidAudio/Private/AndroidAudioDevice.cpp
--- a/Engine/Source/Runtime/Android/AndroidAudio/Private/AndroidAudioDevice.cpp
+++ b/Engine/Source/Runtime/Android/AndroidAudio/Private/AndroidAudioDevice.cpp
@@ -193,6 +193,11 @@ class ICompressedAudioInfo* FSLESAudioDevice::CreateCompressedAudioInfo(USoundWa
 	return nullptr;
 }
 
+FAudioPlatformSettings FSLESAudioDevice::GetPlatformSettings() const
+{
+	return FAudioPlatformSettings::GetPlatformSettings(TEXT("/Script/AndroidRuntimeSettings.AndroidRuntimeSettings"));
+}
+
 /** Check if any background music or sound is playing through the audio device */
 bool FSLESAudioDevice::IsExernalBackgroundSoundActive()
 {
diff --git a/Engine/Source/Runtime/Android/AndroidAudio/Public/AndroidAudioDevice.h b/Engine/Source/Runtime/Android/AndroidAudio/Public/AndroidAudioDevice.h
--- a/Engine/Source/Runtime/Android/AndroidAudio/Public/AndroidAudioDevice.h
+++ b/Engine/Source/Runtime/Android/AndroidAudio/Public/AndroidAudioDevice.h
@@ -316,6 +316,8 @@ public:
 	/** Starts up any platform specific hardware/APIs */
 	virtual bool InitializeHardware() override;
 
+	virtual FAudioPlatformSettings GetPlatformSettings() const override;
+
 	virtual void TeardownHardware() override
 	{
 		Teardown();
diff --git a/Engine/Source/Runtime/Engine/Classes/Sound/AudioSettings.h b/Engine/Source/Runtime/Engine/Classes/Sound/AudioSettings.h
--- a/Engine/Source/Runtime/Engine/Classes/Sound/AudioSettings.h
+++ b/Engine/Source/Runtime/Engine/Classes/Sound/AudioSettings.h
@@ -31,7 +31,7 @@ struct ENGINE_API FAudioPlatformSettings
 		: SampleRate(48000)
 		, CallbackBufferFrameSize(1024)
 		, NumBuffers(2)
-		, MaxChannels(0)
+		, MaxChannels(32)
 		, NumSourceWorkers(0)
 	{
 	}
@@ -63,7 +63,7 @@ enum class EMonoChannelUpmixMethod : int8
 {
 	// The mono channel is split 0.5 left/right
 	Linear,
-	
+
 	// The mono channel is split 0.707 left/right
 	EqualPower,
 
@@ -118,7 +118,7 @@ class ENGINE_API UAudioSettings : public UDeveloperSettings
 	/** The SoundMix to use as base when no other system has speciicefied a Base SoundMix */
 	UPROPERTY(config, EditAnywhere, Category="Audio", meta=(AllowedClasses="SoundMix"))
 	FSoftObjectPath DefaultBaseSoundMix;
-	
+
 	/** Sound class to be used for the VOIP audio component */
 	UPROPERTY(config, EditAnywhere, Category="Audio", meta=(AllowedClasses="SoundClass", DisplayName = "VOIP Sound Class"))
 	FSoftObjectPath VoiPSoundClass;
@@ -170,9 +170,9 @@ class ENGINE_API UAudioSettings : public UDeveloperSettings
 	UPROPERTY(config, EditAnywhere, Category = "Quality", AdvancedDisplay)
 	uint32 MaxWaveInstances;
 
-	/** 
+	/**
 	 * The max number of sources to reserve for "stopping" sounds. A "stopping" sound applies a fast fade in the DSP
-	 * render to prevent discontinuities when stopping sources.  
+	 * render to prevent discontinuities when stopping sources.
 	 */
 	UPROPERTY(config, EditAnywhere, Category = "Quality", AdvancedDisplay)
 	uint32 NumStoppingSources;
@@ -203,7 +203,7 @@ class ENGINE_API UAudioSettings : public UDeveloperSettings
 
 	const FAudioQualitySettings& GetQualityLevelSettings(int32 QualityLevel) const;
 
-	// Sets whether audio mixer is enabled. Set once an audio mixer platform modu le is loaded.
+	// Sets whether audio mixer is enabled. Set once an audio mixer platform module is loaded.
 	void SetAudioMixerEnabled(const bool bInAudioMixerEnabled);
 
 	// Returns if the audio mixer is currently enabled
diff --git a/Engine/Source/Runtime/Engine/Private/AudioDevice.cpp b/Engine/Source/Runtime/Engine/Private/AudioDevice.cpp
--- a/Engine/Source/Runtime/Engine/Private/AudioDevice.cpp
+++ b/Engine/Source/Runtime/Engine/Private/AudioDevice.cpp
@@ -234,19 +234,22 @@ bool FAudioDevice::Init(int32 InMaxChannels)
 
 	PluginListeners.Reset();
 
-	// initialize max channels taking into account platform configurations
-	// Get a copy of the platform-specific settings (overriden by platforms)
+	// Initialize MaxChannels taking into account platform configurations
+	// Get a copy of the platform-specific settings (overridden by platforms)
 	PlatformSettings = GetPlatformSettings();
 
-	// MaxChannels is the min of the platform-specific value and the max value in the quality settings (InMaxChannels)
-	MaxChannels = PlatformSettings.MaxChannels > 0 ? FMath::Min(PlatformSettings.MaxChannels, InMaxChannels) : InMaxChannels;
+	// Max channels is the max value supplied to Init call (quality settings), unless overwritten by the platform settings.
+	// This does not have to be the minimum value in this case (nor is it desired, so platforms can potentially scale up)
+	// as the Sources array has yet to be initialized.
+	MaxChannels = PlatformSettings.MaxChannels > 0 ? PlatformSettings.MaxChannels : InMaxChannels;
 	MaxChannels_GameThread = MaxChannels;
 
+	// Ensure and not assert so if in editor, user can change quality setting and re-serialize if so desired.
+	ensureMsgf(MaxChannels > 0, TEXT("Neither passed nor platform MaxChannel setting was positive value"));
+
 	// Mixed sample rate is set by the platform
 	SampleRate = PlatformSettings.SampleRate;
 
-	check(MaxChannels != 0);
-
 	verify(GConfig->GetInt(TEXT("Audio"), TEXT("CommonAudioPoolSize"), CommonAudioPoolSize, GEngineIni));
 
 	// If this is true, skip the initial startup precache so we can do it later in the flow
@@ -375,7 +378,7 @@ bool FAudioDevice::Init(int32 InMaxChannels)
 
 	// create a platform specific effects manager
 	// if this is the audio mixer, we initialized the effects manager before the hardware
- 	if (!IsAudioMixerEnabled())
+	if (!IsAudioMixerEnabled())
 	{
 		Effects = CreateEffectsManager();
 	}
@@ -422,6 +425,12 @@ void FAudioDevice::PrecacheStartupSounds()
 
 void FAudioDevice::SetMaxChannels(int32 InMaxChannels)
 {
+	if (InMaxChannels <= 0)
+	{
+		UE_LOG(LogAudio, Warning, TEXT("MaxChannels must be set to a positive value."));
+		return;
+	}
+
 	if (!IsInAudioThread())
 	{
 		MaxChannels_GameThread = InMaxChannels;
@@ -437,9 +446,9 @@ void FAudioDevice::SetMaxChannels(int32 InMaxChannels)
 		return;
 	}
 
-	if (InMaxChannels > Sources.Num())
+	if (InMaxChannels > MaxChannels)
 	{
-		UE_LOG(LogAudio, Warning, TEXT("Can't increase channels past starting number!"));
+		UE_LOG(LogAudio, Warning, TEXT("Can't increase MaxChannels past existing setting!"));
 		return;
 	}
 
@@ -480,7 +489,7 @@ int32 FAudioDevice::GetMaxChannels() const
 {
 	if (IsInAudioThread())
 	{
-		if (AudioChannelCountCVar > 0 && AudioChannelCountCVar < Sources.Num())
+		if (AudioChannelCountCVar > 0 && AudioChannelCountCVar < MaxChannels)
 		{
 			return FMath::Max(int32(AudioChannelCountCVar * MaxChannelsScale * AudioChannelCountScaleCVar), 1);
 		}
@@ -489,7 +498,7 @@ int32 FAudioDevice::GetMaxChannels() const
 	}
 	else
 	{
-		if (AudioChannelCountCVar > 0 && AudioChannelCountCVar < Sources.Num())
+		if (AudioChannelCountCVar > 0 && AudioChannelCountCVar < MaxChannels)
 		{
 			return FMath::Max(int32(AudioChannelCountCVar * MaxChannelsScale_GameThread * AudioChannelCountScaleCVar), 1);
 		}
diff --git a/Engine/Source/Runtime/Engine/Private/AudioSettings.cpp b/Engine/Source/Runtime/Engine/Private/AudioSettings.cpp
--- a/Engine/Source/Runtime/Engine/Private/AudioSettings.cpp
+++ b/Engine/Source/Runtime/Engine/Private/AudioSettings.cpp
@@ -65,7 +65,6 @@ void UAudioSettings::AddDefaultSettings()
 {
 	FAudioQualitySettings DefaultSettings;
 	DefaultSettings.DisplayName = LOCTEXT("DefaultSettingsName", "Default");
-	GConfig->GetInt(TEXT("Audio"), TEXT("MaxChannels"), DefaultSettings.MaxChannels, GEngineIni); // for backwards compatibility
 	QualityLevels.Add(DefaultSettings);
 	bAllowVirtualizedSounds = true;
 	DefaultReverbSendLevel = 0.0f;
diff --git a/Engine/Source/Runtime/IOS/IOSAudio/Private/IOSAudioDevice.cpp b/Engine/Source/Runtime/IOS/IOSAudio/Private/IOSAudioDevice.cpp
--- a/Engine/Source/Runtime/IOS/IOSAudio/Private/IOSAudioDevice.cpp
+++ b/Engine/Source/Runtime/IOS/IOSAudio/Private/IOSAudioDevice.cpp
@@ -107,6 +107,11 @@ bool FIOSAudioDevice::Exec(UWorld* InWorld, const TCHAR* Cmd, FOutputDevice& Ar)
 	return FAudioDevice::Exec(InWorld, Cmd, Ar);
 }
 
+FAudioPlatformSettings FIOSAudioDevice::GetPlatformSettings() const
+{
+	return FAudioPlatformSettings::GetPlatformSettings(TEXT("/Script/IOSRuntimeSettings.IOSRuntimeSettings"));
+}
+
 bool FIOSAudioDevice::InitializeHardware()
 {
 	SIZE_T SampleSize = sizeof(AudioSampleType);
diff --git a/Engine/Source/Runtime/IOS/IOSAudio/Public/IOSAudioDevice.h b/Engine/Source/Runtime/IOS/IOSAudio/Public/IOSAudioDevice.h
--- a/Engine/Source/Runtime/IOS/IOSAudio/Public/IOSAudioDevice.h
+++ b/Engine/Source/Runtime/IOS/IOSAudio/Public/IOSAudioDevice.h
@@ -237,7 +237,10 @@ protected:
 
 	/** Check if any background music or sound is playing through the audio device */
 	virtual bool IsExernalBackgroundSoundActive() override;
-	
+
+	/** Retrieves the platform settings */
+	virtual FAudioPlatformSettings GetPlatformSettings() const override;
+
 private:
 
 	// Used to support adpcm streaming
diff --git a/Engine/Source/Runtime/Linux/AudioMixerSDL/Private/AudioMixerPlatformSDL.cpp b/Engine/Source/Runtime/Linux/AudioMixerSDL/Private/AudioMixerPlatformSDL.cpp
--- a/Engine/Source/Runtime/Linux/AudioMixerSDL/Private/AudioMixerPlatformSDL.cpp
+++ b/Engine/Source/Runtime/Linux/AudioMixerSDL/Private/AudioMixerPlatformSDL.cpp
@@ -388,12 +388,7 @@ namespace Audio
 		return FAudioPlatformSettings::GetPlatformSettings(TEXT("/Script/LinuxTargetPlatform.LinuxTargetSettings"));
 #else
 		// On HTML5 and Windows, use default parameters.
-		FAudioPlatformSettings Settings;
-		Settings.SampleRate = 48000;
-		Settings.MaxChannels = 0;
-		Settings.NumBuffers = 2;
-		Settings.CallbackBufferFrameSize = 1024;
-		return Settings;
+		return FAudioPlatformSettings();
 #endif
 	}
 }
diff --git a/Engine/Source/Runtime/Windows/XAudio2/Private/XAudio2Device.cpp b/Engine/Source/Runtime/Windows/XAudio2/Private/XAudio2Device.cpp
--- a/Engine/Source/Runtime/Windows/XAudio2/Private/XAudio2Device.cpp
+++ b/Engine/Source/Runtime/Windows/XAudio2/Private/XAudio2Device.cpp
@@ -732,3 +732,8 @@ void* FXAudio2Device::AllocatePermanentMemory( int32 Size, bool& AllocatedInPool
 	
 	return( Allocation );
 }
+
+FAudioPlatformSettings FXAudio2Device::GetPlatformSettings() const
+{
+	return FAudioPlatformSettings::GetPlatformSettings(TEXT("/Script/WindowsTargetPlatform.WindowsTargetSettings"));
+}
diff --git a/Engine/Source/Runtime/Windows/XAudio2/Public/XAudio2Device.h b/Engine/Source/Runtime/Windows/XAudio2/Public/XAudio2Device.h
--- a/Engine/Source/Runtime/Windows/XAudio2/Public/XAudio2Device.h
+++ b/Engine/Source/Runtime/Windows/XAudio2/Public/XAudio2Device.h
@@ -200,6 +200,8 @@ protected:
 	 */
 	void* AllocatePermanentMemory( int32 Size, /*OUT*/ bool& AllocatedInPool );
 
+	/** Retrieves the platform settings */
+	virtual FAudioPlatformSettings GetPlatformSettings() const override;
 
 	/** 
      * Derives the output matrix to use based on the channel mask and the number of channels
