Fix for crash when manipulating sections in the static mesh editor. When the EditableMesh is controlling the vertex buffers, there are occasions when the render resources will have been released by UStaticMesh::PreEditChange.  In this case, Build() must reinitialize them, whether or not the EditableMesh is in control.

diff --git a/Engine/Source/Runtime/Engine/Private/StaticMeshBuild.cpp b/Engine/Source/Runtime/Engine/Private/StaticMeshBuild.cpp
--- a/Engine/Source/Runtime/Engine/Private/StaticMeshBuild.cpp
+++ b/Engine/Source/Runtime/Engine/Private/StaticMeshBuild.cpp
@@ -76,13 +76,6 @@ void UStaticMesh::Build(bool bSilent, TArray<FText>* OutErrors)
 	if (IsTemplate())
 		return;
 
-	// If we're controlled by an editable mesh do not build. The editable mesh will build us
-	if (EditableMesh)
-	{
-		return;
-	}
-
-
 	if (SourceModels.Num() <= 0)
 	{
 		UE_LOG(LogStaticMesh,Warning,TEXT("Static mesh has no source models: %s"),*GetPathName());
@@ -115,6 +108,16 @@ void UStaticMesh::Build(bool bSilent, TArray<FText>* OutErrors)
 	// allocated, and potentially accessing the UStaticMesh.
 	ReleaseResourcesFence.Wait();
 
+	// If we're controlled by an editable mesh do not build. The editable mesh will build us
+	if (EditableMesh)
+	{
+		if (FApp::CanEverRender())
+		{
+			InitResources();
+		}
+		return;
+	}
+
 	// Remember the derived data key of our current render data if any.
 	FString ExistingDerivedDataKey = RenderData ? RenderData->DerivedDataKey : TEXT("");
 
