#ue4 - CharacterMovement: Deferred limited air control code in PhysFalling() until we actually need the result, which is only after a collision. This reduces the number of calls to CalcVelocity() and NewFallVelocity() from 2 to 1 in normal cases.

Added virtual ShouldLimitAirControl() for deciding if we want to limit air control at all. This is useful for falling modes which are more like flying (low gravity or anti-gravity).

diff --git a/Engine/Source/Runtime/Engine/Classes/GameFramework/CharacterMovementComponent.h b/Engine/Source/Runtime/Engine/Classes/GameFramework/CharacterMovementComponent.h
--- a/Engine/Source/Runtime/Engine/Classes/GameFramework/CharacterMovementComponent.h
+++ b/Engine/Source/Runtime/Engine/Classes/GameFramework/CharacterMovementComponent.h
@@ -1529,6 +1529,14 @@ public:
 	 */
 	virtual FVector GetFallingLateralAcceleration(float DeltaTime);
 	
+	/**
+	 * Returns true if falling movement should limit air control. Limiting air control prevents input acceleration during falling movement
+	 * from allowing velocity to redirect forces upwards while falling, which could result in slower falling or even upward boosting.
+	 *
+	 * @see GetFallingLateralAcceleration(), BoostAirControl(), GetAirControl(), LimitAirControl()
+	 */
+	virtual bool ShouldLimitAirControl(float DeltaTime, const FVector& FallAcceleration) const;
+
 	/**
 	 * Get the air control to use during falling movement.
 	 * Given an initial air control (TickAirControl), applies the result of BoostAirControl().
diff --git a/Engine/Source/Runtime/Engine/Private/Components/CharacterMovementComponent.cpp b/Engine/Source/Runtime/Engine/Private/Components/CharacterMovementComponent.cpp
--- a/Engine/Source/Runtime/Engine/Private/Components/CharacterMovementComponent.cpp
+++ b/Engine/Source/Runtime/Engine/Private/Components/CharacterMovementComponent.cpp
@@ -4089,6 +4089,11 @@ FVector UCharacterMovementComponent::GetFallingLateralAcceleration(float DeltaTi
 }
 
 
+bool UCharacterMovementComponent::ShouldLimitAirControl(float DeltaTime, const FVector& FallAcceleration) const
+{
+	return (FallAcceleration.SizeSquared2D() > 0.f);
+}
+
 FVector UCharacterMovementComponent::GetAirControl(float DeltaTime, float TickAirControl, const FVector& FallAcceleration)
 {
 	// Boost
@@ -4124,7 +4129,7 @@ void UCharacterMovementComponent::PhysFalling(float deltaTime, int32 Iterations)
 
 	FVector FallAcceleration = GetFallingLateralAcceleration(deltaTime);
 	FallAcceleration.Z = 0.f;
-	const bool bHasAirControl = (FallAcceleration.SizeSquared2D() > 0.f);
+	const bool bHasLimitedAirControl = ShouldLimitAirControl(deltaTime, FallAcceleration);
 
 	float remainingTime = deltaTime;
 	while( (remainingTime >= MIN_TICK_TIME) && (Iterations < MaxSimulationIterations) )
@@ -4139,24 +4144,12 @@ void UCharacterMovementComponent::PhysFalling(float deltaTime, int32 Iterations)
 
 		RestorePreAdditiveRootMotionVelocity();
 
-		FVector OldVelocity = Velocity;
-		FVector VelocityNoAirControl = Velocity;
+		const FVector OldVelocity = Velocity;
 
 		// Apply input
+		const float MaxDecel = GetMaxBrakingDeceleration();
 		if (!HasAnimRootMotion() && !CurrentRootMotion.HasOverrideVelocity())
 		{
-			const float MaxDecel = GetMaxBrakingDeceleration();
-			// Compute VelocityNoAirControl
-			if (bHasAirControl)
-			{
-				// Find velocity *without* acceleration.
-				TGuardValue<FVector> RestoreAcceleration(Acceleration, FVector::ZeroVector);
-				TGuardValue<FVector> RestoreVelocity(Velocity, Velocity);
-				Velocity.Z = 0.f;
-				CalcVelocity(timeTick, FallingLateralFriction, false, MaxDecel);
-				VelocityNoAirControl = FVector(Velocity.X, Velocity.Y, OldVelocity.Z);
-			}
-
 			// Compute Velocity
 			{
 				// Acceleration = FallAcceleration for CalcVelocity(), but we restore it after using it.
@@ -4165,12 +4158,6 @@ void UCharacterMovementComponent::PhysFalling(float deltaTime, int32 Iterations)
 				CalcVelocity(timeTick, FallingLateralFriction, false, MaxDecel);
 				Velocity.Z = OldVelocity.Z;
 			}
-
-			// Just copy Velocity to VelocityNoAirControl if they are the same (ie no acceleration).
-			if (!bHasAirControl)
-			{
-				VelocityNoAirControl = Velocity;
-			}
 		}
 
 		// Compute current gravity
@@ -4196,8 +4183,6 @@ void UCharacterMovementComponent::PhysFalling(float deltaTime, int32 Iterations)
 
 		// Apply gravity
 		Velocity = NewFallVelocity(Velocity, Gravity, GravityTime);
-		VelocityNoAirControl = bHasAirControl ? NewFallVelocity(VelocityNoAirControl, Gravity, GravityTime) : Velocity;
-		const FVector AirControlAccel = (Velocity - VelocityNoAirControl) / timeTick;
 
 		// See if we need to sub-step to exactly reach the apex. This is important for avoiding "cutting off the top" of the trajectory as framerate varies.
 		if (CharacterMovementCVars::ForceJumpPeakSubstep && OldVelocity.Z > 0.f && Velocity.Z <= 0.f && NumJumpApexAttempts < MaxJumpApexAttemptsPerSimulation)
@@ -4302,9 +4287,23 @@ void UCharacterMovementComponent::PhysFalling(float deltaTime, int32 Iterations)
 
 				// Limit air control based on what we hit.
 				// We moved to the impact point using air control, but may want to deflect from there based on a limited air control acceleration.
-				if (bHasAirControl)
+				FVector VelocityNoAirControl = OldVelocity;
+				FVector AirControlAccel = Acceleration;
+				if (bHasLimitedAirControl)
 				{
+					// Compute VelocityNoAirControl
+					{
+						// Find velocity *without* acceleration.
+						TGuardValue<FVector> RestoreAcceleration(Acceleration, FVector::ZeroVector);
+						TGuardValue<FVector> RestoreVelocity(Velocity, OldVelocity);
+						Velocity.Z = 0.f;
+						CalcVelocity(timeTick, FallingLateralFriction, false, MaxDecel);
+						VelocityNoAirControl = FVector(Velocity.X, Velocity.Y, OldVelocity.Z);
+						VelocityNoAirControl = NewFallVelocity(VelocityNoAirControl, Gravity, GravityTime);
+					}
+
 					const bool bCheckLandingSpot = false; // we already checked above.
+					AirControlAccel = (Velocity - VelocityNoAirControl) / timeTick;
 					const FVector AirControlDeltaV = LimitAirControl(LastMoveTimeSlice, AirControlAccel, Hit, bCheckLandingSpot) * LastMoveTimeSlice;
 					Adjusted = (VelocityNoAirControl + AirControlDeltaV) * LastMoveTimeSlice;
 				}
@@ -4347,7 +4346,7 @@ void UCharacterMovementComponent::PhysFalling(float deltaTime, int32 Iterations)
 						}
 
 						// Act as if there was no air control on the last move when computing new deflection.
-						if (bHasAirControl && Hit.Normal.Z > VERTICAL_SLOPE_NORMAL_Z)
+						if (bHasLimitedAirControl && Hit.Normal.Z > VERTICAL_SLOPE_NORMAL_Z)
 						{
 							const FVector LastMoveNoAirControl = VelocityNoAirControl * LastMoveTimeSlice;
 							Delta = ComputeSlideVector(LastMoveNoAirControl, 1.f, OldHitNormal, Hit);
@@ -4357,7 +4356,7 @@ void UCharacterMovementComponent::PhysFalling(float deltaTime, int32 Iterations)
 						TwoWallAdjust(Delta, Hit, OldHitNormal);
 
 						// Limit air control, but allow a slide along the second wall.
-						if (bHasAirControl)
+						if (bHasLimitedAirControl)
 						{
 							const bool bCheckLandingSpot = false; // we already checked above.
 							const FVector AirControlDeltaV = LimitAirControl(subTimeTickRemaining, AirControlAccel, Hit, bCheckLandingSpot) * subTimeTickRemaining;
