#UE4 Fix crash that can happen when components that were previously default subobjects are async loaded in editor builds. They load as default subobjects but are detected to not be default subobjects so they are marked pending kill and the linker loading them asserts after they are GC'd

diff --git a/Engine/Source/Editor/UnrealEd/Private/Kismet2/Kismet2.cpp b/Engine/Source/Editor/UnrealEd/Private/Kismet2/Kismet2.cpp
--- a/Engine/Source/Editor/UnrealEd/Private/Kismet2/Kismet2.cpp
+++ b/Engine/Source/Editor/UnrealEd/Private/Kismet2/Kismet2.cpp
@@ -1114,6 +1114,12 @@ static void ConformComponentsUtils::ConformRemovedNativeComponents(UObject* BpCd
 		// else, the component has been removed from our native super class
 
 		Component->DestroyComponent(/*bPromoteChildren =*/false);
+		if (Component->HasAnyInternalFlags(EInternalObjectFlags::AsyncLoading))
+		{
+			// Async loading components cannot be pending kill, or the async loading code will assert when trying to postload them.
+			Component->ClearPendingKill();
+			FLinkerLoad::InvalidateExport(Component);
+		}
 		DestroyedComponents.Add(Component);
 
 		// The DestroyComponent() call above will clear the RootComponent value in this case.
diff --git a/Engine/Source/Runtime/Engine/Classes/Components/ActorComponent.h b/Engine/Source/Runtime/Engine/Classes/Components/ActorComponent.h
--- a/Engine/Source/Runtime/Engine/Classes/Components/ActorComponent.h
+++ b/Engine/Source/Runtime/Engine/Classes/Components/ActorComponent.h
@@ -238,6 +238,9 @@ private:
 #if WITH_EDITOR
 	/** During undo/redo it isn't safe to cache owner */
 	uint8 bCanUseCachedOwner:1;
+
+	/** Invalidates this component as an export once the linker is set. Used to clean up old native default subobjects that were removed from code */
+	uint8 bInvalidateExportWhenLinkerIsSet : 1;
 #endif
 
 	/** True if this component was owned by a net startup actor during level load. */
@@ -789,6 +792,7 @@ public:
 	virtual void PostRename(UObject* OldOuter, const FName OldName) override;
 	virtual void Serialize(FArchive& Ar) override;
 #if WITH_EDITOR
+	virtual void PostLinkerChange() override;
 	virtual bool Modify( bool bAlwaysMarkDirty = true ) override;
 	virtual void PreEditChange(UProperty* PropertyThatWillChange) override;
 	virtual void PostEditChangeProperty(FPropertyChangedEvent& PropertyChangedEvent) override;
diff --git a/Engine/Source/Runtime/Engine/Private/Components/ActorComponent.cpp b/Engine/Source/Runtime/Engine/Private/Components/ActorComponent.cpp
--- a/Engine/Source/Runtime/Engine/Private/Components/ActorComponent.cpp
+++ b/Engine/Source/Runtime/Engine/Private/Components/ActorComponent.cpp
@@ -192,8 +192,19 @@ void UActorComponent::PostInitProperties()
 			{
 				// else: this is a natively created component that thinks its archetype is the CDO of
 				// this class, rather than a template component and this isn't the template component.
-				// Delete this stale component:
-				MarkPendingKill();
+				// Delete this stale component
+#if WITH_EDITOR
+				if (HasAnyInternalFlags(EInternalObjectFlags::AsyncLoading))
+				{
+					// Async loading components cannot be pending kill, or the async loading code will assert when trying to postload them.
+					// Instead, wait until the linker is set and invalidate the export so it does not proceed to load
+					bInvalidateExportWhenLinkerIsSet = true;
+				}
+				else
+#endif // WITH_EDITOR
+				{
+					MarkPendingKill();
+				}
 			}
 		}
 		else
@@ -541,6 +552,17 @@ bool UActorComponent::CallRemoteFunction( UFunction* Function, void* Parameters,
 /** FComponentReregisterContexts for components which have had PreEditChange called but not PostEditChange. */
 static TMap<TWeakObjectPtr<UActorComponent>,FComponentReregisterContext*> EditReregisterContexts;
 
+void UActorComponent::PostLinkerChange()
+{
+	Super::PostLinkerChange();
+
+	if (bInvalidateExportWhenLinkerIsSet)
+	{
+		FLinkerLoad::InvalidateExport(this);
+		bInvalidateExportWhenLinkerIsSet = false;
+	}
+}
+
 bool UActorComponent::Modify( bool bAlwaysMarkDirty/*=true*/ )
 {
 	// If this is a construction script component we don't store them in the transaction buffer.  Instead, mark
