Add setting to allow editor to create separate audio device for PIE #jira none

diff --git a/Engine/Source/Editor/UnrealEd/Classes/Settings/LevelEditorMiscSettings.h b/Engine/Source/Editor/UnrealEd/Classes/Settings/LevelEditorMiscSettings.h
--- a/Engine/Source/Editor/UnrealEd/Classes/Settings/LevelEditorMiscSettings.h
+++ b/Engine/Source/Editor/UnrealEd/Classes/Settings/LevelEditorMiscSettings.h
@@ -48,6 +48,10 @@ public:
 	UPROPERTY(EditAnywhere, config, Category=Sound)
 	uint32 bAllowBackgroundAudio:1;
 
+	/** If checked, will create a new audio engine instance for the play-in-editor window. Otherwise, will re-use the audio device used in content browser. */
+	UPROPERTY(EditAnywhere, config, Category = Sound)
+	uint32 bCreateNewAudioDeviceForPlayInEditor: 1;
+
 	/** If true audio will be enabled in the editor. Does not affect PIE **/
 	UPROPERTY(config)
 	uint32 bEnableRealTimeAudio:1;
diff --git a/Engine/Source/Editor/UnrealEd/Private/Settings/SettingsClasses.cpp b/Engine/Source/Editor/UnrealEd/Private/Settings/SettingsClasses.cpp
--- a/Engine/Source/Editor/UnrealEd/Private/Settings/SettingsClasses.cpp
+++ b/Engine/Source/Editor/UnrealEd/Private/Settings/SettingsClasses.cpp
@@ -348,6 +348,7 @@ ULevelEditorMiscSettings::ULevelEditorMiscSettings( const FObjectInitializer& Ob
 	bPromptWhenAddingToLevelBeforeCheckout = true;
 	bPromptWhenAddingToLevelOutsideBounds = true;
 	PercentageThresholdForPrompt = 20.0f;
+	bCreateNewAudioDeviceForPlayInEditor = true;
 }
 
 
diff --git a/Engine/Source/Runtime/Engine/Private/AudioDeviceManager.cpp b/Engine/Source/Runtime/Engine/Private/AudioDeviceManager.cpp
--- a/Engine/Source/Runtime/Engine/Private/AudioDeviceManager.cpp
+++ b/Engine/Source/Runtime/Engine/Private/AudioDeviceManager.cpp
@@ -11,6 +11,7 @@
 
 #if WITH_EDITOR
 #include "AudioEditorModule.h"
+#include "Settings/LevelEditorMiscSettings.h"
 #endif
 
 // Private consts for helping with index/generation determination in audio device manager
@@ -380,7 +381,13 @@ bool FAudioDeviceManager::CreateAudioDevice(bool bCreateNewDevice, FCreateAudioD
 	bool bRequiresInit = true;
 
 	// For the first PIE window, we'll just use the main audio device
-	if (NumActiveAudioDevices == 1)
+	bool bCreateNewAudioDeviceForPlayInEditor = true;
+
+#if WITH_EDITOR
+	bCreateNewAudioDeviceForPlayInEditor = GetDefault<ULevelEditorMiscSettings>()->bCreateNewAudioDeviceForPlayInEditor;
+#endif
+
+	if (NumActiveAudioDevices == 1 && !bCreateNewAudioDeviceForPlayInEditor)
 	{
 		FAudioDevice* MainAudioDevice = GEngine->GetMainAudioDevice();
 		if (MainAudioDevice)
