Remove MaxWaveInstances as it is never exercised in typical cases where MaxChannels is lower.  Also causes a secondary sort in high count situations which is sub-optimal and obfuscates further why sounds are stopping when they shouldn't.

diff --git a/Engine/Source/Runtime/Engine/Classes/Sound/AudioSettings.h b/Engine/Source/Runtime/Engine/Classes/Sound/AudioSettings.h
--- a/Engine/Source/Runtime/Engine/Classes/Sound/AudioSettings.h
+++ b/Engine/Source/Runtime/Engine/Classes/Sound/AudioSettings.h
@@ -166,10 +166,6 @@ class ENGINE_API UAudioSettings : public UDeveloperSettings
 	UPROPERTY(config, EditAnywhere, Category = "Quality", AdvancedDisplay)
 	uint32 bAllowCenterChannel3DPanning : 1;
 
-	/** The max number of active sounds allowed. Used to cull numbers of active sounds, which reduces CPU cost of audio thread. */
-	UPROPERTY(config, EditAnywhere, Category = "Quality", AdvancedDisplay)
-	uint32 MaxWaveInstances;
-
 	/**
 	 * The max number of sources to reserve for "stopping" sounds. A "stopping" sound applies a fast fade in the DSP
 	 * render to prevent discontinuities when stopping sources.
diff --git a/Engine/Source/Runtime/Engine/Private/AudioDevice.cpp b/Engine/Source/Runtime/Engine/Private/AudioDevice.cpp
--- a/Engine/Source/Runtime/Engine/Private/AudioDevice.cpp
+++ b/Engine/Source/Runtime/Engine/Private/AudioDevice.cpp
@@ -147,7 +147,6 @@ FAudioDevice::FAudioDevice()
 	, MaxChannelsScale(1.0f)
 	, MaxChannelsScale_GameThread(1.0f)
 	, NumStoppingVoices(32)
-	, MaxWaveInstances(0)
 	, SampleRate(0)
 	, NumPrecacheFrames(MONO_PCM_BUFFER_SAMPLES)
 	, CommonAudioPoolSize(0)
@@ -263,7 +262,7 @@ bool FAudioDevice::Init(int32 InMaxChannels)
 		PlatformAudioHeadroom = FMath::Pow(10.0f, Headroom / 20.0f);
 	}
 
-	int32 NumPrecacheFramesSettings = 0; 
+	int32 NumPrecacheFramesSettings = 0;
 	if (GConfig->GetInt(TEXT("Audio"), TEXT("NumPrecacheFrames"), NumPrecacheFramesSettings, GEngineIni))
 	{
 		NumPrecacheFrames = FMath::Min(128, NumPrecacheFramesSettings);
@@ -280,7 +279,6 @@ bool FAudioDevice::Init(int32 InMaxChannels)
 	bAllowCenterChannel3DPanning = AudioSettings->bAllowCenterChannel3DPanning;
 	bAllowVirtualizedSounds = AudioSettings->bAllowVirtualizedSounds;
 	DefaultReverbSendLevel = AudioSettings->DefaultReverbSendLevel;
-	MaxWaveInstances = AudioSettings->MaxWaveInstances;
 
 	const FSoftObjectPath DefaultBaseSoundMixName = GetDefault<UAudioSettings>()->DefaultBaseSoundMix;
 	if (DefaultBaseSoundMixName.IsValid())
@@ -555,7 +553,7 @@ void FAudioDevice::Teardown()
 		SpatializationPluginInterface.Reset();
 		bSpatializationInterfaceEnabled = false;
 	}
-	
+
 	if (ReverbPluginInterface.IsValid())
 	{
 		ReverbPluginInterface->Shutdown();
@@ -3344,7 +3342,7 @@ int32 FAudioDevice::GetSortedActiveWaveInstances(TArray<FWaveInstance*>& WaveIns
 	{
 		StopQuietSoundsDueToMaxConcurrency(WaveInstances, ActiveSoundsCopy);
 	}
-	
+
 	int32 FirstActiveIndex = 0;
 	// Only need to do the wave instance sort if we have any waves and if our wave instances are greater than our max channels.
 	if (WaveInstances.Num() >= 0)
@@ -3869,22 +3867,6 @@ void FAudioDevice::Update(bool bGameTicking)
 		// Check which sounds are active from these wave instances and update passive SoundMixes
 		UpdatePassiveSoundMixModifiers(ActiveWaveInstances, FirstActiveIndex);
 
-		// Update the min priority cull threshold. It will be based off the min priority of the lowest priority sound here.
-		if (ActiveWaveInstances.Num() > MaxWaveInstances)
-		{
-			// Helper function for "Sort" (higher priority sorts last).
-			struct FCompareFWaveInstanceByPriority
-			{
-				FORCEINLINE bool operator()(const FWaveInstance& A, const FWaveInstance& B) const
-				{
-					return A.Priority > B.Priority;
-				}
-			};
-
-			// Sort by priority (lowest priority first).
-			ActiveWaveInstances.Sort(FCompareFWaveInstanceByPriority());
-		}
-
 		// If not paused, update the playback time of the active sounds after we've processed passive mix modifiers
 		// Note that for sounds which play while paused, this will result in longer active sound playback times, which will be ok. If we update the
 		// active sound is updated while paused (for a long time), most sounds will be stopped when unpaused.
@@ -5009,7 +4991,7 @@ UAudioComponent* FAudioDevice::CreateComponent(USoundBase* Sound, UWorld* World,
 	Params->bPlay = bPlay;
 	Params->bStopWhenOwnerDestroyed = bStopWhenOwnerDestroyed;
 	Params->AttenuationSettings = AttenuationSettings;
-	
+
 	if (ConcurrencySettings)
 	{
 		Params->ConcurrencySet.Add(ConcurrencySettings);
diff --git a/Engine/Source/Runtime/Engine/Private/AudioSettings.cpp b/Engine/Source/Runtime/Engine/Private/AudioSettings.cpp
--- a/Engine/Source/Runtime/Engine/Private/AudioSettings.cpp
+++ b/Engine/Source/Runtime/Engine/Private/AudioSettings.cpp
@@ -71,7 +71,6 @@ void UAudioSettings::AddDefaultSettings()
 	bEnableLegacyReverb = false;
 	VoiPSampleRate = EVoiceSampleRate::Low16000Hz;
 	VoipBufferingDelay = 0.2f;
-	MaxWaveInstances = 100;
 	NumStoppingSources = 8;
 }
 
diff --git a/Engine/Source/Runtime/Engine/Public/AudioDevice.h b/Engine/Source/Runtime/Engine/Public/AudioDevice.h
--- a/Engine/Source/Runtime/Engine/Public/AudioDevice.h
+++ b/Engine/Source/Runtime/Engine/Public/AudioDevice.h
@@ -1592,9 +1592,6 @@ public:
 	/** The number of sources to reserve for stopping sounds. */
 	int32 NumStoppingVoices;
 
-	/** The maximum number of wave instances allowed. */
-	int32 MaxWaveInstances;
-
 	/** The sample rate of all the audio devices */
 	int32 SampleRate;
 
