CSV camera stats + events for console commands

diff --git a/Engine/Source/Runtime/Engine/Classes/Engine/GameViewportClient.h b/Engine/Source/Runtime/Engine/Classes/Engine/GameViewportClient.h
--- a/Engine/Source/Runtime/Engine/Classes/Engine/GameViewportClient.h
+++ b/Engine/Source/Runtime/Engine/Classes/Engine/GameViewportClient.h
@@ -881,6 +881,9 @@ private:
 	 */
 	bool SetDisplayConfiguration( const FIntPoint* Dimensions, EWindowMode::Type WindowMode);
 
+	/** Updates CSVProfiler camera stats */
+	void UpdateCsvCameraStats(const FSceneView* View);
+
 #if WITH_EDITOR
 	/** Delegate called when game viewport client received input key */
 	FOnGameViewportInputKey GameViewportInputKeyDelegate;
diff --git a/Engine/Source/Runtime/Engine/Private/GameViewportClient.cpp b/Engine/Source/Runtime/Engine/Private/GameViewportClient.cpp
--- a/Engine/Source/Runtime/Engine/Private/GameViewportClient.cpp
+++ b/Engine/Source/Runtime/Engine/Private/GameViewportClient.cpp
@@ -59,6 +59,8 @@
 #include "DynamicResolutionState.h"
 #include "ProfilingDebugging/CsvProfiler.h"
 
+CSV_DEFINE_CATEGORY(View, true);
+
 #if WITH_EDITOR
 #include "Settings/LevelEditorPlaySettings.h"
 #endif
@@ -114,6 +116,48 @@ static TAutoConsoleVariable<float> CVarSecondaryScreenPercentage( // TODO: make
 	TEXT(" 1: override secondary screen percentage."),
 	ECVF_Default);
 
+
+void UGameViewportClient::UpdateCsvCameraStats(const FSceneView* View)
+{
+#if CSV_PROFILER
+	if (!View)
+	{
+		return;
+	}
+	static uint32 PrevFrameNumber = GFrameNumber;
+	static double PrevTime = 0.0;
+	static FVector PrevViewOrigin = FVector(ForceInitToZero);
+
+	// TODO: support multiple views/view families, e.g for splitscreen. For now, we just output stats for the first one.
+	if (GFrameNumber != PrevFrameNumber)
+	{
+		FVector ViewOrigin = View->ViewMatrices.GetViewOrigin();
+		FVector ForwardVec = View->ViewMatrices.GetOverriddenTranslatedViewMatrix().GetColumn(2);
+		FVector UpVec = View->ViewMatrices.GetOverriddenTranslatedViewMatrix().GetColumn(1);
+		FVector Diff = ViewOrigin - PrevViewOrigin;
+		double CurrentTime = FPlatformTime::Seconds();
+		double DeltaT = CurrentTime - PrevTime;
+		FVector Velocity = Diff / float(DeltaT);
+		float CameraSpeed = Velocity.Size();
+		PrevViewOrigin = ViewOrigin;
+		PrevTime = CurrentTime;
+		PrevFrameNumber = GFrameNumber;
+
+		CSV_CUSTOM_STAT(View, PosX, View->ViewMatrices.GetViewOrigin().X, ECsvCustomStatOp::Set);
+		CSV_CUSTOM_STAT(View, PosY, View->ViewMatrices.GetViewOrigin().Y, ECsvCustomStatOp::Set);
+		CSV_CUSTOM_STAT(View, PosZ, View->ViewMatrices.GetViewOrigin().Z, ECsvCustomStatOp::Set);
+		CSV_CUSTOM_STAT(View, ForwardX, ForwardVec.X, ECsvCustomStatOp::Set);
+		CSV_CUSTOM_STAT(View, ForwardY, ForwardVec.Y, ECsvCustomStatOp::Set);
+		CSV_CUSTOM_STAT(View, ForwardZ, ForwardVec.Z, ECsvCustomStatOp::Set);
+		CSV_CUSTOM_STAT(View, UpX, UpVec.X, ECsvCustomStatOp::Set);
+		CSV_CUSTOM_STAT(View, UpY, UpVec.Y, ECsvCustomStatOp::Set);
+		CSV_CUSTOM_STAT(View, UpZ, UpVec.Z, ECsvCustomStatOp::Set);
+		CSV_CUSTOM_STAT(View, Speed, CameraSpeed, ECsvCustomStatOp::Set);
+	}
+#endif
+}
+
+
 UGameViewportClient::UGameViewportClient(const FObjectInitializer& ObjectInitializer)
 	: Super(ObjectInitializer)
 #if WITH_EDITOR
@@ -1308,6 +1352,10 @@ void UGameViewportClient::Draw(FViewport* InViewport, FCanvas* SceneCanvas)
 					#if RHI_RAYTRACING
 						View->SetupRayTracedRendering();
 					#endif
+
+					#if CSV_PROFILER
+						UpdateCsvCameraStats(View);
+					#endif
 					}
 
 					// Add view information for resource streaming. Allow up to 5X boost for small FOV.
diff --git a/Engine/Source/Runtime/Engine/Private/UnrealEngine.cpp b/Engine/Source/Runtime/Engine/Private/UnrealEngine.cpp
--- a/Engine/Source/Runtime/Engine/Private/UnrealEngine.cpp
+++ b/Engine/Source/Runtime/Engine/Private/UnrealEngine.cpp
@@ -1730,6 +1730,7 @@ void UEngine::TickDeferredCommands()
 	const int32 DeferredCommandsCount = DeferredCommands.Num();
 	for( int32 DeferredCommandsIndex=0; DeferredCommandsIndex<DeferredCommandsCount; DeferredCommandsIndex++ )
 	{
+		CSV_EVENT_GLOBAL(TEXT("Cmd: %s"), *DeferredCommands[DeferredCommandsIndex]);
 		// Use LocalPlayer if available...
 		ULocalPlayer* LocalPlayer = GetDebugLocalPlayer();
 		if( LocalPlayer )
diff --git a/Engine/Source/Runtime/Engine/Private/UserInterface/Console.cpp b/Engine/Source/Runtime/Engine/Private/UserInterface/Console.cpp
--- a/Engine/Source/Runtime/Engine/Private/UserInterface/Console.cpp
+++ b/Engine/Source/Runtime/Engine/Private/UserInterface/Console.cpp
@@ -29,6 +29,7 @@
 #include "Stats/StatsData.h"
 #include "Misc/TextFilter.h"
 #include "HAL/PlatformApplicationMisc.h"
+#include "ProfilingDebugging/CsvProfiler.h"
 
 static const uint32 MAX_AUTOCOMPLETION_LINES = 20;
 
@@ -519,6 +520,7 @@ void UConsole::SetCursorPos( int32 Position )
 
 void UConsole::ConsoleCommand(const FString& Command)
 {
+	CSV_EVENT_GLOBAL(TEXT("Cmd: %s"), *Command);
 	// insert into history buffer
 	{
 		HistoryBuffer.Remove(Command);
