Prevent crashing the game if a NetFieldExportGroup isn't found while playing back a replay.

diff --git a/Engine/Source/Runtime/Engine/Private/PackageMapClient.cpp b/Engine/Source/Runtime/Engine/Private/PackageMapClient.cpp
--- a/Engine/Source/Runtime/Engine/Private/PackageMapClient.cpp
+++ b/Engine/Source/Runtime/Engine/Private/PackageMapClient.cpp
@@ -1608,22 +1608,31 @@ void UPackageMapClient::ReceiveNetFieldExports(FArchive& Archive)
 		}
 		else
 		{
-			NetFieldExportGroup = GuidCache->NetFieldExportGroupIndexToGroup.FindChecked(PathNameIndex);
+			FNetFieldExportGroup** FoundNetFieldExport = GuidCache->NetFieldExportGroupIndexToGroup.Find(PathNameIndex);
+			NetFieldExportGroup = FoundNetFieldExport ? *FoundNetFieldExport : nullptr;
 		}
 
-		TArray<FNetFieldExport>& Exports = NetFieldExportGroup->NetFieldExports;
 		FNetFieldExport Export;
 		Archive << Export;
 
-		if (Exports.IsValidIndex(Export.Handle))
+		if (NetFieldExportGroup)
 		{
-			// preserve compatibility flag
-			Export.bIncompatible = Exports[Export.Handle].bIncompatible;
-			Exports[Export.Handle] = Export;
+			TArray<FNetFieldExport>& Exports = NetFieldExportGroup->NetFieldExports;
+			if (Exports.IsValidIndex(Export.Handle))
+			{
+				// preserve compatibility flag
+				Export.bIncompatible = Exports[Export.Handle].bIncompatible;
+				Exports[Export.Handle] = Export;
+			}
+			else
+			{
+				UE_LOG(LogNetPackageMap, Error, TEXT("ReceiveNetFieldExports: Invalid NetFieldExportHandle '%i', Max '%i'"), Export.Handle, Exports.Num());
+			}
 		}
 		else
 		{
-			UE_LOG(LogNetPackageMap, Error, TEXT("ReceiveNetFieldExports: Invalid NetFieldExportHandle '%i', Max '%i'"), Export.Handle, Exports.Num());
+			UE_LOG(LogNetPackageMap, Error, TEXT("ReceiveNetFieldExports: Unable to find NetFieldExportGroup for export. Export.Handle=%i, Export.Name=%s, PathNameIndex=%lu, WasExported=%d, Archive.IsError()=%d"),
+				Export.Handle, *Export.ExportName.ToString(), PathNameIndex, !!WasExported, !!Archive.IsError());
 		}
 	}
 }
