Timers: Added a new cvar TimerManager.BuildTimerSourceList which tracks all timer expirations while active, attributing them to a source and keeping some basic statistics like total # of expires and max # expires in a given frame
- When non-zero, tracks which timers expire each frame, dumping them during shutdown or when the flag is changed back to 0.
-  0: Off
-  1: On - Group timers by class (useful to focus on entire systems of things, especially bad spikey frames where we care about aggregates)
-  2: On - Do not group timers by class (useful if individual instances are problematic)
- This is gated by UE_ENABLE_TRACKING_TIMER_SOURCES, which defaults to !UE_BUILD_SHIPPING

diff --git a/Engine/Source/Runtime/Engine/Private/TimerManager.cpp b/Engine/Source/Runtime/Engine/Private/TimerManager.cpp
--- a/Engine/Source/Runtime/Engine/Private/TimerManager.cpp
+++ b/Engine/Source/Runtime/Engine/Private/TimerManager.cpp
@@ -48,6 +48,157 @@ static FAutoConsoleVariableRef CVarMaxExpiredTimersToLog(
 	MaxExpiredTimersToLog,
 	TEXT("Maximum number of TimerData exceeding the threshold to log in a single frame."));
 
+
+#if UE_ENABLE_TRACKING_TIMER_SOURCES
+static int32 GBuildTimerSourceList = 0;
+static FAutoConsoleVariableRef CVarGBuildTimerSourceList(
+	TEXT("TimerManager.BuildTimerSourceList"), GBuildTimerSourceList,
+	TEXT("When non-zero, tracks which timers expire each frame, dumping them during shutdown or when the flag is changed back to 0.")
+	TEXT("\n  0: Off")
+	TEXT("\n  1: On - Group timers by class (useful to focus on entire systems of things, especially bad spikey frames where we care about aggregates)")
+	TEXT("\n  2: On - Do not group timers by class (useful if individual instances are problematic)"),
+	ECVF_Default);
+
+// Information about a single timer or timer source being tracked
+struct FTimerSourceEntry
+{
+	uint64 TotalCount = 0;
+	double SumApproxRate = 0.0;
+	uint32 LastFrameID = 0;
+	uint32 NumThisFrame = 0;
+	uint32 MaxPerFrame = 0;
+
+	void UpdateEntry(float Rate, uint32 FrameID, int32 CallCount)
+	{
+		TotalCount += CallCount;
+		SumApproxRate += Rate;
+
+		if (LastFrameID != FrameID)
+		{
+			LastFrameID = FrameID;
+			NumThisFrame = 0;
+		}
+		NumThisFrame += CallCount;
+
+		MaxPerFrame = FMath::Max(NumThisFrame, MaxPerFrame);
+	}
+};
+
+// List of information about all timers that have expired during a tracking window
+struct FTimerSourceList
+{
+	TMap<FString, FTimerSourceEntry> Entries;
+	UGameInstance* OwningGameInstance = nullptr;
+
+	// This is similar to FTimerUnifiedDelegate::ToString() but it tries to find the base class / exclude vptr printout info so that timers are collapsed/aggregated better
+	static FString GetPartialDeduplicateDelegateToString(const FTimerUnifiedDelegate& Delegate)
+	{
+		FString FunctionNameStr;
+		FString ObjectNameStr;
+		bool bDynDelegate = false;
+
+		if (Delegate.FuncDelegate.IsBound())
+		{
+			ObjectNameStr = TEXT("NonDynamicDelegate");
+			FName FunctionName;
+#if USE_DELEGATE_TRYGETBOUNDFUNCTIONNAME
+			FunctionName = Delegate.FuncDelegate.TryGetBoundFunctionName();
+#endif
+			if (FunctionName.IsNone())
+			{
+				uint64 ProgramCounter = Delegate.FuncDelegate.GetBoundProgramCounterForTimerManager();
+				if (ProgramCounter != 0)
+				{
+					// Add the function address
+					if (DumpTimerLogSymbolNames)
+					{
+						// Try to resolve the function address to a symbol
+						FProgramCounterSymbolInfo SymbolInfo;
+						FPlatformStackWalk::ProgramCounterToSymbolInfo(ProgramCounter, /*out*/ SymbolInfo);
+						FunctionNameStr = FString::Printf(TEXT("%s [%s:%d]"), ANSI_TO_TCHAR(SymbolInfo.FunctionName), ANSI_TO_TCHAR(SymbolInfo.Filename), SymbolInfo.LineNumber);
+					}
+					else
+					{
+						FunctionNameStr = FString::Printf(TEXT("func: 0x%llx"), ProgramCounter);
+					}
+				}
+				else
+				{
+					FunctionNameStr = TEXT("0x0");
+				}
+			}
+			else
+			{
+				FunctionNameStr = FunctionName.ToString();
+			}
+		}
+		else if (Delegate.FuncDynDelegate.IsBound())
+		{
+			const FName FuncFName = Delegate.FuncDynDelegate.GetFunctionName();
+			FunctionNameStr = FuncFName.ToString();
+			bDynDelegate = true;
+
+			UClass* SourceClass = nullptr;
+			if (const UObject* Object = Delegate.FuncDynDelegate.GetUObject())
+			{
+				SourceClass = Object->GetClass();
+				if (UFunction* Func = SourceClass->FindFunctionByName(FuncFName))
+				{
+					SourceClass = Func->GetOwnerClass();
+				}
+			}
+
+
+			ObjectNameStr = GetPathNameSafe(SourceClass);
+		}
+		else
+		{
+			ObjectNameStr = TEXT("NotBound!");
+			FunctionNameStr = TEXT("NotBound!");
+		}
+
+		return FString::Printf(TEXT("%s,\"%s\""), *ObjectNameStr, *FunctionNameStr);
+	}
+
+	void AddEntry(const FTimerData& Data, int32 CallCount)
+	{
+		Entries.FindOrAdd((GBuildTimerSourceList == 2) ? Data.TimerDelegate.ToString() : GetPartialDeduplicateDelegateToString(Data.TimerDelegate)).UpdateEntry(Data.Rate, GFrameCounter, CallCount);
+	}
+
+	void DumpEntries()
+	{
+		FString AdditionalInfoAboutGameInstance;
+		if (OwningGameInstance != nullptr)
+		{
+			if (FWorldContext* WorldContext = OwningGameInstance->GetWorldContext())
+			{
+				if (WorldContext->RunAsDedicated)
+				{
+					AdditionalInfoAboutGameInstance = TEXT(" (dedicated server)");
+				}
+			}
+		}
+
+		UE_LOG(LogEngine, Log, TEXT("-----------------------"));
+		UE_LOG(LogEngine, Log, TEXT("Listing Expired Timer Stats for OwningGameInstance=%s%s"), *GetPathNameSafe(OwningGameInstance), *AdditionalInfoAboutGameInstance);
+		
+		UE_LOG(LogEngine, Log, TEXT("TotalExpires,RateAvg,PeakInFrame,ObjectOrClass,FunctionName"));
+		Entries.ValueSort([](const FTimerSourceEntry& A, const FTimerSourceEntry& B) { return A.TotalCount > B.TotalCount; });
+		for (const auto& KVP : Entries)
+		{
+			UE_LOG(LogEngine, Log, TEXT("%llu,%.2f,%u,%s"), KVP.Value.TotalCount, KVP.Value.SumApproxRate / KVP.Value.TotalCount, KVP.Value.MaxPerFrame, *KVP.Key);
+		}
+
+		UE_LOG(LogEngine, Log, TEXT("-----------------------"));
+	}
+
+	~FTimerSourceList()
+	{
+		DumpEntries();
+	}
+};
+#endif // UE_ENABLE_TRACKING_TIMER_SOURCES
+
 namespace
 {
 	void DescribeFTimerDataSafely(FOutputDevice& Ar, const FTimerData& Data)
@@ -622,6 +773,21 @@ void FTimerManager::Tick(float DeltaTime)
 		return;
 	}
 
+#if UE_ENABLE_TRACKING_TIMER_SOURCES
+	if ((TimerSourceList == nullptr) != (GBuildTimerSourceList == 0))
+	{
+		if (TimerSourceList == nullptr)
+		{
+			TimerSourceList.Reset(new FTimerSourceList);
+			TimerSourceList->OwningGameInstance = OwningGameInstance;
+		}
+		else
+		{
+			TimerSourceList.Reset();
+		}
+	}
+#endif
+
 	const double StartTime = FPlatformTime::Seconds();
 	bool bDumpTimerLogsThresholdExceeded = false;
 	int32 NbExpiredTimers = 0;
@@ -673,6 +839,14 @@ void FTimerManager::Tick(float DeltaTime)
 				FMath::TruncToInt( (InternalTime - Top->ExpireTime) / Top->Rate ) + 1
 				: 1;
 
+#if UE_ENABLE_TRACKING_TIMER_SOURCES
+			if (TimerSourceList.IsValid())
+			{
+				//@TODO: The actual call count may be less, e.g., if the delegate clears itself during the loop below
+				TimerSourceList->AddEntry(*Top, CallCount);
+			}
+#endif
+
 			// Now call the function
 			for (int32 CallIdx=0; CallIdx<CallCount; ++CallIdx)
 			{ 
@@ -687,7 +861,6 @@ void FTimerManager::Tick(float DeltaTime)
 				RunTimerDelegates.Add(Top->TimerDelegate);
 #endif
 
-
 				Top->TimerDelegate.Execute();
 
 				// Update Top pointer, in case it has been invalidated by the Execute call
@@ -774,6 +947,18 @@ TStatId FTimerManager::GetStatId() const
 	RETURN_QUICK_DECLARE_CYCLE_STAT(FTimerManager, STATGROUP_Tickables);
 }
 
+void FTimerManager::SetGameInstance(UGameInstance* InGameInstance)
+{
+	OwningGameInstance = InGameInstance;
+
+#if UE_ENABLE_TRACKING_TIMER_SOURCES
+	if (TimerSourceList.IsValid())
+	{
+		TimerSourceList->OwningGameInstance = OwningGameInstance;
+	}
+#endif
+}
+
 void FTimerManager::ListTimers() const
 {
 	TArray<const FTimerData*> ValidActiveTimers;
diff --git a/Engine/Source/Runtime/Engine/Public/TimerManager.h b/Engine/Source/Runtime/Engine/Public/TimerManager.h
--- a/Engine/Source/Runtime/Engine/Public/TimerManager.h
+++ b/Engine/Source/Runtime/Engine/Public/TimerManager.h
@@ -12,9 +12,14 @@
 #include "Engine/EngineTypes.h"
 
 class UGameInstance;
+struct FTimerSourceList;
 
 DECLARE_DELEGATE(FTimerDelegate);
 
+#ifndef UE_ENABLE_TRACKING_TIMER_SOURCES
+#define UE_ENABLE_TRACKING_TIMER_SOURCES !UE_BUILD_SHIPPING
+#endif
+
 /** Simple interface to wrap a timer delegate that can be either native or dynamic. */
 struct FTimerUnifiedDelegate
 {
@@ -387,7 +392,7 @@ public:
 	void ListTimers() const;
 
 	/** Used by the UGameInstance constructor to set this manager's owning game instance. */
-	void SetGameInstance(UGameInstance* InGameInstance) { OwningGameInstance = InGameInstance; }
+	void SetGameInstance(UGameInstance* InGameInstance);
 
 // This should be private, but needs to be public for testing.
 public:
@@ -449,5 +454,10 @@ private:
 
 	/** The game instance that created this timer manager. May be null if this timer manager wasn't created by a game instance. */
 	UGameInstance* OwningGameInstance;
+
+#if UE_ENABLE_TRACKING_TIMER_SOURCES
+	/** Debugging/tracking information used when TimerManager.BuildTimerSourceList is set */
+	TUniquePtr<FTimerSourceList> TimerSourceList;
+#endif
 };
 
