Correctly handle the application of the lightmap to instances depending on the load order of the light scenario vs the level containing the data.

diff --git a/Engine/Source/Runtime/Engine/Private/HierarchicalInstancedStaticMesh.cpp b/Engine/Source/Runtime/Engine/Private/HierarchicalInstancedStaticMesh.cpp
--- a/Engine/Source/Runtime/Engine/Private/HierarchicalInstancedStaticMesh.cpp
+++ b/Engine/Source/Runtime/Engine/Private/HierarchicalInstancedStaticMesh.cpp
@@ -2808,26 +2808,33 @@ void UHierarchicalInstancedStaticMeshComponent::BuildTreeAsync()
 
 void UHierarchicalInstancedStaticMeshComponent::PropagateLightingScenarioChange()
 {
-	if (GIsEditor && PerInstanceRenderData.IsValid())
+	if (PerInstanceRenderData.IsValid())
 	{
-		FComponentRecreateRenderStateContext Context(this);
-
-		const FMeshMapBuildData* MeshMapBuildData = nullptr;
-		if (LODData.Num() > 0)
+		if (GIsEditor)
 		{
-			MeshMapBuildData = GetMeshMapBuildData(LODData[0]);
-		}
+			FComponentRecreateRenderStateContext Context(this);
 
-		if (MeshMapBuildData != nullptr)
-		{
-			for (int32 InstanceIndex = 0; InstanceIndex < PerInstanceSMData.Num(); ++InstanceIndex)
+			const FMeshMapBuildData* MeshMapBuildData = nullptr;
+			if (LODData.Num() > 0)
 			{
-				int32 RenderIndex = InstanceReorderTable.IsValidIndex(InstanceIndex) ? InstanceReorderTable[InstanceIndex] : InstanceIndex;
+				MeshMapBuildData = GetMeshMapBuildData(LODData[0]);
+			}
+
+			if (MeshMapBuildData != nullptr)
+			{
+				for (int32 InstanceIndex = 0; InstanceIndex < PerInstanceSMData.Num(); ++InstanceIndex)
+				{
+					int32 RenderIndex = InstanceReorderTable.IsValidIndex(InstanceIndex) ? InstanceReorderTable[InstanceIndex] : InstanceIndex;
 
-				InstanceUpdateCmdBuffer.SetLightMapData(RenderIndex, MeshMapBuildData->PerInstanceLightmapData[InstanceIndex].LightmapUVBias);
-				InstanceUpdateCmdBuffer.SetShadowMapData(RenderIndex, MeshMapBuildData->PerInstanceLightmapData[InstanceIndex].ShadowmapUVBias);
+					InstanceUpdateCmdBuffer.SetLightMapData(RenderIndex, MeshMapBuildData->PerInstanceLightmapData[InstanceIndex].LightmapUVBias);
+					InstanceUpdateCmdBuffer.SetShadowMapData(RenderIndex, MeshMapBuildData->PerInstanceLightmapData[InstanceIndex].ShadowmapUVBias);
+				}
 			}
 		}
+		else
+		{
+			BuildTreeIfOutdated(true, true);
+		}
 	}
 }
 
@@ -2961,6 +2968,16 @@ void UHierarchicalInstancedStaticMeshComponent::OnPostLoadPerInstanceData()
 		}
 		else
 		{
+			AActor* Owner = GetOwner();
+			ULevel* OwnerLevel = Owner->GetLevel();
+			UWorld* OwnerWorld = OwnerLevel ? OwnerLevel->OwningWorld : nullptr;
+
+			if (OwnerWorld && OwnerWorld->GetActiveLightingScenario() != nullptr && OwnerWorld->GetActiveLightingScenario() != OwnerLevel)
+			{
+				//update the instance data if the lighting scenario isn't the owner level
+				bForceTreeBuild = true;
+			}
+
 			if (!bForceTreeBuild)
 			{
 				// create PerInstanceRenderData either from current data or pre-built instance buffer
