BlendSapce crash when somebody forcefully delete the sample animation

diff --git a/Engine/Source/Runtime/Engine/Private/Animation/BlendSpaceBase.cpp b/Engine/Source/Runtime/Engine/Private/Animation/BlendSpaceBase.cpp
--- a/Engine/Source/Runtime/Engine/Private/Animation/BlendSpaceBase.cpp
+++ b/Engine/Source/Runtime/Engine/Private/Animation/BlendSpaceBase.cpp
@@ -677,8 +677,9 @@ bool UBlendSpaceBase::GetSamplesFromBlendInput(const FVector &BlendInput, TArray
 		{
 			const int32 SampleDataIndex = GridElement.Indices[Ind];		
 			if( SampleData.IsValidIndex(SampleDataIndex) 
-#if WITH_EDITOR
+#if WITH_EDITOR // we check these in editor because these could change when editor is running
 				&& SampleData[SampleDataIndex].bIsValid
+				&& SampleData[SampleDataIndex].Animation 
 				&& SampleData[SampleDataIndex].Animation->GetSkeleton() == GetSkeleton()
 #endif // WITH_EDITOR
 				)
