Seekable Streams Using LPCM/ADPCM
https://p4-swarm.epicgames.net/reviews/6345090

diff --git a/Engine/Plugins/Lumin/MagicLeap/Source/MagicLeapAudio/Private/AudioMixerPlatformMagicLeap.cpp b/Engine/Plugins/Lumin/MagicLeap/Source/MagicLeapAudio/Private/AudioMixerPlatformMagicLeap.cpp
--- a/Engine/Plugins/Lumin/MagicLeap/Source/MagicLeapAudio/Private/AudioMixerPlatformMagicLeap.cpp
+++ b/Engine/Plugins/Lumin/MagicLeap/Source/MagicLeapAudio/Private/AudioMixerPlatformMagicLeap.cpp
@@ -399,6 +399,11 @@ namespace Audio
 
 	ICompressedAudioInfo* FMixerPlatformMagicLeap::CreateCompressedAudioInfo(USoundWave* InSoundWave)
 	{
+		if (InSoundWave->IsSeekableStreaming())
+		{
+			return new FADPCMAudioInfo();
+		}
+
 #if WITH_OGGVORBIS
 		static FName NAME_OGG(TEXT("OGG"));
 		if (InSoundWave->HasCompressedData(NAME_OGG))
@@ -482,7 +487,7 @@ namespace Audio
 		{
 			return;
 		}
-		// Get the callback buffer from MLAudio 
+		// Get the callback buffer from MLAudio
 		MLResult Result = MLAudioGetOutputStreamBuffer(InPlatform->StreamHandle, &CallbackBuffer);
 		if (Result != MLResult_Ok)
 		{
diff --git a/Engine/Source/Developer/AllDesktopTargetPlatform/Private/AllDesktopTargetPlatform.cpp b/Engine/Source/Developer/AllDesktopTargetPlatform/Private/AllDesktopTargetPlatform.cpp
--- a/Engine/Source/Developer/AllDesktopTargetPlatform/Private/AllDesktopTargetPlatform.cpp
+++ b/Engine/Source/Developer/AllDesktopTargetPlatform/Private/AllDesktopTargetPlatform.cpp
@@ -11,7 +11,7 @@
 #endif
 
 
- 
+
 
 /* FAllDesktopTargetPlatform structors
  *****************************************************************************/
@@ -45,9 +45,9 @@ void FAllDesktopTargetPlatform::GetAllPossibleShaderFormats( TArray<FName>& OutF
 	static FName NAME_PCD3D_SM4(TEXT("PCD3D_SM4"));
 	static FName NAME_GLSL_150(TEXT("GLSL_150"));
 	static FName NAME_GLSL_430(TEXT("GLSL_430"));
-	
+
 #if PLATFORM_WINDOWS
-	// right now, only windows can properly compile D3D shaders (this won't corrupt the DDC, but it will 
+	// right now, only windows can properly compile D3D shaders (this won't corrupt the DDC, but it will
 	// make it so that packages cooked on Mac/Linux will only run on Windows with -opengl)
 	OutFormats.AddUnique(NAME_PCD3D_SM5);
 	OutFormats.AddUnique(NAME_PCD3D_SM4);
@@ -78,24 +78,33 @@ FName FAllDesktopTargetPlatform::GetWaveFormat( const class USoundWave* Wave ) c
 {
 	static FName NAME_OGG(TEXT("OGG"));
 	static FName NAME_OPUS(TEXT("OPUS"));
-	
-	if (Wave->IsStreaming())
+	static FName NAME_ADPCM(TEXT("ADPCM"));
+
+	// Seekable streams need to pick a codec which allows fixed-sized frames so we can compute stream chunk index to load
+	if (Wave->IsSeekableStreaming())
+	{
+		return NAME_ADPCM;
+	}
+	else if (Wave->IsStreaming())
 	{
-		// @todo desktop platform: Does Linux support OPUS?
+#if !USE_VORBIS_FOR_STREAMING
 		return NAME_OPUS;
+#endif
 	}
-	
+
 	return NAME_OGG;
 }
 
 
 void FAllDesktopTargetPlatform::GetAllWaveFormats(TArray<FName>& OutFormats) const
 {
+	static FName NAME_ADPCM(TEXT("ADPCM"));
 	static FName NAME_OGG(TEXT("OGG"));
 	static FName NAME_OPUS(TEXT("OPUS"));
+
+	OutFormats.Add(NAME_ADPCM);
 	OutFormats.Add(NAME_OGG);
 	OutFormats.Add(NAME_OPUS);
-
 }
 
 #endif // WITH_ENGINE
diff --git a/Engine/Source/Developer/Android/AndroidTargetPlatform/Private/AndroidTargetPlatform.cpp b/Engine/Source/Developer/Android/AndroidTargetPlatform/Private/AndroidTargetPlatform.cpp
--- a/Engine/Source/Developer/Android/AndroidTargetPlatform/Private/AndroidTargetPlatform.cpp
+++ b/Engine/Source/Developer/Android/AndroidTargetPlatform/Private/AndroidTargetPlatform.cpp
@@ -30,6 +30,7 @@
 
 #if WITH_ENGINE
 #include "AudioCompressionSettings.h"
+#include "Sound/SoundWave.h"
 #endif
 
 #define LOCTEXT_NAMESPACE "FAndroidTargetPlatform"
@@ -206,7 +207,7 @@ FAndroidTargetPlatform::FAndroidTargetPlatform(bool bInIsClient )
 
 
 FAndroidTargetPlatform::~FAndroidTargetPlatform()
-{ 
+{
 	 FTicker::GetCoreTicker().RemoveTicker(TickDelegateHandle);
 }
 
@@ -340,23 +341,23 @@ bool FAndroidTargetPlatform::SupportsFeature( ETargetPlatformFeatures Feature )
 		case ETargetPlatformFeatures::Packaging:
 		case ETargetPlatformFeatures::DeviceOutputLog:
 			return true;
-			
+
 		case ETargetPlatformFeatures::LowQualityLightmaps:
 		case ETargetPlatformFeatures::MobileRendering:
 			return SupportsES31() || SupportsES2() || SupportsVulkan();
-			
+
 		case ETargetPlatformFeatures::HighQualityLightmaps:
 		case ETargetPlatformFeatures::Tessellation:
 		case ETargetPlatformFeatures::DeferredRendering:
 			return SupportsAEP();
-			
+
 		case ETargetPlatformFeatures::SoftwareOcclusion:
 			return SupportsSoftwareOcclusion();
-			
+
 		default:
 			break;
 	}
-	
+
 	return TTargetPlatformBase<FAndroidPlatformProperties>::SupportsFeature(Feature);
 }
 
@@ -408,7 +409,7 @@ void FAndroidTargetPlatform::GetTextureFormats( const UTexture* InTexture, TArra
 	check(InTexture);
 
 	// The order we add texture formats to OutFormats is important. When multiple formats are cooked
-	// and supported by the device, the first supported format listed will be used. 
+	// and supported by the device, the first supported format listed will be used.
 	// eg, ETC1/uncompressed should always be last
 
 	bool bNoCompression = InTexture->CompressionNone				// Code wants the texture uncompressed.
@@ -529,23 +530,23 @@ void FAndroidTargetPlatform::GetAllTextureFormats(TArray<FName>& OutFormats) con
 	OutFormats.Add(AndroidTexFormat::NameG8);
 	OutFormats.Add(AndroidTexFormat::NameG8);
 
-	auto AddAllTextureFormatIfSupports = [=, &OutFormats](bool bIsNonPOT) 
+	auto AddAllTextureFormatIfSupports = [=, &OutFormats](bool bIsNonPOT)
 	{
-		AddTextureFormatIfSupports(AndroidTexFormat::NameAutoPVRTC, OutFormats, bIsNonPOT); 
+		AddTextureFormatIfSupports(AndroidTexFormat::NameAutoPVRTC, OutFormats, bIsNonPOT);
 		AddTextureFormatIfSupports(AndroidTexFormat::NamePVRTC2, OutFormats, bIsNonPOT);
 		AddTextureFormatIfSupports(AndroidTexFormat::NamePVRTC4, OutFormats, bIsNonPOT);
-	
+
 		AddTextureFormatIfSupports(AndroidTexFormat::NameAutoDXT, OutFormats, bIsNonPOT);
 		AddTextureFormatIfSupports(AndroidTexFormat::NameDXT1, OutFormats, bIsNonPOT);
 		AddTextureFormatIfSupports(AndroidTexFormat::NameDXT5, OutFormats, bIsNonPOT);
-	
-		AddTextureFormatIfSupports(AndroidTexFormat::NameATC_RGB, OutFormats, bIsNonPOT); 
+
+		AddTextureFormatIfSupports(AndroidTexFormat::NameATC_RGB, OutFormats, bIsNonPOT);
 		AddTextureFormatIfSupports(AndroidTexFormat::NameATC_RGBA_I, OutFormats, bIsNonPOT);
-	
-		AddTextureFormatIfSupports(AndroidTexFormat::NameAutoETC1, OutFormats, bIsNonPOT); 
+
+		AddTextureFormatIfSupports(AndroidTexFormat::NameAutoETC1, OutFormats, bIsNonPOT);
 		AddTextureFormatIfSupports(AndroidTexFormat::NameAutoETC1a, OutFormats, bIsNonPOT);
 		AddTextureFormatIfSupports(AndroidTexFormat::NameAutoETC2, OutFormats, bIsNonPOT);
-	
+
 		AddTextureFormatIfSupports(AndroidTexFormat::NameAutoATC, OutFormats, bIsNonPOT);
 	};
 
@@ -562,7 +563,7 @@ void FAndroidTargetPlatform::GetReflectionCaptureFormats( TArray<FName>& OutForm
 		// use Full HDR with AEP
 		OutFormats.Add(FName(TEXT("FullHDR")));
 	}
-	
+
 	// always emit encoded
 	OutFormats.Add(FName(TEXT("EncodedHDR")));
 }
@@ -576,39 +577,47 @@ const UTextureLODSettings& FAndroidTargetPlatform::GetTextureLODSettings() const
 
 FName FAndroidTargetPlatform::GetWaveFormat( const class USoundWave* Wave ) const
 {
-	static bool formatRead = false;
-	static FName NAME_FORMAT;
+	static const FName NAME_ADPCM(TEXT("ADPCM"));
+	static const FName NAME_OGG(TEXT("OGG"));
 
-	if (!formatRead)
+	static bool bFormatRead = false;
+	static FName NAME_FORMAT;
+	if (!bFormatRead)
 	{
-		formatRead = true;
+		bFormatRead = true;
 
-		FString audioSetting;
-		if (!GConfig->GetString(TEXT("/Script/AndroidRuntimeSettings.AndroidRuntimeSettings"), TEXT("AndroidAudio"), audioSetting, GEngineIni))
+		FName AudioSetting;
 		{
-			audioSetting = TEXT("DEFAULT");
+			FString AudioSettingStr;
+			if (!GConfig->GetString(TEXT("/Script/AndroidRuntimeSettings.AndroidRuntimeSettings"), TEXT("AndroidAudio"), AudioSettingStr, GEngineIni))
+			{
+				AudioSetting = *AudioSettingStr;
+			}
 		}
 
 #if WITH_OGGVORBIS
-		if (audioSetting == TEXT("OGG") || audioSetting == TEXT("Default"))
+		if (AudioSetting == NAME_OGG || AudioSetting == NAME_None)
 		{
-			static FName NAME_OGG(TEXT("OGG"));
 			NAME_FORMAT = NAME_OGG;
 		}
 #else
-		if (audioSetting == TEXT("OGG"))
+		if (AudioSetting == NAME_OGG)
 		{
 			UE_LOG(LogAudio, Error, TEXT("Attemped to select Ogg Vorbis encoding when the cooker is built without Ogg Vorbis support."));
 		}
 #endif
 		else
 		{
-	
 			// Otherwise return ADPCM as it'll either be option '2' or 'default' depending on WITH_OGGVORBIS config
-			static FName NAME_ADPCM(TEXT("ADPCM"));
 			NAME_FORMAT = NAME_ADPCM;
 		}
 	}
+
+	if (Wave->IsSeekableStreaming())
+	{
+		return NAME_ADPCM;
+	}
+
 	return NAME_FORMAT;
 }
 
@@ -618,7 +627,7 @@ void FAndroidTargetPlatform::GetAllWaveFormats(TArray<FName>& OutFormats) const
 	static FName NAME_OGG(TEXT("OGG"));
 	static FName NAME_ADPCM(TEXT("ADPCM"));
 
-	OutFormats.Add(NAME_OGG); 
+	OutFormats.Add(NAME_OGG);
 	OutFormats.Add(NAME_ADPCM);
 }
 
@@ -646,7 +655,7 @@ namespace
 			{
 				*FoundSampleRate = RetrievedSampleRate;
 			}
-			
+
 		}
 		else
 		{
@@ -804,7 +813,7 @@ bool FAndroidTargetPlatform::HandleTicker( float DeltaTime )
 		FScopeLock ScopeLock(DeviceDetection->GetDeviceMapLock());
 
 		auto DeviceIt = DeviceDetection->GetDeviceMap().CreateConstIterator();
-		
+
 		for (; DeviceIt; ++DeviceIt)
 		{
 			ConnectedDeviceIds.Add(DeviceIt.Key());
diff --git a/Engine/Source/Developer/HTML5/HTML5TargetPlatform/Private/HTML5TargetPlatform.cpp b/Engine/Source/Developer/HTML5/HTML5TargetPlatform/Private/HTML5TargetPlatform.cpp
--- a/Engine/Source/Developer/HTML5/HTML5TargetPlatform/Private/HTML5TargetPlatform.cpp
+++ b/Engine/Source/Developer/HTML5/HTML5TargetPlatform/Private/HTML5TargetPlatform.cpp
@@ -10,6 +10,7 @@
 #include "Misc/ScopeLock.h"
 #if WITH_ENGINE
 #include "AudioCompressionSettings.h"
+#include "Sound/SoundWave.h"
 #endif
 
 #if PLATFORM_WINDOWS
@@ -287,7 +288,14 @@ const UTextureLODSettings& FHTML5TargetPlatform::GetTextureLODSettings() const
 
 FName FHTML5TargetPlatform::GetWaveFormat( const USoundWave* Wave ) const
 {
-	static FName NAME_OGG(TEXT("OGG"));
+	static const FName NAME_OGG(TEXT("OGG"));
+	static const FName NAME_ADPCM(TEXT("ADPCM"));
+
+	if (Wave->IsSeekableStreaming())
+	{
+		return NAME_ADPCM;
+	}
+
 	return NAME_OGG;
 }
 
@@ -295,7 +303,10 @@ FName FHTML5TargetPlatform::GetWaveFormat( const USoundWave* Wave ) const
 
 void FHTML5TargetPlatform::GetAllWaveFormats(TArray<FName>& OutFormats) const
 {
-	static FName NAME_OGG(TEXT("OGG"));
+	static const FName NAME_ADPCM(TEXT("ADPCM"));
+	static const FName NAME_OGG(TEXT("OGG"));
+
+	OutFormats.Add(NAME_ADPCM);
 	OutFormats.Add(NAME_OGG);
 }
 
diff --git a/Engine/Source/Developer/Linux/LinuxTargetPlatform/Private/LinuxTargetPlatform.h b/Engine/Source/Developer/Linux/LinuxTargetPlatform/Private/LinuxTargetPlatform.h
--- a/Engine/Source/Developer/Linux/LinuxTargetPlatform/Private/LinuxTargetPlatform.h
+++ b/Engine/Source/Developer/Linux/LinuxTargetPlatform/Private/LinuxTargetPlatform.h
@@ -36,7 +36,7 @@ class TLinuxTargetPlatform
 	: public TTargetPlatformBase<TProperties>
 {
 public:
-	
+
 	typedef TTargetPlatformBase<TProperties> TSuper;
 
 	/**
@@ -46,12 +46,12 @@ public:
 #if WITH_ENGINE
 		: bChangingDeviceConfig(false)
 #endif // WITH_ENGINE
-	{		
+	{
 #if PLATFORM_LINUX
 		// only add local device if actually running on Linux
 		LocalDevice = MakeShareable(new FLinuxTargetDevice(*this, FPlatformProcess::ComputerName(), nullptr));
 #endif
-	
+
 #if WITH_ENGINE
 		FConfigCacheIni::LoadLocalIniFile(EngineSettings, TEXT("Engine"), true, *this->PlatformName());
 		TextureLODSettings = nullptr;
@@ -91,13 +91,13 @@ public:
 		}
 
 		FTargetDeviceId UATFriendlyId(TEXT("Linux"), DeviceName);
-		Device = MakeShareable(new FLinuxTargetDevice(*this, DeviceName, 
+		Device = MakeShareable(new FLinuxTargetDevice(*this, DeviceName,
 #if WITH_ENGINE
 			[&]() { SaveDevicesToConfig(); }));
 		SaveDevicesToConfig();	// this will do the right thing even if AddDevice() was called from InitDevicesFromConfig
 #else
 			nullptr));
-#endif // WITH_ENGINE	
+#endif // WITH_ENGINE
 
 		DeviceDiscoveredEvent.Broadcast(Device.ToSharedRef());
 		return true;
@@ -179,7 +179,7 @@ public:
 			{
 				return true;
 			}
-			
+
 			// else check for legacy LINUX_ROOT
 			FString ToolchainCompiler = FPlatformMisc::GetEnvironmentVariable(TEXT("LINUX_ROOT"));
 			if (PLATFORM_WINDOWS)
@@ -312,8 +312,14 @@ public:
 
 	virtual FName GetWaveFormat( const class USoundWave* Wave ) const override
 	{
-		static FName NAME_OGG(TEXT("OGG"));
-		static FName NAME_OPUS(TEXT("OPUS"));
+		static const FName NAME_ADPCM(TEXT("ADPCM"));
+		static const FName NAME_OGG(TEXT("OGG"));
+		static const FName NAME_OPUS(TEXT("OPUS"));
+
+		if (Wave->IsSeekableStreaming())
+		{
+			return NAME_ADPCM;
+		}
 
 		if (Wave->IsStreaming())
 		{
@@ -325,8 +331,11 @@ public:
 
 	virtual void GetAllWaveFormats(TArray<FName>& OutFormats) const override
 	{
-		static FName NAME_OGG(TEXT("OGG"));
-		static FName NAME_OPUS(TEXT("OPUS"));
+		static const FName NAME_ADPCM(TEXT("ADPCM"));
+		static const FName NAME_OGG(TEXT("OGG"));
+		static const FName NAME_OPUS(TEXT("OPUS"));
+
+		OutFormats.Add(NAME_ADPCM);
 		OutFormats.Add(NAME_OGG);
 		OutFormats.Add(NAME_OPUS);
 	}
@@ -378,7 +387,7 @@ public:
 	{
 		return DeviceDiscoveredEvent;
 	}
-	
+
 	DECLARE_DERIVED_EVENT(TLinuxTargetPlatform, ITargetPlatform::FOnTargetDeviceLost, FOnTargetDeviceLost);
 	virtual FOnTargetDeviceLost& OnDeviceLost( ) override
 	{
@@ -401,7 +410,7 @@ protected:
 		}
 		bChangingDeviceConfig = true;
 
-		int NumDevices = 0;	
+		int NumDevices = 0;
 		for (;; ++NumDevices)
 		{
 			FString DeviceName, DeviceUser, DevicePass;
@@ -437,7 +446,7 @@ protected:
 				}
 			}
 		}
-		
+
 		bChangingDeviceConfig = false;
 	}
 
diff --git a/Engine/Source/Developer/Mac/MacTargetPlatform/Private/GenericMacTargetPlatform.h b/Engine/Source/Developer/Mac/MacTargetPlatform/Private/GenericMacTargetPlatform.h
--- a/Engine/Source/Developer/Mac/MacTargetPlatform/Private/GenericMacTargetPlatform.h
+++ b/Engine/Source/Developer/Mac/MacTargetPlatform/Private/GenericMacTargetPlatform.h
@@ -309,8 +309,14 @@ return TSuper::SupportsFeature(Feature);
 
 	virtual FName GetWaveFormat( const class USoundWave* Wave ) const override
 	{
-		static FName NAME_OGG(TEXT("OGG"));
-		static FName NAME_OPUS(TEXT("OPUS"));
+		static const FName NAME_ADPCM(TEXT("ADPCM"));
+		static const FName NAME_OGG(TEXT("OGG"));
+		static const FName NAME_OPUS(TEXT("OPUS"));
+
+		if (Wave->IsSeekableStreaming())
+		{
+			return NAME_ADPCM;
+		}
 
 		if (Wave->IsStreaming())
 		{
@@ -322,8 +328,11 @@ return TSuper::SupportsFeature(Feature);
 
 	virtual void GetAllWaveFormats(TArray<FName>& OutFormats) const override
 	{
-		static FName NAME_OGG(TEXT("OGG"));
-		static FName NAME_OPUS(TEXT("OPUS"));
+		static const FName NAME_ADPCM(TEXT("ADPCM"));
+		static const FName NAME_OGG(TEXT("OGG"));
+		static const FName NAME_OPUS(TEXT("OPUS"));
+
+		OutFormats.Add(NAME_ADPCM);
 		OutFormats.Add(NAME_OGG);
 		OutFormats.Add(NAME_OPUS);
 	}
diff --git a/Engine/Source/Developer/Windows/WindowsTargetPlatform/Private/GenericWindowsTargetPlatform.h b/Engine/Source/Developer/Windows/WindowsTargetPlatform/Private/GenericWindowsTargetPlatform.h
--- a/Engine/Source/Developer/Windows/WindowsTargetPlatform/Private/GenericWindowsTargetPlatform.h
+++ b/Engine/Source/Developer/Windows/WindowsTargetPlatform/Private/GenericWindowsTargetPlatform.h
@@ -37,7 +37,7 @@ public:
 		// only add local device if actually running on Windows
 		LocalDevice = MakeShareable(new TLocalPcTargetDevice<PLATFORM_64BITS>(*this));
 #endif
-		
+
 	#if WITH_ENGINE
 		FConfigCacheIni::LoadLocalIniFile(EngineSettings, TEXT("Engine"), true, *this->PlatformName());
 		TextureLODSettings = nullptr; // These are registered by the device profile system.
@@ -89,8 +89,8 @@ public:
 		static FName NAME_SF_VULKAN_ES31(TEXT("SF_VULKAN_ES31"));
 		static FName NAME_OPENGL_150_ES2(TEXT("GLSL_150_ES2"));
 		static FName NAME_OPENGL_150_ES3_1(TEXT("GLSL_150_ES31"));
-		bRequiresEncodedHDRReflectionCaptures =	TargetedShaderFormats.Contains(NAME_SF_VULKAN_ES31) 
-												 || TargetedShaderFormats.Contains(NAME_OPENGL_150_ES2) 
+		bRequiresEncodedHDRReflectionCaptures =	TargetedShaderFormats.Contains(NAME_SF_VULKAN_ES31)
+												 || TargetedShaderFormats.Contains(NAME_OPENGL_150_ES2)
 												 || TargetedShaderFormats.Contains(NAME_OPENGL_150_ES3_1);
 	#endif
 	}
@@ -237,7 +237,7 @@ public:
 		// Gather the list of Target RHIs and filter out any that may be invalid.
 		TArray<FName> PossibleShaderFormats;
 		GetAllPossibleShaderFormats(PossibleShaderFormats);
-		
+
 		for (int32 ShaderFormatIdx = TargetedShaderFormats.Num() - 1; ShaderFormatIdx >= 0; ShaderFormatIdx--)
 		{
 			FString ShaderFormat = TargetedShaderFormats[ShaderFormatIdx];
@@ -406,9 +406,16 @@ public:
 
 	virtual FName GetWaveFormat( const class USoundWave* Wave ) const override
 	{
-		static FName NAME_OGG(TEXT("OGG"));
-		static FName NAME_OPUS(TEXT("OPUS"));
-		
+		static const FName NAME_ADPCM(TEXT("ADPCM"));
+		static const FName NAME_OGG(TEXT("OGG"));
+		static const FName NAME_OPUS(TEXT("OPUS"));
+
+		// Seekable streams need to pick a codec which allows fixed-sized frames so we can compute stream chunk index to load
+		if (Wave->IsSeekableStreaming())
+		{
+			return NAME_ADPCM;
+		}
+
 #if !USE_VORBIS_FOR_STREAMING
 		if (Wave->IsStreaming())
 		{
@@ -421,8 +428,11 @@ public:
 
 	virtual void GetAllWaveFormats(TArray<FName>& OutFormats) const override
 	{
+		static const FName NAME_ADPCM(TEXT("ADPCM"));
 		static FName NAME_OGG(TEXT("OGG"));
 		static FName NAME_OPUS(TEXT("OPUS"));
+
+		OutFormats.Add(NAME_ADPCM);
 		OutFormats.Add(NAME_OGG);
 		OutFormats.Add(NAME_OPUS);
 	}
diff --git a/Engine/Source/Runtime/ALAudio/Private/ALAudioDevice.cpp b/Engine/Source/Runtime/ALAudio/Private/ALAudioDevice.cpp
--- a/Engine/Source/Runtime/ALAudio/Private/ALAudioDevice.cpp
+++ b/Engine/Source/Runtime/ALAudio/Private/ALAudioDevice.cpp
@@ -13,6 +13,7 @@
 #include "ALAudioDevice.h"
 #include "AudioEffect.h"
 #include "VorbisAudioInfo.h"
+#include "ADPCMAudioInfo.h"
 
 DEFINE_LOG_CATEGORY(LogALAudio);
 
@@ -334,6 +335,11 @@ bool FALAudioDevice::HasCompressedAudioInfoClass(USoundWave* SoundWave)
 
 class ICompressedAudioInfo* FALAudioDevice::CreateCompressedAudioInfo(USoundWave* SoundWave)
 {
+	if (SoundWave->IsSeekableStreaming())
+	{
+		return new FADPCMAudioInfo();
+	}
+
 #if WITH_OGGVORBIS
 	return new FVorbisAudioInfo();
 #else
diff --git a/Engine/Source/Runtime/Android/AndroidAudio/Private/AndroidAudioDevice.cpp b/Engine/Source/Runtime/Android/AndroidAudio/Private/AndroidAudioDevice.cpp
--- a/Engine/Source/Runtime/Android/AndroidAudio/Private/AndroidAudioDevice.cpp
+++ b/Engine/Source/Runtime/Android/AndroidAudio/Private/AndroidAudioDevice.cpp
@@ -36,8 +36,8 @@ DEFINE_STAT(STAT_AudioAndroidSourcePlayerRealize);
 void FSLESAudioDevice::Teardown( void )
 {
 	// Flush stops all sources and deletes all buffers so sources can be safely deleted below.
-	Flush( NULL );	
-	
+	Flush( NULL );
+
 	// Destroy all sound sources
 	for( int32 i = 0; i < Sources.Num(); i++ )
 	{
@@ -45,7 +45,7 @@ void FSLESAudioDevice::Teardown( void )
 	}
 
 	UE_LOG( LogAndroidAudio, Warning, TEXT("OpenSLES Tearing Down HW"));
-	
+
 	// Teardown OpenSLES..
 	// Destroy the SLES objects in reverse order of creation:
 	if (SL_OutputMixObject)
@@ -67,7 +67,7 @@ void FSLESAudioDevice::Teardown( void )
 /**
  * Initializes the audio device and creates sources.
  *
- * @warning: 
+ * @warning:
  *
  * @return TRUE if initialization was successful, FALSE otherwise
  */
@@ -78,9 +78,9 @@ bool FSLESAudioDevice::InitializeHardware( void )
 	SLresult result;
 
 	UE_LOG( LogAndroidAudio, Warning, TEXT("OpenSLES Initializing HW"));
-	
+
 	SLEngineOption EngineOption[] = { {(SLuint32) SL_ENGINEOPTION_THREADSAFE, (SLuint32) SL_BOOLEAN_TRUE} };
-	
+
     // create engine
     result = slCreateEngine( &SL_EngineObject, 1, EngineOption, 0, NULL, NULL);
     //check(SL_RESULT_SUCCESS == result);
@@ -88,55 +88,55 @@ bool FSLESAudioDevice::InitializeHardware( void )
 	{
 		UE_LOG( LogAndroidAudio, Error, TEXT("Engine create failed %d"), int32(result));
 	}
-	
+
     // realize the engine
     result = (*SL_EngineObject)->Realize(SL_EngineObject, SL_BOOLEAN_FALSE);
     check(SL_RESULT_SUCCESS == result);
-	
+
     // get the engine interface, which is needed in order to create other objects
     result = (*SL_EngineObject)->GetInterface(SL_EngineObject, SL_IID_ENGINE, &SL_EngineEngine);
     check(SL_RESULT_SUCCESS == result);
-	
-	// create output mix, with environmental reverb specified as a non-required interface    
+
+	// create output mix, with environmental reverb specified as a non-required interface
     result = (*SL_EngineEngine)->CreateOutputMix( SL_EngineEngine, &SL_OutputMixObject, 0, NULL, NULL );
     check(SL_RESULT_SUCCESS == result);
-	
+
     // realize the output mix
     result = (*SL_OutputMixObject)->Realize(SL_OutputMixObject, SL_BOOLEAN_FALSE);
     check(SL_RESULT_SUCCESS == result);
-	
+
 	UE_LOG( LogAndroidAudio, Warning, TEXT("OpenSLES Initialized"));
     // ignore unsuccessful result codes for env
 
-	
+
 	// Default to sensible channel count.
 	if( MaxChannels < 1 )
-	{  
+	{
 		MaxChannels = 12;
 	}
-	
-	
+
+
 	// Initialize channels.
 	for( int32 i = 0; i < FMath::Min( MaxChannels, 12 ); i++ )
 	{
-		FSLESSoundSource* Source = new FSLESSoundSource( this );		
+		FSLESSoundSource* Source = new FSLESSoundSource( this );
 		Sources.Add( Source );
 		FreeSources.Add( Source );
 	}
-	
+
 	if( Sources.Num() < 1 )
 	{
 		UE_LOG( LogAndroidAudio,  Warning, TEXT( "OpenSLAudio: couldn't allocate any sources" ) );
 		return false;
 	}
-	
+
 	// Update MaxChannels in case we couldn't create enough sources.
 	MaxChannels = Sources.Num();
 	UE_LOG( LogAndroidAudio, Warning, TEXT( "OpenSLAudioDevice: Allocated %i sources" ), MaxChannels );
-	
-	// Set up a default (nop) effects manager 
+
+	// Set up a default (nop) effects manager
 	Effects = new FAudioEffectsManager( this );
-	
+
 	return true;
 }
 
@@ -177,6 +177,11 @@ bool FSLESAudioDevice::HasCompressedAudioInfoClass(USoundWave* SoundWave)
 
 class ICompressedAudioInfo* FSLESAudioDevice::CreateCompressedAudioInfo(USoundWave* SoundWave)
 {
+	if (SoundWave->IsSeekableStreaming())
+	{
+		return new FADPCMAudioInfo();
+	}
+
 #if WITH_OGGVORBIS
 	static FName NAME_OGG(TEXT("OGG"));
 	if (SoundWave->bStreaming || SoundWave->HasCompressedData(NAME_OGG))
diff --git a/Engine/Source/Runtime/Android/AudioMixerAndroid/Private/AudioMixerPlatformAndroid.cpp b/Engine/Source/Runtime/Android/AudioMixerAndroid/Private/AudioMixerPlatformAndroid.cpp
--- a/Engine/Source/Runtime/Android/AudioMixerAndroid/Private/AudioMixerPlatformAndroid.cpp
+++ b/Engine/Source/Runtime/Android/AudioMixerAndroid/Private/AudioMixerPlatformAndroid.cpp
@@ -46,7 +46,7 @@ extern int32 AndroidThunkCpp_GetMetaDataInt(const FString& Key);
 #endif
 
 namespace Audio
-{	
+{
 	FMixerPlatformAndroid::FMixerPlatformAndroid()
 		: bSuspended(false)
 		, bInitialized(false)
@@ -55,7 +55,7 @@ namespace Audio
 	}
 
 	FMixerPlatformAndroid::~FMixerPlatformAndroid()
-	{		
+	{
 		if (bInitialized)
 		{
 			TeardownHardware();
@@ -93,30 +93,30 @@ namespace Audio
 		{
 			return false;
 		}
-					
+
 		SLresult Result;
 		SLEngineOption EngineOption[] = { {(SLuint32) SL_ENGINEOPTION_THREADSAFE, (SLuint32) SL_BOOLEAN_TRUE} };
-	
+
 		// Create engine
 		Result = slCreateEngine( &SL_EngineObject, 1, EngineOption, 0, NULL, NULL);
 		OPENSLES_CHECK_ON_FAIL(Result);
-	
+
 		// Realize the engine
 		Result = (*SL_EngineObject)->Realize(SL_EngineObject, SL_BOOLEAN_FALSE);
 		OPENSLES_CHECK_ON_FAIL(Result);
-	
+
 		// get the engine interface, which is needed in order to create other objects
 		Result = (*SL_EngineObject)->GetInterface(SL_EngineObject, SL_IID_ENGINE, &SL_EngineEngine);
 		OPENSLES_CHECK_ON_FAIL(Result);
 
-		// create output mix 
+		// create output mix
 		Result = (*SL_EngineEngine)->CreateOutputMix(SL_EngineEngine, &SL_OutputMixObject, 0, NULL, NULL );
 		OPENSLES_CHECK_ON_FAIL(Result);
 
 		// realize the output mix
 		Result = (*SL_OutputMixObject)->Realize(SL_OutputMixObject, SL_BOOLEAN_FALSE);
 		OPENSLES_CHECK_ON_FAIL(Result);
-	
+
 		bInitialized = true;
 
 		return true;
@@ -128,7 +128,7 @@ namespace Audio
 		{
 			return true;
 		}
-		
+
 		// Teardown OpenSLES..
 		// Destroy the SLES objects in reverse order of creation:
 		if (SL_OutputMixObject)
@@ -146,7 +146,7 @@ namespace Audio
 		}
 
 		bInitialized = false;
-		
+
 		return true;
 	}
 
@@ -211,20 +211,20 @@ namespace Audio
 
 		// data info
 		SLDataLocator_AndroidSimpleBufferQueue LocationBuffer = { SL_DATALOCATOR_ANDROIDSIMPLEBUFFERQUEUE, 1};
-		
+
 		// PCM Info
 		SLDataFormat_PCM PCM_Format = {
-			SL_DATAFORMAT_PCM, 
+			SL_DATAFORMAT_PCM,
 			(SLuint32)AudioStreamInfo.DeviceInfo.NumChannels,
-			(SLuint32)(AudioStreamInfo.DeviceInfo.SampleRate * 1000), // NOTE: OpenSLES has sample rates specified in millihertz. 
-			SL_PCMSAMPLEFORMAT_FIXED_16, 
-			SL_PCMSAMPLEFORMAT_FIXED_16, 
-			SL_SPEAKER_FRONT_LEFT | SL_SPEAKER_FRONT_RIGHT, 
-			SL_BYTEORDER_LITTLEENDIAN 
+			(SLuint32)(AudioStreamInfo.DeviceInfo.SampleRate * 1000), // NOTE: OpenSLES has sample rates specified in millihertz.
+			SL_PCMSAMPLEFORMAT_FIXED_16,
+			SL_PCMSAMPLEFORMAT_FIXED_16,
+			SL_SPEAKER_FRONT_LEFT | SL_SPEAKER_FRONT_RIGHT,
+			SL_BYTEORDER_LITTLEENDIAN
 		};
-		
+
 		SLDataSource SoundDataSource = { &LocationBuffer, &PCM_Format };
-		
+
 		// configure audio sink
 		SLDataLocator_OutputMix OutputMix = { SL_DATALOCATOR_OUTPUTMIX, SL_OutputMixObject };
 		SLDataSink AudioSink = { &OutputMix, nullptr };
@@ -239,7 +239,7 @@ namespace Audio
 		// realize the player
 		Result = (*SL_PlayerObject)->Realize(SL_PlayerObject, SL_BOOLEAN_FALSE);
 		OPENSLES_RETURN_ON_FAIL(Result);
-		
+
 		// get the play interface
 		Result = (*SL_PlayerObject)->GetInterface(SL_PlayerObject, SL_IID_PLAY, &SL_PlayerPlayInterface);
 		OPENSLES_RETURN_ON_FAIL(Result);
@@ -262,17 +262,17 @@ namespace Audio
 		{
 			return false;
 		}
-		
+
 		SLresult Result =(*SL_PlayerBufferQueue)->RegisterCallback(SL_PlayerBufferQueue, nullptr, nullptr);
 
-		(*SL_PlayerObject)->Destroy(SL_PlayerObject);			
+		(*SL_PlayerObject)->Destroy(SL_PlayerObject);
 
 		SL_PlayerObject = nullptr;
 		SL_PlayerPlayInterface = nullptr;
 		SL_PlayerBufferQueue = nullptr;
 
 		AudioStreamInfo.StreamState = EAudioOutputStreamState::Closed;
-		
+
 		return true;
 	}
 
@@ -293,7 +293,7 @@ namespace Audio
 		{
 			return false;
 		}
-		
+
 		if (AudioStreamInfo.StreamState != EAudioOutputStreamState::Stopped)
 		{
 			// set the player's state to stopped
@@ -307,7 +307,7 @@ namespace Audio
 
 			check(AudioStreamInfo.StreamState == EAudioOutputStreamState::Stopped);
 		}
-	
+
 		return true;
 	}
 
@@ -380,9 +380,16 @@ namespace Audio
 		SLresult Result = (*SL_PlayerBufferQueue)->Enqueue(SL_PlayerBufferQueue, Buffer, BufferSize);
 		OPENSLES_LOG_ON_FAIL(Result);
 	}
-	
+
 	FName FMixerPlatformAndroid::GetRuntimeFormat(USoundWave* InSoundWave)
 	{
+		static FName NAME_ADPCM(TEXT("ADPCM"));
+
+		if (InSoundWave->IsSeekableStreaming())
+		{
+			return NAME_ADPCM;
+		}
+
 #if WITH_OGGVORBIS
 		static FName NAME_OGG(TEXT("OGG"));
 		if (InSoundWave->HasCompressedData(NAME_OGG))
@@ -390,57 +397,25 @@ namespace Audio
 			return NAME_OGG;
 		}
 #endif
-
-		static FName NAME_ADPCM(TEXT("ADPCM"));
-
 		return NAME_ADPCM;
 	}
 
 	bool FMixerPlatformAndroid::HasCompressedAudioInfoClass(USoundWave* InSoundWave)
 	{
-#if WITH_OGGVORBIS
-		if (InSoundWave->bStreaming)
-		{
-			return true;
-		}
-
-		static FName NAME_OGG(TEXT("OGG"));
-		if (InSoundWave->HasCompressedData(NAME_OGG))
-		{
-			return true;
-		}
-#endif
-
-		if (InSoundWave->bStreaming)
-		{
-			return true;
-		}
-
-		static FName NAME_ADPCM(TEXT("ADPCM"));
-		if (InSoundWave->HasCompressedData(NAME_ADPCM))
-		{
-			return true;
-		}
-
-		return false;
+		return true;
 	}
 
 	ICompressedAudioInfo* FMixerPlatformAndroid::CreateCompressedAudioInfo(USoundWave* InSoundWave)
 	{
-#if WITH_OGGVORBIS
 		static FName NAME_OGG(TEXT("OGG"));
-		if (InSoundWave->bStreaming || InSoundWave->HasCompressedData(NAME_OGG))
-		{
-			return new FVorbisAudioInfo();
-		}
-#endif
 		static FName NAME_ADPCM(TEXT("ADPCM"));
-		if (InSoundWave->bStreaming || InSoundWave->HasCompressedData(NAME_ADPCM))
+
+		if (InSoundWave->IsSeekableStreaming())
 		{
 			return new FADPCMAudioInfo();
 		}
 
-		return nullptr;
+		return new FVorbisAudioInfo();
 	}
 
 	FString FMixerPlatformAndroid::GetDefaultDeviceName()
diff --git a/Engine/Source/Runtime/Apple/AudioMixerAudioUnit/Private/AudioMixerPlatformAudioUnit.cpp b/Engine/Source/Runtime/Apple/AudioMixerAudioUnit/Private/AudioMixerPlatformAudioUnit.cpp
--- a/Engine/Source/Runtime/Apple/AudioMixerAudioUnit/Private/AudioMixerPlatformAudioUnit.cpp
+++ b/Engine/Source/Runtime/Apple/AudioMixerAudioUnit/Private/AudioMixerPlatformAudioUnit.cpp
@@ -14,10 +14,11 @@
 
 #include "VorbisAudioInfo.h"
 #include "OpusAudioInfo.h"
+#include "ADPCMAudioInfo.h"
 #endif // #if PLATFORM_IOS || PLATFORM_TVOS
 /*
-	This implementation only depends on the audio units API which allows it to run on MacOS, iOS and tvOS. 
-	
+	This implementation only depends on the audio units API which allows it to run on MacOS, iOS and tvOS.
+
 	For now just assume an iOS configuration (only 2 left and right channels on a single device)
 */
 
@@ -40,22 +41,22 @@ DEFINE_LOG_CATEGORY(LogAudioMixerAudioUnit);
 namespace Audio
 {
 #if PLATFORM_IOS || PLATFORM_TVOS
-    static const int32 DefaultBufferSize = 4096;
+	static const int32 DefaultBufferSize = 4096;
 #else
-    static const int32 DefaultBufferSize = 1024;
-    static const int32 AUBufferSize = 256;
+	static const int32 DefaultBufferSize = 1024;
+	static const int32 AUBufferSize = 256;
 #endif //#if PLATFORM_IOS || PLATFORM_TVOS
-    static const double DefaultSampleRate = 48000.0;
-    
-    static int32 SuspendCounter = 0;
-    
+	static const double DefaultSampleRate = 48000.0;
+
+	static int32 SuspendCounter = 0;
+
 	FMixerPlatformAudioUnit::FMixerPlatformAudioUnit()
 		: bInitialized(false)
-        , bInCallback(false)
-        , SubmittedBufferPtr(nullptr)
-        , RemainingBytesInCurrentSubmittedBuffer(0)
-        , BytesPerSubmittedBuffer(0)
-        , GraphSampleRate(DefaultSampleRate)
+		, bInCallback(false)
+		, SubmittedBufferPtr(nullptr)
+		, RemainingBytesInCurrentSubmittedBuffer(0)
+		, BytesPerSubmittedBuffer(0)
+		, GraphSampleRate(DefaultSampleRate)
 	{
 	}
 
@@ -66,105 +67,105 @@ namespace Audio
 			TeardownHardware();
 		}
 	}
-    
-    int32 FMixerPlatformAudioUnit::GetNumFrames(const int32 InNumReqestedFrames)
-    {
+
+	int32 FMixerPlatformAudioUnit::GetNumFrames(const int32 InNumReqestedFrames)
+	{
 #if PLATFORM_IOS || PLATFORM_TVOS
-        AVAudioSession* AudioSession = [AVAudioSession sharedInstance];
-        double BufferSizeInSec = [AudioSession preferredIOBufferDuration];
-        double SampleRate = [AudioSession preferredSampleRate];
-        
-        if (BufferSizeInSec == 0.0)
-        {
-            return DefaultBufferSize;
-        }
-        
-        int32 NumFrames = (int32)(SampleRate * BufferSizeInSec);
-        
-        return NumFrames;
+		AVAudioSession* AudioSession = [AVAudioSession sharedInstance];
+		double BufferSizeInSec = [AudioSession preferredIOBufferDuration];
+		double SampleRate = [AudioSession preferredSampleRate];
+
+		if (BufferSizeInSec == 0.0)
+		{
+			return DefaultBufferSize;
+		}
+
+		int32 NumFrames = (int32)(SampleRate * BufferSizeInSec);
+
+		return NumFrames;
 #else
-       //On MacOS, we hardcode buffer sizes.
-        return DefaultBufferSize;
+	   //On MacOS, we hardcode buffer sizes.
+		return DefaultBufferSize;
 #endif
-    }
+	}
 
 	bool FMixerPlatformAudioUnit::InitializeHardware()
 	{
-        if (bInitialized)
+		if (bInitialized)
 		{
 			return false;
 		}
-        
+
 		OSStatus Status;
-        GraphSampleRate = (double) AudioStreamInfo.DeviceInfo.SampleRate;
-        UInt32 BufferSize = (UInt32) GetNumFrames(OpenStreamParams.NumFrames);
-        const int32 NumChannels = 2;
-
-        if (GraphSampleRate == 0)
-        {
-            GraphSampleRate = DefaultSampleRate;
-        }
-        
-        if (BufferSize == 0)
-        {
-            BufferSize = DefaultBufferSize;
-        }
-        
-        BytesPerSubmittedBuffer = BufferSize * NumChannels * sizeof(float);
-        check(BytesPerSubmittedBuffer != 0);
-        UE_LOG(LogAudioMixerAudioUnit, Display, TEXT("Bytes per submitted buffer: %d"), BytesPerSubmittedBuffer);
-        
+		GraphSampleRate = (double) AudioStreamInfo.DeviceInfo.SampleRate;
+		UInt32 BufferSize = (UInt32) GetNumFrames(OpenStreamParams.NumFrames);
+		const int32 NumChannels = 2;
+
+		if (GraphSampleRate == 0)
+		{
+			GraphSampleRate = DefaultSampleRate;
+		}
+
+		if (BufferSize == 0)
+		{
+			BufferSize = DefaultBufferSize;
+		}
+
+		BytesPerSubmittedBuffer = BufferSize * NumChannels * sizeof(float);
+		check(BytesPerSubmittedBuffer != 0);
+		UE_LOG(LogAudioMixerAudioUnit, Display, TEXT("Bytes per submitted buffer: %d"), BytesPerSubmittedBuffer);
+
 #if PLATFORM_IOS || PLATFORM_TVOS
-        NSError* error;
-        
-        AVAudioSession* AudioSession = [AVAudioSession sharedInstance];
-
-        GraphSampleRate = [AudioSession preferredSampleRate];
-        bool Success = [AudioSession setPreferredSampleRate:GraphSampleRate error:&error];
-        
-        if (!Success)
-        {
-            UE_LOG(LogAudioMixerAudioUnit, Display, TEXT("Error setting sample rate."));
-        }
-        
-        // By calling setPreferredIOBufferDuration, we indicate that we would prefer that the buffer size not change if possible.
-        Success = [AudioSession setPreferredIOBufferDuration:[AudioSession preferredIOBufferDuration] error: &error];
-        
-        UE_LOG(LogAudioMixerAudioUnit, Display, TEXT("Device Sample Rate: %f"), GraphSampleRate);
-        check(GraphSampleRate != 0);
-        
-        Success = [AudioSession setActive:true error:&error];
-        
-        if (!Success)
-        {
-            UE_LOG(LogAudioMixerAudioUnit, Display, TEXT("Error starting audio session."));
-        }
+		NSError* error;
+
+		AVAudioSession* AudioSession = [AVAudioSession sharedInstance];
+
+		GraphSampleRate = [AudioSession preferredSampleRate];
+		bool Success = [AudioSession setPreferredSampleRate:GraphSampleRate error:&error];
+
+		if (!Success)
+		{
+			UE_LOG(LogAudioMixerAudioUnit, Display, TEXT("Error setting sample rate."));
+		}
+
+		// By calling setPreferredIOBufferDuration, we indicate that we would prefer that the buffer size not change if possible.
+		Success = [AudioSession setPreferredIOBufferDuration:[AudioSession preferredIOBufferDuration] error: &error];
+
+		UE_LOG(LogAudioMixerAudioUnit, Display, TEXT("Device Sample Rate: %f"), GraphSampleRate);
+		check(GraphSampleRate != 0);
+
+		Success = [AudioSession setActive:true error:&error];
+
+		if (!Success)
+		{
+			UE_LOG(LogAudioMixerAudioUnit, Display, TEXT("Error starting audio session."));
+		}
 #else
-        AudioObjectID DeviceAudioObjectID;
-        AudioObjectPropertyAddress DevicePropertyAddress;
-        UInt32 AudioDeviceQuerySize;
-        
-        //Get Audio Device ID- this will be used throughout initialization to query the audio hardware.
-        DevicePropertyAddress.mSelector = kAudioHardwarePropertyDefaultOutputDevice;
-        DevicePropertyAddress.mScope = kAudioObjectPropertyScopeGlobal;
-        DevicePropertyAddress.mElement = 0;
-        AudioDeviceQuerySize = sizeof(AudioDeviceID);
-        Status = AudioObjectGetPropertyData(kAudioObjectSystemObject, &DevicePropertyAddress, 0, nullptr, &AudioDeviceQuerySize, &DeviceAudioObjectID);
-        
-        if(Status != 0)
-        {
-            UE_LOG(LogAudioMixerAudioUnit, Display, TEXT("ERROR setting sample rate to %f"), GraphSampleRate);
-        }
-        
-        Status = AudioObjectGetPropertyData(DeviceAudioObjectID, &DevicePropertyAddress, 0, nullptr, &AudioDeviceQuerySize, &GraphSampleRate);
-        
-        if(Status == 0)
-        {
-            UE_LOG(LogAudioMixerAudioUnit, Display, TEXT("Sample Rate: %f"), GraphSampleRate);
-        }
+		AudioObjectID DeviceAudioObjectID;
+		AudioObjectPropertyAddress DevicePropertyAddress;
+		UInt32 AudioDeviceQuerySize;
+
+		//Get Audio Device ID- this will be used throughout initialization to query the audio hardware.
+		DevicePropertyAddress.mSelector = kAudioHardwarePropertyDefaultOutputDevice;
+		DevicePropertyAddress.mScope = kAudioObjectPropertyScopeGlobal;
+		DevicePropertyAddress.mElement = 0;
+		AudioDeviceQuerySize = sizeof(AudioDeviceID);
+		Status = AudioObjectGetPropertyData(kAudioObjectSystemObject, &DevicePropertyAddress, 0, nullptr, &AudioDeviceQuerySize, &DeviceAudioObjectID);
+
+		if(Status != 0)
+		{
+			UE_LOG(LogAudioMixerAudioUnit, Display, TEXT("ERROR setting sample rate to %f"), GraphSampleRate);
+		}
+
+		Status = AudioObjectGetPropertyData(DeviceAudioObjectID, &DevicePropertyAddress, 0, nullptr, &AudioDeviceQuerySize, &GraphSampleRate);
+
+		if(Status == 0)
+		{
+			UE_LOG(LogAudioMixerAudioUnit, Display, TEXT("Sample Rate: %f"), GraphSampleRate);
+		}
 
 #endif // #if PLATFORM_IOS || PLATFORM_TVOS
-        
+
 
 		// Linear PCM stream format
 		OutputFormat.mFormatID         = kAudioFormatLinearPCM;
@@ -188,11 +189,11 @@ namespace Audio
 		// Setup audio output unit
 		UnitDescription.componentType         = kAudioUnitType_Output;
 #if PLATFORM_IOS || PLATFORM_TVOS
-        //On iOS, we'll use the RemoteIO AudioUnit.
+		//On iOS, we'll use the RemoteIO AudioUnit.
 		UnitDescription.componentSubType      = kAudioUnitSubType_RemoteIO;
 #else //PLATFORM_MAC
-        //On MacOS, we'll use the DefaultOutput AudioUnit.
-        UnitDescription.componentSubType      = kAudioUnitSubType_DefaultOutput;
+		//On MacOS, we'll use the DefaultOutput AudioUnit.
+		UnitDescription.componentSubType      = kAudioUnitSubType_DefaultOutput;
 #endif // #if PLATFORM_IOS || PLATFORM_TVOS
 		UnitDescription.componentManufacturer = kAudioUnitManufacturer_Apple;
 		UnitDescription.componentFlags        = 0;
@@ -203,14 +204,14 @@ namespace Audio
 			HandleError(TEXT("Failed to initialize audio output node!"), true);
 			return false;
 		}
-		
+
 		Status = AUGraphOpen(AudioUnitGraph);
 		if (Status != noErr)
 		{
 			HandleError(TEXT("Failed to open audio unit graph"), true);
 			return false;
 		}
-		
+
 		Status = AUGraphNodeInfo(AudioUnitGraph, OutputNode, nullptr, &OutputUnit);
 		if (Status != noErr)
 		{
@@ -219,11 +220,11 @@ namespace Audio
 		}
 
 		Status = AudioUnitSetProperty(OutputUnit,
-		                              kAudioUnitProperty_StreamFormat,
-		                              kAudioUnitScope_Input,
-		                              0,
-		                              &OutputFormat,
-		                              sizeof(AudioStreamBasicDescription));
+									  kAudioUnitProperty_StreamFormat,
+									  kAudioUnitScope_Input,
+									  0,
+									  &OutputFormat,
+									  sizeof(AudioStreamBasicDescription));
 		if (Status != noErr)
 		{
 			HandleError(TEXT("Failed to set output format!"), true);
@@ -231,35 +232,35 @@ namespace Audio
 		}
 
 #if PLATFORM_MAC
-        DevicePropertyAddress.mSelector = kAudioDevicePropertyBufferFrameSize;
-        DevicePropertyAddress.mScope = kAudioObjectPropertyScopeGlobal;
-        DevicePropertyAddress.mElement = 0;
-        AudioDeviceQuerySize = sizeof(AUBufferSize);
-        Status = AudioObjectSetPropertyData(DeviceAudioObjectID, &DevicePropertyAddress, 0, nullptr, AudioDeviceQuerySize, &AUBufferSize);
-        if(Status != 0)
-        {
-            HandleError(TEXT("Failed to set output format!"), true);
-            return false;
-        }
-        
+		DevicePropertyAddress.mSelector = kAudioDevicePropertyBufferFrameSize;
+		DevicePropertyAddress.mScope = kAudioObjectPropertyScopeGlobal;
+		DevicePropertyAddress.mElement = 0;
+		AudioDeviceQuerySize = sizeof(AUBufferSize);
+		Status = AudioObjectSetPropertyData(DeviceAudioObjectID, &DevicePropertyAddress, 0, nullptr, AudioDeviceQuerySize, &AUBufferSize);
+		if(Status != 0)
+		{
+			HandleError(TEXT("Failed to set output format!"), true);
+			return false;
+		}
+
 #endif //#if PLATFORM_MAC
-        
-        AudioStreamInfo.NumOutputFrames = BufferSize;
 
-        AudioStreamInfo.DeviceInfo = GetPlatformDeviceInfo();
-        
+		AudioStreamInfo.NumOutputFrames = BufferSize;
+
+		AudioStreamInfo.DeviceInfo = GetPlatformDeviceInfo();
+
 		AURenderCallbackStruct InputCallback;
 		InputCallback.inputProc = &AudioRenderCallback;
 		InputCallback.inputProcRefCon = this;
 		Status = AUGraphSetNodeInputCallback(AudioUnitGraph,
-		                                     OutputNode,
-		                                     0,
-		                                     &InputCallback);
+											 OutputNode,
+											 0,
+											 &InputCallback);
 		UE_CLOG(Status != noErr, LogAudioMixerAudioUnit, Error, TEXT("Failed to set input callback for audio output node"));
 
-        OpenStreamParams.NumFrames = BufferSize;
+		OpenStreamParams.NumFrames = BufferSize;
 		AudioStreamInfo.StreamState = EAudioOutputStreamState::Closed;
-        
+
 		bInitialized = true;
 
 		return true;
@@ -277,7 +278,7 @@ namespace Audio
 		{
 			return true;
 		}
-		
+
 		StopAudioStream();
 		CloseAudioStream();
 
@@ -288,7 +289,7 @@ namespace Audio
 		OutputUnit = nullptr;
 
 		bInitialized = false;
-		
+
 		return true;
 	}
 
@@ -300,7 +301,7 @@ namespace Audio
 	bool FMixerPlatformAudioUnit::GetNumOutputDevices(uint32& OutNumOutputDevices)
 	{
 		OutNumOutputDevices = 1;
-		
+
 		return true;
 	}
 
@@ -313,7 +314,7 @@ namespace Audio
 	bool FMixerPlatformAudioUnit::GetDefaultOutputDeviceIndex(uint32& OutDefaultDeviceIndex) const
 	{
 		OutDefaultDeviceIndex = 0;
-		
+
 		return true;
 	}
 
@@ -323,12 +324,12 @@ namespace Audio
 		{
 			return false;
 		}
-		
+
 		AudioStreamInfo.OutputDeviceIndex = Params.OutputDeviceIndex;
 		AudioStreamInfo.AudioMixer = Params.AudioMixer;
-		
+
 		OpenStreamParams = Params;
-		
+
 		// Initialize the audio unit graph
 		OSStatus Status = AUGraphInitialize(AudioUnitGraph);
 		if (Status != noErr)
@@ -348,9 +349,9 @@ namespace Audio
 		{
 			return false;
 		}
-		
+
 		AudioStreamInfo.StreamState = EAudioOutputStreamState::Closed;
-		
+
 		return true;
 	}
 
@@ -360,9 +361,9 @@ namespace Audio
 		{
 			return false;
 		}
-        
+
 		BeginGeneratingAudio();
-		
+
 		// This will start the render audio callback
 		OSStatus Status = AUGraphStart(AudioUnitGraph);
 		if (Status != noErr)
@@ -380,11 +381,11 @@ namespace Audio
 		{
 			return false;
 		}
-		
+
 		if (AudioStreamInfo.StreamState != EAudioOutputStreamState::Stopped && AudioStreamInfo.StreamState != EAudioOutputStreamState::Closed)
 		{
 			AUGraphStop(AudioUnitGraph);
-			
+
 			if (AudioStreamInfo.StreamState == EAudioOutputStreamState::Running)
 			{
 				StopGeneratingAudio();
@@ -393,59 +394,65 @@ namespace Audio
 			check(AudioStreamInfo.StreamState == EAudioOutputStreamState::Stopped);
 		}
 
-		
+
 		return true;
 	}
 
 	bool FMixerPlatformAudioUnit::MoveAudioStreamToNewAudioDevice(const FString& InNewDeviceId)
 	{
-        //TODO
-        
+		//TODO
 		return false;
 	}
 
 	FAudioPlatformDeviceInfo FMixerPlatformAudioUnit::GetPlatformDeviceInfo() const
 	{
-        FAudioPlatformDeviceInfo DeviceInfo;
-        
-    #if PLATFORM_IOS || PLATFORM_TVOS
-        AVAudioSession* AudioSession = [AVAudioSession sharedInstance];
-        double SampleRate = [AudioSession preferredSampleRate];
-        DeviceInfo.SampleRate = (int32)SampleRate;
+		FAudioPlatformDeviceInfo DeviceInfo;
+
+	#if PLATFORM_IOS || PLATFORM_TVOS
+		AVAudioSession* AudioSession = [AVAudioSession sharedInstance];
+		double SampleRate = [AudioSession preferredSampleRate];
+		DeviceInfo.SampleRate = (int32)SampleRate;
 #else
-        DeviceInfo.SampleRate = GraphSampleRate;
+		DeviceInfo.SampleRate = GraphSampleRate;
 #endif
-        DeviceInfo.NumChannels = 2;
-        DeviceInfo.Format = EAudioMixerStreamDataFormat::Float;
-        DeviceInfo.OutputChannelArray.SetNum(2);
-        DeviceInfo.OutputChannelArray[0] = EAudioMixerChannel::FrontLeft;
-        DeviceInfo.OutputChannelArray[1] = EAudioMixerChannel::FrontRight;
-        DeviceInfo.bIsSystemDefault = true;
-        
-        return DeviceInfo;
+		DeviceInfo.NumChannels = 2;
+		DeviceInfo.Format = EAudioMixerStreamDataFormat::Float;
+		DeviceInfo.OutputChannelArray.SetNum(2);
+		DeviceInfo.OutputChannelArray[0] = EAudioMixerChannel::FrontLeft;
+		DeviceInfo.OutputChannelArray[1] = EAudioMixerChannel::FrontRight;
+		DeviceInfo.bIsSystemDefault = true;
+
+		return DeviceInfo;
 	}
 
 	void FMixerPlatformAudioUnit::SubmitBuffer(const uint8* Buffer)
 	{
-        SubmittedBufferPtr = (uint8*) Buffer;
-        SubmittedBytes = 0;
-        RemainingBytesInCurrentSubmittedBuffer = BytesPerSubmittedBuffer;
+		SubmittedBufferPtr = (uint8*) Buffer;
+		SubmittedBytes = 0;
+		RemainingBytesInCurrentSubmittedBuffer = BytesPerSubmittedBuffer;
 	}
 
 	FName FMixerPlatformAudioUnit::GetRuntimeFormat(USoundWave* InSoundWave)
 	{
+		static const FName NAME_ADPCM(TEXT("ADPCM"));
+		static const FName NAME_OGG(TEXT("OGG"));
+		static const FName NAME_OPUS(TEXT("OPUS"));
+
 #if PLATFORM_IOS || PLATFORM_TVOS
-		static FName NAME_ADPCM(TEXT("ADPCM"));
 		return NAME_ADPCM;
 #else
-        static FName NAME_OPUS(TEXT("OPUS"));
-        
-        if (InSoundWave->IsStreaming())
-        {
-            return NAME_OPUS;
-        }
-        static FName NAME_OGG(TEXT("OGG"));
-        return NAME_OGG;
+
+		if (InSoundWave->IsSeekableStreaming())
+		{
+			return NAME_ADPCM;
+		}
+
+		if (InSoundWave->IsStreaming())
+		{
+			return NAME_OPUS;
+		}
+
+		return NAME_OGG;
 #endif
 	}
 
@@ -459,31 +466,36 @@ namespace Audio
 #if PLATFORM_IOS || PLATFORM_TVOS
 		return new FADPCMAudioInfo();
 #else
-        check(InSoundWave);
-        
-        if (InSoundWave->IsStreaming())
-        {
-            return new FOpusAudioInfo();
-        }
-        
+		check(InSoundWave);
+
+		if (InSoundWave->IsSeekableStreaming())
+		{
+			return new FADPCMAudioInfo();
+		}
+
+		if (InSoundWave->IsStreaming())
+		{
+			return new FOpusAudioInfo();
+		}
+
 #if WITH_OGGVORBIS
-        static const FName NAME_OGG(TEXT("OGG"));
-        if (FPlatformProperties::RequiresCookedData() ? InSoundWave->HasCompressedData(NAME_OGG) : (InSoundWave->GetCompressedData(NAME_OGG) != nullptr))
-        {
-            ICompressedAudioInfo* CompressedInfo = new FVorbisAudioInfo();
-            if (!CompressedInfo)
-            {
-                UE_LOG(LogAudio, Error, TEXT("Failed to create new FVorbisAudioInfo for SoundWave %s: out of memory."), *InSoundWave->GetName());
-                return nullptr;
-            }
-            return CompressedInfo;
-        }
-        else
-        {
-            return nullptr;
-        }
+		static const FName NAME_OGG(TEXT("OGG"));
+		if (FPlatformProperties::RequiresCookedData() ? InSoundWave->HasCompressedData(NAME_OGG) : (InSoundWave->GetCompressedData(NAME_OGG) != nullptr))
+		{
+			ICompressedAudioInfo* CompressedInfo = new FVorbisAudioInfo();
+			if (!CompressedInfo)
+			{
+				UE_LOG(LogAudio, Error, TEXT("Failed to create new FVorbisAudioInfo for SoundWave %s: out of memory."), *InSoundWave->GetName());
+				return nullptr;
+			}
+			return CompressedInfo;
+		}
+		else
+		{
+			return nullptr;
+		}
 #else
-        return nullptr;
+		return nullptr;
 #endif // WITH_OGGVORBIS
 #endif // PLATFORM_IOS || PLATFORM_TVOS
 	}
@@ -495,52 +507,52 @@ namespace Audio
 
 	FAudioPlatformSettings FMixerPlatformAudioUnit::GetPlatformSettings() const
 	{
-        FAudioPlatformSettings Settings;
-        Settings.NumBuffers = 2;
+		FAudioPlatformSettings Settings;
+		Settings.NumBuffers = 2;
 #if PLATFORM_IOS || PLATFORM_TVOS
-        AVAudioSession* AudioSession = [AVAudioSession sharedInstance];
-        double BufferSizeInSec = [AudioSession preferredIOBufferDuration];
-        double SampleRate = [AudioSession preferredSampleRate];
-        
-        int32 NumFrames;
-        
-        if (BufferSizeInSec == 0.0)
-        {
-            NumFrames = DefaultBufferSize;
-        }
-        else
-        {
-            NumFrames = (int32)(SampleRate * BufferSizeInSec);
-        }
-        Settings.CallbackBufferFrameSize = NumFrames;
-        Settings.SampleRate = SampleRate;
-        Settings.MaxChannels = 32;
-        
+		AVAudioSession* AudioSession = [AVAudioSession sharedInstance];
+		double BufferSizeInSec = [AudioSession preferredIOBufferDuration];
+		double SampleRate = [AudioSession preferredSampleRate];
+
+		int32 NumFrames;
+
+		if (BufferSizeInSec == 0.0)
+		{
+			NumFrames = DefaultBufferSize;
+		}
+		else
+		{
+			NumFrames = (int32)(SampleRate * BufferSizeInSec);
+		}
+		Settings.CallbackBufferFrameSize = NumFrames;
+		Settings.SampleRate = SampleRate;
+		Settings.MaxChannels = 32;
+
 #else
-        Settings.SampleRate = GraphSampleRate;
-        Settings.CallbackBufferFrameSize = DefaultBufferSize;
-        
+		Settings.SampleRate = GraphSampleRate;
+		Settings.CallbackBufferFrameSize = DefaultBufferSize;
+
 #endif //#if PLATFORM_IOS || PLATFORM_TVOS
-        
-        return Settings;
+
+		return Settings;
 	}
 
 	void FMixerPlatformAudioUnit::ResumeContext()
 	{
-        if (SuspendCounter > 0)
-        {
-            FPlatformAtomics::InterlockedDecrement(&SuspendCounter);
-            AUGraphStart(AudioUnitGraph);
-            UE_LOG(LogAudioMixerAudioUnit, Display, TEXT("Resuming Audio"));
-            bSuspended = false;
-        }
+		if (SuspendCounter > 0)
+		{
+			FPlatformAtomics::InterlockedDecrement(&SuspendCounter);
+			AUGraphStart(AudioUnitGraph);
+			UE_LOG(LogAudioMixerAudioUnit, Display, TEXT("Resuming Audio"));
+			bSuspended = false;
+		}
 	}
-	
+
 	void FMixerPlatformAudioUnit::SuspendContext()
 	{
 		if (SuspendCounter == 0)
 		{
-            FPlatformAtomics::InterlockedIncrement(&SuspendCounter);
+			FPlatformAtomics::InterlockedIncrement(&SuspendCounter);
 			AUGraphStop(AudioUnitGraph);
 			UE_LOG(LogAudioMixerAudioUnit, Display, TEXT("Suspending Audio"));
 			bSuspended = true;
@@ -559,47 +571,47 @@ namespace Audio
 	bool FMixerPlatformAudioUnit::PerformCallback(AudioBufferList* OutputBufferData)
 	{
 		bInCallback = true;
-		
+
 		if (AudioStreamInfo.StreamState == EAudioOutputStreamState::Running)
 		{
-            if (!SubmittedBufferPtr)
-            {
-                ReadNextBuffer();
-            }
-
-            // How many bytes we have left over from previous callback
-            int32 SubmittedBufferBytesLeft = BytesPerSubmittedBuffer - SubmittedBytes;
-            int32 OutputBufferBytesLeft = OutputBufferData->mBuffers[0].mDataByteSize;
-            uint8* OutputBufferPtr = (uint8*) OutputBufferData->mBuffers[0].mData;
-            while (OutputBufferBytesLeft > 0)
-            {
-                const int32 BytesToCopy = FMath::Min(SubmittedBufferBytesLeft, OutputBufferBytesLeft);
-                
-                FMemory::Memcpy((void*) OutputBufferPtr, SubmittedBufferPtr + SubmittedBytes, BytesToCopy);
-                
-                OutputBufferBytesLeft -= BytesToCopy;
-                SubmittedBufferBytesLeft -= BytesToCopy;
-                
-                if (SubmittedBufferBytesLeft <= 0)
-                {
-                    ReadNextBuffer();
-                    SubmittedBytes = 0;
-                    SubmittedBufferBytesLeft = BytesPerSubmittedBuffer;
-                }
-                else
-                {
-                    SubmittedBytes += BytesToCopy;
-                }
-                
-                if (OutputBufferBytesLeft <= 0)
-                {
-                    break;
-                }
-                else
-                {
-                    OutputBufferPtr += BytesToCopy;
-                }
-            }
+			if (!SubmittedBufferPtr)
+			{
+				ReadNextBuffer();
+			}
+
+			// How many bytes we have left over from previous callback
+			int32 SubmittedBufferBytesLeft = BytesPerSubmittedBuffer - SubmittedBytes;
+			int32 OutputBufferBytesLeft = OutputBufferData->mBuffers[0].mDataByteSize;
+			uint8* OutputBufferPtr = (uint8*) OutputBufferData->mBuffers[0].mData;
+			while (OutputBufferBytesLeft > 0)
+			{
+				const int32 BytesToCopy = FMath::Min(SubmittedBufferBytesLeft, OutputBufferBytesLeft);
+
+				FMemory::Memcpy((void*) OutputBufferPtr, SubmittedBufferPtr + SubmittedBytes, BytesToCopy);
+
+				OutputBufferBytesLeft -= BytesToCopy;
+				SubmittedBufferBytesLeft -= BytesToCopy;
+
+				if (SubmittedBufferBytesLeft <= 0)
+				{
+					ReadNextBuffer();
+					SubmittedBytes = 0;
+					SubmittedBufferBytesLeft = BytesPerSubmittedBuffer;
+				}
+				else
+				{
+					SubmittedBytes += BytesToCopy;
+				}
+
+				if (OutputBufferBytesLeft <= 0)
+				{
+					break;
+				}
+				else
+				{
+					OutputBufferPtr += BytesToCopy;
+				}
+			}
 		}
 		else
 		{
@@ -608,9 +620,9 @@ namespace Audio
 				memset(OutputBufferData->mBuffers[bufferItr].mData, 0, OutputBufferData->mBuffers[bufferItr].mDataByteSize);
 			}
 		}
-		
+
 		bInCallback = false;
-		
+
 		return true;
 	}
 
@@ -620,9 +632,9 @@ namespace Audio
 	{
 		// Get the user data and cast to our FMixerPlatformCoreAudio object
 		FMixerPlatformAudioUnit* me = (FMixerPlatformAudioUnit*) RefCon;
-		
+
 		me->PerformCallback(IOData);
-		
+
 		return noErr;
 	}
 }
diff --git a/Engine/Source/Runtime/AudioMixer/Private/AudioMixerSource.cpp b/Engine/Source/Runtime/AudioMixer/Private/AudioMixerSource.cpp
--- a/Engine/Source/Runtime/AudioMixer/Private/AudioMixerSource.cpp
+++ b/Engine/Source/Runtime/AudioMixer/Private/AudioMixerSource.cpp
@@ -485,7 +485,7 @@ namespace Audio
 						if (WaveInstance->StartTime > 0.0f || WaveInstance->WaveData->bProcedural || WaveInstance->WaveData->bIsBus || !WaveInstance->WaveData->CachedRealtimeFirstBuffer)
 						{
 							// Before reading more PCMRT data, we first need to seek the buffer
-							if (WaveInstance->StartTime > 0.0f && !(BufferType == EBufferType::Streaming) && !WaveInstance->WaveData->bIsBus && !WaveInstance->WaveData->bProcedural)
+							if (WaveInstance->IsSeekable())
 							{
 								MixerBuffer->Seek(WaveInstance->StartTime);
 							}
diff --git a/Engine/Source/Runtime/Engine/Classes/Sound/SoundWave.h b/Engine/Source/Runtime/Engine/Classes/Sound/SoundWave.h
--- a/Engine/Source/Runtime/Engine/Classes/Sound/SoundWave.h
+++ b/Engine/Source/Runtime/Engine/Classes/Sound/SoundWave.h
@@ -2,7 +2,7 @@
 
 #pragma once
 
-/** 
+/**
  * Playable sound object for raw wave files
  */
 
@@ -297,6 +297,10 @@ class ENGINE_API USoundWave : public USoundBase
 	UPROPERTY(EditAnywhere, Category=Streaming)
 	uint8 bStreaming:1;
 
+	/** Whether this sound supports seeking. This requires recooking with a codec which supports seekability and streaming. */
+	UPROPERTY(EditAnywhere, Category = Streaming, meta = (DisplayName = "Seekable", EditCondition = "bStreaming"))
+	uint8 bSeekableStreaming:1;
+
 	/** Set to true for programmatically-generated, streamed audio. */
 	uint8 bProcedural:1;
 
@@ -340,8 +344,8 @@ class ENGINE_API USoundWave : public USoundBase
 	UPROPERTY(EditAnywhere, Category = "Analysis")
 	USoundWave* OverrideSoundToUseForAnalysis;
 
-	/** 
-		Whether or not we should treat the sound wave used for analysis (this or the override) as looping while performing analysis. 
+	/**
+		Whether or not we should treat the sound wave used for analysis (this or the override) as looping while performing analysis.
 		A looping sound may include the end of the file for inclusion in analysis for envelope and FFT analysis.
 	*/
 	UPROPERTY(EditAnywhere, Category = "Analysis")
@@ -496,7 +500,7 @@ public:
 
 	UPROPERTY(VisibleAnywhere, Instanced, Category=ImportSettings)
 	class UAssetImportData* AssetImportData;
-	
+
 #endif // WITH_EDITORONLY_DATA
 
 protected:
@@ -516,14 +520,14 @@ private:
 	*/
 	static ITargetPlatform* GetRunningPlatform();
 
-public:	
+public:
 	/** Async worker that decompresses the audio data on a different thread */
 	typedef FAsyncTask< class FAsyncAudioDecompressWorker > FAsyncAudioDecompress;	// Forward declare typedef
 	FAsyncAudioDecompress* AudioDecompressor;
 
 	/** Pointer to 16 bit PCM data - used to avoid synchronous operation to obtain first block of the realtime decompressed buffer */
 	uint8* CachedRealtimeFirstBuffer;
-	
+
 	/** The number of frames which have been precached for this sound wave. */
 	int32 NumPrecacheFrames;
 
@@ -568,7 +572,7 @@ public:
 	/** cooked streaming platform data for this sound */
 	TSortedMap<FString, FStreamedAudioPlatformData*> CookedPlatformData;
 
-	//~ Begin UObject Interface. 
+	//~ Begin UObject Interface.
 	virtual void Serialize( FArchive& Ar ) override;
 	virtual void PostInitProperties() override;
 	virtual bool IsReadyForFinishDestroy() override;
@@ -576,13 +580,13 @@ public:
 	virtual void PostLoad() override;
 	virtual void BeginDestroy() override;
 #if WITH_EDITOR
-	virtual void PostEditChangeProperty(FPropertyChangedEvent& PropertyChangedEvent) override;	
+	virtual void PostEditChangeProperty(FPropertyChangedEvent& PropertyChangedEvent) override;
 #endif // WITH_EDITOR
 	virtual void GetResourceSizeEx(FResourceSizeEx& CumulativeResourceSize) override;
 	virtual FName GetExporterName() override;
 	virtual FString GetDesc() override;
 	virtual void GetAssetRegistryTags(TArray<FAssetRegistryTag>& OutTags) const override;
-	//~ End UObject Interface. 
+	//~ End UObject Interface.
 
 	//~ Begin USoundBase Interface.
 	virtual bool IsPlayable() const override;
@@ -646,12 +650,12 @@ public:
 	/** Will clean up the decompressor task if the task has finished or force it finish. Returns true if the decompressor is cleaned up. */
 	bool CleanupDecompressor(bool bForceCleanup = false);
 
-	/** 
+	/**
 	 * Copy the compressed audio data from the bulk data
 	 */
 	virtual void InitAudioResource( FByteBulkData& CompressedData );
 
-	/** 
+	/**
 	 * Copy the compressed audio data from derived data cache
 	 *
 	 * @param Format to get the compressed audio in
@@ -659,22 +663,22 @@ public:
 	 */
 	virtual bool InitAudioResource(FName Format);
 
-	/** 
+	/**
 	 * Remove the compressed audio data associated with the passed in wave
 	 */
 	void RemoveAudioResource();
 
-	/** 
+	/**
 	 * Prints the subtitle associated with the SoundWave to the console
 	 */
 	void LogSubtitle( FOutputDevice& Ar );
 
-	/** 
+	/**
 	 * Handle any special requirements when the sound starts (e.g. subtitles)
 	 */
 	FWaveInstance* HandleStart( FActiveSound& ActiveSound, const UPTRINT WaveInstanceHash ) const;
 
-	/** 
+	/**
 	 * This is only used for DTYPE_Procedural audio. It's recommended to use USynthComponent base class
 	 * for procedurally generated sound vs overriding this function. If a new component is not feasible,
 	 * consider using USoundWaveProcedural base class vs USoundWave base class since as it implements
@@ -682,21 +686,21 @@ public:
 	 */
 	virtual int32 GeneratePCMData(uint8* PCMData, const int32 SamplesNeeded) { ensure(false); return 0; }
 
-	/** 
-	* Return the format of the generated PCM data type. Used in audio mixer to allow generating float buffers and avoid unnecessary format conversions. 
+	/**
+	* Return the format of the generated PCM data type. Used in audio mixer to allow generating float buffers and avoid unnecessary format conversions.
 	* This feature is only supported in audio mixer. If your procedural sound wave needs to be used in both audio mixer and old audio engine,
 	* it's best to generate int16 data as old audio engine only supports int16 formats. Or check at runtime if the audio mixer is enabled.
 	* Audio mixer will convert from int16 to float internally.
 	*/
 	virtual Audio::EAudioMixerStreamDataFormat::Type GetGeneratedPCMDataFormat() const { return Audio::EAudioMixerStreamDataFormat::Int16; }
 
-	/** 
+	/**
 	 * Gets the compressed data size from derived data cache for the specified format
-	 * 
+	 *
 	 * @param Format	format of compressed data
 	 * @param CompressionOverrides Optional argument for compression overrides.
 	 * @return			compressed data size, or zero if it could not be obtained
-	 */ 
+	 */
 	int32 GetCompressedDataSize(FName Format, const FPlatformAudioCookOverrides* CompressionOverrides = GetPlatformCompressionOverridesForCurrentPlatform())
 	{
 		FByteBulkData* Data = GetCompressedData(Format, CompressionOverrides);
@@ -726,10 +730,10 @@ public:
 
 	virtual void BeginGetCompressedData(FName Format, const FPlatformAudioCookOverrides* CompressionOverrides);
 
-	/** 
+	/**
 	 * Gets the compressed data from derived data cache for the specified platform
 	 * Warning, the returned pointer isn't valid after we add new formats
-	 * 
+	 *
 	 * @param Format	format of compressed data
 	 * @param PlatformName optional name of platform we are getting compressed data for.
 	 * @param CompressionOverrides optional platform compression overrides
@@ -737,9 +741,9 @@ public:
 	 */
 	virtual FByteBulkData* GetCompressedData(FName Format, const FPlatformAudioCookOverrides* CompressionOverrides = GetPlatformCompressionOverridesForCurrentPlatform());
 
-	/** 
+	/**
 	 * Change the guid and flush all compressed data
-	 */ 
+	 */
 	void InvalidateCompressedData();
 
 	/** Returns curves associated with this sound wave */
@@ -764,11 +768,12 @@ public:
 	static FName GetCurvePropertyName() { return GET_MEMBER_NAME_CHECKED(USoundWave, Curves); }
 #endif
 
-	/**
-	 * Checks whether sound has been categorised as streaming.
-	 */
+	/** Checks whether sound has been categorised as streaming. */
 	bool IsStreaming(const FPlatformAudioCookOverrides* Overrides = nullptr) const;
 
+	/** Checks whether sound has seekable streaming enabled. */
+	bool IsSeekableStreaming() const;
+
 	/**
 	 * Attempts to update the cached platform data after any changes that might affect it
 	 */
@@ -797,13 +802,13 @@ public:
 	float GetSampleRateForCompressionOverrides(const FPlatformAudioCookOverrides* CompressionOverrides);
 
 #if WITH_EDITORONLY_DATA
-	
+
 #if WITH_EDITOR
 	/*
 	* Returns a sample rate if there is a specific sample rate override for this platform, -1.0 otherwise.
 	*/
 	float GetSampleRateForTargetPlatform(const ITargetPlatform* TargetPlatform);
-	
+
 	/**
 	 * Begins caching platform data in the background for the platform requested
 	 */
@@ -823,7 +828,7 @@ public:
 
 	uint32 bNeedsThumbnailGeneration:1;
 #endif
-	
+
 	/**
 	 * Caches platform data for the sound.
 	 */
diff --git a/Engine/Source/Runtime/Engine/Private/ADPCMAudioInfo.cpp b/Engine/Source/Runtime/Engine/Private/ADPCMAudioInfo.cpp
--- a/Engine/Source/Runtime/Engine/Private/ADPCMAudioInfo.cpp
+++ b/Engine/Source/Runtime/Engine/Private/ADPCMAudioInfo.cpp
@@ -13,62 +13,90 @@
 
 namespace ADPCM
 {
+	const uint32 MaxChunkSize = 256 * 1024;
+
 	void DecodeBlock(const uint8* EncodedADPCMBlock, int32 BlockSize, int16* DecodedPCMData);
 	void DecodeBlockStereo(const uint8* EncodedADPCMBlockLeft, const uint8* EncodedADPCMBlockRight, int32 BlockSize, int16* DecodedPCMData);
 }
 
 FADPCMAudioInfo::FADPCMAudioInfo(void)
+	: UncompressedBlockData(nullptr)
+	, SamplesPerBlock(0)
+	, bSeekPending(false)
 {
-	UncompressedBlockData = NULL;
-	SamplesPerBlock = 0;
 }
 
 FADPCMAudioInfo::~FADPCMAudioInfo(void)
 {
-	if(UncompressedBlockData != NULL)
+	if(UncompressedBlockData != nullptr)
 	{
 		FMemory::Free(UncompressedBlockData);
-		UncompressedBlockData = NULL;
+		UncompressedBlockData = nullptr;
 	}
 }
 
 void FADPCMAudioInfo::SeekToTime(const float SeekTime)
 {
-	if(SeekTime == 0.0f)
+	CurCompressedChunkData = nullptr;
+
+	if(SeekTime <= 0.0f)
 	{
-		CurrentUncompressedBlockSampleIndex = UncompressedBlockSize / sizeof(uint16);
 		CurrentCompressedBlockIndex = 0;
-
 		CurrentUncompressedBlockSampleIndex = 0;
 		CurrentChunkIndex = 0;
 		CurrentChunkBufferOffset = 0;
 		TotalSamplesStreamed = 0;
-		CurCompressedChunkData = NULL;
+
+		return;
 	}
-	else
+
+	// Calculate block index & force SeekTime to be in bounds.
+	check(WaveInfo.pSamplesPerSec != nullptr);
+	const uint32 SamplesPerSec = *WaveInfo.pSamplesPerSec;
+	const uint32 SeekedSamples = static_cast<uint32>(SeekTime * SamplesPerSec);
+	TotalSamplesStreamed = FMath::Min<uint32>(SeekedSamples, TotalSamplesPerChannel - 1);
+
+	const uint32 HeaderOffset = static_cast<uint32>(WaveInfo.SampleDataStart - SrcBufferData);
+
+	if (Format == WAVE_FORMAT_ADPCM)
 	{
-		if (Format == WAVE_FORMAT_LPCM)
+		CurrentCompressedBlockIndex = TotalSamplesStreamed / SamplesPerBlock; // Compute the block index that where SeekTime resides.
+		CurrentChunkIndex = 0;
+		CurrentChunkBufferOffset = HeaderOffset;
+
+		const int32 ChannelBlockSize = BlockSize * NumChannels;
+		for (uint32 BlockIndex = 0; BlockIndex < CurrentCompressedBlockIndex; ++BlockIndex)
 		{
-			// There are no "blocks" on LPCM, so only update the total samples streamed (which is based off sample rate).
-			// Note that TotalSamplesStreamed is per-channel in the ReadCompressedInfo. Channels are taken into account there.
-			uint32 SeekedSamples = (uint32)(SeekTime * (float)(*WaveInfo.pSamplesPerSec));
-			TotalSamplesStreamed = FMath::Clamp<uint32>(SeekedSamples, 0, TotalSamplesPerChannel - 1);
+			if (CurrentChunkBufferOffset + ChannelBlockSize >= ADPCM::MaxChunkSize)
+			{
+				++CurrentChunkIndex;
+				CurrentChunkBufferOffset = 0;
+			}
+
+			// Always add chunks in NumChannels pairs
+			CurrentChunkBufferOffset += ChannelBlockSize;
 		}
-		else
-		{
-			// Figure out the block index that the seek takes us to
-			uint32 SeekedSamples = (uint32)(SeekTime * (float)(*WaveInfo.pSamplesPerSec));
+	}
+	else if (Format == WAVE_FORMAT_LPCM)
+	{
+		const int32 ChannelBlockSize = sizeof(int16) * NumChannels;
 
-			// Clamp to the end of memory in case we have an invalid seek time.
-			SeekedSamples = FMath::Clamp<uint32>(SeekedSamples, 0, TotalSamplesPerChannel - 1);
+		// 1. Find total offset
+		CurrentChunkBufferOffset = HeaderOffset + (TotalSamplesStreamed * ChannelBlockSize);
 
-			// Compute the block index that we're seeked to
-			CurrentCompressedBlockIndex = SeekedSamples / SamplesPerBlock;
+		// 2. Calculate index and remove from offset
+		CurrentChunkIndex = CurrentChunkBufferOffset / ADPCM::MaxChunkSize;
+		CurrentChunkBufferOffset %= ADPCM::MaxChunkSize;
 
-			// Update the samples streamed to the current block index and the samples per block
-			TotalSamplesStreamed = CurrentCompressedBlockIndex * SamplesPerBlock;
-		}
+		// 3. Trim remainder of block size, effectively aligning the block to a channel pair boundary
+		CurrentChunkBufferOffset -= CurrentChunkBufferOffset % ChannelBlockSize;
+	}
+	else
+	{
+		return;
 	}
+
+	bSeekPending = true;
 }
 
 bool FADPCMAudioInfo::ReadCompressedInfo(const uint8* InSrcBufferData, uint32 InSrcBufferDataSize, struct FSoundQualityInfo* QualityInfo)
@@ -85,8 +113,8 @@ bool FADPCMAudioInfo::ReadCompressedInfo(const uint8* InSrcBufferData, uint32 In
 	SrcBufferDataSize = InSrcBufferDataSize;
 
 	void*	FormatHeader;
-	
-	if (!WaveInfo.ReadWaveInfo((uint8*)SrcBufferData, SrcBufferDataSize, NULL, false, &FormatHeader))
+
+	if (!WaveInfo.ReadWaveInfo((uint8*)SrcBufferData, SrcBufferDataSize, nullptr, false, &FormatHeader))
 	{
 		UE_LOG(LogAudio, Warning, TEXT("WaveInfo.ReadWaveInfo Failed"));
 		return false;
@@ -115,7 +143,7 @@ bool FADPCMAudioInfo::ReadCompressedInfo(const uint8* InSrcBufferData, uint32 In
 		TotalDecodedSize = ((WaveInfo.SampleDataSize + CompressedBlockSize - 1) / CompressedBlockSize) * UncompressedBlockSize;
 
 		UncompressedBlockData = (uint8*)FMemory::Realloc(UncompressedBlockData, NumChannels * UncompressedBlockSize);
-		check(UncompressedBlockData != NULL);
+		check(UncompressedBlockData != nullptr);
 		TotalCompressedBlocksPerChannel = (WaveInfo.SampleDataSize + CompressedBlockSize - 1) / CompressedBlockSize / NumChannels;
 	}
 	else if(Format == WAVE_FORMAT_LPCM)
@@ -125,18 +153,18 @@ bool FADPCMAudioInfo::ReadCompressedInfo(const uint8* InSrcBufferData, uint32 In
 		UncompressedBlockSize = 0;
 		CompressedBlockSize = 0;
 		StreamBufferSize = 0;
-		UncompressedBlockData = NULL;
+		UncompressedBlockData = nullptr;
 		TotalCompressedBlocksPerChannel = 0;
-		
+
 		TotalDecodedSize = WaveInfo.SampleDataSize;
-		
+
 		TotalSamplesPerChannel = TotalDecodedSize / sizeof(uint16) / NumChannels;
 	}
 	else
 	{
 		return false;
 	}
-	
+
 	if (QualityInfo)
 	{
 		QualityInfo->SampleRate = *WaveInfo.pSamplesPerSec;
@@ -144,26 +172,27 @@ bool FADPCMAudioInfo::ReadCompressedInfo(const uint8* InSrcBufferData, uint32 In
 		QualityInfo->SampleDataSize = TotalDecodedSize;
 		QualityInfo->Duration = (float)TotalSamplesPerChannel / QualityInfo->SampleRate;
 	}
-	
+
 	CurrentCompressedBlockIndex = 0;
 	TotalSamplesStreamed = 0;
 	// This is set to the max value to trigger the decompression of the first audio block
 	CurrentUncompressedBlockSampleIndex = UncompressedBlockSize / sizeof(uint16);
-	
+
 	return true;
 }
 
 bool FADPCMAudioInfo::ReadCompressedData(uint8* Destination, bool bLooping, uint32 BufferSize)
 {
+	const uint32 ChannelSampleSize = sizeof(uint16) * NumChannels;
+
 	// This correctly handles any BufferSize as long as its a multiple of sample size * number of channels
-	
 	check(Destination);
-	check((BufferSize % (sizeof(uint16) * NumChannels)) == 0);
-	
-	int16*	OutData = (int16*)Destination;
-	
-	bool	reachedEndOfSamples = false;
-	
+	check(BufferSize % ChannelSampleSize == 0);
+
+	int16* OutData = (int16*)Destination;
+
+	bool ReachedEndOfSamples = false;
+
 	if(Format == WAVE_FORMAT_ADPCM)
 	{
 		// We need to loop over the number of samples requested since an uncompressed block will not match the number of frames requested
@@ -172,50 +201,50 @@ bool FADPCMAudioInfo::ReadCompressedData(uint8* Destination, bool bLooping, uint
 			if(CurrentUncompressedBlockSampleIndex >= UncompressedBlockSize / sizeof(uint16))
 			{
 				// we need to decompress another block of compressed data from the current chunk
-				
+
 				// Decompress one block for each channel and store it in UncompressedBlockData
-				for(int32 channelItr = 0; channelItr < NumChannels; ++channelItr)
+				for(int32 ChannelItr = 0; ChannelItr < NumChannels; ++ChannelItr)
 				{
 					ADPCM::DecodeBlock(
-						WaveInfo.SampleDataStart + (channelItr * TotalCompressedBlocksPerChannel + CurrentCompressedBlockIndex) * CompressedBlockSize,
+						WaveInfo.SampleDataStart + (ChannelItr * TotalCompressedBlocksPerChannel + CurrentCompressedBlockIndex) * CompressedBlockSize,
 						CompressedBlockSize,
-						(int16*)(UncompressedBlockData + channelItr * UncompressedBlockSize));
+						(int16*)(UncompressedBlockData + ChannelItr * UncompressedBlockSize));
 				}
-				
+
 				// Update some bookkeeping
 				CurrentUncompressedBlockSampleIndex = 0;
 				++CurrentCompressedBlockIndex;
 			}
-			
+
 			// Only copy over the number of samples we currently have available, we will loop around if needed
-			uint32	decompressedSamplesToCopy = FMath::Min<uint32>(UncompressedBlockSize / sizeof(uint16) - CurrentUncompressedBlockSampleIndex, BufferSize / (sizeof(uint16) * NumChannels));
-			check(decompressedSamplesToCopy > 0);
-			
+			uint32 DecompressedSamplesToCopy = FMath::Min<uint32>(UncompressedBlockSize / sizeof(uint16) - CurrentUncompressedBlockSampleIndex, BufferSize / (sizeof(uint16) * NumChannels));
+			check(DecompressedSamplesToCopy > 0);
+
 			// Ensure we don't go over the number of samples left in the audio data
-			if(decompressedSamplesToCopy > TotalSamplesPerChannel - TotalSamplesStreamed)
+			if(DecompressedSamplesToCopy > TotalSamplesPerChannel - TotalSamplesStreamed)
 			{
-				decompressedSamplesToCopy = TotalSamplesPerChannel - TotalSamplesStreamed;
+				DecompressedSamplesToCopy = TotalSamplesPerChannel - TotalSamplesStreamed;
 			}
-			
+
 			// Copy over the actual sample data
-			for(uint32 sampleItr = 0; sampleItr < decompressedSamplesToCopy; ++sampleItr)
+			for(uint32 SampleItr = 0; SampleItr < DecompressedSamplesToCopy; ++SampleItr)
 			{
-				for(int32 channelItr = 0; channelItr < NumChannels; ++channelItr)
+				for(int32 ChannelItr = 0; ChannelItr < NumChannels; ++ChannelItr)
 				{
-					uint16	value = *(int16*)(UncompressedBlockData + channelItr * UncompressedBlockSize + (CurrentUncompressedBlockSampleIndex + sampleItr) * sizeof(int16));
-					*OutData++ = value;
+					uint16 Value = *(int16*)(UncompressedBlockData + ChannelItr * UncompressedBlockSize + (CurrentUncompressedBlockSampleIndex + SampleItr) * sizeof(int16));
+					*OutData++ = Value;
 				}
 			}
-			
+
 			// Update bookkeeping
-			CurrentUncompressedBlockSampleIndex += decompressedSamplesToCopy;
-			BufferSize -= decompressedSamplesToCopy * sizeof(uint16) * NumChannels;
-			TotalSamplesStreamed += decompressedSamplesToCopy;
-			
+			CurrentUncompressedBlockSampleIndex += DecompressedSamplesToCopy;
+			BufferSize -= DecompressedSamplesToCopy * sizeof(uint16) * NumChannels;
+			TotalSamplesStreamed += DecompressedSamplesToCopy;
+
 			// Check for the end of the audio samples and loop if needed
 			if(TotalSamplesStreamed >= TotalSamplesPerChannel)
 			{
-				reachedEndOfSamples = true;
+				ReachedEndOfSamples = true;
 				// This is set to the max value to trigger the decompression of the first audio block
 				CurrentUncompressedBlockSampleIndex = UncompressedBlockSize / sizeof(uint16);
 				CurrentCompressedBlockIndex = 0;
@@ -223,7 +252,7 @@ bool FADPCMAudioInfo::ReadCompressedData(uint8* Destination, bool bLooping, uint
 				if(!bLooping)
 				{
 					// Set the remaining buffer to 0
-					memset(OutData, 0, BufferSize);
+					FMemory::Memset(OutData, 0, BufferSize);
 					return true;
 				}
 			}
@@ -231,33 +260,33 @@ bool FADPCMAudioInfo::ReadCompressedData(uint8* Destination, bool bLooping, uint
 	}
 	else
 	{
-		uint32	decompressedSamplesToCopy = BufferSize / (sizeof(uint16) * NumChannels);
-			
+		uint32 DecompressedSamplesToCopy = BufferSize / ChannelSampleSize;
+
 		// Ensure we don't go over the number of samples left in the audio data
-		if(decompressedSamplesToCopy > TotalSamplesPerChannel - TotalSamplesStreamed)
+		if(DecompressedSamplesToCopy > TotalSamplesPerChannel - TotalSamplesStreamed)
 		{
-			decompressedSamplesToCopy = TotalSamplesPerChannel - TotalSamplesStreamed;
+			DecompressedSamplesToCopy = TotalSamplesPerChannel - TotalSamplesStreamed;
 		}
 
-		memcpy(OutData, WaveInfo.SampleDataStart + TotalSamplesStreamed * sizeof(uint16) * NumChannels, decompressedSamplesToCopy * sizeof(uint16) * NumChannels);
-		TotalSamplesStreamed += decompressedSamplesToCopy;
-		BufferSize -= decompressedSamplesToCopy * sizeof(uint16) * NumChannels;
-		
+		FMemory::Memcpy(OutData, WaveInfo.SampleDataStart + (TotalSamplesStreamed * ChannelSampleSize), DecompressedSamplesToCopy * ChannelSampleSize);
+		TotalSamplesStreamed += DecompressedSamplesToCopy;
+		BufferSize -= DecompressedSamplesToCopy * ChannelSampleSize;
+
 		// Check for the end of the audio samples and loop if needed
 		if(TotalSamplesStreamed >= TotalSamplesPerChannel)
 		{
-			reachedEndOfSamples = true;
+			ReachedEndOfSamples = true;
 			TotalSamplesStreamed = 0;
 			if(!bLooping)
 			{
 				// Set the remaining buffer to 0
-				memset(OutData, 0, BufferSize);
+				FMemory::Memset(OutData, 0, BufferSize);
 				return true;
 			}
 		}
 	}
-	
-	return reachedEndOfSamples;
+
+	return ReachedEndOfSamples;
 }
 
 void FADPCMAudioInfo::ExpandFile(uint8* DstBuffer, struct FSoundQualityInfo* QualityInfo)
@@ -279,37 +308,40 @@ bool FADPCMAudioInfo::StreamCompressedInfoInternal(USoundWave* Wave, struct FSou
 	check(StreamingSoundWave == Wave);
 
 	// Get the first chunk of audio data (should already be loaded)
-	uint8 const* firstChunk = IStreamingManager::Get().GetAudioStreamingManager().GetLoadedChunk(Wave, 0, &CurrentChunkDataSize);
+	uint8 const* FirstChunk = IStreamingManager::Get().GetAudioStreamingManager().GetLoadedChunk(Wave, 0, &CurrentChunkDataSize);
 
-	if (firstChunk == NULL)
+	if (FirstChunk == nullptr)
 	{
 		return false;
 	}
 
-	SrcBufferData = NULL;
+	SrcBufferData = nullptr;
 	SrcBufferDataSize = 0;
-	
-	void*	FormatHeader;
-	
-	if (!WaveInfo.ReadWaveInfo((uint8*)firstChunk, Wave->RunningPlatformData->Chunks[0].AudioDataSize, NULL, true, &FormatHeader))
+
+	void* FormatHeader;
+
+	if (!WaveInfo.ReadWaveInfo((uint8*)FirstChunk, Wave->RunningPlatformData->Chunks[0].AudioDataSize, nullptr, true, &FormatHeader))
 	{
 		UE_LOG(LogAudio, Warning, TEXT("WaveInfo.ReadWaveInfo Failed"));
 		return false;
 	}
 
-	FirstChunkSampleDataOffset = WaveInfo.SampleDataStart - firstChunk;
+	SrcBufferData = FirstChunk;
+
+	FirstChunkSampleDataOffset = WaveInfo.SampleDataStart - FirstChunk;
 	CurrentChunkBufferOffset = 0;
-	CurCompressedChunkData = NULL;
+	CurCompressedChunkData = nullptr;
 	CurrentUncompressedBlockSampleIndex = 0;
 	CurrentChunkIndex = 0;
 	TotalSamplesStreamed = 0;
 	Format = *WaveInfo.pFormatTag;
 	NumChannels = *WaveInfo.pChannels;
-	
+
 	if (Format == WAVE_FORMAT_ADPCM)
 	{
-		ADPCM::ADPCMFormatHeader* APPCMHeader = (ADPCM::ADPCMFormatHeader*)FormatHeader;
-		TotalSamplesPerChannel = APPCMHeader->SamplesPerChannel;
+		ADPCM::ADPCMFormatHeader* ADPCMHeader = (ADPCM::ADPCMFormatHeader*)FormatHeader;
+		TotalSamplesPerChannel = ADPCMHeader->SamplesPerChannel;
+		SamplesPerBlock = ADPCMHeader->wSamplesPerBlock;
 
 		const uint32 PreambleSize = 7;
 
@@ -322,9 +354,9 @@ bool FADPCMAudioInfo::StreamCompressedInfoInternal(USoundWave* Wave, struct FSou
 		const uint32 targetBlocks = MONO_PCM_BUFFER_SAMPLES / uncompressedBlockSamples;
 		StreamBufferSize = targetBlocks * UncompressedBlockSize;
 		TotalDecodedSize = ((WaveInfo.SampleDataSize + CompressedBlockSize - 1) / CompressedBlockSize) * UncompressedBlockSize;
-		
+
 		UncompressedBlockData = (uint8*)FMemory::Realloc(UncompressedBlockData, NumChannels * UncompressedBlockSize);
-		check(UncompressedBlockData != NULL);
+		check(UncompressedBlockData != nullptr);
 	}
 	else if (Format == WAVE_FORMAT_LPCM)
 	{
@@ -342,7 +374,7 @@ bool FADPCMAudioInfo::StreamCompressedInfoInternal(USoundWave* Wave, struct FSou
 		UE_LOG(LogAudio, Error, TEXT("Unsupported wave format"));
 		return false;
 	}
-	
+
 	if (QualityInfo)
 	{
 		QualityInfo->SampleRate = *WaveInfo.pSamplesPerSec;
@@ -357,104 +389,117 @@ bool FADPCMAudioInfo::StreamCompressedInfoInternal(USoundWave* Wave, struct FSou
 bool FADPCMAudioInfo::StreamCompressedData(uint8* Destination, bool bLooping, uint32 BufferSize)
 {
 	// Destination samples are interlaced by channel, BufferSize is in bytes
-	
+	const int32 ChannelSampleSize = sizeof(uint16) * NumChannels;
+
 	// Ensure that BuffserSize is a multiple of the sample size times the number of channels
-	checkf((BufferSize % (sizeof(uint16) * NumChannels)) == 0, TEXT("Invalid buffer size %d requested for %d channels"), BufferSize, NumChannels);
-	
-	int16*	OutData = (int16*)Destination;
-	
-	bool	reachedEndOfSamples = false;
+	checkf(BufferSize % ChannelSampleSize == 0, TEXT("Invalid buffer size %d requested for %d channels"), BufferSize, NumChannels);
+
+	int16* OutData = (int16*)Destination;
+
+	bool ReachedEndOfSamples = false;
 
 	if(Format == WAVE_FORMAT_ADPCM)
 	{
 		// We need to loop over the number of samples requested since an uncompressed block will not match the number of frames requested
 		while(BufferSize > 0)
 		{
-			if(CurCompressedChunkData == NULL || CurrentUncompressedBlockSampleIndex >= UncompressedBlockSize / sizeof(uint16))
+			if(CurCompressedChunkData == nullptr || CurrentUncompressedBlockSampleIndex >= UncompressedBlockSize / sizeof(uint16))
 			{
 				// we need to decompress another block of compressed data from the current chunk
-				
-				if(CurCompressedChunkData == NULL || CurrentChunkBufferOffset >= CurrentChunkDataSize)
+
+				if(CurCompressedChunkData == nullptr || CurrentChunkBufferOffset >= CurrentChunkDataSize)
 				{
 					// Get another chunk of compressed data from the streaming engine
-					
-					// CurrentChunkIndex is used to keep track of the current chunk for loading/unloading by the streaming engine but chunk 0 is preloaded so we don't want to increment this when getting chunk 0, only later chunks
-					// Also, if we failed to get a previous chunk because it was not ready we don't want to re-increment the CurrentChunkIndex
-					if(CurCompressedChunkData != NULL)
+
+					// CurrentChunkIndex is used to keep track of the current chunk for loading/unloading by the streaming engine, but chunk 0 is
+					// preloaded so don't increment this when getting chunk 0, only later chunks. If previous chunk retrieval failed because it was
+					// not ready, don't re-increment the CurrentChunkIndex.  Only increment if seek not pending as seek determines the current chunk index
+					// and invalidates CurCompressedChunkData.
+					if(CurCompressedChunkData != nullptr)
 					{
 						++CurrentChunkIndex;
 					}
-					
+
 					// Request the next chunk of data from the streaming engine
 					CurCompressedChunkData = IStreamingManager::Get().GetAudioStreamingManager().GetLoadedChunk(StreamingSoundWave, CurrentChunkIndex, &CurrentChunkDataSize);
 
-					if(CurCompressedChunkData == NULL)
+					if(CurCompressedChunkData == nullptr)
 					{
-						// If we did not get it then just bail, CurrentChunkIndex will not get incremented on the next callback so in effect another attempt will be made to fetch the chunk
-						// Since audio streaming depends on the general data streaming mechanism used by other parts of the engine and new data is prefectched on the game tick thread its possible a game hickup can cause this
-						UE_LOG(LogAudio, Verbose, TEXT("Missed Deadline chunk %d"), CurrentChunkIndex);
-						
+						// We only need to worry about missing the stream chunk if we were seeking. Seeking might cause a bit of latency with chunk loading. That is expected.
+						if (!bSeekPending)
+						{
+							// If we did not get it then just bail, CurrentChunkIndex will not get incremented on the next callback so in effect another attempt will be made to fetch the chunk
+							// Since audio streaming depends on the general data streaming mechanism used by other parts of the engine and new data is prefectched on the game tick thread its possible a game hickup can cause this
+							UE_LOG(LogAudio, Verbose, TEXT("Missed Deadline chunk %d"), CurrentChunkIndex);
+						}
+
 						// zero out remaining data and bail
-						memset(OutData, 0, BufferSize);
+						FMemory::Memset(OutData, 0, BufferSize);
 						return false;
 					}
-					
+
 					// Set the current buffer offset accounting for the header in the first chunk
-					CurrentChunkBufferOffset = CurrentChunkIndex == 0 ? FirstChunkSampleDataOffset : 0;
+					if (!bSeekPending)
+					{
+						CurrentChunkBufferOffset = CurrentChunkIndex == 0 ? FirstChunkSampleDataOffset : 0;
+					}
+					bSeekPending = false;
 				}
-				
+
 				// Decompress one block for each channel and store it in UncompressedBlockData
-				for(int32 channelItr = 0; channelItr < NumChannels; ++channelItr)
+				for(int32 ChannelItr = 0; ChannelItr < NumChannels; ++ChannelItr)
 				{
 					ADPCM::DecodeBlock(
-						CurCompressedChunkData + CurrentChunkBufferOffset + channelItr * CompressedBlockSize,
+						CurCompressedChunkData + CurrentChunkBufferOffset + ChannelItr * CompressedBlockSize,
 						CompressedBlockSize,
-						(int16*)(UncompressedBlockData + channelItr * UncompressedBlockSize));
+						(int16*)(UncompressedBlockData + ChannelItr * UncompressedBlockSize));
 				}
-				
+
 				// Update some bookkeeping
 				CurrentUncompressedBlockSampleIndex = 0;
 				CurrentChunkBufferOffset += NumChannels * CompressedBlockSize;
 			}
-			
+
 			// Only copy over the number of samples we currently have available, we will loop around if needed
-			uint32	decompressedSamplesToCopy = FMath::Min<uint32>(UncompressedBlockSize / sizeof(uint16) - CurrentUncompressedBlockSampleIndex, BufferSize / (sizeof(uint16) * NumChannels));
-			check(decompressedSamplesToCopy > 0);
-			
+			uint32 DecompressedSamplesToCopy = FMath::Min<uint32>(
+				(UncompressedBlockSize / sizeof(uint16)) - CurrentUncompressedBlockSampleIndex,
+				BufferSize / (ChannelSampleSize));
+			check(DecompressedSamplesToCopy > 0);
+
 			// Ensure we don't go over the number of samples left in the audio data
-			if(decompressedSamplesToCopy > TotalSamplesPerChannel - TotalSamplesStreamed)
+			if(DecompressedSamplesToCopy > TotalSamplesPerChannel - TotalSamplesStreamed)
 			{
-				decompressedSamplesToCopy = TotalSamplesPerChannel - TotalSamplesStreamed;
+				DecompressedSamplesToCopy = TotalSamplesPerChannel - TotalSamplesStreamed;
 			}
-			
+
 			// Copy over the actual sample data
-			for(uint32 sampleItr = 0; sampleItr < decompressedSamplesToCopy; ++sampleItr)
+			for(uint32 SampleItr = 0; SampleItr < DecompressedSamplesToCopy; ++SampleItr)
 			{
-				for(int32 channelItr = 0; channelItr < NumChannels; ++channelItr)
+				for(int32 ChannelItr = 0; ChannelItr < NumChannels; ++ChannelItr)
 				{
-					uint16	value = *(int16*)(UncompressedBlockData + channelItr * UncompressedBlockSize + (CurrentUncompressedBlockSampleIndex + sampleItr) * sizeof(int16));
-					*OutData++ = value;
+					uint16 Value = *(int16*)(UncompressedBlockData + ChannelItr * UncompressedBlockSize + (CurrentUncompressedBlockSampleIndex + SampleItr) * sizeof(int16));
+					*OutData++ = Value;
 				}
 			}
-			
+
 			// Update bookkeeping
-			CurrentUncompressedBlockSampleIndex += decompressedSamplesToCopy;
-			BufferSize -= decompressedSamplesToCopy * sizeof(uint16) * NumChannels;
-			TotalSamplesStreamed += decompressedSamplesToCopy;
-			
+			CurrentUncompressedBlockSampleIndex += DecompressedSamplesToCopy;
+			BufferSize -= DecompressedSamplesToCopy * ChannelSampleSize;
+			TotalSamplesStreamed += DecompressedSamplesToCopy;
+
 			// Check for the end of the audio samples and loop if needed
 			if(TotalSamplesStreamed >= TotalSamplesPerChannel)
 			{
-				reachedEndOfSamples = true;
+				ReachedEndOfSamples = true;
 				CurrentUncompressedBlockSampleIndex = 0;
 				CurrentChunkIndex = 0;
 				CurrentChunkBufferOffset = 0;
 				TotalSamplesStreamed = 0;
-				CurCompressedChunkData = NULL;
+				CurCompressedChunkData = nullptr;
 				if(!bLooping)
 				{
 					// Set the remaining buffer to 0
-					memset(OutData, 0, BufferSize);
+					FMemory::Memset(OutData, 0, BufferSize);
 					return true;
 				}
 			}
@@ -464,68 +509,82 @@ bool FADPCMAudioInfo::StreamCompressedData(uint8* Destination, bool bLooping, ui
 	{
 		while(BufferSize > 0)
 		{
-			if(CurCompressedChunkData == NULL || CurrentChunkBufferOffset >= CurrentChunkDataSize)
+			if(CurCompressedChunkData == nullptr || CurrentChunkBufferOffset >= CurrentChunkDataSize)
 			{
 				// Get another chunk of compressed data from the streaming engine
-				
-				// CurrentChunkIndex is used to keep track of the current chunk for loading/unloading by the streaming engine but chunk 0 is preloaded so we don't want to increment this when getting chunk 0, only later chunks
+
+				// CurrentChunkIndex is used to keep track of the current chunk for loading/unloading by the streaming engine but chunk 0
+				// is preloaded so we don't want to increment this when getting chunk 0, only later chunks
 				// Also, if we failed to get a previous chunk because it was not ready we don't want to re-increment the CurrentChunkIndex
-				if(CurCompressedChunkData != NULL)
+				if(CurCompressedChunkData != nullptr)
 				{
 					++CurrentChunkIndex;
 				}
-				
+
 				// Request the next chunk of data from the streaming engine
 				CurCompressedChunkData = IStreamingManager::Get().GetAudioStreamingManager().GetLoadedChunk(StreamingSoundWave, CurrentChunkIndex, &CurrentChunkDataSize);
 
-				if(CurCompressedChunkData == NULL)
+				if(CurCompressedChunkData == nullptr)
 				{
-					// If we did not get it then just bail, CurrentChunkIndex will not get incremented on the next callback so in effect another attempt will be made to fetch the chunk
-					// Since audio streaming depends on the general data streaming mechanism used by other parts of the engine and new data is prefectched on the game tick thread its possible a game hickup can cause this
-					UE_LOG(LogAudio, Warning, TEXT("Missed Deadline chunk %d"), CurrentChunkIndex);
-					
-					// zero out remaining data and bail
-					memset(OutData, 0, BufferSize);
+					// Only report missing the stream chunk if we were seeking. Seeking
+					// may cause a bit of latency with chunk loading, which is expected.
+					if (!bSeekPending)
+					{
+						// CurrentChunkIndex will not get incremented on the next callback, effectively causing another attempt to fetch the chunk.
+						// Since audio streaming depends on the general data streaming mechanism used by other parts of the engine and new data is
+						// prefetched on the game tick thread, this may be caused by a game hitch.
+						UE_LOG(LogAudio, Verbose, TEXT("Missed streaming ADPCM deadline chunk %d"), CurrentChunkIndex);
+					}
+
+					FMemory::Memset(OutData, 0, BufferSize);
 					return false;
 				}
-				
+
 				// Set the current buffer offset accounting for the header in the first chunk
-				CurrentChunkBufferOffset = CurrentChunkIndex == 0 ? FirstChunkSampleDataOffset : 0;
+				if (!bSeekPending)
+				{
+					CurrentChunkBufferOffset = CurrentChunkIndex == 0 ? FirstChunkSampleDataOffset : 0;
+				}
+				bSeekPending = false;
 			}
 
-			uint32	decompressedSamplesToCopy = FMath::Min<uint32>((CurrentChunkDataSize - CurrentChunkBufferOffset) / (sizeof(uint16) * NumChannels), BufferSize / (sizeof(uint16) * NumChannels));
-			check(decompressedSamplesToCopy > 0);
-			
+			uint32 DecompressedSamplesToCopy = FMath::Min<uint32>(
+				(CurrentChunkDataSize - CurrentChunkBufferOffset) / ChannelSampleSize,
+				BufferSize / ChannelSampleSize);
+
+			check(DecompressedSamplesToCopy > 0);
+
 			// Ensure we don't go over the number of samples left in the audio data
-			if(decompressedSamplesToCopy > TotalSamplesPerChannel - TotalSamplesStreamed)
+			if(DecompressedSamplesToCopy > TotalSamplesPerChannel - TotalSamplesStreamed)
 			{
-				decompressedSamplesToCopy = TotalSamplesPerChannel - TotalSamplesStreamed;
+				DecompressedSamplesToCopy = TotalSamplesPerChannel - TotalSamplesStreamed;
 			}
-			
-			memcpy(OutData, CurCompressedChunkData + CurrentChunkBufferOffset, decompressedSamplesToCopy * (sizeof(uint16) * NumChannels));
-			
-			OutData += decompressedSamplesToCopy * NumChannels;
-			CurrentChunkBufferOffset += decompressedSamplesToCopy * (sizeof(uint16) * NumChannels);
-			BufferSize -= decompressedSamplesToCopy * (sizeof(uint16) * NumChannels);
-			TotalSamplesStreamed += decompressedSamplesToCopy;
-			
+
+			const uint32 SampleSize = DecompressedSamplesToCopy * ChannelSampleSize;
+			FMemory::Memcpy(OutData, CurCompressedChunkData + CurrentChunkBufferOffset, SampleSize);
+
+			OutData += DecompressedSamplesToCopy * NumChannels;
+			CurrentChunkBufferOffset += SampleSize;
+			BufferSize -= SampleSize;
+			TotalSamplesStreamed += DecompressedSamplesToCopy;
+
 			// Check for the end of the audio samples and loop if needed
 			if(TotalSamplesStreamed >= TotalSamplesPerChannel)
 			{
-				reachedEndOfSamples = true;
+				ReachedEndOfSamples = true;
 				CurrentChunkIndex = 0;
 				CurrentChunkBufferOffset = 0;
 				TotalSamplesStreamed = 0;
-				CurCompressedChunkData = NULL;
+				CurCompressedChunkData = nullptr;
 				if(!bLooping)
 				{
 					// Set the remaining buffer to 0
-					memset(OutData, 0, BufferSize);
+					FMemory::Memset(OutData, 0, BufferSize);
 					return true;
 				}
 			}
 		}
 	}
-	
-	return reachedEndOfSamples;
+
+	return ReachedEndOfSamples;
 }
diff --git a/Engine/Source/Runtime/Engine/Private/Audio.cpp b/Engine/Source/Runtime/Engine/Private/Audio.cpp
--- a/Engine/Source/Runtime/Engine/Private/Audio.cpp
+++ b/Engine/Source/Runtime/Engine/Private/Audio.cpp
@@ -983,6 +983,28 @@ float FWaveInstance::GetVolumeWeightedPriority() const
 	}
 }
 
+bool FWaveInstance::IsSeekable() const
+{
+	check(WaveData);
+
+	if (StartTime == 0.0f)
+	{
+		return false;
+	}
+
+	if (WaveData->bIsBus || WaveData->bProcedural)
+	{
+		return false;
+	}
+
+	if (IsStreaming() && !WaveData->IsSeekableStreaming())
+	{
+		return false;
+	}
+
+	return true;
+}
+
 bool FWaveInstance::IsStreaming() const
 {
 	return FPlatformProperties::SupportsAudioStreaming() && WaveData != nullptr && WaveData->IsStreaming();
diff --git a/Engine/Source/Runtime/Engine/Private/SoundWave.cpp b/Engine/Source/Runtime/Engine/Private/SoundWave.cpp
--- a/Engine/Source/Runtime/Engine/Private/SoundWave.cpp
+++ b/Engine/Source/Runtime/Engine/Private/SoundWave.cpp
@@ -203,13 +203,13 @@ void USoundWave::GetResourceSizeEx(FResourceSizeEx& CumulativeResourceSize)
 			check(!RawPCMData || RawPCMDataSize);
 			CumulativeResourceSize.AddDedicatedSystemMemoryBytes(RawPCMDataSize);
 		}
-		else 
+		else
 		{
 			if (DecompressionType == DTYPE_RealTime && CachedRealtimeFirstBuffer)
 			{
 				CumulativeResourceSize.AddDedicatedSystemMemoryBytes(MONO_PCM_BUFFER_SIZE * NumChannels);
 			}
-			
+
 			if (!FPlatformProperties::SupportsAudioStreaming() || !IsStreaming())
 			{
 				CumulativeResourceSize.AddDedicatedSystemMemoryBytes(GetCompressedDataSize(LocalAudioDevice->GetRuntimeFormat(this)));
@@ -263,7 +263,7 @@ FString USoundWave::GetDesc()
 void USoundWave::GetAssetRegistryTags(TArray<FAssetRegistryTag>& OutTags) const
 {
 	Super::GetAssetRegistryTags(OutTags);
-	
+
 #if WITH_EDITORONLY_DATA
 	if (AssetImportData)
 	{
@@ -296,7 +296,7 @@ void USoundWave::Serialize( FArchive& Ar )
 		Ar << DummyCompressionName;
 	}
 
-	bool bShouldStreamSound = false; 
+	bool bShouldStreamSound = false;
 
 	if (Ar.IsSaving() || Ar.IsCooking())
 	{
@@ -358,7 +358,7 @@ void USoundWave::Serialize( FArchive& Ar )
 					}
 				}
 				bool bMapped = CookingTarget->SupportsFeature(ETargetPlatformFeatures::MemoryMappedFiles) && CookingTarget->SupportsFeature(ETargetPlatformFeatures::MemoryMappedAudio);
-				CompressedFormatData.Serialize(Ar, this, &ActualFormatsToSave, true, DEFAULT_ALIGNMENT, 
+				CompressedFormatData.Serialize(Ar, this, &ActualFormatsToSave, true, DEFAULT_ALIGNMENT,
 					!bMapped, // inline if not mapped
 					bMapped);
 #endif
@@ -395,7 +395,7 @@ void USoundWave::Serialize( FArchive& Ar )
 			}
 		}
 
-#if WITH_EDITORONLY_DATA	
+#if WITH_EDITORONLY_DATA
 		if (Ar.IsLoading() && !Ar.IsTransacting() && !bCooked && !GetOutermost()->HasAnyPackageFlags(PKG_ReloadingForCooker))
 		{
 			BeginCachePlatformData();
@@ -441,7 +441,7 @@ void USoundWave::LogSubtitle( FOutputDevice& Ar )
 }
 
 float USoundWave::GetSubtitlePriority() const
-{ 
+{
 	return SubtitlePriority;
 };
 
@@ -630,7 +630,7 @@ FByteBulkData* USoundWave::GetCompressedData(FName Format, const FPlatformAudioC
 	{
 		return nullptr;
 	}
-	
+
 	FName PlatformSpecificFormat = GetPlatformSpecificFormat(Format, CompressionOverrides);
 
 	bool bContainedValidData = CompressedFormatData.Contains(PlatformSpecificFormat);
@@ -772,7 +772,7 @@ void USoundWave::BeginDestroy()
 
 	// Flag that this sound wave is beginning destroying. This will ensure that all sounds using this in the audio renderer are stopped before GC finishes.
 	bIsBeginDestroy = true;
-	
+
 #if WITH_EDITOR
 	// Flush any async results so we dont leak them in the DDC
 	if (GetDerivedDataCache() && AsyncLoadingDataFormats.Num() > 0)
@@ -886,7 +886,7 @@ void USoundWave::LogBakedData()
 
 	if (CookedEnvelopeTimeData.Num())
 	{
-		FString EnvelopeFileName = FString::Printf(TEXT("%s.envelope.csv"), *SoundWaveName);	
+		FString EnvelopeFileName = FString::Printf(TEXT("%s.envelope.csv"), *SoundWaveName);
 		FString FilenameFull = AnalysisPathName + EnvelopeFileName;
 
 		FArchive* FileAr = IFileManager::Get().CreateDebugFileWriter(*FilenameFull);
@@ -970,7 +970,7 @@ static bool AnyFFTAnalysisPropertiesChanged(const FName& PropertyName)
 			PropertyName == EnableFFTAnalysisFName ||
 			PropertyName == FFTSizeFName ||
 			PropertyName == FFTAnalysisFrameSizeFName ||
-			PropertyName == FrequenciesToAnalyzeFName || 
+			PropertyName == FrequenciesToAnalyzeFName ||
 			PropertyName == FFTAnalysisAttackTimeFName ||
 			PropertyName == FFTAnalysisReleaseTimeFName;
 }
@@ -1023,7 +1023,7 @@ void USoundWave::BakeFFTAnalysis()
 		{
 			SoundWaveToUseForAnalysis = OverrideSoundToUseForAnalysis;
 		}
-		
+
 		if (!SoundWaveToUseForAnalysis->GetImportedSoundWaveData(RawImportedWaveData, RawDataSampleRate, RawDataNumChannels))
 		{
 			return;
@@ -1137,7 +1137,7 @@ void USoundWave::BakeFFTAnalysis()
 				// Don't let the time shift be more than the negative or postive duration
 				float Duration = (float)NumFrames / RawDataSampleRate;
 				float TimeShift = FMath::Clamp((float)FFTAnalysisTimeOffset / 1000, -Duration, Duration);
-				
+
 				NewData.TimeSec = NewData.TimeSec + (float)FFTAnalysisTimeOffset / 1000;
 
 				// Truncate data if time shift is far enough to left that it's before the start of the sound
@@ -1177,7 +1177,7 @@ void USoundWave::BakeFFTAnalysis()
 
 		CookedSpectralTimeData.Sort(FSortSpectralDataByTime());
 
-		// It's possible for the maximum magnitude to be 0.0 if the audio file was silent. 
+		// It's possible for the maximum magnitude to be 0.0 if the audio file was silent.
 		if (MaximumMagnitude > 0.0f)
 		{
 			// Normalize all the magnitude values based on the highest magnitude
@@ -1265,6 +1265,7 @@ void USoundWave::PostEditChangeProperty(FPropertyChangedEvent& PropertyChangedEv
 
 	static const FName CompressionQualityFName = FName(TEXT("CompressionQuality"));
 	static FName StreamingFName = GET_MEMBER_NAME_CHECKED(USoundWave, bStreaming);
+	static FName SeekableStreamingFName = GET_MEMBER_NAME_CHECKED(USoundWave, bSeekableStreaming);
 
 	// Prevent constant re-compression of SoundWave while properties are being changed interactively
 	if (PropertyChangedEvent.ChangeType != EPropertyChangeType::Interactive)
@@ -1273,19 +1274,19 @@ void USoundWave::PostEditChangeProperty(FPropertyChangedEvent& PropertyChangedEv
 		if (UProperty* PropertyThatChanged = PropertyChangedEvent.Property)
 		{
 			const FName& Name = PropertyThatChanged->GetFName();
-			if (Name == CompressionQualityFName || Name == StreamingFName)
+			if (Name == CompressionQualityFName || Name == StreamingFName || Name == SeekableStreamingFName)
 			{
 				InvalidateCompressedData();
 				FreeResources();
 				UpdatePlatformData();
 				MarkPackageDirty();
 			}
-			
+
 			if (AnyFFTAnalysisPropertiesChanged(Name))
 			{
 				BakeFFTAnalysis();
 			}
-			
+
 			if (AnyEnvelopeAnalysisPropertiesChanged(Name))
 			{
 				BakeEnvelopeAnalysis();
@@ -1425,7 +1426,7 @@ bool USoundWave::IsReadyForFinishDestroy()
 				SoundWave->FreeResources();
 			}, GET_STATID(STAT_AudioFreeResources));
 		}
-	
+
 	return NumSourcesPlaying.GetValue() == 0 && ResourceState == ESoundWaveResourceState::Freed;
 }
 
@@ -1525,7 +1526,7 @@ void USoundWave::Parse( FAudioDevice* AudioDevice, const UPTRINT NodeWaveInstanc
 			WaveInstance->RadioFilterVolumeThreshold = SoundClassProperties->RadioFilterVolumeThreshold * ParseParams.VolumeMultiplier;
 			WaveInstance->StereoBleed = SoundClassProperties->StereoBleed;
 			WaveInstance->LFEBleed = SoundClassProperties->LFEBleed;
-			
+
 			WaveInstance->bIsUISound = ActiveSound.bIsUISound || SoundClassProperties->bIsUISound;
 			WaveInstance->bIsMusic = ActiveSound.bIsMusic || SoundClassProperties->bIsMusic;
 			WaveInstance->bCenterChannelOnly = ActiveSound.bCenterChannelOnly || SoundClassProperties->bCenterChannelOnly;
@@ -1714,6 +1715,11 @@ bool USoundWave::IsStreaming(const FPlatformAudioCookOverrides* Overrides /* = n
 			&& Duration > Overrides->AutoStreamingThreshold);
 }
 
+bool USoundWave::IsSeekableStreaming() const
+{
+	return bStreaming && bSeekableStreaming;
+}
+
 bool USoundWave::GetSoundWavesWithCookedAnalysisData(TArray<USoundWave*>& OutSoundWaves)
 {
 	if (CookedSpectralTimeData.Num() > 0 || CookedEnvelopeTimeData.Num() > 0)
@@ -1908,8 +1914,8 @@ bool USoundWave::GetInterpolatedCookedFFTDataForTime(float InTime, uint32& InOut
 				return true;
 			}
 		}
-		// We're looping 
-		else 
+		// We're looping
+		else
 		{
 			// Need to check initial wrap-around case (i.e. we're reading earlier than first data point so need to lerp from last data point to first
 			if (InTime >= 0.0f && InTime < CookedSpectralTimeData[0].TimeSec)
@@ -1938,7 +1944,7 @@ bool USoundWave::GetInterpolatedCookedFFTDataForTime(float InTime, uint32& InOut
 			// Or we've been offset a bit in the negative.
 			else if (InTime < 0.0f)
 			{
-				// Wrap the time to the end of the sound wave file 
+				// Wrap the time to the end of the sound wave file
 				InTime = FMath::Clamp(Duration + InTime, 0.0f, Duration);
 			}
 		}
@@ -2007,7 +2013,7 @@ bool USoundWave::GetInterpolatedCookedEnvelopeDataForTime(float InTime, uint32&
 				InOutLastIndex = 0;
 				return true;
 			}
-		}	
+		}
 		else
 		{
 			// Need to check initial wrap-around case (i.e. we're reading earlier than first data point so need to lerp from last data point to first
@@ -2026,7 +2032,7 @@ bool USoundWave::GetInterpolatedCookedEnvelopeDataForTime(float InTime, uint32&
 			// Or we've been offset a bit in the negative.
 			else if (InTime < 0.0f)
 			{
-				// Wrap the time to the end of the sound wave file 
+				// Wrap the time to the end of the sound wave file
 				InTime = FMath::Clamp(Duration + InTime, 0.0f, Duration);
 			}
 
diff --git a/Engine/Source/Runtime/Engine/Public/ADPCMAudioInfo.h b/Engine/Source/Runtime/Engine/Public/ADPCMAudioInfo.h
--- a/Engine/Source/Runtime/Engine/Public/ADPCMAudioInfo.h
+++ b/Engine/Source/Runtime/Engine/Public/ADPCMAudioInfo.h
@@ -15,7 +15,7 @@
 #define NUM_ADAPTATION_COEFF 7
 
 #define WAVE_FORMAT_LPCM  1
-#ifndef WAVE_FORMAT_ADPCM 
+#ifndef WAVE_FORMAT_ADPCM
 #define WAVE_FORMAT_ADPCM 2
 #endif
 
@@ -125,13 +125,13 @@ public:
 	uint32	TotalDecodedSize;
 	int32	NumChannels;
 	int32	Format;
-	
+
 	uint8*			UncompressedBlockData;			// This holds the current block of compressed data for all channels
 	uint32			CurrentUncompressedBlockSampleIndex;	// This is the sample index within the current uncompressed block data
 	uint32			CurrentChunkIndex;				// This is the index that is currently being used, needed by streaming engine to make sure it stays loaded and the next chunk gets preloaded
 	uint32			CurrentChunkBufferOffset;		// This is this byte offset within the current chunk, used by streaming engine to prioritize a load if more then half way through current chunk
-	uint32			CurrentChunkDataSize;			// The size of the current chunk, the first chunk is bigger to accomodate the header info
-	uint32			TotalSamplesStreamed;			// The number of samples streamed so far
+	uint32			CurrentChunkDataSize;			// The size of the current chunk, the first chunk is bigger to accommodate the header info
+	uint32			TotalSamplesStreamed;			// The number of samples streamed so far (per channel)
 	uint32			TotalSamplesPerChannel;			// Number of samples per channel, used to detect when an audio waveform has ended
 	uint32			SamplesPerBlock;				// The number of samples per block
 	uint32			FirstChunkSampleDataOffset;		// The size of the header in the first chunk, used to skip over it when looping or starting the sample over
@@ -139,4 +139,5 @@ public:
 
 	uint32			CurrentCompressedBlockIndex;		// For non disk streaming - the current compressed block in the compressed source data
 	uint32			TotalCompressedBlocksPerChannel;	// For non disk streaming - the total number of compressed blocks per channel
+	uint8			bSeekPending : 1;					// Whether or not seek has been requested and pending a read
 };
diff --git a/Engine/Source/Runtime/Engine/Public/Audio.h b/Engine/Source/Runtime/Engine/Public/Audio.h
--- a/Engine/Source/Runtime/Engine/Public/Audio.h
+++ b/Engine/Source/Runtime/Engine/Public/Audio.h
@@ -513,6 +513,8 @@ public:
 	/** Returns the volume due to application behavior for the wave instance. */
 	float GetVolumeApp() const { return VolumeApp; }
 
+	bool IsSeekable() const;
+
 	/** Checks whether wave is streaming and streaming is supported */
 	bool IsStreaming() const;
 
diff --git a/Engine/Source/Runtime/Linux/AudioMixerSDL/Private/AudioMixerPlatformSDL.cpp b/Engine/Source/Runtime/Linux/AudioMixerSDL/Private/AudioMixerPlatformSDL.cpp
--- a/Engine/Source/Runtime/Linux/AudioMixerSDL/Private/AudioMixerPlatformSDL.cpp
+++ b/Engine/Source/Runtime/Linux/AudioMixerSDL/Private/AudioMixerPlatformSDL.cpp
@@ -324,12 +324,19 @@ namespace Audio
 
 	FName FMixerPlatformSDL::GetRuntimeFormat(USoundWave* InSoundWave)
 	{
+		static FName NAME_OGG(TEXT("OGG"));
+		static FName NAME_OPUS(TEXT("OPUS"));
+		static FName NAME_ADPCM(TEXT("ADPCM"));
+
 		if (InSoundWave->IsStreaming())
 		{
-			return FName(TEXT("OPUS"));
-		}
+			if (InSoundWave->IsSeekableStreaming())
+			{
+				return NAME_ADPCM;
+			}
 
-		static FName NAME_OGG(TEXT("OGG"));
+			return NAME_OPUS;
+		}
 		return NAME_OGG;
 	}
 
@@ -344,6 +351,11 @@ namespace Audio
 
 		if (InSoundWave->IsStreaming())
 		{
+			if (InSoundWave->IsSeekableStreaming())
+			{
+				return new FADPCMAudioInfo();
+			}
+
 			return new FOpusAudioInfo();
 		}
 
diff --git a/Engine/Source/Runtime/Mac/CoreAudio/Private/CoreAudioDevice.cpp b/Engine/Source/Runtime/Mac/CoreAudio/Private/CoreAudioDevice.cpp
--- a/Engine/Source/Runtime/Mac/CoreAudio/Private/CoreAudioDevice.cpp
+++ b/Engine/Source/Runtime/Mac/CoreAudio/Private/CoreAudioDevice.cpp
@@ -13,6 +13,7 @@
 #include "OpusAudioInfo.h"
 #include "AudioEffect.h"
 #include "CoreAudioEffects.h"
+#include "ADPCMAudioInfo.h"
 
 DEFINE_LOG_CATEGORY(LogCoreAudio);
 
@@ -513,6 +514,10 @@ class ICompressedAudioInfo* FCoreAudioDevice::CreateCompressedAudioInfo(USoundWa
 
 	if (SoundWave->IsStreaming())
 	{
+		if (SoundWave->IsSeekableStreaming())
+		{
+			return new FADPCMAudioInfo();
+		}
 		return new FOpusAudioInfo();
 	}
 
diff --git a/Engine/Source/Runtime/Windows/AudioMixerXAudio2/Private/AudioMixerPlatformXAudio2.cpp b/Engine/Source/Runtime/Windows/AudioMixerXAudio2/Private/AudioMixerPlatformXAudio2.cpp
--- a/Engine/Source/Runtime/Windows/AudioMixerXAudio2/Private/AudioMixerPlatformXAudio2.cpp
+++ b/Engine/Source/Runtime/Windows/AudioMixerXAudio2/Private/AudioMixerPlatformXAudio2.cpp
@@ -20,7 +20,7 @@
 #endif  //#if WITH_XMA2
 #include "OpusAudioInfo.h"
 #include "VorbisAudioInfo.h"
-
+#include "ADPCMAudioInfo.h"
 
 #include "CoreGlobals.h"
 #include "Misc/ConfigCacheIni.h"
@@ -813,20 +813,24 @@ namespace Audio
 		static FName NAME_OGG(TEXT("OGG"));
 		static FName NAME_OPUS(TEXT("OPUS"));
 		static FName NAME_XMA(TEXT("XMA"));
+		static FName NAME_ADPCM(TEXT("ADPCM"));
 
-#if WITH_XMA2 && USE_XMA2_FOR_STREAMING
-		if (InSoundWave->IsStreaming() && InSoundWave->NumChannels <= 2)
+		if (InSoundWave->IsStreaming())
 		{
-			return NAME_XMA;
-		}
+			if (InSoundWave->IsSeekableStreaming())
+			{
+				return NAME_ADPCM;
+			}
+
+#if WITH_XMA2 && USE_XMA2_FOR_STREAMING
+			if (InSoundWave->NumChannels <= 2)
+			{
+				return NAME_XMA;
+			}
 #endif
 
-		if (InSoundWave->IsStreaming())
-		{
 #if USE_VORBIS_FOR_STREAMING
 			return NAME_OGG;
-#else
-			return NAME_OPUS;
 #endif
 		}
 
@@ -849,6 +853,14 @@ namespace Audio
 	{
 		check(InSoundWave);
 
+		if (InSoundWave->IsStreaming())
+		{
+			if (InSoundWave->IsSeekableStreaming())
+			{
+				return new FADPCMAudioInfo();
+			}
+		}
+
 #if WITH_XMA2 && USE_XMA2_FOR_STREAMING
 		if (InSoundWave->IsStreaming() && InSoundWave->NumChannels <= 2 )
 		{
diff --git a/Engine/Source/Runtime/Windows/XAudio2/Private/XAudio2Device.cpp b/Engine/Source/Runtime/Windows/XAudio2/Private/XAudio2Device.cpp
--- a/Engine/Source/Runtime/Windows/XAudio2/Private/XAudio2Device.cpp
+++ b/Engine/Source/Runtime/Windows/XAudio2/Private/XAudio2Device.cpp
@@ -15,6 +15,7 @@
 #include "AudioEffect.h"
 #include "OpusAudioInfo.h"
 #include "VorbisAudioInfo.h"
+#include "ADPCMAudioInfo.h"
 #include "XAudio2Effects.h"
 #include "Interfaces/IAudioFormat.h"
 #include "HAL/PlatformAffinity.h"
@@ -364,21 +365,26 @@ class ICompressedAudioInfo* FXAudio2Device::CreateCompressedAudioInfo(USoundWave
 {
 	check(SoundWave);
 
-#if WITH_XMA2 && USE_XMA2_FOR_STREAMING
-	if (SoundWave->IsStreaming() && SoundWave->NumChannels <= 2 )
+	if (SoundWave->IsStreaming())
 	{
-		ICompressedAudioInfo* CompressedInfo = new FXMAAudioInfo();
-		if (!CompressedInfo)
+		if (SoundWave->IsSeekableStreaming())
 		{
-			UE_LOG(LogAudio, Error, TEXT("Failed to create new FXMAAudioInfo for streaming SoundWave %s: out of memory."), *SoundWave->GetName());
-			return nullptr;
+			return new FADPCMAudioInfo();
+		}
+
+#if WITH_XMA2 && USE_XMA2_FOR_STREAMING
+		if (SoundWave->NumChannels <= 2 )
+		{
+			ICompressedAudioInfo* CompressedInfo = new FXMAAudioInfo();
+			if (!CompressedInfo)
+			{
+				UE_LOG(LogAudio, Error, TEXT("Failed to create new FXMAAudioInfo for streaming SoundWave %s: out of memory."), *SoundWave->GetName());
+				return nullptr;
+			}
+			return CompressedInfo;
 		}
-		return CompressedInfo;
-	}
 #endif
 
-	if (SoundWave->IsStreaming())
-	{
 #if USE_VORBIS_FOR_STREAMING
 		return new FVorbisAudioInfo();
 #else
diff --git a/Engine/Source/Runtime/Windows/XAudio2/Private/XAudio2Source.cpp b/Engine/Source/Runtime/Windows/XAudio2/Private/XAudio2Source.cpp
--- a/Engine/Source/Runtime/Windows/XAudio2/Private/XAudio2Source.cpp
+++ b/Engine/Source/Runtime/Windows/XAudio2/Private/XAudio2Source.cpp
@@ -571,7 +571,7 @@ bool FXAudio2SoundSource::Init(FWaveInstance* InWaveInstance)
 
 			if (WaveInstance->StartTime > 0.0f)
 			{
-				if (WaveInstance->WaveData->bStreaming)
+				if (WaveInstance->WaveData->bStreaming && !WaveInstance->WaveData->bSeekableStreaming)
 				{
 					UE_LOG(LogXAudio2, Verbose, TEXT("Seeking (aka start time) is not supported for streaming sound waves ('%s')."), *InWaveInstance->GetName());
 				}
diff --git a/Engine/Source/Runtime/Windows/XAudio2/Public/XAudio2Device.h b/Engine/Source/Runtime/Windows/XAudio2/Public/XAudio2Device.h
--- a/Engine/Source/Runtime/Windows/XAudio2/Public/XAudio2Device.h
+++ b/Engine/Source/Runtime/Windows/XAudio2/Public/XAudio2Device.h
@@ -119,30 +119,32 @@ private:
 
 	virtual FName GetRuntimeFormat(USoundWave* SoundWave) override
 	{
-		static FName NAME_OPUS(TEXT("OPUS"));
-		static FName NAME_OGG(TEXT("OGG"));
-		static FName NAME_XMA(TEXT("XMA"));
+		static const FName NAME_ADPCM(TEXT("ADPCM"));
+		static const FName NAME_OGG(TEXT("OGG"));
+		static const FName NAME_OPUS(TEXT("OPUS"));
+		static const FName NAME_XMA(TEXT("XMA"));
 
-#if WITH_XMA2 && USE_XMA2_FOR_STREAMING
-		if (SoundWave->IsStreaming() && SoundWave->NumChannels <= 2)
+		if (SoundWave->IsStreaming())
 		{
-			return NAME_XMA;
-		}
-#endif
-
+			if (SoundWave->IsSeekableStreaming())
+			{
+				return NAME_ADPCM;
+			}
 
-#if WITH_OGGVORBIS
+#if WITH_XMA2 && USE_XMA2_FOR_STREAMING
+			else if (SoundWave->NumChannels <= 2)
+			{
+				return NAME_XMA;
+			}
+#endif // WITH_XMA2 && USE_XMA2_FOR_STREAMING
 
-		if (SoundWave->IsStreaming())
-		{
 #if USE_VORBIS_FOR_STREAMING
 			return NAME_OGG;
 #else
 			return NAME_OPUS;
-#endif
+#endif // else USE_VORBIS_FOR_STREAMING
 		}
 
-
 #if WITH_XMA2
 		if (SoundWave->NumChannels > 2)
 		{
@@ -153,13 +155,9 @@ private:
 		{
 			return NAME_XMA;
 		}
-#else
-		return NAME_OGG;
-#endif
+#endif // WITH_XMA2
 
-#else //WITH_OGGVORBIS
-		return NAME_XMA;
-#endif //WITH_OGGVORBIS
+		return NAME_OGG;
 	}
 
 	virtual bool HasCompressedAudioInfoClass(USoundWave* SoundWave) override;
