Port CL# 5371058 to Release-4.22

Fix crash with leaving a world that is being kept alive in GC by forcing pending kill on the world as it is destroyed.

diff --git a/Engine/Source/Runtime/Engine/Classes/Engine/World.h b/Engine/Source/Runtime/Engine/Classes/Engine/World.h
--- a/Engine/Source/Runtime/Engine/Classes/Engine/World.h
+++ b/Engine/Source/Runtime/Engine/Classes/Engine/World.h
@@ -2515,7 +2515,7 @@ public:
 	void DestroyWorld( bool bInformEngineOfWorld, UWorld* NewWorld = nullptr );
 
 	/** 
-	 * Marks all objects that have this World as an Outer as pending kill
+	 * Marks this world and all objects within as pending kill
 	 */
 	void MarkObjectsPendingKill();
 
diff --git a/Engine/Source/Runtime/Engine/Private/World.cpp b/Engine/Source/Runtime/Engine/Private/World.cpp
--- a/Engine/Source/Runtime/Engine/Private/World.cpp
+++ b/Engine/Source/Runtime/Engine/Private/World.cpp
@@ -1564,6 +1564,8 @@ void UWorld::MarkObjectsPendingKill()
 		Object->MarkPendingKill();
 	};
 	ForEachObjectWithOuter(this, MarkObjectPendingKill, true, RF_NoFlags, EInternalObjectFlags::PendingKill);
+
+	MarkPendingKill();
 	bMarkedObjectsPendingKill = true;
 }
 
