Fix comment on function FBodyInstance::GetMassSpaceLocal.

diff --git a/Engine/Source/Runtime/Engine/Classes/PhysicsEngine/BodyInstance.h b/Engine/Source/Runtime/Engine/Classes/PhysicsEngine/BodyInstance.h
--- a/Engine/Source/Runtime/Engine/Classes/PhysicsEngine/BodyInstance.h
+++ b/Engine/Source/Runtime/Engine/Classes/PhysicsEngine/BodyInstance.h
@@ -565,6 +565,8 @@ public:
 
 	/** Returns the mass coordinate system to world space transform (position is world center of mass, rotation is world inertia orientation) */
 	FTransform GetMassSpaceToWorldSpace() const;
+
+	/** Returns the mass coordinate system to local space transform (position is local center of mass, rotation should be identity) */
 	FTransform GetMassSpaceLocal() const;
 
 	/** TODO: this only works at runtime when the physics state has been created. Any changes that result in recomputing mass properties will not properly remember this */
