Correct detection of a mobile PIE session within UWorld::IsPlayInMobilePreview()

diff --git a/Engine/Source/Runtime/Engine/Private/World.cpp b/Engine/Source/Runtime/Engine/Private/World.cpp
--- a/Engine/Source/Runtime/Engine/Private/World.cpp
+++ b/Engine/Source/Runtime/Engine/Private/World.cpp
@@ -96,6 +96,7 @@
 	#include "HierarchicalLODUtilitiesModule.h"
 	#include "ObjectTools.h"
 	#include "Engine/LODActor.h"
+	#include "PIEPreviewDeviceProfileSelectorModule.h"
 #endif
 
 
@@ -6400,6 +6401,14 @@ bool UWorld::IsPlayInPreview() const
 
 bool UWorld::IsPlayInMobilePreview() const
 {
+#if WITH_EDITOR
+	if (FPIEPreviewDeviceModule::IsRequestingPreviewDevice()
+		|| FParse::Param(FCommandLine::Get(), TEXT("featureleveles2"))
+		|| FParse::Param(FCommandLine::Get(), TEXT("featureleveles31")))
+	{
+		return true;
+	}
+#endif // WITH_EDITOR
 	return FParse::Param(FCommandLine::Get(), TEXT("simmobile")) && !IsPlayInVulkanPreview();
 }
 
