Use accessors for LastRenderTime and LastRenderTimeOnScreen in preparation for them being made private
#rnx

diff --git a/Engine/Plugins/Developer/AnimationSharing/Source/AnimationSharing/Private/AnimationSharingManager.cpp b/Engine/Plugins/Developer/AnimationSharing/Source/AnimationSharing/Private/AnimationSharingManager.cpp
--- a/Engine/Plugins/Developer/AnimationSharing/Source/AnimationSharing/Private/AnimationSharingManager.cpp
+++ b/Engine/Plugins/Developer/AnimationSharing/Source/AnimationSharing/Private/AnimationSharingManager.cpp
@@ -1357,7 +1357,7 @@ void UAnimSharingInstance::TickActorStates()
 			ActorData.bRequiresTick = ActorData.SignificanceValue >= ScalabilitySettings->TickSignificanceValue.Default;
 			for (int32 ComponentIndex : ActorData.ComponentIndices)
 			{
-				if (PerComponentData[ComponentIndex].Component->LastRenderTime > (WorldTime - 1.f))
+				if (PerComponentData[ComponentIndex].Component->GetLastRenderTime() > (WorldTime - 1.f))
 				{
 					PerComponentData[ComponentIndex].Component->bRecentlyRendered = true;
 					ActorData.bRequiresTick = true;
diff --git a/Engine/Plugins/FX/Niagara/Source/Niagara/Private/NiagaraSystemInstance.cpp b/Engine/Plugins/FX/Niagara/Source/Niagara/Private/NiagaraSystemInstance.cpp
--- a/Engine/Plugins/FX/Niagara/Source/Niagara/Private/NiagaraSystemInstance.cpp
+++ b/Engine/Plugins/FX/Niagara/Source/Niagara/Private/NiagaraSystemInstance.cpp
@@ -415,7 +415,7 @@ void FNiagaraSystemInstance::Reset(FNiagaraSystemInstance::EResetMode Mode, bool
 		return;
 	}
 
-	Component->LastRenderTime = Component->GetWorld()->GetTimeSeconds();
+	Component->SetLastRenderTime(Component->GetWorld()->GetTimeSeconds());
 
 	SetPaused(false);
 
@@ -1162,7 +1162,7 @@ void FNiagaraSystemInstance::TickInstanceParameters(float DeltaSeconds)
 	check(World);
 	float WorldTime = World->GetTimeSeconds();
 	//Bias the LastRenderTime slightly to account for any delay as it's written by the RT.
-	float SafeTimeSinceRendererd = FMath::Max(0.0f, WorldTime - Component->LastRenderTime - GLastRenderTimeSafetyBias);
+	float SafeTimeSinceRendererd = FMath::Max(0.0f, WorldTime - Component->GetLastRenderTime() - GLastRenderTimeSafetyBias);
 	SystemTimeSinceRenderedParam.SetValue(SafeTimeSinceRendererd);
 	
 	OwnerExecutionStateParam.SetValue((int32)RequestedExecutionState);
diff --git a/Engine/Plugins/Runtime/AnimationBudgetAllocator/Source/AnimationBudgetAllocator/Private/AnimationBudgetAllocator.cpp b/Engine/Plugins/Runtime/AnimationBudgetAllocator/Source/AnimationBudgetAllocator/Private/AnimationBudgetAllocator.cpp
--- a/Engine/Plugins/Runtime/AnimationBudgetAllocator/Source/AnimationBudgetAllocator/Private/AnimationBudgetAllocator.cpp
+++ b/Engine/Plugins/Runtime/AnimationBudgetAllocator/Source/AnimationBudgetAllocator/Private/AnimationBudgetAllocator.cpp
@@ -220,7 +220,7 @@ void FAnimationBudgetAllocator::QueueSortedComponentIndices(float InDeltaSeconds
 			{
 				auto ShouldComponentTick = [WorldTime](const USkeletalMeshComponentBudgeted* InComponent, const FComponentData& InComponentData)
 				{
-					return ((InComponent->LastRenderTime > WorldTime) ||
+					return ((InComponent->GetLastRenderTime() > WorldTime) ||
 						(InComponent->GetShouldUseActorRenderedFlag() && InComponent->GetAttachmentRootActor() && InComponent->GetAttachmentRootActor()->WasRecentlyRendered())  ||
 							InComponentData.bTickEvenIfNotRendered ||
 							InComponent->ShouldTickPose() ||
@@ -246,7 +246,7 @@ void FAnimationBudgetAllocator::QueueSortedComponentIndices(float InDeltaSeconds
 				{
 					// Push into a separate limited list if we are 'tick even if not rendered'.
 					// Skip offscreen components with a significance of zero or less.
-					if(Component->LastRenderTime <= WorldTime && ComponentData.bTickEvenIfNotRendered && ComponentData.Significance > 0.0f)
+					if(Component->GetLastRenderTime() <= WorldTime && ComponentData.bTickEvenIfNotRendered && ComponentData.Significance > 0.0f)
 					{
 						NonRenderedComponentData.Add(ComponentIndex);
 					}
diff --git a/Engine/Source/Runtime/Engine/Classes/Components/PrimitiveComponent.h b/Engine/Source/Runtime/Engine/Classes/Components/PrimitiveComponent.h
--- a/Engine/Source/Runtime/Engine/Classes/Components/PrimitiveComponent.h
+++ b/Engine/Source/Runtime/Engine/Classes/Components/PrimitiveComponent.h
@@ -607,7 +607,6 @@ public:
 	float BoundsScale;
 
 	/** Last time the component was submitted for rendering (called FScene::AddPrimitive). */
-	UPROPERTY(transient)
 	float LastSubmitTime;
 
 	/**
@@ -615,13 +614,17 @@ public:
 	 * from the render thread, which is up to a frame behind the game thread, so you should allow this time to
 	 * be at least a frame behind the game thread's world time before you consider the actor non-visible.
 	 */
-	UPROPERTY(transient)
 	float LastRenderTime;
 
 	/** Same as LastRenderTimeOnScreen but only updated if the component is on screen. Used by the texture streamer. */
-	UPROPERTY(transient)
 	float LastRenderTimeOnScreen;
 
+public:
+
+	void SetLastRenderTime(float InLastRenderTime);
+	float GetLastRenderTime() const { return LastRenderTime; }
+	float GetLastRenderTimeOnScreen() const { return LastRenderTimeOnScreen; }
+
 	/**
 	 * Set of actors to ignore during component sweeps in MoveComponent().
 	 * All components owned by these actors will be ignored when this component moves or updates overlaps.
diff --git a/Engine/Source/Runtime/Engine/Private/Components/MovementComponent.cpp b/Engine/Source/Runtime/Engine/Private/Components/MovementComponent.cpp
--- a/Engine/Source/Runtime/Engine/Private/Components/MovementComponent.cpp
+++ b/Engine/Source/Runtime/Engine/Private/Components/MovementComponent.cpp
@@ -336,7 +336,7 @@ bool UMovementComponent::ShouldSkipUpdate(float DeltaTime) const
 
 		const float RenderTimeThreshold = 0.41f;
 		UWorld* TheWorld = GetWorld();
-		if (UpdatedPrimitive && TheWorld->TimeSince(UpdatedPrimitive->LastRenderTime) <= RenderTimeThreshold)
+		if (UpdatedPrimitive && TheWorld->TimeSince(UpdatedPrimitive->GetLastRenderTime()) <= RenderTimeThreshold)
 		{
 			return false; // Rendered, don't skip it.
 		}
@@ -349,7 +349,7 @@ bool UMovementComponent::ShouldSkipUpdate(float DeltaTime) const
 			const UPrimitiveComponent* PrimitiveChild = Cast<UPrimitiveComponent>(Child);
 			if (PrimitiveChild)
 			{
-				if (PrimitiveChild->IsRegistered() && TheWorld->TimeSince(PrimitiveChild->LastRenderTime) <= RenderTimeThreshold)
+				if (PrimitiveChild->IsRegistered() && TheWorld->TimeSince(PrimitiveChild->GetLastRenderTime()) <= RenderTimeThreshold)
 				{
 					return false; // Rendered, don't skip it.
 				}
diff --git a/Engine/Source/Runtime/Engine/Private/Components/PrimitiveComponent.cpp b/Engine/Source/Runtime/Engine/Private/Components/PrimitiveComponent.cpp
--- a/Engine/Source/Runtime/Engine/Private/Components/PrimitiveComponent.cpp
+++ b/Engine/Source/Runtime/Engine/Private/Components/PrimitiveComponent.cpp
@@ -3464,6 +3464,11 @@ void UPrimitiveComponent::SetCustomNavigableGeometry(const EHasCustomNavigableGe
 	bHasCustomNavigableGeometry = InType;
 }
  
+void UPrimitiveComponent::SetLastRenderTime(float InLastRenderTime)
+{
+	LastRenderTime = InLastRenderTime;
+}
+
 #if WITH_EDITOR
 const bool UPrimitiveComponent::ShouldGenerateAutoLOD(const int32 HierarchicalLevelIndex) const
 {	
diff --git a/Engine/Source/Runtime/Engine/Private/Components/SkinnedMeshComponent.cpp b/Engine/Source/Runtime/Engine/Private/Components/SkinnedMeshComponent.cpp
--- a/Engine/Source/Runtime/Engine/Private/Components/SkinnedMeshComponent.cpp
+++ b/Engine/Source/Runtime/Engine/Private/Components/SkinnedMeshComponent.cpp
@@ -779,7 +779,7 @@ void USkinnedMeshComponent::TickComponent(float DeltaTime, enum ELevelTick TickT
 	Super::TickComponent(DeltaTime, TickType, ThisTickFunction);
 
 	// See if this mesh was rendered recently. This has to happen first because other data will rely on this
-	bRecentlyRendered = (LastRenderTime > GetWorld()->TimeSeconds - 1.0f);
+	bRecentlyRendered = (GetLastRenderTime() > GetWorld()->TimeSeconds - 1.0f);
 
 	// Update component's LOD settings
 	// This must be done BEFORE animation Update and Evaluate (TickPose and RefreshBoneTransforms respectively)
diff --git a/Engine/Source/Runtime/Engine/Private/LocalPlayer.cpp b/Engine/Source/Runtime/Engine/Private/LocalPlayer.cpp
--- a/Engine/Source/Runtime/Engine/Private/LocalPlayer.cpp
+++ b/Engine/Source/Runtime/Engine/Private/LocalPlayer.cpp
@@ -1227,7 +1227,7 @@ bool ULocalPlayer::HandleListSkelMeshesCommand( const TCHAR* Cmd, FOutputDevice&
 				check(SkeletalMeshComponent);
 				UWorld* World = SkeletalMeshComponent->GetWorld();
 				check(World);
-				float TimeSinceLastRender = World->GetTimeSeconds() - SkeletalMeshComponent->LastRenderTime;
+				float TimeSinceLastRender = World->GetTimeSeconds() - SkeletalMeshComponent->GetLastRenderTime();
 
 				UE_LOG(LogPlayerManagement, Log, TEXT("%s%2i  Component    : %s"),
 					(TimeSinceLastRender > 0.5) ? TEXT(" ") : TEXT("*"),
diff --git a/Engine/Source/Runtime/Engine/Private/Particles/ParticleComponents.cpp b/Engine/Source/Runtime/Engine/Private/Particles/ParticleComponents.cpp
--- a/Engine/Source/Runtime/Engine/Private/Particles/ParticleComponents.cpp
+++ b/Engine/Source/Runtime/Engine/Private/Particles/ParticleComponents.cpp
@@ -3566,7 +3566,7 @@ bool UParticleSystemComponent::CanConsiderInvisible()const
 		const float ClampedMaxSecondsBeforeInactive = MaxSecondsBeforeInactive > 0 ? FMath::Max(MaxSecondsBeforeInactive, 0.1f) : 0.0f;
 		if (ClampedMaxSecondsBeforeInactive > 0.0f && AccumTickTime > ClampedMaxSecondsBeforeInactive && World->IsGameWorld())
 		{
-			return World->GetTimeSeconds() > (LastRenderTime + ClampedMaxSecondsBeforeInactive);
+			return World->GetTimeSeconds() > (GetLastRenderTime() + ClampedMaxSecondsBeforeInactive);
 		}
 	}
 	return false;
@@ -6030,7 +6030,7 @@ void UParticleSystemComponent::ActivateSystem(bool bFlagAsJustAttached)
 	World = GetWorld(); // refresh the world pointer as it may have changed by this point
 	if(!bWasDeactivated && !bWasCompleted && ensure(World))
 	{
-		LastRenderTime = World->GetTimeSeconds();
+		SetLastRenderTime(World->GetTimeSeconds());
 	}
 }
 
@@ -6136,7 +6136,7 @@ void UParticleSystemComponent::DeactivateSystem()
 	//TODO: What if there are immortal particles but bKillOnDeactivate is false? Need to mark emitters with currently immortal particles, kill them and warn the user.
 	SetComponentTickEnabled(true);
 
-	LastRenderTime = World->GetTimeSeconds();
+	SetLastRenderTime(World->GetTimeSeconds());
 }
 
 void UParticleSystemComponent::CancelAutoAttachment(bool bDetachFromParent)
diff --git a/Engine/Source/Runtime/Engine/Private/Particles/ParticleModules_Collision.cpp b/Engine/Source/Runtime/Engine/Private/Particles/ParticleModules_Collision.cpp
--- a/Engine/Source/Runtime/Engine/Private/Particles/ParticleModules_Collision.cpp
+++ b/Engine/Source/Runtime/Engine/Private/Particles/ParticleModules_Collision.cpp
@@ -217,7 +217,7 @@ void UParticleModuleCollision::Update(FParticleEmitterInstance* Owner, int32 Off
 
 		// LOD collision based on visibility
 		// This is at the 'emitter instance' level as it will be true or false for the whole instance...
-		if (bCollideOnlyIfVisible && ((World->TimeSeconds - Owner->Component->LastRenderTime) > GParticleCollisionIgnoreInvisibleTime))
+		if (bCollideOnlyIfVisible && ((World->TimeSeconds - Owner->Component->GetLastRenderTime()) > GParticleCollisionIgnoreInvisibleTime))
 		{
 			// no collision if not recently rendered
 			bIgnoreAllCollision = true;
diff --git a/Engine/Source/Runtime/Engine/Private/Streaming/TextureInstanceState.cpp b/Engine/Source/Runtime/Engine/Private/Streaming/TextureInstanceState.cpp
--- a/Engine/Source/Runtime/Engine/Private/Streaming/TextureInstanceState.cpp
+++ b/Engine/Source/Runtime/Engine/Private/Streaming/TextureInstanceState.cpp
@@ -16,7 +16,7 @@
 
 int32 FRenderAssetInstanceState::AddBounds(const UPrimitiveComponent* Component)
 {
-	return AddBounds(Component->Bounds, PackedRelativeBox_Identity, Component, Component->LastRenderTimeOnScreen, Component->Bounds.Origin, 0, 0, FLT_MAX);
+	return AddBounds(Component->Bounds, PackedRelativeBox_Identity, Component, Component->GetLastRenderTimeOnScreen(), Component->Bounds.Origin, 0, 0, FLT_MAX);
 }
 
 int32 FRenderAssetInstanceState::AddBounds(const FBoxSphereBounds& Bounds, uint32 PackedRelativeBox, const UPrimitiveComponent* InComponent, float LastRenderTime, const FVector4& RangeOrigin, float MinDistanceSq, float MinRangeSq, float MaxRangeSq)
@@ -334,7 +334,7 @@ EAddComponentResult FRenderAssetInstanceState::AddComponent(const UPrimitiveComp
 				++NumOfBoundReferences;
 			}
 
-			const int32 BoundsIndex = AddBounds(FBoxSphereBounds(ForceInit), Info.PackedRelativeBox, Component, Component->LastRenderTimeOnScreen, FVector(ForceInit), 0, 0, FLT_MAX);
+			const int32 BoundsIndex = AddBounds(FBoxSphereBounds(ForceInit), Info.PackedRelativeBox, Component, Component->GetLastRenderTimeOnScreen(), FVector(ForceInit), 0, 0, FLT_MAX);
 			AddRenderAssetElements(Component, TArrayView<FStreamingRenderAssetPrimitiveInfo>(RenderAssetInstanceInfos.GetData() + InfoIndex, NumOfBoundReferences), BoundsIndex, ComponentLink);
 			BoundsToUnpack.Push(BoundsIndex);
 
@@ -374,7 +374,7 @@ EAddComponentResult FRenderAssetInstanceState::AddComponent(const UPrimitiveComp
 			}
 
 			GetDistanceAndRange(Component, Info.Bounds, MinDistanceSq, MinRangeSq, MaxRangeSq);
-			const int32 BoundsIndex = AddBounds(Info.Bounds, PackedRelativeBox_Identity, Component, Component->LastRenderTimeOnScreen, Component->Bounds.Origin, MinDistanceSq, MinRangeSq, MaxRangeSq);
+			const int32 BoundsIndex = AddBounds(Info.Bounds, PackedRelativeBox_Identity, Component, Component->GetLastRenderTimeOnScreen(), Component->Bounds.Origin, MinDistanceSq, MinRangeSq, MaxRangeSq);
 			AddRenderAssetElements(Component, TArrayView<FStreamingRenderAssetPrimitiveInfo>(RenderAssetInstanceInfos.GetData() + InfoIndex, NumOfBoundReferences), BoundsIndex, ComponentLink);
 
 			InfoIndex += NumOfBoundReferences;
@@ -492,7 +492,7 @@ void FRenderAssetInstanceState::UpdateBounds(const UPrimitiveComponent* Componen
 			const FElement& Element = Elements[ElementIndex];
 			if (Element.BoundsIndex != INDEX_NONE)
 			{
-				Bounds4[Element.BoundsIndex / 4].FullUpdate(Element.BoundsIndex % 4, Component->Bounds, Component->LastRenderTimeOnScreen);
+				Bounds4[Element.BoundsIndex / 4].FullUpdate(Element.BoundsIndex % 4, Component->Bounds, Component->GetLastRenderTimeOnScreen());
 			}
 			ElementIndex = Element.NextComponentLink;
 		}
@@ -504,7 +504,7 @@ bool FRenderAssetInstanceState::UpdateBounds(int32 BoundIndex)
 	const UPrimitiveComponent* Component = ensure(Bounds4Components.IsValidIndex(BoundIndex)) ? Bounds4Components[BoundIndex] : nullptr;
 	if (Component)
 	{
-		Bounds4[BoundIndex / 4].FullUpdate(BoundIndex % 4, Component->Bounds, Component->LastRenderTimeOnScreen);
+		Bounds4[BoundIndex / 4].FullUpdate(BoundIndex % 4, Component->Bounds, Component->GetLastRenderTimeOnScreen());
 		return true;
 	}
 	else
@@ -531,13 +531,13 @@ bool FRenderAssetInstanceState::ConditionalUpdateBounds(int32 BoundIndex)
 
 			if (0.5f * FMath::Min3<float>(XSquared, YSquared, ZSquared) <= RadiusSquared && RadiusSquared <= 2.f * (XSquared + YSquared + ZSquared))
 			{
-				Bounds4[BoundIndex / 4].FullUpdate(BoundIndex % 4, Bounds, Component->LastRenderTimeOnScreen);
+				Bounds4[BoundIndex / 4].FullUpdate(BoundIndex % 4, Bounds, Component->GetLastRenderTimeOnScreen());
 				return true;
 			}
 		}
 		else // Otherwise we assume it is guarantied to be good.
 		{
-			Bounds4[BoundIndex / 4].FullUpdate(BoundIndex % 4, Component->Bounds, Component->LastRenderTimeOnScreen);
+			Bounds4[BoundIndex / 4].FullUpdate(BoundIndex % 4, Component->Bounds, Component->GetLastRenderTimeOnScreen());
 			return true;
 		}
 	}
@@ -552,7 +552,7 @@ void FRenderAssetInstanceState::UpdateLastRenderTimeAndMaxDrawDistance(int32 Bou
 	{
 		const int32 Bounds4Idx = BoundIndex / 4;
 		const int32 SubIdx = BoundIndex & 3;
-		Bounds4[Bounds4Idx].UpdateLastRenderTime(SubIdx, Component->LastRenderTimeOnScreen);
+		Bounds4[Bounds4Idx].UpdateLastRenderTime(SubIdx, Component->GetLastRenderTimeOnScreen());
 		// The min draw distances of HLODs can change dynamically (see Tick, PauseDitherTransition,
 		// and StartDitherTransition methods of ALODActor)
 		const UPrimitiveComponent* LODParent = Component->GetLODParentPrimitive();
diff --git a/Engine/Source/Runtime/UMG/Private/Components/WidgetComponent.cpp b/Engine/Source/Runtime/UMG/Private/Components/WidgetComponent.cpp
--- a/Engine/Source/Runtime/UMG/Private/Components/WidgetComponent.cpp
+++ b/Engine/Source/Runtime/UMG/Private/Components/WidgetComponent.cpp
@@ -1069,7 +1069,7 @@ bool UWidgetComponent::ShouldDrawWidget() const
 	if ( IsVisible() )
 	{
 		// If we don't tick when off-screen, don't bother ticking if it hasn't been rendered recently
-		if ( TickWhenOffscreen || GetWorld()->TimeSince(LastRenderTime) <= RenderTimeThreshold )
+		if ( TickWhenOffscreen || GetWorld()->TimeSince(GetLastRenderTime()) <= RenderTimeThreshold )
 		{
 			if ( ( GetCurrentTime() - LastWidgetRenderTime) >= RedrawTime )
 			{
