Rework of the render asset streaming synchronization :
- Update object are now refcounted, allowed async task to have references without preventing the main thread from releasing references
- This allows a workaround for a GC bug where async task are not scheduled and the engine times out because streaming cancelation does not completed
- The new framework allows recursive implementation where async task could eventually execute rendercommands once better synchro is added to RHIAsyncCreateTexture2D and RHIAsyncReallocateTexture2D.

diff --git a/Engine/Source/Runtime/Engine/Classes/Engine/StaticMesh.h b/Engine/Source/Runtime/Engine/Classes/Engine/StaticMesh.h
--- a/Engine/Source/Runtime/Engine/Classes/Engine/StaticMesh.h
+++ b/Engine/Source/Runtime/Engine/Classes/Engine/StaticMesh.h
@@ -24,6 +24,7 @@
 #include "MeshDescription.h"
 #include "MeshAttributes.h"
 #include "MeshAttributeArray.h"
+#include "RenderAssetUpdate.h"
 
 #include "StaticMesh.generated.h"
 
@@ -787,7 +788,7 @@ protected:
 	UPROPERTY(EditAnywhere, AdvancedDisplay, Instanced, Category = StaticMesh)
 	TArray<UAssetUserData*> AssetUserData;
 
-	FStaticMeshUpdate* PendingUpdate;
+	TRefCountPtr<FRenderAssetUpdate> PendingUpdate;
 
 	friend struct FStaticMeshUpdateContext;
 	friend class FStaticMeshUpdate;
diff --git a/Engine/Source/Runtime/Engine/Classes/Engine/Texture2D.h b/Engine/Source/Runtime/Engine/Classes/Engine/Texture2D.h
--- a/Engine/Source/Runtime/Engine/Classes/Engine/Texture2D.h
+++ b/Engine/Source/Runtime/Engine/Classes/Engine/Texture2D.h
@@ -8,6 +8,7 @@
 #include "UObject/ScriptMacros.h"
 #include "Engine/Texture.h"
 #include "TextureResource.h"
+#include "RenderAssetUpdate.h"
 #include "Texture2D.generated.h"
 
 class FTexture2DResourceMem;
@@ -90,7 +91,7 @@ public:
 protected:
 
 	/** Helper to manage the current pending update following a call to StreamIn() or StreamOut(). */
-	class FTexture2DUpdate* PendingUpdate;
+	TRefCountPtr<FRenderAssetUpdate> PendingUpdate;
 	friend class FTexture2DUpdate;
 
 public:
diff --git a/Engine/Source/Runtime/Engine/Private/StaticMesh.cpp b/Engine/Source/Runtime/Engine/Private/StaticMesh.cpp
--- a/Engine/Source/Runtime/Engine/Private/StaticMesh.cpp
+++ b/Engine/Source/Runtime/Engine/Private/StaticMesh.cpp
@@ -4797,10 +4797,6 @@ bool UStaticMesh::StreamOut(int32 NewMipCount)
 	if (bIsStreamable && !PendingUpdate && RenderData.IsValid() && RenderData->bReadyForStreaming && NewMipCount < GetNumResidentMips())
 	{
 		PendingUpdate = new FStaticMeshStreamOut(this, NewMipCount);
-
-		// The object starts in the locked state while it is being initialized.
-		PendingUpdate->DoUnlock();
-
 		return !PendingUpdate->IsCancelled();
 	}
 	return false;
@@ -4835,9 +4831,6 @@ bool UStaticMesh::StreamIn(int32 NewMipCount, bool bHighPrio)
 				PendingUpdate = new FStaticMeshStreamIn_IO_RenderThread(this, NewMipCount, bHighPrio);
 			}
 		}
-
-		// The object starts in the locked state while it is being initialized.
-		PendingUpdate->DoUnlock();
 		return !PendingUpdate->IsCancelled();
 	}
 	return false;
@@ -4853,8 +4846,16 @@ bool UStaticMesh::UpdateStreamingStatus(bool bWaitForMipFading)
 			PendingUpdate->Abort();
 		}
 
-		// When the renderthread is the gamethread, allow the Tick to execute rendercommands.
-		PendingUpdate->Tick(this, GIsThreadedRendering ? FStaticMeshUpdate::TT_None : FStaticMeshUpdate::TT_Render);
+		// When there is no renderthread, allow the gamethread to tick as the renderthread.
+		FRenderAssetUpdate::EThreadType TickThread = GIsThreadedRendering ? FRenderAssetUpdate::TT_None : FRenderAssetUpdate::TT_Render;
+		if (HasAnyFlags(RF_BeginDestroyed) && PendingUpdate->GetRelevantThread() == FRenderAssetUpdate::TT_Async)
+		{
+			// To avoid async tasks from timing out the GC, we tick as Async to force completion if this is relevant.
+			// This could lead the asset from releasing the PendingUpdate, which will be deleted once the async task completes.
+			TickThread = FRenderAssetUpdate::TT_Async;
+		}
+		PendingUpdate->Tick(TickThread);
+
 		if (!PendingUpdate->IsCompleted())
 		{
 			return true;
@@ -4864,8 +4865,7 @@ bool UStaticMesh::UpdateStreamingStatus(bool bWaitForMipFading)
 		const bool bRebuildPlatformData = PendingUpdate->DDCIsInvalid() && !IsPendingKillOrUnreachable();
 #endif
 
-		delete PendingUpdate;
-		PendingUpdate = nullptr;
+		PendingUpdate.SafeRelease();
 
 #if WITH_EDITOR
 		if (GIsEditor)
diff --git a/Engine/Source/Runtime/Engine/Private/Streaming/RenderAssetUpdate.cpp b/Engine/Source/Runtime/Engine/Private/Streaming/RenderAssetUpdate.cpp
--- a/Engine/Source/Runtime/Engine/Private/Streaming/RenderAssetUpdate.cpp
+++ b/Engine/Source/Runtime/Engine/Private/Streaming/RenderAssetUpdate.cpp
@@ -4,7 +4,7 @@
 RenderAssetUpdate.cpp: Base class of helpers to stream in and out texture/mesh LODs.
 =============================================================================*/
 
-#include "Streaming/RenderAssetUpdate.h"
+#include "RenderAssetUpdate.h"
 #include "Engine/StreamableRenderAsset.h"
 #include "Streaming/TextureStreamingHelpers.h"
 #include "UObject/UObjectIterator.h"
@@ -17,6 +17,12 @@ static FAutoConsoleVariableRef CVarStreamingFlushTimeOut(
 	ECVF_Default
 );
 
+static TAutoConsoleVariable<int32> CVarStreamingStressTestExtraAsyncLatency(
+	TEXT("r.Streaming.StressTest.ExtraAsyncLatency"),
+	0,
+	TEXT("An extra latency in milliseconds for each async task when doing the stress test."),
+	ECVF_Cheat);
+
 
 volatile int32 GRenderAssetStreamingSuspension = 0;
 
@@ -108,3 +114,159 @@ void ResumeRenderAssetStreaming()
 	FPlatformAtomics::InterlockedDecrement(&GRenderAssetStreamingSuspension);
 	ensure(GRenderAssetStreamingSuspension >= 0);
 }
+
+FRenderAssetUpdate::FRenderAssetUpdate(UStreamableRenderAsset* InAsset, int32 InRequestedMips)
+	: PendingFirstMip(INDEX_NONE)
+	, RequestedMips(INDEX_NONE)
+	, ScheduledRenderTasks(0)
+	, ScheduledAsyncTasks(0)
+	, StreamableAsset(InAsset)
+	, bIsCancelled(false)
+	, TaskState(TS_Init)
+{
+	check(InAsset);
+
+	const int32 NonStreamingMipCount = InAsset->GetNumNonStreamingMips();
+	const int32 MaxMipCount = InAsset->GetNumMipsForStreaming();
+	InRequestedMips = FMath::Clamp<int32>(InRequestedMips, NonStreamingMipCount, MaxMipCount);
+
+	if (InRequestedMips > 0 && InRequestedMips != InAsset->GetNumResidentMips() && InAsset->bIsStreamable)
+	{
+		RequestedMips = InRequestedMips;
+		PendingFirstMip = MaxMipCount - RequestedMips;
+	}
+	else // This shouldn't happen but if it does, then the update is canceled
+	{
+		bIsCancelled = true;
+	}
+}
+
+FRenderAssetUpdate::~FRenderAssetUpdate()
+{
+	// Work must be done here because derived destructors have been called now and so derived members are invalid.
+	check(TaskSynchronization.GetValue() == 0 && TaskState == TS_Done);
+}
+
+uint32 FRenderAssetUpdate::Release() const
+{
+	uint32 NewValue = (uint32)NumRefs.Decrement();
+	if (NewValue == 0)
+	{
+		if (ensure(TaskState == TS_Done) && !TaskSynchronization.GetValue())
+		{
+			delete this;
+		}
+		else
+		{
+			// Can't delete this object if some other system has some token to decrement.
+			UE_LOG(LogContentStreaming, Error, TEXT("RenderAssetUpdate is leaking (%s, State=%d)"), *StreamableAsset->GetFullName(), (int32)TaskState);
+		}
+	}
+	return NewValue;
+}
+
+void FRenderAssetUpdate::Tick(EThreadType InCurrentThread)
+{
+	if (TaskState != TS_Done)
+	{
+		bool bIsLocked = true;
+
+		// Should we do aynthing about FApp::ShouldUseThreadingForPerformance()? For example to prevent async thread from stalling game/renderthreads.
+		// When the renderthread is the gamethread, don't lock if this is the renderthread to prevent stalling on low priority async tasks.
+		if (InCurrentThread == TT_None || (InCurrentThread == TT_Render && !GIsThreadedRendering))
+		{
+			bIsLocked = CS.TryLock();
+		}
+		else
+		{
+			CS.Lock();
+		}
+		
+		if (bIsLocked)
+		{
+			// This will happens in PushTask() or when ScheduleRenderTask() ends up executing the command.
+			const bool bWasAlreadyLocked = (TaskState == TS_Locked);
+			TaskState = TS_Locked;
+
+			ETaskState TickResult;
+			do // Iterate as longs as there is progress
+			{
+				// Only test for suspension the first time and in normal progress.
+				// When cancelled, we want the update to complete without interruptions, allowing reference to be freed.
+				TickResult = TickInternal(InCurrentThread, !bWasAlreadyLocked && !bIsCancelled);
+			} 
+			while (TickResult == TS_Locked);
+
+			// We do this to prevent updating the TaskState while in the Lock.
+			if (!bWasAlreadyLocked)
+			{
+				TaskState = TickResult;
+			}
+			CS.Unlock();
+		}
+	}
+}
+
+void FRenderAssetUpdate::ScheduleRenderTask()
+{
+	check(TaskState == TS_Locked);
+
+	// Notify that a tick is scheduled on the render thread.
+	++ScheduledRenderTasks;
+	// Increment refcount because we don't use a TRefCountPtr with ENQUEUE_RENDER_COMMAND.
+	AddRef();
+
+	ENQUEUE_RENDER_COMMAND(RenderAssetUpdateCommand)(
+		[&](FRHICommandListImmediate&)
+	{
+		// Recompute the context has things might have changed!
+		Tick(TT_Render);
+
+		--ScheduledRenderTasks;
+
+		// Decrement refcount because we don't use a TRefCountPtr with ENQUEUE_RENDER_COMMAND.
+		Release();
+	});
+}
+
+void FRenderAssetUpdate::ScheduleAsyncTask()
+{
+	check(TaskState == TS_Locked);
+
+	// Notify that an async tick is scheduled.
+	++ScheduledAsyncTasks;
+	(new FAsyncMipUpdateTask(this))->StartBackgroundTask();
+}
+
+void FRenderAssetUpdate::FMipUpdateTask::DoWork()
+{
+	check(PendingUpdate.IsValid());
+
+#if !UE_BUILD_SHIPPING
+	const int32 ExtraSyncLatency = CVarStreamingStressTestExtraAsyncLatency.GetValueOnAnyThread();
+	if (ExtraSyncLatency > 0)
+	{
+		// Slow down the async. Used to test GC issues.
+		FPlatformProcess::Sleep(ExtraSyncLatency * .001f);
+	}
+#endif
+
+	// Recompute the context has things might have changed!
+	PendingUpdate->Tick(FRenderAssetUpdate::TT_Async);
+	
+	--PendingUpdate->ScheduledAsyncTasks;
+}
+
+FRenderAssetUpdate::ETaskState FRenderAssetUpdate::DoLock() 
+{  
+	CS.Lock(); 
+	ETaskState PreviousTaskState = TaskState;
+	TaskState = TS_Locked; 
+	return PreviousTaskState;
+}
+
+void FRenderAssetUpdate::DoUnlock(ETaskState PreviousTaskState)
+{
+	TaskState = PreviousTaskState; 
+	CS.Unlock();
+}
diff --git a/Engine/Source/Runtime/Engine/Private/Streaming/RenderAssetUpdate.h b/Engine/Source/Runtime/Engine/Private/Streaming/RenderAssetUpdate.h
deleted file mode 100644
diff --git a/Engine/Source/Runtime/Engine/Private/Streaming/RenderAssetUpdate.inl b/Engine/Source/Runtime/Engine/Private/Streaming/RenderAssetUpdate.inl
new file mode 100644
--- /dev/null
+++ b/Engine/Source/Runtime/Engine/Private/Streaming/RenderAssetUpdate.inl
@@ -0,0 +1,96 @@
+// Copyright 1998-2019 Epic Games, Inc. All Rights Reserved.
+
+/*=============================================================================
+RenderAssetUpdate.h: Base class of helpers to stream in and out texture/mesh LODs
+=============================================================================*/
+
+#pragma once
+
+#include "RenderAssetUpdate.h"
+
+template <typename TContext>
+TRenderAssetUpdate<TContext>::TRenderAssetUpdate(UStreamableRenderAsset* InAsset, int32 InRequestedMips)
+	: FRenderAssetUpdate(InAsset, InRequestedMips)
+	, TaskThread(TT_None)
+	, TaskCallback(nullptr)
+	, CancelationThread(TT_None)
+	, CancelationCallback(nullptr)
+{
+
+}
+
+template <typename TContext>
+void TRenderAssetUpdate<TContext>::PushTask(const FContext& Context, EThreadType InTaskThread, const FCallback& InTaskCallback, EThreadType InCancelationThread, const FCallback& InCancelationCallback)
+{
+	// PushTask can only be called by the one thread/callback that is doing the processing. 
+	// This means we don't need to check whether other threads could be trying to push tasks.
+	check(TaskState == TS_Locked || TaskState == TS_Init);
+	checkSlow((bool)InTaskCallback == (InTaskThread != TT_None));
+	checkSlow((bool)InCancelationCallback == (InCancelationThread != TT_None));
+
+	TaskThread = InTaskThread;
+	TaskCallback = InTaskCallback;
+	CancelationThread = InCancelationThread;
+	CancelationCallback = InCancelationCallback;
+}
+
+template <typename TContext>
+FRenderAssetUpdate::ETaskState TRenderAssetUpdate<TContext>::TickInternal(EThreadType InCurrentThread, bool bCheckForSuspension)
+{
+	check(TaskState == TS_Locked);
+
+	// Thread, callback and asset must be coherent because Abort() could be called while this is executing.
+	EThreadType RelevantThread = TaskThread;
+	FCallback RelevantCallback = TaskCallback;
+	UStreamableRenderAsset* RelevantAsset = StreamableAsset;
+	if (bIsCancelled || !RelevantAsset)
+	{
+		RelevantThread = CancelationThread;
+		RelevantCallback = CancelationCallback;
+		// RelevantAsset = nullptr; // TODO once Cancel supports it!
+	}
+
+	if (RelevantThread == TT_None)
+	{
+		ClearCallbacks();
+		return TS_Done;
+	}
+	else if (TaskSynchronization.GetValue() > 0)
+	{
+		return TS_Suspended;
+	}
+	else if (bCheckForSuspension && IsAssetStreamingSuspended())
+	{
+		return TS_Suspended;
+	}
+	else if (RelevantThread == InCurrentThread)
+	{
+		ClearCallbacks();
+		RelevantCallback(FContext(RelevantAsset, InCurrentThread));
+		return TS_Locked;
+	} 
+	else if (RelevantThread == FRenderAssetUpdate::TT_Render && !ScheduledRenderTasks)
+	{
+		if (GIsThreadedRendering)
+		{
+			// Note that it could execute now if this is the renderthread.
+			ScheduleRenderTask();
+			return TS_Locked;
+		}
+		else
+		{
+			// If render commands do not have their own thread, avoid unrelated thread to enqueue commands.
+			// This is to avoid issues from using RHIAsyncCreateTexture2D which doesn't give any feedback of when it is ready.
+			return TS_InProgress;
+		}
+	}
+	else if (RelevantThread == FRenderAssetUpdate::TT_Async && !ScheduledAsyncTasks)
+	{
+		ScheduleAsyncTask();
+		return TS_Locked;
+	}
+	else
+	{
+		return TS_InProgress;
+	}
+}
diff --git a/Engine/Source/Runtime/Engine/Private/Streaming/StaticMeshUpdate.cpp b/Engine/Source/Runtime/Engine/Private/Streaming/StaticMeshUpdate.cpp
--- a/Engine/Source/Runtime/Engine/Private/Streaming/StaticMeshUpdate.cpp
+++ b/Engine/Source/Runtime/Engine/Private/Streaming/StaticMeshUpdate.cpp
@@ -10,6 +10,10 @@ StaticMeshUpdate.cpp: Helpers to stream in and out static mesh LODs.
 #include "Streaming/TextureStreamingHelpers.h"
 #include "HAL/PlatformFilemanager.h"
 #include "Serialization/MemoryReader.h"
+#include "Streaming/RenderAssetUpdate.inl"
+
+// Instantiate TRenderAssetUpdate for FStaticMeshUpdateContext
+template class TRenderAssetUpdate<FStaticMeshUpdateContext>;
 
 static constexpr uint32 GMaxNumResourceUpdatesPerLOD = 14;
 static constexpr uint32 GMaxNumResourceUpdatesPerBatch = (MAX_STATIC_MESH_LODS - 1) * GMaxNumResourceUpdatesPerLOD;
@@ -296,10 +300,9 @@ void FStaticMeshStreamIn_IO::FCancelIORequestsTask::DoWork()
 	check(PendingUpdate);
 	// Acquire the lock of this object in order to cancel any pending IO.
 	// If the object is currently being ticked, wait.
-	PendingUpdate->DoLock();
+	const ETaskState PreviousTaskState = PendingUpdate->DoLock();
 	PendingUpdate->CancelIORequest();
-	PendingUpdate->DoUnlock();
-	FPlatformAtomics::InterlockedDecrement(&PendingUpdate->ScheduledTaskCount);
+	PendingUpdate->DoUnlock(PreviousTaskState);
 }
 
 FStaticMeshStreamIn_IO::FStaticMeshStreamIn_IO(UStaticMesh* InMesh, int32 InRequestedMips, bool bHighPrio)
@@ -320,7 +323,6 @@ void FStaticMeshStreamIn_IO::Abort()
 		{
 			// Prevent the update from being considered done before this is finished.
 			// By checking that it was not already cancelled, we make sure this doesn't get called twice.
-			FPlatformAtomics::InterlockedIncrement(&ScheduledTaskCount);
 			(new FAsyncCancelIORequestsTask(this))->StartBackgroundTask();
 		}
 	}
@@ -361,7 +363,7 @@ void FStaticMeshStreamIn_IO::SetAsyncFileCallback(const FContext& Context)
 
 		// The tick here is intended to schedule the success or cancel callback.
 		// Using TT_None ensure gets which could create a dead lock.
-		Tick(Context.Mesh, FStaticMeshUpdate::TT_None);
+		Tick(FStaticMeshUpdate::TT_None);
 	};
 }
 
@@ -404,6 +406,7 @@ void FStaticMeshStreamIn_IO::SetIORequest(const FContext& Context, const FString
 		TaskSynchronization.Increment();
 		LODData.Empty(ReadSize);
 		LODData.AddUninitialized(ReadSize);
+
 		IORequest = IOFileHandle->ReadRequest(
 			ReadOffset,
 			ReadSize,
@@ -478,6 +481,7 @@ TStaticMeshStreamIn_IO<bRenderThread>::TStaticMeshStreamIn_IO(UStaticMesh* InMes
 template <bool bRenderThread>
 void TStaticMeshStreamIn_IO<bRenderThread>::DoInitiateIO(const FContext& Context)
 {
+	DECLARE_SCOPE_CYCLE_COUNTER(TEXT("TStaticMeshStreamIn_IO::DoInitiateIO"), STAT_StaticMeshStreamInIO_DoInitiateIO, STATGROUP_StreamingDetails);
 	check(Context.CurrentThread == TT_Async);
 	const FString IOFilename = GetIOFilename(Context);
 	SetIORequest(Context, IOFilename);
@@ -487,6 +491,7 @@ void TStaticMeshStreamIn_IO<bRenderThread>::DoInitiateIO(const FContext& Context
 template <bool bRenderThread>
 void TStaticMeshStreamIn_IO<bRenderThread>::DoSerializeLODData(const FContext& Context)
 {
+	DECLARE_SCOPE_CYCLE_COUNTER(TEXT("TStaticMeshStreamIn_IO::DoSerializeLODData"), STAT_StaticMeshStreamInIO_DoSerializeLODData, STATGROUP_StreamingDetails);
 	check(Context.CurrentThread == TT_Async);
 	SerializeLODData(Context);
 	ClearIORequest(Context);
@@ -498,6 +503,7 @@ void TStaticMeshStreamIn_IO<bRenderThread>::DoSerializeLODData(const FContext& C
 template <bool bRenderThread>
 void TStaticMeshStreamIn_IO<bRenderThread>::DoCreateBuffers(const FContext& Context)
 {
+	DECLARE_SCOPE_CYCLE_COUNTER(TEXT("TStaticMeshStreamIn_IO::DoCreateBuffers"), STAT_StaticMeshStreamInIO_DoCreateBuffers, STATGROUP_StreamingDetails);
 	if (bRenderThread)
 	{
 		CreateBuffers_RenderThread(Context);
@@ -513,6 +519,7 @@ void TStaticMeshStreamIn_IO<bRenderThread>::DoCreateBuffers(const FContext& Cont
 template <bool bRenderThread>
 void TStaticMeshStreamIn_IO<bRenderThread>::DoCancelIO(const FContext& Context)
 {
+	DECLARE_SCOPE_CYCLE_COUNTER(TEXT("TStaticMeshStreamIn_IO::DoCancelIO"), STAT_StaticMeshStreamInIO_DoCancelIO, STATGROUP_StreamingDetails);
 	ClearIORequest(Context);
 	PushTask(Context, TT_None, nullptr, (EThreadType)Context.CurrentThread, SRA_UPDATE_CALLBACK(DoCancel));
 }
@@ -542,6 +549,7 @@ TStaticMeshStreamIn_DDC<bRenderThread>::TStaticMeshStreamIn_DDC(UStaticMesh* InM
 template <bool bRenderThread>
 void TStaticMeshStreamIn_DDC<bRenderThread>::DoLoadNewLODsFromDDC(const FContext& Context)
 {
+	DECLARE_SCOPE_CYCLE_COUNTER(TEXT("TStaticMeshStreamInDDC::DoLoadNewLODsFromDDC"), STAT_StaticMeshStreamInDDC_DoLoadNewLODsFromDDC, STATGROUP_StreamingDetails);
 	LoadNewLODsFromDDC(Context);
 	check(!TaskSynchronization.GetValue());
 	const EThreadType TThread = bRenderThread ? TT_Render : TT_Async;
@@ -552,6 +560,7 @@ void TStaticMeshStreamIn_DDC<bRenderThread>::DoLoadNewLODsFromDDC(const FContext
 template <bool bRenderThread>
 void TStaticMeshStreamIn_DDC<bRenderThread>::DoCreateBuffers(const FContext& Context)
 {
+	DECLARE_SCOPE_CYCLE_COUNTER(TEXT("TStaticMeshStreamInDDC::DoCreateBuffers"), STAT_StaticMeshStreamInDDC_DoCreateBuffers, STATGROUP_StreamingDetails);
 	if (bRenderThread)
 	{
 		CreateBuffers_RenderThread(Context);
diff --git a/Engine/Source/Runtime/Engine/Private/Streaming/StaticMeshUpdate.h b/Engine/Source/Runtime/Engine/Private/Streaming/StaticMeshUpdate.h
--- a/Engine/Source/Runtime/Engine/Private/Streaming/StaticMeshUpdate.h
+++ b/Engine/Source/Runtime/Engine/Private/Streaming/StaticMeshUpdate.h
@@ -7,7 +7,6 @@ StaticMeshUpdate.h: Helpers to stream in and out static mesh LODs.
 #pragma once
 
 #include "CoreMinimal.h"
-#include "Streaming/RenderAssetUpdate.h"
 #include "Engine/StaticMesh.h"
 #include "Async/AsyncFileHandle.h"
 
@@ -42,6 +41,9 @@ struct FStaticMeshUpdateContext
 	EThreadType CurrentThread;
 };
 
+// Declare that TRenderAssetUpdate is instantiated for FStaticMeshUpdateContext
+extern template class TRenderAssetUpdate<FStaticMeshUpdateContext>;
+
 /**
 * This class provides a framework for loading and unloading the LODs of static meshes.
 * Each thread essentially calls Tick() until the job is done.
@@ -52,19 +54,15 @@ class FStaticMeshUpdate : public TRenderAssetUpdate<FStaticMeshUpdateContext>
 public:
 	FStaticMeshUpdate(UStaticMesh* InMesh, int32 InRequestedMips);
 
-	virtual ~FStaticMeshUpdate() {}
-
 	virtual void Abort()
 	{
 		TRenderAssetUpdate<FStaticMeshUpdateContext>::Abort();
 	}
 
-#if WITH_EDITOR
-	/** Returns whether DDC of this texture needs to be regenerated.  */
-	virtual bool DDCIsInvalid() const { return false; }
-#endif
-
 protected:
+
+	virtual ~FStaticMeshUpdate() {}
+
 	/** Cached index of current first LOD that will be replaced by PendingFirstMip */
 	int32 CurrentFirstLODIdx;
 };
@@ -161,7 +159,7 @@ protected:
 		}
 
 	private:
-		FStaticMeshStreamIn_IO* PendingUpdate;
+		TRefCountPtr<FStaticMeshStreamIn_IO> PendingUpdate;
 	};
 
 	typedef FAutoDeleteAsyncTask<FCancelIORequestsTask> FAsyncCancelIORequestsTask;
@@ -221,7 +219,7 @@ public:
 
 	virtual ~FStaticMeshStreamIn_DDC() {}
 
-	virtual bool DDCIsInvalid() const override { return bDerivedDataInvalid; }
+	bool DDCIsInvalid() const override { return bDerivedDataInvalid; }
 
 protected:
 	void LoadNewLODsFromDDC(const FContext& Context);
diff --git a/Engine/Source/Runtime/Engine/Private/Streaming/StreamingManagerTexture.cpp b/Engine/Source/Runtime/Engine/Private/Streaming/StreamingManagerTexture.cpp
--- a/Engine/Source/Runtime/Engine/Private/Streaming/StreamingManagerTexture.cpp
+++ b/Engine/Source/Runtime/Engine/Private/Streaming/StreamingManagerTexture.cpp
@@ -23,7 +23,6 @@
 #include "ProfilingDebugging/CsvProfiler.h"
 #include "Interfaces/ITargetPlatformManagerModule.h"
 #include "Interfaces/ITargetPlatform.h"
-#include "Streaming/RenderAssetUpdate.h"
 
 CSV_DECLARE_CATEGORY_MODULE_EXTERN(CORE_API, Basic);
 
diff --git a/Engine/Source/Runtime/Engine/Private/Streaming/Texture2DStreamIn_IO.cpp b/Engine/Source/Runtime/Engine/Private/Streaming/Texture2DStreamIn_IO.cpp
--- a/Engine/Source/Runtime/Engine/Private/Streaming/Texture2DStreamIn_IO.cpp
+++ b/Engine/Source/Runtime/Engine/Private/Streaming/Texture2DStreamIn_IO.cpp
@@ -94,7 +94,7 @@ void FTexture2DStreamIn_IO::SetIOFilename(const FContext& Context)
 
 void FTexture2DStreamIn_IO::SetIORequests(const FContext& Context)
 {
-	SetAsyncFileCallback(Context);
+	SetAsyncFileCallback();
 
 	check(!IOFileHandle);
 	IOFileHandle = FPlatformFileManager::Get().GetPlatformFile().OpenAsyncRead(*IOFilename);
@@ -169,9 +169,9 @@ void FTexture2DStreamIn_IO::ClearIORequests(const FContext& Context)
 	}
 }
 
-void FTexture2DStreamIn_IO::SetAsyncFileCallback(const FContext& Context)
+void FTexture2DStreamIn_IO::SetAsyncFileCallback()
 {
-	AsyncFileCallBack = [this, Context](bool bWasCancelled, IAsyncReadRequest* Req)
+	AsyncFileCallBack = [this](bool bWasCancelled, IAsyncReadRequest* Req)
 	{
 		// At this point task synchronization would hold the number of pending requests.
 		TaskSynchronization.Decrement();
@@ -191,7 +191,7 @@ void FTexture2DStreamIn_IO::SetAsyncFileCallback(const FContext& Context)
 
 		// The tick here is intended to schedule the success or cancel callback.
 		// Using TT_None ensure gets which could create a dead lock.
-		Tick(Context.Texture, FTexture2DUpdate::TT_None);
+		Tick(FTexture2DUpdate::TT_None);
 	};
 }
 
@@ -206,7 +206,6 @@ void FTexture2DStreamIn_IO::Abort()
 		{
 			// Prevent the update from being considered done before this is finished.
 			// By checking that it was not already cancelled, we make sure this doesn't get called twice.
-			FPlatformAtomics::InterlockedIncrement(&ScheduledTaskCount);
 			(new FAsyncCancelIORequestsTask(this))->StartBackgroundTask();
 		}
 	}
@@ -217,8 +216,7 @@ void FTexture2DStreamIn_IO::FCancelIORequestsTask::DoWork()
 	check(PendingUpdate);
 	// Acquire the lock of this object in order to cancel any pending IO.
 	// If the object is currently being ticked, wait.
-	PendingUpdate->DoLock();
+	const ETaskState PreviousTaskState = PendingUpdate->DoLock();
 	PendingUpdate->CancelIORequests();
-	PendingUpdate->DoUnlock();
-	FPlatformAtomics::InterlockedDecrement(&PendingUpdate->ScheduledTaskCount);
+	PendingUpdate->DoUnlock(PreviousTaskState);
 }
diff --git a/Engine/Source/Runtime/Engine/Private/Streaming/Texture2DStreamIn_IO.h b/Engine/Source/Runtime/Engine/Private/Streaming/Texture2DStreamIn_IO.h
--- a/Engine/Source/Runtime/Engine/Private/Streaming/Texture2DStreamIn_IO.h
+++ b/Engine/Source/Runtime/Engine/Private/Streaming/Texture2DStreamIn_IO.h
@@ -30,7 +30,7 @@ protected:
 	// Cancel / destroy each requests created in SetIORequests()
 	void ClearIORequests(const FContext& Context);
 	// Set the IO callback used for streaming the mips.
-	void SetAsyncFileCallback(const FContext& Context);
+	void SetAsyncFileCallback();
 
 	// Cancel all IO requests.
 	void CancelIORequests();
@@ -50,7 +50,7 @@ private:
 			RETURN_QUICK_DECLARE_CYCLE_STAT(FCancelIORequestsTask_Texture, STATGROUP_ThreadPoolAsyncTasks);
 		}
 	protected:
-		FTexture2DStreamIn_IO* PendingUpdate;
+		TRefCountPtr<FTexture2DStreamIn_IO> PendingUpdate;
 	};
 
 	typedef FAutoDeleteAsyncTask<FCancelIORequestsTask> FAsyncCancelIORequestsTask;
@@ -68,7 +68,5 @@ private:
 
 	IAsyncReadFileHandle* IOFileHandle;
 	FAsyncFileCallBack AsyncFileCallBack;
-
-
 };
 
diff --git a/Engine/Source/Runtime/Engine/Private/Streaming/Texture2DUpdate.cpp b/Engine/Source/Runtime/Engine/Private/Streaming/Texture2DUpdate.cpp
--- a/Engine/Source/Runtime/Engine/Private/Streaming/Texture2DUpdate.cpp
+++ b/Engine/Source/Runtime/Engine/Private/Streaming/Texture2DUpdate.cpp
@@ -7,6 +7,10 @@ Texture2DUpdate.cpp: Helpers to stream in and out mips.
 #include "Streaming/Texture2DUpdate.h"
 #include "RenderUtils.h"
 #include "Containers/ResourceArray.h"
+#include "Streaming/RenderAssetUpdate.inl"
+
+// Instantiate TRenderAssetUpdate for FTexture2DUpdateContext
+template class TRenderAssetUpdate<FTexture2DUpdateContext>;
 
 #if STATS
 extern volatile int64 GPending2DUpdateCount;
diff --git a/Engine/Source/Runtime/Engine/Private/Streaming/Texture2DUpdate.h b/Engine/Source/Runtime/Engine/Private/Streaming/Texture2DUpdate.h
--- a/Engine/Source/Runtime/Engine/Private/Streaming/Texture2DUpdate.h
+++ b/Engine/Source/Runtime/Engine/Private/Streaming/Texture2DUpdate.h
@@ -7,7 +7,7 @@ Texture2DUpdate.h: Helpers to stream in and out mips.
 #pragma once
 
 #include "CoreMinimal.h"
-#include "Streaming/RenderAssetUpdate.h"
+#include "RenderAssetUpdate.h"
 #include "Engine/Texture2D.h"
 #include "Async/AsyncFileHandle.h"
 
@@ -44,6 +44,9 @@ struct FTexture2DUpdateContext
 	EThreadType CurrentThread;
 };
 
+// Declare that TRenderAssetUpdate is instantiated for FTexture2DUpdateContext
+extern template class TRenderAssetUpdate<FTexture2DUpdateContext>;
+
 /**
  * This class provides a framework for loading and unloading the mips of 2D textures.
  * Each thread essentially calls Tick() until the job is done.
@@ -54,19 +57,10 @@ class FTexture2DUpdate : public TRenderAssetUpdate<FTexture2DUpdateContext>
 public:
 	FTexture2DUpdate(UTexture2D* InTexture, int32 InRequestedMips);
 
-	virtual ~FTexture2DUpdate();
-
-	virtual void Abort()
-	{
-		TRenderAssetUpdate<FTexture2DUpdateContext>::Abort();
-	}
+protected:
 
-#if WITH_EDITORONLY_DATA
-	/** Returns whether DDC of this texture needs to be regenerated.  */
-	virtual bool DDCIsInvalid() const { return false; }
-#endif
+	virtual ~FTexture2DUpdate();
 
-protected:
 	// ****************************
 	// ********* Helpers **********
 	// ****************************
diff --git a/Engine/Source/Runtime/Engine/Private/Texture2D.cpp b/Engine/Source/Runtime/Engine/Private/Texture2D.cpp
--- a/Engine/Source/Runtime/Engine/Private/Texture2D.cpp
+++ b/Engine/Source/Runtime/Engine/Private/Texture2D.cpp
@@ -788,8 +788,16 @@ bool UTexture2D::UpdateStreamingStatus( bool bWaitForMipFading /*= false*/ )
 			PendingUpdate->Abort();
 		}
 
-		// When the renderthread is the gamethread, allow the Tick to execute rendercommands.
-		PendingUpdate->Tick(this, GIsThreadedRendering ? FTexture2DUpdate::TT_None : FTexture2DUpdate::TT_Render);
+		// When there is no renderthread, allow the gamethread to tick as the renderthread.
+		FRenderAssetUpdate::EThreadType TickThread = GIsThreadedRendering ? FRenderAssetUpdate::TT_None : FRenderAssetUpdate::TT_Render;
+		if (HasAnyFlags(RF_BeginDestroyed) && PendingUpdate->GetRelevantThread() == FRenderAssetUpdate::TT_Async)
+		{
+			// To avoid async tasks from timing out the GC, we tick as Async to force completion if this is relevant.
+			// This could lead the asset from releasing the PendingUpdate, which will be deleted once the async task completes.
+			TickThread = FRenderAssetUpdate::TT_Async;
+		}
+		PendingUpdate->Tick(TickThread);
+
 		if (!PendingUpdate->IsCompleted())
 		{
 			return true;
@@ -799,8 +807,8 @@ bool UTexture2D::UpdateStreamingStatus( bool bWaitForMipFading /*= false*/ )
 		const bool bRebuildPlatformData = PendingUpdate->DDCIsInvalid() && !IsPendingKillOrUnreachable();
 #endif
 
-		delete PendingUpdate;
-		PendingUpdate = nullptr;
+		PendingUpdate.SafeRelease();
+
 
 #if WITH_EDITOR
 		if (GIsEditor)
@@ -1777,9 +1785,6 @@ bool UTexture2D::StreamIn(int32 NewMipCount, bool bHighPrio)
 				PendingUpdate = new FTexture2DStreamIn_IO_AsyncReallocate(this, NewMipCount, bHighPrio);
 			}
 		}
-
-		// The object starts in the locked state while it is being initialized.
-		PendingUpdate->DoUnlock();
 		return !PendingUpdate->IsCancelled();
 	}
 	return false;
@@ -1804,10 +1809,6 @@ bool UTexture2D::StreamOut(int32 NewMipCount)
 		{
 			PendingUpdate = new FTexture2DStreamOut_AsyncReallocate(this, NewMipCount);
 		}
-
-		// The object starts in the locked state while it is being initialized.
-		PendingUpdate->DoUnlock();
-
 		return !PendingUpdate->IsCancelled();
 	}
 	return false;
diff --git a/Engine/Source/Runtime/Engine/Public/RenderAssetUpdate.h b/Engine/Source/Runtime/Engine/Public/RenderAssetUpdate.h
new file mode 100644
--- /dev/null
+++ b/Engine/Source/Runtime/Engine/Public/RenderAssetUpdate.h
@@ -0,0 +1,245 @@
+// Copyright 1998-2019 Epic Games, Inc. All Rights Reserved.
+
+/*=============================================================================
+RenderAssetUpdate.h: Base class of helpers to stream in and out texture/mesh LODs
+=============================================================================*/
+
+#pragma once
+
+#include "CoreMinimal.h"
+#include "HAL/ThreadSafeCounter.h"
+#include "Async/AsyncWork.h"
+#include "Engine/StreamableRenderAsset.h"
+#include "RenderingThread.h"
+
+/** SRA stands for StreamableRenderAsset */
+#define SRA_UPDATE_CALLBACK(FunctionName) [this](const FContext& C){ FunctionName(C); }
+
+ // Allows yield to lower priority threads
+#define RENDER_ASSET_STREAMING_SLEEP_DT (0.010f)
+
+ENGINE_API bool IsAssetStreamingSuspended();
+
+/**
+* This class provides a framework for loading and unloading the texture/mesh LODs.
+* Each thread essentially calls Tick() until the job is done.
+* The object can be safely deleted when IsCompleted() returns true.
+*/
+class FRenderAssetUpdate : public IRefCountedObject
+{
+public:
+
+	/** A thread type used for doing a part of the update process.  */
+	enum EThreadType
+	{
+		TT_None,	// No thread.
+		TT_Render,	// The render thread.
+		TT_Async	// An async work thread.
+	};
+
+	/**  The state of scheduled work for the update process. */
+	enum ETaskState
+	{
+		TS_Done,		// Nothing to do.
+		TS_Suspended,	// Waiting for external events, like IO requests, or when the renderthread is suspended.
+		TS_InProgress,	// Callbacks are about to be executed.
+		TS_Locked,		// The object is locked, and no one is allowed to process or look at the next task.
+		TS_Init			// The object is in initialization
+	};
+
+	FRenderAssetUpdate(UStreamableRenderAsset* InAsset, int32 InRequestedMips);
+
+	/**
+	* Do or schedule any pending work for a given texture.
+	*
+	* @param InAsset - the texture/mesh being updated, this must be the same texture/mesh as the texture/mesh used to create this object.
+	* @param InCurrentThread - the thread from which the tick is being called. Using TT_None ensures that no work will be immediately performed.
+	*/
+	void Tick(EThreadType InCurrentThread);
+
+	/** Returns whether the task has finished executing and there is no other thread possibly accessing it. */
+	bool IsCompleted() const
+	{
+		return TaskState == TS_Done;
+	}
+
+	/**
+	* Cancel the current update. Will also attempt to cancel pending IO requests, see FTexture2DStreamIn_IO::Abort().
+	* This is called outside of the update loop when a cancelation is required for external events.
+	*/
+	virtual void Abort()
+	{
+		MarkAsCancelled();
+	}
+
+	/** Returns whether the task was aborted through Abort() or cancelled.  */
+	bool IsCancelled() const
+	{
+		return bIsCancelled;
+	}
+
+	/** Perform a lock on the object, preventing any other thread from processing a pending task in Tick(). */
+	ETaskState DoLock();
+
+	/** Release any lock on the object, allowing other thread to modify it. */
+	void DoUnlock(ETaskState PreviousTaskState);
+
+	bool IsLocked() const
+	{
+		return TaskState == TS_Locked;
+	}
+
+	/** Get the number of requested mips for this update, ignoring cancellation attempts. */
+	int32 GetNumRequestedMips() const
+	{
+		return RequestedMips;
+	}
+
+#if WITH_EDITORONLY_DATA
+	/** Returns whether DDC of this texture needs to be regenerated.  */
+	virtual bool DDCIsInvalid() const { return false; }
+#endif
+
+	/** Return the thread relevant to the next step of execution. */
+	virtual EThreadType GetRelevantThread() const = 0;
+
+	//****************** IRefCountedObject ****************/
+
+	uint32 AddRef() const final override 
+	{ 
+		return (uint32)NumRefs.Increment(); 
+	}
+
+	uint32 Release() const final override;
+
+	uint32 GetRefCount() const final override
+	{
+		return (uint32)NumRefs.GetValue();
+	}
+
+protected:
+
+	virtual ~FRenderAssetUpdate();
+
+	virtual ETaskState TickInternal(EThreadType InCurrentThread, bool bCheckForSuspension) = 0;
+
+	/** Set the task state as cancelled. This is internally called in Abort() and when any critical conditions are not met when performing the update. */
+	void MarkAsCancelled()
+	{
+		// StreamableAsset = nullptr; // TODO once Cancel supports it!
+		bIsCancelled = true;
+	}
+
+	void ScheduleRenderTask();
+	void ScheduleAsyncTask();
+
+
+	/** An async task used to call tick on the pending update. */
+	class FMipUpdateTask : public FNonAbandonableTask
+	{
+	public:
+		FMipUpdateTask(FRenderAssetUpdate* InPendingUpdate) : PendingUpdate(InPendingUpdate) {}
+
+		void DoWork();
+
+		FORCEINLINE TStatId GetStatId() const
+		{
+			RETURN_QUICK_DECLARE_CYCLE_STAT(FMipUpdateTask, STATGROUP_ThreadPoolAsyncTasks);
+		}
+
+	protected:
+		TRefCountPtr<FRenderAssetUpdate> PendingUpdate;
+	};
+
+	/** The async task to update this object, only one can be active at anytime. It just calls Tick(). */
+	typedef FAutoDeleteAsyncTask<FMipUpdateTask> FAsyncMipUpdateTask;
+
+	/** The index of mip that will end as being the first mip of the intermediate (future) texture/mesh. */
+	int32 PendingFirstMip;
+	/** The total number of mips of the intermediate (future) texture/mesh. */
+	int32 RequestedMips;
+
+	/** Critical Section. */
+	FCriticalSection CS;
+
+	/** Number of Ticks scheduled on the render thread. */
+	int32 ScheduledRenderTasks;
+
+	/** Number of Ticks scheduled on async tasks. */
+	int32 ScheduledAsyncTasks;
+
+	/** The asset updated **/
+	UStreamableRenderAsset* StreamableAsset;
+
+	/** Synchronization used for trigger the task next step execution. */
+	FThreadSafeCounter	TaskSynchronization;
+
+	/** Whether the task has been cancelled because the update could not proceed or because the user called Abort(). */
+	bool bIsCancelled;
+
+	/** The state of the work yet to be performed to complete the update or cancelation. */
+	volatile ETaskState TaskState;
+
+private:
+
+	/** Ref couting **/
+	mutable FThreadSafeCounter NumRefs;
+};
+
+/**
+* This class provides a framework for loading and unloading the texture/mesh LODs.
+* Each thread essentially calls Tick() until the job is done.
+* The object can be safely deleted when IsCompleted() returns true.
+*/
+template <typename TContext>
+class TRenderAssetUpdate : public FRenderAssetUpdate
+{
+public:
+
+	typedef TContext FContext;
+
+	/** A callback used to perform a task in the update process. Each task must be executed on a specific thread. */
+	typedef TFunction<void(const FContext& Context)> FCallback;
+
+	TRenderAssetUpdate(UStreamableRenderAsset* InAsset, int32 InRequestedMips);
+
+	/**
+	* Defines the next step to be executed. The next step will be executed by calling the callback on the specified thread.
+	* The callback (for both success and cancelation) will only be executed if TaskSynchronization reaches 0.
+	* If all requirements are immediately satisfied when calling the PushTask the relevant callback will be called immediately.
+	*
+	* @param Context - The context defining which texture is being updated and on which thread this is being called.
+	* @param InTaskThread - The thread on which to call the next step of the update, being TaskCallback.
+	* @param InTaskCallback - The callback that will perform the next step of the update.
+	* @param InCancelationThread - The thread on which to call the cancellation of the update (only if the update gets cancelled).
+	* @param InCancelationCallback - The callback handling the cancellation of the update (only if the update gets cancelled).
+	*/
+	void PushTask(const FContext& Context, EThreadType InTaskThread, const FCallback& InTaskCallback, EThreadType InCancelationThread, const FCallback& InCancelationCallback);
+
+	/** Return the thread relevant to the next step of execution. */
+	EThreadType GetRelevantThread() const final override { return bIsCancelled ? CancelationThread : TaskThread;  }
+
+protected:
+
+	/** The thread on which to call the next step of the update, being TaskCallback. */
+	EThreadType TaskThread;
+	/** The callback that will perform the next step of the update. */
+	FCallback TaskCallback;
+	/** The thread on which to call the cancellation of the update (only if the update gets cancelled). */
+	EThreadType CancelationThread; // The thread on which the callbacks should be called.
+	/** The callback handling the cancellation of the update (only if the update gets cancelled). */
+	FCallback CancelationCallback;
+
+	ETaskState TickInternal(EThreadType InCurrentThread, bool bCheckForSuspension) final override;
+
+	void ClearCallbacks()
+	{
+		TaskThread = TT_None;
+		TaskCallback = nullptr;
+		CancelationThread = TT_None;
+		CancelationCallback = nullptr;
+	}
+};
+
+void SuspendRenderAssetStreaming();
+void ResumeRenderAssetStreaming();
