Fixed missing heade

diff --git a/Engine/Source/Runtime/Engine/Classes/Sound/SoundSourceBusSend.h b/Engine/Source/Runtime/Engine/Classes/Sound/SoundSourceBusSend.h
--- a/Engine/Source/Runtime/Engine/Classes/Sound/SoundSourceBusSend.h
+++ b/Engine/Source/Runtime/Engine/Classes/Sound/SoundSourceBusSend.h
@@ -3,6 +3,7 @@
 #pragma once
 
 #include "CoreMinimal.h"
+#include "Curves/CurveFloat.h"
 #include "SoundSourceBusSend.generated.h"
 
 class USoundSourceBus;
