Allow games to detect when a replay scrub is in the middle of finalizing the fast forward process (calling onreps)

diff --git a/Engine/Source/Runtime/Engine/Classes/Engine/DemoNetDriver.h b/Engine/Source/Runtime/Engine/Classes/Engine/DemoNetDriver.h
--- a/Engine/Source/Runtime/Engine/Classes/Engine/DemoNetDriver.h
+++ b/Engine/Source/Runtime/Engine/Classes/Engine/DemoNetDriver.h
@@ -523,6 +523,7 @@ private:
 	bool bIsFastForwardingForCheckpoint;
 	bool bWasStartStreamingSuccessful;
 	bool bIsLoadingCheckpoint;
+	bool bIsFinalizingFastForward;
 
 	TArray<FNetworkGUID> NonQueuedGUIDsForScrubbing;
 
@@ -826,6 +827,7 @@ public:
 	void RequestEventDataForActiveReplay(const FString& EventID, const int32 UserIndex, const FRequestEventDataCallback& Delegate);
 
 	bool IsFastForwarding() const { return bIsFastForwarding; }
+	bool IsFinalizingFastForward() const { return bIsFinalizingFastForward; }
 
 	FReplayExternalDataArray* GetExternalDataArrayForObject(UObject* Object);
 
diff --git a/Engine/Source/Runtime/Engine/Private/DemoNetDriver.cpp b/Engine/Source/Runtime/Engine/Private/DemoNetDriver.cpp
--- a/Engine/Source/Runtime/Engine/Private/DemoNetDriver.cpp
+++ b/Engine/Source/Runtime/Engine/Private/DemoNetDriver.cpp
@@ -821,6 +821,7 @@ void UDemoNetDriver::ResetDemoState()
 	bIsLoadingCheckpoint = false;
 	bIsWaitingForHeaderDownload = false;
 	bIsWaitingForStream = false;
+	bIsFinalizingFastForward = false;
 
 	ExternalDataToObjectMap.Empty();
 	PlaybackPackets.Empty();
@@ -3929,6 +3930,8 @@ void UDemoNetDriver::FinalizeFastForward(const double StartTime)
 {
 	DECLARE_SCOPE_CYCLE_COUNTER(TEXT("Demo_FinalizeFastForward"), Demo_FinalizeFastForward, STATGROUP_Net);
 
+	TGuardValue<bool> FinalizingFastForward(bIsFinalizingFastForward, true);
+
 	// This must be set before we CallRepNotifies or they might be skipped again
 	bIsFastForwarding = false;
 
