Add more comments covering general NetDriver / NetConnection flows.

diff --git a/Engine/Plugins/Online/OnlineSubsystemUtils/Source/OnlineSubsystemUtils/Private/IpNetDriver.cpp b/Engine/Plugins/Online/OnlineSubsystemUtils/Source/OnlineSubsystemUtils/Private/IpNetDriver.cpp
--- a/Engine/Plugins/Online/OnlineSubsystemUtils/Source/OnlineSubsystemUtils/Private/IpNetDriver.cpp
+++ b/Engine/Plugins/Online/OnlineSubsystemUtils/Source/OnlineSubsystemUtils/Private/IpNetDriver.cpp
@@ -25,98 +25,6 @@ Notes:
 #include "IPAddress.h"
 #include "Sockets.h"
 
-/**
- * Initiating Connections / Handshaking Flow.
- *
- * Both Server and Clients will have have their own NetDrivers, and all UE Replicated Game traffic will be sent by or Received
- * from the IpNetDriver. This traffic also includes logic for establishing connections, and re-establishing connections when
- * something goes wrong.
- *
- * Handshaking is split across a couple of different places: NetDriver, PendingNetGame, World, PacketHandlers, and maybe others.
- * The split is due to having separate needs, things such as: determining whether or not an incoming connection is sending data in "UE-Protocol",
- * determining whether or not an address appears to be malicious, whether or not a given client has the correct version of a game, etc.
- *
- *****************************************************************************************
- * IpNetDriver / IpConnection Startup and Handshaking
- *****************************************************************************************
- *
- * Whenenever a server Loads a map (via UEngine::LoadMap), we will make a call into UWorld::Listen.
- * That code is responsible for creating the main Game Net Driver, parsing out settings, and calling UNetDriver::InitListen.
- * Ultimately, that code will be responsible for figuring out what IP / Port we bind to by calls to our configured Socket Subsystem
- * (see ISocketSubsystem::GetLocalBindAddr and ISocketSubsystem::BindNextPort).
- *
- * Once the server is listening, it's ready to start accepting client connections.
- *
- * Whenever a client wants to Join a server, they will first establish a new UPendingNetGame in UEngine::Browse with the server's IP.
- * UPendingNetGame::Initialize and UPendingNetGame::InitNetDriver are responsible for initializing settings and setting up the NetDriver respectively.
- * Clients will immediately setup a UNetConnection for the server as a part of this initialization, and will start sending data to the server on that connection,
- * initiating the handshaking process.
- *
- * On both Clients and Server, UIpNetDriver::TickDispatch is responsible for receiving packets.
- * Whenever we receive a packet, we inspect its address and see whether or not it's from a connection we already know about.
- * We determine whether or not we've established a connection for a given source address by simply keeping a map from FInternetAddr to UNetConnection.
- *
- * If a packet is from a connection that's already established, we pass the packet along to the connection via UNetConnection::ReceivedRawPacket.
- * If a packet is not from a connection that's already established, we treat is as "connectionless" and begin the handshaking process.
- *
- * See StatelessConnectionHandlerComponent.cpp for details on how this handshaking works.
- *
- *****************************************************************************************
- * UWorld / UPendingNetGame / AGameModeBase Startup and Handshaking
- *****************************************************************************************
- *
- * After the UIpNetDriver and UIpNetConnection have completed their handshaking process on Client and Server,
- * UPendingNetGame::SendInitialJoin will be called on the Client to kick off game level handshaking.
- *
- * Game Level Handshaking is done through a more structured and involved set of FNetControlMessages.
- * The full set of control messages can be found in DataChannel.h.
- *
- * Most of the work for handling these control messages are done either in UWorld::NotifyControlMessage,
- * and UPendingNetGame::NotifyControlMessage. Briefly, the flow looks like this:
- *
- * Client's UPendingNetGame::SendInitialJoin sends NMT_Hello.
- *
- * Server's UWorld::NotifyControlMessage receives NMT_Hello, sends NMT_Challenge.
- *
- * Client's UPendingNetGame::NotifyControlMessage receives NMT_Challenge, and sends back data in NMT_Login.
- *
- * Server's UWorld::NotifyControlMessage receives NMT_Login, verifies challenge data, and then calls AGameModeBase::PreLogin.
- * If PreLogin doesn't report any errors, Server calls UWorld::WelcomePlayer, which call AGameModeBase::GameWelcomePlayer,
- * and send NMT_Welcome with map information.
- *
- * Client's UPendingNetGame::NotifyControlMessage receives NMT_Welcome, reads the map info (so it can start loading later),
- * and sends an NMT_NetSpeed message with the configured Net Speed of the client.
- *
- * Server's UWorld::NotifyControlMessage receives NMT_NetSpeed, and adjusts the connections Net Speed appropriately.
- *
- * At this point, the handshaking is considered to be complete, and the player is fully connected to the game.
- * Depending on how long it takes to load the map, the client could still receive some non-handshake control messages
- * on UPendingNetGame before control transitions to UWorld.
- *
- * There are also additional steps for handling Encryption when desired.
- *
- *****************************************************************************************
- * Reestablishing Lost Connections
- *****************************************************************************************
- *
- * Throughout the course of a game, it's possible for connections to be lost for a number of reasons.
- * Internet could drop out, users could switch from LTE to WIFI, they could leave a game, etc.
- *
- * If the server initiated one of these disconnects, or is otherwise aware of it (due to a timeout or error),
- * then the disconnect will be handled by closing the UNetConnection and notifying the game.
- * At that point, it's up to a game to decide whether or not they support Join In Progress or Rejoins.
- * If the game does support it, we will completely restart the handshaking flow as above.
- *
- * If something just briefly interrupts the client's connection, but the server is never made aware,
- * then the engine / game will typically recover automatically (albeit with some packet loss / lag spike).
- *
- * However, if the Client's IP Address or Port change for any reason, but the server isn't aware of this,
- * then we will begin a recovery process by redoing the low level handshake. In this case, game code
- * will not be alerted.
- *
- * This process is covered in StatlessConnectionHandlerComponent.cpp.
- */
-
 /** For backwards compatibility with the engine stateless connect code */
 #ifndef STATELESSCONNECT_HAS_RANDOM_SEQUENCE
 	#define STATELESSCONNECT_HAS_RANDOM_SEQUENCE 0
diff --git a/Engine/Source/Runtime/Engine/Classes/Engine/NetDriver.h b/Engine/Source/Runtime/Engine/Classes/Engine/NetDriver.h
--- a/Engine/Source/Runtime/Engine/Classes/Engine/NetDriver.h
+++ b/Engine/Source/Runtime/Engine/Classes/Engine/NetDriver.h
@@ -19,6 +19,298 @@
 
 #include "NetDriver.generated.h"
 
+/**
+ *****************************************************************************************
+ * NetDrivers, NetConnections, and Channels
+ *****************************************************************************************
+ *
+ *
+ * UNetDrivers are responsible for managing sets of UNetConnections, and data that can be shared between them.
+ * There is typically a relatively small number of UNetDrivers for a given game. These may include:
+ *	- The Game NetDriver, responsible for standard game network traffic.
+ *	- The Demo NetDriver, responsible for recording or playing back previously recorded game data. This is how Replays work.
+ *	- The Beacon NetDriver, responsible for network traffic that falls outside of "normal" gameplay traffic.
+ *
+ * Custom NetDrivers can also be implemented by games or applications and used.
+ * NetConnections represent individual clients that are connected to a game (or more generally, to a NetDriver).
+ *
+ * End point data isn't directly handled by NetConnections. Instead NetConnections will route data to Channels.
+ * Every NetConnection will have its own set of channels.
+ *
+ * Common types of channels:
+ *
+ *	- A Control Channel is used to send information regarding state of a connection (whether or not the connection should close, etc.)
+ *	- A Voice Channel may be used to send voice data between client and server.
+ *	- A Unique Actor Channel will exist for every Actor replicated from the server to the client.
+ *
+ * Custom Channels can also be created and used for specialized purposes (although, this isn't very common).
+ *
+ *
+ *****************************************************************************************
+ * Game Net Drivers, Net Connections, and Channels
+ *****************************************************************************************
+ *
+ *
+ * Under normal circumstances, there will exist only a single NetDriver (created on Client and Server) for "standard" game traffic and connections.
+ *
+ * The Server NetDriver will maintain a list of NetConnections, each representing a player that is in the game. It is responsible for
+ * replicating Actor Data.
+ *
+ * Client NetDrivers will have a single NetConnection representing the connection to the Server.
+ *
+ * On both Server and Client, the NetDriver is responsible for receiving Packets from the network and passing those to the appropriate
+ * NetConnection (and establishing new NetConnections when necessary).
+ *
+ *
+ *****************************************************************************************
+ *****************************************************************************************
+ *****************************************************************************************
+ * Initiating Connections / Handshaking Flow.
+ *****************************************************************************************
+ *****************************************************************************************
+ *****************************************************************************************
+ *
+ *
+ * UIpNetDriver and UIpConnection (or derived classes) are the engine defaults for almost every platform, and everything below
+ * describes how they establish and manage connections. These processes can differ between implementations of NetDriver, however.
+ *
+ * Both Server and Clients will have have their own NetDrivers, and all UE Replicated Game traffic will be sent by or Received
+ * from the IpNetDriver. This traffic also includes logic for establishing connections, and re-establishing connections when
+ * something goes wrong.
+ *
+ * Handshaking is split across a couple of different places: NetDriver, PendingNetGame, World, PacketHandlers, and maybe others.
+ * The split is due to having separate needs, things such as: determining whether or not an incoming connection is sending data in "UE-Protocol",
+ * determining whether or not an address appears to be malicious, whether or not a given client has the correct version of a game, etc.
+ *
+ *
+ *****************************************************************************************
+ * Startup and Handshaking
+ *****************************************************************************************
+ *
+ *
+ * Whenenever a server Loads a map (via UEngine::LoadMap), we will make a call into UWorld::Listen.
+ * That code is responsible for creating the main Game Net Driver, parsing out settings, and calling UNetDriver::InitListen.
+ * Ultimately, that code will be responsible for figuring out what how exactly we listen for client connections.
+ * For example, in IpNetDriver, that is where we determine the IP / Port we will bind to by calls to our configured Socket Subsystem
+ * (see ISocketSubsystem::GetLocalBindAddr and ISocketSubsystem::BindNextPort).
+ *
+ * Once the server is listening, it's ready to start accepting client connections.
+ *
+ * Whenever a client wants to Join a server, they will first establish a new UPendingNetGame in UEngine::Browse with the server's IP.
+ * UPendingNetGame::Initialize and UPendingNetGame::InitNetDriver are responsible for initializing settings and setting up the NetDriver respectively.
+ * Clients will immediately setup a UNetConnection for the server as a part of this initialization, and will start sending data to the server on that connection,
+ * initiating the handshaking process.
+ *
+ * On both Clients and Server, UNetDriver::TickDispatch is typically responsible for receiving network data.
+ * Typically, when we receive a packet, we inspect its address and see whether or not it's from a connection we already know about.
+ * We determine whether or not we've established a connection for a given source address by simply keeping a map from FInternetAddr to UNetConnection.
+ *
+ * If a packet is from a connection that's already established, we pass the packet along to the connection via UNetConnection::ReceivedRawPacket.
+ * If a packet is not from a connection that's already established, we treat is as "connectionless" and begin the handshaking process.
+ *
+ * See StatelessConnectionHandlerComponent.cpp for details on how this handshaking works.
+ *
+ *
+ *****************************************************************************************
+ * UWorld / UPendingNetGame / AGameModeBase Startup and Handshaking
+ *****************************************************************************************
+ *
+ *
+ * After the UNetDriver and UNetConnection have completed their handshaking process on Client and Server,
+ * UPendingNetGame::SendInitialJoin will be called on the Client to kick off game level handshaking.
+ *
+ * Game Level Handshaking is done through a more structured and involved set of FNetControlMessages.
+ * The full set of control messages can be found in DataChannel.h.
+ *
+ * Most of the work for handling these control messages are done either in UWorld::NotifyControlMessage,
+ * and UPendingNetGame::NotifyControlMessage. Briefly, the flow looks like this:
+ *
+ * Client's UPendingNetGame::SendInitialJoin sends NMT_Hello.
+ *
+ * Server's UWorld::NotifyControlMessage receives NMT_Hello, sends NMT_Challenge.
+ *
+ * Client's UPendingNetGame::NotifyControlMessage receives NMT_Challenge, and sends back data in NMT_Login.
+ *
+ * Server's UWorld::NotifyControlMessage receives NMT_Login, verifies challenge data, and then calls AGameModeBase::PreLogin.
+ * If PreLogin doesn't report any errors, Server calls UWorld::WelcomePlayer, which call AGameModeBase::GameWelcomePlayer,
+ * and send NMT_Welcome with map information.
+ *
+ * Client's UPendingNetGame::NotifyControlMessage receives NMT_Welcome, reads the map info (so it can start loading later),
+ * and sends an NMT_NetSpeed message with the configured Net Speed of the client.
+ *
+ * Server's UWorld::NotifyControlMessage receives NMT_NetSpeed, and adjusts the connections Net Speed appropriately.
+ *
+ * At this point, the handshaking is considered to be complete, and the player is fully connected to the game.
+ * Depending on how long it takes to load the map, the client could still receive some non-handshake control messages
+ * on UPendingNetGame before control transitions to UWorld.
+ *
+ * There are also additional steps for handling Encryption when desired.
+ *
+ *
+ *****************************************************************************************
+ * Reestablishing Lost Connections
+ *****************************************************************************************
+ *
+ *
+ * Throughout the course of a game, it's possible for connections to be lost for a number of reasons.
+ * Internet could drop out, users could switch from LTE to WIFI, they could leave a game, etc.
+ *
+ * If the server initiated one of these disconnects, or is otherwise aware of it (due to a timeout or error),
+ * then the disconnect will be handled by closing the UNetConnection and notifying the game.
+ * At that point, it's up to a game to decide whether or not they support Join In Progress or Rejoins.
+ * If the game does support it, we will completely restart the handshaking flow as above.
+ *
+ * If something just briefly interrupts the client's connection, but the server is never made aware,
+ * then the engine / game will typically recover automatically (albeit with some packet loss / lag spike).
+ *
+ * However, if the Client's IP Address or Port change for any reason, but the server isn't aware of this,
+ * then we will begin a recovery process by redoing the low level handshake. In this case, game code
+ * will not be alerted.
+ *
+ * This process is covered in StatlessConnectionHandlerComponent.cpp.
+ *
+ *
+ *****************************************************************************************
+ *****************************************************************************************
+ *****************************************************************************************
+ * Data Transmission
+ *****************************************************************************************
+ *****************************************************************************************
+ *****************************************************************************************
+ *
+ *
+ * Game NetConnections and NetDrivers are generally agnostic to the underlying communication method / technology used.
+ * That is is left up to subclasses to decide (classes such as UIpConnection / UIpNetDriver or UWebSocketConnection / UWebSocketNetDriver).
+ *
+ * Instead, UNetDriver and UNetConnection work with Packets and Bunches.
+ *
+ * Packets are blobs of data that are sent between pairs of NetConnections on Host and Client.
+ * Packets consist of meta data about the packet (such as header information and acknowledgments), and Bunches.
+ *
+ * Bunches are blobs of data that are sent between pairs of Channels on Host and Client.
+ * When a Connection receives a Packet, that packet will be disassembled into individual bunches.
+ * Those bunches are then passed along to individual Channels to be processed further.
+ *
+ * A Packet may contain no bunches, a single bunch, or multiple bunches.
+ * Because size limits for bunches may be larger than the size limits of a single packet, UE4 supports
+ * the notion of partial bunches.
+ *
+ * When a bunch is too large, before transmission we will slice it into a number of smaller bunches.
+ * these bunches will be flagged as PartialInitial, Partial, or PartialFinal. Using this information,
+ * we can reassemble the bunches on the receiving end.
+ *
+ *	Example: Client RPC to Server.
+ *		- Client makes a call to Server_RPC.
+ *		- That request is forwarded (via NetDriver and NetConnection) to the Actor Channel that owns the Actor on which the RPC was called.
+ *		- The Actor Channel will serialize the RPC Identifier and parameters into a Bunch. The Bunch will also contain the ID of its Actor Channel.
+ *		- The Actor Channel will then request the NetConnection send the Bunch.
+ *		- Later, the NetConnection will assemble this (and other) data into a Packet which it will send to the server.
+ *		- On the Server, the Packet will be received by the NetDriver.
+ *		- The NetDriver will inspect the Address that sent the Packet, and hand the Packet over to the appropriate NetConnection.
+ *		- The NetConnection will disassemble the Packet into its Bunches (one by one).
+ *		- The NetConnection will use the Channel ID on the bunch to Route the bunch to the corresponding Actor Channel.
+ *		- The ActorChannel will them disassemble the bunch, see it contains RPC data, and use the RPC ID and serialized parameters
+ *			to call the appropriate function on the Actor.
+ *
+ *
+ *****************************************************************************************
+ * Reliability and Retransmission
+ *****************************************************************************************
+ *
+ *
+ * UE4 Networking typically assumes reliability isn't guaranteed by the underlying network protocol.
+ * Instead, it implements its own reliability and retransmission of both packets and bunches.
+ *
+ * When a NetConnection is established, it will establish a Sequence Number for its packets and bunches.
+ * These can either be fixed, or randomized (when randomized, the sequence will be sent by the server).
+ *
+ * The packet number is per NetConnection, incremented for every packet sent, every packet will include its packet number,
+ * and we will never retransmit a packet with the same packet number.
+ *
+ * The bunch number is per Channel, incremented for every **reliable** bunch sent, and every **reliable** bunch will
+ * include its bunch number. Unlike packets, though, exact (reliable) bunches may be retransmitted. This means we
+ * will resend bunches with the same bunch number.
+ *
+ * Note, throughout the code what are described above as both bunch numbers and packet numbers are commonly referred to
+ * just as sequence numbers. We make the distinction here for clearer understanding.
+ *
+ *	--- Detecting Incoming Dropped Packets ---
+ *
+ *
+ *	By assigning packet numbers, we can easily detect when incoming packets are lost.
+ *	This is done simply by taking the difference between the last successfully received packet number, and the
+ *	packet number of the current packet being processed.
+ *
+ *	Under good conditions, all packets will be received in the order they are sent. This means that the difference will
+ *	be +1.
+ *
+ *	If the difference is greater than 1, that indicates that we missed some packets. We will just
+ *	assume that the missing packets were dropped, but consider the current packet to have been successfully received,
+ *	and use its number going forward.
+ *
+ *	If the difference is negative (or 0), that indicates that we either received some packets out of order, or an external
+ *	service is trying to resend data to us (remember, the engine will not reuse sequence numbers).
+ *
+ *	In either case, the engine will typically ignore the missing or invalid packets, and will not send ACKs for them.
+ *
+ *	We do have methods for "fixing" out of order packets that are received on the same frame.
+ *	When enabled, if we detect missing packets (difference > 1), we won't process the current packet immediately.
+ *	Instead, it will add it to a queue. The next time we receive a packet successfully (difference == 1), we will
+ *	see if the head of our queue is properly ordered. If so, we will process it, otherwise we will continue
+ *	receiving packets.
+ *
+ *	Once we've read all packets that are currently available, we will flush this queue processing any remaining packets.
+ *	Anything that's missing at this point will be assumed to have been dropped.
+ *
+ *	Every packet successfully received will have its packet number sent back to the sender as an acknowledgment (ACK).
+ *
+ *
+ *	--- Detecting Outgoing Dropped Packets ---
+ *
+ *
+ *	As mentioned above, whenever a packet is received successfully the recipient will send back an ACK.
+ *	These ACKs will contain the packet numbers of successfully received packets, in sequence order.
+ *
+ *	Similar to how the recipient tracks the packet number, the sender will track the highest ACKed packet number.
+ *
+ *	When ACKs are being processed, any ACK below our last received ACK is ignored and any gaps in packet
+ *	numbers are considered Not Acknowledged (NAKed).
+ *
+ *	It is the sender's responsibility to handle these ACKs and NAKs and resend any missing data.
+ *	The new data will be added to new outgoing packets (again, we will not resend packets we've already sent,
+ *	or reuse packet sequence numbers).
+ *
+ *
+ *	--- Resending Missing Data ---
+ *
+ *
+ *	As mentioned above, packets alone don't contain useful game data. Instead, it's the bunches that comprise them
+ *	that have meaningful data.
+ *
+ *	Bunches can either be marked as Reliable or Unreliable.
+ *
+ *	The engine will make no attempt at resending unreliable bunches if they are dropped. Therefore, if bunches
+ *	are marked unreliable, the game / engine should be able to continue without them, or external retry
+ *	mechanisms must be put in place, or the data must be sent redundantly. Therefore, everything below only
+ *	applies to reliable bunches.
+ *
+ *	However, the engine will attempt to resend reliable bunches. Whenever a reliable bunch is sent, it will
+ *	be added to a list of un-ACKed reliable bunches. If we receive a NAK for a packet that contained the bunch,
+ *	the engine will retransmit an exact copy of that bunch. Note, because bunches may be partial, dropping even
+ *	a single partial bunch will result in retransmission of the entire bunch. When all packets containing a bunch
+ *	have been ACKed, we will remove it from the list.
+ *
+ *	Similar to packets, we will compare the bunch number for received reliable bunches to the last successfully
+ *	received bunch. If we detect that the difference is negative, we simply ignore the bunch. If the difference
+ *	is greater than one, we will assume we missed a bunch. Unlike packet handling, we will not discard this data.
+ *	Instead, we will queue the bunch and pause processing of **any** bunches, reliable or unreliable.
+ *	Processing will not be resumed until we detect have received the missing bunches, at which point we will process
+ *	them, and then start processing our queued bunches.
+ *	Any new bunches that are received while waiting for the missing bunches, or while we still have any bunches in our
+ *	queue, will be added to the queue instead of being processed immediately.
+ *
+ */
+
 class Error;
 class FNetGUIDCache;
 class FNetworkNotify;
