Remove the checkSlow from FRepChangedPropertyTracker::SetCustomIsActiveOverride.
IsConditional does not imply COND_Custom, and no code bothers to check those flags before checking the active override value.

Trivial

diff --git a/Engine/Source/Runtime/Engine/Public/Net/RepLayout.h b/Engine/Source/Runtime/Engine/Public/Net/RepLayout.h
--- a/Engine/Source/Runtime/Engine/Public/Net/RepLayout.h
+++ b/Engine/Source/Runtime/Engine/Public/Net/RepLayout.h
@@ -178,8 +178,6 @@ public:
 	{
 		FRepChangedParent & Parent = Parents[RepIndex];
 
-		checkSlow(Parent.IsConditional);
-
 		Parent.Active = (bIsActive || bIsClientReplayRecording) ? 1 : 0;
 		Parent.OldActive = Parent.Active;
 	}
