Fix bad conversion to linear from dB

diff --git a/Engine/Source/Runtime/Engine/Private/AudioDevice.cpp b/Engine/Source/Runtime/Engine/Private/AudioDevice.cpp
--- a/Engine/Source/Runtime/Engine/Private/AudioDevice.cpp
+++ b/Engine/Source/Runtime/Engine/Private/AudioDevice.cpp
@@ -3543,12 +3543,11 @@ void FAudioDevice::StopSources(TArray<FWaveInstance*>& WaveInstances, int32 Firs
 
 #if STATS
 	uint32 AudibleInactiveSounds = 0;
-	static const float MinReportVol = Audio::ConvertToDecibels(-60.f);
 	// Count how many sounds are not being played but were audible
 	for (int32 InstanceIndex = 0; InstanceIndex < FirstActiveIndex; InstanceIndex++)
 	{
 		FWaveInstance* WaveInstance = WaveInstances[InstanceIndex];
-		if (WaveInstance->GetVolumeWithDistanceAttenuation() > MinReportVol)
+		if (WaveInstance->GetVolumeWithDistanceAttenuation() > KINDA_SMALL_NUMBER)
 		{
 			AudibleInactiveSounds++;
 		}
