Make sure subobjects are added to the replication map of the channel when they are initially replicated by an RPC

diff --git a/Engine/Source/Runtime/Engine/Classes/Engine/ActorChannel.h b/Engine/Source/Runtime/Engine/Classes/Engine/ActorChannel.h
--- a/Engine/Source/Runtime/Engine/Classes/Engine/ActorChannel.h
+++ b/Engine/Source/Runtime/Engine/Classes/Engine/ActorChannel.h
@@ -183,6 +183,9 @@ public:
 	/** Queue a function bunch for this channel to be sent on the next property update. */
 	void QueueRemoteFunctionBunch( UObject* CallTarget, UFunction* Func, FOutBunch &Bunch );
 
+	/** If not queueing the RPC, prepare the channel for replicating the call.  */
+	void PrepareForRemoteFunction(UObject* TargetObj);
+	
 	/** Returns true if channel is ready to go dormant (e.g., all outstanding property updates have been ACK'd) */
 	virtual bool ReadyForDormancy(bool debug=false) override;
 	
diff --git a/Engine/Source/Runtime/Engine/Private/DataChannel.cpp b/Engine/Source/Runtime/Engine/Private/DataChannel.cpp
--- a/Engine/Source/Runtime/Engine/Private/DataChannel.cpp
+++ b/Engine/Source/Runtime/Engine/Private/DataChannel.cpp
@@ -3079,6 +3079,13 @@ void UActorChannel::Serialize(FArchive& Ar)
 	}
 }
 
+void UActorChannel::PrepareForRemoteFunction(UObject* TargetObj)
+{
+	// Make sure we create a replicator in case we destroy a sub object before we ever try to replicate its properties,
+	// otherwise it will not be in the ReplicationMap and we'll never send the deletion to clients
+	FindOrCreateReplicator(TargetObj);
+}
+
 void UActorChannel::QueueRemoteFunctionBunch( UObject* CallTarget, UFunction* Func, FOutBunch &Bunch )
 {
 	FindOrCreateReplicator(CallTarget).Get().QueueRemoteFunctionBunch( Func, Bunch );
diff --git a/Engine/Source/Runtime/Engine/Private/NetworkDriver.cpp b/Engine/Source/Runtime/Engine/Private/NetworkDriver.cpp
--- a/Engine/Source/Runtime/Engine/Private/NetworkDriver.cpp
+++ b/Engine/Source/Runtime/Engine/Private/NetworkDriver.cpp
@@ -2025,6 +2025,8 @@ void UNetDriver::ProcessRemoteFunctionForChannel(UActorChannel* Ch, const FClass
 		}
 		else
 		{
+			Ch->PrepareForRemoteFunction(TargetObj);
+
 			FNetBitWriter TempBlockWriter( Bunch.PackageMap, 0 );
 			Ch->WriteFieldHeaderAndPayload( TempBlockWriter, ClassCache, FieldCache, NetFieldExportGroup, TempWriter );
 			ParameterBits = TempBlockWriter.GetNumBits();
