Fix for looping discontinuity w/ audio mixer and ADPCM.

diff --git a/Engine/Source/Runtime/Engine/Private/ADPCMAudioInfo.cpp b/Engine/Source/Runtime/Engine/Private/ADPCMAudioInfo.cpp
--- a/Engine/Source/Runtime/Engine/Private/ADPCMAudioInfo.cpp
+++ b/Engine/Source/Runtime/Engine/Private/ADPCMAudioInfo.cpp
@@ -133,14 +133,13 @@ bool FADPCMAudioInfo::ReadCompressedInfo(const uint8* InSrcBufferData, uint32 In
 		BlockSize = *WaveInfo.pBlockAlign;
 
 		// ADPCM starts with 2 uncompressed samples and then the remaining compressed sample data has 2 samples per byte
-		UncompressedBlockSize = (2 + (BlockSize - PreambleSize) * 2) * sizeof(int16);
+		const uint32 UncompressedBlockSamples = (2 + (BlockSize - PreambleSize) * 2);
+		UncompressedBlockSize = UncompressedBlockSamples * sizeof(int16);
 		CompressedBlockSize = BlockSize;
 
-		const uint32 uncompressedBlockSamples = (2 + (BlockSize - PreambleSize) * 2);
-		const uint32 targetBlocks = MONO_PCM_BUFFER_SAMPLES / uncompressedBlockSamples;
-		StreamBufferSize = targetBlocks * UncompressedBlockSize;
-		// Ensure TotalDecodedSize is a even multiple of the compressed block size so that the buffer is not over read on the last block
-		TotalDecodedSize = ((WaveInfo.SampleDataSize + CompressedBlockSize - 1) / CompressedBlockSize) * UncompressedBlockSize;
+		const uint32 TargetBlocks = MONO_PCM_BUFFER_SAMPLES / UncompressedBlockSamples;
+		StreamBufferSize = TargetBlocks * UncompressedBlockSize;
+		TotalDecodedSize = TotalSamplesPerChannel * NumChannels * sizeof(int16);
 
 		UncompressedBlockData = (uint8*)FMemory::Realloc(UncompressedBlockData, NumChannels * UncompressedBlockSize);
 		check(UncompressedBlockData != nullptr);
@@ -249,12 +248,6 @@ bool FADPCMAudioInfo::ReadCompressedData(uint8* Destination, bool bLooping, uint
 				CurrentUncompressedBlockSampleIndex = UncompressedBlockSize / sizeof(uint16);
 				CurrentCompressedBlockIndex = 0;
 				TotalSamplesStreamed = 0;
-				if(!bLooping)
-				{
-					// Set the remaining buffer to 0
-					FMemory::Memset(OutData, 0, BufferSize);
-					return true;
-				}
 			}
 		}
 	}
