Adding path info to the nil matrix warning log

diff --git a/Engine/Source/Runtime/Engine/Private/Components/SkinnedMeshComponent.cpp b/Engine/Source/Runtime/Engine/Private/Components/SkinnedMeshComponent.cpp
--- a/Engine/Source/Runtime/Engine/Private/Components/SkinnedMeshComponent.cpp
+++ b/Engine/Source/Runtime/Engine/Private/Components/SkinnedMeshComponent.cpp
@@ -2505,7 +2505,7 @@ void USkinnedMeshComponent::SetRefPoseOverride(const TArray<FTransform>& NewRefP
 			ZAxis.IsNearlyZero(SMALL_NUMBER))
 		{
 			// this is not allowed, warn them 
-			UE_LOG(LogSkeletalMesh, Warning, TEXT("Reference Pose for joint (%s) includes NIL matrix. Zero scale isn't allowed on ref pose. "), *SkeletalMesh->RefSkeleton.GetBoneName(BoneIndex).ToString());
+			UE_LOG(LogSkeletalMesh, Warning, TEXT("Reference Pose for asset %s for joint (%s) includes NIL matrix. Zero scale isn't allowed on ref pose. "), *SkeletalMesh->GetPathName(), *SkeletalMesh->RefSkeleton.GetBoneName(BoneIndex).ToString());
 		}
 
 		// Precompute inverse so we can use from-refpose-skin vertices.
diff --git a/Engine/Source/Runtime/Engine/Private/SkeletalMesh.cpp b/Engine/Source/Runtime/Engine/Private/SkeletalMesh.cpp
--- a/Engine/Source/Runtime/Engine/Private/SkeletalMesh.cpp
+++ b/Engine/Source/Runtime/Engine/Private/SkeletalMesh.cpp
@@ -1332,7 +1332,7 @@ void USkeletalMesh::CalculateInvRefMatrices()
 				ZAxis.IsNearlyZero(SMALL_NUMBER))
 			{
 				// this is not allowed, warn them 
-				UE_LOG(LogSkeletalMesh, Warning, TEXT("Reference Pose for joint (%s) includes NIL matrix. Zero scale isn't allowed on ref pose. "), *RefSkeleton.GetBoneName(b).ToString());
+				UE_LOG(LogSkeletalMesh, Warning, TEXT("Reference Pose for asset %s for joint (%s) includes NIL matrix. Zero scale isn't allowed on ref pose. "), *GetPathName(), *RefSkeleton.GetBoneName(b).ToString());
 			}
 
 			// Precompute inverse so we can use from-refpose-skin vertices.
