- Added virtual memory display to "stat unitmax" with peak physical and virtual as well (peteS will check in a change that makes Peak work correctly on ios)

diff --git a/Engine/Source/Runtime/Engine/Private/StatsRender2.cpp b/Engine/Source/Runtime/Engine/Private/StatsRender2.cpp
--- a/Engine/Source/Runtime/Engine/Private/StatsRender2.cpp
+++ b/Engine/Source/Runtime/Engine/Private/StatsRender2.cpp
@@ -529,29 +529,6 @@ static int32 RenderMemoryHeadings( class FCanvas* Canvas, const int32 X, const
 	return Globals.GetFontHeight() + (Globals.GetFontHeight() / 3);
 }
 
-// @param bAutoType true: automatically choose GB/MB/KB/... false: always use MB for easier comparisons
-static FString GetMemoryString( const double Value, const bool bAutoType = true )
-{
-	if (bAutoType)
-	{
-		if (Value > 1024.0 * 1024.0 * 1024.0)
-		{
-			return FString::Printf( TEXT( "%.2f GB" ), float( Value / (1024.0 * 1024.0 * 1024.0) ) );
-		}
-		if (Value > 1024.0 * 1024.0)
-		{
-			return FString::Printf( TEXT( "%.2f MB" ), float( Value / (1024.0 * 1024.0) ) );
-		}
-		if (Value > 1024.0)
-		{
-			return FString::Printf( TEXT( "%.2f KB" ), float( Value / (1024.0) ) );
-		}
-		return FString::Printf( TEXT( "%.2f B" ), float( Value ) );
-	}
-
-	return FString::Printf( TEXT( "%.2f MB" ), float( Value / (1024.0 * 1024.0) ) );
-}
-
 static int32 RenderMemoryCounter( const FGameThreadStatsData& ViewData, const FComplexStatMessage& All, class FCanvas* Canvas, const int32 X, const int32 Y, const float Budget, const bool bIsBudgetIgnored )
 {
 	FPlatformMemory::EMemoryCounterRegion Region = FPlatformMemory::EMemoryCounterRegion(All.NameAndInfo.GetField<EMemoryRegion>());
diff --git a/Engine/Source/Runtime/Engine/Private/UnrealClient.cpp b/Engine/Source/Runtime/Engine/Private/UnrealClient.cpp
--- a/Engine/Source/Runtime/Engine/Private/UnrealClient.cpp
+++ b/Engine/Source/Runtime/Engine/Private/UnrealClient.cpp
@@ -365,6 +365,29 @@ TArray<FColor>* FScreenshotRequest::GetHighresScreenshotMaskColorArray()
 }
 
 
+// @param bAutoType true: automatically choose GB/MB/KB/... false: always use MB for easier comparisons
+FString GetMemoryString( const double Value, const bool bAutoType )
+{
+	if (bAutoType)
+	{
+		if (Value > 1024.0 * 1024.0 * 1024.0)
+		{
+			return FString::Printf( TEXT( "%.2f GB" ), float( Value / (1024.0 * 1024.0 * 1024.0) ) );
+		}
+		if (Value > 1024.0 * 1024.0)
+		{
+			return FString::Printf( TEXT( "%.2f MB" ), float( Value / (1024.0 * 1024.0) ) );
+		}
+		if (Value > 1024.0)
+		{
+			return FString::Printf( TEXT( "%.2f KB" ), float( Value / (1024.0) ) );
+		}
+		return FString::Printf( TEXT( "%.2f B" ), float( Value ) );
+	}
+	
+	return FString::Printf( TEXT( "%.2f MB" ), float( Value / (1024.0 * 1024.0) ) );
+}
+
 FOnScreenshotRequestProcessed FScreenshotRequest::ScreenshotProcessedDelegate;
 bool FScreenshotRequest::bIsScreenshotRequested = false;
 FString FScreenshotRequest::Filename;
@@ -615,16 +638,31 @@ int32 FStatUnitData::DrawStat(FViewport* InViewport, FCanvas* InCanvas, int32 In
 			InY += RowHeight;
 		}
 		{
-			uint64 MemoryUsed = FPlatformMemory::GetMemoryUsedFast();
-			if (MemoryUsed > 0)
+			if (bShowUnitMaxTimes)
 			{
-				// print out currently used memory
+				FPlatformMemoryStats Stats = FPlatformMemory::GetStats();
+
 				InCanvas->DrawShadowedString(X1, InY, TEXT("Mem:"), Font, bShowUnitTimeGraph ? FColor(100, 100, 255) : FColor::White);
-				double MemInGb = MemoryUsed / (1024.0 * 1024.0 * 1024.0);
-				double MemInMb = MemoryUsed / (1024.0 * 1024.0);
-				InCanvas->DrawShadowedString(X2, InY, *FString::Printf(TEXT("%3.2f%s"), MemInGb > 1.0 ? MemInGb : MemInMb, MemInGb > 1.0 ? TEXT("GB") : TEXT("MB")), Font, FColor::Green);
+				InCanvas->DrawShadowedString(X2, InY, *GetMemoryString(Stats.UsedPhysical), Font, FColor::Green);
+				InCanvas->DrawShadowedString(X3, InY, *GetMemoryString(Stats.PeakUsedPhysical), Font, FColor::Green);
+				InY += RowHeight;
+				
+				InCanvas->DrawShadowedString(X1, InY, TEXT("VMem:"), Font, bShowUnitTimeGraph ? FColor(100, 100, 255) : FColor::White);
+				InCanvas->DrawShadowedString(X2, InY, *GetMemoryString(Stats.UsedVirtual), Font, FColor::Green);
+				InCanvas->DrawShadowedString(X3, InY, *GetMemoryString(Stats.PeakUsedVirtual), Font, FColor::Green);
 				InY += RowHeight;
 			}
+			else
+			{
+				uint64 MemoryUsed = FPlatformMemory::GetMemoryUsedFast();
+				if (MemoryUsed > 0)
+				{
+					// print out currently used memory
+					InCanvas->DrawShadowedString(X1, InY, TEXT("Mem:"), Font, bShowUnitTimeGraph ? FColor(100, 100, 255) : FColor::White);
+					InCanvas->DrawShadowedString(X2, InY, *GetMemoryString(MemoryUsed), Font, FColor::Green);
+					InY += RowHeight;
+				}
+			}
 		}
 
 		ERHIFeatureLevel::Type FeatureLevel = InCanvas->GetFeatureLevel();
diff --git a/Engine/Source/Runtime/Engine/Public/UnrealClient.h b/Engine/Source/Runtime/Engine/Public/UnrealClient.h
--- a/Engine/Source/Runtime/Engine/Public/UnrealClient.h
+++ b/Engine/Source/Runtime/Engine/Public/UnrealClient.h
@@ -196,6 +196,9 @@ private:
 	static TArray<FColor> HighresScreenshotMaskColorArray;
 };
 
+// @param bAutoType true: automatically choose GB/MB/KB/... false: always use MB for easier comparisons
+ENGINE_API FString GetMemoryString( const double Value, const bool bAutoType = true );
+
 /** Data needed to display perframe stat tracking when STAT UNIT is enabled */
 struct FStatUnitData
 {
