Added to the Projectile Movement Component the flag to keep the rotation vertical when matching rotation to velocity. This allows objects moving up/down to maintain pitch/roll while adjusting yaw to match direction.

zak.middleton

diff --git a/Engine/Source/Runtime/Engine/Classes/GameFramework/ProjectileMovementComponent.h b/Engine/Source/Runtime/Engine/Classes/GameFramework/ProjectileMovementComponent.h
--- a/Engine/Source/Runtime/Engine/Classes/GameFramework/ProjectileMovementComponent.h
+++ b/Engine/Source/Runtime/Engine/Classes/GameFramework/ProjectileMovementComponent.h
@@ -38,6 +38,10 @@ class ENGINE_API UProjectileMovementComponent : public UMovementComponent
 	UPROPERTY(EditAnywhere, BlueprintReadWrite, Category=Projectile)
 	uint8 bRotationFollowsVelocity:1;
 
+	/** If true, this projectile will have its rotation updated each frame to maintain the rotations Yaw only. (bRotationFollowsVelocity is required to be true) */
+	UPROPERTY(EditAnywhere, BlueprintReadWrite, Category=Projectile, meta = (EditCondition = "bRotationFollowsVelocity"))
+	uint8 bRotationRemainsVertical:1;
+
 	/** If true, simple bounces will be simulated. Set this to false to stop simulating on contact. */
 	UPROPERTY(EditAnywhere, BlueprintReadWrite, Category=ProjectileBounces)
 	uint8 bShouldBounce:1;
diff --git a/Engine/Source/Runtime/Engine/Private/Components/ProjectileMovementComponent.cpp b/Engine/Source/Runtime/Engine/Private/Components/ProjectileMovementComponent.cpp
--- a/Engine/Source/Runtime/Engine/Private/Components/ProjectileMovementComponent.cpp
+++ b/Engine/Source/Runtime/Engine/Private/Components/ProjectileMovementComponent.cpp
@@ -91,7 +91,15 @@ void UProjectileMovementComponent::InitializeComponent()
 		{
 			if (UpdatedComponent)
 			{
-				UpdatedComponent->SetWorldRotation(Velocity.Rotation());
+				FRotator DesiredRotation = Velocity.Rotation();
+				if (bRotationRemainsVertical)
+				{
+					DesiredRotation.Pitch = 0.0f;
+					DesiredRotation.Yaw = FRotator::NormalizeAxis(DesiredRotation.Yaw);
+					DesiredRotation.Roll = 0.0f;
+				}
+
+				UpdatedComponent->SetWorldRotation(DesiredRotation);
 			}
 		}
 
@@ -182,7 +190,16 @@ void UProjectileMovementComponent::TickComponent(float DeltaTime, enum ELevelTic
 		Hit.Time = 1.f;
 		const FVector OldVelocity = Velocity;
 		const FVector MoveDelta = ComputeMoveDelta(OldVelocity, TimeTick);
-		const FQuat NewRotation = (bRotationFollowsVelocity && !OldVelocity.IsNearlyZero(0.01f)) ? OldVelocity.ToOrientationQuat() : UpdatedComponent->GetComponentQuat();
+		FQuat NewRotation = (bRotationFollowsVelocity && !OldVelocity.IsNearlyZero(0.01f)) ? OldVelocity.ToOrientationQuat() : UpdatedComponent->GetComponentQuat();
+
+		if (bRotationFollowsVelocity && bRotationRemainsVertical)
+		{
+			FRotator DesiredRotation = NewRotation.Rotator();
+			DesiredRotation.Pitch = 0.0f;
+			DesiredRotation.Yaw = FRotator::NormalizeAxis(DesiredRotation.Yaw);
+			DesiredRotation.Roll = 0.0f;
+			NewRotation = DesiredRotation.Quaternion();
+		}
 
 		// Move the component
 		if (bShouldBounce)
