Disabling UActorChannel::PrepareForRemoteFunction for now as it is causing an assert

#rnx

diff --git a/Engine/Source/Runtime/Engine/Private/DataChannel.cpp b/Engine/Source/Runtime/Engine/Private/DataChannel.cpp
--- a/Engine/Source/Runtime/Engine/Private/DataChannel.cpp
+++ b/Engine/Source/Runtime/Engine/Private/DataChannel.cpp
@@ -3083,7 +3083,7 @@ void UActorChannel::PrepareForRemoteFunction(UObject* TargetObj)
 {
 	// Make sure we create a replicator in case we destroy a sub object before we ever try to replicate its properties,
 	// otherwise it will not be in the ReplicationMap and we'll never send the deletion to clients
-	FindOrCreateReplicator(TargetObj);
+	//FindOrCreateReplicator(TargetObj);
 }
 
 void UActorChannel::QueueRemoteFunctionBunch( UObject* CallTarget, UFunction* Func, FOutBunch &Bunch )
