character movement prediction: added virtual accessors to all operations on control rotation value

diff --git a/Engine/Source/Runtime/Engine/Classes/GameFramework/CharacterMovementComponent.h b/Engine/Source/Runtime/Engine/Classes/GameFramework/CharacterMovementComponent.h
--- a/Engine/Source/Runtime/Engine/Classes/GameFramework/CharacterMovementComponent.h
+++ b/Engine/Source/Runtime/Engine/Classes/GameFramework/CharacterMovementComponent.h
@@ -2120,6 +2120,8 @@ public:
 	virtual void ResetPredictionData_Client() override;
 	virtual void ResetPredictionData_Server() override;
 
+	static uint32 PackYawAndPitchTo32(const float Yaw, const float Pitch);
+
 protected:
 	class FNetworkPredictionData_Client_Character* ClientPredictionData;
 	class FNetworkPredictionData_Server_Character* ServerPredictionData;
@@ -2141,8 +2143,6 @@ protected:
 	/** Update mesh location based on interpolated values. */
 	void SmoothClientPosition_UpdateVisuals();
 
-	static uint32 PackYawAndPitchTo32(const float Yaw, const float Pitch); 
-
 	/*
 	========================================================================
 	Here's how player movement prediction, replication and correction works in network games:
@@ -2642,6 +2642,12 @@ public:
 	/** Returns a byte containing encoded special movement information (jumping, crouching, etc.)	 */
 	virtual uint8 GetCompressedFlags() const;
 
+	/** Compare current control rotation with stored starting data */
+	virtual bool IsMatchingStartControlRotation(const APlayerController* PC) const;
+
+	/** Packs control rotation for network transport */
+	virtual void GetPackedAngles(uint32& YawAndPitchPack, uint8& RollPack) const;
+
 	// Bit masks used by GetCompressedFlags() to encode movement information.
 	enum CompressedFlags
 	{
diff --git a/Engine/Source/Runtime/Engine/Private/Components/CharacterMovementComponent.cpp b/Engine/Source/Runtime/Engine/Private/Components/CharacterMovementComponent.cpp
--- a/Engine/Source/Runtime/Engine/Private/Components/CharacterMovementComponent.cpp
+++ b/Engine/Source/Runtime/Engine/Private/Components/CharacterMovementComponent.cpp
@@ -8107,8 +8107,9 @@ void UCharacterMovementComponent::CallServerMove
 	check(NewMove != nullptr);
 
 	// Compress rotation down to 5 bytes
-	const uint32 ClientYawPitchINT = PackYawAndPitchTo32(NewMove->SavedControlRotation.Yaw, NewMove->SavedControlRotation.Pitch);
-	const uint8 ClientRollBYTE = FRotator::CompressAxisToByte(NewMove->SavedControlRotation.Roll);
+	uint32 ClientYawPitchINT = 0;
+	uint8 ClientRollBYTE = 0;
+	NewMove->GetPackedAngles(ClientYawPitchINT, ClientRollBYTE);
 
 	// Determine if we send absolute or relative location
 	UPrimitiveComponent* ClientMovementBase = NewMove->EndBase.Get();
@@ -8124,7 +8125,9 @@ void UCharacterMovementComponent::CallServerMove
 	FNetworkPredictionData_Client_Character* ClientData = GetPredictionData_Client_Character();
 	if (const FSavedMove_Character* const PendingMove = ClientData->PendingMove.Get())
 	{
-		const uint32 OldClientYawPitchINT = PackYawAndPitchTo32(ClientData->PendingMove->SavedControlRotation.Yaw, ClientData->PendingMove->SavedControlRotation.Pitch);
+		uint32 OldClientYawPitchINT = 0;
+		uint8 OldClientRollBYTE = 0;
+		ClientData->PendingMove->GetPackedAngles(OldClientYawPitchINT, OldClientRollBYTE);
 
 		// If we delayed a move without root motion, and our new move has root motion, send these through a special function, so the server knows how to process them.
 		if ((PendingMove->RootMotionMontage == NULL) && (NewMove->RootMotionMontage != NULL))
@@ -10703,7 +10706,7 @@ float UCharacterMovementComponent::GetClientNetSendDeltaTime(const APlayerContro
 		}
 
 		// Lower frequency for standing still and not rotating camera
-		if (Acceleration.IsZero() && Velocity.IsZero() && ClientData->LastAckedMove.IsValid() && ClientData->LastAckedMove->StartControlRotation.Equals(PC->GetControlRotation(), CharacterMovementCVars::NetStationaryRotationTolerance))
+		if (Acceleration.IsZero() && Velocity.IsZero() && ClientData->LastAckedMove.IsValid() && ClientData->LastAckedMove->IsMatchingStartControlRotation(PC))
 		{
 			NetMoveDelta = FMath::Max(GameNetworkManager->ClientNetSendMoveDeltaTimeStationary, NetMoveDelta);
 		}
@@ -10712,6 +10715,18 @@ float UCharacterMovementComponent::GetClientNetSendDeltaTime(const APlayerContro
 	return NetMoveDelta;
 }
 
+bool FSavedMove_Character::IsMatchingStartControlRotation(const APlayerController* PC) const
+{
+	return PC ? StartControlRotation.Equals(PC->GetControlRotation(), CharacterMovementCVars::NetStationaryRotationTolerance) : false;
+}
+
+void FSavedMove_Character::GetPackedAngles(uint32& YawAndPitchPack, uint8& RollPack) const
+{
+	// Compress rotation down to 5 bytes
+	YawAndPitchPack = UCharacterMovementComponent::PackYawAndPitchTo32(SavedControlRotation.Yaw, SavedControlRotation.Pitch);
+	RollPack = FRotator::CompressAxisToByte(SavedControlRotation.Roll);
+}
+
 bool FSavedMove_Character::CanCombineWith(const FSavedMovePtr& NewMovePtr, ACharacter* Character, float MaxDelta) const
 {
 	const FSavedMove_Character* NewMove = NewMovePtr.Get();
