Use more efficient Reevaluate than manually doing Remove/Add when updating level streaming priority
#rnx

diff --git a/Engine/Source/Runtime/Engine/Private/World.cpp b/Engine/Source/Runtime/Engine/Private/World.cpp
--- a/Engine/Source/Runtime/Engine/Private/World.cpp
+++ b/Engine/Source/Runtime/Engine/Private/World.cpp
@@ -3287,10 +3287,7 @@ void UWorld::UpdateStreamingLevelPriority(ULevelStreaming* StreamingLevel)
 {
 	if (StreamingLevel)
 	{
-		if (StreamingLevelsToConsider.Remove(StreamingLevel))
-		{
-			StreamingLevelsToConsider.Add(StreamingLevel);
-		}
+		StreamingLevelsToConsider.Reevaluate(StreamingLevel);
 	}
 }
 
