ShouldCreateRayTracedCascade now takes into account r.DistanceFieldShadowing. If project enables mesh distance fields generation, on platform that supports the feature, but platform has distance field disabled, function now returns false

#rnx

diff --git a/Engine/Source/Runtime/Engine/Private/Components/DirectionalLightComponent.cpp b/Engine/Source/Runtime/Engine/Private/Components/DirectionalLightComponent.cpp
--- a/Engine/Source/Runtime/Engine/Private/Components/DirectionalLightComponent.cpp
+++ b/Engine/Source/Runtime/Engine/Private/Components/DirectionalLightComponent.cpp
@@ -339,6 +339,12 @@ public:
 
 	virtual bool ShouldCreateRayTracedCascade(ERHIFeatureLevel::Type InFeatureLevel, bool bPrecomputedLightingIsValid, int32 MaxNearCascades) const override
 	{
+		static auto CVarDistanceFieldShadowing = IConsoleManager::Get().FindConsoleVariable(TEXT("r.DistanceFieldShadowing"));
+		if (CVarDistanceFieldShadowing != nullptr && CVarDistanceFieldShadowing->GetInt() == 0)
+		{
+			return false;
+		}
+
 		const uint32 NumCascades = GetNumShadowMappedCascades(MaxNearCascades, bPrecomputedLightingIsValid);
 		const float RayTracedShadowDistance = GetDistanceFieldShadowDistance();
 		const bool bCreateWithCSM = NumCascades > 0 && RayTracedShadowDistance > GetCSMMaxDistance(bPrecomputedLightingIsValid, MaxNearCascades);
