Added support for the new Shading Model type through the MakeMaterialAttributes material expression.

#rnx

diff --git a/Engine/Source/Runtime/Engine/Classes/Materials/MaterialExpressionMakeMaterialAttributes.h b/Engine/Source/Runtime/Engine/Classes/Materials/MaterialExpressionMakeMaterialAttributes.h
--- a/Engine/Source/Runtime/Engine/Classes/Materials/MaterialExpressionMakeMaterialAttributes.h
+++ b/Engine/Source/Runtime/Engine/Classes/Materials/MaterialExpressionMakeMaterialAttributes.h
@@ -80,6 +80,7 @@ class UMaterialExpressionMakeMaterialAttributes : public UMaterialExpression
 	virtual int32 Compile(class FMaterialCompiler* Compiler, int32 OutputIndex) override;
 	virtual void GetCaption(TArray<FString>& OutCaptions) const override;
 	virtual bool IsResultMaterialAttributes(int32 OutputIndex) override {return true;}
+	virtual uint32 GetInputType(int32 InputIndex) override;
 #endif // WITH_EDITOR
 	//~ End UMaterialExpression Interface
 };
diff --git a/Engine/Source/Runtime/Engine/Private/Materials/MaterialExpressions.cpp b/Engine/Source/Runtime/Engine/Private/Materials/MaterialExpressions.cpp
--- a/Engine/Source/Runtime/Engine/Private/Materials/MaterialExpressions.cpp
+++ b/Engine/Source/Runtime/Engine/Private/Materials/MaterialExpressions.cpp
@@ -4589,6 +4589,18 @@ void UMaterialExpressionMakeMaterialAttributes::GetCaption(TArray<FString>& OutC
 {
 	OutCaptions.Add(TEXT("MakeMaterialAttributes"));
 }
+
+uint32 UMaterialExpressionMakeMaterialAttributes::GetInputType(int32 InputIndex)
+{
+	if (GetInputName(InputIndex).IsEqual("ShadingModel"))
+	{
+		return MCT_ShadingModel;
+	}
+	else
+	{
+		return UMaterialExpression::GetInputType(InputIndex);
+	}
+}
 #endif // WITH_EDITOR
 
 // -----
