Fixing UMaterialExpressionBreakMaterialAttributes to work with the new ShadingModel output.

#rnx

diff --git a/Engine/Source/Runtime/Engine/Classes/Materials/MaterialExpressionBreakMaterialAttributes.h b/Engine/Source/Runtime/Engine/Classes/Materials/MaterialExpressionBreakMaterialAttributes.h
--- a/Engine/Source/Runtime/Engine/Classes/Materials/MaterialExpressionBreakMaterialAttributes.h
+++ b/Engine/Source/Runtime/Engine/Classes/Materials/MaterialExpressionBreakMaterialAttributes.h
@@ -30,6 +30,7 @@ class UMaterialExpressionBreakMaterialAttributes : public UMaterialExpression
 	virtual FName GetInputName(int32 InputIndex) const override;
 	virtual bool IsInputConnectionRequired(int32 InputIndex) const override;
 	virtual uint32 GetInputType(int32 InputIndex) override {return MCT_MaterialAttributes;}
+	virtual uint32 GetOutputType(int32 OutputIndex) override;
 #endif // WITH_EDITOR
 	//~ End UMaterialExpression Interface
 };
diff --git a/Engine/Source/Runtime/Engine/Private/Materials/MaterialExpressions.cpp b/Engine/Source/Runtime/Engine/Private/Materials/MaterialExpressions.cpp
--- a/Engine/Source/Runtime/Engine/Private/Materials/MaterialExpressions.cpp
+++ b/Engine/Source/Runtime/Engine/Private/Materials/MaterialExpressions.cpp
@@ -4656,7 +4656,7 @@ UMaterialExpressionBreakMaterialAttributes::UMaterialExpressionBreakMaterialAttr
 	}
 
 	Outputs.Add(FExpressionOutput(TEXT("PixelDepthOffset"), 1, 1, 0, 0, 0));
-	Outputs.Add(FExpressionOutput(TEXT("ShadingModel"), 1, 1, 0, 0, 0));
+	Outputs.Add(FExpressionOutput(TEXT("ShadingModel"), 0, 0, 0, 0, 0));
 #endif
 }
 
@@ -4696,15 +4696,16 @@ void UMaterialExpressionBreakMaterialAttributes::Serialize(FStructuredArchive::F
 		}
 
 		Outputs[OutputIndex].SetMask(1, 1, 0, 0, 0); ++OutputIndex;// PixelDepthOffset
-		Outputs[OutputIndex].SetMask(1, 1, 0, 0, 0); // ShadingModelFromMaterialExpression
+		Outputs[OutputIndex].SetMask(0, 0, 0, 0, 0); // ShadingModelFromMaterialExpression
 	}
 #endif // WITH_EDITOR
 }
 
 #if WITH_EDITOR
-int32 UMaterialExpressionBreakMaterialAttributes::Compile(class FMaterialCompiler* Compiler, int32 OutputIndex)
+static TMap<EMaterialProperty, int32> PropertyToIOIndexMap;
+
+static void BuildPropertyToIOIndexMap()
 {
-	static TMap<EMaterialProperty, int32> PropertyToIOIndexMap;
 	if (PropertyToIOIndexMap.Num() == 0)
 	{
 		PropertyToIOIndexMap.Add(MP_BaseColor, 0);
@@ -4734,6 +4735,11 @@ int32 UMaterialExpressionBreakMaterialAttributes::Compile(class FMaterialCompile
 		PropertyToIOIndexMap.Add(MP_PixelDepthOffset, 24);
 		PropertyToIOIndexMap.Add(MP_ShadingModel, 25);
 	}
+}
+
+int32 UMaterialExpressionBreakMaterialAttributes::Compile(class FMaterialCompiler* Compiler, int32 OutputIndex)
+{
+	BuildPropertyToIOIndexMap();
 
 	// Here we don't care about any multiplex index coming in.
 	// We pass through our output index as the multiplex index so the MakeMaterialAttriubtes node at the other end can send us the right data.
@@ -4785,6 +4791,22 @@ bool UMaterialExpressionBreakMaterialAttributes::IsInputConnectionRequired(int32
 {
 	return true;
 }
+
+uint32 UMaterialExpressionBreakMaterialAttributes::GetOutputType(int32 OutputIndex)
+{
+	BuildPropertyToIOIndexMap();
+
+	const EMaterialProperty* Property = PropertyToIOIndexMap.FindKey(OutputIndex);
+
+	if (Property && *Property == EMaterialProperty::MP_ShadingModel)
+	{
+		return MCT_ShadingModel;
+	}
+	else
+	{
+		return UMaterialExpression::GetOutputType(OutputIndex);
+	}
+}
 #endif // WITH_EDITOR
 
 // -----
