Fix OverrideFarClippingPlaneDistance not being taken into account in HISM Culling

diff --git a/Engine/Source/Runtime/Engine/Private/HierarchicalInstancedStaticMesh.cpp b/Engine/Source/Runtime/Engine/Private/HierarchicalInstancedStaticMesh.cpp
--- a/Engine/Source/Runtime/Engine/Private/HierarchicalInstancedStaticMesh.cpp
+++ b/Engine/Source/Runtime/Engine/Private/HierarchicalInstancedStaticMesh.cpp
@@ -1620,6 +1620,10 @@ void FHierarchicalStaticMeshSceneProxy::GetDynamicMeshElements(const TArray<cons
 				{
 					FinalCull = ComputeBoundsDrawDistance(MinSize, SphereRadius, View->ViewMatrices.GetProjectionMatrix()) * LODScale;
 				}
+				if (View->SceneViewInitOptions.OverrideFarClippingPlaneDistance > 0.0f)
+				{
+					FinalCull = FMath::Min(FinalCull, View->SceneViewInitOptions.OverrideFarClippingPlaneDistance * MaxDrawDistanceScale);
+				}
 				if (UserData_AllInstances.EndCullDistance > 0.0f)
 				{
 					FinalCull = FMath::Min(FinalCull, UserData_AllInstances.EndCullDistance * MaxDrawDistanceScale);
@@ -1796,6 +1800,10 @@ void FHierarchicalStaticMeshSceneProxy::GetDynamicMeshElements(const TArray<cons
 						{
 							FinalCull = ComputeBoundsDrawDistance(MinSize, SphereRadius, View->ViewMatrices.GetProjectionMatrix()) * LODScale;
 						}
+						if (View->SceneViewInitOptions.OverrideFarClippingPlaneDistance > 0.0f)
+						{
+							FinalCull = FMath::Min(FinalCull, View->SceneViewInitOptions.OverrideFarClippingPlaneDistance * MaxDrawDistanceScale);
+						}
 						if (UserData_AllInstances.EndCullDistance > 0.0f)
 						{
 							FinalCull = FMath::Min(FinalCull, UserData_AllInstances.EndCullDistance * MaxDrawDistanceScale);
