Fix for active sounds crash in GC

diff --git a/Engine/Source/Runtime/Engine/Private/AudioDevice.cpp b/Engine/Source/Runtime/Engine/Private/AudioDevice.cpp
--- a/Engine/Source/Runtime/Engine/Private/AudioDevice.cpp
+++ b/Engine/Source/Runtime/Engine/Private/AudioDevice.cpp
@@ -634,6 +634,11 @@ void FAudioDevice::AddReferencedObjects(FReferenceCollector& Collector)
 		ActiveSound->AddReferencedObjects(Collector);
 	}
 
+	for (FActiveSound* ActiveSound : PendingSoundsToDelete)
+	{
+		ActiveSound->AddReferencedObjects(Collector);
+	}	
+
 	for (FVirtualLoopPair& Pair : VirtualLoops)
 	{
 		Pair.Key->AddReferencedObjects(Collector);
