Fix for possible null pointer crash when attempting to get the game state by unique net id.

diff --git a/Engine/Source/Runtime/Engine/Private/GameStateBase.cpp b/Engine/Source/Runtime/Engine/Private/GameStateBase.cpp
--- a/Engine/Source/Runtime/Engine/Private/GameStateBase.cpp
+++ b/Engine/Source/Runtime/Engine/Private/GameStateBase.cpp
@@ -229,7 +229,7 @@ APlayerState* AGameStateBase::GetPlayerStateFromUniqueNetId(const FUniqueNetIdWr
 	const TArray<APlayerState*>& Players = PlayerArray;
 	for (APlayerState* Player : Players)
 	{
-		if (Player->UniqueId == InPlayerId)
+		if (Player && Player->UniqueId == InPlayerId)
 		{
 			return Player;
 		}
