Improved lifetime management of the shared resource owned by UVectorFieldStatic.

diff --git a/Engine/Source/Runtime/Engine/Private/Particles/ParticleGpuSimulation.cpp b/Engine/Source/Runtime/Engine/Private/Particles/ParticleGpuSimulation.cpp
--- a/Engine/Source/Runtime/Engine/Private/Particles/ParticleGpuSimulation.cpp
+++ b/Engine/Source/Runtime/Engine/Private/Particles/ParticleGpuSimulation.cpp
@@ -3918,15 +3918,15 @@ private:
 		//
 		if (!Simulation->bDestroyed_GameThread)
 		{
-			FVectorFieldResource* Resource = Simulation->LocalVectorField.Resource;
-			ENQUEUE_RENDER_COMMAND(FResetVectorFieldCommand)(
-				[Resource](FRHICommandList& RHICmdList)
-				{
-					if (Resource)
-					{
+			FVectorFieldResource* Resource = Simulation->LocalVectorField.Resource.GetReference();
+			if (Resource)
+			{
+			    ENQUEUE_RENDER_COMMAND(FResetVectorFieldCommand)(
+				    [Resource](FRHICommandList& RHICmdList)
+				    {
 						Resource->ResetVectorField();
-					}
-				});
+					});
+		    }
 		}
 	}
 
@@ -4634,23 +4634,26 @@ static void SetParametersForVectorField(FVectorFieldUniformParameters& OutParame
 	check(VectorFieldInstance && VectorFieldInstance->Resource);
 	check(Index < MAX_VECTOR_FIELDS);
 
-	FVectorFieldResource* Resource = VectorFieldInstance->Resource;
-	const float Intensity = VectorFieldInstance->Intensity * Resource->Intensity * EmitterScale;
-
-	// Override vector field tightness if value is set (greater than 0). This override is only used for global vector fields.
-	float Tightness = EmitterTightness;
-	if(EmitterTightness == -1)
+	FVectorFieldResource* Resource = VectorFieldInstance->Resource.GetReference();
+	if (Resource)
 	{
-		Tightness = FMath::Clamp<float>(VectorFieldInstance->Tightness, 0.0f, 1.0f);
-	}
+		const float Intensity = VectorFieldInstance->Intensity * Resource->Intensity * EmitterScale;
+
+		// Override vector field tightness if value is set (greater than 0). This override is only used for global vector fields.
+		float Tightness = EmitterTightness;
+		if(EmitterTightness == -1)
+		{
+			Tightness = FMath::Clamp<float>(VectorFieldInstance->Tightness, 0.0f, 1.0f);
+		}
 
-	OutParameters.WorldToVolume[Index] = VectorFieldInstance->WorldToVolume;
-	OutParameters.VolumeToWorld[Index] = VectorFieldInstance->VolumeToWorldNoScale;
-	OutParameters.VolumeSize[Index] = FVector4(Resource->SizeX, Resource->SizeY, Resource->SizeZ, 0);
-	OutParameters.IntensityAndTightness[Index] = FVector4(Intensity, Tightness, 0, 0 );
-	OutParameters.TilingAxes[Index].X = VectorFieldInstance->bTileX ? 1.0f : 0.0f;
-	OutParameters.TilingAxes[Index].Y = VectorFieldInstance->bTileY ? 1.0f : 0.0f;
-	OutParameters.TilingAxes[Index].Z = VectorFieldInstance->bTileZ ? 1.0f : 0.0f;
+		OutParameters.WorldToVolume[Index] = VectorFieldInstance->WorldToVolume;
+		OutParameters.VolumeToWorld[Index] = VectorFieldInstance->VolumeToWorldNoScale;
+		OutParameters.VolumeSize[Index] = FVector4(Resource->SizeX, Resource->SizeY, Resource->SizeZ, 0);
+		OutParameters.IntensityAndTightness[Index] = FVector4(Intensity, Tightness, 0, 0 );
+		OutParameters.TilingAxes[Index].X = VectorFieldInstance->bTileX ? 1.0f : 0.0f;
+		OutParameters.TilingAxes[Index].Y = VectorFieldInstance->bTileY ? 1.0f : 0.0f;
+		OutParameters.TilingAxes[Index].Z = VectorFieldInstance->bTileZ ? 1.0f : 0.0f;
+	}
 }
 
 bool FFXSystem::UsesGlobalDistanceFieldInternal() const
diff --git a/Engine/Source/Runtime/Engine/Private/VectorField.cpp b/Engine/Source/Runtime/Engine/Private/VectorField.cpp
--- a/Engine/Source/Runtime/Engine/Private/VectorField.cpp
+++ b/Engine/Source/Runtime/Engine/Private/VectorField.cpp
@@ -57,15 +57,16 @@ FVectorFieldInstance::~FVectorFieldInstance()
 {
 	if (Resource && bInstancedResource)
 	{
-		FVectorFieldResource* InResource = Resource;
+		FVectorFieldResource* InResource = Resource.GetReference();
+		InResource->AddRef();
 		ENQUEUE_RENDER_COMMAND(FDestroyVectorFieldResourceCommand)(
 			[InResource](FRHICommandList& RHICmdList)
 			{
 				InResource->ReleaseResource();
-				delete InResource;
+				InResource->Release();
 			});
-		Resource = nullptr;
 	}
+	Resource = nullptr;
 }
 
 /**
@@ -75,7 +76,7 @@ FVectorFieldInstance::~FVectorFieldInstance()
  */
 void FVectorFieldInstance::Init(FVectorFieldResource* InResource, bool bInstanced)
 {
-	check(Resource == NULL);
+	check(!Resource);
 	Resource = InResource;
 	bInstancedResource = bInstanced;
 }
@@ -188,6 +189,7 @@ public:
 		InVectorField->SourceData.GetCopy(&VolumeData, /*bDiscardInternalCopy=*/ true);
 	}
 
+protected:
 	/** Destructor. */
 	virtual ~FVectorFieldStaticResource()
 	{
@@ -195,6 +197,8 @@ public:
 		VolumeData = NULL;
 	}
 
+public:
+
 	/**
 	 * Initialize RHI resources.
 	 */
@@ -243,19 +247,19 @@ public:
 		ENQUEUE_RENDER_COMMAND(FUpdateStaticVectorFieldCommand)(
 			[Resource, UpdateParams](FRHICommandListImmediate& RHICmdList)
 			{
-					// Free any existing volume data on the resource.
-					FMemory::Free(Resource->VolumeData);
-
-					// Update settings on this resource.
-					Resource->SizeX = UpdateParams.SizeX;
-					Resource->SizeY = UpdateParams.SizeY;
-					Resource->SizeZ = UpdateParams.SizeZ;
-					Resource->Intensity = UpdateParams.Intensity;
-					Resource->LocalBounds = UpdateParams.Bounds;
-					Resource->VolumeData = UpdateParams.VolumeData;
-
-					// Update RHI resources.
-					Resource->UpdateRHI();
+				// Free any existing volume data on the resource.
+				FMemory::Free(Resource->VolumeData);
+
+				// Update settings on this resource.
+				Resource->SizeX = UpdateParams.SizeX;
+				Resource->SizeY = UpdateParams.SizeY;
+				Resource->SizeZ = UpdateParams.SizeZ;
+				Resource->Intensity = UpdateParams.Intensity;
+				Resource->LocalBounds = UpdateParams.Bounds;
+				Resource->VolumeData = UpdateParams.VolumeData;
+
+				// Update RHI resources.
+				Resource->UpdateRHI();
 			});
 	}
 
@@ -279,19 +283,21 @@ void UVectorFieldStatic::InitInstance(FVectorFieldInstance* Instance, bool bPrev
 
 void UVectorFieldStatic::InitResource()
 {
-	check(Resource == NULL);
+	check(!Resource);
 
 	// Loads and copies the bulk data into CPUData if bAllowCPUAccess is set, otherwise clear CPUData. 
 	UpdateCPUData();
 
 	Resource = new FVectorFieldStaticResource(this);
+	Resource->AddRef(); // Increment refcount because of UVectorFieldStatic::Resource is not a TRefCountPtr.
+
 	BeginInitResource(Resource);
 }
 
 
 void UVectorFieldStatic::UpdateResource()
 {
-	check(Resource != NULL);
+	check(Resource);
 
 	// Loads and copies the bulk data into CPUData if bAllowCPUAccess is set, otherwise clears CPUData. 
 	UpdateCPUData();
@@ -377,17 +383,18 @@ FRHITexture* UVectorFieldStatic::GetVolumeTextureRef()
 
 void UVectorFieldStatic::ReleaseResource()
 {
-	if ( Resource != NULL )
+	if (Resource)
 	{
-		FRenderResource* InResource = Resource;
+		FVectorFieldResource* InResource = Resource;
 		ENQUEUE_RENDER_COMMAND(ReleaseVectorFieldCommand)(
 			[InResource](FRHICommandList& RHICmdList)
 			{
 				InResource->ReleaseResource();
-				delete InResource;
+				// Decrement the refcount and possibly delete (see refs from FVectorFieldInstance).
+				InResource->Release(); 
 			});
+		Resource = nullptr;
 	}
-	Resource = nullptr;
 }
 
 void UVectorFieldStatic::Serialize(FArchive& Ar)
diff --git a/Engine/Source/Runtime/Engine/Private/VectorField.h b/Engine/Source/Runtime/Engine/Private/VectorField.h
--- a/Engine/Source/Runtime/Engine/Private/VectorField.h
+++ b/Engine/Source/Runtime/Engine/Private/VectorField.h
@@ -67,6 +67,30 @@ public:
 	 * Resets the vector field simulation.
 	 */
 	virtual void ResetVectorField() {}
+
+
+	/** Ref couting **/
+	uint32 AddRef() const
+	{ 
+		return (uint32)NumRefs.Increment(); 
+	}
+
+	uint32 Release() const
+	{
+		uint32 NewValue = (uint32)NumRefs.Decrement();
+		if (NewValue == 0)
+		{
+			delete this;
+		}
+		return NewValue;
+	}
+
+protected:
+
+	virtual ~FVectorFieldResource(){}
+
+	/** Ref couting. Required becauce of how resources are shared between FVectorFieldInstance and UVectorField... **/
+	mutable FThreadSafeCounter NumRefs;
 };
 
 /**
@@ -76,8 +100,8 @@ class FVectorFieldInstance
 {
 public:
 
-	/** The vector field resource. */
-	FVectorFieldResource* Resource;
+	/** The vector field resource. Kept an active ref count to manage release with UVectorFieldStatic. */
+	TRefCountPtr<FVectorFieldResource> Resource;
 	/** Bounds of the vector field in world space. */
 	FBox WorldBounds;
 	/** Transform from the vector field's local space to world space, no scaling is applied. */
diff --git a/Engine/Source/Runtime/Engine/Private/VectorFieldVisualization.cpp b/Engine/Source/Runtime/Engine/Private/VectorFieldVisualization.cpp
--- a/Engine/Source/Runtime/Engine/Private/VectorFieldVisualization.cpp
+++ b/Engine/Source/Runtime/Engine/Private/VectorFieldVisualization.cpp
@@ -244,7 +244,7 @@ void GetVectorFieldMesh(
 	int32 ViewIndex,
 	FMeshElementCollector& Collector)
 {
-	FVectorFieldResource* Resource = VectorFieldInstance->Resource;
+	FVectorFieldResource* Resource = VectorFieldInstance->Resource.GetReference();
 
 	if (Resource && IsValidRef(Resource->VolumeTextureRHI))
 	{
