Port CL# 6414161 to Release-4.22

Fix morph target reset for 1 frame when changing material
- removed ActiveMorphtarget clearing on create render state because we're clearing in Destroy Render state
- We update morphtarget info if not disabled and allowed, so that we don't lose the current acitve target data

diff --git a/Engine/Source/Runtime/Engine/Private/Components/SkinnedMeshComponent.cpp b/Engine/Source/Runtime/Engine/Private/Components/SkinnedMeshComponent.cpp
--- a/Engine/Source/Runtime/Engine/Private/Components/SkinnedMeshComponent.cpp
+++ b/Engine/Source/Runtime/Engine/Private/Components/SkinnedMeshComponent.cpp
@@ -586,9 +586,9 @@ void USkinnedMeshComponent::CreateRenderState_Concurrent()
 				const bool bMorphTargetsAllowed = CVarEnableMorphTargets.GetValueOnAnyThread(true) != 0;
 
 				// Are morph targets disabled for this LOD?
-				if(bDisableMorphTarget || !bMorphTargetsAllowed)
+				if (!bDisableMorphTarget && bMorphTargetsAllowed)
 				{
-					ActiveMorphTargets.Empty();
+					RefreshMorphTargets();
 				}
 
 				MeshObject->Update(PredictedLODLevel, this, ActiveMorphTargets, MorphTargetWeights, EPreviousBoneTransformUpdateMode::UpdatePrevious);  // send to rendering thread
