Expose bForcePVRTC4 in the editor on textures.

diff --git a/Engine/Source/Runtime/Engine/Classes/Engine/Texture.h b/Engine/Source/Runtime/Engine/Classes/Engine/Texture.h
--- a/Engine/Source/Runtime/Engine/Classes/Engine/Texture.h
+++ b/Engine/Source/Runtime/Engine/Classes/Engine/Texture.h
@@ -470,7 +470,7 @@ public:
 	uint32 bFlipGreenChannel:1;
 
 	/** For DXT1 textures, setting this will cause the texture to be twice the size, but better looking, on iPhone */
-	UPROPERTY()
+	UPROPERTY(EditAnywhere, Category = Texture, AdvancedDisplay)
 	uint32 bForcePVRTC4:1;
 
 	/** How to pad the texture to a power of 2 size (if necessary) */
