Optimize cook ULevel::PreSave by moving the validation of light guids to LightComponent::PreSave().

Results:
ULevel::PreSave down from ~18.5 seconds to ~0.5 seconds (VTune user-mode sampling cpu time) on a full cook saving ~183 000 packages.

diff --git a/Engine/Source/Runtime/Engine/Classes/Components/LightComponent.h b/Engine/Source/Runtime/Engine/Classes/Components/LightComponent.h
--- a/Engine/Source/Runtime/Engine/Classes/Components/LightComponent.h
+++ b/Engine/Source/Runtime/Engine/Classes/Components/LightComponent.h
@@ -359,6 +359,7 @@ public:
 	virtual void Serialize(FArchive& Ar) override;
 	virtual void PostLoad() override;
 #if WITH_EDITOR
+	virtual void PreSave(const class ITargetPlatform* TargetPlatform) override;
 	virtual bool CanEditChange(const UProperty* InProperty) const override;
 	virtual void PostEditChangeProperty(FPropertyChangedEvent& PropertyChangedEvent) override;
 	virtual void UpdateLightSpriteTexture() override;
diff --git a/Engine/Source/Runtime/Engine/Classes/Engine/Level.h b/Engine/Source/Runtime/Engine/Classes/Engine/Level.h
--- a/Engine/Source/Runtime/Engine/Classes/Engine/Level.h
+++ b/Engine/Source/Runtime/Engine/Classes/Engine/Level.h
@@ -740,11 +740,6 @@ public:
 	 */
 	void InvalidateModelSurface();
 
-	/**
-	 * Makes sure that all light components have valid GUIDs associated.
-	 */
-	void ValidateLightGUIDs();
-
 	/**
 	 * Sorts the actor list by net relevancy and static behaviour. First all not net relevant static
 	 * actors, then all net relevant static actors and then the rest. This is done to allow the dynamic
diff --git a/Engine/Source/Runtime/Engine/Private/Components/LightComponent.cpp b/Engine/Source/Runtime/Engine/Private/Components/LightComponent.cpp
--- a/Engine/Source/Runtime/Engine/Private/Components/LightComponent.cpp
+++ b/Engine/Source/Runtime/Engine/Private/Components/LightComponent.cpp
@@ -542,6 +542,12 @@ void ULightComponent::PostLoad()
 }
 
 #if WITH_EDITOR
+void ULightComponent::PreSave(const class ITargetPlatform* TargetPlatform)
+{
+	Super::PreSave(TargetPlatform);
+	ValidateLightGUIDs();
+}
+
 bool ULightComponent::CanEditChange(const UProperty* InProperty) const
 {
 	if (InProperty)
diff --git a/Engine/Source/Runtime/Engine/Private/Level.cpp b/Engine/Source/Runtime/Engine/Private/Level.cpp
--- a/Engine/Source/Runtime/Engine/Private/Level.cpp
+++ b/Engine/Source/Runtime/Engine/Private/Level.cpp
@@ -515,21 +515,6 @@ void ULevel::SortActorList()
 }
 
 
-void ULevel::ValidateLightGUIDs()
-{
-	for( TObjectIterator<ULightComponent> It; It; ++It )
-	{
-		ULightComponent*	LightComponent	= *It;
-		bool				IsInLevel		= LightComponent->IsIn( this );
-
-		if( IsInLevel )
-		{
-			LightComponent->ValidateLightGUIDs();
-		}
-	}
-}
-
-
 void ULevel::PreSave(const class ITargetPlatform* TargetPlatform)
 {
 	Super::PreSave(TargetPlatform);
@@ -537,10 +522,6 @@ void ULevel::PreSave(const class ITargetPlatform* TargetPlatform)
 #if WITH_EDITOR
 	if( !IsTemplate() )
 	{
-		UPackage* Package = GetOutermost();
-
-		ValidateLightGUIDs();
-
 		// Clear out any crosslevel references
 		for( int32 ActorIdx = 0; ActorIdx < Actors.Num(); ActorIdx++ )
 		{
@@ -550,8 +531,6 @@ void ULevel::PreSave(const class ITargetPlatform* TargetPlatform)
 				Actor->ClearCrossLevelReferences();
 			}
 		}
-
-		// CheckTextureStreamingBuild(this);
 	}
 #endif // WITH_EDITOR
 }
