Add option to have outer bounding regions for MBP.

diff --git a/Engine/Source/Runtime/Engine/Classes/GameFramework/WorldSettings.h b/Engine/Source/Runtime/Engine/Classes/GameFramework/WorldSettings.h
--- a/Engine/Source/Runtime/Engine/Classes/GameFramework/WorldSettings.h
+++ b/Engine/Source/Runtime/Engine/Classes/GameFramework/WorldSettings.h
@@ -366,7 +366,9 @@ struct FBroadphaseSettings
 	FBroadphaseSettings()
 		: bUseMBPOnClient(false)
 		, bUseMBPOnServer(false)
+		, bUseMBPOuterBounds(false)
 		, MBPBounds(EForceInit::ForceInitToZero)
+		, MBPOuterBounds(EForceInit::ForceInitToZero)
 		, MBPNumSubdivs(2)
 	{
 
@@ -379,10 +381,18 @@ struct FBroadphaseSettings
 	UPROPERTY(EditAnywhere, Category = Broadphase)
 	bool bUseMBPOnServer;
 
+	/** Whether to have MBP grid over concentrated inner bounds with loose outer bounds */
+	UPROPERTY(EditAnywhere, Category = Broadphase)
+	bool bUseMBPOuterBounds;
+
 	/** Total bounds for MBP, must cover the game world or collisions are disabled for out of bounds actors */
 	UPROPERTY(EditAnywhere, Category = Broadphase, meta = (EditCondition = bUseMBP))
 	FBox MBPBounds;
 
+	/** Total bounds for MBP, should cover absolute maximum bounds of the game world where physics is required */
+	UPROPERTY(EditAnywhere, Category = Broadphase, meta = (EditCondition = bUseMBP))
+	FBox MBPOuterBounds;
+
 	/** Number of times to subdivide the MBP bounds, final number of regions is MBPNumSubdivs^2 */
 	UPROPERTY(EditAnywhere, Category = Broadphase, meta = (EditCondition = bUseMBP, ClampMin=1, ClampMax=16))
 	uint32 MBPNumSubdivs;
diff --git a/Engine/Source/Runtime/Engine/Private/PhysicsEngine/PhysScene_PhysX.cpp b/Engine/Source/Runtime/Engine/Private/PhysicsEngine/PhysScene_PhysX.cpp
--- a/Engine/Source/Runtime/Engine/Private/PhysicsEngine/PhysScene_PhysX.cpp
+++ b/Engine/Source/Runtime/Engine/Private/PhysicsEngine/PhysScene_PhysX.cpp
@@ -1964,7 +1964,7 @@ void FPhysScene_PhysX::InitPhysScene(const AWorldSettings* Settings)
 		}
 
 		// Must have at least one and no more than 256 regions, subdivision is num^2 so only up to 16
-		NumSubdivisions = FMath::Clamp<uint32>(NumSubdivisions, 1, 16);
+		NumSubdivisions = FMath::Clamp<uint32>(NumSubdivisions, 1, BroadphaseSettings.bUseMBPOuterBounds ? 15 : 16);
 
 		const FBox& Bounds = BroadphaseSettings.MBPBounds;
 		PxBounds3 MbpBounds(U2PVector(Bounds.Min), U2PVector(Bounds.Max));
@@ -1984,6 +1984,100 @@ void FPhysScene_PhysX::InitPhysScene(const AWorldSettings* Settings)
 
 			PScene->addBroadPhaseRegion(NewRegion);
 		}
+			
+		if (BroadphaseSettings.bUseMBPOuterBounds)
+		{
+			const FBox& OuterBounds = BroadphaseSettings.MBPOuterBounds;
+			if (!(OuterBounds.Min.X >= Bounds.Min.X || OuterBounds.Min.Y >= Bounds.Min.Y ||
+				OuterBounds.Max.X <= Bounds.Max.X || OuterBounds.Max.Y <= Bounds.Max.Y))
+			{
+				{
+					//outer 1
+					const FBox LeftBounds(FVector(Bounds.Min.X, OuterBounds.Min.Y, Bounds.Min.Z), FVector(OuterBounds.Max.X, Bounds.Min.Y, Bounds.Max.Z));
+					PxBounds3 MbpLeftBounds(U2PVector(LeftBounds.Min), U2PVector(LeftBounds.Max));
+
+					// Storage for generated regions, the generation function will create num^2 regions
+					TArray<PxBounds3> GeneratedRegionsLeft;
+					GeneratedRegionsLeft.AddZeroed(1);
+
+					// Final parameter is up axis (2 == Z for Unreal Engine)
+					PxBroadPhaseExt::createRegionsFromWorldBounds(GeneratedRegionsLeft.GetData(), MbpLeftBounds, 1, 2);
+
+					for (const PxBounds3& Region : GeneratedRegionsLeft)
+					{
+						PxBroadPhaseRegion NewRegion;
+						NewRegion.bounds = Region;
+						NewRegion.userData = nullptr; // No need to track back to an Unreal instance at the moment
+
+						PScene->addBroadPhaseRegion(NewRegion);
+					}
+				}
+				{
+					//outer 2
+					const FBox RightBounds(FVector(OuterBounds.Min.X, Bounds.Max.Y, Bounds.Min.Z), FVector(Bounds.Max.X, OuterBounds.Max.Y, Bounds.Max.Z));
+					PxBounds3 MbpRightBounds(U2PVector(RightBounds.Min), U2PVector(RightBounds.Max));
+
+					// Storage for generated regions, the generation function will create num^2 regions
+					TArray<PxBounds3> GeneratedRegionsRight;
+					GeneratedRegionsRight.AddZeroed(1);
+
+					// Final parameter is up axis (2 == Z for Unreal Engine)
+					PxBroadPhaseExt::createRegionsFromWorldBounds(GeneratedRegionsRight.GetData(), MbpRightBounds, 1, 2);
+
+					for (const PxBounds3& Region : GeneratedRegionsRight)
+					{
+						PxBroadPhaseRegion NewRegion;
+						NewRegion.bounds = Region;
+						NewRegion.userData = nullptr; // No need to track back to an Unreal instance at the moment
+
+						PScene->addBroadPhaseRegion(NewRegion);
+					}
+				}
+				{
+					//outer 3
+					const FBox TopBounds(FVector(Bounds.Max.X, Bounds.Min.Y, Bounds.Min.Z), FVector(OuterBounds.Max.X, OuterBounds.Max.Y, Bounds.Max.Z));
+					PxBounds3 MbpTopBounds(U2PVector(TopBounds.Min), U2PVector(TopBounds.Max));
+
+					// Storage for generated regions, the generation function will create num^2 regions
+					TArray<PxBounds3> GeneratedRegionsTop;
+					GeneratedRegionsTop.AddZeroed(1);
+
+					// Final parameter is up axis (2 == Z for Unreal Engine)
+					PxBroadPhaseExt::createRegionsFromWorldBounds(GeneratedRegionsTop.GetData(), MbpTopBounds, 1, 2);
+
+					for (const PxBounds3& Region : GeneratedRegionsTop)
+					{
+						PxBroadPhaseRegion NewRegion;
+						NewRegion.bounds = Region;
+						NewRegion.userData = nullptr; // No need to track back to an Unreal instance at the moment
+
+						PScene->addBroadPhaseRegion(NewRegion);
+					}
+				}
+				{
+					//outer 4
+					const FBox BottomBounds(FVector(OuterBounds.Min.X, OuterBounds.Min.Y, Bounds.Min.Z), FVector(Bounds.Min.X, Bounds.Max.Y, Bounds.Max.Z));
+					PxBounds3 MbpBottomBounds(U2PVector(BottomBounds.Min), U2PVector(BottomBounds.Max));
+
+					// Storage for generated regions, the generation function will create num^2 regions
+					TArray<PxBounds3> GeneratedRegionsBottom;
+					GeneratedRegionsBottom.AddZeroed(1);
+
+					// Final parameter is up axis (2 == Z for Unreal Engine)
+					PxBroadPhaseExt::createRegionsFromWorldBounds(GeneratedRegionsBottom.GetData(), MbpBottomBounds, 1, 2);
+
+					for (const PxBounds3& Region : GeneratedRegionsBottom)
+					{
+						PxBroadPhaseRegion NewRegion;
+						NewRegion.bounds = Region;
+						NewRegion.userData = nullptr; // No need to track back to an Unreal instance at the moment
+
+						PScene->addBroadPhaseRegion(NewRegion);
+					}
+				}
+			}
+		}
+		
 	}
 
 #if WITH_APEX
