Fix UMaterialExpression::PostInitProperties() to not mark the package as dirty on load (the problem was introduced in 2293642)
Don't mark the package as dirty when doing fixup from code in PostInitProperties.
Make UpdateMaterialExpressionGuid calls consistent with the UpdateParameterGuid calls.

Fixes non-determinism in cooked asset registries since the runtime asset registry skips update on load for dirty packages.
Fixes issue with ResavePackages -OnlySaveDirtyPackages that will save the same package over and over again.

diff --git a/Engine/Source/Runtime/Engine/Private/Materials/MaterialExpressions.cpp b/Engine/Source/Runtime/Engine/Private/Materials/MaterialExpressions.cpp
--- a/Engine/Source/Runtime/Engine/Private/Materials/MaterialExpressions.cpp
+++ b/Engine/Source/Runtime/Engine/Private/Materials/MaterialExpressions.cpp
@@ -712,7 +712,7 @@ void UMaterialExpression::PostInitProperties()
 
 	UpdateParameterGuid(false, false);
 	
-	UpdateMaterialExpressionGuid(false, true);
+	UpdateMaterialExpressionGuid(false, false);
 }
 
 void UMaterialExpression::PostLoad()
