Force automation process to not use the ddc when executing merge actor test for staticmesh

diff --git a/Engine/Source/Developer/MeshMergeUtilities/Private/MeshMergeUtilities.cpp b/Engine/Source/Developer/MeshMergeUtilities/Private/MeshMergeUtilities.cpp
--- a/Engine/Source/Developer/MeshMergeUtilities/Private/MeshMergeUtilities.cpp
+++ b/Engine/Source/Developer/MeshMergeUtilities/Private/MeshMergeUtilities.cpp
@@ -2584,6 +2584,13 @@ void FMeshMergeUtilities::MergeComponentsToStaticMesh(const TArray<UPrimitiveCom
 		StaticMesh->ImportVersion = EImportStaticMeshVersion::LastVersion;
 		StaticMesh->LightMapResolution = InSettings.bComputedLightMapResolution ? DataTracker.GetLightMapDimension() : InSettings.TargetLightMapResolution;
 
+#if WITH_EDITOR
+		//If we are running the automation test
+		if (GIsAutomationTesting)
+		{
+			StaticMesh->BuildCacheAutomationTestGuid = FGuid::NewGuid();
+		}
+#endif
 		StaticMesh->Build(bSilent);
 
 		if (ImposterBounds.IsValid)
diff --git a/Engine/Source/Runtime/Engine/Private/StaticMesh.cpp b/Engine/Source/Runtime/Engine/Private/StaticMesh.cpp
--- a/Engine/Source/Runtime/Engine/Private/StaticMesh.cpp
+++ b/Engine/Source/Runtime/Engine/Private/StaticMesh.cpp
@@ -2296,7 +2296,19 @@ void FStaticMeshRenderData::Cache(UStaticMesh* Owner, const FStaticMeshLODSettin
 				// TODO: Save non-inlined LODs separately
 			}
 
-			GetDerivedDataCacheRef().Put(*DerivedDataKey, DerivedData);
+			bool bSaveDDC = true;
+#if WITH_EDITOR
+			//Do not save ddc when we are forcing the regeneration of ddc in automation test
+			//No need to take more space in the ddc.
+			if (GIsAutomationTesting && Owner->BuildCacheAutomationTestGuid.IsValid())
+			{
+				bSaveDDC = false;
+			}
+#endif
+			if (bSaveDDC)
+			{
+				GetDerivedDataCacheRef().Put(*DerivedDataKey, DerivedData);
+			}
 
 			int32 T1 = FPlatformTime::Cycles();
 			UE_LOG(LogStaticMesh,Log,TEXT("Built static mesh [%.2fs] %s"),
