Resolve erroneous HLOD cluster warnings
#fix Make sure we check whether or not a cluster has valid sub-actors and whether or not they are still linked to the LODActor, before marking it as an unbuilt cluster

diff --git a/Engine/Source/Runtime/Engine/Classes/Engine/LODActor.h b/Engine/Source/Runtime/Engine/Classes/Engine/LODActor.h
--- a/Engine/Source/Runtime/Engine/Classes/Engine/LODActor.h
+++ b/Engine/Source/Runtime/Engine/Classes/Engine/LODActor.h
@@ -99,6 +99,9 @@ public:
 	const bool IsBuilt(bool bInForce = false) const;
 #endif
 
+	/** Returns whether or not this LODActor has valid SubActors and whether or not their contained Primitive Components are linked (LODParentPrimitive) to StaticMeshComponent*/
+	const bool HasValidLODChildren() const;
+
 #if WITH_EDITOR
 	/** Force this actor to appear unbuilt (zeros-out key) */
 	void ForceUnbuilt();
diff --git a/Engine/Source/Runtime/Engine/Private/LODActor.cpp b/Engine/Source/Runtime/Engine/Private/LODActor.cpp
--- a/Engine/Source/Runtime/Engine/Private/LODActor.cpp
+++ b/Engine/Source/Runtime/Engine/Private/LODActor.cpp
@@ -143,11 +143,11 @@ static void ListUnbuiltHLODActors(const TArray<FString>& Args, UWorld* World)
 	for (TActorIterator<ALODActor> HLODIt(World); HLODIt; ++HLODIt)
 	{
 		ALODActor* Actor = *HLODIt;
-		if (!Actor->IsBuilt())
+		if (!Actor->IsBuilt() && Actor->HasValidLODChildren())
 		{
 			++NumUnbuilt;
 			FString ActorPathName = Actor->GetPathName(World);
-			UE_LOG(LogHLOD, Warning, TEXT("HLOD %s is unbuilt"), *ActorPathName);
+			UE_LOG(LogHLOD, Warning, TEXT("HLOD %s is unbuilt (HLOD level %i)"), *ActorPathName, Actor->LODLevel);
 		}
 	}
 
@@ -496,9 +496,38 @@ const bool ALODActor::IsBuilt(bool bInForce/*=false*/) const
 
 	return bCachedIsBuilt;
 }
-
 #endif
 
+const bool ALODActor::HasValidLODChildren() const
+{
+	if (SubActors.Num() > 0)
+	{
+		for (const AActor* Actor : SubActors)
+		{
+			if (Actor)
+			{
+				// Retrieve contained components for all sub-actors
+				TArray<const UPrimitiveComponent*> Components;
+				Actor->GetComponents(Components);
+
+				// Try and find the parent primitive(s) and see if it matches this LODActor's static mesh component
+				for (const UPrimitiveComponent* PrimitiveComponent : Components)
+				{
+					if (const UPrimitiveComponent* ParentPrimitiveComponent = PrimitiveComponent ? PrimitiveComponent->GetLODParentPrimitive() : nullptr)
+					{
+						if (ParentPrimitiveComponent && ParentPrimitiveComponent == GetStaticMeshComponent())
+						{
+							return true;
+						}
+					}
+				}
+			}
+		}		
+	}
+
+	return false;
+}
+
 #if WITH_EDITOR
 
 void ALODActor::ForceUnbuilt()
diff --git a/Engine/Source/Runtime/Engine/Private/UnrealEngine.cpp b/Engine/Source/Runtime/Engine/Private/UnrealEngine.cpp
--- a/Engine/Source/Runtime/Engine/Private/UnrealEngine.cpp
+++ b/Engine/Source/Runtime/Engine/Private/UnrealEngine.cpp
@@ -9652,7 +9652,7 @@ float DrawMapWarnings(UWorld* World, FViewport* Viewport, FCanvas* Canvas, UCanv
 			GUnbuiltHLODCount = 0;
 			for (TActorIterator<ALODActor> HLODIt(World); HLODIt; ++HLODIt)
 			{
-				if (!HLODIt->IsBuilt())
+				if (!HLODIt->IsBuilt() && HLODIt->HasValidLODChildren())
 				{
 					++GUnbuiltHLODCount;
 				}
