Reverting 6444227

Some code paths resulted in asserts

diff --git a/Engine/Source/Runtime/Engine/Private/Components/SkinnedMeshComponent.cpp b/Engine/Source/Runtime/Engine/Private/Components/SkinnedMeshComponent.cpp
--- a/Engine/Source/Runtime/Engine/Private/Components/SkinnedMeshComponent.cpp
+++ b/Engine/Source/Runtime/Engine/Private/Components/SkinnedMeshComponent.cpp
@@ -586,9 +586,9 @@ void USkinnedMeshComponent::CreateRenderState_Concurrent()
 				const bool bMorphTargetsAllowed = CVarEnableMorphTargets.GetValueOnAnyThread(true) != 0;
 
 				// Are morph targets disabled for this LOD?
-				if (!bDisableMorphTarget && bMorphTargetsAllowed)
+				if (bDisableMorphTarget || !bMorphTargetsAllowed)
 				{
-					RefreshMorphTargets();
+					ActiveMorphTargets.Empty();
 				}
 
 				MeshObject->Update(PredictedLODLevel, this, ActiveMorphTargets, MorphTargetWeights, EPreviousBoneTransformUpdateMode::UpdatePrevious);  // send to rendering thread
