Adding default value to the shading model field on UMaterial to make sure that code-created materials don't have an invalid shading model field. The back-compat is now taken care of by checking if the shading model is MSM_FromMaterialExpression instead of the validity of the shading model field.

#rnx

diff --git a/Engine/Source/Runtime/Engine/Private/Materials/Material.cpp b/Engine/Source/Runtime/Engine/Private/Materials/Material.cpp
--- a/Engine/Source/Runtime/Engine/Private/Materials/Material.cpp
+++ b/Engine/Source/Runtime/Engine/Private/Materials/Material.cpp
@@ -830,6 +830,7 @@ UMaterial::UMaterial(const FObjectInitializer& ObjectInitializer)
 {
 	BlendMode = BLEND_Opaque;
 	ShadingModel = MSM_DefaultLit;
+	ShadingModels = FMaterialShadingModelField(ShadingModel); 
 	TranslucencyLightingMode = TLM_VolumetricNonDirectional;
 	TranslucencyDirectionalLightingIntensity = 1.0f;
 	TranslucentShadowDensityScale = 0.5f;
@@ -3892,8 +3893,8 @@ void UMaterial::PostLoad()
 		ShadingModel = MSM_DefaultLit;
 	}
 
-	// If this material hasn't been compiled and saved with support for multiple shading models yet, propagate the existing shading model
-	if (!ShadingModels.IsValid())
+	// Take care of loading materials that were not compiled when the shading model field existed
+	if (ShadingModel != MSM_FromMaterialExpression)
 	{
 		ShadingModels = FMaterialShadingModelField(ShadingModel);
 	}
