- Replaced console command for ChangeActiveGameFeature to direct function call, on IOS and Android
- Disabled console commands over the web bridge in SHIPPING

diff --git a/Engine/Source/Runtime/Engine/Private/GameEngine.cpp b/Engine/Source/Runtime/Engine/Private/GameEngine.cpp
--- a/Engine/Source/Runtime/Engine/Private/GameEngine.cpp
+++ b/Engine/Source/Runtime/Engine/Private/GameEngine.cpp
@@ -759,8 +759,25 @@
 
 		FEmbeddedDelegates::GetNativeToEmbeddedParamsDelegateForSubsystem(TEXT("engine")).AddLambda([](const FEmbeddedCallParamsHelper& Message)
 		{
+			if (Message.Command == TEXT("StartUE4Live"))
+			{
+				FName Requester = *Message.Parameters.FindRef(TEXT("requester"));
+				bool bTickOnly = Message.Parameters.FindRef(TEXT("tickonly")) == TEXT("true");
+				
+				FEmbeddedCommunication::KeepAwake(Requester, !bTickOnly);
+				Message.OnCompleteDelegate({}, TEXT(""));
+			}
+			else if (Message.Command == TEXT("StopUE4Live"))
+			{
+				FName Requester = *Message.Parameters.FindRef(TEXT("requester"));
+				
+				FEmbeddedCommunication::AllowSleep(Requester);
+				Message.OnCompleteDelegate({}, TEXT(""));
+			}
+// disable things that are too scary to expose to a shipping external app/webview
+#if !UE_BUILD_SHIPPING
 			// execute any console commands
-			if (Message.Command == TEXT("console"))
+			else if (Message.Command == TEXT("console"))
 			{
 				// executed too early, must fail
 				if (GEngine == nullptr)
@@ -792,6 +809,26 @@
 				// call the completion delegate with all text output
 				Message.OnCompleteDelegate({ { TEXT("output"), Output } }, TEXT(""));
 			}
+			else if (Message.Command == TEXT("setconfig"))
+			{
+				FString File = Message.Parameters.FindRef(TEXT("file"));
+				FString Section = Message.Parameters.FindRef(TEXT("section"));
+				FString Key = Message.Parameters.FindRef(TEXT("key"));
+				FString Value = Message.Parameters.FindRef(TEXT("value"));
+				bool bSkipSave = Message.Parameters.FindRef(TEXT("skipsave")) == TEXT("true");
+				
+				FString& ConfigFile = GetConfigFromName(File);
+				
+				GConfig->SetString(*Section, *Key, *Value, ConfigFile);
+				if (!bSkipSave)
+				{
+					GConfig->Flush(false, ConfigFile);
+				}
+				
+				// send back empty reply, nothing to report
+				Message.OnCompleteDelegate({ }, TEXT(""));
+			}
+#endif
 			else if (Message.Command == TEXT("getconfig"))
 			{
 				FString File = Message.Parameters.FindRef(TEXT("file"));
@@ -811,25 +848,6 @@
 					Message.OnCompleteDelegate({ }, FString::Printf(TEXT("Config key [%s] : %s in %s was not found"), *Section, *Key, *File));
 				}
 			}
-			else if (Message.Command == TEXT("setconfig"))
-			{
-				FString File = Message.Parameters.FindRef(TEXT("file"));
-				FString Section = Message.Parameters.FindRef(TEXT("section"));
-				FString Key = Message.Parameters.FindRef(TEXT("key"));
-				FString Value = Message.Parameters.FindRef(TEXT("value"));
-				bool bSkipSave = Message.Parameters.FindRef(TEXT("skipsave")) == TEXT("true");
-
-				FString& ConfigFile = GetConfigFromName(File);
-				
-				GConfig->SetString(*Section, *Key, *Value, ConfigFile);
-				if (!bSkipSave)
-				{
-					GConfig->Flush(false, ConfigFile);
-				}
-					
-				// send back empty reply, nothing to report
-				Message.OnCompleteDelegate({ }, TEXT(""));
-			}
 			else if (Message.Command == TEXT("cvar"))
 			{
 				FString Name = Message.Parameters.FindRef(TEXT("name"));
@@ -844,21 +862,6 @@
 					Message.OnCompleteDelegate({ }, FString::Printf(TEXT("CVar %s not found"), *Name));
 				}
 			}
-			else if (Message.Command == TEXT("StartUE4Live"))
-			{
-				FName Requester = *Message.Parameters.FindRef(TEXT("requester"));
-				bool bTickOnly = Message.Parameters.FindRef(TEXT("tickonly")) == TEXT("true");
-				
-				FEmbeddedCommunication::KeepAwake(Requester, !bTickOnly);
-				Message.OnCompleteDelegate({}, TEXT(""));
-			}
-			else if (Message.Command == TEXT("StopUE4Live"))
-			{
-				FName Requester = *Message.Parameters.FindRef(TEXT("requester"));
-				
-				FEmbeddedCommunication::AllowSleep(Requester);
-				Message.OnCompleteDelegate({}, TEXT(""));
-			}
 			else
 			{
 				Message.OnCompleteDelegate({}, TEXT("Unknown command"));
