Fixed static analyzer warning

#rnx

diff --git a/Engine/Source/Runtime/Engine/Private/StaticMesh.cpp b/Engine/Source/Runtime/Engine/Private/StaticMesh.cpp
--- a/Engine/Source/Runtime/Engine/Private/StaticMesh.cpp
+++ b/Engine/Source/Runtime/Engine/Private/StaticMesh.cpp
@@ -1304,6 +1304,8 @@ void FStaticMeshRenderData::Serialize(FArchive& Ar, UStaticMesh* Owner, bool bCo
 
 						if (bDownSampling)
 						{
+							CA_ASSUME(LOD.DistanceFieldData != nullptr);
+
 							FDistanceFieldVolumeData DownSampledDFVolumeData = *LOD.DistanceFieldData;
 							IMeshUtilities& MeshUtilities = FModuleManager::Get().LoadModuleChecked<IMeshUtilities>(TEXT("MeshUtilities"));
 
@@ -1316,7 +1318,7 @@ void FStaticMeshRenderData::Serialize(FArchive& Ar, UStaticMesh* Owner, bool bCo
 					if (!bDownSampling)
 #endif
 					{
-						if (!LOD.DistanceFieldData)
+						if (LOD.DistanceFieldData == nullptr)
 						{
 							LOD.DistanceFieldData = new FDistanceFieldVolumeData();
 						}
