Fixing Crash in USoundClass::BeginDestroy when called by UnhashUnreachableObjects.

diff --git a/Engine/Source/Runtime/Engine/Private/AudioThread.cpp b/Engine/Source/Runtime/Engine/Private/AudioThread.cpp
--- a/Engine/Source/Runtime/Engine/Private/AudioThread.cpp
+++ b/Engine/Source/Runtime/Engine/Private/AudioThread.cpp
@@ -112,6 +112,9 @@ FAudioThread::FAudioThread()
 
 	FCoreUObjectDelegates::GetPreGarbageCollectDelegate().AddRaw(this, &FAudioThread::OnPreGarbageCollect);
 	FCoreUObjectDelegates::GetPostGarbageCollect().AddRaw(this, &FAudioThread::OnPostGarbageCollect);
+
+	FCoreUObjectDelegates::PreGarbageCollectConditionalBeginDestroy.AddRaw(this, &FAudioThread::OnPreGarbageCollect);
+	FCoreUObjectDelegates::PostGarbageCollectConditionalBeginDestroy.AddRaw(this, &FAudioThread::OnPostGarbageCollect);
 }
 
 FAudioThread::~FAudioThread()
@@ -119,10 +122,15 @@ FAudioThread::~FAudioThread()
 	FCoreUObjectDelegates::GetPreGarbageCollectDelegate().RemoveAll(this);
 	FCoreUObjectDelegates::GetPostGarbageCollect().RemoveAll(this);
 
+	FCoreUObjectDelegates::PreGarbageCollectConditionalBeginDestroy.RemoveAll(this);
+	FCoreUObjectDelegates::PostGarbageCollectConditionalBeginDestroy.RemoveAll(this);
+
 	FPlatformProcess::ReturnSynchEventToPool(TaskGraphBoundSyncEvent);
 	TaskGraphBoundSyncEvent = nullptr;
 }
 
+static int32 AudioThreadSuspendCount = 0;
+
 void FAudioThread::SuspendAudioThread()
 {
 	check(FPlatformTLS::GetCurrentThreadId() == GGameThreadId);
@@ -159,12 +167,20 @@ void FAudioThread::ResumeAudioThread()
 
 void FAudioThread::OnPreGarbageCollect()
 {
-	SuspendAudioThread();
+	AudioThreadSuspendCount++;
+	if (AudioThreadSuspendCount == 1)
+	{
+		SuspendAudioThread();
+	}
 }
 
 void FAudioThread::OnPostGarbageCollect()
 {
-	ResumeAudioThread();
+	AudioThreadSuspendCount--;
+	if (AudioThreadSuspendCount == 0)
+	{
+		ResumeAudioThread();
+	}
 }
 
 bool FAudioThread::Init()
