Collapsing asset redirectors when loading the redirector map.
This should prevent issues where:
A->B
B->C
C->D
Previously A would not get redirect to D.

diff --git a/Engine/Source/Runtime/Engine/Private/AssetManager.cpp b/Engine/Source/Runtime/Engine/Private/AssetManager.cpp
--- a/Engine/Source/Runtime/Engine/Private/AssetManager.cpp
+++ b/Engine/Source/Runtime/Engine/Private/AssetManager.cpp
@@ -2065,6 +2065,30 @@ void UAssetManager::LoadRedirectorMaps()
 	{
 		AssetPathRedirects.Add(FName(*Redirect.Old), FName(*Redirect.New));
 	}
+
+	// Collapse all redirects to resolve recursive relationships.
+	for (const TPair<FName, FName>& Pair : AssetPathRedirects)
+	{
+		const FName OldPath = Pair.Key;
+		FName NewPath = Pair.Value;
+		TSet<FName> CollapsedPaths;
+		CollapsedPaths.Add(OldPath);
+		CollapsedPaths.Add(NewPath);
+		while (FName* NewPathValue = AssetPathRedirects.Find(NewPath)) // Does the NewPath exist as a key?
+		{
+			NewPath = *NewPathValue;
+			if (CollapsedPaths.Contains(NewPath))
+			{
+				UE_LOG(LogAssetManager, Error, TEXT("AssetPathRedirect cycle detected when redirecting: %s to %s"), *OldPath.ToString(), *NewPath.ToString());
+				break;
+			}
+			else
+			{
+				CollapsedPaths.Add(NewPath);
+			}
+		}
+		AssetPathRedirects[OldPath] = NewPath;
+	}
 }
 
 FPrimaryAssetId UAssetManager::GetRedirectedPrimaryAssetId(const FPrimaryAssetId& OldId) const
