Force UNetConnection to hold references to GuidCache objects while channels have bunches queued.
If garbage collection happens while bunches are queued, previously resolved GUIDs (especially for BP Classes) could be invalidated, breaking things.

diff --git a/Engine/Source/Runtime/Engine/Classes/Engine/ActorChannel.h b/Engine/Source/Runtime/Engine/Classes/Engine/ActorChannel.h
--- a/Engine/Source/Runtime/Engine/Classes/Engine/ActorChannel.h
+++ b/Engine/Source/Runtime/Engine/Classes/Engine/ActorChannel.h
@@ -98,7 +98,8 @@ public:
 	// Async networking loading support state
 	TArray< class FInBunch * >			QueuedBunches;			// Queued bunches waiting on pending guids to resolve
 	double								QueuedBunchStartTime;	// Time when since queued bunches was last empty
-	TSet< FNetworkGUID >				PendingGuidResolves;	// These guids are waiting for their resolves, we need to queue up bunches until these are resolved
+
+	TSet<FNetworkGUID> PendingGuidResolves;	// These guids are waiting for their resolves, we need to queue up bunches until these are resolved
 
 	UPROPERTY()
 	TArray< UObject* >					CreateSubObjects;		// Any sub-object we created on this channel
@@ -117,20 +118,9 @@ public:
 	/**
 	 * Default constructor
 	 */
-	UActorChannel(const FObjectInitializer& ObjectInitializer = FObjectInitializer::Get())
-		: UChannel(ObjectInitializer)
-#if !UE_BUILD_SHIPPING
-		, bBlockChannelFailure(false)
-#endif
-	{
-		PRAGMA_DISABLE_DEPRECATION_WARNINGS
-		ChType = CHTYPE_Actor;
-		PRAGMA_ENABLE_DEPRECATION_WARNINGS
-		ChName = NAME_Actor;
-		bClearRecentActorRefs = true;
-		bHoldQueuedExportBunchesAndGUIDs = false;
-		QueuedCloseReason = EChannelCloseReason::Destroyed;
-	}
+	UActorChannel(const FObjectInitializer& ObjectInitializer = FObjectInitializer::Get());
+
+	~UActorChannel();
 
 public:
 
@@ -333,4 +323,10 @@ protected:
 
 	/** Closes the actor channel but with a 'dormant' flag set so it can be reopened */
 	virtual void BecomeDormant() override;
+
+private:
+
+	// TODO: It would be nice to merge the tracking of these with PendingGuidResolves, to not duplicate memory,
+	// especially since both of these sets should be empty most of the time for most channels.
+	TSet<TSharedRef<struct FQueuedBunchObjectReference>> QueuedBunchObjectReferences;
 };
diff --git a/Engine/Source/Runtime/Engine/Classes/Engine/PackageMapClient.h b/Engine/Source/Runtime/Engine/Classes/Engine/PackageMapClient.h
--- a/Engine/Source/Runtime/Engine/Classes/Engine/PackageMapClient.h
+++ b/Engine/Source/Runtime/Engine/Classes/Engine/PackageMapClient.h
@@ -206,6 +206,23 @@ enum class EAppendNetExportFlags : uint32
 ENUM_CLASS_FLAGS(EAppendNetExportFlags);
 
 
+/** Convenience type for holding onto references to objects while we have queued bunches referring to those objects. */
+struct FQueuedBunchObjectReference
+{
+private:
+
+	friend class FNetGUIDCache;
+
+	FQueuedBunchObjectReference(const FNetworkGUID InNetGUID, UObject* InObject) :
+		NetGUID(InNetGUID),
+		Object(InObject)
+	{
+	}
+
+	FNetworkGUID NetGUID;
+	UObject* Object;
+};
+
 class ENGINE_API FNetGUIDCache
 {
 public:
@@ -241,6 +258,7 @@ public:
 	void			RegisterNetGUIDFromPath_Server( const FNetworkGUID& NetGUID, const FString& PathName, const FNetworkGUID& OuterGUID, const uint32 NetworkChecksum, const bool bNoLoad, const bool bIgnoreWhenMissing );
 	UObject *		GetObjectFromNetGUID( const FNetworkGUID& NetGUID, const bool bIgnoreMustBeMapped );
 	bool			ShouldIgnoreWhenMissing( const FNetworkGUID& NetGUID ) const;
+	FNetGuidCacheObject const * const GetCacheObject(const FNetworkGUID& NetGUID) const;
 	bool			IsGUIDRegistered( const FNetworkGUID& NetGUID ) const;
 	bool			IsGUIDLoaded( const FNetworkGUID& NetGUID ) const;
 	bool			IsGUIDBroken( const FNetworkGUID& NetGUID, const bool bMustBeRegistered ) const;
@@ -265,6 +283,9 @@ public:
 	const FNetAsyncLoadDelinquencyAnalytics& GetAsyncLoadDelinquencyAnalytics() const;
 	void ResetAsyncLoadDelinquencyAnalytics();
 
+	void CollectReferences(class FReferenceCollector& ReferenceCollector);
+	TSharedRef<FQueuedBunchObjectReference> TrackQueuedBunchObjectReference(const FNetworkGUID InNetGUID, UObject* InObject);
+
 	TMap< FNetworkGUID, FNetGuidCacheObject >		ObjectLookup;
 	TMap< TWeakObjectPtr< UObject >, FNetworkGUID >	NetGUIDLookup;
 	int32											UniqueNetIDs[2];
@@ -327,6 +348,16 @@ private:
 
 	void StartAsyncLoadingPackage(FNetGuidCacheObject& Object, const FNetworkGUID ObjectGUID, const bool bWasAlreadyAsyncLoading);
 	void ValidateAsyncLoadingPackage(FNetGuidCacheObject& Object, const FNetworkGUID ObjectGUID);
+
+	void UpdateQueuedBunchObjectReference(const FNetworkGUID NetGUID, UObject* NewObject);
+
+	/**
+	 * Set of all current Objects that we've been requested to be referenced while channels
+	 * resolve their queued bunches. This is used to prevent objects (especially async load objects,
+	 * which may have no other references) from being GC'd while a channel is waiting for more
+	 * pending guids. 
+	 */
+	TMap<FNetworkGUID, TWeakPtr<FQueuedBunchObjectReference>> QueuedBunchObjectReferences;
 };
 
 class ENGINE_API FPackageMapAckState
diff --git a/Engine/Source/Runtime/Engine/Private/DataChannel.cpp b/Engine/Source/Runtime/Engine/Private/DataChannel.cpp
--- a/Engine/Source/Runtime/Engine/Private/DataChannel.cpp
+++ b/Engine/Source/Runtime/Engine/Private/DataChannel.cpp
@@ -825,6 +825,27 @@ void UChannel::AppendMustBeMappedGuids( FOutBunch* Bunch )
 }
 PRAGMA_ENABLE_DEPRECATION_WARNINGS
 
+
+UActorChannel::UActorChannel(const FObjectInitializer& ObjectInitializer)
+	: UChannel(ObjectInitializer)
+#if !UE_BUILD_SHIPPING
+	, bBlockChannelFailure(false)
+#endif
+{
+	PRAGMA_DISABLE_DEPRECATION_WARNINGS
+	ChType = CHTYPE_Actor;
+	PRAGMA_ENABLE_DEPRECATION_WARNINGS
+	ChName = NAME_Actor;
+	bClearRecentActorRefs = true;
+	bHoldQueuedExportBunchesAndGUIDs = false;
+	QueuedCloseReason = EChannelCloseReason::Destroyed;
+}
+
+UActorChannel::~UActorChannel()
+{
+
+}
+
 void UActorChannel::AppendExportBunches( TArray<FOutBunch *>& OutExportBunches )
 {
 	if (bHoldQueuedExportBunchesAndGUIDs)
@@ -1979,6 +2000,7 @@ bool UActorChannel::CleanUp(const bool bForDestroy, EChannelCloseReason CloseRea
 
 	// We don't care about any leftover pending guids at this point
 	PendingGuidResolves.Empty();
+	QueuedBunchObjectReferences.Empty();
 
 	// Free export bunches list
 	for ( int32 i = 0; i < QueuedExportBunches.Num(); i++ )
@@ -2309,6 +2331,8 @@ bool UActorChannel::ProcessQueuedBunches()
 		{
 			PackageMapClient->SetHasQueuedBunches(ActorNetGUID, false);
 		}
+
+		QueuedBunchObjectReferences.Empty();
 	}
 
 	// Update the driver with our time spent
@@ -2324,75 +2348,91 @@ bool UActorChannel::ProcessQueuedBunches()
 
 void UActorChannel::ReceivedBunch( FInBunch & Bunch )
 {
-	check( !Closing );
+	check(!Closing);
 
-	if ( Broken || bTornOff )
+	if (Broken || bTornOff)
 	{
 		return;
 	}
 
-	if ( Connection->Driver->IsServer() )
+	TArray<TPair<FNetworkGUID, UObject*>> QueuedObjectsToTrack;
+
+	if (Connection->Driver->IsServer())
 	{
-		if ( Bunch.bHasMustBeMappedGUIDs )
+		if (Bunch.bHasMustBeMappedGUIDs)
 		{
-			UE_LOG( LogNetTraffic, Error, TEXT( "UActorChannel::ReceivedBunch: Client attempted to set bHasMustBeMappedGUIDs. Actor: %s" ), Actor != NULL ? *Actor->GetName() : TEXT( "NULL" ) );
+			UE_LOG(LogNetTraffic, Error, TEXT("UActorChannel::ReceivedBunch: Client attempted to set bHasMustBeMappedGUIDs. Actor: %s"), *GetNameSafe(Actor));
 			Bunch.SetError();
 			return;
 		}
 	}
 	else
 	{
-		if ( Bunch.bHasMustBeMappedGUIDs )
+		if (Bunch.bHasMustBeMappedGUIDs)
 		{
 			// If this bunch has any guids that must be mapped, we need to wait until they resolve before we can 
 			// process the rest of the stream on this channel
 			uint16 NumMustBeMappedGUIDs = 0;
 			Bunch << NumMustBeMappedGUIDs;
 
+			QueuedObjectsToTrack.SetNum(NumMustBeMappedGUIDs);
 			//UE_LOG( LogNetTraffic, Warning, TEXT( "Read must be mapped GUID's. NumMustBeMappedGUIDs: %i" ), NumMustBeMappedGUIDs );
 
-			UPackageMapClient * PackageMapClient = CastChecked< UPackageMapClient >( Connection->PackageMap );
+			UPackageMapClient* PackageMapClient = CastChecked<UPackageMapClient>(Connection->PackageMap);
+			FNetGUIDCache* GuidCache = Connection->Driver->GuidCache.Get();
 
-			for ( int32 i = 0; i < NumMustBeMappedGUIDs; i++ )
+			for (int32 i = 0; i < NumMustBeMappedGUIDs; i++)
 			{
 				FNetworkGUID NetGUID;
 				Bunch << NetGUID;
 
 				// If we have async package map loading disabled, we have to ignore NumMustBeMappedGUIDs
 				//	(this is due to the fact that async loading could have been enabled on the server side)
-				if ( !Connection->Driver->GuidCache->ShouldAsyncLoad() )
+				if (!Connection->Driver->GuidCache->ShouldAsyncLoad())
 				{
 					continue;
 				}
 
-				// This GUID better have been exported before we get here, which means it must be registered by now
-				if (!Connection->Driver->GuidCache->IsGUIDRegistered(NetGUID))
+				if (FNetGuidCacheObject const * const GuidCacheObject = GuidCache->GetCacheObject(NetGUID))
 				{
-					UE_LOG(LogNet, Warning, TEXT("UActorChannel::ReceivedBunch: Received a MustBeMappedGUID that is not registered. ChIndex: %i NetGUID: %s Channel: %s Bunch: %s" ), ChIndex, *NetGUID.ToString(), *Describe(), *Bunch.ToString());
-					Bunch.SetError();
-					return;
-				}
+					if (UObject* const Object = GuidCacheObject->Object.Get())
+					{
+						// Note this must be mapped guid / object pair.
+						// If we are already queuing bunches, then we'll track it below.
+						QueuedObjectsToTrack.Emplace(NetGUID, Object);
+					}
+					else
+					{
+						PendingGuidResolves.Add(NetGUID);
+
+						// Start ticking this channel so that we try to resolve the pending GUID
+						Connection->StartTickingChannel(this);
 
-				if ( !Connection->Driver->GuidCache->IsGUIDLoaded( NetGUID ) )
+						// We know we're going to be queuing bunches and will need to track this object,
+						// so don't bother throwing it in the array, and just track it immediately.
+						QueuedBunchObjectReferences.Add(GuidCache->TrackQueuedBunchObjectReference(NetGUID, nullptr));
+					}
+				}
+				else
 				{
-					PendingGuidResolves.Add( NetGUID );
-					
-					// Start ticking this channel so that we try to resolve the pending GUID
-					Connection->StartTickingChannel(this);
+					// This GUID better have been exported before we get here, which means it must be registered by now
+					UE_LOG(LogNet, Warning, TEXT("UActorChannel::ReceivedBunch: Received a MustBeMappedGUID that is not registered. ChIndex: %i NetGUID: %s Channel: %s Bunch: %s"), ChIndex, *NetGUID.ToString(), *Describe(), *Bunch.ToString());
+					Bunch.SetError();
+					return;
 				}
 			}
 		}
 
-		if ( Actor == NULL && Bunch.bOpen )
+		if (Actor == NULL && Bunch.bOpen)
 		{
 			// Take a sneak peak at the actor guid so we have a copy of it now
-			FBitReaderMark Mark( Bunch );
+			FBitReaderMark Mark(Bunch);
 
-			NET_CHECKSUM( Bunch );
+			NET_CHECKSUM(Bunch);
 
 			Bunch << ActorNetGUID;
 
-			Mark.Pop( Bunch );
+			Mark.Pop(Bunch);
 
 			// we can now map guid to channel, even if all the bunches get queued
 			if (Connection->InternalAck)
@@ -2406,39 +2446,47 @@ void UActorChannel::ReceivedBunch( FInBunch & Bunch )
 		//	2. We already have queued up bunches
 		//	3. If this actor was previously on a channel that is now still processing bunches after a close
 		//	4. The driver is requesting queuing for this GUID
-		if ( PendingGuidResolves.Num() > 0 || QueuedBunches.Num() > 0 || Connection->KeepProcessingActorChannelBunchesMap.Contains( ActorNetGUID ) ||
-			 ( Connection->Driver->ShouldQueueBunchesForActorGUID( ActorNetGUID ) ) )
+		if (PendingGuidResolves.Num() > 0 || QueuedBunches.Num() > 0 || Connection->KeepProcessingActorChannelBunchesMap.Contains(ActorNetGUID) ||
+			 (Connection->Driver->ShouldQueueBunchesForActorGUID(ActorNetGUID)))
 		{
-			if ( Connection->KeepProcessingActorChannelBunchesMap.Contains( ActorNetGUID ) )
+			if (Connection->KeepProcessingActorChannelBunchesMap.Contains(ActorNetGUID))
 			{
-				UE_LOG( LogNet, Log, TEXT( "UActorChannel::ReceivedBunch: Queuing bunch because another channel (that closed) is processing bunches for this guid still. ActorNetGUID: %s" ), *ActorNetGUID.ToString() );
+				UE_LOG(LogNet, Log, TEXT("UActorChannel::ReceivedBunch: Queuing bunch because another channel (that closed) is processing bunches for this guid still. ActorNetGUID: %s"), *ActorNetGUID.ToString());
 			}
 
-			if ( QueuedBunches.Num() == 0 )
+			if (QueuedBunches.Num() == 0)
 			{
 				// Remember when we first started queuing
 				QueuedBunchStartTime = FPlatformTime::Seconds();
 			}
 
-			QueuedBunches.Add( new FInBunch( Bunch ) );
+			QueuedBunches.Add(new FInBunch(Bunch));
 			
 			// Start ticking this channel so we can process the queued bunches when possible
 			Connection->StartTickingChannel(this);
 
 			// Register this as being queued
-			UPackageMapClient * PackageMapClient = Cast< UPackageMapClient >(Connection->PackageMap);
+			UPackageMapClient* PackageMapClient = Cast<UPackageMapClient>(Connection->PackageMap);
 
 			if (PackageMapClient)
 			{
 				PackageMapClient->SetHasQueuedBunches(ActorNetGUID, true);
 			}
 
+			if (FNetGUIDCache* LocalGuidCache = Connection->Driver->GuidCache.Get())
+			{
+				for (TPair<FNetworkGUID, UObject*>& NetGUIDObjectPair : QueuedObjectsToTrack)
+				{
+					QueuedBunchObjectReferences.Add(LocalGuidCache->TrackQueuedBunchObjectReference(NetGUIDObjectPair.Key, NetGUIDObjectPair.Value));
+				}
+			}
+
 			return;
 		}
 	}
 
 	// We can process this bunch now
-	ProcessBunch( Bunch );
+	ProcessBunch(Bunch);
 }
 
 void UActorChannel::ProcessBunch( FInBunch & Bunch )
@@ -3800,6 +3848,7 @@ void UActorChannel::AddedToChannelPool()
 	check(ReplicationMap.Num() == 0);
 	check(QueuedBunches.Num() == 0);
 	check(PendingGuidResolves.Num() == 0);
+	check(QueuedBunchObjectReferences.Num() == 0);
 	check(QueuedMustBeMappedGuidsInLastBunch.Num() == 0);
 	check(QueuedExportBunches.Num() == 0);
 
diff --git a/Engine/Source/Runtime/Engine/Private/NetworkDriver.cpp b/Engine/Source/Runtime/Engine/Private/NetworkDriver.cpp
--- a/Engine/Source/Runtime/Engine/Private/NetworkDriver.cpp
+++ b/Engine/Source/Runtime/Engine/Private/NetworkDriver.cpp
@@ -3173,6 +3173,11 @@ void UNetDriver::AddReferencedObjects(UObject* InThis, FReferenceCollector& Coll
 	{
 		Collector.AddReferencedObject(It.Value(), This);
 	}
+
+	if (This->GuidCache.IsValid())
+	{
+		This->GuidCache->CollectReferences(Collector);
+	}
 }
 
 
diff --git a/Engine/Source/Runtime/Engine/Private/PackageMapClient.cpp b/Engine/Source/Runtime/Engine/Private/PackageMapClient.cpp
--- a/Engine/Source/Runtime/Engine/Private/PackageMapClient.cpp
+++ b/Engine/Source/Runtime/Engine/Private/PackageMapClient.cpp
@@ -2778,6 +2778,8 @@ void FNetGUIDCache::AsyncPackageCallback(const FName& PackageName, UPackage * Pa
 
 			if (UObject* Object = CacheObject->Object.Get())
 			{
+				UpdateQueuedBunchObjectReference(PendingLoadRequest->NetGUID, Object);
+
 				if (UWorld* World = Object->GetWorld())
 				{
 					if (AGameStateBase* GS = World->GetGameState())
@@ -3078,6 +3080,9 @@ UObject* FNetGUIDCache::GetObjectFromNetGUID( const FNetworkGUID& NetGUID, const
 		NetGUIDLookup.Add( Object, NetGUID );
 	}
 
+	// Update our QueuedObjectReference if one exists.
+	UpdateQueuedBunchObjectReference(NetGUID, Object);
+
 	return Object;
 }
 
@@ -3141,106 +3146,56 @@ bool FNetGUIDCache::IsGUIDRegistered( const FNetworkGUID& NetGUID ) const
 	return ObjectLookup.Contains( NetGUID );
 }
 
-bool FNetGUIDCache::IsGUIDLoaded( const FNetworkGUID& NetGUID ) const
+FNetGuidCacheObject const * const FNetGUIDCache::GetCacheObject(const FNetworkGUID& NetGUID) const
 {
-	if ( !NetGUID.IsValid() )
-	{
-		return false;
-	}
-
-	if ( NetGUID.IsDefault() )
-	{
-		return false;
-	}
-
-	const FNetGuidCacheObject* CacheObjectPtr = ObjectLookup.Find( NetGUID );
-
-	if ( CacheObjectPtr == NULL )
-	{
-		return false;
-	}
-
-	return CacheObjectPtr->Object != NULL;
+	return (!NetGUID.IsValid() || NetGUID.IsDefault()) ? nullptr : ObjectLookup.Find(NetGUID);
 }
 
-bool FNetGUIDCache::IsGUIDBroken( const FNetworkGUID& NetGUID, const bool bMustBeRegistered ) const
+bool FNetGUIDCache::IsGUIDLoaded(const FNetworkGUID& NetGUID) const
 {
-	if ( !NetGUID.IsValid() )
-	{
-		return false;
-	}
-
-	if ( NetGUID.IsDefault() )
-	{
-		return false;
-	}
-
-	const FNetGuidCacheObject* CacheObjectPtr = ObjectLookup.Find( NetGUID );
-
-	if ( CacheObjectPtr == NULL )
-	{
-		return bMustBeRegistered;
-	}
-
-	return CacheObjectPtr->bIsBroken;
+	FNetGuidCacheObject const * const CacheObjectPtr = GetCacheObject(NetGUID);
+	return CacheObjectPtr && CacheObjectPtr->Object != nullptr;
 }
 
-bool FNetGUIDCache::IsGUIDNoLoad( const FNetworkGUID& NetGUID ) const
+bool FNetGUIDCache::IsGUIDBroken(const FNetworkGUID& NetGUID, const bool bMustBeRegistered) const
 {
-	if ( !NetGUID.IsValid() )
+	if (!NetGUID.IsValid())
 	{
 		return false;
 	}
 
-	if ( NetGUID.IsDefault() )
+	if (NetGUID.IsDefault())
 	{
 		return false;
 	}
 
-	const FNetGuidCacheObject* const CacheObjectPtr = ObjectLookup.Find( NetGUID );
-
-	if ( CacheObjectPtr == nullptr )
+	if (FNetGuidCacheObject const * const CacheObjectPtr = ObjectLookup.Find(NetGUID))
 	{
-		return false;
+		return CacheObjectPtr->bIsBroken;
 	}
 
-	return CacheObjectPtr->bNoLoad;
+	return bMustBeRegistered;
 }
 
-bool FNetGUIDCache::IsGUIDPending( const FNetworkGUID& NetGUID ) const
+bool FNetGUIDCache::IsGUIDNoLoad(const FNetworkGUID& NetGUID) const
 {
-	if ( !NetGUID.IsValid() )
-	{
-		return false;
-	}
-
-	if ( NetGUID.IsDefault() )
-	{
-		return false;
-	}
-
-	const FNetGuidCacheObject* const CacheObjectPtr = ObjectLookup.Find( NetGUID );
-
-	if ( CacheObjectPtr == nullptr )
-	{
-		return false;
-	}
+	FNetGuidCacheObject const * const CacheObjectPtr = GetCacheObject(NetGUID);
+	return CacheObjectPtr && CacheObjectPtr->bNoLoad;
+}
 
-	return CacheObjectPtr->bIsPending;
+bool FNetGUIDCache::IsGUIDPending(const FNetworkGUID& NetGUID) const
+{
+	FNetGuidCacheObject const * const CacheObjectPtr = GetCacheObject(NetGUID);
+	return CacheObjectPtr && CacheObjectPtr->bIsPending;
 }
 
 FNetworkGUID FNetGUIDCache::GetOuterNetGUID( const FNetworkGUID& NetGUID ) const
 {
 	FNetworkGUID OuterGUID;
 
-	if ( NetGUID.IsValid() && !NetGUID.IsDefault() )
+	if (FNetGuidCacheObject const * const CacheObjectPtr = GetCacheObject(NetGUID))
 	{
-		const FNetGuidCacheObject* const CacheObjectPtr = ObjectLookup.Find( NetGUID );
-
-		if ( CacheObjectPtr != nullptr)
-		{
-			OuterGUID = CacheObjectPtr->OuterGUID;
-		}
+		OuterGUID = CacheObjectPtr->OuterGUID;
 	}
 
 	return OuterGUID;
@@ -3425,6 +3380,18 @@ PRAGMA_ENABLE_DEPRECATION_WARNINGS
 #endif
 
 	GRANULAR_NETWORK_MEMORY_TRACKING_TRACK("DelinquentAsyncLoads", DelinquentAsyncLoads.CountBytes(Ar));
+
+	GRANULAR_NETWORK_MEMORY_TRACKING_TRACK("QueuedBunchObjectReferences",
+
+		QueuedBunchObjectReferences.CountBytes(Ar);
+		for (const auto& QueuedBunchObjectReferencePair : QueuedBunchObjectReferences)
+		{
+			if (QueuedBunchObjectReferencePair.Value.IsValid())
+			{
+				Ar.CountBytes(sizeof(FQueuedBunchObjectReference), sizeof(FQueuedBunchObjectReference));
+			}
+		}
+	);
 }
 
 void FNetFieldExport::CountBytes(FArchive& Ar) const
@@ -3484,6 +3451,66 @@ void UPackageMapClient::ResetQueuedActorDelinquencyAnalytics()
 	DelinquentQueuedActors.Reset();
 }
 
+void FNetGUIDCache::CollectReferences(class FReferenceCollector& ReferenceCollector)
+{
+	for (auto It = QueuedBunchObjectReferences.CreateIterator(); It; ++It)
+	{
+		TSharedPtr<FQueuedBunchObjectReference> SharedQueuedObjectReference = It.Value().Pin();
+		FQueuedBunchObjectReference* ObjectRefPtr = SharedQueuedObjectReference.Get();
+		if (ObjectRefPtr)
+		{
+			// Don't bother adding the reference if we don't have an object.
+			if (ObjectRefPtr->Object)
+			{
+				// AddReferencedObject will set our reference to nullptr if the object is pending kill.
+				ReferenceCollector.AddReferencedObject(ObjectRefPtr->Object, Driver);
+
+				if (!ObjectRefPtr->Object)
+				{
+					UE_LOG(LogNetPackageMap, Warning, TEXT("FNetGUIDCache::CollectReferences: QueuedBunchObjectReference was killed by GC. NetGUID=%s"), *It.Key().ToString());
+				}
+			}
+		}
+		else
+		{
+			// If our weak pointer doesn't resolve, we're clearly the last referencer.
+			// No need to keep it around.
+			It.RemoveCurrent();
+		}
+	}
+
+	QueuedBunchObjectReferences.Compact();
+}
+
+TSharedRef<FQueuedBunchObjectReference> FNetGUIDCache::TrackQueuedBunchObjectReference(const FNetworkGUID InNetGUID, UObject* InObject)
+{
+	if (TWeakPtr<FQueuedBunchObjectReference>* ExistingRef = QueuedBunchObjectReferences.Find(InNetGUID))
+	{
+		TSharedPtr<FQueuedBunchObjectReference> ExistingRefShared = ExistingRef->Pin();
+		if (ExistingRefShared.IsValid())
+		{
+			return ExistingRefShared.ToSharedRef();
+		}
+	}
+
+	// Using MakeShareable instead of MakeShared to allow private constructor.
+	TSharedRef<FQueuedBunchObjectReference> NewRef = MakeShareable<FQueuedBunchObjectReference>(new FQueuedBunchObjectReference(InNetGUID, InObject));
+	QueuedBunchObjectReferences.Add(InNetGUID, NewRef);
+	return NewRef;
+}
+
+void FNetGUIDCache::UpdateQueuedBunchObjectReference(const FNetworkGUID NetGUID, UObject* NewObject)
+{
+	if (TWeakPtr<FQueuedBunchObjectReference>* WeakObjectReference = QueuedBunchObjectReferences.Find(NetGUID))
+	{
+		TSharedPtr<FQueuedBunchObjectReference> SharedObjectReference = WeakObjectReference->Pin();
+		if (FQueuedBunchObjectReference* ObjectReference = SharedObjectReference.Get())
+		{
+			ObjectReference->Object = NewObject;
+		}
+	}
+}
+
 //------------------------------------------------------
 // Debug command to see how many times we've exported each NetGUID
 // Used for measuring inefficiencies. Some duplication is unavoidable since we cannot garuntee atomicicity across multiple channels.
