Specultative fix for recurring crash.

diff --git a/Engine/Source/Runtime/Engine/Private/PhysicsEngine/PhysAnim.cpp b/Engine/Source/Runtime/Engine/Private/PhysicsEngine/PhysAnim.cpp
--- a/Engine/Source/Runtime/Engine/Private/PhysicsEngine/PhysAnim.cpp
+++ b/Engine/Source/Runtime/Engine/Private/PhysicsEngine/PhysAnim.cpp
@@ -275,15 +275,28 @@ void USkeletalMeshComponent::PerformBlendPhysicsBones(const TArray<FBoneIndexTyp
 				}
 			}
 
+			if (!(ensure(BoneIndex < EditableComponentSpaceTransforms.Num())))
+			{
+				continue;
+			}
+
 			// Update SpaceBases entry for this bone now
 			if( BoneIndex == 0 )
 			{
+				if (!(ensure(InBoneSpaceTransforms.Num())))
+				{
+					continue;
+				}
 				EditableComponentSpaceTransforms[0] = InBoneSpaceTransforms[0];
 			}
 			else
 			{
 				if(bLocalSpaceKinematics || BodyIndex == INDEX_NONE || Bodies[BodyIndex]->IsInstanceSimulatingPhysics())
 				{
+					if (!(ensure(BoneIndex < InBoneSpaceTransforms.Num())))
+					{
+						continue;
+					}
 					const int32 ParentIndex = SkeletalMesh->RefSkeleton.GetParentIndex(BoneIndex);
 					EditableComponentSpaceTransforms[BoneIndex] = InBoneSpaceTransforms[BoneIndex] * EditableComponentSpaceTransforms[ParentIndex];
 
