Use the right type to compare AActor::Role / AActor::RemoteRole in RepLayout.

diff --git a/Engine/Source/Runtime/Engine/Private/RepLayout.cpp b/Engine/Source/Runtime/Engine/Private/RepLayout.cpp
--- a/Engine/Source/Runtime/Engine/Private/RepLayout.cpp
+++ b/Engine/Source/Runtime/Engine/Private/RepLayout.cpp
@@ -835,7 +835,7 @@ static void CompareRoleProperties(
 		const FRepLayoutCmd& RemoteRoleCmd = SharedParams.Cmds[RemoteRoleParent.CmdStart];
 		const uint16 RemoteRoleHandle = RemoteRoleCmd.RelativeHandle;
 
-		const ENetRole ObjectRemoteRole = *(const ENetRole*)(Data + RemoteRoleParent).Data;
+		const TEnumAsByte<ENetRole> ObjectRemoteRole = *(const TEnumAsByte<ENetRole>*)(Data + RemoteRoleParent).Data;
 		if (SharedParams.bForceFail || RepState->SavedRemoteRole != ObjectRemoteRole)
 		{
 			RepState->SavedRemoteRole = ObjectRemoteRole;
@@ -846,7 +846,7 @@ static void CompareRoleProperties(
 		const FRepLayoutCmd& RoleCmd = SharedParams.Cmds[RoleParent.CmdStart];
 		const uint16 RoleHandle = RoleCmd.RelativeHandle;
 
-		const ENetRole ObjectRole = *(const ENetRole*)(Data + RoleParent).Data;
+		const TEnumAsByte<ENetRole> ObjectRole = *(const TEnumAsByte<ENetRole>*)(Data + RoleParent).Data;
 		if (SharedParams.bForceFail || RepState->SavedRole != ObjectRole)
 		{
 			RepState->SavedRole = ObjectRole;
@@ -893,7 +893,7 @@ static void CompareParentProperties(
 		// In that case, just allow this to fail and perform the old logic.
 		if (UNLIKELY(RepState && ParentIndex == SharedParams.RoleIndex))
 		{
-			const ENetRole ObjectRole = *(const ENetRole*)(Data + Parent).Data;
+			const TEnumAsByte<ENetRole> ObjectRole = *(const TEnumAsByte<ENetRole>*)(Data + Parent).Data;
 			if (SharedParams.bForceFail || RepState->SavedRole != ObjectRole)
 			{
 				RepState->SavedRole = ObjectRole;
@@ -902,7 +902,7 @@ static void CompareParentProperties(
 		}
 		else if (UNLIKELY(RepState && ParentIndex == SharedParams.RemoteRoleIndex))
 		{
-			const ENetRole ObjectRemoteRole = *(const ENetRole*)(Data + Parent).Data;
+			const TEnumAsByte<ENetRole> ObjectRemoteRole = *(const TEnumAsByte<ENetRole>*)(Data + Parent).Data;
 			if (SharedParams.bForceFail || RepState->SavedRemoteRole != ObjectRemoteRole)
 			{
 				RepState->SavedRemoteRole = ObjectRemoteRole;
diff --git a/Engine/Source/Runtime/Engine/Public/Net/RepLayout.h b/Engine/Source/Runtime/Engine/Public/Net/RepLayout.h
--- a/Engine/Source/Runtime/Engine/Public/Net/RepLayout.h
+++ b/Engine/Source/Runtime/Engine/Public/Net/RepLayout.h
@@ -558,8 +558,8 @@ public:
 	// Cache off the RemoteRole and Role per connection to avoid issues with
 	// FScopedRoleDowngrade. See UE-66313 (among others).
 
-	ENetRole SavedRemoteRole = ROLE_MAX;
-	ENetRole SavedRole = ROLE_MAX;
+	TEnumAsByte<ENetRole> SavedRemoteRole = ROLE_MAX;
+	TEnumAsByte<ENetRole> SavedRole = ROLE_MAX;
 
 	/** Index in the buffer where changelist history starts (i.e., the Oldest changelist). */
 	int32 HistoryStart;
