#UE4 Fix the order of the output for the mobile parts of the ListStaticMeshes console command

diff --git a/Engine/Source/Runtime/Engine/Private/UnrealEngine.cpp b/Engine/Source/Runtime/Engine/Private/UnrealEngine.cpp
--- a/Engine/Source/Runtime/Engine/Private/UnrealEngine.cpp
+++ b/Engine/Source/Runtime/Engine/Private/UnrealEngine.cpp
@@ -5194,10 +5194,10 @@ bool UEngine::HandleListStaticMeshesCommand(const TCHAR* Cmd, FOutputDevice& Ar)
 				SortedMesh.VertexCountTotal,
 				SortedMesh.VertexCountCollision,
 				SortedMesh.ShapeCountCollision,
+				SortedMesh.UsageCount,
 				SortedMesh.ResKBIncMobile,
-				SortedMesh.MobileMinLOD,
 				SortedMesh.VertexCountTotalMobile,
-				SortedMesh.UsageCount,
+				SortedMesh.MobileMinLOD,
 				*SortedMesh.Name);
 		}
 		else
