#UE4 Better FGCObject reporting in obj refs. FGCObject names are now reported in the "obj refs" console command, provided that they pass in FGCObject::GGCObjectReferencer as the referencer when calling AddReferencedObjects. FGCObjects may now also implement GetReferencerPropertyName to describe more details about how it references an object.

diff --git a/Engine/Source/Runtime/CoreUObject/Private/Misc/GCObjectReferencer.cpp b/Engine/Source/Runtime/CoreUObject/Private/Misc/GCObjectReferencer.cpp
--- a/Engine/Source/Runtime/CoreUObject/Private/Misc/GCObjectReferencer.cpp
+++ b/Engine/Source/Runtime/CoreUObject/Private/Misc/GCObjectReferencer.cpp
@@ -61,6 +61,11 @@ bool UGCObjectReferencer::GetReferencerName(UObject* Object, FString& OutName) c
 		if (ObjectArray.Contains(Object))
 		{
 			OutName = GCReporter->GetReferencerName();
+			FString ReferencerProperty;
+			if (GCReporter->GetReferencerPropertyName(Object, ReferencerProperty))
+			{
+				OutName += TEXT(":") + ReferencerProperty;
+			}
 			return true;
 		}
 	}
diff --git a/Engine/Source/Runtime/CoreUObject/Private/Serialization/AsyncLoading.cpp b/Engine/Source/Runtime/CoreUObject/Private/Serialization/AsyncLoading.cpp
--- a/Engine/Source/Runtime/CoreUObject/Private/Serialization/AsyncLoading.cpp
+++ b/Engine/Source/Runtime/CoreUObject/Private/Serialization/AsyncLoading.cpp
@@ -5531,9 +5531,9 @@ FAsyncPackage::~FAsyncPackage()
 
 void FAsyncPackage::AddReferencedObjects(FReferenceCollector& Collector)
 {
-	Collector.AddReferencedObjects(ReferencedObjects);
-	Collector.AddReferencedObjects(DeferredFinalizeObjects);
-	Collector.AddReferencedObjects(PackageObjLoaded);
+	Collector.AddReferencedObjects(ReferencedObjects, FGCObject::GGCObjectReferencer);
+	Collector.AddReferencedObjects(DeferredFinalizeObjects, FGCObject::GGCObjectReferencer);
+	Collector.AddReferencedObjects(PackageObjLoaded, FGCObject::GGCObjectReferencer);
 }
 
 void FAsyncPackage::AddObjectReference(UObject* InObject)
diff --git a/Engine/Source/Runtime/CoreUObject/Private/UObject/ReferenceChainSearch.cpp b/Engine/Source/Runtime/CoreUObject/Private/UObject/ReferenceChainSearch.cpp
--- a/Engine/Source/Runtime/CoreUObject/Private/UObject/ReferenceChainSearch.cpp
+++ b/Engine/Source/Runtime/CoreUObject/Private/UObject/ReferenceChainSearch.cpp
@@ -5,6 +5,7 @@
 #include "UObject/UObjectIterator.h"
 #include "UObject/UnrealType.h"
 #include "UObject/FastReferenceCollector.h"
+#include "UObject/GCObject.h"
 #include "HAL/ThreadHeartBeat.h"
 
 DEFINE_LOG_CATEGORY_STATIC(LogReferenceChain, Log, All);
@@ -349,7 +350,20 @@ public:
 			else
 			{
 				RefInfo.Type = FReferenceChainSearch::EReferenceType::AddReferencedObjects;
-				if (ReferencingObject)
+				if (ReferencingObject == FGCObject::GGCObjectReferencer)
+				{
+					FString RefName;
+					if (FGCObject::GGCObjectReferencer->GetReferencerName(Object, RefName))
+					{
+						RefInfo.Type = FReferenceChainSearch::EReferenceType::Property;
+						RefInfo.ReferencerName = FName(*RefName);
+					}
+					else
+					{
+						RefInfo.ReferencerName = ReferencingObject->GetFName();
+					}
+				}
+				else if (ReferencingObject)
 				{
 					RefInfo.ReferencerName = ReferencingObject->GetFName();
 				}
diff --git a/Engine/Source/Runtime/CoreUObject/Public/UObject/GCObject.h b/Engine/Source/Runtime/CoreUObject/Public/UObject/GCObject.h
--- a/Engine/Source/Runtime/CoreUObject/Public/UObject/GCObject.h
+++ b/Engine/Source/Runtime/CoreUObject/Public/UObject/GCObject.h
@@ -164,5 +164,13 @@ public:
 	{
 		return "Unknown FGCObject";
 	}
+
+	/**
+	 * Use this method to report how the specified object is referenced, if necessary
+	 */
+	virtual bool GetReferencerPropertyName(UObject* Object, FString& OutPropertyName) const
+	{
+		return false;
+	}
 };
 
diff --git a/Engine/Source/Runtime/Engine/Classes/Engine/StreamableManager.h b/Engine/Source/Runtime/Engine/Classes/Engine/StreamableManager.h
--- a/Engine/Source/Runtime/Engine/Classes/Engine/StreamableManager.h
+++ b/Engine/Source/Runtime/Engine/Classes/Engine/StreamableManager.h
@@ -333,6 +333,8 @@ struct ENGINE_API FStreamableManager : public FGCObject
 
 	/** Add referenced objects to stop them from GCing */
 	virtual void AddReferencedObjects(FReferenceCollector& Collector) override;
+	virtual FString GetReferencerName() const override { return TEXT("FStreamableManager"); }
+	virtual bool GetReferencerPropertyName(UObject* Object, FString& OutPropertyName) const override;
 
 	FStreamableManager();
 	~FStreamableManager();
diff --git a/Engine/Source/Runtime/Engine/Private/AssetManager.cpp b/Engine/Source/Runtime/Engine/Private/AssetManager.cpp
--- a/Engine/Source/Runtime/Engine/Private/AssetManager.cpp
+++ b/Engine/Source/Runtime/Engine/Private/AssetManager.cpp
@@ -1203,6 +1203,14 @@ TSharedPtr<FStreamableHandle> UAssetManager::ChangeBundleStateForPrimaryAssets(c
 				DebugName += TEXT(")");
 			}
 
+			if (PathsToLoad.Num() == 0)
+			{
+				// New state has no assets to load. Set the CurrentState's bundles and clear the handle
+				NameData->CurrentState.BundleNames = NewBundleState;
+				NameData->CurrentState.Handle.Reset();
+				continue;
+			}
+
 			NewHandle = LoadAssetList(PathsToLoad.Array(), FStreamableDelegate(), Priority, DebugName);
 
 			if (!NewHandle.IsValid())
diff --git a/Engine/Source/Runtime/Engine/Private/Slate/DeferredCleanupSlateBrush.cpp b/Engine/Source/Runtime/Engine/Private/Slate/DeferredCleanupSlateBrush.cpp
--- a/Engine/Source/Runtime/Engine/Private/Slate/DeferredCleanupSlateBrush.cpp
+++ b/Engine/Source/Runtime/Engine/Private/Slate/DeferredCleanupSlateBrush.cpp
@@ -43,7 +43,7 @@ TSharedRef<FDeferredCleanupSlateBrush> FDeferredCleanupSlateBrush::CreateBrush(
 
 void FDeferredCleanupSlateBrush::AddReferencedObjects(FReferenceCollector& Collector)
 {
-	InternalBrush.AddReferencedObjects(Collector);
+	InternalBrush.AddReferencedObjects(Collector, FGCObject::GGCObjectReferencer);
 }
 
 FString FDeferredCleanupSlateBrush::GetReferencerName() const
diff --git a/Engine/Source/Runtime/Engine/Private/StreamableManager.cpp b/Engine/Source/Runtime/Engine/Private/StreamableManager.cpp
--- a/Engine/Source/Runtime/Engine/Private/StreamableManager.cpp
+++ b/Engine/Source/Runtime/Engine/Private/StreamableManager.cpp
@@ -782,7 +782,7 @@ void FStreamableManager::AddReferencedObjects(FReferenceCollector& Collector)
 		FStreamable* Existing = It.Value();
 		if (Existing->Target)
 		{
-			Collector.AddReferencedObject(Existing->Target);
+			Collector.AddReferencedObject(Existing->Target, FGCObject::GGCObjectReferencer);
 		}
 	}
 
@@ -791,11 +791,44 @@ void FStreamableManager::AddReferencedObjects(FReferenceCollector& Collector)
 		FRedirectedPath& Existing = It.Value();
 		if (Existing.LoadedRedirector)
 		{
-			Collector.AddReferencedObject(Existing.LoadedRedirector);
+			Collector.AddReferencedObject(Existing.LoadedRedirector, FGCObject::GGCObjectReferencer);
 		}
 	}
 }
 
+bool FStreamableManager::GetReferencerPropertyName(UObject* Object, FString& OutPropertyName) const
+{
+	if (Object)
+	{
+		// Report the first handle
+		for (TStreamableMap::TConstIterator It(StreamableItems); It; ++It)
+		{
+			const FStreamable* Existing = It.Value();
+			check(Existing);
+			if (Existing->Target == Object)
+			{
+				for (const TWeakPtr<FStreamableHandle>& WeakHandle : Existing->ActiveHandles)
+				{
+					TSharedPtr<FStreamableHandle> Handle = WeakHandle.Pin();
+					if (Handle.IsValid())
+					{
+						OutPropertyName = Handle->GetDebugName();
+						return true;
+					}
+				}
+
+				if (Existing->LoadingHandles.Num() > 0)
+				{
+					OutPropertyName = Existing->LoadingHandles[0]->GetDebugName();
+					return true;
+				}
+			}
+		}
+	}
+
+	return false;
+}
+
 FStreamable* FStreamableManager::FindStreamable(const FSoftObjectPath& Target) const
 {
 	FStreamable* Existing = StreamableItems.FindRef(Target);
diff --git a/Engine/Source/Runtime/SlateCore/Public/Styling/SlateBrush.h b/Engine/Source/Runtime/SlateCore/Public/Styling/SlateBrush.h
--- a/Engine/Source/Runtime/SlateCore/Public/Styling/SlateBrush.h
+++ b/Engine/Source/Runtime/SlateCore/Public/Styling/SlateBrush.h
@@ -283,9 +283,9 @@ public:
 	}
 
 	/** Report any references to UObjects to the reference collector. */
-	void AddReferencedObjects(FReferenceCollector& Collector)
+	void AddReferencedObjects(FReferenceCollector& Collector, UObject* ReferencingObject = nullptr)
 	{
-		Collector.AddReferencedObject(ResourceObject);
+		Collector.AddReferencedObject(ResourceObject, ReferencingObject);
 	}
 
 	/**
diff --git a/Engine/Source/Runtime/UMG/Private/Slate/SObjectWidget.cpp b/Engine/Source/Runtime/UMG/Private/Slate/SObjectWidget.cpp
--- a/Engine/Source/Runtime/UMG/Private/Slate/SObjectWidget.cpp
+++ b/Engine/Source/Runtime/UMG/Private/Slate/SObjectWidget.cpp
@@ -86,7 +86,7 @@ FString SObjectWidget::GetReferencerName() const
 
 void SObjectWidget::AddReferencedObjects(FReferenceCollector& Collector)
 {
-	Collector.AddReferencedObject(WidgetObject);
+	Collector.AddReferencedObject(WidgetObject, FGCObject::GGCObjectReferencer);
 }
 
 void SObjectWidget::SetPadding(const TAttribute<FMargin>& InMargin)
