Static analyzer false positive fix

#rnx

diff --git a/Engine/Source/Runtime/Engine/Private/StaticMesh.cpp b/Engine/Source/Runtime/Engine/Private/StaticMesh.cpp
--- a/Engine/Source/Runtime/Engine/Private/StaticMesh.cpp
+++ b/Engine/Source/Runtime/Engine/Private/StaticMesh.cpp
@@ -1304,8 +1304,6 @@ void FStaticMeshRenderData::Serialize(FArchive& Ar, UStaticMesh* Owner, bool bCo
 
 						if (bDownSampling)
 						{
-							CA_ASSUME(LOD.DistanceFieldData != nullptr);
-
 							FDistanceFieldVolumeData DownSampledDFVolumeData = *LOD.DistanceFieldData;
 							IMeshUtilities& MeshUtilities = FModuleManager::Get().LoadModuleChecked<IMeshUtilities>(TEXT("MeshUtilities"));
 
@@ -1316,15 +1314,17 @@ void FStaticMeshRenderData::Serialize(FArchive& Ar, UStaticMesh* Owner, bool bCo
 					}
 
 					if (!bDownSampling)
-#endif
 					{
-						if (LOD.DistanceFieldData == nullptr)
-						{
-							LOD.DistanceFieldData = new FDistanceFieldVolumeData();
-						}
-
 						Ar << *(LOD.DistanceFieldData);
 					}
+#else
+					if (LOD.DistanceFieldData == nullptr)
+					{
+						LOD.DistanceFieldData = new FDistanceFieldVolumeData();
+					}
+
+					Ar << *(LOD.DistanceFieldData);
+#endif
 				}
 			}
 		}
