Make it so that asset scan paths can handle having a trailing slash.
#rnx

diff --git a/Engine/Source/Runtime/Engine/Private/AssetManagerTypes.cpp b/Engine/Source/Runtime/Engine/Private/AssetManagerTypes.cpp
--- a/Engine/Source/Runtime/Engine/Private/AssetManagerTypes.cpp
+++ b/Engine/Source/Runtime/Engine/Private/AssetManagerTypes.cpp
@@ -48,7 +48,15 @@ void FPrimaryAssetTypeInfo::FillRuntimeData(bool& bIsValid, bool& bBaseClassWasL
 	{
 		if (!PathRef.Path.IsEmpty())
 		{
-			AssetScanPaths.AddUnique(PathRef.Path);
+			if (PathRef.Path.EndsWith(TEXT("/")))
+			{
+				// Remove the trailing slash so searching with this behaves correctly
+				AssetScanPaths.AddUnique(PathRef.Path.LeftChop(1));
+			}
+			else
+			{
+				AssetScanPaths.AddUnique(PathRef.Path);
+			}
 		}
 	}
 
