Fixed enabling raytracing does not enable skin cache automatically if the ini file permissions are set to read only

diff --git a/Engine/Source/Runtime/Engine/Private/RendererSettings.cpp b/Engine/Source/Runtime/Engine/Private/RendererSettings.cpp
--- a/Engine/Source/Runtime/Engine/Private/RendererSettings.cpp
+++ b/Engine/Source/Runtime/Engine/Private/RendererSettings.cpp
@@ -8,6 +8,7 @@
 #include "Misc/MessageDialog.h"
 #include "UnrealEdMisc.h"
 #include "Misc/ConfigCacheIni.h"
+#include "HAL/PlatformFilemanager.h"
 
 /** The editor object. */
 extern UNREALED_API class UEditorEngine* GEditor;
@@ -109,6 +110,9 @@ void URendererSettings::PostEditChangeProperty(FPropertyChangedEvent& PropertyCh
 			&& bEnableRayTracing 
 			&& !bSupportSkinCacheShaders)
 		{
+			FString FullPath = FPaths::ConvertRelativePathToFull(GetDefaultConfigFilename());
+			FPlatformFileManager::Get().GetPlatformFile().SetReadOnly(*FullPath, false);
+
 			if (FMessageDialog::Open(EAppMsgType::YesNo, LOCTEXT("Skin Cache Disabled", "Ray Tracing requires enabling skin cache. Do you want to automatically enable skin cache now?")) == EAppReturnType::Yes)
 			{
 				bSupportSkinCacheShaders = 1;
