LiveCoding: Prevent Ctrl+Alt+F11 triggering fullscreen toggle.

diff --git a/Engine/Source/Runtime/Engine/Private/GameViewportClient.cpp b/Engine/Source/Runtime/Engine/Private/GameViewportClient.cpp
--- a/Engine/Source/Runtime/Engine/Private/GameViewportClient.cpp
+++ b/Engine/Source/Runtime/Engine/Private/GameViewportClient.cpp
@@ -496,7 +496,7 @@ UGameInstance* UGameViewportClient::GetGameInstance() const
 bool UGameViewportClient::TryToggleFullscreenOnInputKey(FKey Key, EInputEvent EventType)
 {
 	if ((Key == EKeys::Enter && EventType == EInputEvent::IE_Pressed && FSlateApplication::Get().GetModifierKeys().IsAltDown() && GetDefault<UInputSettings>()->bAltEnterTogglesFullscreen)
-		|| (IsRunningGame() && Key == EKeys::F11 && EventType == EInputEvent::IE_Pressed && GetDefault<UInputSettings>()->bF11TogglesFullscreen))
+		|| (IsRunningGame() && Key == EKeys::F11 && EventType == EInputEvent::IE_Pressed && GetDefault<UInputSettings>()->bF11TogglesFullscreen) && !FSlateApplication::Get().GetModifierKeys().AreModifersDown(EModifierKey::Control | EModifierKey::Alt))
 	{
 		HandleToggleFullscreenCommand();
 		return true;
