Moved the Custom Primitive Data up to the Primitive Component/Primitive Scene Proxy level to allow all primitives to use this. This is compatible with existing uses of SetCustomPrimitiveData* calls.
Removed the use of SendRenderDynamicData_Concurrent() path and introduced a direct path when setting the data instead.

#rnx

diff --git a/Engine/Source/Runtime/Engine/Classes/Components/PrimitiveComponent.h b/Engine/Source/Runtime/Engine/Classes/Components/PrimitiveComponent.h
--- a/Engine/Source/Runtime/Engine/Classes/Components/PrimitiveComponent.h
+++ b/Engine/Source/Runtime/Engine/Classes/Components/PrimitiveComponent.h
@@ -511,6 +511,9 @@ private:
 
 	FMaskFilter MoveIgnoreMask;
 
+	/** Custom data that can be read by a material through a material parameter expression. Set data using SetCustomPrimitiveData* functions */
+	UPROPERTY()
+	FCustomPrimitiveData CustomPrimitiveData;
 public:
 	/**
 	 * Determine whether a Character can step up onto this component.
@@ -710,6 +713,28 @@ public:
 	/** Get the mask filter checked when others move into us. */
 	FMaskFilter GetMaskFilterOnBodyInstance(FMaskFilter InMaskFilter) const { return BodyInstance.GetMaskFilter(); }
 
+	/** Set custom primitive data at index DataIndex. */
+	UFUNCTION(BlueprintCallable, Category="Rendering|Material")
+	void SetCustomPrimitiveDataFloat(int32 DataIndex, float Value);
+
+	/** Set custom primitive data, two floats at once, from index DataIndex to index DataIndex + 2. */
+	UFUNCTION(BlueprintCallable, Category="Rendering|Material")
+	void SetCustomPrimitiveDataVector2(int32 DataIndex, FVector2D Value);
+
+	/** Set custom primitive data, three floats at once, from index DataIndex to index DataIndex + 3. */
+	UFUNCTION(BlueprintCallable, Category="Rendering|Material")
+	void SetCustomPrimitiveDataVector3(int32 DataIndex, FVector Value);
+
+	/** Set custom primitive data, four floats at once, from index DataIndex to index DataIndex + 4. */
+	UFUNCTION(BlueprintCallable, Category="Rendering|Material")
+	void SetCustomPrimitiveDataVector4(int32 DataIndex, FVector4 Value);
+
+	/** 
+	 * Get the custom primitive data for this primitive component.
+	 * @return The payload of custom data that will be set on the primitive and accessible in the material through a material expression.
+	 */
+	const FCustomPrimitiveData& GetCustomPrimitiveData() const { return CustomPrimitiveData; }
+
 #if WITH_EDITOR
 	/** Override delegate used for checking the selection state of a component */
 	DECLARE_DELEGATE_RetVal_OneParam( bool, FSelectionOverride, const UPrimitiveComponent* );
@@ -717,6 +742,10 @@ public:
 #endif
 
 protected:
+
+	/** Insert an array of floats into the CustomPrimitiveData, starting at the given index */
+	void SetCustomPrimitiveDataInternal(int32 DataIndex, const TArray<float>& Values);
+
 	/** Set of components that this component is currently overlapping. */
 	TArray<FOverlapInfo> OverlappingComponents;
 
diff --git a/Engine/Source/Runtime/Engine/Classes/Components/StaticMeshComponent.h b/Engine/Source/Runtime/Engine/Classes/Components/StaticMeshComponent.h
--- a/Engine/Source/Runtime/Engine/Classes/Components/StaticMeshComponent.h
+++ b/Engine/Source/Runtime/Engine/Classes/Components/StaticMeshComponent.h
@@ -195,11 +195,6 @@ private:
 	class UStaticMesh* StaticMesh;
 
 public:
-
-	/** Custom data that can be read by a material through a material expression */
-	UPROPERTY()
-	FCustomPrimitiveData CustomPrimitiveData;
-
 	/** Helper function to get the FName of the private static mesh member */
 	static const FName GetMemberNameChecked_StaticMesh() { return GET_MEMBER_NAME_CHECKED(UStaticMeshComponent, StaticMesh); }
 
@@ -387,22 +382,6 @@ public:
 	UFUNCTION(BlueprintCallable, Category="Components|StaticMesh")
 	void GetLocalBounds(FVector& Min, FVector& Max) const;
 
-	/** Set custom primitive data at index DataIndex. */
-	UFUNCTION(BlueprintCallable, Category="Components|StaticMesh")
-	void SetCustomPrimitiveDataFloat(int32 DataIndex, float Value);
-
-	/** Set custom primitive data, two floats at once, from index DataIndex to index DataIndex + 2. */
-	UFUNCTION(BlueprintCallable, Category="Components|StaticMesh")
-	void SetCustomPrimitiveDataVector2(int32 DataIndex, FVector2D Value);
-
-	/** Set custom primitive data, three floats at once, from index DataIndex to index DataIndex + 3. */
-	UFUNCTION(BlueprintCallable, Category="Components|StaticMesh")
-	void SetCustomPrimitiveDataVector3(int32 DataIndex, FVector Value);
-
-	/** Set custom primitive data, four floats at once, from index DataIndex to index DataIndex + 4. */
-	UFUNCTION(BlueprintCallable, Category="Components|StaticMesh")
-	void SetCustomPrimitiveDataVector4(int32 DataIndex, FVector4 Value);
-
 	/** 
 	 * Set forced reverse culling
 	 */
@@ -459,7 +438,6 @@ protected:
 	virtual void CreateRenderState_Concurrent() override;
 	virtual void OnCreatePhysicsState() override;
 	virtual void OnDestroyPhysicsState() override;
-	virtual void SendRenderDynamicData_Concurrent() override;
 public:
 	virtual void InvalidateLightingCacheDetailed(bool bInvalidateBuildEnqueuedLighting, bool bTranslationOnly) override;
 	virtual UObject const* AdditionalStatObject() const override;
@@ -640,9 +618,6 @@ private:
 
 	/** Update the vertex override colors */
 	void PrivateFixupOverrideColors();
-
-	/** Insert an array of floats into the CustomPrimitiveData, starting at the given index */
-	void SetCustomPrimitiveDataInternal(int32 DataIndex, const TArray<float>& Values);
 protected:
 
 	/** Whether the component type supports static lighting. */
diff --git a/Engine/Source/Runtime/Engine/Private/Components/PrimitiveComponent.cpp b/Engine/Source/Runtime/Engine/Private/Components/PrimitiveComponent.cpp
--- a/Engine/Source/Runtime/Engine/Private/Components/PrimitiveComponent.cpp
+++ b/Engine/Source/Runtime/Engine/Private/Components/PrimitiveComponent.cpp
@@ -1535,6 +1535,53 @@ UMaterialInstanceDynamic* UPrimitiveComponent::CreateDynamicMaterialInstance(int
 	return MID;
 }
 
+void UPrimitiveComponent::SetCustomPrimitiveDataInternal(int32 DataIndex, const TArray<float>& Values)
+{
+	// Can only set data on valid indices and only if there's actually any data to set
+	if (DataIndex < FCustomPrimitiveData::NumCustomPrimitiveDataFloats && Values.Num() > 0)
+	{
+		// Number of floats needed in the custom primitive data array
+		const int32 NeededFloats = FMath::Min(DataIndex + Values.Num(), FCustomPrimitiveData::NumCustomPrimitiveDataFloats);
+
+		// Number of value to copy into the custom primitive data array at index DataIndex. Capped to not overflow
+		const int32 NumValuesToSet = FMath::Min(Values.Num(), FCustomPrimitiveData::NumCustomPrimitiveDataFloats - DataIndex);
+
+		// If trying to set data on an index which doesn't exist yet, allocate up to it
+		if (NeededFloats > CustomPrimitiveData.Data.Num())
+		{
+			CustomPrimitiveData.Data.SetNumZeroed(NeededFloats);
+		}
+
+		// Only update data if it has changed
+		if (FMemory::Memcmp(&CustomPrimitiveData.Data[DataIndex], Values.GetData(), NumValuesToSet * sizeof(float)) != 0)
+		{
+			FMemory::Memcpy(&CustomPrimitiveData.Data[DataIndex], Values.GetData(), NumValuesToSet * sizeof(float));
+
+			GetWorld()->Scene->UpdateCustomPrimitiveData(this);
+		}
+	}
+}
+
+void UPrimitiveComponent::SetCustomPrimitiveDataFloat(int32 DataIndex, float Value)
+{
+	SetCustomPrimitiveDataInternal(DataIndex, {Value});
+}
+
+void UPrimitiveComponent::SetCustomPrimitiveDataVector2(int32 DataIndex, FVector2D Value)
+{
+	SetCustomPrimitiveDataInternal(DataIndex, {Value.X, Value.Y});
+}
+
+void UPrimitiveComponent::SetCustomPrimitiveDataVector3(int32 DataIndex, FVector Value)
+{
+	SetCustomPrimitiveDataInternal(DataIndex, {Value.X, Value.Y, Value.Z});
+}
+
+void UPrimitiveComponent::SetCustomPrimitiveDataVector4(int32 DataIndex, FVector4 Value)
+{
+	SetCustomPrimitiveDataInternal(DataIndex, {Value.X, Value.Y, Value.Z, Value.W});
+}
+
 UMaterialInterface* UPrimitiveComponent::GetMaterialFromCollisionFaceIndex(int32 FaceIndex, int32& SectionIndex) const
 {
 	//This function should be overriden
diff --git a/Engine/Source/Runtime/Engine/Private/Components/StaticMeshComponent.cpp b/Engine/Source/Runtime/Engine/Private/Components/StaticMeshComponent.cpp
--- a/Engine/Source/Runtime/Engine/Private/Components/StaticMeshComponent.cpp
+++ b/Engine/Source/Runtime/Engine/Private/Components/StaticMeshComponent.cpp
@@ -624,13 +624,6 @@ void UStaticMeshComponent::OnDestroyPhysicsState()
 	bNavigationRelevant = IsNavigationRelevant();
 }
 
-void UStaticMeshComponent::SendRenderDynamicData_Concurrent()
-{
-	GetWorld()->Scene->UpdateCustomPrimitiveData(this);
-
-	Super::SendRenderDynamicData_Concurrent();
-}
-
 #if WITH_EDITORONLY_DATA
 
 /** Return the total number of LOD sections in the LOD resources */
@@ -1881,51 +1874,6 @@ void UStaticMeshComponent::SetDistanceFieldSelfShadowBias(float NewValue)
 	}
 }
 
-void UStaticMeshComponent::SetCustomPrimitiveDataInternal(int32 DataIndex, const TArray<float>& Values)
-{
-	// Can only set data on valid indices and only if there's actually any data to set
-	if (DataIndex < FCustomPrimitiveData::NumCustomPrimitiveDataFloats && Values.Num() > 0)
-	{
-		// Number of floats needed in the custom primitive data array
-		const int32 NeededFloats = FMath::Min(DataIndex + Values.Num(), FCustomPrimitiveData::NumCustomPrimitiveDataFloats);
-
-		// Number of value to copy into the custom primitive data array at index DataIndex. Capped to not overflow
-		const int32 NumValuesToSet = FMath::Min(Values.Num(), FCustomPrimitiveData::NumCustomPrimitiveDataFloats - DataIndex);
-
-		// If trying to set data on an index which doesn't exist yet, allocate up to it
-		// TODO, use a sparse structure instead?
-		if (NeededFloats > CustomPrimitiveData.Data.Num())
-		{
-			CustomPrimitiveData.Data.SetNumZeroed(NeededFloats);
-		}
-	
-		FMemory::Memcpy(&CustomPrimitiveData.Data[DataIndex], Values.GetData(), NumValuesToSet * sizeof(float));
-
-		// Make sure the render dynamic data is marked dirty to ensure an update if the primitive buffer before next draw
-		MarkRenderDynamicDataDirty();
-	}
-}
-
-void UStaticMeshComponent::SetCustomPrimitiveDataFloat(int32 DataIndex, float Value)
-{
-	SetCustomPrimitiveDataInternal(DataIndex, {Value});
-}
-
-void UStaticMeshComponent::SetCustomPrimitiveDataVector2(int32 DataIndex, FVector2D Value)
-{
-	SetCustomPrimitiveDataInternal(DataIndex, {Value.X, Value.Y});
-}
-
-void UStaticMeshComponent::SetCustomPrimitiveDataVector3(int32 DataIndex, FVector Value)
-{
-	SetCustomPrimitiveDataInternal(DataIndex, {Value.X, Value.Y, Value.Z});
-}
-
-void UStaticMeshComponent::SetCustomPrimitiveDataVector4(int32 DataIndex, FVector4 Value)
-{
-	SetCustomPrimitiveDataInternal(DataIndex, {Value.X, Value.Y, Value.Z, Value.W});
-}
-
 void UStaticMeshComponent::SetReverseCulling(bool ReverseCulling)
 {
 	if (ReverseCulling != bReverseCulling)
diff --git a/Engine/Source/Runtime/Engine/Private/PrimitiveSceneProxy.cpp b/Engine/Source/Runtime/Engine/Private/PrimitiveSceneProxy.cpp
--- a/Engine/Source/Runtime/Engine/Private/PrimitiveSceneProxy.cpp
+++ b/Engine/Source/Runtime/Engine/Private/PrimitiveSceneProxy.cpp
@@ -72,7 +72,8 @@ FPrimitiveSceneProxy::FPrimitiveSceneProxy(const UPrimitiveComponent* InComponen
 ,	PropertyColor(FLinearColor::White)
 ,	
 #endif
-	TranslucencySortPriority(FMath::Clamp(InComponent->TranslucencySortPriority, SHRT_MIN, SHRT_MAX))
+	CustomPrimitiveData(InComponent->GetCustomPrimitiveData())
+,	TranslucencySortPriority(FMath::Clamp(InComponent->TranslucencySortPriority, SHRT_MIN, SHRT_MAX))
 ,	Mobility(InComponent->Mobility)
 ,	LightmapType(InComponent->LightmapType)
 ,	StatId()
diff --git a/Engine/Source/Runtime/Engine/Private/StaticMeshRender.cpp b/Engine/Source/Runtime/Engine/Private/StaticMeshRender.cpp
--- a/Engine/Source/Runtime/Engine/Private/StaticMeshRender.cpp
+++ b/Engine/Source/Runtime/Engine/Private/StaticMeshRender.cpp
@@ -108,7 +108,6 @@ static FAutoConsoleCommand GToggleForceDefaultMaterialCmd(
 /** Initialization constructor. */
 FStaticMeshSceneProxy::FStaticMeshSceneProxy(UStaticMeshComponent* InComponent, bool bForceLODsShareStaticLighting)
 	: FPrimitiveSceneProxy(InComponent, InComponent->GetStaticMesh()->GetFName())
-	, CustomPrimitiveData(InComponent->CustomPrimitiveData)
 	, RenderData(InComponent->GetStaticMesh()->RenderData.Get())
 	, OccluderData(InComponent->GetStaticMesh()->OccluderData.Get())
 	, ForcedLodModel(InComponent->ForcedLodModel)
diff --git a/Engine/Source/Runtime/Engine/Public/PrimitiveSceneProxy.h b/Engine/Source/Runtime/Engine/Public/PrimitiveSceneProxy.h
--- a/Engine/Source/Runtime/Engine/Public/PrimitiveSceneProxy.h
+++ b/Engine/Source/Runtime/Engine/Public/PrimitiveSceneProxy.h
@@ -732,7 +732,7 @@ public:
 	 * Get the custom primitive data for this scene proxy.
 	 * @return The payload of custom data that will be set on the primitive and accessible in the material through a material expression.
 	 */
-	ENGINE_API virtual const FCustomPrimitiveData* GetCustomPrimitiveData() const { return nullptr; }
+	ENGINE_API const FCustomPrimitiveData* GetCustomPrimitiveData() const { return &CustomPrimitiveData; }
 
 protected:
 
@@ -755,6 +755,9 @@ private:
 
 	friend class FScene;
 
+	/** Custom primitive data */
+	FCustomPrimitiveData CustomPrimitiveData;
+
 	/** The translucency sort priority */
 	int16 TranslucencySortPriority;
 
diff --git a/Engine/Source/Runtime/Engine/Public/SceneInterface.h b/Engine/Source/Runtime/Engine/Public/SceneInterface.h
--- a/Engine/Source/Runtime/Engine/Public/SceneInterface.h
+++ b/Engine/Source/Runtime/Engine/Public/SceneInterface.h
@@ -70,11 +70,11 @@ public:
 	/** Updates primitive attachment state. */
 	virtual void UpdatePrimitiveAttachment(UPrimitiveComponent* Primitive) = 0;
 	/** 
-	 * Updates the custom primitive data of a static mesh which has already been added to the scene. 
+	 * Updates the custom primitive data of a primitive component which has already been added to the scene. 
 	 * 
-	 * @param StaticMesh - StaticMesh component to update
+	 * @param Primitive - Primitive component to update
 	 */
-	virtual void UpdateCustomPrimitiveData(UStaticMeshComponent* StaticMesh) = 0;
+	virtual void UpdateCustomPrimitiveData(UPrimitiveComponent* Primitive) = 0;
 	/**
 	 * Updates distance field scene data (transforms, uv scale, self-shadow bias, etc.) but doesn't change allocation in the atlas
 	 */
diff --git a/Engine/Source/Runtime/Engine/Public/StaticMeshResources.h b/Engine/Source/Runtime/Engine/Public/StaticMeshResources.h
--- a/Engine/Source/Runtime/Engine/Public/StaticMeshResources.h
+++ b/Engine/Source/Runtime/Engine/Public/StaticMeshResources.h
@@ -772,7 +772,6 @@ public:
 	virtual void GetMeshDescription(int32 LODIndex, TArray<FMeshBatch>& OutMeshElements) const override;
 
 	virtual void GetDynamicMeshElements(const TArray<const FSceneView*>& Views, const FSceneViewFamily& ViewFamily, uint32 VisibilityMap, FMeshElementCollector& Collector) const override;
-	virtual const FCustomPrimitiveData* GetCustomPrimitiveData() const override { return &CustomPrimitiveData; }
 
 #if RHI_RAYTRACING
 	virtual bool IsRayTracingRelevant() const override { return true; }
@@ -795,9 +794,6 @@ public:
 	virtual int32 GetLightMapResolution() const override { return LightMapResolution; }
 #endif
 
-	/** The custom data for this Scene Proxy */
-	FCustomPrimitiveData CustomPrimitiveData;
-
 protected:
 	/** Information used by the proxy about a single LOD of the mesh. */
 	class FLODInfo : public FLightCacheInterface
diff --git a/Engine/Source/Runtime/Renderer/Private/RendererScene.cpp b/Engine/Source/Runtime/Renderer/Private/RendererScene.cpp
--- a/Engine/Source/Runtime/Renderer/Private/RendererScene.cpp
+++ b/Engine/Source/Runtime/Renderer/Private/RendererScene.cpp
@@ -1413,24 +1413,24 @@ void FScene::UpdatePrimitiveAttachment(UPrimitiveComponent* Primitive)
 	}
 }
 
-void FScene::UpdateCustomPrimitiveData(UStaticMeshComponent* StaticMesh)
+void FScene::UpdateCustomPrimitiveData(UPrimitiveComponent* Primitive)
 {
 	SCOPE_CYCLE_COUNTER(STAT_UpdateCustomPrimitiveDataGT);
 
 	// This path updates the primitive data directly in the GPUScene. 
-	if(StaticMesh->SceneProxy) 
+	if(Primitive->SceneProxy) 
 	{
 		struct FUpdateParams
 		{
 			FScene* Scene;
-			FStaticMeshSceneProxy* PrimitiveSceneProxy;
+			FPrimitiveSceneProxy* PrimitiveSceneProxy;
 			FCustomPrimitiveData CustomPrimitiveData;
 		};
 
 		FUpdateParams UpdateParams;
 		UpdateParams.Scene = this;
-		UpdateParams.PrimitiveSceneProxy = static_cast<FStaticMeshSceneProxy*>(StaticMesh->SceneProxy);
-		UpdateParams.CustomPrimitiveData = StaticMesh->CustomPrimitiveData; 
+		UpdateParams.PrimitiveSceneProxy = Primitive->SceneProxy;
+		UpdateParams.CustomPrimitiveData = Primitive->GetCustomPrimitiveData(); 
 
 		ENQUEUE_RENDER_COMMAND(UpdateCustomPrimitiveDataCommand)(
 			[UpdateParams](FRHICommandListImmediate& RHICmdList)
@@ -3451,7 +3451,7 @@ public:
 	/** Updates the transform of a primitive which has already been added to the scene. */
 	virtual void UpdatePrimitiveTransform(UPrimitiveComponent* Primitive) override {}
 	virtual void UpdatePrimitiveAttachment(UPrimitiveComponent* Primitive) override {}
-	virtual void UpdateCustomPrimitiveData(UStaticMeshComponent* Primitive) override {}
+	virtual void UpdateCustomPrimitiveData(UPrimitiveComponent* Primitive) override {}
 
 	virtual void AddLight(ULightComponent* Light) override {}
 	virtual void RemoveLight(ULightComponent* Light) override {}
diff --git a/Engine/Source/Runtime/Renderer/Private/ScenePrivate.h b/Engine/Source/Runtime/Renderer/Private/ScenePrivate.h
--- a/Engine/Source/Runtime/Renderer/Private/ScenePrivate.h
+++ b/Engine/Source/Runtime/Renderer/Private/ScenePrivate.h
@@ -2682,7 +2682,7 @@ public:
 	virtual void ReleasePrimitive(UPrimitiveComponent* Primitive) override;
 	virtual void UpdatePrimitiveTransform(UPrimitiveComponent* Primitive) override;
 	virtual void UpdatePrimitiveAttachment(UPrimitiveComponent* Primitive) override;
-	virtual void UpdateCustomPrimitiveData(UStaticMeshComponent* StaticMesh) override;
+	virtual void UpdateCustomPrimitiveData(UPrimitiveComponent* Primitive) override;
 	virtual void UpdatePrimitiveDistanceFieldSceneData_GameThread(UPrimitiveComponent* Primitive) override;
 	virtual FPrimitiveSceneInfo* GetPrimitiveSceneInfo(int32 PrimitiveIndex) override;
 	virtual bool GetPreviousLocalToWorld(const FPrimitiveSceneInfo* PrimitiveSceneInfo, FMatrix& OutPreviousLocalToWorld) const override;
