Check that there exists a world and scene before attempting to set primtive data

#rnx

diff --git a/Engine/Source/Runtime/Engine/Private/Components/PrimitiveComponent.cpp b/Engine/Source/Runtime/Engine/Private/Components/PrimitiveComponent.cpp
--- a/Engine/Source/Runtime/Engine/Private/Components/PrimitiveComponent.cpp
+++ b/Engine/Source/Runtime/Engine/Private/Components/PrimitiveComponent.cpp
@@ -1557,7 +1557,11 @@ void UPrimitiveComponent::SetCustomPrimitiveDataInternal(int32 DataIndex, const
 		{
 			FMemory::Memcpy(&CustomPrimitiveData.Data[DataIndex], Values.GetData(), NumValuesToSet * sizeof(float));
 
-			GetWorld()->Scene->UpdateCustomPrimitiveData(this);
+			UWorld* World = GetWorld();
+			if (World && World->Scene)
+			{
+				World->Scene->UpdateCustomPrimitiveData(this);
+			}
 		}
 	}
 }
