Fix for repnotifies not being called on dynamic actors spawned into sublevels that are streamed out and back in without a gc during replay playback

diff --git a/Engine/Source/Runtime/Engine/Private/DemoNetDriver.cpp b/Engine/Source/Runtime/Engine/Private/DemoNetDriver.cpp
--- a/Engine/Source/Runtime/Engine/Private/DemoNetDriver.cpp
+++ b/Engine/Source/Runtime/Engine/Private/DemoNetDriver.cpp
@@ -4802,43 +4802,6 @@ bool UDemoNetDriver::FastForwardLevels(const FGotoResult& GotoResult)
 
 	LevelsPendingFastForward.Reset();
 
-	// It's possible that the level we're streaming in may spawn Dynamic Actors.
-	// In that case, we want to make sure we track them so we can process them below.
-	// We only care about the actors if they're outered to the Level.
-	struct FDynamicActorTracker
-	{
-		FDynamicActorTracker(UWorld* InTrackWorld, TSet<ULevel*>& InCareAboutLevels, TSet<TWeakObjectPtr<AActor>>& InActorSet) :
-			TrackWorld(InTrackWorld),
-			CareAboutLevels(InCareAboutLevels),
-			ActorSet(InActorSet)
-		{
-			FOnActorSpawned::FDelegate TrackActorDelegate = FOnActorSpawned::FDelegate::CreateRaw(this, &FDynamicActorTracker::TrackActor);
-			TrackActorHandle = TrackWorld->AddOnActorSpawnedHandler(TrackActorDelegate);
-		}
-
-		~FDynamicActorTracker()
-		{
-			TrackWorld->RemoveOnActorSpawnedHandler(TrackActorHandle);
-		}
-
-	private:
-
-		void TrackActor(AActor* Actor)
-		{
-			if (Actor && CareAboutLevels.Contains(Actor->GetLevel()))
-			{
-				UE_LOG(LogDemo, Verbose, TEXT("FastForwardLevels - Sublevel spawned dynamic actor."));
-				ActorSet.Add(Actor);
-			}
-		}
-
-		UWorld* TrackWorld;
-		TSet<ULevel*> CareAboutLevels;
-		TSet<TWeakObjectPtr<AActor>>& ActorSet;
-		FDelegateHandle TrackActorHandle;
-
-	} ActorTracker(World, LocalLevels, StartupActors);
-
 	{
 		TGuardValue<bool> FastForward(bIsFastForwarding, true);
 		FScopedIgnoreOpenChannels ScopedIgnoreChannels(ServerConnection);
@@ -4894,23 +4857,25 @@ bool UDemoNetDriver::FastForwardLevels(const FGotoResult& GotoResult)
 
 			// Since we know this is an actor channel, should be safe to do a static_cast.
 			UActorChannel* const ActorChannel = static_cast<UActorChannel*>(Channel);
-			AActor* Actor = ActorChannel->GetActor();
-
-			// We only need to consider startup actors, or dynamic that were spawned and outered
-			// to one of our sublevels.
-			if (Actor == nullptr || !StartupActors.Contains(Actor))
+			if (AActor* Actor = ActorChannel->GetActor())
 			{
-				continue;
-			}
+				const bool bDynamicInLevel = !Actor->IsNetStartupActor() && LocalLevels.Contains(Actor->GetLevel());
 
-			ChannelsToUpdate.Add(ActorChannel);
-			if (const FObjectReplicator* const ActorReplicator = ActorChannel->ActorReplicator.Get())
-			{
-				FReceivingRepState* ReceivingRepState = ActorReplicator->RepState->GetReceivingRepState();
-				FRepShadowDataBuffer ShadowData(ReceivingRepState->StaticBuffer.GetData());
-				FConstRepObjectDataBuffer ActorData(Actor);
+				// We only need to consider startup actors, or dynamic that were spawned and outered
+				// to one of our sublevels.
+				if (bDynamicInLevel || StartupActors.Contains(Actor))
+				{
+					ChannelsToUpdate.Add(ActorChannel);
 
-				ActorReplicator->RepLayout->DiffProperties(&(ReceivingRepState->RepNotifies), ShadowData, ActorData, EDiffPropertiesFlags::Sync);
+					if (const FObjectReplicator* const ActorReplicator = ActorChannel->ActorReplicator.Get())
+					{
+						FReceivingRepState* ReceivingRepState = ActorReplicator->RepState->GetReceivingRepState();
+						FRepShadowDataBuffer ShadowData(ReceivingRepState->StaticBuffer.GetData());
+						FConstRepObjectDataBuffer ActorData(Actor);
+
+						ActorReplicator->RepLayout->DiffProperties(&(ReceivingRepState->RepNotifies), ShadowData, ActorData, EDiffPropertiesFlags::Sync);
+					}
+				}
 			}
 		}
 
