Duplicating CL#5672140 and 5688685 from //UE4/Dev-Framework

As a cook/asset audit optimization, only compute the editor preview for asset chunks if we're set to generate chunks
Issue originally brought up by https://udn.unrealengine.com/questions/491315/significant-time-spent-creating-asset-registry.html

diff --git a/Engine/Source/Runtime/Engine/Private/AssetManager.cpp b/Engine/Source/Runtime/Engine/Private/AssetManager.cpp
--- a/Engine/Source/Runtime/Engine/Private/AssetManager.cpp
+++ b/Engine/Source/Runtime/Engine/Private/AssetManager.cpp
@@ -23,6 +23,7 @@
 #include "Widgets/Notifications/SNotificationList.h"
 #include "Framework/Notifications/NotificationManager.h"
 #include "Commandlets/ChunkDependencyInfo.h"
+#include "Settings/ProjectPackagingSettings.h"
 #endif
 
 #define LOCTEXT_NAMESPACE "AssetManager"
@@ -3087,23 +3088,28 @@ void UAssetManager::UpdateManagementDatabase(bool bForceRefresh)
 		AssetRegistry.SetManageReferences(PrimaryAssetIdManagementMap, false, EAssetRegistryDependencyType::None);
 	}
 
-	// Update chunk package list for all chunks
-	for (FName PackageName : PackagesToUpdateChunksFor)
+	UProjectPackagingSettings* ProjectPackagingSettings = GetMutableDefault<UProjectPackagingSettings>();
+	if (ProjectPackagingSettings && ProjectPackagingSettings->bGenerateChunks)
 	{
-		ChunkList.Reset();
-		OverrideChunkList.Reset();
-		GetPackageChunkIds(PackageName, nullptr, ExistingChunkList, ChunkList, &OverrideChunkList);
-
-		if (ChunkList.Num() > 0)
+		// Update the editor preview chunk package list for all chunks, but only if we actually care about chunks
+		// bGenerateChunks is settable per platform, but should be enabled on the default platform for preview to work
+		for (FName PackageName : PackagesToUpdateChunksFor)
 		{
-			for (int32 ChunkId : ChunkList)
-			{
-				CachedChunkMap.FindOrAdd(ChunkId).AllAssets.Add(PackageName);
+			ChunkList.Reset();
+			OverrideChunkList.Reset();
+			GetPackageChunkIds(PackageName, nullptr, ExistingChunkList, ChunkList, &OverrideChunkList);
 
-				if (OverrideChunkList.Contains(ChunkId))
+			if (ChunkList.Num() > 0)
+			{
+				for (int32 ChunkId : ChunkList)
 				{
-					// This was in the override list, so add an explicit dependency
-					CachedChunkMap.FindOrAdd(ChunkId).ExplicitAssets.Add(PackageName);
+					CachedChunkMap.FindOrAdd(ChunkId).AllAssets.Add(PackageName);
+
+					if (OverrideChunkList.Contains(ChunkId))
+					{
+						// This was in the override list, so add an explicit dependency
+						CachedChunkMap.FindOrAdd(ChunkId).ExplicitAssets.Add(PackageName);
+					}
 				}
 			}
 		}
