Layered blend per bone node now has an option "Mesh Space Scale Blend" which performs scale blending in mesh space instead of local space.

Useful for fixing scaling issues arising from "Layered blend per bone" node. For example, when blending lower body locomotion with upper body aiming and feathering the lower body weight up the spine, scaling on the spine can propogate through the skeleton to the hands causing them to drift from the weapon. This removes the need to manually fixup the spine's scale after blending the poses.

BlendMeshPosesPerBoneWeights and BlendLocalPosesPerBoneWeights are now folded into a single function: BlendPosesPerBoneFilter.

#rns

diff --git a/Engine/Source/Runtime/AnimGraphRuntime/Private/AnimNodes/AnimNode_LayeredBoneBlend.cpp b/Engine/Source/Runtime/AnimGraphRuntime/Private/AnimNodes/AnimNode_LayeredBoneBlend.cpp
--- a/Engine/Source/Runtime/AnimGraphRuntime/Private/AnimNodes/AnimNode_LayeredBoneBlend.cpp
+++ b/Engine/Source/Runtime/AnimGraphRuntime/Private/AnimNodes/AnimNode_LayeredBoneBlend.cpp
@@ -265,7 +265,16 @@ void FAnimNode_LayeredBoneBlend::Evaluate_AnyThread(FPoseContext& Output)
 			}
 		}
 
-		FAnimationRuntime::BlendPosesPerBoneFilter(BasePoseContext.Pose, TargetBlendPoses, BasePoseContext.Curve, TargetBlendCurves, Output.Pose, Output.Curve, CurrentBoneBlendWeights, bMeshSpaceRotationBlend, CurveBlendOption);
+		FAnimationRuntime::EBlendPosesPerBoneFilterFlags blendFlags = FAnimationRuntime::EBlendPosesPerBoneFilterFlags::None;
+		if (bMeshSpaceRotationBlend)
+		{
+			blendFlags |= FAnimationRuntime::EBlendPosesPerBoneFilterFlags::MeshSpaceRotation;
+		}
+		if (bMeshSpaceScaleBlend)
+		{
+			blendFlags |= FAnimationRuntime::EBlendPosesPerBoneFilterFlags::MeshSpaceScale;
+		}
+		FAnimationRuntime::BlendPosesPerBoneFilter(BasePoseContext.Pose, TargetBlendPoses, BasePoseContext.Curve, TargetBlendCurves, Output.Pose, Output.Curve, CurrentBoneBlendWeights, blendFlags, CurveBlendOption);
 	}
 }
 
diff --git a/Engine/Source/Runtime/AnimGraphRuntime/Public/AnimNodes/AnimNode_LayeredBoneBlend.h b/Engine/Source/Runtime/AnimGraphRuntime/Public/AnimNodes/AnimNode_LayeredBoneBlend.h
--- a/Engine/Source/Runtime/AnimGraphRuntime/Public/AnimNodes/AnimNode_LayeredBoneBlend.h
+++ b/Engine/Source/Runtime/AnimGraphRuntime/Public/AnimNodes/AnimNode_LayeredBoneBlend.h
@@ -38,6 +38,10 @@ public:
 	/** Whether to blend bone rotations in mesh space or in local space */
 	UPROPERTY(EditAnywhere, BlueprintReadWrite, Category=Config)
 	bool bMeshSpaceRotationBlend;
+
+	/** Whether to blend bone scales in mesh space or in local space */
+	UPROPERTY(EditAnywhere, BlueprintReadWrite, Category=Config)
+	bool bMeshSpaceScaleBlend;
 	
 	/** How to blend the layers together */
 	UPROPERTY(EditAnywhere, BlueprintReadWrite, Category=Config)
@@ -82,11 +86,12 @@ protected:
 public:	
 	FAnimNode_LayeredBoneBlend()
 		: bMeshSpaceRotationBlend(false)
+		, bMeshSpaceScaleBlend(false)
 		, CurveBlendOption(ECurveBlendOption::MaxWeight)
+		, bBlendRootMotionBasedOnRootBone(true)
 		, bHasRelevantPoses(false)
 		, LODThreshold(INDEX_NONE)
 	{
-		bBlendRootMotionBasedOnRootBone = true;
 	}
 
 	// FAnimNode_Base interface
diff --git a/Engine/Source/Runtime/Engine/Private/Animation/AnimationRuntime.cpp b/Engine/Source/Runtime/Engine/Private/Animation/AnimationRuntime.cpp
--- a/Engine/Source/Runtime/Engine/Private/Animation/AnimationRuntime.cpp
+++ b/Engine/Source/Runtime/Engine/Private/Animation/AnimationRuntime.cpp
@@ -993,26 +993,86 @@ void FAnimationRuntime::ExcludeBonesWithNoParents(const TArray<int32>& BoneIndic
 	}
 }
 
-struct FBlendMeshPosesPerBoneWeightsScratchArea : public TThreadSingleton<FBlendMeshPosesPerBoneWeightsScratchArea>
+void FAnimationRuntime::UpdateDesiredBoneWeight(const TArray<FPerBoneBlendWeight>& SrcBoneBlendWeights, TArray<FPerBoneBlendWeight>& TargetBoneBlendWeights, const TArray<float>& BlendWeights)
 {
-	TCustomBoneIndexArray<FQuat, FCompactPoseBoneIndex> SourceRotations;
-	TCustomBoneIndexArray<FQuat, FCompactPoseBoneIndex> BlendRotations;
-	TCustomBoneIndexArray<FQuat, FCompactPoseBoneIndex> TargetRotations;
+	// in the future, cache this outside
+	ensure (TargetBoneBlendWeights.Num() == SrcBoneBlendWeights.Num());
+
+	FMemory::Memzero(TargetBoneBlendWeights.GetData(), TargetBoneBlendWeights.Num() * sizeof(FPerBoneBlendWeight));
+
+	for (int32 BoneIndex = 0; BoneIndex < SrcBoneBlendWeights.Num(); ++BoneIndex)
+	{
+		const int32 PoseIndex = SrcBoneBlendWeights[BoneIndex].SourceIndex;
+		check(PoseIndex < BlendWeights.Num());
+		float TargetBlendWeight = BlendWeights[PoseIndex] * SrcBoneBlendWeights[BoneIndex].BlendWeight;
+		
+		// if relevant, otherwise all initialized as zero
+		if (FAnimWeight::IsRelevant(TargetBlendWeight))
+		{
+			TargetBoneBlendWeights[BoneIndex].SourceIndex = PoseIndex;
+			TargetBoneBlendWeights[BoneIndex].BlendWeight = TargetBlendWeight;
+		}
+	}
+}
+
+struct FBlendPosesPerBoneFilterScratchArea : public TThreadSingleton<FBlendPosesPerBoneFilterScratchArea>
+{
+	using RotationArray = TCustomBoneIndexArray<FQuat, FCompactPoseBoneIndex>;
+	using ScaleArray = TCustomBoneIndexArray<FVector, FCompactPoseBoneIndex>;
+
+	RotationArray SourceRotations;
+	RotationArray BlendRotations;
+	RotationArray TargetRotations;
+
+	ScaleArray SourceScales;
+	ScaleArray BlendScales;
+	ScaleArray TargetScales;
+
 	TArray<float> MaxPoseWeights;
 	TArray<const FBlendedCurve*> SourceCurves;
 	TArray<float> SourceWeights;
 };
 
-void FAnimationRuntime::BlendMeshPosesPerBoneWeights(
-		struct FCompactPose& BasePose,
-		const TArray<struct FCompactPose>& BlendPoses,
-		struct FBlendedCurve& BaseCurve,
-		const TArray<struct FBlendedCurve>& BlendedCurves,
-		const TArray<FPerBoneBlendWeight>& BoneBlendWeights,
-		ECurveBlendOption::Type CurveBlendOption,
-		/*out*/ FCompactPose& OutPose,
-		/*out*/ struct FBlendedCurve& OutCurve)
+void FAnimationRuntime::BlendPosesPerBoneFilter(
+	struct FCompactPose& BasePose,
+	const TArray<struct FCompactPose>& BlendPoses,
+	struct FBlendedCurve& BaseCurve,
+	const TArray<struct FBlendedCurve>& BlendedCurves,
+	struct FCompactPose& OutPose,
+	struct FBlendedCurve& OutCurve,
+	TArray<FPerBoneBlendWeight>& BoneBlendWeights,
+	bool bMeshSpaceRotationBlend,
+	ECurveBlendOption::Type CurveBlendOption)
+{
+	EBlendPosesPerBoneFilterFlags blendFlags = EBlendPosesPerBoneFilterFlags::None;
+	if (bMeshSpaceRotationBlend)
+	{
+		blendFlags |= EBlendPosesPerBoneFilterFlags::MeshSpaceRotation;
+	}
+
+	return BlendPosesPerBoneFilter(BasePose, BlendPoses, BaseCurve, BlendedCurves, OutPose, OutCurve, BoneBlendWeights, blendFlags, CurveBlendOption);
+}
+
+void FAnimationRuntime::BlendPosesPerBoneFilter(
+	struct FCompactPose& BasePose, 
+	const TArray<struct FCompactPose>& BlendPoses, 
+	struct FBlendedCurve& BaseCurve, 
+	const TArray<struct FBlendedCurve>& BlendedCurves, 
+	struct FCompactPose& OutPose, 
+	struct FBlendedCurve& OutCurve, 
+	TArray<FPerBoneBlendWeight>& BoneBlendWeights, 
+	EBlendPosesPerBoneFilterFlags blendFlags,
+	ECurveBlendOption::Type CurveBlendOption)
 {
+	SCOPE_CYCLE_COUNTER(STAT_BlendPosesPerBoneFilter);
+
+	// if no blendpose, outpose = basepose
+	if (BlendPoses.Num() == 0)
+	{
+		OutPose = BasePose;
+		return;
+	}
+
 	const int32 NumBones = BasePose.GetNumBones();
 	check(BoneBlendWeights.Num() == NumBones);
 	check(OutPose.GetNumBones() == NumBones);
@@ -1029,170 +1089,271 @@ void FAnimationRuntime::BlendMeshPosesPerBoneWeights(
 		check(BlendPose.GetNumBones() == NumBones);
 	}
 
-	const FBoneContainer& BoneContainer = BasePose.GetBoneContainer();
+	FBlendPosesPerBoneFilterScratchArea& ScratchArea = FBlendPosesPerBoneFilterScratchArea::Get();
+	TArray<float>& MaxPoseWeights = ScratchArea.MaxPoseWeights;
+	MaxPoseWeights.Reset();
+	MaxPoseWeights.AddZeroed(NumPoses);
 
-	TCustomBoneIndexArray<FQuat, FCompactPoseBoneIndex>& SourceRotations = FBlendMeshPosesPerBoneWeightsScratchArea::Get().SourceRotations;
-	TCustomBoneIndexArray<FQuat, FCompactPoseBoneIndex>& BlendRotations = FBlendMeshPosesPerBoneWeightsScratchArea::Get().BlendRotations;
-	TCustomBoneIndexArray<FQuat, FCompactPoseBoneIndex>& TargetRotations = FBlendMeshPosesPerBoneWeightsScratchArea::Get().TargetRotations;
+	using RotationArray = FBlendPosesPerBoneFilterScratchArea::RotationArray;
+	using ScaleArray = FBlendPosesPerBoneFilterScratchArea::ScaleArray;
+	RotationArray& SourceRotations = ScratchArea.SourceRotations;
+	RotationArray& BlendRotations = ScratchArea.BlendRotations;
+	RotationArray& TargetRotations = ScratchArea.TargetRotations;
+	ScaleArray& SourceScales = ScratchArea.SourceScales;
+	ScaleArray& BlendScales = ScratchArea.BlendScales;
+	ScaleArray& TargetScales = ScratchArea.TargetScales;
 
-	SourceRotations.Reset();
-	SourceRotations.AddUninitialized(NumBones);
-	BlendRotations.Reset();
-	BlendRotations.AddUninitialized(NumBones);
-	TargetRotations.Reset();
-	TargetRotations.AddUninitialized(NumBones);
+	const bool bMeshSpaceRotationBlend = EnumHasAnyFlags(blendFlags, EBlendPosesPerBoneFilterFlags::MeshSpaceRotation);
+	const bool bMeshSpaceScaleBlend = EnumHasAnyFlags(blendFlags, EBlendPosesPerBoneFilterFlags::MeshSpaceScale);
 
-	TArray<float>& MaxPoseWeights = FBlendMeshPosesPerBoneWeightsScratchArea::Get().MaxPoseWeights;
-	MaxPoseWeights.Reset();
-	MaxPoseWeights.AddZeroed(NumPoses);
+	if (bMeshSpaceRotationBlend)
+	{
+		SourceRotations.Reset();
+		SourceRotations.AddUninitialized(NumBones);
+		BlendRotations.Reset();
+		BlendRotations.AddUninitialized(NumBones);
+		TargetRotations.Reset();
+		TargetRotations.AddUninitialized(NumBones);
+	}
 
-	for (const FCompactPoseBoneIndex BoneIndex : BasePose.ForEachBoneIndex())
+	if (bMeshSpaceScaleBlend)
 	{
-		const int32 PoseIndex = BoneBlendWeights[BoneIndex.GetInt()].SourceIndex;
-		const FCompactPoseBoneIndex ParentIndex = BoneContainer.GetParentBoneIndex(BoneIndex);
+		SourceScales.Reset();
+		SourceScales.AddUninitialized(NumBones);
+		BlendScales.Reset();
+		BlendScales.AddUninitialized(NumBones);
+		TargetScales.Reset();
+		TargetScales.AddUninitialized(NumBones);
+	}
 
-		FQuat SrcRotationInMesh;
-		FQuat TargetRotationInMesh;
+	// helpers for mesh space transform accumulation
+	auto AccumulateMeshSpaceRotation = [&](int32 PoseIndex, FCompactPoseBoneIndex BoneIndex, const FQuat& ParentSourceRotation, const FQuat& ParentTargetRotation)
+	{
+		SourceRotations[BoneIndex] = ParentSourceRotation * BasePose[BoneIndex].GetRotation();
+		TargetRotations[BoneIndex] = ParentTargetRotation * BlendPoses[PoseIndex][BoneIndex].GetRotation();
+	};
+	auto AccumulateMeshSpaceScale = [&](int32 PoseIndex, FCompactPoseBoneIndex BoneIndex, const FVector& ParentSourceScale, const FVector& ParentTargetScale)
+	{
+		SourceScales[BoneIndex] = ParentSourceScale * BasePose[BoneIndex].GetScale3D();
+		TargetScales[BoneIndex] = ParentTargetScale * BlendPoses[PoseIndex][BoneIndex].GetScale3D();
+	};
 
-		if (ParentIndex != INDEX_NONE)
-		{
-			SrcRotationInMesh = SourceRotations[ParentIndex] * BasePose[BoneIndex].GetRotation();
-			TargetRotationInMesh = TargetRotations[ParentIndex] * BlendPoses[PoseIndex][BoneIndex].GetRotation();
-		}
-		else
-		{
-			SrcRotationInMesh = BasePose[BoneIndex].GetRotation();
-			TargetRotationInMesh = BlendPoses[PoseIndex][BoneIndex].GetRotation();
-		}
 
-		// update mesh based rotations
-		SourceRotations[BoneIndex] = SrcRotationInMesh;
-		TargetRotations[BoneIndex] = TargetRotationInMesh;
+	// helpers for mesh space to local space transformation
+	auto ConvertMeshToLocalSpaceRotation = [&](FTransform& BlendAtom, FCompactPoseBoneIndex ParentIndex, FCompactPoseBoneIndex BoneIndex)
+	{
+		// local -> mesh -> local transformations can cause loss of precision for long bone chains, we have to normalize rotation there.
+		FQuat LocalBlendQuat = BlendRotations[ParentIndex].Inverse() * BlendRotations[BoneIndex];
+		LocalBlendQuat.Normalize();
+		BlendAtom.SetRotation(LocalBlendQuat);
+	};
+	auto ConvertMeshToLocalSpaceScale = [&](FTransform& BlendAtom, FCompactPoseBoneIndex ParentIndex, FCompactPoseBoneIndex BoneIndex)
+	{
+		FVector ParentScaleInv = FTransform::GetSafeScaleReciprocal(BlendScales[ParentIndex], SMALL_NUMBER);
+		FVector LocalBlendScale = ParentScaleInv * BlendScales[BoneIndex];
+		BlendAtom.SetScale3D(LocalBlendScale);
+	};
+
 
-		// now update outer
-		FTransform BaseAtom = BasePose[BoneIndex];
-		FTransform TargetAtom = BlendPoses[PoseIndex][BoneIndex];
-		FTransform BlendAtom;
+	// helpers for mesh space lerping
+	auto LerpMeshSpaceRotation = [&](FCompactPoseBoneIndex BoneIndex, float BlendWeight)
+	{
+		// Fast lerp produces un-normalized quaternions, so we'll re-normalize.
+		BlendRotations[BoneIndex] = FQuat::FastLerp(SourceRotations[BoneIndex], TargetRotations[BoneIndex], BlendWeight);
+		BlendRotations[BoneIndex].Normalize();
+	};
+	auto LerpMeshSpaceScale = [&](FCompactPoseBoneIndex BoneIndex, float BlendWeight)
+	{
+		BlendScales[BoneIndex] = FMath::Lerp(SourceScales[BoneIndex], TargetScales[BoneIndex], BlendWeight);
+	};
 
-		const float BlendWeight = FMath::Clamp(BoneBlendWeights[BoneIndex.GetInt()].BlendWeight, 0.f, 1.f);
-		MaxPoseWeights[PoseIndex] = FMath::Max(MaxPoseWeights[PoseIndex], BlendWeight);
+	const FBoneContainer& BoneContainer = BasePose.GetBoneContainer();
 
-		if (!FAnimWeight::IsRelevant(BlendWeight))
-		{
-			BlendAtom = BaseAtom;
-			BlendRotations[BoneIndex] = SourceRotations[BoneIndex];
-		}
-		else if (FAnimWeight::IsFullWeight(BlendWeight))
-		{
-			BlendAtom = TargetAtom;
-			BlendRotations[BoneIndex] = TargetRotations[BoneIndex];
-		}
-		else // we want blend here
+	// blend poses with both mesh space rotation and scaling (we assume uniform scale)
+	if (bMeshSpaceRotationBlend && bMeshSpaceScaleBlend)
+	{
+		for (const FCompactPoseBoneIndex BoneIndex : BasePose.ForEachBoneIndex())
 		{
-			BlendAtom = BaseAtom;
-			BlendAtom.BlendWith(TargetAtom, BlendWeight);
+			const int32 PoseIndex = BoneBlendWeights[BoneIndex.GetInt()].SourceIndex;
+			const FCompactPoseBoneIndex ParentIndex = BoneContainer.GetParentBoneIndex(BoneIndex);
 
-			// blend rotation in mesh space
-			BlendRotations[BoneIndex] = FQuat::FastLerp(SourceRotations[BoneIndex], TargetRotations[BoneIndex], BlendWeight);
+			if (ParentIndex != INDEX_NONE)
+			{
+				AccumulateMeshSpaceRotation(PoseIndex, BoneIndex, SourceRotations[ParentIndex], TargetRotations[ParentIndex]);
+				AccumulateMeshSpaceScale(PoseIndex, BoneIndex, SourceScales[ParentIndex], TargetScales[ParentIndex]);
+			}
+			else
+			{
+				AccumulateMeshSpaceRotation(PoseIndex, BoneIndex, FQuat::Identity, FQuat::Identity);
+				AccumulateMeshSpaceScale(PoseIndex, BoneIndex, FVector(1.0f), FVector(1.0f));
+			}
 
-			// Fast lerp produces un-normalized quaternions, re-normalize.
-			BlendRotations[BoneIndex].Normalize();
-		}
+			FTransform BaseAtom = BasePose[BoneIndex];
+			FTransform TargetAtom = BlendPoses[PoseIndex][BoneIndex];
+			FTransform BlendAtom;
 
-		OutPose[BoneIndex] = BlendAtom;
-		if (ParentIndex != INDEX_NONE)
-		{
-			FQuat LocalBlendQuat = BlendRotations[ParentIndex].Inverse() * BlendRotations[BoneIndex];
+			const float BlendWeight = FMath::Clamp(BoneBlendWeights[BoneIndex.GetInt()].BlendWeight, 0.f, 1.f);
+			MaxPoseWeights[PoseIndex] = FMath::Max(MaxPoseWeights[PoseIndex], BlendWeight);
+
+			if (!FAnimWeight::IsRelevant(BlendWeight))
+			{
+				BlendAtom = BaseAtom;
+				BlendRotations[BoneIndex] = SourceRotations[BoneIndex];
+				BlendScales[BoneIndex] = SourceScales[BoneIndex];
+			}
+			else if (FAnimWeight::IsFullWeight(BlendWeight))
+			{
+				BlendAtom = TargetAtom;
+				BlendRotations[BoneIndex] = TargetRotations[BoneIndex];
+				BlendScales[BoneIndex] = TargetScales[BoneIndex];
+			}
+			else
+			{
+				BlendAtom = BaseAtom;
+				BlendAtom.BlendWith(TargetAtom, BlendWeight);
+				LerpMeshSpaceRotation(BoneIndex, BlendWeight);
+				LerpMeshSpaceScale(BoneIndex, BlendWeight);
+			}
+
+			if (ParentIndex != INDEX_NONE)
+			{
+				ConvertMeshToLocalSpaceRotation(BlendAtom, ParentIndex, BoneIndex);
+				ConvertMeshToLocalSpaceScale(BlendAtom, ParentIndex, BoneIndex);
+			}
 
-			// local -> mesh -> local transformations can cause loss of precision for long bone chains, we have to normalize rotation there.
-			LocalBlendQuat.Normalize();
-			OutPose[BoneIndex].SetRotation(LocalBlendQuat);
+			OutPose[BoneIndex] = BlendAtom;
 		}
 	}
 
-	// time to blend curves
-	// the way we blend curve per bone
-	// is to find out max weight per that pose, and then apply that weight to the curve
+	// blend poses with mesh space rotation and local scale
+	else if (bMeshSpaceRotationBlend)
 	{
-		TArray<const FBlendedCurve*>& SourceCurves = FBlendMeshPosesPerBoneWeightsScratchArea::Get().SourceCurves;
-		TArray<float>& SourceWeights = FBlendMeshPosesPerBoneWeightsScratchArea::Get().SourceWeights;
-
-		SourceCurves.Reset();
-		SourceCurves.SetNumUninitialized(NumPoses + 1);
-		SourceWeights.Reset();
-		SourceWeights.SetNumUninitialized(NumPoses + 1);
+		for (const FCompactPoseBoneIndex BoneIndex : BasePose.ForEachBoneIndex())
+		{
+			const int32 PoseIndex = BoneBlendWeights[BoneIndex.GetInt()].SourceIndex;
+			const FCompactPoseBoneIndex ParentIndex = BoneContainer.GetParentBoneIndex(BoneIndex);
 
-		SourceCurves[0] = &BaseCurve;
-		SourceWeights[0] = 1.f;
+			if (ParentIndex != INDEX_NONE)
+			{
+				AccumulateMeshSpaceRotation(PoseIndex, BoneIndex, SourceRotations[ParentIndex], TargetRotations[ParentIndex]);
+			}
+			else
+			{
+				AccumulateMeshSpaceRotation(PoseIndex, BoneIndex, FQuat::Identity, FQuat::Identity);
+			}
 
-		for(int32 Idx=0; Idx<NumPoses; ++Idx)
-		{
-			SourceCurves[Idx+1] = &BlendedCurves[Idx];
-			SourceWeights[Idx+1] = MaxPoseWeights[Idx];
-		}
+			FTransform BaseAtom = BasePose[BoneIndex];
+			FTransform TargetAtom = BlendPoses[PoseIndex][BoneIndex];
+			FTransform BlendAtom;
 
-		BlendCurves(SourceCurves, SourceWeights, OutCurve, CurveBlendOption);
-	}
-}
+			const float BlendWeight = FMath::Clamp(BoneBlendWeights[BoneIndex.GetInt()].BlendWeight, 0.f, 1.f);
+			MaxPoseWeights[PoseIndex] = FMath::Max(MaxPoseWeights[PoseIndex], BlendWeight);
 
-struct FBlendLocalPosesPerBoneWeightsScratchArea : public TThreadSingleton<FBlendLocalPosesPerBoneWeightsScratchArea>
-{
-	TArray<float> MaxPoseWeights;
-	TArray<const FBlendedCurve*> SourceCurves;
-	TArray<float> SourceWeights;
-};
+			if (!FAnimWeight::IsRelevant(BlendWeight))
+			{
+				BlendAtom = BaseAtom;
+				BlendRotations[BoneIndex] = SourceRotations[BoneIndex];
+			}
+			else if (FAnimWeight::IsFullWeight(BlendWeight))
+			{
+				BlendAtom = TargetAtom;
+				BlendRotations[BoneIndex] = TargetRotations[BoneIndex];
+			}
+			else
+			{
+				BlendAtom = BaseAtom;
+				BlendAtom.BlendWith(TargetAtom, BlendWeight);
+				LerpMeshSpaceRotation(BoneIndex, BlendWeight);
+			}
 
-void FAnimationRuntime::BlendLocalPosesPerBoneWeights(
-	FCompactPose& BasePose,
-	const TArray<FCompactPose>& BlendPoses,
-	struct FBlendedCurve& BaseCurve,
-	const TArray<struct FBlendedCurve>& BlendedCurves,
-	const TArray<FPerBoneBlendWeight>& BoneBlendWeights,
-	ECurveBlendOption::Type CurveBlendOption,
-	/*out*/ FCompactPose& OutPose, 
-	/*out*/ struct FBlendedCurve& OutCurve)
-{
-	const int32 NumBones = BasePose.GetNumBones();
-	check(BoneBlendWeights.Num() == NumBones);
-	check(OutPose.GetNumBones() == NumBones);
+			if (ParentIndex != INDEX_NONE)
+			{
+				ConvertMeshToLocalSpaceRotation(BlendAtom, ParentIndex, BoneIndex);
+			}
 
-	const int32 NumPoses = BlendPoses.Num();
-	for (const FPerBoneBlendWeight& PerBoneBlendWeight : BoneBlendWeights)
-	{
-		check(PerBoneBlendWeight.SourceIndex >= 0);
-		check(PerBoneBlendWeight.SourceIndex < NumPoses);
+			OutPose[BoneIndex] = BlendAtom;
+		}
 	}
 
-	for (const FCompactPose& BlendPose : BlendPoses)
+	// blend poses with mesh space scaling (we assume uniform scale) and local rotation
+	else if (bMeshSpaceScaleBlend)
 	{
-		check(BlendPose.GetNumBones() == NumBones);
-	}
+		for (const FCompactPoseBoneIndex BoneIndex : BasePose.ForEachBoneIndex())
+		{
+			const int32 PoseIndex = BoneBlendWeights[BoneIndex.GetInt()].SourceIndex;
+			const FCompactPoseBoneIndex ParentIndex = BoneContainer.GetParentBoneIndex(BoneIndex);
 
-	TArray<float>& MaxPoseWeights = FBlendLocalPosesPerBoneWeightsScratchArea::Get().MaxPoseWeights;
-	MaxPoseWeights.Reset();
-	MaxPoseWeights.AddZeroed(NumPoses);
+			if (ParentIndex != INDEX_NONE)
+			{
+				AccumulateMeshSpaceScale(PoseIndex, BoneIndex, SourceScales[ParentIndex], TargetScales[ParentIndex]);
+			}
+			else
+			{
+				AccumulateMeshSpaceScale(PoseIndex, BoneIndex, FVector(1.0f), FVector(1.0f));
+			}
 
-	for (FCompactPoseBoneIndex BoneIndex : BasePose.ForEachBoneIndex())
-	{
-		const int32 PoseIndex = BoneBlendWeights[BoneIndex.GetInt()].SourceIndex;
-		const FTransform& BaseAtom = BasePose[BoneIndex];
+			FTransform BaseAtom = BasePose[BoneIndex];
+			FTransform TargetAtom = BlendPoses[PoseIndex][BoneIndex];
+			FTransform BlendAtom;
 
-		const float BlendWeight = FMath::Clamp(BoneBlendWeights[BoneIndex.GetInt()].BlendWeight, 0.f, 1.f);
-		MaxPoseWeights[PoseIndex] = FMath::Max(MaxPoseWeights[PoseIndex], BlendWeight);
+			const float BlendWeight = FMath::Clamp(BoneBlendWeights[BoneIndex.GetInt()].BlendWeight, 0.f, 1.f);
+			MaxPoseWeights[PoseIndex] = FMath::Max(MaxPoseWeights[PoseIndex], BlendWeight);
 
-		if (!FAnimWeight::IsRelevant(BlendWeight))
-		{
-			OutPose[BoneIndex] = BaseAtom;
-		}
-		else if (FAnimWeight::IsFullWeight(BlendWeight))
-		{
-			OutPose[BoneIndex] = BlendPoses[PoseIndex][BoneIndex];
+			if (!FAnimWeight::IsRelevant(BlendWeight))
+			{
+				BlendAtom = BaseAtom;
+				BlendScales[BoneIndex] = SourceScales[BoneIndex];
+			}
+			else if (FAnimWeight::IsFullWeight(BlendWeight))
+			{
+				BlendAtom = TargetAtom;
+				BlendScales[BoneIndex] = TargetScales[BoneIndex];
+			}
+			else
+			{
+				BlendAtom = BaseAtom;
+				BlendAtom.BlendWith(TargetAtom, BlendWeight);
+				LerpMeshSpaceScale(BoneIndex, BlendWeight);
+			}
+
+			if (ParentIndex != INDEX_NONE)
+			{
+				ConvertMeshToLocalSpaceScale(BlendAtom, ParentIndex, BoneIndex);
+			}
+
+			OutPose[BoneIndex] = BlendAtom;
 		}
-		else // we want blend here
+	}
+
+	// blend poses with local rotation and scaling
+	else
+	{
+		for (const FCompactPoseBoneIndex BoneIndex : BasePose.ForEachBoneIndex())
 		{
-			FTransform BlendAtom = BaseAtom;
-			const FTransform& TargetAtom = BlendPoses[PoseIndex][BoneIndex];
-			BlendAtom.BlendWith(TargetAtom, BlendWeight);
+			const int32 PoseIndex = BoneBlendWeights[BoneIndex.GetInt()].SourceIndex;
+			const FCompactPoseBoneIndex ParentIndex = BoneContainer.GetParentBoneIndex(BoneIndex);
+
+			FTransform BaseAtom = BasePose[BoneIndex];
+			FTransform TargetAtom = BlendPoses[PoseIndex][BoneIndex];
+			FTransform BlendAtom;
+
+			const float BlendWeight = FMath::Clamp(BoneBlendWeights[BoneIndex.GetInt()].BlendWeight, 0.f, 1.f);
+			MaxPoseWeights[PoseIndex] = FMath::Max(MaxPoseWeights[PoseIndex], BlendWeight);
+
+			if (!FAnimWeight::IsRelevant(BlendWeight))
+			{
+				BlendAtom = BaseAtom;
+			}
+			else if (FAnimWeight::IsFullWeight(BlendWeight))
+			{
+				BlendAtom = TargetAtom;
+			}
+			else
+			{
+				BlendAtom = BaseAtom;
+				BlendAtom.BlendWith(TargetAtom, BlendWeight);
+			}
+
 			OutPose[BoneIndex] = BlendAtom;
 		}
 	}
@@ -1201,8 +1362,8 @@ void FAnimationRuntime::BlendLocalPosesPerBoneWeights(
 	// the way we blend curve per bone
 	// is to find out max weight per that pose, and then apply that weight to the curve
 	{
-		TArray<const FBlendedCurve*>& SourceCurves = FBlendLocalPosesPerBoneWeightsScratchArea::Get().SourceCurves;
-		TArray<float>& SourceWeights = FBlendLocalPosesPerBoneWeightsScratchArea::Get().SourceWeights;
+		TArray<const FBlendedCurve*>& SourceCurves = ScratchArea.SourceCurves;
+		TArray<float>& SourceWeights = ScratchArea.SourceWeights;
 
 		SourceCurves.Reset();
 		SourceCurves.SetNumUninitialized(NumPoses + 1);
@@ -1212,67 +1373,13 @@ void FAnimationRuntime::BlendLocalPosesPerBoneWeights(
 		SourceCurves[0] = &BaseCurve;
 		SourceWeights[0] = 1.f;
 
-		for (int32 Idx=0; Idx<NumPoses; ++Idx)
+		for(int32 Idx=0; Idx<NumPoses; ++Idx)
 		{
 			SourceCurves[Idx+1] = &BlendedCurves[Idx];
 			SourceWeights[Idx+1] = MaxPoseWeights[Idx];
 		}
-		
-		BlendCurves(SourceCurves, SourceWeights, OutCurve, CurveBlendOption);
-	}
-}
-
-void FAnimationRuntime::UpdateDesiredBoneWeight(const TArray<FPerBoneBlendWeight>& SrcBoneBlendWeights, TArray<FPerBoneBlendWeight>& TargetBoneBlendWeights, const TArray<float>& BlendWeights)
-{
-	// in the future, cache this outside
-	ensure (TargetBoneBlendWeights.Num() == SrcBoneBlendWeights.Num());
-
-	FMemory::Memzero(TargetBoneBlendWeights.GetData(), TargetBoneBlendWeights.Num() * sizeof(FPerBoneBlendWeight));
-
-	for (int32 BoneIndex = 0; BoneIndex < SrcBoneBlendWeights.Num(); ++BoneIndex)
-	{
-		const int32 PoseIndex = SrcBoneBlendWeights[BoneIndex].SourceIndex;
-		check(PoseIndex < BlendWeights.Num());
-		float TargetBlendWeight = BlendWeights[PoseIndex] * SrcBoneBlendWeights[BoneIndex].BlendWeight;
-		
-		// if relevant, otherwise all initialized as zero
-		if (FAnimWeight::IsRelevant(TargetBlendWeight))
-		{
-			TargetBoneBlendWeights[BoneIndex].SourceIndex = PoseIndex;
-			TargetBoneBlendWeights[BoneIndex].BlendWeight = TargetBlendWeight;
-		}
-	}
-}
 
-void FAnimationRuntime::BlendPosesPerBoneFilter(
-	struct FCompactPose& BasePose, 
-	const TArray<struct FCompactPose>& BlendPoses, 
-	struct FBlendedCurve& BaseCurve, 
-	const TArray<struct FBlendedCurve>& BlendedCurves, 
-	struct FCompactPose& OutPose, 
-	struct FBlendedCurve& OutCurve, 
-	TArray<FPerBoneBlendWeight>& BoneBlendWeights, 
-	bool bMeshSpaceRotationBlending, 
-	ECurveBlendOption::Type CurveBlendOption)
-{
-	SCOPE_CYCLE_COUNTER(STAT_BlendPosesPerBoneFilter);
-
-	ensure(OutPose.GetNumBones() == BasePose.GetNumBones());
-
-	if (BlendPoses.Num() != 0)
-	{
-		if (bMeshSpaceRotationBlending)
-		{
-			BlendMeshPosesPerBoneWeights(BasePose, BlendPoses, BaseCurve, BlendedCurves, BoneBlendWeights, CurveBlendOption, OutPose, OutCurve);
-		}
-		else
-		{
-			BlendLocalPosesPerBoneWeights(BasePose, BlendPoses, BaseCurve, BlendedCurves, BoneBlendWeights, CurveBlendOption, OutPose, OutCurve);
-		}
-	}
-	else // if no blendpose, outpose = basepose
-	{
-		OutPose = BasePose;
+		BlendCurves(SourceCurves, SourceWeights, OutCurve, CurveBlendOption);
 	}
 }
 
diff --git a/Engine/Source/Runtime/Engine/Public/AnimationRuntime.h b/Engine/Source/Runtime/Engine/Public/AnimationRuntime.h
--- a/Engine/Source/Runtime/Engine/Public/AnimationRuntime.h
+++ b/Engine/Source/Runtime/Engine/Public/AnimationRuntime.h
@@ -8,6 +8,7 @@
 #pragma once
 
 #include "CoreMinimal.h"
+#include "Misc/EnumClassFlags.h"
 #include "BoneIndices.h"
 #include "Animation/AnimTypes.h"
 #include "Animation/AnimationAsset.h"
@@ -225,26 +226,48 @@ public:
 		/*out*/ FCompactPose& ResultPose,
 		/*out*/ FBlendedCurve& ResultCurve);
 
+
+	/** Blending flags for BlendPosesPerBoneFilter */
+	enum class EBlendPosesPerBoneFilterFlags : uint32 {
+		None = 0,
+		MeshSpaceRotation = (1 << 0),	//! Blend bone rotations in mesh space instead of local space
+		MeshSpaceScale    = (1 << 1),	//! Blend bone scales in mesh space instead of local space
+	};
+
 	/**
 	* Blend Poses per bone weights : The BasePoses + BlendPoses(SourceIndex) * Blend Weights(BoneIndex)
 	* Please note BlendWeights are array, so you can define per bone base
 	* This supports multi per bone blending, but only one pose as blend at a time per track
 	* PerBoneBlendWeights.Num() == Atoms.Num()
 	*
+	* @note : This optionally blends rotation and/or scale in mesh space. Translation is always in local space.
+	*
 	* I had multiple debates about having PerBoneBlendWeights array(for memory reasons),
 	* but it so much simplifies multiple purpose for this function instead of searching bonenames or
 	* having multiple bone names with multiple weights, and filtering through which one is correct one
 	* I assume all those things should be determined before coming here and this only cares about weights
 	**/
 	static void BlendPosesPerBoneFilter(
-		FCompactPose& BasePose, 
-		const TArray<FCompactPose>& BlendPoses, 
-		FBlendedCurve& BaseCurve, 
-		const TArray<FBlendedCurve>& BlendCurves, 
-		FCompactPose& OutPose, 
-		FBlendedCurve& OutCurve, 
-		TArray<FPerBoneBlendWeight>& BoneBlendWeights, 
-		bool bMeshSpaceRotationBlending,
+		FCompactPose& BasePose,
+		const TArray<FCompactPose>& BlendPoses,
+		FBlendedCurve& BaseCurve,
+		const TArray<FBlendedCurve>& BlendCurves,
+		FCompactPose& OutPose,
+		FBlendedCurve& OutCurve,
+		TArray<FPerBoneBlendWeight>& BoneBlendWeights,
+		EBlendPosesPerBoneFilterFlags blendFlags,
+		enum ECurveBlendOption::Type CurveBlendOption);
+
+	UE_DEPRECATED(4.23, "Please use the BlendPosesPerBoneFilter function that accepts EBlendPosesPerBoneFilterFlags.")
+	static void BlendPosesPerBoneFilter(
+		FCompactPose& BasePose,
+		const TArray<FCompactPose>& BlendPoses,
+		FBlendedCurve& BaseCurve,
+		const TArray<FBlendedCurve>& BlendCurves,
+		FCompactPose& OutPose,
+		FBlendedCurve& OutCurve,
+		TArray<FPerBoneBlendWeight>& BoneBlendWeights,
+		bool bMeshSpaceRotationBlend,
 		enum ECurveBlendOption::Type CurveBlendOption);
 
 	static void UpdateDesiredBoneWeight(const TArray<FPerBoneBlendWeight>& SrcBoneBlendWeights, TArray<FPerBoneBlendWeight>& TargetBoneBlendWeights, const TArray<float>& BlendWeights);
@@ -433,53 +456,7 @@ public:
 	* @param	RequiredBones		BoneContainer
 	*/
 	static void RetargetBoneTransform(const USkeleton* MySkeleton, const FName& RetargetSource, FTransform& BoneTransform, const int32 SkeletonBoneIndex, const FCompactPoseBoneIndex& BoneIndex, const FBoneContainer& RequiredBones, const bool bIsBakedAdditive);
-private:
-	/** 
-	* Blend Poses per bone weights : The BasePose + BlendPoses(SourceIndex) * Blend Weights(BoneIndex)
-	* Please note BlendWeights are array, so you can define per bone base 
-	* This supports multi per bone blending, but only one pose as blend at a time per track
-	* PerBoneBlendWeights.Num() == Atoms.Num()
-	*
-	* @note : This blends rotation in mesh space, but translation in local space
-	*
-	* I had multiple debates about having PerBoneBlendWeights array(for memory reasons),  
-	* but it so much simplifies multiple purpose for this function instead of searching bonenames or 
-	* having multiple bone names with multiple weights, and filtering through which one is correct one
-	* I assume all those things should be determined before coming here and this only cares about weights
-	**/
-	static void BlendMeshPosesPerBoneWeights(
-				FCompactPose& BasePose,
-				const TArray<FCompactPose>& BlendPoses,
-				FBlendedCurve& BaseCurve, 
-				const TArray<FBlendedCurve>& BlendedCurves, 
-				const TArray<FPerBoneBlendWeight>& BoneBlendWeights,
-				enum ECurveBlendOption::Type CurveBlendOption,
-				/*out*/ FCompactPose& OutPose, 
-				/*out*/ FBlendedCurve& OutCurve);
 
-	/** 
-	* Blend Poses per bone weights : The BasePoses + BlendPoses(SourceIndex) * Blend Weights(BoneIndex)
-	* Please note BlendWeights are array, so you can define per bone base 
-	* This supports multi per bone blending, but only one pose as blend at a time per track
-	* PerBoneBlendWeights.Num() == Atoms.Num()
-	*
-	* @note : This blends all in local space
-	* 
-	* I had multiple debates about having PerBoneBlendWeights array(for memory reasons),  
-	* but it so much simplifies multiple purpose for this function instead of searching bonenames or 
-	* having multiple bone names with multiple weights, and filtering through which one is correct one
-	* I assume all those things should be determined before coming here and this only cares about weights
-	**/
-	static void BlendLocalPosesPerBoneWeights(
-			FCompactPose& BasePose,
-			const TArray<FCompactPose>& BlendPoses,
-			FBlendedCurve& BaseCurve, 
-			const TArray<FBlendedCurve>& BlendedCurves, 
-			const TArray<FPerBoneBlendWeight>& BoneBlendWeights,
-			enum ECurveBlendOption::Type CurveBlendOption,
-			/*out*/ FCompactPose& OutPose, 
-			/*out*/ FBlendedCurve& OutCurve);
-public:
 	/** 
 	 * Calculate distance how close two strings are. 
 	 * By close, it calculates how many operations to transform First to Second 
@@ -489,3 +466,4 @@ public:
 	static int32 GetStringDistance(const FString& First, const FString& Second);
 };
 
+ENUM_CLASS_FLAGS(FAnimationRuntime::EBlendPosesPerBoneFilterFlags);
