Added flushing of single deferred resource update. This is used to ensure that the deferred update has happened before the resource is used in the case of the Render To Render Target BP function.

diff --git a/Engine/Source/Runtime/Engine/Private/KismetRenderingLibrary.cpp b/Engine/Source/Runtime/Engine/Private/KismetRenderingLibrary.cpp
--- a/Engine/Source/Runtime/Engine/Private/KismetRenderingLibrary.cpp
+++ b/Engine/Source/Runtime/Engine/Private/KismetRenderingLibrary.cpp
@@ -113,10 +113,12 @@ void UKismetRenderingLibrary::DrawMaterialToRenderTarget(UObject* WorldContextOb
 	{
 		World->SendAllEndOfFrameUpdates();
 
+		FTextureRenderTargetResource* RenderTargetResource = TextureRenderTarget->GameThread_GetRenderTargetResource();
+
 		UCanvas* Canvas = World->GetCanvasForDrawMaterialToRenderTarget();
 
 		FCanvas RenderCanvas(
-			TextureRenderTarget->GameThread_GetRenderTargetResource(), 
+			RenderTargetResource, 
 			nullptr, 
 			World,
 			World->FeatureLevel);
@@ -124,14 +126,14 @@ void UKismetRenderingLibrary::DrawMaterialToRenderTarget(UObject* WorldContextOb
 		Canvas->Init(TextureRenderTarget->SizeX, TextureRenderTarget->SizeY, nullptr, &RenderCanvas);
 		Canvas->Update();
 
-
-
 		TDrawEvent<FRHICommandList>* DrawMaterialToTargetEvent = new TDrawEvent<FRHICommandList>();
 
 		FName RTName = TextureRenderTarget->GetFName();
 		ENQUEUE_RENDER_COMMAND(BeginDrawEventCommand)(
-			[RTName, DrawMaterialToTargetEvent](FRHICommandList& RHICmdList)
+			[RTName, DrawMaterialToTargetEvent, RenderTargetResource](FRHICommandListImmediate& RHICmdList)
 			{
+				RenderTargetResource->FlushDeferredResourceUpdate(RHICmdList);
+
 				BEGIN_DRAW_EVENTF(
 					RHICmdList, 
 					DrawCanvasToTarget, 
@@ -144,7 +146,6 @@ void UKismetRenderingLibrary::DrawMaterialToRenderTarget(UObject* WorldContextOb
 		RenderCanvas.Flush_GameThread();
 		Canvas->Canvas = NULL;
 
-		FTextureRenderTargetResource* RenderTargetResource = TextureRenderTarget->GameThread_GetRenderTargetResource();
 		ENQUEUE_RENDER_COMMAND(CanvasRenderTargetResolveCommand)(
 			[RenderTargetResource, DrawMaterialToTargetEvent](FRHICommandList& RHICmdList)
 			{
diff --git a/Engine/Source/Runtime/Engine/Private/TextureRenderTarget.cpp b/Engine/Source/Runtime/Engine/Private/TextureRenderTarget.cpp
--- a/Engine/Source/Runtime/Engine/Private/TextureRenderTarget.cpp
+++ b/Engine/Source/Runtime/Engine/Private/TextureRenderTarget.cpp
@@ -125,15 +125,10 @@ void FDeferredUpdateResource::UpdateResources(FRHICommandListImmediate& RHICmdLi
 			FDeferredUpdateResource* RTResource = *ResourceIt;
 			// iterate to next resource before removing an entry
 			ResourceIt.Next();
+
 			if( RTResource )
 			{
-				// update each resource
-				RTResource->UpdateDeferredResource(RHICmdList);
-				if( RTResource->bOnlyUpdateOnce )
-				{
-					// remove from list if only a single update was requested
-					RTResource->RemoveFromDeferredUpdateList();
-				}
+				RTResource->FlushDeferredResourceUpdate(RHICmdList);
 			}
 		}
 		// since the updates should only occur once globally
@@ -142,6 +137,24 @@ void FDeferredUpdateResource::UpdateResources(FRHICommandListImmediate& RHICmdLi
 	}
 }
 
+/**
+ * Performs a deferred resource update on this resource if it exists in the UpdateList.
+ */
+void FDeferredUpdateResource::FlushDeferredResourceUpdate( FRHICommandListImmediate& RHICmdList )
+{
+	if( UpdateListLink.IsLinked() )
+	{
+		checkf(bNeedsUpdate, TEXT("The update list does not need to be updated at this point"));
+
+		UpdateDeferredResource(RHICmdList);
+		if( bOnlyUpdateOnce )
+		{
+			// Remove from list if only a single update was requested
+			RemoveFromDeferredUpdateList();
+		}
+	}
+}
+
 /**
 * Add this resource to deferred update list
 * @param OnlyUpdateOnce - flag this resource for a single update if true
diff --git a/Engine/Source/Runtime/Engine/Public/TextureResource.h b/Engine/Source/Runtime/Engine/Public/TextureResource.h
--- a/Engine/Source/Runtime/Engine/Public/TextureResource.h
+++ b/Engine/Source/Runtime/Engine/Public/TextureResource.h
@@ -423,6 +423,11 @@ public:
 	 */
 	ENGINE_API static void UpdateResources( FRHICommandListImmediate& RHICmdList );
 
+	/**
+	 * Performs a deferred resource update on this resource if it exists in the UpdateList.
+	 */
+	ENGINE_API void FlushDeferredResourceUpdate( FRHICommandListImmediate& RHICmdList );
+
 	/** 
 	 * This is reset after all viewports have been rendered
 	 */
