- Added embedded call to getting the memory bucket to pass memory info to native app

diff --git a/Engine/Source/Runtime/Engine/Private/GameEngine.cpp b/Engine/Source/Runtime/Engine/Private/GameEngine.cpp
--- a/Engine/Source/Runtime/Engine/Private/GameEngine.cpp
+++ b/Engine/Source/Runtime/Engine/Private/GameEngine.cpp
@@ -774,13 +774,16 @@ public:
 				FEmbeddedCommunication::AllowSleep(Requester);
 				Message.OnCompleteDelegate({}, TEXT(""));
 			}
-			// execute any console commands
 			else if (Message.Command == TEXT("devicesleep"))
 			{
 				bool bEnabled = Message.Parameters.FindRef(TEXT("enabled")).ToBool();
 				FPlatformApplicationMisc::ControlScreensaver(bEnabled ? FGenericPlatformApplicationMisc::Enable : FGenericPlatformApplicationMisc::Disable);
 				Message.OnCompleteDelegate({}, TEXT(""));
 			}
+			else if (Message.Command == TEXT("getmemorybucket"))
+			{
+				Message.OnCompleteDelegate({ { TEXT("bucket"), LexToString((int)FPlatformMemory::GetMemorySizeBucket()) } }, TEXT(""));
+			}
 
 // disable things that are too scary to expose to a shipping external app/webview
 #if !UE_BUILD_SHIPPING
