Added GameEngine call to support showing a share dialog box through embedded.
#tests iOS Client

diff --git a/Engine/Source/Runtime/Engine/Private/GameEngine.cpp b/Engine/Source/Runtime/Engine/Private/GameEngine.cpp
--- a/Engine/Source/Runtime/Engine/Private/GameEngine.cpp
+++ b/Engine/Source/Runtime/Engine/Private/GameEngine.cpp
@@ -873,6 +873,26 @@ public:
 					Message.OnCompleteDelegate({ }, FString::Printf(TEXT("CVar %s not found"), *Name));
 				}
 			}
+            else if (Message.Command == TEXT("shareurl"))
+            {
+                FString URL = Message.Parameters.FindRef(TEXT("url"));
+                FString Description = Message.Parameters.FindRef(TEXT("description"));
+                FString Subject = Message.Parameters.FindRef(TEXT("title"));
+                
+                FString ParsedX = Message.Parameters.FindRef(TEXT("x"));
+                FString ParsedY = Message.Parameters.FindRef(TEXT("y"));
+                
+                int ConvertedX = 0;
+                int ConvertedY = 0;
+                if ((!ParsedX.IsEmpty()
+                     && !ParsedY.IsEmpty()))
+                {
+                    LexFromString(ConvertedX, *ParsedX);
+                    LexFromString(ConvertedY, *ParsedY);
+                }
+                
+                FPlatformMisc::ShareURL(URL, FText::FromString(Description), ConvertedX, ConvertedY);
+            }
 			else
 			{
 				Message.OnCompleteDelegate({}, TEXT("Unknown command"));
