Undo CL 6641697. #fyi aaron.mcleran #rb none

diff --git a/Engine/Source/Runtime/Engine/Private/ADPCMAudioInfo.cpp b/Engine/Source/Runtime/Engine/Private/ADPCMAudioInfo.cpp
--- a/Engine/Source/Runtime/Engine/Private/ADPCMAudioInfo.cpp
+++ b/Engine/Source/Runtime/Engine/Private/ADPCMAudioInfo.cpp
@@ -20,7 +20,7 @@ namespace ADPCM
 }
 
 FADPCMAudioInfo::FADPCMAudioInfo(void)
-	: DeinterleavedUncompressedAudio(nullptr)
+	: UncompressedBlockData(nullptr)
 	, SamplesPerBlock(0)
 	, bSeekPending(false)
 {
@@ -28,10 +28,10 @@ FADPCMAudioInfo::FADPCMAudioInfo(void)
 
 FADPCMAudioInfo::~FADPCMAudioInfo(void)
 {
-	if (DeinterleavedUncompressedAudio != nullptr)
+	if(UncompressedBlockData != nullptr)
 	{
-		FMemory::Free(DeinterleavedUncompressedAudio);
-		DeinterleavedUncompressedAudio = nullptr;
+		FMemory::Free(UncompressedBlockData);
+		UncompressedBlockData = nullptr;
 	}
 }
 
@@ -39,14 +39,14 @@ void FADPCMAudioInfo::SeekToTime(const float SeekTime)
 {
 	CurCompressedChunkData = nullptr;
 
-	if (SeekTime <= 0.0f)
+	if(SeekTime <= 0.0f)
 	{
 		CurrentCompressedBlockIndex = 0;
-
-		CurrentUncompressedFrameIndex = 0;
+		CurrentUncompressedBlockSampleIndex = 0;
 		CurrentChunkIndex = 0;
 		CurrentChunkBufferOffset = 0;
-		TotalFramesStreamed = 0;
+		TotalSamplesStreamed = 0;
+
 		return;
 	}
 
@@ -54,13 +54,13 @@ void FADPCMAudioInfo::SeekToTime(const float SeekTime)
 	check(WaveInfo.pSamplesPerSec != nullptr);
 	const uint32 SamplesPerSec = *WaveInfo.pSamplesPerSec;
 	const uint32 SeekedSamples = static_cast<uint32>(SeekTime * SamplesPerSec);
-	TotalFramesStreamed = FMath::Min<uint32>(SeekedSamples, TotalFrames - 1);
+	TotalSamplesStreamed = FMath::Min<uint32>(SeekedSamples, TotalSamplesPerChannel - 1);
 
 	const uint32 HeaderOffset = static_cast<uint32>(WaveInfo.SampleDataStart - SrcBufferData);
 
 	if (Format == WAVE_FORMAT_ADPCM)
 	{
-		CurrentCompressedBlockIndex = TotalFramesStreamed / SamplesPerBlock; // Compute the block index that where SeekTime resides.
+		CurrentCompressedBlockIndex = TotalSamplesStreamed / SamplesPerBlock; // Compute the block index that where SeekTime resides.
 		CurrentChunkIndex = 0;
 		CurrentChunkBufferOffset = HeaderOffset;
 
@@ -82,7 +82,7 @@ void FADPCMAudioInfo::SeekToTime(const float SeekTime)
 		const int32 ChannelBlockSize = sizeof(int16) * NumChannels;
 
 		// 1. Find total offset
-		CurrentChunkBufferOffset = HeaderOffset + (TotalFramesStreamed * ChannelBlockSize);
+		CurrentChunkBufferOffset = HeaderOffset + (TotalSamplesStreamed * ChannelBlockSize);
 
 		// 2. Calculate index and remove from offset
 		CurrentChunkIndex = CurrentChunkBufferOffset / ADPCM::MaxChunkSize;
@@ -126,37 +126,38 @@ bool FADPCMAudioInfo::ReadCompressedInfo(const uint8* InSrcBufferData, uint32 In
 	if (Format == WAVE_FORMAT_ADPCM)
 	{
 		ADPCM::ADPCMFormatHeader* ADPCMHeader = (ADPCM::ADPCMFormatHeader*)FormatHeader;
-		TotalFrames = ADPCMHeader->SamplesPerChannel;
+		TotalSamplesPerChannel = ADPCMHeader->SamplesPerChannel;
 		SamplesPerBlock = ADPCMHeader->wSamplesPerBlock;
 
 		const uint32 PreambleSize = 7;
 		BlockSize = *WaveInfo.pBlockAlign;
 
 		// ADPCM starts with 2 uncompressed samples and then the remaining compressed sample data has 2 samples per byte
-		const uint32 NumSamplesPerUncompressedChannel = (2 + (BlockSize - PreambleSize) * 2);
-		NumBytesPerUncompressedChannel = NumSamplesPerUncompressedChannel * sizeof(int16);
+		const uint32 UncompressedBlockSamples = (2 + (BlockSize - PreambleSize) * 2);
+		UncompressedBlockSize = UncompressedBlockSamples * sizeof(int16);
 		CompressedBlockSize = BlockSize;
 
-		const uint32 TargetBlocks = MONO_PCM_BUFFER_SAMPLES / NumSamplesPerUncompressedChannel;
-		StreamBufferSize = TargetBlocks * NumBytesPerUncompressedChannel;
-		TotalDecodedSize = TotalFrames * NumChannels * sizeof(int16);
+		const uint32 TargetBlocks = MONO_PCM_BUFFER_SAMPLES / UncompressedBlockSamples;
+		StreamBufferSize = TargetBlocks * UncompressedBlockSize;
+		TotalDecodedSize = TotalSamplesPerChannel * NumChannels * sizeof(int16);
 
-		DeinterleavedUncompressedAudio = (int16*)FMemory::Realloc(DeinterleavedUncompressedAudio, NumChannels * NumBytesPerUncompressedChannel);
-		check(DeinterleavedUncompressedAudio != nullptr);
+		UncompressedBlockData = (uint8*)FMemory::Realloc(UncompressedBlockData, NumChannels * UncompressedBlockSize);
+		check(UncompressedBlockData != nullptr);
 		TotalCompressedBlocksPerChannel = (WaveInfo.SampleDataSize + CompressedBlockSize - 1) / CompressedBlockSize / NumChannels;
 	}
-	else if (Format == WAVE_FORMAT_LPCM)
+	else if(Format == WAVE_FORMAT_LPCM)
 	{
 		// There are no "blocks" in this case
 		BlockSize = 0;
-		NumBytesPerUncompressedChannel = 0;
+		UncompressedBlockSize = 0;
 		CompressedBlockSize = 0;
 		StreamBufferSize = 0;
-		DeinterleavedUncompressedAudio = nullptr;
+		UncompressedBlockData = nullptr;
 		TotalCompressedBlocksPerChannel = 0;
 
 		TotalDecodedSize = WaveInfo.SampleDataSize;
-		TotalFrames = TotalDecodedSize / (sizeof(int16) * NumChannels);
+
+		TotalSamplesPerChannel = TotalDecodedSize / sizeof(uint16) / NumChannels;
 	}
 	else
 	{
@@ -168,90 +169,89 @@ bool FADPCMAudioInfo::ReadCompressedInfo(const uint8* InSrcBufferData, uint32 In
 		QualityInfo->SampleRate = *WaveInfo.pSamplesPerSec;
 		QualityInfo->NumChannels = *WaveInfo.pChannels;
 		QualityInfo->SampleDataSize = TotalDecodedSize;
-		QualityInfo->Duration = (float)TotalFrames / QualityInfo->SampleRate;
+		QualityInfo->Duration = (float)TotalSamplesPerChannel / QualityInfo->SampleRate;
 	}
 
 	CurrentCompressedBlockIndex = 0;
-	TotalFramesStreamed = 0;
+	TotalSamplesStreamed = 0;
 	// This is set to the max value to trigger the decompression of the first audio block
-	CurrentUncompressedFrameIndex = NumBytesPerUncompressedChannel / sizeof(uint16);
+	CurrentUncompressedBlockSampleIndex = UncompressedBlockSize / sizeof(uint16);
+
 	return true;
 }
 
 bool FADPCMAudioInfo::ReadCompressedData(uint8* Destination, bool bLooping, uint32 BufferSize)
 {
+	const uint32 ChannelSampleSize = sizeof(uint16) * NumChannels;
+
 	// This correctly handles any BufferSize as long as its a multiple of sample size * number of channels
 	check(Destination);
-	check((BufferSize % (sizeof(uint16) * NumChannels)) == 0);
-	
+	check(BufferSize % ChannelSampleSize == 0);
+
 	int16* OutData = (int16*)Destination;
-	
-	bool ReachedEndOfSamples = false;
 
-	const int32 NumSamplesPerUncompressedChannel = NumBytesPerUncompressedChannel / sizeof(int16);
+	bool ReachedEndOfSamples = false;
 
-	if (Format == WAVE_FORMAT_ADPCM)
+	if(Format == WAVE_FORMAT_ADPCM)
 	{
-		int32 NumFramesToOutput = BufferSize / (NumChannels * sizeof(int16));
-
 		// We need to loop over the number of samples requested since an uncompressed block will not match the number of frames requested
-		while (NumFramesToOutput > 0)
+		while(BufferSize > 0)
 		{
-			if (CurrentUncompressedFrameIndex >= NumSamplesPerUncompressedChannel)
+			if(CurrentUncompressedBlockSampleIndex >= UncompressedBlockSize / sizeof(uint16))
 			{
 				// we need to decompress another block of compressed data from the current chunk
-				check(CurrentCompressedBlockIndex < TotalCompressedBlocksPerChannel);
+
 				// Decompress one block for each channel and store it in UncompressedBlockData
-				for (int32 ChannelIndex = 0; ChannelIndex < NumChannels; ++ChannelIndex)
+				for(int32 ChannelItr = 0; ChannelItr < NumChannels; ++ChannelItr)
 				{
-					const int32 CompressedAudioByteOffset = (ChannelIndex * TotalCompressedBlocksPerChannel + CurrentCompressedBlockIndex) * CompressedBlockSize;
-					const int32 OffsetIntoUncompressedBuffer = ChannelIndex * NumSamplesPerUncompressedChannel;
-
 					ADPCM::DecodeBlock(
-						&WaveInfo.SampleDataStart[CompressedAudioByteOffset],
+						WaveInfo.SampleDataStart + (ChannelItr * TotalCompressedBlocksPerChannel + CurrentCompressedBlockIndex) * CompressedBlockSize,
 						CompressedBlockSize,
-						&DeinterleavedUncompressedAudio[OffsetIntoUncompressedBuffer]);
+						(int16*)(UncompressedBlockData + ChannelItr * UncompressedBlockSize));
 				}
 
-				// Move on to the next compressed block:
-				CurrentUncompressedFrameIndex = 0;
+				// Update some bookkeeping
+				CurrentUncompressedBlockSampleIndex = 0;
 				++CurrentCompressedBlockIndex;
 			}
 
 			// Only copy over the number of samples we currently have available, we will loop around if needed
-			int32 NumFramesToInterleave = FMath::Min<int32>(NumSamplesPerUncompressedChannel - CurrentUncompressedFrameIndex, NumFramesToOutput);
+			uint32 DecompressedSamplesToCopy = FMath::Min<uint32>(UncompressedBlockSize / sizeof(uint16) - CurrentUncompressedBlockSampleIndex, BufferSize / (sizeof(uint16) * NumChannels));
+			check(DecompressedSamplesToCopy > 0);
 
-			check(NumFramesToInterleave > 0);
-			const int32 NumFramesLeft = TotalFrames - TotalFramesStreamed;
-			NumFramesToInterleave = FMath::Clamp<int32>(NumFramesToInterleave, 0, NumFramesLeft);
+			// Ensure we don't go over the number of samples left in the audio data
+			if(DecompressedSamplesToCopy > TotalSamplesPerChannel - TotalSamplesStreamed)
+			{
+				DecompressedSamplesToCopy = TotalSamplesPerChannel - TotalSamplesStreamed;
+			}
 
-			// Interleave audio from DeinterleavedUncompressedAudio into OutData:
-			for (int32 FrameIndex = 0; FrameIndex < NumFramesToInterleave; ++FrameIndex)
+			// Copy over the actual sample data
+			for(uint32 SampleItr = 0; SampleItr < DecompressedSamplesToCopy; ++SampleItr)
 			{
-				for (int32 ChannelIndex = 0; ChannelIndex < NumChannels; ++ChannelIndex)
+				for(int32 ChannelItr = 0; ChannelItr < NumChannels; ++ChannelItr)
 				{
-					const int32 OffsetIntoDeinterleavedBuffer = ChannelIndex * NumSamplesPerUncompressedChannel + CurrentUncompressedFrameIndex + FrameIndex;
-					*OutData++ = DeinterleavedUncompressedAudio[OffsetIntoDeinterleavedBuffer];
+					uint16 Value = *(int16*)(UncompressedBlockData + ChannelItr * UncompressedBlockSize + (CurrentUncompressedBlockSampleIndex + SampleItr) * sizeof(int16));
+					*OutData++ = Value;
 				}
 			}
 
 			// Update bookkeeping
-			CurrentUncompressedFrameIndex += NumFramesToInterleave;
-			NumFramesToOutput -= NumFramesToInterleave;
-			TotalFramesStreamed += NumFramesToInterleave;
-			
+			CurrentUncompressedBlockSampleIndex += DecompressedSamplesToCopy;
+			BufferSize -= DecompressedSamplesToCopy * sizeof(uint16) * NumChannels;
+			TotalSamplesStreamed += DecompressedSamplesToCopy;
+
 			// Check for the end of the audio samples and loop if needed
-			if (TotalFramesStreamed >= TotalFrames || CurrentCompressedBlockIndex >= TotalCompressedBlocksPerChannel)
+			if(TotalSamplesStreamed >= TotalSamplesPerChannel)
 			{
 				ReachedEndOfSamples = true;
 				// This is set to the max value to trigger the decompression of the first audio block
-				CurrentUncompressedFrameIndex = NumBytesPerUncompressedChannel / sizeof(int16);
+				CurrentUncompressedBlockSampleIndex = UncompressedBlockSize / sizeof(uint16);
 				CurrentCompressedBlockIndex = 0;
-				TotalFramesStreamed = 0;
-				if (!bLooping)
+				TotalSamplesStreamed = 0;
+				if(!bLooping)
 				{
-					// Set the remaining buffer to 0
-					FMemory::Memzero(OutData, NumFramesToOutput * sizeof(int16) * NumChannels);
+					// Zero remaining buffer
+					FMemory::Memzero(OutData, BufferSize);
 					return true;
 				}
 			}
@@ -259,26 +259,26 @@ bool FADPCMAudioInfo::ReadCompressedData(uint8* Destination, bool bLooping, uint
 	}
 	else
 	{
-		// For LPCM, we simply copy the audio from WaveInfo.SampleDataStart into OutData.
-		uint32	NumFramesToCopyToOutData = BufferSize / (sizeof(int16) * NumChannels);
-			
-		// Ensure we don't read past the number of frames left in the compressed audio buffer.
-		const int32 NumFramesLeftInBuffer = TotalFrames - TotalFramesStreamed;
-		NumFramesToCopyToOutData = FMath::Clamp<int32>(NumFramesToCopyToOutData, 0, NumFramesLeftInBuffer);
-
-		int32 OffsetIntoCompressedAudio = TotalFramesStreamed * sizeof(int16) * NumChannels;
-		FMemory::Memcpy(OutData, &WaveInfo.SampleDataStart[OffsetIntoCompressedAudio], NumFramesToCopyToOutData * sizeof(int16) * NumChannels);
-		TotalFramesStreamed += NumFramesToCopyToOutData;
-		BufferSize -= NumFramesToCopyToOutData * sizeof(int16) * NumChannels;
-		
+		uint32 DecompressedSamplesToCopy = BufferSize / ChannelSampleSize;
+
+		// Ensure we don't go over the number of samples left in the audio data
+		if(DecompressedSamplesToCopy > TotalSamplesPerChannel - TotalSamplesStreamed)
+		{
+			DecompressedSamplesToCopy = TotalSamplesPerChannel - TotalSamplesStreamed;
+		}
+
+		FMemory::Memcpy(OutData, WaveInfo.SampleDataStart + (TotalSamplesStreamed * ChannelSampleSize), DecompressedSamplesToCopy * ChannelSampleSize);
+		TotalSamplesStreamed += DecompressedSamplesToCopy;
+		BufferSize -= DecompressedSamplesToCopy * ChannelSampleSize;
+
 		// Check for the end of the audio samples and loop if needed
-		if (TotalFramesStreamed >= TotalFrames)
+		if(TotalSamplesStreamed >= TotalSamplesPerChannel)
 		{
 			ReachedEndOfSamples = true;
-			TotalFramesStreamed = 0;
-			if (!bLooping)
+			TotalSamplesStreamed = 0;
+			if(!bLooping)
 			{
-				// Set the remaining buffer to 0
+				// Zero remaining buffer
 				FMemory::Memzero(OutData, BufferSize);
 				return true;
 			}
@@ -330,43 +330,43 @@ bool FADPCMAudioInfo::StreamCompressedInfoInternal(USoundWave* Wave, struct FSou
 	FirstChunkSampleDataOffset = WaveInfo.SampleDataStart - FirstChunk;
 	CurrentChunkBufferOffset = 0;
 	CurCompressedChunkData = nullptr;
-	CurrentUncompressedFrameIndex = 0;
+	CurrentUncompressedBlockSampleIndex = 0;
 	CurrentChunkIndex = 0;
-	TotalFramesStreamed = 0;
+	TotalSamplesStreamed = 0;
 	Format = *WaveInfo.pFormatTag;
 	NumChannels = *WaveInfo.pChannels;
 
 	if (Format == WAVE_FORMAT_ADPCM)
 	{
 		ADPCM::ADPCMFormatHeader* ADPCMHeader = (ADPCM::ADPCMFormatHeader*)FormatHeader;
-		TotalFrames = ADPCMHeader->SamplesPerChannel;
+		TotalSamplesPerChannel = ADPCMHeader->SamplesPerChannel;
 		SamplesPerBlock = ADPCMHeader->wSamplesPerBlock;
 
 		const uint32 PreambleSize = 7;
 
 		BlockSize = *WaveInfo.pBlockAlign;
-		NumBytesPerUncompressedChannel = (2 + (BlockSize - PreambleSize) * 2) * sizeof(int16);
+		UncompressedBlockSize = (2 + (BlockSize - PreambleSize) * 2) * sizeof(int16);
 		CompressedBlockSize = BlockSize;
 
 		// Calculate buffer sizes and total number of samples
-		const uint32 NumUncompressedSamplesPerBlock = (2 + (BlockSize - PreambleSize) * 2);
-		const uint32 TargetNumBlocks = MONO_PCM_BUFFER_SAMPLES / NumUncompressedSamplesPerBlock;
-		StreamBufferSize = TargetNumBlocks * NumBytesPerUncompressedChannel;
-		TotalDecodedSize = ((WaveInfo.SampleDataSize + CompressedBlockSize - 1) / CompressedBlockSize) * NumBytesPerUncompressedChannel;
+		const uint32 uncompressedBlockSamples = (2 + (BlockSize - PreambleSize) * 2);
+		const uint32 targetBlocks = MONO_PCM_BUFFER_SAMPLES / uncompressedBlockSamples;
+		StreamBufferSize = targetBlocks * UncompressedBlockSize;
+		TotalDecodedSize = ((WaveInfo.SampleDataSize + CompressedBlockSize - 1) / CompressedBlockSize) * UncompressedBlockSize;
 
-		DeinterleavedUncompressedAudio = (int16*)FMemory::Realloc(DeinterleavedUncompressedAudio, NumChannels * NumBytesPerUncompressedChannel);
-		check(DeinterleavedUncompressedAudio != nullptr);
+		UncompressedBlockData = (uint8*)FMemory::Realloc(UncompressedBlockData, NumChannels * UncompressedBlockSize);
+		check(UncompressedBlockData != nullptr);
 	}
 	else if (Format == WAVE_FORMAT_LPCM)
 	{
 		BlockSize = 0;
-		NumBytesPerUncompressedChannel = 0;
+		UncompressedBlockSize = 0;
 		CompressedBlockSize = 0;
 
 		// This is uncompressed, so decoded size and buffer size are the same
 		TotalDecodedSize = WaveInfo.SampleDataSize;
 		StreamBufferSize = WaveInfo.SampleDataSize;
-		TotalFrames = StreamBufferSize / sizeof(uint16) / NumChannels;
+		TotalSamplesPerChannel = StreamBufferSize / sizeof(uint16) / NumChannels;
 	}
 	else
 	{
@@ -379,7 +379,7 @@ bool FADPCMAudioInfo::StreamCompressedInfoInternal(USoundWave* Wave, struct FSou
 		QualityInfo->SampleRate = *WaveInfo.pSamplesPerSec;
 		QualityInfo->NumChannels = *WaveInfo.pChannels;
 		QualityInfo->SampleDataSize = TotalDecodedSize;
-		QualityInfo->Duration = (float)TotalFrames / QualityInfo->SampleRate;
+		QualityInfo->Duration = (float)TotalSamplesPerChannel / QualityInfo->SampleRate;
 	}
 
 	return true;
@@ -391,30 +391,22 @@ bool FADPCMAudioInfo::StreamCompressedData(uint8* Destination, bool bLooping, ui
 	const int32 ChannelSampleSize = sizeof(uint16) * NumChannels;
 
 	// Ensure that BuffserSize is a multiple of the sample size times the number of channels
-	checkf((BufferSize % ChannelSampleSize) == 0, TEXT("Invalid buffer size %d requested for %d channels"), BufferSize, NumChannels);
-	
-	// If UncompressedBlockData is NULL, we haven't called FADPCM::StreamCompressedInfo yet.
-	if (DeinterleavedUncompressedAudio == nullptr || Destination == nullptr)
-	{
-		return false;
-	}
+	checkf(BufferSize % ChannelSampleSize == 0, TEXT("Invalid buffer size %d requested for %d channels"), BufferSize, NumChannels);
 
 	int16* OutData = (int16*)Destination;
 
-	int32 NumFramesLeftInOutData = BufferSize / (sizeof(int16) * NumChannels);
 	bool ReachedEndOfSamples = false;
 
-	if (Format == WAVE_FORMAT_ADPCM)
+	if(Format == WAVE_FORMAT_ADPCM)
 	{
 		// We need to loop over the number of samples requested since an uncompressed block will not match the number of frames requested
-		while (NumFramesLeftInOutData > 0)
+		while(BufferSize > 0)
 		{
-			const int32 NumSamplesPerUncompressedChannel = NumBytesPerUncompressedChannel / sizeof(uint16);
-			if (CurCompressedChunkData == nullptr || CurrentUncompressedFrameIndex >= NumSamplesPerUncompressedChannel)
+			if(CurCompressedChunkData == nullptr || CurrentUncompressedBlockSampleIndex >= UncompressedBlockSize / sizeof(uint16))
 			{
 				// we need to decompress another block of compressed data from the current chunk
 
-				if (CurCompressedChunkData == nullptr || CurrentChunkBufferOffset >= CurrentChunkDataSize)
+				if(CurCompressedChunkData == nullptr || CurrentChunkBufferOffset >= CurrentChunkDataSize)
 				{
 					// Get another chunk of compressed data from the streaming engine
 
@@ -422,7 +414,7 @@ bool FADPCMAudioInfo::StreamCompressedData(uint8* Destination, bool bLooping, ui
 					// preloaded so don't increment this when getting chunk 0, only later chunks. If previous chunk retrieval failed because it was
 					// not ready, don't re-increment the CurrentChunkIndex.  Only increment if seek not pending as seek determines the current chunk index
 					// and invalidates CurCompressedChunkData.
-					if (CurCompressedChunkData != nullptr)
+					if(CurCompressedChunkData != nullptr)
 					{
 						++CurrentChunkIndex;
 					}
@@ -430,18 +422,18 @@ bool FADPCMAudioInfo::StreamCompressedData(uint8* Destination, bool bLooping, ui
 					// Request the next chunk of data from the streaming engine
 					CurCompressedChunkData = IStreamingManager::Get().GetAudioStreamingManager().GetLoadedChunk(StreamingSoundWave, CurrentChunkIndex, &CurrentChunkDataSize);
 
-					if (CurCompressedChunkData == nullptr)
+					if(CurCompressedChunkData == nullptr)
 					{
 						// We only need to worry about missing the stream chunk if we were seeking. Seeking might cause a bit of latency with chunk loading. That is expected.
 						if (!bSeekPending)
 						{
 							// If we did not get it then just bail, CurrentChunkIndex will not get incremented on the next callback so in effect another attempt will be made to fetch the chunk
 							// Since audio streaming depends on the general data streaming mechanism used by other parts of the engine and new data is prefectched on the game tick thread its possible a game hickup can cause this
-							UE_LOG(LogAudio, Verbose, TEXT("Didn't load chunk %d in time for playback!"), CurrentChunkIndex);
+							UE_LOG(LogAudio, Verbose, TEXT("Missed Deadline chunk %d"), CurrentChunkIndex);
 						}
 
 						// zero out remaining data and bail
-						FMemory::Memzero(OutData, NumFramesLeftInOutData * NumChannels * sizeof(int16));
+						FMemory::Memset(OutData, 0, BufferSize);
 						return false;
 					}
 
@@ -454,57 +446,59 @@ bool FADPCMAudioInfo::StreamCompressedData(uint8* Destination, bool bLooping, ui
 				}
 
 				// Decompress one block for each channel and store it in UncompressedBlockData
-				for (int32 ChannelIndex = 0; ChannelIndex < NumChannels; ++ChannelIndex)
+				for(int32 ChannelItr = 0; ChannelItr < NumChannels; ++ChannelItr)
 				{
-					const int32 CompressedBufferOffset = CurrentChunkBufferOffset + ChannelIndex * CompressedBlockSize;
-					const int32 DeinterleavedUncompressedBufferOffset = ChannelIndex * NumSamplesPerUncompressedChannel;
 					ADPCM::DecodeBlock(
-						&CurCompressedChunkData[CompressedBufferOffset],
+						CurCompressedChunkData + CurrentChunkBufferOffset + ChannelItr * CompressedBlockSize,
 						CompressedBlockSize,
-						&DeinterleavedUncompressedAudio[DeinterleavedUncompressedBufferOffset]);
+						(int16*)(UncompressedBlockData + ChannelItr * UncompressedBlockSize));
 				}
 
 				// Update some bookkeeping
-				CurrentUncompressedFrameIndex = 0;
+				CurrentUncompressedBlockSampleIndex = 0;
 				CurrentChunkBufferOffset += NumChannels * CompressedBlockSize;
 			}
 
 			// Only copy over the number of samples we currently have available, we will loop around if needed
-			int32 NumFramesToInterleave = FMath::Min<int32>(NumSamplesPerUncompressedChannel - CurrentUncompressedFrameIndex, NumFramesLeftInOutData);
-			check(NumFramesToInterleave > 0);
-			
+			uint32 DecompressedSamplesToCopy = FMath::Min<uint32>(
+				(UncompressedBlockSize / sizeof(uint16)) - CurrentUncompressedBlockSampleIndex,
+				BufferSize / (ChannelSampleSize));
+			check(DecompressedSamplesToCopy > 0);
+
 			// Ensure we don't go over the number of samples left in the audio data
-			const int32 FramesLeft = TotalFrames - TotalFramesStreamed;
-			NumFramesToInterleave = FMath::Clamp<int32>(NumFramesToInterleave, 0, FramesLeft);
+			if(DecompressedSamplesToCopy > TotalSamplesPerChannel - TotalSamplesStreamed)
+			{
+				DecompressedSamplesToCopy = TotalSamplesPerChannel - TotalSamplesStreamed;
+			}
 
-			// Interleave audio from DeinterleavedUncompressedAudio
-			for (int32 FrameIndex = 0; FrameIndex < NumFramesToInterleave; ++FrameIndex)
+			// Copy over the actual sample data
+			for(uint32 SampleItr = 0; SampleItr < DecompressedSamplesToCopy; ++SampleItr)
 			{
-				for (int32 ChannelIndex = 0; ChannelIndex < NumChannels; ++ChannelIndex)
+				for(int32 ChannelItr = 0; ChannelItr < NumChannels; ++ChannelItr)
 				{
-					const int32 OffsetInDeinterleavedBuffer = ChannelIndex * NumSamplesPerUncompressedChannel + (CurrentUncompressedFrameIndex + FrameIndex);
-					*OutData++ = DeinterleavedUncompressedAudio[OffsetInDeinterleavedBuffer];
+					uint16 Value = *(int16*)(UncompressedBlockData + ChannelItr * UncompressedBlockSize + (CurrentUncompressedBlockSampleIndex + SampleItr) * sizeof(int16));
+					*OutData++ = Value;
 				}
 			}
 
 			// Update bookkeeping
-			CurrentUncompressedFrameIndex += NumFramesToInterleave;
-			NumFramesLeftInOutData -= NumFramesToInterleave;
-			TotalFramesStreamed += NumFramesToInterleave;
-			
+			CurrentUncompressedBlockSampleIndex += DecompressedSamplesToCopy;
+			BufferSize -= DecompressedSamplesToCopy * ChannelSampleSize;
+			TotalSamplesStreamed += DecompressedSamplesToCopy;
+
 			// Check for the end of the audio samples and loop if needed
-			if (TotalFramesStreamed >= TotalFrames)
+			if(TotalSamplesStreamed >= TotalSamplesPerChannel)
 			{
 				ReachedEndOfSamples = true;
-				CurrentUncompressedFrameIndex = 0;
+				CurrentUncompressedBlockSampleIndex = 0;
 				CurrentChunkIndex = 0;
 				CurrentChunkBufferOffset = 0;
-				TotalFramesStreamed = 0;
+				TotalSamplesStreamed = 0;
 				CurCompressedChunkData = nullptr;
-				if (!bLooping)
+				if(!bLooping)
 				{
 					// Set the remaining buffer to 0
-					FMemory::Memzero(OutData, NumFramesLeftInOutData * NumChannels * sizeof(int16));
+					FMemory::Memset(OutData, 0, BufferSize);
 					return true;
 				}
 			}
@@ -512,16 +506,16 @@ bool FADPCMAudioInfo::StreamCompressedData(uint8* Destination, bool bLooping, ui
 	}
 	else
 	{
-		while (NumFramesLeftInOutData > 0)
+		while(BufferSize > 0)
 		{
-			if (CurCompressedChunkData == nullptr || CurrentChunkBufferOffset >= CurrentChunkDataSize)
+			if(CurCompressedChunkData == nullptr || CurrentChunkBufferOffset >= CurrentChunkDataSize)
 			{
 				// Get another chunk of compressed data from the streaming engine
 
 				// CurrentChunkIndex is used to keep track of the current chunk for loading/unloading by the streaming engine but chunk 0
 				// is preloaded so we don't want to increment this when getting chunk 0, only later chunks
 				// Also, if we failed to get a previous chunk because it was not ready we don't want to re-increment the CurrentChunkIndex
-				if (CurCompressedChunkData != nullptr)
+				if(CurCompressedChunkData != nullptr)
 				{
 					++CurrentChunkIndex;
 				}
@@ -529,7 +523,7 @@ bool FADPCMAudioInfo::StreamCompressedData(uint8* Destination, bool bLooping, ui
 				// Request the next chunk of data from the streaming engine
 				CurCompressedChunkData = IStreamingManager::Get().GetAudioStreamingManager().GetLoadedChunk(StreamingSoundWave, CurrentChunkIndex, &CurrentChunkDataSize);
 
-				if (CurCompressedChunkData == nullptr)
+				if(CurCompressedChunkData == nullptr)
 				{
 					// Only report missing the stream chunk if we were seeking. Seeking
 					// may cause a bit of latency with chunk loading, which is expected.
@@ -538,10 +532,10 @@ bool FADPCMAudioInfo::StreamCompressedData(uint8* Destination, bool bLooping, ui
 						// CurrentChunkIndex will not get incremented on the next callback, effectively causing another attempt to fetch the chunk.
 						// Since audio streaming depends on the general data streaming mechanism used by other parts of the engine and new data is
 						// prefetched on the game tick thread, this may be caused by a game hitch.
-						UE_LOG(LogAudio, Warning, TEXT("Didn't load chunk %d in time for playback!"), CurrentChunkIndex);
+						UE_LOG(LogAudio, Verbose, TEXT("Missed streaming ADPCM deadline chunk %d"), CurrentChunkIndex);
 					}
 
-					FMemory::Memzero(OutData, NumFramesLeftInOutData * NumChannels * sizeof(int16));
+					FMemory::Memset(OutData, 0, BufferSize);
 					return false;
 				}
 
@@ -553,33 +547,38 @@ bool FADPCMAudioInfo::StreamCompressedData(uint8* Destination, bool bLooping, ui
 				bSeekPending = false;
 			}
 
-			int32 NumFramesToInterleave = FMath::Min<int32>((CurrentChunkDataSize - CurrentChunkBufferOffset), NumFramesLeftInOutData);
-			check(NumFramesToInterleave > 0);
-			
+			uint32 DecompressedSamplesToCopy = FMath::Min<uint32>(
+				(CurrentChunkDataSize - CurrentChunkBufferOffset) / ChannelSampleSize,
+				BufferSize / ChannelSampleSize);
+
+			check(DecompressedSamplesToCopy > 0);
+
 			// Ensure we don't go over the number of samples left in the audio data
-			const int32 FramesLeft = TotalFrames - TotalFramesStreamed;
-			check(FramesLeft > 0);
-			FMath::Clamp<int32>(NumFramesToInterleave, 0, FramesLeft);
-
-			FMemory::Memcpy(OutData, &CurCompressedChunkData[CurrentChunkBufferOffset], NumFramesToInterleave * (sizeof(int16) * NumChannels));
-			
-			OutData += NumFramesToInterleave * NumChannels;
-			CurrentChunkBufferOffset += NumFramesToInterleave * (sizeof(uint16) * NumChannels);
-			BufferSize -= NumFramesToInterleave * (sizeof(uint16) * NumChannels);
-			TotalFramesStreamed += NumFramesToInterleave;
-			
+			if(DecompressedSamplesToCopy > TotalSamplesPerChannel - TotalSamplesStreamed)
+			{
+				DecompressedSamplesToCopy = TotalSamplesPerChannel - TotalSamplesStreamed;
+			}
+
+			const uint32 SampleSize = DecompressedSamplesToCopy * ChannelSampleSize;
+			FMemory::Memcpy(OutData, CurCompressedChunkData + CurrentChunkBufferOffset, SampleSize);
+
+			OutData += DecompressedSamplesToCopy * NumChannels;
+			CurrentChunkBufferOffset += SampleSize;
+			BufferSize -= SampleSize;
+			TotalSamplesStreamed += DecompressedSamplesToCopy;
+
 			// Check for the end of the audio samples and loop if needed
-			if (TotalFramesStreamed >= TotalFrames)
+			if(TotalSamplesStreamed >= TotalSamplesPerChannel)
 			{
 				ReachedEndOfSamples = true;
 				CurrentChunkIndex = 0;
 				CurrentChunkBufferOffset = 0;
-				TotalFramesStreamed = 0;
+				TotalSamplesStreamed = 0;
 				CurCompressedChunkData = nullptr;
-				if (!bLooping)
+				if(!bLooping)
 				{
 					// Set the remaining buffer to 0
-					FMemory::Memzero(OutData, NumFramesLeftInOutData * NumChannels * sizeof(int16));
+					FMemory::Memset(OutData, 0, BufferSize);
 					return true;
 				}
 			}
diff --git a/Engine/Source/Runtime/Engine/Public/ADPCMAudioInfo.h b/Engine/Source/Runtime/Engine/Public/ADPCMAudioInfo.h
--- a/Engine/Source/Runtime/Engine/Public/ADPCMAudioInfo.h
+++ b/Engine/Source/Runtime/Engine/Public/ADPCMAudioInfo.h
@@ -118,7 +118,7 @@ public:
 	const uint8*	SrcBufferData;
 	uint32			SrcBufferDataSize;
 
-	uint32	NumBytesPerUncompressedChannel;
+	uint32	UncompressedBlockSize;
 	uint32	CompressedBlockSize;
 	uint32	BlockSize;
 	int32	StreamBufferSize;
@@ -126,13 +126,13 @@ public:
 	int32	NumChannels;
 	int32	Format;
 
-	int16*			DeinterleavedUncompressedAudio;			// This holds the current block of compressed data for all channels
-	int32			CurrentUncompressedFrameIndex;	// This is the sample index within the current uncompressed block data
+	uint8*			UncompressedBlockData;			// This holds the current block of compressed data for all channels
+	uint32			CurrentUncompressedBlockSampleIndex;	// This is the sample index within the current uncompressed block data
 	uint32			CurrentChunkIndex;				// This is the index that is currently being used, needed by streaming engine to make sure it stays loaded and the next chunk gets preloaded
 	uint32			CurrentChunkBufferOffset;		// This is this byte offset within the current chunk, used by streaming engine to prioritize a load if more then half way through current chunk
-	uint32			CurrentChunkDataSize;			// The size of the current chunk, the first chunk is bigger to accomodate the header info
-	uint32			TotalFramesStreamed;			// The number of samples streamed so far
-	uint32			TotalFrames;			// Number of samples per channel, used to detect when an audio waveform has ended
+	uint32			CurrentChunkDataSize;			// The size of the current chunk, the first chunk is bigger to accommodate the header info
+	uint32			TotalSamplesStreamed;			// The number of samples streamed so far (per channel)
+	uint32			TotalSamplesPerChannel;			// Number of samples per channel, used to detect when an audio waveform has ended
 	uint32			SamplesPerBlock;				// The number of samples per block
 	uint32			FirstChunkSampleDataOffset;		// The size of the header in the first chunk, used to skip over it when looping or starting the sample over
 	const uint8*	CurCompressedChunkData;			// A pointer to the current chunk of data
