Minor tweaks for code style

diff --git a/Engine/Source/Runtime/AnimGraphRuntime/Private/AnimNodes/AnimNode_LayeredBoneBlend.cpp b/Engine/Source/Runtime/AnimGraphRuntime/Private/AnimNodes/AnimNode_LayeredBoneBlend.cpp
--- a/Engine/Source/Runtime/AnimGraphRuntime/Private/AnimNodes/AnimNode_LayeredBoneBlend.cpp
+++ b/Engine/Source/Runtime/AnimGraphRuntime/Private/AnimNodes/AnimNode_LayeredBoneBlend.cpp
@@ -265,16 +265,16 @@ void FAnimNode_LayeredBoneBlend::Evaluate_AnyThread(FPoseContext& Output)
 			}
 		}
 
-		FAnimationRuntime::EBlendPosesPerBoneFilterFlags blendFlags = FAnimationRuntime::EBlendPosesPerBoneFilterFlags::None;
+		FAnimationRuntime::EBlendPosesPerBoneFilterFlags BlendFlags = FAnimationRuntime::EBlendPosesPerBoneFilterFlags::None;
 		if (bMeshSpaceRotationBlend)
 		{
-			blendFlags |= FAnimationRuntime::EBlendPosesPerBoneFilterFlags::MeshSpaceRotation;
+			BlendFlags |= FAnimationRuntime::EBlendPosesPerBoneFilterFlags::MeshSpaceRotation;
 		}
 		if (bMeshSpaceScaleBlend)
 		{
-			blendFlags |= FAnimationRuntime::EBlendPosesPerBoneFilterFlags::MeshSpaceScale;
+			BlendFlags |= FAnimationRuntime::EBlendPosesPerBoneFilterFlags::MeshSpaceScale;
 		}
-		FAnimationRuntime::BlendPosesPerBoneFilter(BasePoseContext.Pose, TargetBlendPoses, BasePoseContext.Curve, TargetBlendCurves, Output.Pose, Output.Curve, CurrentBoneBlendWeights, blendFlags, CurveBlendOption);
+		FAnimationRuntime::BlendPosesPerBoneFilter(BasePoseContext.Pose, TargetBlendPoses, BasePoseContext.Curve, TargetBlendCurves, Output.Pose, Output.Curve, CurrentBoneBlendWeights, BlendFlags, CurveBlendOption);
 	}
 }
 
diff --git a/Engine/Source/Runtime/Engine/Private/Animation/AnimationRuntime.cpp b/Engine/Source/Runtime/Engine/Private/Animation/AnimationRuntime.cpp
--- a/Engine/Source/Runtime/Engine/Private/Animation/AnimationRuntime.cpp
+++ b/Engine/Source/Runtime/Engine/Private/Animation/AnimationRuntime.cpp
@@ -1044,13 +1044,13 @@ void FAnimationRuntime::BlendPosesPerBoneFilter(
 	bool bMeshSpaceRotationBlend,
 	ECurveBlendOption::Type CurveBlendOption)
 {
-	EBlendPosesPerBoneFilterFlags blendFlags = EBlendPosesPerBoneFilterFlags::None;
+	EBlendPosesPerBoneFilterFlags BlendFlags = EBlendPosesPerBoneFilterFlags::None;
 	if (bMeshSpaceRotationBlend)
 	{
-		blendFlags |= EBlendPosesPerBoneFilterFlags::MeshSpaceRotation;
+		BlendFlags |= EBlendPosesPerBoneFilterFlags::MeshSpaceRotation;
 	}
 
-	return BlendPosesPerBoneFilter(BasePose, BlendPoses, BaseCurve, BlendedCurves, OutPose, OutCurve, BoneBlendWeights, blendFlags, CurveBlendOption);
+	return BlendPosesPerBoneFilter(BasePose, BlendPoses, BaseCurve, BlendedCurves, OutPose, OutCurve, BoneBlendWeights, BlendFlags, CurveBlendOption);
 }
 
 void FAnimationRuntime::BlendPosesPerBoneFilter(
@@ -1061,7 +1061,7 @@ void FAnimationRuntime::BlendPosesPerBoneFilter(
 	struct FCompactPose& OutPose, 
 	struct FBlendedCurve& OutCurve, 
 	TArray<FPerBoneBlendWeight>& BoneBlendWeights, 
-	EBlendPosesPerBoneFilterFlags blendFlags,
+	EBlendPosesPerBoneFilterFlags BlendFlags,
 	ECurveBlendOption::Type CurveBlendOption)
 {
 	SCOPE_CYCLE_COUNTER(STAT_BlendPosesPerBoneFilter);
@@ -1103,8 +1103,8 @@ void FAnimationRuntime::BlendPosesPerBoneFilter(
 	ScaleArray& BlendScales = ScratchArea.BlendScales;
 	ScaleArray& TargetScales = ScratchArea.TargetScales;
 
-	const bool bMeshSpaceRotationBlend = EnumHasAnyFlags(blendFlags, EBlendPosesPerBoneFilterFlags::MeshSpaceRotation);
-	const bool bMeshSpaceScaleBlend = EnumHasAnyFlags(blendFlags, EBlendPosesPerBoneFilterFlags::MeshSpaceScale);
+	const bool bMeshSpaceRotationBlend = EnumHasAnyFlags(BlendFlags, EBlendPosesPerBoneFilterFlags::MeshSpaceRotation);
+	const bool bMeshSpaceScaleBlend = EnumHasAnyFlags(BlendFlags, EBlendPosesPerBoneFilterFlags::MeshSpaceScale);
 
 	if (bMeshSpaceRotationBlend)
 	{
