[NO NON-PII COMMIT MESSAGE]

diff --git a/Engine/Source/Runtime/Engine/Private/ADPCMAudioInfo.cpp b/Engine/Source/Runtime/Engine/Private/ADPCMAudioInfo.cpp
--- a/Engine/Source/Runtime/Engine/Private/ADPCMAudioInfo.cpp
+++ b/Engine/Source/Runtime/Engine/Private/ADPCMAudioInfo.cpp
@@ -20,7 +20,14 @@ namespace ADPCM
 }
 
 FADPCMAudioInfo::FADPCMAudioInfo(void)
-	: UncompressedBlockData(nullptr)
+	: UncompressedBlockSize(0)
+	, CompressedBlockSize(0)
+	, BlockSize(0)
+	, StreamBufferSize(0)
+	, TotalDecodedSize(0)
+	, NumChannels(0)
+	, Format(0)
+	, UncompressedBlockData(nullptr)
 	, SamplesPerBlock(0)
 	, bSeekPending(false)
 {
@@ -387,11 +394,30 @@ bool FADPCMAudioInfo::StreamCompressedInfoInternal(USoundWave* Wave, struct FSou
 
 bool FADPCMAudioInfo::StreamCompressedData(uint8* Destination, bool bLooping, uint32 BufferSize)
 {
+	// Initial sanity checks:
+	if (Destination == nullptr || BufferSize == 0)
+	{
+		UE_LOG(LogAudio, Error, TEXT("Stream Compressed Info not called!"));
+		return false;
+	}
+
+	if (NumChannels == 0)
+	{
+		UE_LOG(LogAudio, Error, TEXT("Stream Compressed Info not called!"));
+		FMemory::Memzero(Destination, BufferSize);
+		return true;
+	}
+
 	// Destination samples are interlaced by channel, BufferSize is in bytes
 	const int32 ChannelSampleSize = sizeof(uint16) * NumChannels;
 
 	// Ensure that BuffserSize is a multiple of the sample size times the number of channels
-	checkf(BufferSize % ChannelSampleSize == 0, TEXT("Invalid buffer size %d requested for %d channels"), BufferSize, NumChannels);
+	if (BufferSize % ChannelSampleSize != 0)
+	{
+		UE_LOG(LogAudio, Error, TEXT("Invalid buffer size %d requested for %d channels."), BufferSize, NumChannels);
+		FMemory::Memzero(Destination, BufferSize);
+		return true;
+	}
 
 	int16* OutData = (int16*)Destination;
 
