Added support to store the TotalPlayerTicks in the Performance Charts. The value is driven by game implementation.

#rnx

diff --git a/Engine/Source/Runtime/Engine/Private/ChartCreation.cpp b/Engine/Source/Runtime/Engine/Private/ChartCreation.cpp
--- a/Engine/Source/Runtime/Engine/Private/ChartCreation.cpp
+++ b/Engine/Source/Runtime/Engine/Private/ChartCreation.cpp
@@ -620,6 +620,9 @@ void FPerformanceTrackingChart::Reset(const FDateTime& InStartTime)
 	MaxDrawCalls = 0;
 	MinDrawCalls = INT_MAX;
 	TotalDrawCalls = 0;
+	MaxPlayerTicks = 0;
+	MinPlayerTicks = INT_MAX;
+	TotalPlayerTicks = 0;
 	MaxDrawnPrimitives = 0;
 	MinDrawnPrimitives = INT_MAX;
 	TotalDrawnPrimitives = 0;
@@ -687,6 +690,9 @@ void FPerformanceTrackingChart::AccumulateWith(const FPerformanceTrackingChart&
 	MaxDrawCalls = FMath::Max(MaxDrawCalls, Chart.MaxDrawCalls);
 	MinDrawCalls = FMath::Min(MinDrawCalls, Chart.MinDrawCalls);
 	TotalDrawCalls += Chart.TotalDrawCalls;
+	MaxPlayerTicks = FMath::Max(MaxPlayerTicks, Chart.MaxPlayerTicks);
+	MinPlayerTicks = FMath::Min(MinPlayerTicks, Chart.MinPlayerTicks);
+	TotalPlayerTicks += Chart.TotalPlayerTicks;
 	MaxDrawnPrimitives = FMath::Max(MaxDrawnPrimitives, Chart.MaxDrawnPrimitives);
 	MinDrawnPrimitives = FMath::Min(MinDrawnPrimitives, Chart.MinDrawnPrimitives);
 	TotalDrawnPrimitives += Chart.TotalDrawnPrimitives;
diff --git a/Engine/Source/Runtime/Engine/Public/ChartCreation.h b/Engine/Source/Runtime/Engine/Public/ChartCreation.h
--- a/Engine/Source/Runtime/Engine/Public/ChartCreation.h
+++ b/Engine/Source/Runtime/Engine/Public/ChartCreation.h
@@ -185,6 +185,11 @@ public:
 	int32 MinDrawCalls;
 	int32 TotalDrawCalls;
 
+	/** Total number of player ticks (note: it is up to the game to populate this field) */
+	int32 MaxPlayerTicks;
+	int32 MinPlayerTicks;
+	int32 TotalPlayerTicks;
+
 	/** Total number of primitives drawn */
 	int32 MaxDrawnPrimitives;
 	int32 MinDrawnPrimitives;
