Fixed: EngineTest Animation Editor - AllowCurve screenshot tests failed

diff --git a/Engine/Source/Runtime/Engine/Private/SkeletalRenderGPUSkin.cpp b/Engine/Source/Runtime/Engine/Private/SkeletalRenderGPUSkin.cpp
--- a/Engine/Source/Runtime/Engine/Private/SkeletalRenderGPUSkin.cpp
+++ b/Engine/Source/Runtime/Engine/Private/SkeletalRenderGPUSkin.cpp
@@ -74,6 +74,7 @@ void FMorphVertexBuffer::InitDynamicRHI()
 	FRHIResourceCreateInfo CreateInfo;
 
 	const bool bUseGPUMorphTargets = UseGPUMorphTargets(GMaxRHIShaderPlatform);
+	bUsesComputeShader = bUseGPUMorphTargets;
 
 #if PLATFORM_PS4
 	// PS4 requires non-static buffers in order to be updated on the GPU while the CPU is writing into it
