InstancedStaticMesh - Added more efficient batch transform update

diff --git a/Engine/Source/Runtime/Engine/Classes/Components/HierarchicalInstancedStaticMeshComponent.h b/Engine/Source/Runtime/Engine/Classes/Components/HierarchicalInstancedStaticMeshComponent.h
--- a/Engine/Source/Runtime/Engine/Classes/Components/HierarchicalInstancedStaticMeshComponent.h
+++ b/Engine/Source/Runtime/Engine/Classes/Components/HierarchicalInstancedStaticMeshComponent.h
@@ -210,6 +210,9 @@ public:
 	virtual int32 AddInstance(const FTransform& InstanceTransform) override;
 	virtual bool RemoveInstance(int32 InstanceIndex) override;
 	virtual bool UpdateInstanceTransform(int32 InstanceIndex, const FTransform& NewInstanceTransform, bool bWorldSpace, bool bMarkRenderStateDirty = false, bool bTeleport = false) override;
+	virtual bool BatchUpdateInstancesTransforms(int32 StartInstanceIndex, const TArray<FTransform>& NewInstancesTransforms, bool bWorldSpace=false, bool bMarkRenderStateDirty=false, bool bTeleport=false) override;
+	virtual bool BatchUpdateInstancesTransform(int32 StartInstanceIndex, int32 NumInstances, const FTransform& NewInstancesTransform, bool bWorldSpace=false, bool bMarkRenderStateDirty=false, bool bTeleport=false) override;
+
 	virtual void ClearInstances() override;
 	virtual TArray<int32> GetInstancesOverlappingSphere(const FVector& Center, float Radius, bool bSphereInWorldSpace = true) const override;
 	virtual TArray<int32> GetInstancesOverlappingBox(const FBox& Box, bool bBoxInWorldSpace = true) const override;
diff --git a/Engine/Source/Runtime/Engine/Classes/Components/InstancedStaticMeshComponent.h b/Engine/Source/Runtime/Engine/Classes/Components/InstancedStaticMeshComponent.h
--- a/Engine/Source/Runtime/Engine/Classes/Components/InstancedStaticMeshComponent.h
+++ b/Engine/Source/Runtime/Engine/Classes/Components/InstancedStaticMeshComponent.h
@@ -181,7 +181,7 @@ class ENGINE_API UInstancedStaticMeshComponent : public UStaticMeshComponent
 	*
 	* @param InstanceIndex			The index of the instance to update
 	* @param NewInstanceTransform	The new transform
-	* @param bWorldSpace			If true, the new transform interpreted as a World Space transform, otherwise it is interpreted as Local Space
+	* @param bWorldSpace			If true, the new transform is interpreted as a World Space transform, otherwise it is interpreted as Local Space
 	* @param bMarkRenderStateDirty	If true, the change should be visible immediately. If you are updating many instances you should only set this to true for the last instance.
 	* @param bTeleport				Whether or not the instance's physics should be moved normally, or teleported (moved instantly, ignoring velocity).
 	* @return						True on success.
@@ -189,6 +189,33 @@ class ENGINE_API UInstancedStaticMeshComponent : public UStaticMeshComponent
 	UFUNCTION(BlueprintCallable, Category = "Components|InstancedStaticMesh")
 	virtual bool UpdateInstanceTransform(int32 InstanceIndex, const FTransform& NewInstanceTransform, bool bWorldSpace=false, bool bMarkRenderStateDirty=false, bool bTeleport=false);
 
+    /**
+	* Update the transform for an array of instances.
+	*
+	* @param StartInstanceIndex		The starting index of the instances to update
+	* @param NewInstancesTransforms	The new transforms
+	* @param bWorldSpace			If true, the new transforms are interpreted as a World Space transform, otherwise it is interpreted as Local Space
+	* @param bMarkRenderStateDirty	If true, the change should be visible immediately. If you are updating many instances you should only set this to true for the last instance.
+	* @param bTeleport				Whether or not the instances physics should be moved normally, or teleported (moved instantly, ignoring velocity).
+	* @return						True on success.
+	*/
+	UFUNCTION(BlueprintCallable, Category = "Components|InstancedStaticMesh")
+	virtual bool BatchUpdateInstancesTransforms(int32 StartInstanceIndex, const TArray<FTransform>& NewInstancesTransforms, bool bWorldSpace=false, bool bMarkRenderStateDirty=false, bool bTeleport=false);
+
+	/**
+	* Update the transform for a number of instances.
+	*
+	* @param StartInstanceIndex		The starting index of the instances to update
+	* @param NumInstances			The number of instances to update
+	* @param NewInstancesTransform	The new transform
+	* @param bWorldSpace			If true, the new transform is interpreted as a World Space transform, otherwise it is interpreted as Local Space
+	* @param bMarkRenderStateDirty	If true, the change should be visible immediately. If you are updating many instances you should only set this to true for the last instance.
+	* @param bTeleport				Whether or not the instances physics should be moved normally, or teleported (moved instantly, ignoring velocity).
+	* @return						True on success.
+	*/
+	UFUNCTION(BlueprintCallable, Category = "Components|InstancedStaticMesh")
+	virtual bool BatchUpdateInstancesTransform(int32 StartInstanceIndex, int32 NumInstances, const FTransform& NewInstancesTransform, bool bWorldSpace=false, bool bMarkRenderStateDirty=false, bool bTeleport=false);
+
 	/** Remove the instance specified. Returns True on success. Note that this will leave the array in order, but may shrink it. */
 	UFUNCTION(BlueprintCallable, Category = "Components|InstancedStaticMesh")
 	virtual bool RemoveInstance(int32 InstanceIndex);
diff --git a/Engine/Source/Runtime/Engine/Private/HierarchicalInstancedStaticMesh.cpp b/Engine/Source/Runtime/Engine/Private/HierarchicalInstancedStaticMesh.cpp
--- a/Engine/Source/Runtime/Engine/Private/HierarchicalInstancedStaticMesh.cpp
+++ b/Engine/Source/Runtime/Engine/Private/HierarchicalInstancedStaticMesh.cpp
@@ -2253,6 +2253,36 @@ bool UHierarchicalInstancedStaticMeshComponent::UpdateInstanceTransform(int32 In
 	return Result;
 }
 
+bool UHierarchicalInstancedStaticMeshComponent::BatchUpdateInstancesTransforms(int32 StartInstanceIndex, const TArray<FTransform>& NewInstancesTransforms, bool bWorldSpace, bool bMarkRenderStateDirty, bool bTeleport)
+{
+	bool BatchResult = true;
+
+	int32 InstanceIndex = StartInstanceIndex;
+	for(const FTransform& NewInstanceTransform : NewInstancesTransforms)
+	{
+		bool Result = UpdateInstanceTransform(InstanceIndex, NewInstanceTransform, bWorldSpace, bMarkRenderStateDirty, bTeleport);
+		BatchResult = BatchResult && Result;
+		
+		InstanceIndex++;
+	}
+
+	return BatchResult;
+}
+
+bool UHierarchicalInstancedStaticMeshComponent::BatchUpdateInstancesTransform(int32 StartInstanceIndex, int32 NumInstances, const FTransform& NewInstancesTransform, bool bWorldSpace, bool bMarkRenderStateDirty, bool bTeleport)
+{
+	bool BatchResult = true;
+
+	int32 EndInstanceIndex = StartInstanceIndex + NumInstances;
+	for(int32 InstanceIndex = StartInstanceIndex; InstanceIndex < EndInstanceIndex; ++InstanceIndex)
+	{
+		bool Result = UpdateInstanceTransform(InstanceIndex, NewInstancesTransform, bWorldSpace, bMarkRenderStateDirty, bTeleport);
+		BatchResult = BatchResult && Result;
+	}
+
+	return BatchResult;
+}
+
 void UHierarchicalInstancedStaticMeshComponent::ApplyComponentInstanceData(FInstancedStaticMeshComponentInstanceData* InstancedMeshData)
 {
 	UInstancedStaticMeshComponent::ApplyComponentInstanceData(InstancedMeshData);
diff --git a/Engine/Source/Runtime/Engine/Private/InstancedStaticMesh.cpp b/Engine/Source/Runtime/Engine/Private/InstancedStaticMesh.cpp
--- a/Engine/Source/Runtime/Engine/Private/InstancedStaticMesh.cpp
+++ b/Engine/Source/Runtime/Engine/Private/InstancedStaticMesh.cpp
@@ -1958,6 +1958,94 @@ bool UInstancedStaticMeshComponent::UpdateInstanceTransform(int32 InstanceIndex,
 	return true;
 }
 
+bool UInstancedStaticMeshComponent::BatchUpdateInstancesTransforms(int32 StartInstanceIndex, const TArray<FTransform>& NewInstancesTransforms, bool bWorldSpace, bool bMarkRenderStateDirty, bool bTeleport)
+{
+	if (!PerInstanceSMData.IsValidIndex(StartInstanceIndex) || !PerInstanceSMData.IsValidIndex(StartInstanceIndex + NewInstancesTransforms.Num() - 1))
+	{
+		return false;
+	}
+
+	Modify();
+
+	int32 InstanceIndex = StartInstanceIndex;
+	for (const FTransform& NewInstanceTransform : NewInstancesTransforms)
+	{
+		FInstancedStaticMeshInstanceData& InstanceData = PerInstanceSMData[InstanceIndex];
+
+		// TODO: Computing LocalTransform is useless when we're updating the world location for the entire mesh.
+		// Should find some way around this for performance.
+
+		// Render data uses local transform of the instance
+		FTransform LocalTransform = bWorldSpace ? NewInstanceTransform.GetRelativeTransform(GetComponentTransform()) : NewInstanceTransform;
+		InstanceData.Transform = LocalTransform.ToMatrixWithScale();
+
+		if(bPhysicsStateCreated)
+		{
+			// Physics uses world transform of the instance
+			FTransform WorldTransform = bWorldSpace ? NewInstanceTransform : (LocalTransform * GetComponentTransform());
+			UpdateInstanceBodyTransform(InstanceIndex, WorldTransform, bTeleport);
+		}
+
+		InstanceIndex++;
+	}
+
+	// Request navigation update - Execute on a single index as it updates everything anyway
+	PartialNavigationUpdate(StartInstanceIndex);
+
+	// Force recreation of the render data when proxy is created
+	InstanceUpdateCmdBuffer.Edit();
+
+	if (bMarkRenderStateDirty)
+	{
+		MarkRenderStateDirty();
+	}
+
+	return true;
+}
+
+bool UInstancedStaticMeshComponent::BatchUpdateInstancesTransform(int32 StartInstanceIndex, int32 NumInstances, const FTransform& NewInstancesTransform, bool bWorldSpace, bool bMarkRenderStateDirty, bool bTeleport)
+{
+	if(!PerInstanceSMData.IsValidIndex(StartInstanceIndex) || !PerInstanceSMData.IsValidIndex(StartInstanceIndex + NumInstances - 1))
+	{
+		return false;
+	}
+
+	Modify();
+
+	int32 EndInstanceIndex = StartInstanceIndex + NumInstances;
+	for(int32 InstanceIndex = StartInstanceIndex; InstanceIndex < EndInstanceIndex; ++InstanceIndex)
+	{
+		FInstancedStaticMeshInstanceData& InstanceData = PerInstanceSMData[InstanceIndex];
+
+		// TODO: Computing LocalTransform is useless when we're updating the world location for the entire mesh.
+		// Should find some way around this for performance.
+
+		// Render data uses local transform of the instance
+		FTransform LocalTransform = bWorldSpace ? NewInstancesTransform.GetRelativeTransform(GetComponentTransform()) : NewInstancesTransform;
+		InstanceData.Transform = LocalTransform.ToMatrixWithScale();
+
+		if(bPhysicsStateCreated)
+		{
+			// Physics uses world transform of the instance
+			FTransform WorldTransform = bWorldSpace ? NewInstancesTransform : (LocalTransform * GetComponentTransform());
+			UpdateInstanceBodyTransform(InstanceIndex, WorldTransform, bTeleport);
+		}
+	}
+
+	// Request navigation update - Execute on a single index as it updates everything anyway
+	PartialNavigationUpdate(StartInstanceIndex);
+
+	// Force recreation of the render data when proxy is created
+	InstanceUpdateCmdBuffer.Edit();
+
+	if(bMarkRenderStateDirty)
+	{
+		MarkRenderStateDirty();
+	}
+
+	return true;
+}
+
 TArray<int32> UInstancedStaticMeshComponent::GetInstancesOverlappingSphere(const FVector& Center, float Radius, bool bSphereInWorldSpace) const
 {
 	TArray<int32> Result;
