HLOD Imposters
* Now extracted from sub levels and brought in the main map to be rendered as instanced static meshes.
* Impostors meshes created during proxy creation.

diff --git a/Engine/Source/Developer/HierarchicalLODUtilities/Private/HierarchicalLODUtilities.cpp b/Engine/Source/Developer/HierarchicalLODUtilities/Private/HierarchicalLODUtilities.cpp
--- a/Engine/Source/Developer/HierarchicalLODUtilities/Private/HierarchicalLODUtilities.cpp
+++ b/Engine/Source/Developer/HierarchicalLODUtilities/Private/HierarchicalLODUtilities.cpp
@@ -232,6 +232,107 @@ bool FHierarchicalLODUtilities::BuildStaticMeshForLODActor(ALODActor* LODActor,
 	return BuildStaticMeshForLODActor(LODActor, Proxy, LODSetup, InBaseMaterial);
 }
 
+static FString GetImposterMeshName(const UMaterialInterface* InImposterMaterial)
+{
+	UPackage* MaterialOuterMost = InImposterMaterial->GetOutermost();
+
+	const FString BaseName = FPackageName::GetShortName(MaterialOuterMost->GetPathName());
+	return FString::Printf(TEXT("%s_ImposterMesh"), *BaseName);
+}
+
+static FString GetImposterMeshPackageName(const UMaterialInterface* InImposterMaterial)
+{
+	UPackage* MaterialOuterMost = InImposterMaterial->GetOutermost();
+
+	const FString PathName = FPackageName::GetLongPackagePath(MaterialOuterMost->GetPathName());
+	const FString BaseName = FPackageName::GetShortName(MaterialOuterMost->GetPathName());
+	return FString::Printf(TEXT("%s/%s_ImposterMesh"), *PathName, *BaseName);
+}
+
+UPackage* CreateOrRetrieveImposterMeshPackage(const UMaterialInterface* InImposterMaterial)
+{
+	checkf(InImposterMaterial != nullptr, TEXT("Invalid material supplied"));
+
+	const FString MeshPackageName = GetImposterMeshPackageName(InImposterMaterial);
+
+	UPackage* MeshPackage = CreatePackage(NULL, *MeshPackageName);
+	MeshPackage->FullyLoad();
+
+	// Target filename
+	const FString MeshPackageFileName = FPackageName::LongPackageNameToFilename(MeshPackageName);
+	// This is a hack to avoid save file dialog when we will be saving imposter mesh package
+	MeshPackage->FileName = FName(*MeshPackageFileName);
+
+	return MeshPackage;
+}
+
+UStaticMesh* CreateImposterStaticMesh(UStaticMeshComponent* InComponent, UMaterialInterface* InMaterial, const FMeshProxySettings& InProxySettings)
+{
+	UPackage* MeshPackage = CreateOrRetrieveImposterMeshPackage(InMaterial);
+
+	// check if our asset exists
+	const FString ImposterMeshName = GetImposterMeshName(InMaterial);
+	UStaticMesh* StaticMesh = FindObject<UStaticMesh>(MeshPackage, *ImposterMeshName);
+	if (StaticMesh == nullptr)
+	{
+		// Create the UStaticMesh object.
+		StaticMesh = NewObject<UStaticMesh>(MeshPackage, *ImposterMeshName, RF_Public | RF_Standalone);
+		StaticMesh->InitResources();
+
+		// make sure it has a new lighting guid
+		StaticMesh->LightingGuid = FGuid::NewGuid();
+
+		// Set it to use textured lightmaps. Note that Build Lighting will do the error-checking (texcoordindex exists for all LODs, etc).
+		StaticMesh->LightMapResolution = InProxySettings.LightMapResolution;
+		StaticMesh->LightMapCoordinateIndex = 1;
+
+		// Add one LOD for the base mesh
+		StaticMesh->SetNumSourceModels(0);
+		FStaticMeshSourceModel& SrcModel = StaticMesh->AddSourceModel();
+		/*Don't allow the engine to recalculate normals*/
+		SrcModel.BuildSettings.bRecomputeNormals = false;
+		SrcModel.BuildSettings.bRecomputeTangents = false;
+		SrcModel.BuildSettings.bRemoveDegenerates = true;
+		SrcModel.BuildSettings.bUseHighPrecisionTangentBasis = false;
+		SrcModel.BuildSettings.bUseFullPrecisionUVs = false;
+		SrcModel.BuildSettings.bGenerateLightmapUVs = InProxySettings.bGenerateLightmapUVs;
+		SrcModel.BuildSettings.bBuildReversedIndexBuffer = false;
+		SrcModel.BuildSettings.bBuildAdjacencyBuffer = InProxySettings.bAllowAdjacency;
+		if (!InProxySettings.bAllowDistanceField)
+		{
+			SrcModel.BuildSettings.DistanceFieldResolutionScale = 0.0f;
+		}
+
+		FMeshDescription* ImposterMesh = StaticMesh->CreateMeshDescription(0);
+		const IMeshMergeUtilities& MeshMergeUtilities = FModuleManager::Get().LoadModuleChecked<IMeshMergeModule>("MeshMergeUtilities").GetUtilities();
+		MeshMergeUtilities.ExtractImposterToRawMesh(InComponent, *ImposterMesh);
+	
+		// Commit mesh description and materials list to static mesh
+		StaticMesh->CommitMeshDescription(0);
+		StaticMesh->StaticMaterials = { InMaterial };
+
+		//Set the Imported version before calling the build
+		StaticMesh->ImportVersion = EImportStaticMeshVersion::LastVersion;
+
+		StaticMesh->PostEditChange();
+
+		// Our imposters meshes are flat, but they actually represent a volume.
+		// Extend the imposteor bounds using the original mesh bounds.
+		if (StaticMesh->GetBoundingBox().GetVolume() == 0)
+		{
+			const FBox StaticMeshBox = StaticMesh->GetBoundingBox();
+			const FBox CombinedBox = StaticMeshBox + InComponent->GetStaticMesh()->GetBoundingBox();
+			StaticMesh->PositiveBoundsExtension = (CombinedBox.Max - StaticMeshBox.Max);
+			StaticMesh->NegativeBoundsExtension = (StaticMeshBox.Min - CombinedBox.Min);
+			StaticMesh->CalculateExtendedBounds();
+		}	
+	
+		StaticMesh->MarkPackageDirty();
+	}
+
+	return StaticMesh;
+}
+
 bool FHierarchicalLODUtilities::BuildStaticMeshForLODActor(ALODActor* LODActor, UHLODProxy* Proxy, const FHierarchicalSimplification& LODSetup, UMaterialInterface* InBaseMaterial)
 {
 	if (Proxy && LODActor)
@@ -250,6 +351,25 @@ bool FHierarchicalLODUtilities::BuildStaticMeshForLODActor(ALODActor* LODActor,
 		TArray<UPrimitiveComponent*> AllComponents;
 		UHLODProxy::ExtractComponents(LODActor, AllComponents);
 
+		TArray<UStaticMeshComponent*> AllImposters;
+		if (LODSetup.MergeSetting.bIncludeImposters)
+		{			
+			// Retrieve all imposters.
+			for (UPrimitiveComponent* Component : AllComponents)
+			{
+				if (UStaticMeshComponent* StaticMeshComponent = Cast<UStaticMeshComponent>(Component))
+				{
+					if (LODActor->GetImposterMaterial(StaticMeshComponent))
+					{
+						AllImposters.Add(StaticMeshComponent);
+					}
+				}
+			}
+
+			// Imposters won't be merged in the HLOD mesh
+			AllComponents.RemoveAll([&](UPrimitiveComponent* Component) { return AllImposters.Contains(Component); });
+		}
+
 		// it shouldn't even have come here if it didn't have any staticmesh
 		if (ensure(AllComponents.Num() > 0))
 		{
@@ -296,7 +416,6 @@ bool FHierarchicalLODUtilities::BuildStaticMeshForLODActor(ALODActor* LODActor,
 				Algo::Transform(AllComponents, StaticMeshComponents, [](UPrimitiveComponent* InPrimitiveComponent) { return Cast<UStaticMeshComponent>(InPrimitiveComponent); });
 
 				MeshMergeUtilities.CreateProxyMesh(StaticMeshComponents, ProxySettings, InBaseMaterial, AssetsOuter, PackageName, JobID, Processor->GetCallbackDelegate(), true, OverrideLODSetup.TransitionScreenSize);
-				return true;
 			}
 			else
 			{
@@ -385,9 +504,42 @@ bool FHierarchicalLODUtilities::BuildStaticMeshForLODActor(ALODActor* LODActor,
 
 				// Collect garbage to clean up old unreferenced data in the HLOD package
 				CollectGarbage(GARBAGE_COLLECTION_KEEPFLAGS);
+			}
 
-				return true;
+			// Add imposters
+			if (AllImposters.Num() > 0)
+			{
+				struct FLODImposterBatch
+				{
+					UStaticMesh*		StaticMesh;
+					TArray<FTransform>	Transforms;
+				};
+
+				// Get all meshes + transforms for all imposters type (per material)
+				TMap<UMaterialInterface*, FLODImposterBatch> ImposterBatches;
+				for(UStaticMeshComponent* Imposter : AllImposters)
+				{
+					UMaterialInterface* Material = LODActor->GetImposterMaterial(Imposter);
+					check(Material);
+
+					FLODImposterBatch& LODImposterBatch = ImposterBatches.FindOrAdd(Material);
+					LODImposterBatch.Transforms.Add(Imposter->GetOwner()->GetActorTransform());
+
+					// The static mesh hasn't been created yet, do it.
+					if(LODImposterBatch.StaticMesh == nullptr)
+					{
+						LODImposterBatch.StaticMesh = CreateImposterStaticMesh(Imposter, Material, LODSetup.ProxySetting);
+					}
+				}
+
+				// Add imposters to the LODActor
+				for(const TPair<UMaterialInterface*, FLODImposterBatch> ImposterBatch : ImposterBatches)
+				{
+					LODActor->SetupImposters(ImposterBatch.Key, ImposterBatch.Value.StaticMesh, ImposterBatch.Value.Transforms);
+				}
 			}
+
+			return true;
 		}
 	}
 	return false;
diff --git a/Engine/Source/Developer/MeshMergeUtilities/Private/MeshMergeHelpers.cpp b/Engine/Source/Developer/MeshMergeUtilities/Private/MeshMergeHelpers.cpp
--- a/Engine/Source/Developer/MeshMergeUtilities/Private/MeshMergeHelpers.cpp
+++ b/Engine/Source/Developer/MeshMergeUtilities/Private/MeshMergeHelpers.cpp
@@ -1277,6 +1277,13 @@ void FMeshMergeHelpers::AppendRawMesh(FMeshDescription& InTarget, const FMeshDes
 	}
 }
 
+void FMeshMergeHelpers::ExtractImposterToRawMesh(const UStaticMeshComponent* InImposterComponent, FMeshDescription& InRawMesh)
+{
+	// Retrieve imposter LOD mesh	
+	const int32 LODIndex = InImposterComponent->GetStaticMesh()->GetNumLODs() - 1;
+	UStaticMesh::RegisterMeshAttributes(InRawMesh);
+	FMeshMergeHelpers::RetrieveMesh(InImposterComponent->GetStaticMesh(), LODIndex, InRawMesh);
+}
 
 void FMeshMergeHelpers::MergeImpostersToRawMesh(TArray<const UStaticMeshComponent*> ImposterComponents, FMeshDescription& InRawMesh, const FVector& InPivot, int32 InBaseMaterialIndex, TArray<UMaterialInterface*>& OutImposterMaterials)
 {
diff --git a/Engine/Source/Developer/MeshMergeUtilities/Private/MeshMergeHelpers.h b/Engine/Source/Developer/MeshMergeUtilities/Private/MeshMergeHelpers.h
--- a/Engine/Source/Developer/MeshMergeUtilities/Private/MeshMergeHelpers.h
+++ b/Engine/Source/Developer/MeshMergeUtilities/Private/MeshMergeHelpers.h
@@ -80,4 +80,7 @@ public:
 
 	/** Merges imposter meshes into a raw mesh. */
 	static void MergeImpostersToRawMesh(TArray<const UStaticMeshComponent*> ImposterComponents, FMeshDescription& InRawMesh, const FVector& InPivot, int32 BaseMaterialIndex, TArray<UMaterialInterface*>& OutImposterMaterials);
+
+	/** Extract imposter meshes to a raw mesh. */
+	static void ExtractImposterToRawMesh(const UStaticMeshComponent* InImposterComponent, FMeshDescription& InRawMesh);
 };
diff --git a/Engine/Source/Developer/MeshMergeUtilities/Private/MeshMergeUtilities.cpp b/Engine/Source/Developer/MeshMergeUtilities/Private/MeshMergeUtilities.cpp
--- a/Engine/Source/Developer/MeshMergeUtilities/Private/MeshMergeUtilities.cpp
+++ b/Engine/Source/Developer/MeshMergeUtilities/Private/MeshMergeUtilities.cpp
@@ -1864,30 +1864,23 @@ void FMeshMergeUtilities::MergeComponentsToStaticMesh(const TArray<UPrimitiveCom
 			// Determine LOD to use for merging, either user specified or calculated index and ensure we clamp to the maximum LOD index for this adapter 
 			const int32 LODIndex = [&]()
 			{
+				int32 LowestDetailLOD = Adapter.GetNumberOfLODs() - 1;
+				if (Component->bUseMaxLODAsImposter && !InSettings.bIncludeImposters)
+				{
+					LowestDetailLOD = FMath::Max(0, LowestDetailLOD - 1);
+				}
+
 				switch (InSettings.LODSelectionType)
 				{
 				case EMeshLODSelectionType::SpecificLOD:
-					return FMath::Min(Adapter.GetNumberOfLODs() - 1, InSettings.SpecificLOD);
+					return FMath::Min(LowestDetailLOD, InSettings.SpecificLOD);
+
 				case EMeshLODSelectionType::CalculateLOD:
-				  {
-					  int32 Min = Adapter.GetNumberOfLODs() - 1;
-					  if (Component->bUseMaxLODAsImposter && !InSettings.bIncludeImposters)
-					  {
-						  Min = FMath::Max(0, Min - 1);
-					  }
-					  return FMath::Min(Min, Utilities->GetLODLevelForScreenSize(Component, FMath::Clamp(ScreenSize, 0.0f, 1.0f)));
-				  }
-					
-				default:
+					return FMath::Min(LowestDetailLOD, Utilities->GetLODLevelForScreenSize(Component, FMath::Clamp(ScreenSize, 0.0f, 1.0f)));
+
 				case EMeshLODSelectionType::LowestDetailLOD:
-				  {
-					  if (Component->bUseMaxLODAsImposter && (!InSettings.bIncludeImposters))
-					  {						
-						  return FMath::Max(0, Adapter.GetNumberOfLODs() - 2);											
-					  }
-  
-					  return Adapter.GetNumberOfLODs() - 1;
-				  }					
+				default:
+					return LowestDetailLOD;
 				}
 			}();
 
@@ -2609,6 +2602,12 @@ void FMeshMergeUtilities::MergeComponentsToStaticMesh(const TArray<UPrimitiveCom
 	}
 }
 
+void FMeshMergeUtilities::ExtractImposterToRawMesh(const UStaticMeshComponent* InImposterComponent, FMeshDescription& InImposterMesh) const
+{
+	check(InImposterComponent->bUseMaxLODAsImposter);
+	FMeshMergeHelpers::ExtractImposterToRawMesh(InImposterComponent, InImposterMesh);
+}
+
 void FMeshMergeUtilities::CreateMergedRawMeshes(FMeshMergeDataTracker& InDataTracker, const FMeshMergingSettings& InSettings, const TArray<UStaticMeshComponent*>& InStaticMeshComponentsToMerge, const TArray<UMaterialInterface*>& InUniqueMaterials, const TMap<UMaterialInterface*, UMaterialInterface*>& InCollapsedMaterialMap, const TMultiMap<FMeshLODKey, MaterialRemapPair>& InOutputMaterialsMap, bool bInMergeAllLODs, bool bInMergeMaterialData, const FVector& InMergedAssetPivot, TArray<FMeshDescription>& OutMergedRawMeshes) const
 {
 	if (bInMergeAllLODs)
diff --git a/Engine/Source/Developer/MeshMergeUtilities/Private/MeshMergeUtilities.h b/Engine/Source/Developer/MeshMergeUtilities/Private/MeshMergeUtilities.h
--- a/Engine/Source/Developer/MeshMergeUtilities/Private/MeshMergeUtilities.h
+++ b/Engine/Source/Developer/MeshMergeUtilities/Private/MeshMergeUtilities.h
@@ -34,6 +34,8 @@ public:
 	UE_DEPRECATED(4.20, "Use different signature containing BaseMaterial parameter")
 	virtual void CreateProxyMesh(const TArray<UStaticMeshComponent*>& InStaticMeshComps, const struct FMeshProxySettings& InMeshProxySettings, UPackage* InOuter, const FString& InProxyBasePackageName, const FGuid InGuid, const FCreateProxyDelegate& InProxyCreatedDelegate, const bool bAllowAsync = false, const float ScreenSize = 1.0f) const override;
 
+	virtual void ExtractImposterToRawMesh(const UStaticMeshComponent* InImposterComponent, FMeshDescription& OutImposterMesh) const override;
+
 	virtual bool IsValidBaseMaterial(const UMaterialInterface* InBaseMaterial, bool bShowToaster) const override;
 
 	virtual void RegisterExtension(IMeshMergeExtension* InExtension) override;
diff --git a/Engine/Source/Developer/MeshMergeUtilities/Public/IMeshMergeUtilities.h b/Engine/Source/Developer/MeshMergeUtilities/Public/IMeshMergeUtilities.h
--- a/Engine/Source/Developer/MeshMergeUtilities/Public/IMeshMergeUtilities.h
+++ b/Engine/Source/Developer/MeshMergeUtilities/Public/IMeshMergeUtilities.h
@@ -125,6 +125,14 @@ public:
 	*/
 	virtual void CreateProxyMesh(const TArray<UStaticMeshComponent*>& InComponents, const struct FMeshProxySettings& InMeshProxySettings, UMaterialInterface* InBaseMaterial, UPackage* InOuter, const FString& InProxyBasePackageName, const FGuid InGuid, const FCreateProxyDelegate& InProxyCreatedDelegate, const bool bAllowAsync = false, const float ScreenSize = 1.0f) const = 0;
 
+    /**
+	*	Retrieve the mesh description for an imposter static mesh component.
+	*
+	*	@param	InImposterComponent		The static mesh component whose static mesh lowest LOD is an imposter mesh.
+	*	@param	OutImposterMesh			The extracted mesh description.
+	*/
+	virtual void ExtractImposterToRawMesh(const UStaticMeshComponent* InImposterComponent, FMeshDescription& OutImposterMesh) const = 0;
+
 	/** Checks whether or not the give material is valid as a base for the Proxy Material. This as it requires certain Material Parameters to be included to get a desirable material */
 	virtual bool IsValidBaseMaterial(const UMaterialInterface* InBaseMaterial, bool bShowToaster) const = 0;
 	
diff --git a/Engine/Source/Editor/UnrealEd/Classes/Commandlets/ResavePackagesCommandlet.h b/Engine/Source/Editor/UnrealEd/Classes/Commandlets/ResavePackagesCommandlet.h
--- a/Engine/Source/Editor/UnrealEd/Classes/Commandlets/ResavePackagesCommandlet.h
+++ b/Engine/Source/Editor/UnrealEd/Classes/Commandlets/ResavePackagesCommandlet.h
@@ -204,6 +204,8 @@ protected:
 
 	bool CanCheckoutFile(const FString& Filename, FString& CheckedOutUser);
 
+	void CheckoutAndSavePackage(UPackage* Package, TArray<FString>& SublevelFilenames);
+
 	// Print out a message only if running in very verbose mode
 	void VerboseMessage(const FString& Message);
 
diff --git a/Engine/Source/Editor/UnrealEd/Private/Commandlets/ContentCommandlets.cpp b/Engine/Source/Editor/UnrealEd/Private/Commandlets/ContentCommandlets.cpp
--- a/Engine/Source/Editor/UnrealEd/Private/Commandlets/ContentCommandlets.cpp
+++ b/Engine/Source/Editor/UnrealEd/Private/Commandlets/ContentCommandlets.cpp
@@ -1339,6 +1339,31 @@ bool UResavePackagesCommandlet::CanCheckoutFile(const FString& Filename, FString
 	return bCanCheckout;
 }
 
+void UResavePackagesCommandlet::CheckoutAndSavePackage(UPackage* Package, TArray<FString>& SublevelFilenames)
+{
+	check(Package);
+
+	FString PackageFilename;
+	if (FPackageName::TryConvertLongPackageNameToFilename(Package->GetName(), PackageFilename, FPackageName::GetAssetPackageExtension()))
+	{
+		if (IFileManager::Get().FileExists(*PackageFilename))
+		{
+			if (CheckoutFile(PackageFilename, true))
+			{
+				SublevelFilenames.Add(PackageFilename);
+			}
+			SavePackageHelper(Package, PackageFilename);
+		}
+		else
+		{
+			SavePackageHelper(Package, PackageFilename);
+			if (CheckoutFile(PackageFilename, true))
+			{
+				SublevelFilenames.Add(PackageFilename);
+			}
+		}
+	}
+}
 
 void UResavePackagesCommandlet::PerformAdditionalOperations(class UWorld* World, bool& bSavePackage)
 {
@@ -1383,8 +1408,8 @@ void UResavePackagesCommandlet::PerformAdditionalOperations(class UWorld* World,
 		WorldContext.SetCurrentWorld(World);
 		GWorld = World;
 
-		TArray<FString> SublevelFilenames;
-		auto CheckOutLevelFile = [this,&bShouldProceedWithRebuild, &SublevelFilenames](ULevel* InLevel)
+		TArray<FString> CheckedOutPackagesFilenames;
+		auto CheckOutLevelFile = [this,&bShouldProceedWithRebuild, &CheckedOutPackagesFilenames](ULevel* InLevel)
 		{
 			if (InLevel && InLevel->MapBuildData)
 			{
@@ -1396,7 +1421,7 @@ void UResavePackagesCommandlet::PerformAdditionalOperations(class UWorld* World,
 					{
 						if (CheckoutFile(MapBuildDataPackageName))
 						{
-							SublevelFilenames.Add(MapBuildDataPackageName);
+							CheckedOutPackagesFilenames.Add(MapBuildDataPackageName);
 						}
 						else
 						{
@@ -1420,8 +1445,8 @@ void UResavePackagesCommandlet::PerformAdditionalOperations(class UWorld* World,
 			{
 				if (CanCheckoutFile(WorldPackageName, WorldPackageCheckedOutUser) || !bSkipCheckedOutFiles)
 				{
-				SublevelFilenames.Add(WorldPackageName);
-			}
+					CheckedOutPackagesFilenames.Add(WorldPackageName);
+				}
 				else 
 				{
 					bShouldProceedWithRebuild = false;
@@ -1432,7 +1457,7 @@ void UResavePackagesCommandlet::PerformAdditionalOperations(class UWorld* World,
 				// if we can't check out the main map or it's not up to date then we can't do the lighting rebuild at all!
 				if (CheckoutFile(WorldPackageName))
 				{
-					SublevelFilenames.Add(WorldPackageName);
+					CheckedOutPackagesFilenames.Add(WorldPackageName);
 
 					CheckOutLevelFile(World->PersistentLevel);
 				}
@@ -1469,8 +1494,8 @@ void UResavePackagesCommandlet::PerformAdditionalOperations(class UWorld* World,
 						FString CurrentlyCheckedOutUser;
 						if (CanCheckoutFile(StreamingLevelPackageFilename, CurrentlyCheckedOutUser) || !bSkipCheckedOutFiles)
 						{
-						SublevelFilenames.Add(StreamingLevelPackageFilename);
-					}
+							CheckedOutPackagesFilenames.Add(StreamingLevelPackageFilename);
+						}
 						else 
 						{
 							UE_LOG(LogContentCommandlet, Warning, TEXT("[REPORT] Skipping %s as it is checked out by %s"), *StreamingLevelPackageFilename, *CurrentlyCheckedOutUser);
@@ -1481,7 +1506,7 @@ void UResavePackagesCommandlet::PerformAdditionalOperations(class UWorld* World,
 						// check to see if we need to check this package out
 						if (CheckoutFile(StreamingLevelPackageFilename))
 						{
-							SublevelFilenames.Add(StreamingLevelPackageFilename);
+							CheckedOutPackagesFilenames.Add(StreamingLevelPackageFilename);
 						}
 						else
 						{
@@ -1586,44 +1611,22 @@ void UResavePackagesCommandlet::PerformAdditionalOperations(class UWorld* World,
 					GShaderCompilingManager->ProcessAsyncResults(false, false);
 				}
 
-				IHierarchicalLODUtilities* Utilities = Module.GetUtilities();
-				for (const ULevel* Level : GWorld->GetLevels())
+				// Get the list of packages needs to be saved.
+				TSet<UPackage*> PackagesToSave;
+				for(ULevel* Level : GWorld->GetLevels())
 				{
-					// Only meshes for clusters that are in a visible level
-					if (Level->bIsVisible)
+					if(Level->bIsVisible)
 					{
-						const int32 NumHLODLevels = Level->GetWorldSettings()->GetNumHierarchicalLODLevels();
-
-						for (int32 HLODIndex = 0; HLODIndex < NumHLODLevels; ++HLODIndex)
-						{
-							UPackage* HLODPackage = Utilities->CreateOrRetrieveLevelHLODPackage(Level, HLODIndex);
+						Builder.GetMeshesPackagesToSave(Level, PackagesToSave);
+					}
+				}
 
-							// skip packages we havent modified
-							if(HLODPackage->IsDirty())
-							{
-								FString HLODDataFilename;
-								if (FPackageName::TryConvertLongPackageNameToFilename(HLODPackage->GetName(), HLODDataFilename, FPackageName::GetAssetPackageExtension()))
-								{
-									if (IFileManager::Get().FileExists(*HLODDataFilename))
-									{
-										if (CheckoutFile(HLODDataFilename, true))
-										{
-											SublevelFilenames.Add(HLODDataFilename);
-										}
-
-										SavePackageHelper(HLODPackage, HLODDataFilename);
-									}
-									else
-									{
-										SavePackageHelper(HLODPackage, HLODDataFilename);
-										if (CheckoutFile(HLODDataFilename, true))
-										{
-											SublevelFilenames.Add(HLODDataFilename);
-										}
-									}
-								}
-							}
-						}
+				// Checkout and save each dirty package
+				for (UPackage* Package : PackagesToSave)
+				{
+					if (Package->IsDirty())
+					{
+						CheckoutAndSavePackage(Package, CheckedOutPackagesFilenames);
 					}
 				}
 			}
@@ -1653,37 +1656,14 @@ void UResavePackagesCommandlet::PerformAdditionalOperations(class UWorld* World,
 
 				FEditorDelegates::OnLightingBuildFailed.Remove(BuildFailedDelegateHandle);
 			}
-			auto SaveMapBuildData = [this, &SublevelFilenames](ULevel* InLevel)
+			auto SaveMapBuildData = [this, &CheckedOutPackagesFilenames](ULevel* InLevel)
 			{
 				if (InLevel && InLevel->MapBuildData && (bShouldBuildLighting || bShouldBuildHLOD || bShouldBuildReflectionCaptures) )
 				{
 					UPackage* MapBuildDataPackage = InLevel->MapBuildData->GetOutermost();
-					FString MapBuildDataPackageName = MapBuildDataPackage->GetName();
-
 					if (MapBuildDataPackage != InLevel->GetOutermost())
 					{
-						FString MapBuildDataFilename;
-
-						if (FPackageName::TryConvertLongPackageNameToFilename(MapBuildDataPackageName, MapBuildDataFilename, FPackageName::GetAssetPackageExtension()))
-						{
-							if (IFileManager::Get().FileExists(*MapBuildDataFilename))
-							{
-								if ( CheckoutFile(MapBuildDataFilename, true) )
-								{
-									SublevelFilenames.Add(MapBuildDataFilename);
-								}
-
-								SavePackageHelper(MapBuildDataPackage, MapBuildDataFilename); 
-							}
-							else
-							{
-								SavePackageHelper(MapBuildDataPackage, MapBuildDataFilename);
-								if (CheckoutFile(MapBuildDataFilename, true))
-								{
-									SublevelFilenames.Add(MapBuildDataFilename);
-								}
-							}
-						}
+						CheckoutAndSavePackage(MapBuildDataPackage, CheckedOutPackagesFilenames);
 					}
 				}
 			};
@@ -1698,7 +1678,7 @@ void UResavePackagesCommandlet::PerformAdditionalOperations(class UWorld* World,
 				{
 					FString StreamingLevelPackageFilename;
 					const FString StreamingLevelWorldAssetPackageName = NextStreamingLevel->GetWorldAssetPackageName();
-					if (FPackageName::DoesPackageExist(StreamingLevelWorldAssetPackageName, NULL, &StreamingLevelPackageFilename) && SublevelFilenames.Contains(StreamingLevelPackageFilename))
+					if (FPackageName::DoesPackageExist(StreamingLevelWorldAssetPackageName, NULL, &StreamingLevelPackageFilename) && CheckedOutPackagesFilenames.Contains(StreamingLevelPackageFilename))
 					{
 						UPackage* SubLevelPackage = NextStreamingLevel->GetLoadedLevel()->GetOutermost();
 						bool bSaveSubLevelPackage = true;
@@ -1751,16 +1731,16 @@ void UResavePackagesCommandlet::PerformAdditionalOperations(class UWorld* World,
 			ISourceControlProvider& SourceControlProvider = ISourceControlModule::Get().GetProvider();
 			
 			// revert all our packages
-			for (const auto& SublevelFilename : SublevelFilenames)
+			for (const auto& CheckedOutPackageFilename : CheckedOutPackagesFilenames)
 			{
-				SourceControlProvider.Execute(ISourceControlOperation::Create<FRevert>(), *SublevelFilename);
+				SourceControlProvider.Execute(ISourceControlOperation::Create<FRevert>(), *CheckedOutPackageFilename);
 			}
 		}
 		else
 		{
-			for(const auto& SublevelFilename : SublevelFilenames)
+			for(const auto& CheckedOutPackageFilename : CheckedOutPackagesFilenames)
 			{
-				FilesToSubmit.AddUnique(SublevelFilename);
+				FilesToSubmit.AddUnique(CheckedOutPackageFilename);
 			}
 
 			if(bShouldBuildHLOD && !bBuildingNonHLODData)
diff --git a/Engine/Source/Editor/UnrealEd/Private/HierarchicalLOD.cpp b/Engine/Source/Editor/UnrealEd/Private/HierarchicalLOD.cpp
--- a/Engine/Source/Editor/UnrealEd/Private/HierarchicalLOD.cpp
+++ b/Engine/Source/Editor/UnrealEd/Private/HierarchicalLOD.cpp
@@ -729,6 +729,83 @@ void FHierarchicalLODBuilder::BuildMeshesForLODActors(bool bForceAll)
 	}
 }
 
+void FHierarchicalLODBuilder::GetMeshesPackagesToSave(ULevel* InLevel, TSet<UPackage*>& InHLODPackagesToSave, const FString& PreviousLevelName /*= ""*/)
+{
+	const TArray<FHierarchicalSimplification>& BuildLODLevelSettings = InLevel->GetWorldSettings()->GetHierarchicalLODSetup();
+	UMaterialInterface* BaseMaterial = InLevel->GetWorldSettings()->GetHierarchicalLODBaseMaterial();
+	TArray<TArray<ALODActor*>> LODLevelActors;
+	LODLevelActors.AddDefaulted(BuildLODLevelSettings.Num());
+
+	if (InLevel->Actors.Num() > 0)
+	{
+		FHierarchicalLODUtilitiesModule& Module = FModuleManager::LoadModuleChecked<FHierarchicalLODUtilitiesModule>("HierarchicalLODUtilities");
+		IHierarchicalLODUtilities* Utilities = Module.GetUtilities();
+
+		// Retrieve LOD actors from the level
+		uint32 NumLODActors = 0;
+		for (int32 ActorId = 0; ActorId < InLevel->Actors.Num(); ++ActorId)
+		{
+			AActor* Actor = InLevel->Actors[ActorId];
+			if (Actor && Actor->IsA<ALODActor>())
+			{
+				ALODActor* LODActor = CastChecked<ALODActor>(Actor);
+
+				LODLevelActors[LODActor->LODLevel - 1].Add(LODActor);
+				NumLODActors++;
+			}
+		}
+
+		if (NumLODActors)
+		{
+			const int32 NumLODLevels = LODLevelActors.Num();
+			for (int32 LODIndex = 0; LODIndex < NumLODLevels; ++LODIndex)
+			{
+				UPackage* AssetsOuter = Utilities->RetrieveLevelHLODPackage(InLevel, LODIndex);
+				if (AssetsOuter)
+				{
+					InHLODPackagesToSave.Add(AssetsOuter);
+				}
+				// If we couldn't find the HLOD package, the level may have been renamed, 
+				// so we need to relocate our old HLOD package before saving it.
+				else if (PreviousLevelName.Len() != 0)
+				{
+					FString NewLevelName = InLevel->GetOutermost()->GetName();
+
+					FString OldHLODProxyName = Utilities->GetLevelHLODProxyName(PreviousLevelName, LODIndex);
+					UHLODProxy* OldHLODPProxy = FindObject<UHLODProxy>(nullptr, *OldHLODProxyName);
+					if (OldHLODPProxy)
+					{
+						FString NewHLODProxyName = Utilities->GetLevelHLODProxyName(NewLevelName, LODIndex);
+						FString OldHLODPackageName = FPackageName::ObjectPathToPackageName(OldHLODProxyName);
+						FString NewHLODPackageName = FPackageName::ObjectPathToPackageName(NewHLODProxyName);
+						UPackage* OldHLODPackage = FindObject<UPackage>(nullptr, *OldHLODPackageName);
+						if (OldHLODPackage)
+						{
+							OldHLODPProxy->Rename(*FPackageName::ObjectPathToObjectName(NewHLODProxyName), OldHLODPackage, REN_NonTransactional | REN_DontCreateRedirectors);
+							OldHLODPackage->Rename(*NewHLODPackageName, nullptr, REN_NonTransactional | REN_DontCreateRedirectors);
+
+							InHLODPackagesToSave.Add(OldHLODPackage);
+
+							// Mark the level package as dirty as we have changed export locations, and without a resave we will not pick up
+							// HLOD packages when reloaded.
+							InLevel->GetOutermost()->MarkPackageDirty();
+						}
+					}
+				}
+
+				// We might have created imposters static mesh packages during the HLOD generation, we must save them too.
+				for (ALODActor* LODActor : LODLevelActors[LODIndex])
+				{
+					for (UInstancedStaticMeshComponent* Component : LODActor->GetImpostersStaticMeshComponents())
+					{
+						InHLODPackagesToSave.Add(Component->GetStaticMesh()->GetOutermost());
+					}
+				}
+			}
+		}
+	}
+}
+
 void FHierarchicalLODBuilder::SaveMeshesForActors()
 {
 	TArray<UPackage*> LevelPackagesToSave;
@@ -756,79 +833,23 @@ void FHierarchicalLODBuilder::SaveMeshesForActors()
 	{
 		check(LevelPackagesToSave.Num() == OldLevelPackageNames.Num() && LevelPackagesToSave.Num() == Levels.Num());
 
-		TArray<UPackage*> HLODPackagesToSave;
-
+		TSet<UPackage*> HLODPackagesToSave;
 		for(int32 PackageIndex = 0; PackageIndex < LevelPackagesToSave.Num(); ++PackageIndex)
 		{
-			const ULevel* Level = Levels[PackageIndex];
-
-			HLODPackagesToSave.Add(Level->GetOutermost());
-
-			const TArray<FHierarchicalSimplification>& BuildLODLevelSettings = Level->GetWorldSettings()->GetHierarchicalLODSetup();
-			UMaterialInterface* BaseMaterial = Level->GetWorldSettings()->GetHierarchicalLODBaseMaterial();
-			TArray<TArray<ALODActor*>> LODLevelActors;
-			LODLevelActors.AddDefaulted(BuildLODLevelSettings.Num());
-
-			if (Level->Actors.Num() > 0)
+			FString PreviousLevelName; 
+			bool bLevelRenamed = bUnsavedLevel && LevelPackagesToSave[PackageIndex]->GetName() != OldLevelPackageNames[PackageIndex];
+			if (bLevelRenamed)
 			{
-				FHierarchicalLODUtilitiesModule& Module = FModuleManager::LoadModuleChecked<FHierarchicalLODUtilitiesModule>("HierarchicalLODUtilities");
-				IHierarchicalLODUtilities* Utilities = Module.GetUtilities();
-
-				// Retrieve LOD actors from the level
-				uint32 NumLODActors = 0;
-				for (int32 ActorId = 0; ActorId < Level->Actors.Num(); ++ActorId)
-				{
-					AActor* Actor = Level->Actors[ActorId];
-					if (Actor && Actor->IsA<ALODActor>())
-					{
-						ALODActor* LODActor = CastChecked<ALODActor>(Actor);
-						
-						LODLevelActors[LODActor->LODLevel - 1].Add(LODActor);
-						NumLODActors++;
-					}
-				}		
-
-				if (NumLODActors)
-				{
-					const int32 NumLODLevels = LODLevelActors.Num();
-					for (int32 LODIndex = 0; LODIndex < NumLODLevels; ++LODIndex)
-					{
-						UPackage* AssetsOuter = Utilities->RetrieveLevelHLODPackage(Level, LODIndex);
-						if(AssetsOuter)
-						{
-							HLODPackagesToSave.Add(AssetsOuter);
-						}
-						// If we couldn't find the HLOD package, the level may have been renamed while saving above, 
-						// so we need to relocate our old HLOD package before saving it.
-						else if(bUnsavedLevel && LevelPackagesToSave[PackageIndex]->GetName() != OldLevelPackageNames[PackageIndex])
-						{
-							FString OldHLODProxyName = Utilities->GetLevelHLODProxyName(OldLevelPackageNames[PackageIndex], LODIndex);
-							UHLODProxy* OldHLODPProxy = FindObject<UHLODProxy>(nullptr, *OldHLODProxyName);
-							if(OldHLODPProxy)
-							{
-								FString NewHLODProxyName = Utilities->GetLevelHLODProxyName(LevelPackagesToSave[PackageIndex]->GetName(), LODIndex);
-								FString OldHLODPackageName = FPackageName::ObjectPathToPackageName(OldHLODProxyName);
-								FString NewHLODPackageName = FPackageName::ObjectPathToPackageName(NewHLODProxyName);
-								UPackage* OldHLODPackage = FindObject<UPackage>(nullptr, *OldHLODPackageName);
-								if(OldHLODPackage)
-								{
-									OldHLODPProxy->Rename(*FPackageName::ObjectPathToObjectName(NewHLODProxyName), OldHLODPackage, REN_NonTransactional | REN_DontCreateRedirectors);
-									OldHLODPackage->Rename(*NewHLODPackageName, nullptr, REN_NonTransactional | REN_DontCreateRedirectors);
-									
-									HLODPackagesToSave.Add(OldHLODPackage);
-
-									// Mark the level package as dirty as we have changed export locations, and without a resave we will not pick up
-									// HLOD packages when reloaded.
-									Level->GetOutermost()->MarkPackageDirty();
-								}
-							}
-						}
-					}
-				}
+				PreviousLevelName = OldLevelPackageNames[PackageIndex];
 			}
+
+			ULevel* Level = Levels[PackageIndex];
+			HLODPackagesToSave.Add(Level->GetOutermost());
+			GetMeshesPackagesToSave(Level, HLODPackagesToSave, PreviousLevelName);
 		}
 
-		UEditorLoadingAndSavingUtils::SavePackagesWithDialog(HLODPackagesToSave, true);
+		TArray<UPackage*> PackagesToSave = HLODPackagesToSave.Array();
+		UEditorLoadingAndSavingUtils::SavePackagesWithDialog(PackagesToSave, true);
 	}
 }
 
diff --git a/Engine/Source/Editor/UnrealEd/Public/HierarchicalLOD.h b/Engine/Source/Editor/UnrealEd/Public/HierarchicalLOD.h
--- a/Engine/Source/Editor/UnrealEd/Public/HierarchicalLOD.h
+++ b/Engine/Source/Editor/UnrealEd/Public/HierarchicalLOD.h
@@ -98,6 +98,9 @@ struct UNREALED_API FHierarchicalLODBuilder
 	/** Saves HLOD meshes for actors in all the World's levels */
 	void SaveMeshesForActors();
 
+	/** Get the list of mesh packages to save for a given level */
+	void GetMeshesPackagesToSave(ULevel* InLevel, TSet<UPackage*>& InHLODPackagesToSave, const FString& PreviousLevelName = "");
+
 	/** 
 	 * @param	bInForce	Whether to force the recalculation of this actor's build flag. If this is false then the cached flag is used an only recalculated every so often.
 	 * @return whether a build is needed (i.e. any LOD actors are dirty) 
diff --git a/Engine/Source/Runtime/Engine/Classes/Components/PrimitiveComponent.h b/Engine/Source/Runtime/Engine/Classes/Components/PrimitiveComponent.h
--- a/Engine/Source/Runtime/Engine/Classes/Components/PrimitiveComponent.h
+++ b/Engine/Source/Runtime/Engine/Classes/Components/PrimitiveComponent.h
@@ -216,18 +216,18 @@ public:
 	ELightmapType LightmapType;
 
 #if WITH_EDITORONLY_DATA
+	/** Which specific HLOD levels this component should be excluded from */
+	UPROPERTY(EditAnywhere, AdvancedDisplay, BlueprintReadWrite, Category = HLOD)
+	TArray<int32> ExcludeForSpecificHLODLevels;
+
 	/** If true, and if World setting has bEnableHierarchicalLOD equal to true, then this component will be included when generating a Proxy mesh for the parent Actor */
 	UPROPERTY(EditAnywhere, AdvancedDisplay, Category = HLOD, meta = (DisplayName = "Include Component for HLOD Mesh generation"))
 	uint8 bEnableAutoLODGeneration : 1;
+#endif 
 
 	/** Use the Maximum LOD Mesh (imposter) instead of including Mesh data from this component in the Proxy Generation process */
 	UPROPERTY(EditAnywhere, AdvancedDisplay, Category = HLOD)
-	uint8 bUseMaxLODAsImposter: 1;
-
-	/** Which specific HLOD levels this component should be excluded from */
-	UPROPERTY(EditAnywhere, AdvancedDisplay, BlueprintReadWrite, Category = HLOD)
-	TArray<int32> ExcludeForSpecificHLODLevels;
-#endif 
+	uint8 bUseMaxLODAsImposter : 1;
 
 	/**
 	 * When enabled this object will not be culled by distance. This is ignored if a child of a HLOD.
diff --git a/Engine/Source/Runtime/Engine/Classes/Engine/LODActor.h b/Engine/Source/Runtime/Engine/Classes/Engine/LODActor.h
--- a/Engine/Source/Runtime/Engine/Classes/Engine/LODActor.h
+++ b/Engine/Source/Runtime/Engine/Classes/Engine/LODActor.h
@@ -8,6 +8,7 @@
 #include "Engine/MaterialMerging.h"
 #include "GameFramework/Actor.h"
 #include "Components/StaticMeshComponent.h"
+#include "Components/InstancedStaticMeshComponent.h"
 #include "LODActor.generated.h"
 
 class UStaticMesh;
@@ -35,6 +36,10 @@ private:
 	UPROPERTY(Category=LODActor, VisibleAnywhere)
 	UStaticMeshComponent* StaticMeshComponent;
 
+	/** Imposters to be rendered as instanced static meshes */
+	UPROPERTY(Category=LODActor, VisibleAnywhere)
+	TMap<const UMaterialInterface*, UInstancedStaticMeshComponent*> ImpostersStaticMeshComponents;
+
 	/** The mesh proxy used to display this LOD */
 	UPROPERTY(Category=LODActor, VisibleAnywhere)
 	UHLODProxy* Proxy;
@@ -81,6 +86,9 @@ public:
 	/** Sets StaticMesh and IsPreviewActor to true if InStaticMesh equals nullptr */
 	void SetStaticMesh(UStaticMesh* InStaticMesh);
 
+	/** Add imposters instances to this LODActor. */
+	void SetupImposters(UMaterialInterface* InImposterMaterial, UStaticMesh* InStaticMesh, const TArray<FTransform>& InTransforms);
+
 	/** Sets the LOD draw distance and updates the Static Mesh Component's min drawing distance */
 	void SetDrawDistance(float InDistance);
 
@@ -107,7 +115,7 @@ public:
 	void ForceUnbuilt();
 
 	/**
-	* Adds InAcor to the SubActors array and set its LODParent to this
+	* Adds InActor to the SubActors array and set its LODParent to this
 	* @param InActor - Actor to add
 	*/
 	void AddSubActor(AActor* InActor);
@@ -118,7 +126,6 @@ public:
 	*/
 	const bool RemoveSubActor(AActor* InActor);
 
-
 	/**
 	 * Determines whether or not this LODActor has valid SubActors and can be built
 	 * @return true if the subactor(s) contain at least two static mesh components
@@ -218,6 +225,9 @@ public:
 	/** Returns StaticMeshComponent subobject **/
 	UStaticMeshComponent* GetStaticMeshComponent() const { return StaticMeshComponent; }
 
+	/** Returns instanced static mesh components that represents the imposters found in this LODActor. */
+	TArray<UInstancedStaticMeshComponent*> GetImpostersStaticMeshComponents() const { TArray<UInstancedStaticMeshComponent*> ISM; ImpostersStaticMeshComponents.GenerateValueArray(ISM); return ISM; }
+
 private:
 	// Called when CVars are changed to check to see if the maximum HLOD level value has changed
 	static void OnCVarsChanged();
@@ -227,10 +237,32 @@ private:
 	// Updates the transition distance according to values (if) set in r.HLOD.DistanceOverride
 	void UpdateOverrideTransitionDistance();
 
-
 	// Called to make sure autoregistration/manual registration state matches based on the LOD override cvar and this actor's lod level
 	void UpdateRegistrationToMatchMaximumLODLevel();
 
+	// Setup a LOD static mesh component.
+	void SetupComponent(UStaticMeshComponent* InComponent);
+
+	// Utility methods to act on all static mesh components owned by this actor.
+	void SetComponentsMinDrawDistance(float InMinDrawDistance, bool bInMarkRenderStateDirty);
+	void RegisterMeshComponents();
+	void UnregisterMeshComponents();
+
+#if !(UE_BUILD_SHIPPING || UE_BUILD_TEST)
+	// Get/Create the LOD instanced static mesh component for a given imposter material.
+	UInstancedStaticMeshComponent* GetLODComponentForMaterial(const UMaterialInterface* InMaterial) const;
+	UInstancedStaticMeshComponent* GetOrCreateLODComponentForMaterial(const UMaterialInterface* InMaterial);
+
+	// Get/Create the LOD static mesh component to use for a given actor.
+	UStaticMeshComponent* GetLODComponentForActor(const AActor* InActor, bool bFallbackToDefault = true) const;
+	UStaticMeshComponent* GetOrCreateLODComponentForActor(const AActor* InActor);
+
+public:
+	// Get the imposter material (if any) of an actor or SMC
+	UMaterialInterface* GetImposterMaterial(const AActor* InActor) const;
+	UMaterialInterface* GetImposterMaterial(const UStaticMeshComponent* InComponent) const;
+#endif
+
 private:
  	// Have we already tried to register components? (a cache to avoid having to query the owning world when the global HLOD max level setting is changed)
  	uint8 bHasActorTriedToRegisterComponents : 1;
diff --git a/Engine/Source/Runtime/Engine/Private/LODActor.cpp b/Engine/Source/Runtime/Engine/Private/LODActor.cpp
--- a/Engine/Source/Runtime/Engine/Private/LODActor.cpp
+++ b/Engine/Source/Runtime/Engine/Private/LODActor.cpp
@@ -23,6 +23,7 @@
 #if WITH_EDITOR
 #include "Editor.h"
 #include "HierarchicalLODUtilitiesModule.h"
+#include "IHierarchicalLODUtilities.h"
 #include "ObjectTools.h"
 #include "HierarchicalLOD.h"
 #endif
@@ -175,11 +176,6 @@ ALODActor::ALODActor(const FObjectInitializer& ObjectInitializer)
 {
 	bCanBeDamaged = false;
 
-	// Cast shadows if any sub-actors do
-	bool bCastsShadow = false;
-	bool bCastsStaticShadow = false;
-	bool bCastsDynamicShadow = false;
-
 	PrimaryActorTick.bCanEverTick = true;
 	PrimaryActorTick.bStartWithTickEnabled = false;
 	PrimaryActorTick.bAllowTickOnDedicatedServer = false;
@@ -195,14 +191,8 @@ ALODActor::ALODActor(const FObjectInitializer& ObjectInitializer)
 #endif // WITH_EDITORONLY_DATA
 
 	StaticMeshComponent = CreateDefaultSubobject<UStaticMeshComponent>(TEXT("StaticMeshComponent0"));
-	StaticMeshComponent->SetCollisionProfileName(UCollisionProfile::NoCollision_ProfileName);
-	StaticMeshComponent->Mobility = EComponentMobility::Static;
-	StaticMeshComponent->SetGenerateOverlapEvents(false);
-	StaticMeshComponent->CastShadow = bCastsShadow;
-	StaticMeshComponent->bCastStaticShadow = bCastsStaticShadow;
-	StaticMeshComponent->bCastDynamicShadow = bCastsDynamicShadow;
-	StaticMeshComponent->bAllowCullDistanceVolume = false;
-	StaticMeshComponent->bNeverDistanceCull = true;
+	SetupComponent(StaticMeshComponent);
+
 	bNeedsDrawDistanceReset = false;
 	bHasPatchedUpParent = false;
 	ResetDrawDistanceTime = 0.0f;
@@ -215,6 +205,25 @@ ALODActor::ALODActor(const FObjectInitializer& ObjectInitializer)
 #endif
 }
 
+void ALODActor::SetupComponent(UStaticMeshComponent* Component)
+{
+	// Cast shadows if any sub-actors do
+	bool bCastsShadow = false;
+	bool bCastsStaticShadow = false;
+	bool bCastsDynamicShadow = false;
+
+	Component->SetCollisionProfileName(UCollisionProfile::NoCollision_ProfileName);
+	Component->Mobility = EComponentMobility::Static;
+	Component->SetGenerateOverlapEvents(false);
+	Component->CastShadow = bCastsShadow;
+	Component->bCastStaticShadow = bCastsStaticShadow;
+	Component->bCastDynamicShadow = bCastsDynamicShadow;
+	Component->bAllowCullDistanceVolume = false;
+	Component->bNeverDistanceCull = true;
+
+	Component->MinDrawDistance = LODDrawDistance;
+}
+
 FString ALODActor::GetDetailedInfoInternal() const
 {
 	return StaticMeshComponent ? StaticMeshComponent->GetDetailedInfoInternal() : TEXT("No_StaticMeshComponent");
@@ -223,7 +232,7 @@ FString ALODActor::GetDetailedInfoInternal() const
 void ALODActor::PostLoad()
 {
 	Super::PostLoad();
-	StaticMeshComponent->MinDrawDistance = LODDrawDistance;
+	SetComponentsMinDrawDistance(LODDrawDistance, false);
 	StaticMeshComponent->bCastDynamicShadow = false;	
 	UpdateRegistrationToMatchMaximumLODLevel();
 
@@ -241,6 +250,10 @@ void ALODActor::PostLoad()
 			const float ScreenHeight = 1080.0f;
 			const FPerspectiveMatrix ProjMatrix(HalfFOV, ScreenWidth, ScreenHeight, 1.0f);
 			FBoxSphereBounds Bounds = GetStaticMeshComponent()->CalcBounds(FTransform());
+			for (const TPair<const UMaterialInterface*, UInstancedStaticMeshComponent*>& Component : ImpostersStaticMeshComponents)
+			{
+				Bounds = Bounds + Component.Value->CalcBounds(FTransform());
+			}
 
 			// legacy transition screen size was previously a screen AREA fraction using resolution-scaled values, so we need to convert to distance first to correctly calculate the threshold
 			const float ScreenArea = TransitionScreenSize * (ScreenWidth * ScreenHeight);
@@ -288,6 +301,26 @@ void ALODActor::PostLoad()
 	UpdateOverrideTransitionDistance();
 }
 
+void ALODActor::SetComponentsMinDrawDistance(float InMinDrawDistance, bool bInMarkRenderStateDirty)
+{
+	float MinDrawDistance = FMath::Max(0.0f, InMinDrawDistance);
+
+	StaticMeshComponent->MinDrawDistance = MinDrawDistance;
+	if (bInMarkRenderStateDirty)
+	{
+		StaticMeshComponent->MarkRenderStateDirty();
+	}
+
+	for (const TPair<const UMaterialInterface*, UInstancedStaticMeshComponent*>& Component : ImpostersStaticMeshComponents)
+	{
+		Component.Value->MinDrawDistance = MinDrawDistance;
+		if (bInMarkRenderStateDirty)
+		{
+			Component.Value->MarkRenderStateDirty();
+		}
+	}
+}
+
 void ALODActor::UpdateOverrideTransitionDistance()
 {
 	const int32 NumDistances = ALODActor::HLODDistances.Num();
@@ -311,9 +344,8 @@ void ALODActor::UpdateOverrideTransitionDistance()
 
 	if (DistanceIndex != INDEX_NONE)
 	{
-		StaticMeshComponent->MinDrawDistance = (!HLODDistances.IsValidIndex(DistanceIndex) || FMath::IsNearlyZero(HLODDistances[DistanceIndex])) ? LODDrawDistance : ALODActor::HLODDistances[DistanceIndex];
-		StaticMeshComponent->MinDrawDistance = FMath::Max(0.0f, StaticMeshComponent->MinDrawDistance);
-		StaticMeshComponent->MarkRenderStateDirty();
+		float MinDrawDistance = (!HLODDistances.IsValidIndex(DistanceIndex) || FMath::IsNearlyZero(HLODDistances[DistanceIndex])) ? LODDrawDistance : ALODActor::HLODDistances[DistanceIndex];
+		SetComponentsMinDrawDistance(MinDrawDistance, true);
 	}
 }
 
@@ -363,8 +395,7 @@ void ALODActor::Tick(float DeltaSeconds)
 				MinDrawDistance = HLODDistanceOverride;
 			}
 
-			StaticMeshComponent->MinDrawDistance = FMath::Max(0.0f, MinDrawDistance);
-			StaticMeshComponent->MarkRenderStateDirty();
+			SetComponentsMinDrawDistance(MinDrawDistance, true);
 			bNeedsDrawDistanceReset = false;
 			ResetDrawDistanceTime = 0.0f;
 			PrimaryActorTick.SetTickFunctionEnable(false);
@@ -379,8 +410,7 @@ void ALODActor::Tick(float DeltaSeconds)
 
 void ALODActor::PauseDitherTransition()
 {
-	StaticMeshComponent->MinDrawDistance = 0.0f;
-	StaticMeshComponent->MarkRenderStateDirty();
+	SetComponentsMinDrawDistance(0.0f, true);
 	bNeedsDrawDistanceReset = true;
 	ResetDrawDistanceTime = 0.0f;
 }
@@ -406,14 +436,14 @@ void ALODActor::UpdateRegistrationToMatchMaximumLODLevel()
 		if (!bAllowShowingThisLevel && StaticMeshComponent->IsRegistered())
 		{
 			ensure(bHasActorTriedToRegisterComponents);
-			StaticMeshComponent->UnregisterComponent();
+			UnregisterMeshComponents();
 		}
 		else if (bAllowShowingThisLevel && !StaticMeshComponent->IsRegistered())
 		{
 			// We should only register components if the actor had already tried to register before (otherwise it'll be taken care of in the normal flow)
 			if (bHasActorTriedToRegisterComponents)
 			{
-				StaticMeshComponent->RegisterComponent();
+				RegisterMeshComponents();
 			}
 		}
 	}
@@ -428,7 +458,7 @@ void ALODActor::PostRegisterAllComponents()
 	// In case we patched up the subactors to a parent LOD actor, we can unregister this component as it's not used anymore
 	if (bHasPatchedUpParent)
 	{
-		StaticMeshComponent->UnregisterComponent();
+		UnregisterMeshComponents();
 	}
 
 #if WITH_EDITOR
@@ -442,10 +472,50 @@ void ALODActor::PostRegisterAllComponents()
 #endif
 }
 
+void ALODActor::RegisterMeshComponents()
+{
+	if (!StaticMeshComponent->IsRegistered())
+	{
+		StaticMeshComponent->RegisterComponent();
+	}
+
+	for (const TPair<const UMaterialInterface*, UInstancedStaticMeshComponent*>& Component : ImpostersStaticMeshComponents)
+	{
+		if (!Component.Value->IsRegistered())
+		{
+			Component.Value->RegisterComponent();
+		}
+	}
+}
+
+void ALODActor::UnregisterMeshComponents()
+{
+	if (StaticMeshComponent->IsRegistered())
+	{
+		StaticMeshComponent->UnregisterComponent();
+	}
+	else
+	{
+		StaticMeshComponent->bAutoRegister = false;
+	}
+
+	for (const TPair<const UMaterialInterface*, UInstancedStaticMeshComponent*>& Component : ImpostersStaticMeshComponents)
+	{
+		if (Component.Value->IsRegistered())
+		{
+			Component.Value->UnregisterComponent();
+		}
+		else
+		{
+			Component.Value->bAutoRegister = false;
+		}
+	}
+}
+
 void ALODActor::SetDrawDistance(float InDistance)
 {
 	LODDrawDistance = InDistance;
-	StaticMeshComponent->MinDrawDistance = LODDrawDistance;
+	SetComponentsMinDrawDistance(LODDrawDistance, false);
 }
 
 #if !(UE_BUILD_SHIPPING || UE_BUILD_TEST)
@@ -454,30 +524,34 @@ const bool ALODActor::IsBuilt(bool bInForce/*=false*/) const
 {
 	auto IsBuiltHelper = [this]()
 	{
-		// No static mesh
-		if(StaticMeshComponent->GetStaticMesh() == nullptr)
+		// Ensure all subactors are linked to a LOD static mesh component.
+		for (AActor* SubActor : SubActors)
 		{
-			return false;
+			UStaticMeshComponent* LODComponent = GetLODComponentForActor(SubActor, false);
+			if (LODComponent == nullptr || LODComponent->GetStaticMesh() == nullptr)
+			{
+				return false;
+			}
 		}
 
 		// No proxy mesh
-		if(Proxy == nullptr)
+		if (Proxy == nullptr)
 		{
 			return false;
 		}
 
 		// Mismatched key
-		if(!Proxy->ContainsDataForActor(this))
+		if (!Proxy->ContainsDataForActor(this))
 		{
 			return false;
 		}
 
 		// Unbuilt children
-		for(AActor* SubActor : SubActors)
+		for (AActor* SubActor : SubActors)
 		{
-			if(ALODActor* SubLODActor = Cast<ALODActor>(SubActor))
+			if (ALODActor* SubLODActor = Cast<ALODActor>(SubActor))
 			{
-				if(!SubLODActor->IsBuilt(true))
+				if (!SubLODActor->IsBuilt(true))
 				{
 					return false;
 				}
@@ -488,7 +562,7 @@ const bool ALODActor::IsBuilt(bool bInForce/*=false*/) const
 	};
 
 	const double CurrentTime = FPlatformTime::Seconds();
-	if(bInForce || (CurrentTime - LastIsBuiltTime > 0.5))
+	if (bInForce || (CurrentTime - LastIsBuiltTime > 0.5))
 	{
 		bCachedIsBuilt = IsBuiltHelper();
 		LastIsBuiltTime = CurrentTime;
@@ -515,7 +589,7 @@ const bool ALODActor::HasValidLODChildren() const
 				{
 					if (const UPrimitiveComponent* ParentPrimitiveComponent = PrimitiveComponent ? PrimitiveComponent->GetLODParentPrimitive() : nullptr)
 					{
-						if (ParentPrimitiveComponent && ParentPrimitiveComponent == GetStaticMeshComponent())
+						if (ParentPrimitiveComponent && GetComponents().Contains(const_cast<UPrimitiveComponent*>(ParentPrimitiveComponent)))
 						{
 							return true;
 						}
@@ -554,11 +628,11 @@ void ALODActor::PostEditChangeProperty(FPropertyChangedEvent& PropertyChangedEve
 	
 	if (PropertyName == GET_MEMBER_NAME_CHECKED(ALODActor, bOverrideTransitionScreenSize) || PropertyName == GET_MEMBER_NAME_CHECKED(ALODActor, TransitionScreenSize))
 	{
-		float CalculateSreenSize = 0.0f;
+		float CalculateScreenSize = 0.0f;
 
 		if (bOverrideTransitionScreenSize)
 		{
-			CalculateSreenSize = TransitionScreenSize;
+			CalculateScreenSize = TransitionScreenSize;
 		}
 		else
 		{
@@ -566,10 +640,10 @@ void ALODActor::PostEditChangeProperty(FPropertyChangedEvent& PropertyChangedEve
 			check(World != nullptr);
 			const TArray<struct FHierarchicalSimplification>& HierarchicalLODSetups = World->GetWorldSettings()->GetHierarchicalLODSetup();
 			checkf(HierarchicalLODSetups.IsValidIndex(LODLevel - 1), TEXT("Out of range HLOD level (%i) found in LODActor (%s)"), LODLevel - 1, *GetName());
-			CalculateSreenSize = HierarchicalLODSetups[LODLevel - 1].TransitionScreenSize;
+			CalculateScreenSize = HierarchicalLODSetups[LODLevel - 1].TransitionScreenSize;
 		}
 
-		RecalculateDrawingDistance(CalculateSreenSize);
+		RecalculateDrawingDistance(CalculateScreenSize);
 	}
 
 	UpdateRegistrationToMatchMaximumLODLevel();
@@ -664,7 +738,9 @@ void ALODActor::EditorApplyMirror(const FVector& MirrorScale, const FVector& Piv
 void ALODActor::AddSubActor(AActor* InActor)
 {
 	SubActors.Add(InActor);
-	InActor->SetLODParent(StaticMeshComponent, LODDrawDistance);
+
+	UStaticMeshComponent* LODComponent = GetOrCreateLODComponentForActor(InActor);
+	InActor->SetLODParent(LODComponent, LODDrawDistance);
 
 	// Adding number of triangles
 	if (!InActor->IsA<ALODActor>())
@@ -685,7 +761,6 @@ void ALODActor::AddSubActor(AActor* InActor)
 	{
 		ALODActor* LODActor = Cast<ALODActor>(InActor);
 		NumTrianglesInSubActors += LODActor->GetNumTrianglesInSubActors();
-		
 	}
 	
 	// Reset the shadowing flags and determine them according to our current sub actors
@@ -741,35 +816,47 @@ const bool ALODActor::RemoveSubActor(AActor* InActor)
 void ALODActor::DetermineShadowingFlags()
 {
 	// Cast shadows if any sub-actors do
-	bool bCastsShadow = false;
-	bool bCastsStaticShadow = false;
-	bool bCastsDynamicShadow = false;
-	bool bCastFarShadow = false;
+	StaticMeshComponent->CastShadow = false;
+	StaticMeshComponent->bCastStaticShadow = false;
+	StaticMeshComponent->bCastDynamicShadow = false;
+	StaticMeshComponent->bCastFarShadow = false;
+	StaticMeshComponent->MarkRenderStateDirty();
+
+	for (const TPair<const UMaterialInterface*, UInstancedStaticMeshComponent*>& Component : ImpostersStaticMeshComponents)
+	{
+		Component.Value->CastShadow = false;
+		Component.Value->bCastStaticShadow = false;
+		Component.Value->bCastDynamicShadow = false;
+		Component.Value->bCastFarShadow = false;
+		Component.Value->MarkRenderStateDirty();
+	}
+
 	for (AActor* Actor : SubActors)
 	{
 		if (Actor)
 		{
+			UStaticMeshComponent* LODComponent = GetLODComponentForActor(Actor);
+
 			TArray<UStaticMeshComponent*> StaticMeshComponents;
 			Actor->GetComponents<UStaticMeshComponent>(StaticMeshComponents);
 			for (UStaticMeshComponent* Component : StaticMeshComponents)
 			{
-				bCastsShadow |= Component->CastShadow;
-				bCastsStaticShadow |= Component->bCastStaticShadow;
-				bCastsDynamicShadow |= Component->bCastDynamicShadow;
-				bCastFarShadow |= Component->bCastFarShadow;
+				LODComponent->CastShadow |= Component->CastShadow;
+				LODComponent->bCastStaticShadow |= Component->bCastStaticShadow;
+				LODComponent->bCastDynamicShadow |= Component->bCastDynamicShadow;
+				LODComponent->bCastFarShadow |= Component->bCastFarShadow;
 			}
 		}
 	}
-
-	StaticMeshComponent->CastShadow = bCastsShadow;
-	StaticMeshComponent->bCastStaticShadow = bCastsStaticShadow;
-	StaticMeshComponent->bCastDynamicShadow = bCastsDynamicShadow;
-	StaticMeshComponent->bCastFarShadow = bCastFarShadow;
-	StaticMeshComponent->MarkRenderStateDirty();
 }
 
 const bool ALODActor::HasValidSubActors() const
 {
+#if WITH_EDITOR
+	FHierarchicalLODUtilitiesModule& Module = FModuleManager::LoadModuleChecked<FHierarchicalLODUtilitiesModule>("HierarchicalLODUtilities");
+	IHierarchicalLODUtilities* Utilities = Module.GetUtilities();
+#endif
+
 	int32 NumMeshes = 0;
 
 	// Make sure there is at least one mesh in the subactors
@@ -777,12 +864,7 @@ const bool ALODActor::HasValidSubActors() const
 	{
 		if (SubActor)
 		{
-#if WITH_EDITOR
-			FHierarchicalLODUtilitiesModule& Module = FModuleManager::LoadModuleChecked<FHierarchicalLODUtilitiesModule>("HierarchicalLODUtilities");
-			IHierarchicalLODUtilities* Utilities = Module.GetUtilities();
-#endif
-
-			for (UActorComponent* Comp : GetComponents())
+			for (UActorComponent* Comp : SubActor->GetComponents())
 			{
 				if (UStaticMeshComponent* Component = Cast<UStaticMeshComponent>(Comp))
 				{
@@ -815,15 +897,13 @@ const bool ALODActor::HasAnySubActors() const
 void ALODActor::ToggleForceView()
 {
 	// Toggle the forced viewing of this LODActor, set drawing distance to 0.0f or LODDrawDistance
-	StaticMeshComponent->MinDrawDistance = (StaticMeshComponent->MinDrawDistance == 0.0f) ? LODDrawDistance : 0.0f;
-	StaticMeshComponent->MarkRenderStateDirty();
+	SetComponentsMinDrawDistance((StaticMeshComponent->MinDrawDistance == 0.0f) ? LODDrawDistance : 0.0f, true);
 }
 
 void ALODActor::SetForcedView(const bool InState)
 {
 	// Set forced viewing state of this LODActor, set drawing distance to 0.0f or LODDrawDistance
-	StaticMeshComponent->MinDrawDistance = (InState) ? 0.0f : LODDrawDistance;
-	StaticMeshComponent->MarkRenderStateDirty();
+	SetComponentsMinDrawDistance(InState ? 0.0f : LODDrawDistance, true);
 }
 
 void ALODActor::SetHiddenFromEditorView(const bool InState, const int32 ForceLODLevel )
@@ -844,7 +924,14 @@ void ALODActor::SetHiddenFromEditorView(const bool InState, const int32 ForceLOD
 				}
 
 				// Toggle/set the LOD parent to nullptr or this
-				Actor->SetLODParent((InState) ? nullptr : StaticMeshComponent, (InState) ? 0.0f : LODDrawDistance);
+				if (InState)
+				{
+					Actor->SetLODParent(nullptr, 0.0f);
+				}
+				else
+				{
+					Actor->SetLODParent(GetLODComponentForActor(Actor), LODDrawDistance);
+				}
 			}
 		}
 	}
@@ -876,13 +963,38 @@ void ALODActor::SetStaticMesh(class UStaticMesh* InStaticMesh)
 	}
 }
 
+void ALODActor::SetupImposters(UMaterialInterface* InMaterial, UStaticMesh* InStaticMesh, const TArray<FTransform>& InTransforms)
+{
+	check(InMaterial);
+	check(InStaticMesh);
+	check(InTransforms.Num() > 0);
+
+	UInstancedStaticMeshComponent* Component = GetOrCreateLODComponentForMaterial(InMaterial);
+	Component->SetStaticMesh(InStaticMesh);
+	
+	for (const FTransform& Transform : InTransforms)
+	{
+		Component->AddInstanceWorldSpace(Transform);
+	}
+
+	// Ensure parenting is up to date and take into account the newly created component.
+	UpdateSubActorLODParents();
+}
+
 void ALODActor::UpdateSubActorLODParents()
 {
 	for (AActor* Actor : SubActors)
 	{	
-		if (Actor)
+		if (Actor && Actor->IsA<ALODActor>())
 		{
-			Actor->SetLODParent(StaticMeshComponent, StaticMeshComponent->MinDrawDistance);
+			UStaticMeshComponent* LODComponent = GetLODComponentForActor(Actor);
+			Actor->SetLODParent(LODComponent, LODComponent->MinDrawDistance);
+		}
+
+		if(Actor)
+		{
+			UStaticMeshComponent* LODComponent = GetLODComponentForActor(Actor);
+			Actor->SetLODParent(LODComponent, LODComponent->MinDrawDistance);
 		}
 	}
 }
@@ -906,16 +1018,181 @@ void ALODActor::RecalculateDrawingDistance(const float InTransitionScreenSize)
 	static const float FOVRad = 90.0f * (float)PI / 360.0f;
 	static const FMatrix ProjectionMatrix = FPerspectiveMatrix(FOVRad, 1920, 1080, 0.01f);
 	FBoxSphereBounds Bounds = GetStaticMeshComponent()->CalcBounds(FTransform());
-	LODDrawDistance = ComputeBoundsDrawDistance(InTransitionScreenSize, Bounds.SphereRadius, ProjectionMatrix);
+	for (const TPair<const UMaterialInterface*, UInstancedStaticMeshComponent*>& Component : ImpostersStaticMeshComponents)
+	{
+		Bounds = Bounds + Component.Value->CalcBounds(FTransform());
+	}
 
-	StaticMeshComponent->MinDrawDistance = LODDrawDistance;	
+	float DrawDistance = ComputeBoundsDrawDistance(InTransitionScreenSize, Bounds.SphereRadius, ProjectionMatrix);
+	SetDrawDistance(DrawDistance);
 
 	UpdateSubActorLODParents();
 }
 
-
 #endif // WITH_EDITOR
 
+#if !(UE_BUILD_SHIPPING || UE_BUILD_TEST)
+
+bool IsImposter(const UStaticMeshComponent* InComponent)
+{
+	check(InComponent);
+
+	if (!InComponent->bUseMaxLODAsImposter)
+	{
+		return false;
+	}
+
+	const UStaticMesh* StaticMesh = InComponent->GetStaticMesh();
+	if (StaticMesh == nullptr)
+	{
+		return false;
+	}
+
+	if (StaticMesh->GetNumLODs() == 0)
+	{
+		return false;
+	}
+
+	const int32 LODIndex = StaticMesh->GetNumLODs() - 1;
+	const FStaticMeshLODResources& StaticMeshLOD = StaticMesh->RenderData->LODResources[LODIndex];
+
+	if (StaticMeshLOD.Sections.Num() != 1)
+	{
+		return false;
+	}
+
+	// Currently, the bUseMaxLODAsImposter flag serves multiple purposes. Until this is sorted out,
+	// only take into accounts quads & octagons meshes.
+	const int32 NumWedges = StaticMeshLOD.IndexBuffer.GetNumIndices();
+	const int32 NumVertexPositions = StaticMeshLOD.VertexBuffers.PositionVertexBuffer.GetNumVertices();
+	const int32 NumFaces = NumWedges / 3;
+
+	return (NumFaces == 2 && NumVertexPositions == 4) ||
+		   (NumFaces == 8 && NumVertexPositions == 9);
+}
+
+UMaterialInterface* ALODActor::GetImposterMaterial(const UStaticMeshComponent* InComponent) const
+{
+	UMaterialInterface* ImposterMaterial = nullptr;
+
+	if (IsImposter(InComponent))
+	{
+		UStaticMesh* StaticMesh = InComponent->GetStaticMesh();
+		check(StaticMesh);
+
+		// Retrieve imposter LOD mesh and material
+		const int32 LODIndex = StaticMesh->GetNumLODs() - 1;
+
+		// Retrieve the sections, we're expect 1 for imposter meshes
+		const TArray<FStaticMeshSection>& Sections = StaticMesh->RenderData->LODResources[LODIndex].Sections;
+		if (Sections.Num() == 1)
+		{
+			// Retrieve material for this section
+			ImposterMaterial = InComponent->GetMaterial(Sections[0].MaterialIndex);
+		}
+		else
+		{
+			UE_LOG(LogHLOD, Warning, TEXT("Imposter's static mesh %s has multiple mesh sections for its lowest LOD"), *StaticMesh->GetName());
+		}
+	}
+
+	return ImposterMaterial;
+}
+
+UMaterialInterface* ALODActor::GetImposterMaterial(const AActor* InActor) const
+{
+	UMaterialInterface* ImposterMaterial = nullptr;
+
+	TArray<UStaticMeshComponent*> Components;
+	InActor->GetComponents<UStaticMeshComponent>(Components);
+	Components.RemoveAll([&](UStaticMeshComponent* Val)
+	{
+#if WITH_EDITOR
+		return Val->GetStaticMesh() == nullptr || !Val->ShouldGenerateAutoLOD(LODLevel - 1);
+#else
+		return Val->GetStaticMesh() == nullptr;
+#endif
+	});
+
+	if (Components.Num() == 1)
+	{
+		ImposterMaterial = GetImposterMaterial(Components[0]);
+	}
+
+	return ImposterMaterial;
+}
+
+UInstancedStaticMeshComponent* ALODActor::GetLODComponentForMaterial(const UMaterialInterface* InMaterial) const
+{
+	return ImpostersStaticMeshComponents.FindRef(InMaterial);
+}
+
+UInstancedStaticMeshComponent* ALODActor::GetOrCreateLODComponentForMaterial(const UMaterialInterface* InMaterial)
+{
+	UInstancedStaticMeshComponent* LODComponent = ImpostersStaticMeshComponents.FindRef(InMaterial);
+	if (LODComponent == nullptr)
+	{
+		LODComponent = NewObject<UInstancedStaticMeshComponent>(this);
+		SetupComponent(LODComponent);
+
+		LODComponent->SetupAttachment(GetRootComponent());
+
+		if (StaticMeshComponent->IsRegistered())
+		{
+			LODComponent->RegisterComponent();
+		}
+		else
+		{
+			LODComponent->bAutoRegister = StaticMeshComponent->bAutoRegister;
+		}
+
+		ImpostersStaticMeshComponents.Emplace(InMaterial, LODComponent);
+	}
+
+	check(LODComponent);
+	return LODComponent;
+}
+
+UStaticMeshComponent* ALODActor::GetLODComponentForActor(const AActor* InActor, bool bInFallbackToDefault) const
+{
+	UStaticMeshComponent* LODComponent = StaticMeshComponent;
+
+	if (!InActor->IsA<ALODActor>())
+	{
+		UMaterialInterface* ImposterMaterial = GetImposterMaterial(InActor);
+		if (ImposterMaterial != nullptr)
+		{
+			LODComponent = GetLODComponentForMaterial(ImposterMaterial);
+			if (LODComponent == nullptr && bInFallbackToDefault)
+			{
+				// Needs to be rebuilt... fallback to default component
+				LODComponent = StaticMeshComponent;
+			}
+		}
+	}
+
+	return LODComponent;
+}
+
+UStaticMeshComponent* ALODActor::GetOrCreateLODComponentForActor(const AActor* InActor)
+{
+	UStaticMeshComponent* LODComponent = StaticMeshComponent;
+
+	if (!InActor->IsA<ALODActor>())
+	{
+		UMaterialInterface* ImposterMaterial = GetImposterMaterial(InActor);
+		if (ImposterMaterial != nullptr)
+		{
+			LODComponent = GetOrCreateLODComponentForMaterial(ImposterMaterial);
+		}
+	}
+
+	check(LODComponent != nullptr);
+	return LODComponent;
+}
+
+#endif // !(UE_BUILD_SHIPPING || UE_BUILD_TEST)
+
 FBox ALODActor::GetComponentsBoundingBox(bool bNonColliding) const 
 {
 	FBox BoundBox = Super::GetComponentsBoundingBox(bNonColliding);
@@ -928,20 +1205,31 @@ FBox ALODActor::GetComponentsBoundingBox(bool bNonColliding) const
 
 	if (bNonColliding)
 	{
-		if (StaticMeshComponent && StaticMeshComponent->GetStaticMesh())
+		bool bHasStaticMeshes = StaticMeshComponent && StaticMeshComponent->GetStaticMesh();
+
+		if (!bHasStaticMeshes)
 		{
-			FBoxSphereBounds StaticBound = StaticMeshComponent->GetStaticMesh()->GetBounds();
-			FBox StaticBoundBox(BoundBox.GetCenter()-StaticBound.BoxExtent, BoundBox.GetCenter()+StaticBound.BoxExtent);
-			BoundBox += StaticBoundBox;
+			for (const TPair<const UMaterialInterface*, UInstancedStaticMeshComponent*>& Component : ImpostersStaticMeshComponents)
+			{
+				UStaticMesh* StaticMesh = Component.Value->GetStaticMesh();
+				if (StaticMesh)
+				{
+					bHasStaticMeshes = true;
+					break;
+				}
+			}
 		}
-		else
+
+		// No valid static meshes found, use sub actors bounds instead.
+		if (!bHasStaticMeshes)
 		{
+			FBox SMBoundBox(ForceInit);
 			for (AActor* Actor : SubActors)
 			{
 				if (Actor)
 				{
 					BoundBox += Actor->GetComponentsBoundingBox(bNonColliding);
-				}				
+				}
 			}
 		}
 	}
