Fix PIE Debug Camera buffer visualization overview so arrows only move buffer selection, not camera. Changed ToggleDebugCamera hotkey from backslash to semicolon.

diff --git a/Engine/Config/BaseInput.ini b/Engine/Config/BaseInput.ini
--- a/Engine/Config/BaseInput.ini
+++ b/Engine/Config/BaseInput.ini
@@ -46,7 +46,7 @@ TriggerKeyIsConsideredReleased=0.25
 +DebugExecBindings=(Key=Apostrophe,Command="EnableGDT")
 +DebugExecBindings=(Key=Quote,Command="EnableGDT")
 ; --- PIE only binding
-+DebugExecBindings=(Key=Backslash,Command="ToggleDebugCamera")
++DebugExecBindings=(Key=Semicolon,Command="ToggleDebugCamera")
 
 [/Script/Engine.Console]
 HistoryBot=-1
diff --git a/Engine/Source/Runtime/Engine/Classes/Engine/DebugCameraController.h b/Engine/Source/Runtime/Engine/Classes/Engine/DebugCameraController.h
--- a/Engine/Source/Runtime/Engine/Classes/Engine/DebugCameraController.h
+++ b/Engine/Source/Runtime/Engine/Classes/Engine/DebugCameraController.h
@@ -41,14 +41,18 @@ class ENGINE_API ADebugCameraController
 	UPROPERTY()
 	uint32 bOrbitPivotUseCenter:1;
 
-	/** Whether set view mode to display GBuffer visualization */
+	/** Whether set view mode to display GBuffer visualization overview */
 	UPROPERTY()
 	uint32 bEnableBufferVisualization : 1;
 
-	/** Whether set view mode to display GBuffer visualization */
+	/** Whether set view mode to display GBuffer visualization full */
 	UPROPERTY()
 	uint32 bEnableBufferVisualizationFullMode : 1;
 
+	/** Whether GBuffer visualization overview inputs are set up  */
+	UPROPERTY()
+	uint32 bIsBufferVisualizationInputSetup : 1;
+
 	/** Visualizes the frustum of the camera */
 	UPROPERTY()
 	class UDrawFrustumComponent* DrawFrustum;
@@ -130,6 +134,9 @@ class ENGINE_API ADebugCameraController
 	/** Buffer overview move left */
 	void BufferVisualizationMoveLeft();
 
+	/** Ignores axis motion */
+	void ConsumeAxisMotion(float Val);
+
 	/** Toggle buffer visualization full display */
 	void ToggleBufferVisualizationFullMode();
 
@@ -195,6 +202,9 @@ protected:
 	virtual void ApplySpeedScale();
 	virtual void SetupInputComponent() override;
 
+	/** Sets up or clears input for buffer visualization overview */
+	void SetupBufferVisualizationOverviewInput();
+
 public:
 
 	/** 
diff --git a/Engine/Source/Runtime/Engine/Private/DebugCameraController.cpp b/Engine/Source/Runtime/Engine/Private/DebugCameraController.cpp
--- a/Engine/Source/Runtime/Engine/Private/DebugCameraController.cpp
+++ b/Engine/Source/Runtime/Engine/Private/DebugCameraController.cpp
@@ -59,6 +59,7 @@ ADebugCameraController::ADebugCameraController(const FObjectInitializer& ObjectI
 
 	bEnableBufferVisualization = false;
 	bEnableBufferVisualizationFullMode = false;
+	bIsBufferVisualizationInputSetup = false;
 	LastViewModeIndex = VMI_Lit;
 	LastViewModeSettingsIndex = 0;
 }
@@ -84,10 +85,6 @@ void InitializeDebugCameraInputBindings()
 		UPlayerInput::AddEngineDefinedActionMapping(FInputActionKeyMapping("DebugCamera_OrbitCenter", EKeys::O, true));
 		UPlayerInput::AddEngineDefinedActionMapping(FInputActionKeyMapping("DebugCamera_CycleViewMode", EKeys::V));
 		UPlayerInput::AddEngineDefinedActionMapping(FInputActionKeyMapping("DebugCamera_ToggleBufferVisualizationOverview", EKeys::B));
-		UPlayerInput::AddEngineDefinedActionMapping(FInputActionKeyMapping("DebugCamera_BufferVisualizationUp", EKeys::Up));
-		UPlayerInput::AddEngineDefinedActionMapping(FInputActionKeyMapping("DebugCamera_BufferVisualizationDown", EKeys::Down));
-		UPlayerInput::AddEngineDefinedActionMapping(FInputActionKeyMapping("DebugCamera_BufferVisualizationRight", EKeys::Right));
-		UPlayerInput::AddEngineDefinedActionMapping(FInputActionKeyMapping("DebugCamera_BufferVisualizationLeft", EKeys::Left));
 		UPlayerInput::AddEngineDefinedActionMapping(FInputActionKeyMapping("DebugCamera_ToggleBufferVisualizationFull", EKeys::Enter));
 
 		UPlayerInput::AddEngineDefinedActionMapping(FInputActionKeyMapping("DebugCamera_Select", EKeys::Gamepad_RightTrigger));
@@ -97,6 +94,17 @@ void InitializeDebugCameraInputBindings()
 		UPlayerInput::AddEngineDefinedActionMapping(FInputActionKeyMapping("DebugCamera_DecreaseFOV", EKeys::Gamepad_DPad_Down));
 		UPlayerInput::AddEngineDefinedActionMapping(FInputActionKeyMapping("DebugCamera_ToggleDisplay", EKeys::Gamepad_FaceButton_Left));
 		UPlayerInput::AddEngineDefinedActionMapping(FInputActionKeyMapping("DebugCamera_FreezeRendering", EKeys::Gamepad_FaceButton_Top));
+
+		UPlayerInput::AddEngineDefinedActionMapping(FInputActionKeyMapping("DebugCamera_BufferVisualizationUp", EKeys::Up));
+		UPlayerInput::AddEngineDefinedActionMapping(FInputActionKeyMapping("DebugCamera_BufferVisualizationDown", EKeys::Down));
+		UPlayerInput::AddEngineDefinedActionMapping(FInputActionKeyMapping("DebugCamera_BufferVisualizationLeft", EKeys::Left));
+		UPlayerInput::AddEngineDefinedActionMapping(FInputActionKeyMapping("DebugCamera_BufferVisualizationRight", EKeys::Right));
+
+		// The following axis mappings must be defined to override ADefaultPawn axis mappings when buffer visualization is enabled
+		UPlayerInput::AddEngineDefinedAxisMapping(FInputAxisKeyMapping("DebugCamera_DisableAxisMotion", EKeys::Up, 1.f));
+		UPlayerInput::AddEngineDefinedAxisMapping(FInputAxisKeyMapping("DebugCamera_DisableAxisMotion", EKeys::Down, -1.f));
+		UPlayerInput::AddEngineDefinedAxisMapping(FInputAxisKeyMapping("DebugCamera_DisableAxisMotion", EKeys::Left, -1.f));
+		UPlayerInput::AddEngineDefinedAxisMapping(FInputAxisKeyMapping("DebugCamera_DisableAxisMotion", EKeys::Right, 1.f));
 	}
 }
 
@@ -120,10 +128,6 @@ void ADebugCameraController::SetupInputComponent()
 	InputComponent->BindAction("DebugCamera_OrbitCenter", IE_Pressed, this, &ADebugCameraController::ToggleOrbitCenter);
 	InputComponent->BindAction("DebugCamera_CycleViewMode", IE_Pressed, this, &ADebugCameraController::CycleViewMode);
 	InputComponent->BindAction("DebugCamera_ToggleBufferVisualizationOverview", IE_Pressed, this, &ADebugCameraController::ToggleBufferVisualizationOverviewMode);
-	InputComponent->BindAction("DebugCamera_BufferVisualizationUp", IE_Pressed, this, &ADebugCameraController::BufferVisualizationMoveUp);
-	InputComponent->BindAction("DebugCamera_BufferVisualizationDown", IE_Pressed, this, &ADebugCameraController::BufferVisualizationMoveDown);
-	InputComponent->BindAction("DebugCamera_BufferVisualizationRight", IE_Pressed, this, &ADebugCameraController::BufferVisualizationMoveRight);
-	InputComponent->BindAction("DebugCamera_BufferVisualizationLeft", IE_Pressed, this, &ADebugCameraController::BufferVisualizationMoveLeft);
 	InputComponent->BindAction("DebugCamera_ToggleBufferVisualizationFull", IE_Pressed, this, &ADebugCameraController::ToggleBufferVisualizationFullMode);
 
 	InputComponent->BindTouch(IE_Pressed, this, &ADebugCameraController::OnTouchBegin);
@@ -131,6 +135,56 @@ void ADebugCameraController::SetupInputComponent()
 	InputComponent->BindTouch(IE_Repeat, this, &ADebugCameraController::OnFingerMove);
 }
 
+void ADebugCameraController::SetupBufferVisualizationOverviewInput()
+{
+	if (InputComponent)
+	{
+		if (bEnableBufferVisualization && !bEnableBufferVisualizationFullMode)
+		{
+			if (!bIsBufferVisualizationInputSetup)
+			{
+				InputComponent->BindAction("DebugCamera_BufferVisualizationUp", IE_Pressed, this, &ADebugCameraController::BufferVisualizationMoveUp);
+				InputComponent->BindAction("DebugCamera_BufferVisualizationDown", IE_Pressed, this, &ADebugCameraController::BufferVisualizationMoveDown);
+				InputComponent->BindAction("DebugCamera_BufferVisualizationRight", IE_Pressed, this, &ADebugCameraController::BufferVisualizationMoveRight);
+				InputComponent->BindAction("DebugCamera_BufferVisualizationLeft", IE_Pressed, this, &ADebugCameraController::BufferVisualizationMoveLeft);
+				InputComponent->BindAxis("DebugCamera_DisableAxisMotion", this, &ADebugCameraController::ConsumeAxisMotion);
+				bIsBufferVisualizationInputSetup = true;
+			}
+		}
+		else
+		{
+			if (bIsBufferVisualizationInputSetup)
+			{
+				// find any bindings that match the action names and remove them
+				for (int32 CurrentBindingIndex = 0; CurrentBindingIndex < InputComponent->GetNumActionBindings(); ++CurrentBindingIndex)
+				{
+					const FInputActionBinding& Binding = InputComponent->GetActionBinding(CurrentBindingIndex);
+					FName ActionName = Binding.GetActionName();
+					if (ActionName == "DebugCamera_BufferVisualizationUp"    ||
+						ActionName == "DebugCamera_BufferVisualizationDown"  ||
+						ActionName == "DebugCamera_BufferVisualizationRight" ||
+						ActionName == "DebugCamera_BufferVisualizationLeft")
+					{
+						InputComponent->RemoveActionBinding(CurrentBindingIndex);
+						--CurrentBindingIndex;
+					}
+				}
+
+				// find the axis binding and remove it
+				for (int32 CurrentAxisBindingIndex = 0; CurrentAxisBindingIndex < InputComponent->AxisBindings.Num(); ++CurrentAxisBindingIndex)
+				{
+					const FInputAxisBinding& Binding = InputComponent->AxisBindings[CurrentAxisBindingIndex];
+					if (Binding.AxisName == "DebugCamera_DisableAxisMotion")
+					{
+						InputComponent->AxisBindings.RemoveAt(CurrentAxisBindingIndex, 1, false);
+						--CurrentAxisBindingIndex;
+					}
+				}
+				bIsBufferVisualizationInputSetup = false;
+			}
+		}
+	}
+}
 
 void ADebugCameraController::OnTouchBegin(ETouchIndex::Type FingerIndex, FVector Location)
 {
@@ -392,6 +446,11 @@ void ADebugCameraController::OnDeactivate( APlayerController* RestoredPC )
 		bIsFrozenRendering = false;
 	}
 
+	if (bEnableBufferVisualization)
+	{
+		ToggleBufferVisualizationOverviewMode();
+	}
+
 	bIsOrbitingSelectedActor = false;
 
 	DrawFrustum->SetVisibility(false);
@@ -646,7 +705,6 @@ void ADebugCameraController::CycleViewMode()
 				FString Cmd(TEXT("VIEWMODE "));
 				Cmd += NextViewModeName;
 				GameViewportClient->ConsoleCommand(Cmd);
-
 			}
 		}
 
@@ -712,6 +770,8 @@ void ADebugCameraController::ToggleBufferVisualizationOverviewMode()
 		}
 
 		GameViewportClient->ConsoleCommand(Cmd);
+
+		SetupBufferVisualizationOverviewInput();
 	}
 }
 
@@ -738,7 +798,7 @@ void ADebugCameraController::UpdateVisualizeBufferPostProcessing(FFinalPostProce
 
 void ADebugCameraController::GetNextBuffer(int32 Step)
 {
-	if (bEnableBufferVisualization)
+	if (bEnableBufferVisualization && !bEnableBufferVisualizationFullMode)
 	{
 		TArray<FString> OverviewBuffers = GetBufferVisualizationOverviewTargets();
 		GetNextBuffer(OverviewBuffers, Step);
@@ -747,7 +807,7 @@ void ADebugCameraController::GetNextBuffer(int32 Step)
 
 void ADebugCameraController::GetNextBuffer(const TArray<FString>& OverviewBuffers, int32 Step)
 {
-	if (bEnableBufferVisualization)
+	if (bEnableBufferVisualization && !bEnableBufferVisualizationFullMode)
 	{
 		int32 BufferIndex = 0;
 
@@ -838,9 +898,15 @@ void ADebugCameraController::BufferVisualizationMoveLeft()
 	GetNextBuffer(-1);
 }
 
+void ADebugCameraController::ConsumeAxisMotion(float Val)
+{
+	// Just ignore the axis motion.
+}
+
 void ADebugCameraController::ToggleBufferVisualizationFullMode()
 {
 	SetBufferVisualizationFullMode(!bEnableBufferVisualizationFullMode);
+	SetupBufferVisualizationOverviewInput();
 }
 
 void ADebugCameraController::SetBufferVisualizationFullMode(bool bFullMode)
diff --git a/Engine/Source/Runtime/Engine/Private/DebugCameraHUD.cpp b/Engine/Source/Runtime/Engine/Private/DebugCameraHUD.cpp
--- a/Engine/Source/Runtime/Engine/Private/DebugCameraHUD.cpp
+++ b/Engine/Source/Runtime/Engine/Private/DebugCameraHUD.cpp
@@ -199,7 +199,7 @@ void ADebugCameraHUD::PostRender()
 
 
 			// controls display
-			yl += Y*14;
+			yl += Y*10;
 			
 			Canvas->SetDrawColor(64, 64, 255, 255);
 			Canvas->DrawText(RenderFont, TEXT("Controls"), X, yl, 1.f, 1.f, FontRenderInfo);
