Fix crash when reloading skeletalmesh import data

diff --git a/Engine/Source/Runtime/Engine/Public/Rendering/SkeletalMeshLODImporterData.cpp b/Engine/Source/Runtime/Engine/Public/Rendering/SkeletalMeshLODImporterData.cpp
--- a/Engine/Source/Runtime/Engine/Public/Rendering/SkeletalMeshLODImporterData.cpp
+++ b/Engine/Source/Runtime/Engine/Public/Rendering/SkeletalMeshLODImporterData.cpp
@@ -383,11 +383,41 @@ void FReductionBaseSkeletalMeshBulkData::Serialize(FArchive& Ar, TArray<FReducti
 
 void FReductionBaseSkeletalMeshBulkData::Serialize(FArchive& Ar, UObject* Owner)
 {
+	if (Ar.IsLoading())
+	{
+		//Save the custom version so we can load FReductionSkeletalMeshData later
+		SerializeLoadingCustomVersionContainer = Ar.GetCustomVersions();
+		bUseSerializeLoadingCustomVersion = false;
+		const FCustomVersionContainer& RegisteredCustomVersionContainer = FCustomVersionContainer::GetRegistered();
+		const FCustomVersionArray& ArchiveCustomVersionArray = SerializeLoadingCustomVersionContainer.GetAllVersions();
+		for (const FCustomVersion& ArchiveVersion : ArchiveCustomVersionArray)
+		{
+			const FCustomVersion* RegisteredVersion = RegisteredCustomVersionContainer.GetVersion(ArchiveVersion.Key);
+			if (!RegisteredVersion || RegisteredVersion->Version != ArchiveVersion.Version)
+			{
+				bUseSerializeLoadingCustomVersion = true;
+				break;
+			}
+		}
+	}
+
+	if (Ar.IsSaving() && bUseSerializeLoadingCustomVersion == true)
+	{
+		//We need to update the FReductionSkeletalMeshData serialize version to the latest in case we save the Parent bulkdata
+		FSkeletalMeshLODModel BaseLODModel;
+		TMap<FString, TArray<FMorphTargetDelta>> BaseLODMorphTargetData;
+		LoadReductionData(BaseLODModel, BaseLODMorphTargetData);
+		SaveReductionData(BaseLODModel, BaseLODMorphTargetData);
+	}
+
 	BulkData.Serialize(Ar, Owner);
 }
 
 void FReductionBaseSkeletalMeshBulkData::SaveReductionData(FSkeletalMeshLODModel& BaseLODModel, TMap<FString, TArray<FMorphTargetDelta>>& BaseLODMorphTargetData)
 {
+	//Saving the bulk data mean we do not need anymore the SerializeLoadingCustomVersionContainer of the parent bulk data
+	SerializeLoadingCustomVersionContainer.Empty();
+	bUseSerializeLoadingCustomVersion = false;
 	FReductionSkeletalMeshData ReductionSkeletalMeshData(BaseLODModel, BaseLODMorphTargetData);
 	TArray<uint8> TempBytes;
 	FMemoryWriter Ar(TempBytes, /*bIsPersistent=*/ true);
@@ -408,6 +438,11 @@ void FReductionBaseSkeletalMeshBulkData::LoadReductionData(FSkeletalMeshLODModel
 			BulkData.Lock(LOCK_READ_ONLY), BulkData.GetElementCount(),
 			/*bInFreeOnClose=*/ false, /*bIsPersistent=*/ true
 		);
+		//If we load the bulkdata content and the bulk was using customVersion, we must set the same until we save the reduction data
+		if (bUseSerializeLoadingCustomVersion)
+		{
+			Ar.SetCustomVersions(SerializeLoadingCustomVersionContainer);
+		}
 		Ar << ReductionSkeletalMeshData;
 		BulkData.Unlock();
 	}
diff --git a/Engine/Source/Runtime/Engine/Public/Rendering/SkeletalMeshLODImporterData.h b/Engine/Source/Runtime/Engine/Public/Rendering/SkeletalMeshLODImporterData.h
--- a/Engine/Source/Runtime/Engine/Public/Rendering/SkeletalMeshLODImporterData.h
+++ b/Engine/Source/Runtime/Engine/Public/Rendering/SkeletalMeshLODImporterData.h
@@ -357,6 +357,9 @@ struct FReductionBaseSkeletalMeshBulkData
 	/** Internally store bulk data as bytes. */
 	FByteBulkData BulkData;
 
+	//The custom version when this was load
+	FCustomVersionContainer SerializeLoadingCustomVersionContainer;
+	bool bUseSerializeLoadingCustomVersion = false;
 public:
 	/** Default constructor. */
 	ENGINE_API FReductionBaseSkeletalMeshBulkData();
