Fix component memory regressions
USceneComponent 560 -> 544
UPrimtiveComponent 1040 -> 1008 (including 16 from Scene)
#rnx

diff --git a/Engine/Source/Runtime/Engine/Classes/Components/PrimitiveComponent.h b/Engine/Source/Runtime/Engine/Classes/Components/PrimitiveComponent.h
--- a/Engine/Source/Runtime/Engine/Classes/Components/PrimitiveComponent.h
+++ b/Engine/Source/Runtime/Engine/Classes/Components/PrimitiveComponent.h
@@ -511,9 +511,6 @@ private:
 
 	FMaskFilter MoveIgnoreMask;
 
-	/** Custom data that can be read by a material through a material parameter expression. Set data using SetCustomPrimitiveData* functions */
-	UPROPERTY()
-	FCustomPrimitiveData CustomPrimitiveData;
 public:
 	/**
 	 * Determine whether a Character can step up onto this component.
@@ -538,6 +535,13 @@ public:
 	UPROPERTY(EditAnywhere, AdvancedDisplay, BlueprintReadOnly, Category=Rendering,  meta=(UIMin = "0", UIMax = "255", editcondition = "bRenderCustomDepth", DisplayName = "CustomDepth Stencil Value"))
 	int32 CustomDepthStencilValue;
 
+private:
+	/** Custom data that can be read by a material through a material parameter expression. Set data using SetCustomPrimitiveData* functions */
+	UPROPERTY()
+	FCustomPrimitiveData CustomPrimitiveData;
+
+public:
+
 	/**
 	 * Translucent objects with a lower sort priority draw behind objects with a higher priority.
 	 * Translucent objects with the same priority are rendered from back-to-front based on their bounds origin.
diff --git a/Engine/Source/Runtime/Engine/Classes/Components/SceneComponent.h b/Engine/Source/Runtime/Engine/Classes/Components/SceneComponent.h
--- a/Engine/Source/Runtime/Engine/Classes/Components/SceneComponent.h
+++ b/Engine/Source/Runtime/Engine/Classes/Components/SceneComponent.h
@@ -160,11 +160,6 @@ public:
 	UPROPERTY(BlueprintReadOnly, ReplicatedUsing=OnRep_Transform, interp, Category=Transform)
 	FVector RelativeScale3D;
 
-private:
-	/** Current transform of the component, relative to the world */
-	FTransform ComponentToWorld;
-
-public:
 	/**
 	* Velocity of the component.
 	* @see GetComponentVelocity()
@@ -273,6 +268,10 @@ public:
 	UPROPERTY(EditAnywhere, AdvancedDisplay, BlueprintReadOnly, Category = LOD)
 	TEnumAsByte<enum EDetailMode> DetailMode;
 
+	/** Delegate that will be called when PhysicsVolume has been changed **/
+	UPROPERTY(BlueprintAssignable, Category=PhysicsVolume, meta=(DisplayName="Physics Volume Changed"))
+	FPhysicsVolumeChanged PhysicsVolumeChangedDelegate;
+
 #if WITH_EDITORONLY_DATA
 protected:
 	/** Editor only component used to display the sprite so as to be able to see the location of the Audio Component  */
@@ -286,6 +285,9 @@ private:
 	/** Cache that avoids Quat<->Rotator conversions if possible. Only to be used with RelativeRotation. */
 	FRotationConversionCache RelativeRotationCache;
 
+	/** Current transform of the component, relative to the world */
+	FTransform ComponentToWorld;
+
 public:
 	/** Sets the RelativeRotationCache. Used to ensure component ends up with the same RelativeRotation after calling SetWorldTransform(). */
 	void SetRelativeRotationCache(const FRotationConversionCache& InCache);
@@ -293,10 +295,6 @@ public:
 	/** Get the RelativeRotationCache.  */
 	FORCEINLINE const FRotationConversionCache& GetRelativeRotationCache() const { return RelativeRotationCache; }
 
-	/** Delegate that will be called when PhysicsVolume has been changed **/
-	UPROPERTY(BlueprintAssignable, Category=PhysicsVolume, meta=(DisplayName="Physics Volume Changed"))
-	FPhysicsVolumeChanged PhysicsVolumeChangedDelegate;
-
 	/** Delegate called when this component is moved */
 	FTransformUpdated TransformUpdated;
 
