Fixed outgoing packet loss displayed in NetDebugWidget
- Added counter for number of total notified packets that can be used to accurately measure outgoing packetloss

diff --git a/Engine/Source/Runtime/Engine/Classes/Engine/NetConnection.h b/Engine/Source/Runtime/Engine/Classes/Engine/NetConnection.h
--- a/Engine/Source/Runtime/Engine/Classes/Engine/NetConnection.h
+++ b/Engine/Source/Runtime/Engine/Classes/Engine/NetConnection.h
@@ -1107,6 +1107,11 @@ public:
 	 * @param bWasSaturated		True if we failed to replicate all data because we were saturated.
 	 */
 	ENGINE_API void TrackReplicationForAnalytics(const bool bWasSaturated);
+
+	/**
+	 * Get the current number of sent packets for which we have received a delivery notification
+	 */
+	ENGINE_API uint32 GetOutTotalNotifiedPackets() const { return OutTotalNotifiedPackets; }
 	
 protected:
 
@@ -1190,6 +1195,9 @@ private:
 	/** Full PacketId  of last sent packet that we have received notification for (i.e. we know if it was delivered or not). Related to OutAckPacketId which is tha last successfully delivered PacketId */
 	int32 LastNotifiedPacketId;
 
+	/** Count the number of notified packets, i.e. packets that we know if they are delivered or not. Used to reliably measure outgoing packet loss */
+	uint32 OutTotalNotifiedPackets;
+
 	/** Keep old behavior where we send a packet with only acks even if we have no other outgoing data if we got incoming data */
 	uint32 HasDirtyAcks;
 	
diff --git a/Engine/Source/Runtime/Engine/Classes/Engine/NetDriver.h b/Engine/Source/Runtime/Engine/Classes/Engine/NetDriver.h
--- a/Engine/Source/Runtime/Engine/Classes/Engine/NetDriver.h
+++ b/Engine/Source/Runtime/Engine/Classes/Engine/NetDriver.h
@@ -1520,6 +1520,16 @@ public:
 		return bMaySendProperties;
 	}
 
+	/**
+	 * Get the current number of sent packets for which we have received a delivery notification
+	 */
+	ENGINE_API uint32 GetOutTotalNotifiedPackets() const { return OutTotalNotifiedPackets; }
+
+	/**
+	 * Increase the current number of sent packets for which we have received a delivery notification
+	 */
+	inline void IncreaseOutTotalNotifiedPackets() { ++OutTotalNotifiedPackets; }
+
 protected:
 
 	bool bMaySendProperties;
@@ -1551,4 +1561,7 @@ private:
 
 	/** NetDriver time to end packet loss burst simulation. */
 	float PacketLossBurstEndTime;
+
+	/** Count the number of notified packets, i.e. packets that we know if they are delivered or not. Used to reliably measure outgoing packet loss */
+	uint32 OutTotalNotifiedPackets;
 };
\ No newline at end of file
diff --git a/Engine/Source/Runtime/Engine/Private/NetConnection.cpp b/Engine/Source/Runtime/Engine/Private/NetConnection.cpp
--- a/Engine/Source/Runtime/Engine/Private/NetConnection.cpp
+++ b/Engine/Source/Runtime/Engine/Private/NetConnection.cpp
@@ -169,6 +169,7 @@ PRAGMA_ENABLE_DEPRECATION_WARNINGS
 ,	PlayerOnlinePlatformName( NAME_None )
 ,	ClientWorldPackageName( NAME_None )
 ,	LastNotifiedPacketId( -1 )
+,	OutTotalNotifiedPackets(0)
 ,	HasDirtyAcks(0u)
 ,	bHasWarnedAboutChannelLimit(false)
 ,	bConnectionPendingCloseDueToSocketSendFailure(false)
@@ -1764,6 +1765,8 @@ void UNetConnection::ReceivedPacket( FBitReader& Reader )
 		{
 			// Increase LastNotifiedPacketId, this is a full packet Id
 			++LastNotifiedPacketId;
+			++OutTotalNotifiedPackets;
+			Driver->IncreaseOutTotalNotifiedPackets();
 
 			// Sanity check
 			if (FNetPacketNotify::SequenceNumberT(LastNotifiedPacketId) != AckedSequence)
diff --git a/Engine/Source/Runtime/Engine/Private/NetworkDriver.cpp b/Engine/Source/Runtime/Engine/Private/NetworkDriver.cpp
--- a/Engine/Source/Runtime/Engine/Private/NetworkDriver.cpp
+++ b/Engine/Source/Runtime/Engine/Private/NetworkDriver.cpp
@@ -271,6 +271,7 @@ UNetDriver::UNetDriver(const FObjectInitializer& ObjectInitializer)
 ,	LagState(ENetworkLagState::NotLagging)
 ,	DuplicateLevelID(INDEX_NONE)
 ,	PacketLossBurstEndTime(-1.0f)
+,	OutTotalNotifiedPackets(0u)
 {
 PRAGMA_DISABLE_DEPRECATION_WARNINGS
 	ChannelClasses[CHTYPE_Control]	= UControlChannel::StaticClass();
