Removing a legacy data reset which caused unintentional data loss when projects were migrated to newer engine releases.

diff --git a/Engine/Source/Runtime/Engine/Private/Components/PrimitiveComponent.cpp b/Engine/Source/Runtime/Engine/Private/Components/PrimitiveComponent.cpp
--- a/Engine/Source/Runtime/Engine/Private/Components/PrimitiveComponent.cpp
+++ b/Engine/Source/Runtime/Engine/Private/Components/PrimitiveComponent.cpp
@@ -805,11 +805,6 @@ void UPrimitiveComponent::Serialize(FArchive& Ar)
 			LightmapType = ELightmapType::ForceSurface;
 		}
 	}
-	
-	if (Ar.CustomVer(FFortniteMainBranchObjectVersion::GUID) < FFortniteMainBranchObjectVersion::CullDistanceRefactor_RemovedDefaultDistance)
-	{
-		LDMaxDrawDistance = 0.f;
-	}
 }
 
 #if WITH_EDITOR
