Copying Dev-Audio to UE-Main from CL 6717121

diff --git a/Engine/Source/Runtime/Engine/Classes/Sound/SoundAttenuation.h b/Engine/Source/Runtime/Engine/Classes/Sound/SoundAttenuation.h
--- a/Engine/Source/Runtime/Engine/Classes/Sound/SoundAttenuation.h
+++ b/Engine/Source/Runtime/Engine/Classes/Sound/SoundAttenuation.h
@@ -47,7 +47,7 @@ enum class EReverbSendMethod : uint8
 	// A reverb send based on a supplied curve
 	CustomCurve,
 
-	// A manual reverb send level (Uses the specififed constant send level value. Useful for 2D sounds.)
+	// A manual reverb send level (Uses the specified constant send level value. Useful for 2D sounds.)
 	Manual,
 };
 
@@ -57,15 +57,15 @@ struct ENGINE_API FSoundAttenuationPluginSettings
 {
 	GENERATED_USTRUCT_BODY()
 
-	/** Settings to use with occlusion audio plugin. These are defined by the plugin creator. Not all audio plugins utilize this feature. Note that this is an array so multiple plugins can have settings. */
+	/** Settings to use with occlusion audio plugin. These are defined by the plugin creator. Not all audio plugins utilize this feature. This is an array so multiple plugins can have settings. */
 	UPROPERTY(EditAnywhere, BlueprintReadWrite, Category = AttenuationSpatialization, meta = (DisplayName = "Spatialization Plugin Settings"))
 	TArray<USpatializationPluginSourceSettingsBase*> SpatializationPluginSettingsArray;
 
-	/** Settings to use with occlusion audio plugin. These are defined by the plugin creator. Not all audio plugins utilize this feature. Note that this is an array so multiple plugins can have settings. */
+	/** Settings to use with occlusion audio plugin. These are defined by the plugin creator. Not all audio plugins utilize this feature. This  is an array so multiple plugins can have settings. */
 	UPROPERTY(EditAnywhere, BlueprintReadWrite, Category = AttenuationOcclusion, meta = (DisplayName = "Occlusion Plugin Settings"))
 	TArray<UOcclusionPluginSourceSettingsBase*> OcclusionPluginSettingsArray;
 
-	/** Settings to use with reverb audio plugin. These are defined by the plugin creator. Not all audio plugins utilize this feature. Note that this is an array so multiple plugins can have settings. */
+	/** Settings to use with reverb audio plugin. These are defined by the plugin creator. Not all audio plugins utilize this feature. This  is an array so multiple plugins can have settings. */
 	UPROPERTY(EditAnywhere, BlueprintReadWrite, Category = AttenuationReverbSend, meta = (DisplayName = "Reverb Plugin Settings"))
 	TArray<UReverbPluginSourceSettingsBase*> ReverbPluginSettingsArray;
 };
diff --git a/Engine/Source/Runtime/Engine/Classes/Sound/SoundBase.h b/Engine/Source/Runtime/Engine/Classes/Sound/SoundBase.h
--- a/Engine/Source/Runtime/Engine/Classes/Sound/SoundBase.h
+++ b/Engine/Source/Runtime/Engine/Classes/Sound/SoundBase.h
@@ -1,29 +1,26 @@
 // Copyright 1998-2019 Epic Games, Inc. All Rights Reserved.
-
 #pragma once
 
-/** 
- * The base class for a playable sound object 
+/**
+ * The base class for a playable sound object
  */
 
+#include "Audio.h"
 #include "CoreMinimal.h"
-#include "UObject/ObjectMacros.h"
 #include "UObject/Object.h"
-#include "Audio.h"
-#include "Sound/SoundConcurrency.h"
+#include "UObject/ObjectMacros.h"
+#include "SoundConcurrency.h"
+#include "SoundSourceBusSend.h"
 #include "SoundSubmix.h"
-#include "Sound/SoundSourceBusSend.h"
+#include "IAudioExtensionPlugin.h"
+
 #include "SoundBase.generated.h"
 
-class USoundConcurrency;
+
 class USoundEffectSourcePreset;
-class USoundSubmix;
 class USoundSourceBus;
 class USoundEffectSourcePresetChain;
 
-struct FSoundConcurrencySettings;
-struct FSoundSubmixSendInfo;
-struct FSoundSourceBusSendInfo;
 struct FActiveSound;
 struct FSoundParseParameters;
 
@@ -63,9 +60,12 @@ public:
 	UPROPERTY()
 	uint8 bHasConcatenatorNode : 1;
 
-	/** Whether a sound has virtualize when silent enabled (i.e. for a sound cue, if any sound wave player has it enabled). */
 	UPROPERTY()
-	uint8 bHasVirtualizeWhenSilent:1;
+	uint8 bHasVirtualizeWhenSilent_DEPRECATED:1;
+
+	/** Whether a sound has play when silent enabled (i.e. for a sound cue, if any sound wave player has it enabled). */
+	UPROPERTY()
+	uint8 bHasPlayWhenSilent : 1;
 
 	/** Bypass volume weighting priority upon evaluating whether sound should remain active when max channel count is met (See platform Audio Settings). */
 	UPROPERTY(EditAnywhere, BlueprintReadWrite, Category = Priority)
@@ -120,8 +120,12 @@ public:
 	UPROPERTY(EditAnywhere, Category=Attenuation)
 	USoundAttenuation* AttenuationSettings;
 
-	/** Sound submix this sound belongs to. 
-	  * Audio will play here and traverse through the submix graph. 
+	/** Modulation for the sound */
+	UPROPERTY(EditAnywhere, Category = Modulation)
+	FSoundModulation Modulation;
+
+	/** Sound submix this sound belongs to.
+	  * Audio will play here and traverse through the submix graph.
 	  * A null entry will make the sound obey the default master effects graph.
 	  */
 	UPROPERTY(EditAnywhere, BlueprintReadWrite, Category = Effects, meta = (DisplayName = "Sound Submix"))
@@ -143,7 +147,7 @@ public:
 	UPROPERTY(EditAnywhere, BlueprintReadWrite, Category = Effects, meta = (DisplayName = "Pre-Effect Bus Sends"))
 	TArray<FSoundSourceBusSendInfo> PreEffectBusSends;
 
-public:	
+public:
 
 	//~ Begin UObject Interface.
 	virtual void PostInitProperties() override;
@@ -151,13 +155,14 @@ public:
 	virtual bool CanBeClusterRoot() const override;
 	virtual bool CanBeInCluster() const override;
 	virtual void Serialize(FArchive& Ar) override;
+
 	//~ End UObject interface.
 
 	/** Returns whether the sound base is set up in a playable manner */
 	virtual bool IsPlayable() const;
 
-	/** Returns whether a sound is allowed to be virtualized. */
-	virtual bool IsAllowedVirtual() const;
+	/** Returns whether sound supports subtitles. */
+	virtual bool SupportsSubtitles() const;
 
 	/** Returns whether or not this sound base has an attenuation node. */
 	virtual bool HasAttenuationNode() const;
@@ -170,7 +175,7 @@ public:
 	 */
 	virtual float GetMaxDistance() const;
 
-	/** 
+	/**
 	 * Returns the length of the sound
 	 */
 	virtual float GetDuration();
@@ -181,15 +186,15 @@ public:
 	/** Returns whether or not this sound has a sequencer node, which means it's possible for the owning active sound to persist even though it's not generating audio. */
 	bool HasConcatenatorNode() const;
 
-	/** Returns true if any of the sounds in the sound have "virtualize when silent" enabled. */
-	bool IsVirtualizeWhenSilent() const;
+	/** Returns true if any of the sounds in the sound have "play when silent" enabled. */
+	bool IsPlayWhenSilent() const;
 
 	virtual float GetVolumeMultiplier();
 	virtual float GetPitchMultiplier();
 
 	/** Returns the subtitle priority */
 	virtual float GetSubtitlePriority() const { return DEFAULT_SUBTITLE_PRIORITY; };
-	
+
 	/** Returns whether or not any part of this sound wants interior volumes applied to it */
 	virtual bool ShouldApplyInteriorVolumes();
 
diff --git a/Engine/Source/Runtime/Engine/Classes/Sound/SoundNode.h b/Engine/Source/Runtime/Engine/Classes/Sound/SoundNode.h
--- a/Engine/Source/Runtime/Engine/Classes/Sound/SoundNode.h
+++ b/Engine/Source/Runtime/Engine/Classes/Sound/SoundNode.h
@@ -16,18 +16,18 @@ struct FPropertyChangedEvent;
 struct FWaveInstance;
 
 /*-----------------------------------------------------------------------------
-	USoundNode helper macros. 
+	USoundNode helper macros.
 -----------------------------------------------------------------------------*/
 
 struct FActiveSound;
 
 #define DECLARE_SOUNDNODE_ELEMENT(Type,Name)													\
 	Type& Name = *((Type*)(Payload));															\
-	Payload += sizeof(Type);														
+	Payload += sizeof(Type);
 
 #define DECLARE_SOUNDNODE_ELEMENT_PTR(Type,Name)												\
 	Type* Name = (Type*)(Payload);																\
-	Payload += sizeof(Type);														
+	Payload += sizeof(Type);
 
 #define	RETRIEVE_SOUNDNODE_PAYLOAD( Size )														\
 		uint8*	Payload					= NULL;													\
@@ -82,32 +82,32 @@ public:
 	//~ End UObject Interface
 
 	//
-	//~ Begin USoundNode Interface. 
+	//~ Begin USoundNode Interface.
 	//
 
 	/**
 	 * Notifies the sound node that a wave instance in its subtree has finished.
 	 *
-	 * @param WaveInstance	WaveInstance that was finished 
+	 * @param WaveInstance	WaveInstance that was finished
 	 */
 	virtual bool NotifyWaveInstanceFinished( struct FWaveInstance* WaveInstance )
 	{
 		return( false );
 	}
 
-	/** 
+	/**
 	 * Returns the maximum distance this sound can be heard from.
 	 */
 	virtual float GetMaxDistance() const;
 
 	/** Returns if this node has been set to be allowed virtual. Only the sound node wave player implements this. */
-	virtual bool IsAllowedVirtual() const
+	virtual bool SupportsSubtitles() const
 	{
 		return false;
 	}
 
-	/** 
-	 * Returns the maximum duration this sound node will play for. 
+	/**
+	 * Returns the maximum duration this sound node will play for.
 	 *
 	 * @return	float of number of seconds this sound will play for. INDEFINITELY_LOOPING_DURATION means its looping.
 	 */
@@ -119,27 +119,27 @@ public:
 	/** Returns whether the sound has a sequencer node. */
 	virtual bool HasConcatenatorNode() const;
 
-	/** Returns true if the sound node is virtualize when silent. */
-	virtual bool IsVirtualizeWhenSilent() const;
+	/** Returns true if the sound node is set to play when silent. */
+	virtual bool IsPlayWhenSilent() const;
 
 	virtual void ParseNodes( FAudioDevice* AudioDevice, const UPTRINT NodeWaveInstanceHash, FActiveSound& ActiveSound, const struct FSoundParseParameters& ParseParams, TArray<FWaveInstance*>& WaveInstances );
 
 	/**
 	 * Returns an array of all (not just active) nodes.
 	 */
-	virtual void GetAllNodes( TArray<USoundNode*>& SoundNodes ); 
+	virtual void GetAllNodes( TArray<USoundNode*>& SoundNodes );
 
 	/**
 	 * Returns the maximum number of child nodes this node can possibly have
 	 */
 	virtual int32 GetMaxChildNodes() const
-	{ 
-		return 1 ; 
+	{
+		return 1 ;
 	}
 
 	/** Returns the minimum number of child nodes this node must have */
 	virtual int32 GetMinChildNodes() const
-	{ 
+	{
 		return 0;
 	}
 
@@ -147,8 +147,8 @@ public:
 	virtual int32 GetNumSounds(const UPTRINT NodeWaveInstanceHash, FActiveSound& ActiveSound) const;
 
 
-	/** 
-	 * Editor interface. 
+	/**
+	 * Editor interface.
 	 */
 
 	/**
@@ -180,7 +180,7 @@ public:
 	virtual void OnEndPIE(const bool bIsSimulating) {};
 #endif //WITH_EDITOR
 
-	/** 
+	/**
 	 * Used to create a unique string to identify unique nodes
 	 */
 	static UPTRINT GetNodeWaveInstanceHash(const UPTRINT ParentWaveInstanceHash, const USoundNode* ChildNode, const uint32 ChildIndex);
diff --git a/Engine/Source/Runtime/Engine/Classes/Sound/SoundNodeWavePlayer.h b/Engine/Source/Runtime/Engine/Classes/Sound/SoundNodeWavePlayer.h
--- a/Engine/Source/Runtime/Engine/Classes/Sound/SoundNodeWavePlayer.h
+++ b/Engine/Source/Runtime/Engine/Classes/Sound/SoundNodeWavePlayer.h
@@ -16,7 +16,7 @@ struct FPropertyChangedEvent;
 struct FSoundParseParameters;
 struct FWaveInstance;
 
-/** 
+/**
  * Sound node that contains a reference to the raw wave file to be played
  */
 UCLASS(hidecategories=Object, editinlinenew, MinimalAPI, meta=( DisplayName="Wave Player" ))
@@ -34,7 +34,7 @@ private:
 	void OnSoundWaveLoaded(const FName& PackageName, UPackage * Package, EAsyncLoadingResult::Type Result, bool bAddToRoot);
 
 	uint32 bAsyncLoading:1;
-public:	
+public:
 
 	UPROPERTY(EditAnywhere, Category=WavePlayer)
 	uint32 bLooping:1;
@@ -52,7 +52,7 @@ public:
 	//~ Begin USoundNode Interface
 	virtual int32 GetMaxChildNodes() const override;
 	virtual float GetDuration() override;
-	virtual bool IsVirtualizeWhenSilent() const override;
+	virtual bool IsPlayWhenSilent() const override;
 	virtual int32 GetNumSounds(const UPTRINT NodeWaveInstanceHash, FActiveSound& ActiveSound) const { return 1; }
 	virtual void ParseNodes(FAudioDevice* AudioDevice, const UPTRINT NodeWaveInstanceHash, FActiveSound& ActiveSound, const FSoundParseParameters& ParseParams, TArray<FWaveInstance*>& WaveInstances) override;
 #if WITH_EDITOR
diff --git a/Engine/Source/Runtime/Engine/Classes/Sound/SoundWave.h b/Engine/Source/Runtime/Engine/Classes/Sound/SoundWave.h
--- a/Engine/Source/Runtime/Engine/Classes/Sound/SoundWave.h
+++ b/Engine/Source/Runtime/Engine/Classes/Sound/SoundWave.h
@@ -328,9 +328,12 @@ class ENGINE_API USoundWave : public USoundBase
 	UPROPERTY(EditAnywhere, Category=Subtitles )
 	uint8 bSingleLine:1;
 
+	UPROPERTY()
+	uint8 bVirtualizeWhenSilent_DEPRECATED:1;
+
 	/** Allows sound to continue playing when silent. This prevents issues with sounds restarting when coming back in range, etc. */
 	UPROPERTY(EditAnywhere, Category=Sound, meta = (DisplayName = "Play When Silent"))
-	uint8 bVirtualizeWhenSilent:1;
+	uint8 bPlayWhenSilent:1;
 
 	/** Whether or not this source is ambisonics file format. */
 	UPROPERTY(EditAnywhere, Category = Sound)
@@ -593,7 +596,7 @@ public:
 	virtual void Parse( class FAudioDevice* AudioDevice, const UPTRINT NodeWaveInstanceHash, FActiveSound& ActiveSound, const FSoundParseParameters& ParseParams, TArray<FWaveInstance*>& WaveInstances ) override;
 	virtual float GetDuration() override;
 	virtual float GetSubtitlePriority() const override;
-	virtual bool IsAllowedVirtual() const override;
+	virtual bool SupportsSubtitles() const override;
 	virtual bool GetSoundWavesWithCookedAnalysisData(TArray<USoundWave*>& OutSoundWaves) override;
 	virtual bool HasCookedFFTData() const override;
 	virtual bool HasCookedAmplitudeEnvelopeData() const override;
diff --git a/Engine/Source/Runtime/Engine/Private/ActiveSound.cpp b/Engine/Source/Runtime/Engine/Private/ActiveSound.cpp
--- a/Engine/Source/Runtime/Engine/Private/ActiveSound.cpp
+++ b/Engine/Source/Runtime/Engine/Private/ActiveSound.cpp
@@ -26,22 +26,6 @@ FAutoConsoleVariableRef CVarAudioDisableConcurrencyStopSilentForLoops(
 	TEXT("Disables (1) or enables (0) audio concurrency for loops subscribing to stop silent.\n"),
 	ECVF_Default);
 
-static int32 ActiveSoundVisualizeModeCVar = 1;
-FAutoConsoleVariableRef CVarAudioVisualizeActiveSoundsMode(
-	TEXT("au.3dVisualize.ActiveSounds"),
-	ActiveSoundVisualizeModeCVar,
-	TEXT("Visualization mode for active sounds. \n")
-	TEXT("0: Not Enabled, 1: Volume (Lin), 2: Volume (dB), 3: Distance, 4: Random color"),
-	ECVF_Default);
-
-static int32 ActiveSoundVisualizeTypeCVar = 0;
-FAutoConsoleVariableRef CVarAudioVisualizeActiveSounds(
-	TEXT("au.3dVisualize.ActiveSounds.Type"),
-	ActiveSoundVisualizeTypeCVar,
-	TEXT("Whether to show all sounds, on AudioComponents (Components Only), or off of AudioComponents (Non-Component Only). \n")
-	TEXT("0: All, 1: Components Only, 2: Non-Component Only"),
-	ECVF_Default);
-
 FTraceDelegate FActiveSound::ActiveSoundTraceDelegate;
 TMap<FTraceHandle, FActiveSound::FAsyncTraceDetails> FActiveSound::TraceToActiveSoundMap;
 
@@ -121,7 +105,7 @@ FActiveSound::FActiveSound()
 	, CurrentInteriorLPF(MAX_FILTER_FREQUENCY)
 	, EnvelopeFollowerAttackTime(10)
 	, EnvelopeFollowerReleaseTime(100)
-#if !(UE_BUILD_SHIPPING || UE_BUILD_TEST)
+#if ENABLE_AUDIO_DEBUG
 	, DebugColor(FColor::Black)
 #endif // !(UE_BUILD_SHIPPING || UE_BUILD_TEST)
 	, ClosestListenerPtr(nullptr)
@@ -182,6 +166,16 @@ void FActiveSound::AddReferencedObjects( FReferenceCollector& Collector)
 	}
 }
 
+void FActiveSound::SetPitch(float Value)
+{
+	PitchMultiplier = Value;
+}
+
+void FActiveSound::SetVolume(float Value)
+{
+	VolumeMultiplier = Value;
+}
+
 void FActiveSound::SetWorld(UWorld* InWorld)
 {
 	World = InWorld;
@@ -222,16 +216,15 @@ void FActiveSound::SetAudioComponent(const FActiveSound& ActiveSound)
 	OwnerName = ActiveSound.OwnerName;
 }
 
-void FActiveSound::SetAudioComponent(UAudioComponent* Component)
+void FActiveSound::SetAudioComponent(const UAudioComponent& Component)
 {
 	check(IsInGameThread());
-	check(Component);
 
-	AActor* Owner = Component->GetOwner();
+	AActor* Owner = Component.GetOwner();
 
-	AudioComponentID = Component->GetAudioComponentID();
-	AudioComponentUserID = Component->GetAudioComponentUserID();
-	AudioComponentName = Component->GetFName();
+	AudioComponentID = Component.GetAudioComponentID();
+	AudioComponentUserID = Component.GetAudioComponentUserID();
+	AudioComponentName = Component.GetFName();
 
 	SetOwner(Owner);
 }
@@ -445,7 +438,7 @@ void FActiveSound::UpdateWaveInstances(TArray<FWaveInstance*> &InWaveInstances,
 		ClosestListenerIndex = FindClosestListener(Listeners);
 	}
 
-	// Cache the closest listener ptr 
+	// Cache the closest listener ptr
 	ClosestListenerPtr = &Listeners[ClosestListenerIndex];
 
 	FSoundParseParameters ParseParams;
@@ -462,16 +455,10 @@ void FActiveSound::UpdateWaveInstances(TArray<FWaveInstance*> &InWaveInstances,
 	// in parse params.
 	UpdateConcurrencyVolumeScalars(DeltaTime);
 
-	// If the sound is a preview sound, then ignore the transient master volume and application volume
-	if (!bIsPreviewSound)
-	{
-		ParseParams.VolumeApp = AudioDevice->GetMasterVolume();
-	}
-
 	ParseParams.VolumeMultiplier = GetVolume();
 
 	ParseParams.Priority = Priority;
-	ParseParams.Pitch *= PitchMultiplier * Sound->GetPitchMultiplier();
+	ParseParams.Pitch *= GetPitch() * Sound->GetPitchMultiplier();
 	ParseParams.bEnableLowPassFilter = bEnableLowPassFilter;
 	ParseParams.LowPassFilterFrequency = LowPassFilterFrequency;
 	ParseParams.SoundClass = GetSoundClass();
@@ -487,7 +474,7 @@ void FActiveSound::UpdateWaveInstances(TArray<FWaveInstance*> &InWaveInstances,
 		GetSoundSourceBusSends((EBusSendType)BusSendType, ParseParams.SoundSourceBusSends[BusSendType]);
 	}
 
-	// Set up the base source effect chain. 
+	// Set up the base source effect chain.
 	ParseParams.SourceEffectChain = Sound->SourceEffectChain;
 
 	// Setup the envelope attack and release times
@@ -508,8 +495,7 @@ void FActiveSound::UpdateWaveInstances(TArray<FWaveInstance*> &InWaveInstances,
 		LastLocation = ParseParams.Transform.GetTranslation();
 	}
 
-	static TArray<FWaveInstance*> ThisSoundsWaveInstances;
-	ThisSoundsWaveInstances.Reset();
+	TArray<FWaveInstance*> ThisSoundsWaveInstances;
 
 	// Recurse nodes, have SoundWave's create new wave instances and update bFinished unless we finished fading out.
 	bFinished = true;
@@ -526,6 +512,8 @@ void FActiveSound::UpdateWaveInstances(TArray<FWaveInstance*> &InWaveInstances,
 			ParseParams.ManualReverbSendLevel = AudioDevice->GetDefaultReverbSendLevel();
 		}
 
+		ParseParams.ModulationPluginSettings = FindModulationSettings();
+
 		// if the closest listener is not the primary one, transform the sound transform so it's panned relative to primary listener position
 		if (ClosestListenerIndex != 0)
 		{
@@ -645,114 +633,44 @@ void FActiveSound::UpdateWaveInstances(TArray<FWaveInstance*> &InWaveInstances,
 
 	}
 
-	if (ActiveSoundVisualizeModeCVar != 0)
+#if ENABLE_AUDIO_DEBUG
+	if (DebugColor == FColor::Black)
 	{
-		DrawDebugInfo(&ThisSoundsWaveInstances);
+		DebugColor = FColor::MakeRandomColor();
 	}
+	FAudioDebugger::DrawDebugInfo(*this, ThisSoundsWaveInstances);
+#endif // ENABLE_AUDIO_DEBUG
 
 	InWaveInstances.Append(ThisSoundsWaveInstances);
 }
 
-void FActiveSound::DrawDebugInfo(const TArray<FWaveInstance*>* ThisSoundsWaveInstances)
+USoundModulationPluginSourceSettingsBase* FActiveSound::FindModulationSettings() const
 {
-#if ENABLE_DRAW_DEBUG
-	// Only draw spatialized sounds
-	if (!Sound || !bAllowSpatialization)
-	{
-		return;
-	}
-
-	if (ActiveSoundVisualizeTypeCVar > 0)
-	{
-		if (ActiveSoundVisualizeTypeCVar == 1 && AudioComponentID == 0)
-		{
-			return;
-		}
-
-		if (ActiveSoundVisualizeTypeCVar == 2 && AudioComponentID > 0)
-		{
-			return;
-		}
-	}
-
-	FAudioDeviceManager* DeviceManager = GEngine->GetAudioDeviceManager();
-	if (DeviceManager && DeviceManager->IsVisualizeDebug3dEnabled())
+	if (UClass* PluginClass = GetAudioPluginCustomSettingsClass(EAudioPlugin::MODULATION))
 	{
-		DECLARE_CYCLE_STAT(TEXT("FAudioThreadTask.DrawActiveSoundDebugInfo"), STAT_AudioDrawActiveSoundDebugInfo, STATGROUP_TaskGraphTasks);
-
-		const FString Name = Sound->GetName();
-		const FTransform CurTransform = Transform;
-		FColor TextColor = FColor::White;
-		const float CurMaxDistance = MaxDistance;
-		float DisplayValue = 0.0f;
-		if (ActiveSoundVisualizeModeCVar == 1 || ActiveSoundVisualizeModeCVar == 2)
-		{
-			DisplayValue = ThisSoundsWaveInstances && ThisSoundsWaveInstances->Num() == 1
-				? (*ThisSoundsWaveInstances)[0]->GetVolumeWithDistanceAttenuation()
-				: GetVolume();
-		}
-		else if (ActiveSoundVisualizeModeCVar == 3)
-		{
-			if (AudioDevice)
-			{
-				DisplayValue = AudioDevice->GetDistanceToNearestListener(Transform.GetLocation()) / CurMaxDistance;
-			}
-		}
-#if !(UE_BUILD_SHIPPING || UE_BUILD_TEST)
-		else if (ActiveSoundVisualizeModeCVar == 4)
+		if (Sound->Modulation.Settings.Num() > 0)
 		{
-			if (DebugColor == FColor::Black)
+			for (USoundModulationPluginSourceSettingsBase* Settings : Sound->Modulation.Settings)
 			{
-				DebugColor = FColor::MakeRandomColor();
+				if (Settings && Settings->IsA(PluginClass))
+				{
+					return Settings;
+				}
 			}
-
-			TextColor = DebugColor;
 		}
-#endif // !(UE_BUILD_SHIPPING || UE_BUILD_TEST)
-
-		TWeakObjectPtr<UWorld> WorldPtr = GetWeakWorld();
-		FAudioThread::RunCommandOnGameThread([Name, TextColor, CurTransform, DisplayValue, WorldPtr, CurMaxDistance]()
+		else if (UAudioComponent* AudioComponent = UAudioComponent::GetAudioComponentFromID(AudioComponentID))
 		{
-			if (WorldPtr.IsValid())
+			for (USoundModulationPluginSourceSettingsBase* Settings : AudioComponent->Modulation.Settings)
 			{
-				static const float ColorRedHue = 0.0f;
-				static const float ColorGreenHue = 85.0f;
-
-				const FVector Location = CurTransform.GetLocation();
-				UWorld* DebugWorld = WorldPtr.Get();
-				DrawDebugSphere(DebugWorld, Location, 10.0f, 8, FColor::White, false, -1.0f, SDPG_Foreground);
-				FColor Color = TextColor;
-
-				FString Descriptor;
-				if (ActiveSoundVisualizeModeCVar == 1 || ActiveSoundVisualizeModeCVar == 2)
-				{
-					const float DisplayDbVolume = Audio::ConvertToDecibels(DisplayValue);
-					if (ActiveSoundVisualizeModeCVar == 1)
-					{
-						Descriptor = FString::Printf(TEXT(" (Vol: %.3f)"), DisplayValue);
-					}
-					else
-					{
-						Descriptor = FString::Printf(TEXT(" (Vol: %.3f dB)"), DisplayDbVolume);
-					}
-					static const float DbColorMinVol = -30.0f;
-					const float DbVolume = FMath::Clamp(DisplayDbVolume, DbColorMinVol, 0.0f);
-					const float Hue = FMath::Lerp(ColorRedHue, ColorGreenHue, (-1.0f * DbVolume / DbColorMinVol) + 1.0f);
-					Color = FLinearColor::MakeFromHSV8(static_cast<uint8>(Hue), 255u, 255u).ToFColor(true);
-				}
-				else if (ActiveSoundVisualizeModeCVar == 3)
+				if (Settings && Settings->IsA(PluginClass))
 				{
-					Descriptor = FString::Printf(TEXT(" (Dist: %.3f, Max: %.3f)"), DisplayValue * CurMaxDistance, CurMaxDistance);
-					const float Hue = FMath::Lerp(ColorGreenHue, ColorRedHue, DisplayValue);
-					Color = FLinearColor::MakeFromHSV8(static_cast<uint8>(FMath::Clamp(Hue, 0.0f, 255.f)), 255u, 255u).ToFColor(true);
+					return Settings;
 				}
-
-				const FString Description = FString::Printf(TEXT("%s%s"), *Name, *Descriptor);
-				DrawDebugString(DebugWorld, Location + FVector(0, 0, 32), *Description, nullptr, Color, 0.03f, false);
 			}
-		}, GET_STATID(STAT_AudioDrawActiveSoundDebugInfo));
+		}
 	}
-#endif // ENABLE_DRAW_DEBUG
+
+	return nullptr;
 }
 
 void FActiveSound::Stop(bool bStopNow)
@@ -804,7 +722,7 @@ void FActiveSound::Stop(bool bStopNow)
 		if (Source)
 		{
 			if (!Source->IsStopping())
-			{		
+			{
 				Source->StopNow();
 
 				delete WaveInstance;
@@ -882,7 +800,7 @@ bool FActiveSound::UpdateStoppingSources(uint64 CurrentTick, bool bEnsureStopped
 			}
 			else
 			{
-				// We have a wave instance but no source for it, so just delete it. 
+				// We have a wave instance but no source for it, so just delete it.
 				delete WaveInstance;
 				WaveInstance = nullptr;
 			}
@@ -937,7 +855,7 @@ void FActiveSound::UpdateAdjustVolumeMultiplier(const float DeltaTime)
 
 	// Apply final clamp
 	CurrentAdjustVolumeMultiplier = FMath::Clamp(CurrentAdjustVolumeMultiplier, MinValue, MaxValue);
-} 
+}
 
 void FActiveSound::OcclusionTraceDone(const FTraceHandle& TraceHandle, FTraceDatum& TraceDatum)
 {
@@ -1003,7 +921,7 @@ void FActiveSound::CheckOcclusion(const FVector ListenerLocation, const FVector
 
 				if (UWorld* WorldPtr = World.Get())
 				{
-					// LineTraceTestByChannel is generally threadsafe, but there is a very narrow race condition here 
+					// LineTraceTestByChannel is generally threadsafe, but there is a very narrow race condition here
 					// if World goes invalid before the scene lock and queries begin.
 					bIsOccluded = WorldPtr->LineTraceTestByChannel(SoundLocation, ListenerLocation, OcclusionTraceChannel, Params);
 				}
diff --git a/Engine/Source/Runtime/Engine/Private/Audio.cpp b/Engine/Source/Runtime/Engine/Private/Audio.cpp
--- a/Engine/Source/Runtime/Engine/Private/Audio.cpp
+++ b/Engine/Source/Runtime/Engine/Private/Audio.cpp
@@ -11,8 +11,9 @@
 #include "Components/AudioComponent.h"
 #include "ContentStreaming.h"
 #include "DrawDebugHelpers.h"
-#include "AudioThread.h"
+#include "Audio/AudioDebug.h"
 #include "AudioDevice.h"
+#include "AudioThread.h"
 #include "Sound/SoundBase.h"
 #include "Sound/SoundCue.h"
 #include "Sound/SoundWave.h"
@@ -77,13 +78,6 @@ FAutoConsoleVariableRef CVarAllowAudioSpatializationCVar(
 	TEXT("0: Disable, >0: Enable"),
 	ECVF_Default);
 
-static int32 SpatialSourceVisualizeEnabledCVar = 1;
-FAutoConsoleVariableRef CVarAudioVisualizeSpatialSourceEnabled(
-	TEXT("au.3dVisualize.SpatialSources"),
-	SpatialSourceVisualizeEnabledCVar,
-	TEXT("Whether or not audio spatialized sources are visible when 3d visualize is enabled. \n")
-	TEXT("0: Not Enabled, 1: Enabled"),
-	ECVF_Default);
 
 bool IsAudioPluginEnabled(EAudioPlugin PluginType)
 {
@@ -91,13 +85,12 @@ bool IsAudioPluginEnabled(EAudioPlugin PluginType)
 	{
 	case EAudioPlugin::SPATIALIZATION:
 		return AudioPluginUtilities::GetDesiredSpatializationPlugin(AudioPluginUtilities::CurrentPlatform) != nullptr;
-		break;
 	case EAudioPlugin::REVERB:
 		return AudioPluginUtilities::GetDesiredReverbPlugin(AudioPluginUtilities::CurrentPlatform) != nullptr;
-		break;
 	case EAudioPlugin::OCCLUSION:
 		return AudioPluginUtilities::GetDesiredOcclusionPlugin(AudioPluginUtilities::CurrentPlatform) != nullptr;
-		break;
+	case EAudioPlugin::MODULATION:
+		return AudioPluginUtilities::GetDesiredModulationPlugin(AudioPluginUtilities::CurrentPlatform) != nullptr;
 	default:
 		return false;
 		break;
@@ -110,8 +103,7 @@ UClass* GetAudioPluginCustomSettingsClass(EAudioPlugin PluginType)
 	{
 		case EAudioPlugin::SPATIALIZATION:
 		{
-			IAudioSpatializationFactory* Factory = AudioPluginUtilities::GetDesiredSpatializationPlugin(AudioPluginUtilities::CurrentPlatform);
-			if (Factory)
+			if (IAudioSpatializationFactory* Factory = AudioPluginUtilities::GetDesiredSpatializationPlugin(AudioPluginUtilities::CurrentPlatform))
 			{
 				return Factory->GetCustomSpatializationSettingsClass();
 			}
@@ -120,8 +112,7 @@ UClass* GetAudioPluginCustomSettingsClass(EAudioPlugin PluginType)
 
 		case EAudioPlugin::REVERB:
 		{
-			IAudioReverbFactory* Factory = AudioPluginUtilities::GetDesiredReverbPlugin(AudioPluginUtilities::CurrentPlatform);
-			if (Factory)
+			if (IAudioReverbFactory* Factory = AudioPluginUtilities::GetDesiredReverbPlugin(AudioPluginUtilities::CurrentPlatform))
 			{
 				return Factory->GetCustomReverbSettingsClass();
 			}
@@ -130,15 +121,24 @@ UClass* GetAudioPluginCustomSettingsClass(EAudioPlugin PluginType)
 
 		case EAudioPlugin::OCCLUSION:
 		{
-			IAudioOcclusionFactory* Factory = AudioPluginUtilities::GetDesiredOcclusionPlugin(AudioPluginUtilities::CurrentPlatform);
-			if (Factory)
+			if (IAudioOcclusionFactory* Factory = AudioPluginUtilities::GetDesiredOcclusionPlugin(AudioPluginUtilities::CurrentPlatform))
 			{
 				return Factory->GetCustomOcclusionSettingsClass();
 			}
 		}
 		break;
 
+		case EAudioPlugin::MODULATION:
+		{
+			if (IAudioModulationFactory* Factory = AudioPluginUtilities::GetDesiredModulationPlugin(AudioPluginUtilities::CurrentPlatform))
+			{
+				return Factory->GetCustomModulationSettingsClass();
+			}
+		}
+		break;
+
 		default:
+			static_assert(static_cast<uint32>(EAudioPlugin::COUNT) == 4, "Possible missing audio plugin type case coverage");
 		break;
 	}
 
@@ -354,37 +354,53 @@ float FSoundSource::SetLFEBleed()
 
 void FSoundSource::SetFilterFrequency()
 {
-	LPFFrequency = MAX_FILTER_FREQUENCY;
-
-	if (AudioDevice->GetMixDebugState() == DEBUGSTATE_TestLPF)
-	{
-		// If in debug mode, lets set all sounds to a LPF of MIN_FILTER_FREQUENCY
-		LPFFrequency = MIN_FILTER_FREQUENCY;
-	}
-	else if (AudioDevice->GetMixDebugState() != DEBUGSTATE_DisableLPF)
+	// HPF is only available with audio mixer enabled
+	switch (AudioDevice->GetMixDebugState())
 	{
-		// If so, override the frequency with the occluded filter frequency
-		LPFFrequency = WaveInstance->OcclusionFilterFrequency;
+		case DEBUGSTATE_TestLPF:
+		{
+			LPFFrequency = MIN_FILTER_FREQUENCY;
+		}
+		break;
+
+		case DEBUGSTATE_DisableLPF:
+		{
+			LPFFrequency = MAX_FILTER_FREQUENCY;
+		}
+		break;
 
-		// Set the LPFFrequency to the manual LowPassFilterFrequency if it's lower
-		if (WaveInstance->bEnableLowPassFilter && WaveInstance->LowPassFilterFrequency < LPFFrequency)
+		default:
 		{
-			LPFFrequency = WaveInstance->LowPassFilterFrequency;
+			// Set the LPFFrequency to lowest provided value
+			LPFFrequency = FMath::Min(WaveInstance->OcclusionFilterFrequency, WaveInstance->LowPassFilterFrequency);
+			LPFFrequency = FMath::Min(LPFFrequency, WaveInstance->AmbientZoneFilterFrequency);
+			LPFFrequency = FMath::Min(LPFFrequency, WaveInstance->AttenuationLowpassFilterFrequency);
+			LPFFrequency = FMath::Min(LPFFrequency, WaveInstance->SoundModulationControls.Lowpass);
 		}
+		break;
+	}
 
-		// Set the LPFFrequency to the ambient filter frequency if it's lower
-		if (WaveInstance->AmbientZoneFilterFrequency < LPFFrequency)
+	// HPF is only available with audio mixer enabled
+	switch (AudioDevice->GetMixDebugState())
+	{
+		case DEBUGSTATE_TestHPF:
 		{
-			LPFFrequency = WaveInstance->AmbientZoneFilterFrequency;
+			HPFFrequency = MAX_FILTER_FREQUENCY;
 		}
+		break;
 
-		if (WaveInstance->AttenuationLowpassFilterFrequency < LPFFrequency)
+		case DEBUGSTATE_DisableHPF:
 		{
-			LPFFrequency = WaveInstance->AttenuationLowpassFilterFrequency;
+			HPFFrequency = MIN_FILTER_FREQUENCY;
 		}
+		break;
 
-		// This is only used in audio mixer, and only one thing is setting HPF
-		HPFFrequency = WaveInstance->AttenuationHighpassFilterFrequency;
+		default:
+		{
+			// Set the HPFFrequency to highest provided value
+			HPFFrequency = FMath::Max(WaveInstance->AttenuationHighpassFilterFrequency, WaveInstance->SoundModulationControls.Highpass);
+		}
+		break;
 	}
 }
 
@@ -419,52 +435,6 @@ void FSoundSource::UpdateStereoEmitterPositions()
 	}
 }
 
-void FSoundSource::DrawDebugInfo()
-{
-#if ENABLE_DRAW_DEBUG
-	if (!WaveInstance)
-	{
-		return;
-	}
-
-	const FActiveSound* ActiveSound = WaveInstance->ActiveSound;
-	if (!ActiveSound)
-	{
-		return;
-	}
-
-	if (!SpatialSourceVisualizeEnabledCVar)
-	{
-		return;
-	}
-
-	FAudioDeviceManager* DeviceManager = GEngine->GetAudioDeviceManager();
-	if (DeviceManager && DeviceManager->IsVisualizeDebug3dEnabled())
-	{
-		DECLARE_CYCLE_STAT(TEXT("FAudioThreadTask.DrawSourceDebugInfo"), STAT_AudioDrawSourceDebugInfo, STATGROUP_TaskGraphTasks);
-
-		const bool bSpatialized = Buffer->NumChannels == 2 && WaveInstance->GetUseSpatialization();
-		if (bSpatialized)
-		{
-			const FRotator Rotator = ActiveSound->Transform.GetRotation().Rotator();
-
-			TWeakObjectPtr<UWorld> WorldPtr = WaveInstance->ActiveSound->GetWeakWorld();
-			const FVector LeftChannelSourceLoc = LeftChannelSourceLocation;
-			const FVector RightChannelSourceLoc = RightChannelSourceLocation;
-			FAudioThread::RunCommandOnGameThread([LeftChannelSourceLoc, RightChannelSourceLoc, Rotator, WorldPtr]()
-			{
-				if (WorldPtr.IsValid())
-				{
-					UWorld* World = WorldPtr.Get();
-					DrawDebugCrosshairs(World, LeftChannelSourceLoc, Rotator, 20.0f, FColor::Red, false, -1.0f, SDPG_Foreground);
-					DrawDebugCrosshairs(World, RightChannelSourceLoc, Rotator, 20.0f, FColor::Green, false, -1.0f, SDPG_Foreground);
-				}
-			}, GET_STATID(STAT_AudioDrawSourceDebugInfo));
-		}
-	}
-#endif // ENABLE_DRAW_DEBUG
-}
-
 float FSoundSource::GetDebugVolume(const float InVolume)
 {
 	float OutVolume = InVolume;
@@ -615,7 +585,7 @@ void FSoundSource::UpdateCommon()
 {
 	check(WaveInstance);
 
-	Pitch = WaveInstance->Pitch;
+	Pitch = WaveInstance->GetPitch();
 
 	// Don't apply global pitch scale to UI sounds
 	if (!WaveInstance->bIsUISound)
@@ -647,6 +617,13 @@ float FSoundSource::GetPlaybackPercent() const
 
 }
 
+void FSoundSource::GetChannelLocations(FVector& Left, FVector&Right) const
+{
+	Left = LeftChannelSourceLocation;
+	Right = RightChannelSourceLocation;
+}
+
+
 void FSoundSource::NotifyPlaybackData()
 {
 	const uint64 AudioComponentID = WaveInstance->ActiveSound->GetAudioComponentID();
@@ -783,7 +760,6 @@ FWaveInstance::FWaveInstance( FActiveSound* InActiveSound )
 	, Volume(0.0f)
 	, DistanceAttenuation(1.0f)
 	, VolumeMultiplier(1.0f)
-	, VolumeApp(1.0f)
 	, EnvelopValue(0.0f)
 	, EnvelopeFollowerAttackTime(10)
 	, EnvelopeFollowerReleaseTime(100)
@@ -905,10 +881,17 @@ void FWaveInstance::AddReferencedObjects( FReferenceCollector& Collector )
 float FWaveInstance::GetActualVolume() const
 {
 	// Include all volumes
-	float ActualVolume = GetVolume() * VolumeApp * DistanceAttenuation;
+	float ActualVolume = GetVolume() * DistanceAttenuation;
 	if (ActualVolume != 0.0f)
 	{
 		ActualVolume *= GetDynamicVolume();
+
+		check(ActiveSound);
+		if (!ActiveSound->bIsPreviewSound)
+		{
+			check(ActiveSound->AudioDevice);
+			ActualVolume *= ActiveSound->AudioDevice->GetMasterVolume();
+		}
 	}
 
 	return ActualVolume;
@@ -945,10 +928,15 @@ float FWaveInstance::GetVolumeWithDistanceAttenuation() const
 	return GetVolume() * DistanceAttenuation;
 }
 
+float FWaveInstance::GetPitch() const
+{
+	return Pitch * SoundModulationControls.Pitch;
+}
+
 float FWaveInstance::GetVolume() const
 {
 	// Only includes non-attenuation and non-app volumes
-	return Volume * VolumeMultiplier;
+	return Volume * VolumeMultiplier * SoundModulationControls.Volume;
 }
 
 bool FWaveInstance::ShouldStopDueToMaxConcurrency() const
diff --git a/Engine/Source/Runtime/Engine/Private/AudioDebug.cpp b/Engine/Source/Runtime/Engine/Private/AudioDebug.cpp
new file mode 100644
--- /dev/null
+++ b/Engine/Source/Runtime/Engine/Private/AudioDebug.cpp
@@ -0,0 +1,1272 @@
+// Copyright 1998-2019 Epic Games, Inc. All Rights Reserved.
+#include "Audio/AudioDebug.h"
+
+#include "ActiveSound.h"
+#include "Audio.h"
+#include "AudioDevice.h"
+#include "AudioDeviceManager.h"
+#include "AudioEffect.h"
+#include "AudioVirtualLoop.h"
+#include "CanvasTypes.h"
+#include "DrawDebugHelpers.h"
+#include "Engine/Font.h"
+#include "GameFramework/GameUserSettings.h"
+#include "HAL/IConsoleManager.h"
+#include "Misc/CommandLine.h"
+#include "Misc/ConfigCacheIni.h"
+#include "Sound/AudioSettings.h"
+#include "Sound/AudioVolume.h"
+#include "Sound/ReverbEffect.h"
+#include "Sound/SoundAttenuation.h"
+#include "Sound/SoundClass.h"
+#include "Sound/SoundConcurrency.h"
+#include "Sound/SoundMix.h"
+#include "Sound/SoundSourceBus.h"
+#include "Sound/SoundWave.h"
+#include "UnrealEngine.h"
+
+
+#if ENABLE_AUDIO_DEBUG
+
+// Console variables
+static int32 ActiveSoundVisualizeModeCVar = 1;
+FAutoConsoleVariableRef CVarAudioVisualizeActiveSoundsMode(
+	TEXT("au.3dVisualize.ActiveSounds"),
+	ActiveSoundVisualizeModeCVar,
+	TEXT("Visualization mode for active sounds. \n")
+	TEXT("0: Not Enabled, 1: Volume (Lin), 2: Volume (dB), 3: Distance, 4: Random color"),
+	ECVF_Default);
+
+static int32 ActiveSoundVisualizeTypeCVar = 0;
+FAutoConsoleVariableRef CVarAudioVisualizeActiveSounds(
+	TEXT("au.3dVisualize.ActiveSounds.Type"),
+	ActiveSoundVisualizeTypeCVar,
+	TEXT("Whether to show all sounds, on AudioComponents (Components Only), or off of AudioComponents (Non-Component Only). \n")
+	TEXT("0: All, 1: Components Only, 2: Non-Component Only"),
+	ECVF_Default);
+
+static int32 SpatialSourceVisualizeEnabledCVar = 1;
+FAutoConsoleVariableRef CVarAudioVisualizeSpatialSourceEnabled(
+	TEXT("au.3dVisualize.SpatialSources"),
+	SpatialSourceVisualizeEnabledCVar,
+	TEXT("Whether or not audio spatialized sources are visible when 3d visualize is enabled. \n")
+	TEXT("0: Not Enabled, 1: Enabled"),
+	ECVF_Default);
+
+static int32 VirtualLoopsVisualizeEnabledCVar = 1;
+FAutoConsoleVariableRef CVarAudioVisualizeVirtualLoopsEnabled(
+	TEXT("au.3dVisualize.VirtualLoops"),
+	VirtualLoopsVisualizeEnabledCVar,
+	TEXT("Whether or not virtualized loops are visible when 3d visualize is enabled. \n")
+	TEXT("0: Not Enabled, 1: Enabled"),
+	ECVF_Default);
+
+// Test Data
+namespace ERequestedAudioStats
+{
+	static const uint8 SoundWaves = 0x1;
+	static const uint8 SoundCues = 0x2;
+	static const uint8 Sounds = 0x4;
+	static const uint8 SoundMixes = 0x8;
+	static const uint8 SoundModulation = 0x10;
+	static const uint8 DebugSounds = 0x20;
+	static const uint8 LongSoundNames = 0x40;
+};
+
+namespace
+{
+	const FColor HeaderColor = FColor::Green;
+	const FColor BodyColor = FColor::White;
+	const int32  TabWidth = 12;
+
+	const float MinDisplayVolume = KINDA_SMALL_NUMBER; // -80 dB
+
+	FAudioDevice* GetWorldAudio(UWorld* World)
+	{
+		check(IsInGameThread());
+
+		if (!World)
+		{
+			return nullptr;
+		}
+
+		return World->GetAudioDevice();
+	}
+
+	struct FAudioStats
+	{
+		enum class EDisplayFlags : uint8
+		{
+			Debug = 0x01,
+			Sort_Distance = 0x02,
+			Sort_Class = 0x04,
+			Sort_Name = 0x08,
+			Sort_WavesNum = 0x10,
+			Sort_Disabled = 0x20,
+			Long_Names = 0x40,
+		};
+
+		struct FStatWaveInstanceInfo
+		{
+			FString Description;
+			float Volume;
+			int32 InstanceIndex;
+			FName WaveInstanceName;
+			uint8 bPlayWhenSilent:1;
+		};
+
+		struct FStatSoundInfo
+		{
+			FString SoundName;
+			FName SoundClassName;
+			float Distance;
+			uint32 AudioComponentID;
+			FTransform Transform;
+			TArray<FStatWaveInstanceInfo> WaveInstanceInfos;
+			TMultiMap<EAttenuationShape::Type, FBaseAttenuationSettings::AttenuationShapeDetails> ShapeDetailsMap;
+		};
+
+		struct FStatSoundMix
+		{
+			FString MixName;
+			float InterpValue;
+			int32 RefCount;
+			bool bIsCurrentEQ;
+		};
+
+		uint8 DisplayFlags;
+		TArray<FTransform> ListenerTransforms;
+		TArray<FStatSoundInfo> StatSoundInfos;
+		TArray<FStatSoundMix> StatSoundMixes;
+
+		FAudioStats()
+			: DisplayFlags(0)
+		{
+		}
+	};
+
+	struct FAudioStats_AudioThread
+	{
+		uint8 RequestedStats;
+
+		FAudioStats_AudioThread()
+			: RequestedStats(0)
+		{
+		}
+	};
+
+	TMap<uint32, FAudioStats> AudioDeviceStats;
+	TMap<uint32, FAudioStats_AudioThread> AudioDeviceStats_AudioThread;
+
+	void HandleDumpActiveSounds(UWorld* World)
+	{
+		if (GEngine)
+		{
+			GEngine->GetAudioDeviceManager()->GetDebugger().DumpActiveSounds();
+		}
+	}
+
+	FAutoConsoleCommandWithWorld DumpActiveSounds(TEXT("Audio.DumpActiveSounds"), TEXT("Outputs data about all the currently active sounds."), FConsoleCommandWithWorldDelegate::CreateStatic(&HandleDumpActiveSounds), ECVF_Cheat);
+} // namespace <>
+
+
+/**
+ * Audio Debugger Implementation
+ */
+FAudioDebugger::FAudioDebugger()
+	: bVisualize3dDebug(0)
+{
+}
+
+bool FAudioDebugger::IsVisualizeDebug3dEnabled() const
+{
+	return bVisualize3dDebug;
+}
+
+void FAudioDebugger::ToggleVisualizeDebug3dEnabled()
+{
+	bVisualize3dDebug = !bVisualize3dDebug;
+}
+
+void FAudioDebugger::DrawDebugInfo(const FSoundSource& SoundSource)
+{
+#if ENABLE_DRAW_DEBUG
+	const FWaveInstance* WaveInstance = SoundSource.GetWaveInstance();
+	if (!WaveInstance)
+	{
+		return;
+	}
+
+	const FActiveSound* ActiveSound = WaveInstance->ActiveSound;
+	if (!ActiveSound)
+	{
+		return;
+	}
+
+	if (!SpatialSourceVisualizeEnabledCVar)
+	{
+		return;
+	}
+
+	FAudioDeviceManager* DeviceManager = GEngine->GetAudioDeviceManager();
+	if (DeviceManager && DeviceManager->IsVisualizeDebug3dEnabled())
+	{
+		DECLARE_CYCLE_STAT(TEXT("FAudioThreadTask.DrawSourceDebugInfo"), STAT_AudioDrawSourceDebugInfo, STATGROUP_TaskGraphTasks);
+
+		const FSoundBuffer* Buffer = SoundSource.GetBuffer();
+		const bool bSpatialized = Buffer && Buffer->NumChannels == 2 && WaveInstance->GetUseSpatialization();
+		if (bSpatialized)
+		{
+			const FRotator Rotator = ActiveSound->Transform.GetRotation().Rotator();
+
+			TWeakObjectPtr<UWorld> WorldPtr = WaveInstance->ActiveSound->GetWeakWorld();
+			FVector LeftChannelSourceLoc;
+			FVector RightChannelSourceLoc;
+			SoundSource.GetChannelLocations(LeftChannelSourceLoc, RightChannelSourceLoc);
+			FAudioThread::RunCommandOnGameThread([LeftChannelSourceLoc, RightChannelSourceLoc, Rotator, WorldPtr]()
+			{
+				if (WorldPtr.IsValid())
+				{
+					UWorld* World = WorldPtr.Get();
+					DrawDebugCrosshairs(World, LeftChannelSourceLoc, Rotator, 20.0f, FColor::Red, false, -1.0f, SDPG_Foreground);
+					DrawDebugCrosshairs(World, RightChannelSourceLoc, Rotator, 20.0f, FColor::Green, false, -1.0f, SDPG_Foreground);
+				}
+			}, GET_STATID(STAT_AudioDrawSourceDebugInfo));
+		}
+	}
+#endif // ENABLE_DRAW_DEBUG
+}
+
+void FAudioDebugger::DrawDebugInfo(const FActiveSound& ActiveSound, const TArray<FWaveInstance*>& ThisSoundsWaveInstances)
+{
+#if ENABLE_DRAW_DEBUG
+	if (!ActiveSoundVisualizeModeCVar)
+	{
+		return;
+	}
+
+	// Only draw spatialized sounds
+	const USoundBase* Sound = ActiveSound.GetSound();
+	if (!Sound || !ActiveSound.bAllowSpatialization)
+	{
+		return;
+	}
+
+	if (ActiveSoundVisualizeTypeCVar > 0)
+	{
+		if (ActiveSoundVisualizeTypeCVar == 1 && ActiveSound.GetAudioComponentID() == 0)
+		{
+			return;
+		}
+
+		if (ActiveSoundVisualizeTypeCVar == 2 && ActiveSound.GetAudioComponentID() > 0)
+		{
+			return;
+		}
+	}
+
+	FAudioDeviceManager* DeviceManager = GEngine->GetAudioDeviceManager();
+	if (DeviceManager && DeviceManager->IsVisualizeDebug3dEnabled())
+	{
+		DECLARE_CYCLE_STAT(TEXT("FAudioThreadTask.DrawActiveSoundDebugInfo"), STAT_AudioDrawActiveSoundDebugInfo, STATGROUP_TaskGraphTasks);
+
+		const FString Name = Sound->GetName();
+		const FTransform CurTransform = ActiveSound.Transform;
+		FColor TextColor = FColor::White;
+		const float CurMaxDistance = ActiveSound.MaxDistance;
+		float DisplayValue = 0.0f;
+		if (ActiveSoundVisualizeModeCVar == 1 || ActiveSoundVisualizeModeCVar == 2)
+		{
+			DisplayValue = ActiveSound.GetVolume();
+			if (ThisSoundsWaveInstances.Num() == 1)
+			{
+				const FWaveInstance* WaveInstance = ThisSoundsWaveInstances[0];
+				DisplayValue = WaveInstance->GetVolumeWithDistanceAttenuation() * WaveInstance->GetDynamicVolume();
+			}
+		}
+		else if (ActiveSoundVisualizeModeCVar == 3)
+		{
+			if (ActiveSound.AudioDevice)
+			{
+				DisplayValue = ActiveSound.AudioDevice->GetDistanceToNearestListener(ActiveSound.Transform.GetLocation()) / CurMaxDistance;
+			}
+		}
+		else if (ActiveSoundVisualizeModeCVar == 4)
+		{
+			TextColor = ActiveSound.DebugColor;
+		}
+
+		TWeakObjectPtr<UWorld> WorldPtr = ActiveSound.GetWeakWorld();
+		FAudioThread::RunCommandOnGameThread([Name, TextColor, CurTransform, DisplayValue, WorldPtr, CurMaxDistance]()
+		{
+			if (WorldPtr.IsValid())
+			{
+				static const float ColorRedHue = 0.0f;
+				static const float ColorGreenHue = 85.0f;
+
+				const FVector Location = CurTransform.GetLocation();
+				UWorld* DebugWorld = WorldPtr.Get();
+				DrawDebugSphere(DebugWorld, Location, 10.0f, 8, FColor::White, false, -1.0f, SDPG_Foreground);
+				FColor Color = TextColor;
+
+				FString Descriptor;
+				if (ActiveSoundVisualizeModeCVar == 1 || ActiveSoundVisualizeModeCVar == 2)
+				{
+					const float DisplayDbVolume = Audio::ConvertToDecibels(DisplayValue);
+					if (ActiveSoundVisualizeModeCVar == 1)
+					{
+						Descriptor = FString::Printf(TEXT(" (Vol: %.3f)"), DisplayValue);
+					}
+					else
+					{
+						Descriptor = FString::Printf(TEXT(" (Vol: %.3f dB)"), DisplayDbVolume);
+					}
+					static const float DbColorMinVol = -30.0f;
+					const float DbVolume = FMath::Clamp(DisplayDbVolume, DbColorMinVol, 0.0f);
+					const float Hue = FMath::Lerp(ColorRedHue, ColorGreenHue, (-1.0f * DbVolume / DbColorMinVol) + 1.0f);
+					Color = FLinearColor::MakeFromHSV8(static_cast<uint8>(Hue), 255u, 255u).ToFColor(true);
+				}
+				else if (ActiveSoundVisualizeModeCVar == 3)
+				{
+					Descriptor = FString::Printf(TEXT(" (Dist: %.3f, Max: %.3f)"), DisplayValue * CurMaxDistance, CurMaxDistance);
+					const float Hue = FMath::Lerp(ColorGreenHue, ColorRedHue, DisplayValue);
+					Color = FLinearColor::MakeFromHSV8(static_cast<uint8>(FMath::Clamp(Hue, 0.0f, 255.f)), 255u, 255u).ToFColor(true);
+				}
+
+				const FString Description = FString::Printf(TEXT("%s%s"), *Name, *Descriptor);
+				DrawDebugString(DebugWorld, Location + FVector(0, 0, 32), *Description, nullptr, Color, 0.03f, false);
+			}
+		}, GET_STATID(STAT_AudioDrawActiveSoundDebugInfo));
+	}
+#endif // ENABLE_DRAW_DEBUG
+}
+
+void FAudioDebugger::DrawDebugInfo(const FAudioVirtualLoop& VirtualLoop)
+{
+#if ENABLE_DRAW_DEBUG
+	if (!GEngine)
+	{
+		return;
+	}
+
+	if (!VirtualLoopsVisualizeEnabledCVar)
+	{
+		return;
+	}
+
+	FAudioDeviceManager* DeviceManager = GEngine->GetAudioDeviceManager();
+	if (DeviceManager && DeviceManager->IsVisualizeDebug3dEnabled())
+	{
+		DECLARE_CYCLE_STAT(TEXT("FAudioThreadTask.DrawVirtualLoopDebugInfo"), STAT_AudioDrawVirtualLoopDebugInfo, STATGROUP_TaskGraphTasks);
+
+		const FActiveSound& ActiveSound = VirtualLoop.GetActiveSound();
+		USoundBase* Sound = ActiveSound.GetSound();
+		check(Sound);
+
+		const FTransform Transform = ActiveSound.Transform;
+		const TWeakObjectPtr<UWorld> World = ActiveSound.GetWeakWorld();
+		const FString Name = Sound->GetName();
+		const float DrawInterval = VirtualLoop.GetUpdateInterval();
+		FAudioThread::RunCommandOnGameThread([World, Transform, Name, DrawInterval]()
+		{
+			if (World.IsValid())
+			{
+				const FString Description = FString::Printf(TEXT("%s [V]"), *Name);
+				FVector Location = Transform.GetLocation();
+				FRotator Rotation = Transform.GetRotation().Rotator();
+				DrawDebugCrosshairs(World.Get(), Location, Rotation, 20.0f, FColor::Blue, false, DrawInterval, SDPG_Foreground);
+				DrawDebugString(World.Get(), Location + FVector(0, 0, 32), *Description, nullptr, FColor::Blue, DrawInterval, false);
+			}
+		}, GET_STATID(STAT_AudioDrawVirtualLoopDebugInfo));
+	}
+#endif // ENABLE_DRAW_DEBUG
+}
+
+void FAudioDebugger::DumpActiveSounds() const
+{
+	if (!GEngine)
+	{
+		return;
+	}
+
+	if (!IsInAudioThread())
+	{
+		DECLARE_CYCLE_STAT(TEXT("FAudioThreadTask.DumpActiveSounds"), STAT_AudioDumpActiveSounds, STATGROUP_TaskGraphTasks);
+		FAudioThread::RunCommandOnAudioThread([this]()
+		{
+			DumpActiveSounds();
+		}, GET_STATID(STAT_AudioDumpActiveSounds));
+		return;
+	}
+
+	FAudioDevice* AudioDevice = GEngine->GetAudioDeviceManager()->GetActiveAudioDevice();
+	if (!AudioDevice)
+	{
+		return;
+	}
+
+	const TArray<FActiveSound*>& ActiveSounds = AudioDevice->GetActiveSounds();
+	UE_LOG(LogAudio, Display, TEXT("Active Sound Count: %d"), ActiveSounds.Num());
+	UE_LOG(LogAudio, Display, TEXT("------------------------"), ActiveSounds.Num());
+
+	for (const FActiveSound* ActiveSound : ActiveSounds)
+	{
+		if (ActiveSound)
+		{
+			UE_LOG(LogAudio, Display, TEXT("%s (%.3g) - %s"), *ActiveSound->GetSound()->GetName(), ActiveSound->GetSound()->GetDuration(), *ActiveSound->GetAudioComponentName());
+
+			for (const TPair<UPTRINT, FWaveInstance*>& WaveInstancePair : ActiveSound->WaveInstances)
+			{
+				const FWaveInstance* WaveInstance = WaveInstancePair.Value;
+				UE_LOG(LogAudio, Display, TEXT("   %s (%.3g) (%d) - %.3g"),
+					*WaveInstance->GetName(), WaveInstance->WaveData->GetDuration(),
+					WaveInstance->WaveData->GetResourceSizeBytes(EResourceSizeMode::EstimatedTotal),
+					WaveInstance->GetVolumeWithDistanceAttenuation() * WaveInstance->GetDynamicVolume());
+			}
+		}
+	}
+}
+
+void FAudioDebugger::ResolveDesiredStats(FViewportClient* ViewportClient)
+{
+	if (!ViewportClient)
+	{
+		return;
+	}
+
+	FAudioDevice* AudioDevice = GetWorldAudio(ViewportClient->GetWorld());
+	if (!AudioDevice)
+	{
+		return;
+	}
+
+	uint8 SetStats = 0;
+	uint8 ClearStats = 0;
+
+	if (ViewportClient->IsStatEnabled(TEXT("SoundCues")))
+	{
+		SetStats |= ERequestedAudioStats::SoundCues;
+	}
+	else
+	{
+		ClearStats |= ERequestedAudioStats::SoundCues;
+	}
+
+	if (ViewportClient->IsStatEnabled(TEXT("SoundWaves")))
+	{
+		SetStats |= ERequestedAudioStats::SoundWaves;
+	}
+	else
+	{
+		ClearStats |= ERequestedAudioStats::SoundWaves;
+	}
+
+	if (ViewportClient->IsStatEnabled(TEXT("SoundMixes")))
+	{
+		SetStats |= ERequestedAudioStats::SoundMixes;
+	}
+	else
+	{
+		ClearStats |= ERequestedAudioStats::SoundMixes;
+	}
+
+	if (ViewportClient->IsStatEnabled(TEXT("Sounds")))
+	{
+		FAudioStats& Stats = AudioDeviceStats.FindOrAdd(AudioDevice->DeviceHandle);
+		SetStats |= ERequestedAudioStats::Sounds;
+
+		if (Stats.DisplayFlags & static_cast<uint8>(FAudioStats::EDisplayFlags::Debug))
+		{
+			SetStats |= ERequestedAudioStats::DebugSounds;
+		}
+		else
+		{
+			ClearStats |= ERequestedAudioStats::DebugSounds;
+		}
+
+		if (Stats.DisplayFlags & static_cast<uint8>(FAudioStats::EDisplayFlags::Long_Names))
+		{
+			SetStats |= ERequestedAudioStats::LongSoundNames;
+		}
+		else
+		{
+			ClearStats |= ERequestedAudioStats::LongSoundNames;
+		}
+	}
+	else
+	{
+		ClearStats |= ERequestedAudioStats::Sounds;
+		ClearStats |= ERequestedAudioStats::DebugSounds;
+		ClearStats |= ERequestedAudioStats::LongSoundNames;
+	}
+
+	DECLARE_CYCLE_STAT(TEXT("FAudioThreadTask.ResolveDesiredStats"), STAT_AudioResolveDesiredStats, STATGROUP_TaskGraphTasks);
+
+	const uint32 DeviceID = AudioDevice->DeviceHandle;
+	if (IsInAudioThread())
+	{
+		FAudioStats_AudioThread& Stats = AudioDeviceStats_AudioThread.FindOrAdd(DeviceID);
+		Stats.RequestedStats |= SetStats;
+		Stats.RequestedStats &= ~ClearStats;
+	}
+	else
+	{
+		FAudioThread::RunCommandOnAudioThread([SetStats, ClearStats, DeviceID]()
+		{
+			FAudioStats_AudioThread& Stats = AudioDeviceStats_AudioThread.FindOrAdd(DeviceID);
+			Stats.RequestedStats |= SetStats;
+			Stats.RequestedStats &= ~ClearStats;
+		}, GET_STATID(STAT_AudioResolveDesiredStats));
+	}
+}
+
+int32 FAudioDebugger::RenderStatCues(UWorld* World, FViewport* Viewport, FCanvas* Canvas, int32 X, int32 Y, const FVector* ViewLocation, const FRotator* ViewRotation)
+{
+	FAudioDevice* AudioDevice = GetWorldAudio(World);
+	if (!AudioDevice)
+	{
+		return Y;
+	}
+
+	const int32 FontHeight = GetStatsFont()->GetMaxCharHeight() + 2;
+	Canvas->DrawShadowedString(X, Y, TEXT("Active Sound Cues:"), GetStatsFont(), HeaderColor);
+	Y += FontHeight;
+
+	int32 ActiveSoundCount = 0;
+	FAudioStats& AudioStats = AudioDeviceStats.FindOrAdd(AudioDevice->DeviceHandle);
+	for (const FAudioStats::FStatSoundInfo& StatSoundInfo : AudioStats.StatSoundInfos)
+	{
+		for (const FAudioStats::FStatWaveInstanceInfo& WaveInstanceInfo : StatSoundInfo.WaveInstanceInfos)
+		{
+			if (WaveInstanceInfo.Volume >= MinDisplayVolume)
+			{
+				const FString TheString = FString::Printf(TEXT("%4i. %s %s"), ActiveSoundCount++, *StatSoundInfo.SoundName, *StatSoundInfo.SoundClassName.ToString());
+				Canvas->DrawShadowedString(X, Y, *TheString, GetStatsFont(), FColor::White);
+				Y += FontHeight;
+				break;
+			}
+		}
+	}
+
+	Canvas->DrawShadowedString(X, Y, *FString::Printf(TEXT("Total: %i"), ActiveSoundCount), GetStatsFont(), BodyColor);
+	Y += FontHeight;
+
+	return Y;
+}
+
+int32 FAudioDebugger::RenderStatMixes(UWorld* World, FViewport* Viewport, FCanvas* Canvas, int32 X, int32 Y, const FVector* ViewLocation, const FRotator* ViewRotation)
+{
+	FAudioDevice* AudioDevice = GetWorldAudio(World);
+	if (!AudioDevice)
+	{
+		return Y;
+	}
+
+	const int32 FontHeight = GetStatsFont()->GetMaxCharHeight() + 2;
+	Canvas->DrawShadowedString(X, Y, TEXT("Active Sound Mixes:"), GetStatsFont(), HeaderColor);
+	Y += FontHeight;
+
+	bool bDisplayedSoundMixes = false;
+
+	FAudioStats& AudioStats = AudioDeviceStats.FindOrAdd(AudioDevice->DeviceHandle);
+	if (AudioStats.StatSoundMixes.Num() > 0)
+	{
+		bDisplayedSoundMixes = true;
+
+		for (const FAudioStats::FStatSoundMix& StatSoundMix : AudioStats.StatSoundMixes)
+		{
+			const FString TheString = FString::Printf(TEXT("%s - Fade Proportion: %1.2f - Total Ref Count: %i"), *StatSoundMix.MixName, StatSoundMix.InterpValue, StatSoundMix.RefCount);
+
+			const FColor& TextColour = (StatSoundMix.bIsCurrentEQ ? FColor::Yellow : FColor::White);
+
+			Canvas->DrawShadowedString(X + TabWidth, Y, *TheString, GetStatsFont(), TextColour);
+			Y += FontHeight;
+		}
+	}
+
+	if (!bDisplayedSoundMixes)
+	{
+		Canvas->DrawShadowedString(X + TabWidth, Y, TEXT("None"), GetStatsFont(), FColor::White);
+		Y += FontHeight;
+	}
+
+	return Y;
+}
+
+int32 FAudioDebugger::RenderStatModulators(UWorld* World, FViewport* Viewport, FCanvas* Canvas, int32 X, int32 Y, const FVector* ViewLocation, const FRotator* ViewRotation)
+{
+	FAudioDevice* AudioDevice = GetWorldAudio(World);
+	if (!AudioDevice)
+	{
+		return Y;
+	}
+
+	const int32 FontHeight = GetStatsFont()->GetMaxCharHeight() + 2;
+	Canvas->DrawShadowedString(X, Y, TEXT("Active Sound Modulation:"), GetStatsFont(), HeaderColor);
+	Y += FontHeight;
+
+	bool bDisplayedSoundModulationInfo = false;
+
+	if (IAudioModulation* Modulation = AudioDevice->ModulationInterface.Get())
+	{
+		const int32 YInit = Y;
+		Y = Modulation->OnRenderStat(Viewport, Canvas, X, Y, *GetStatsFont(), ViewLocation, ViewRotation);
+		bDisplayedSoundModulationInfo = Y != YInit;
+	}
+
+	if (!bDisplayedSoundModulationInfo)
+	{
+		Canvas->DrawShadowedString(X + TabWidth, Y, TEXT("None"), GetStatsFont(), FColor::White);
+		Y += FontHeight;
+	}
+
+	return Y;
+}
+
+int32 FAudioDebugger::RenderStatReverb(UWorld* World, FViewport* Viewport, FCanvas* Canvas, int32 X, int32 Y, const FVector* ViewLocation, const FRotator* ViewRotation)
+{
+	FAudioDevice* AudioDevice = GetWorldAudio(World);
+	if (!AudioDevice)
+	{
+		return Y;
+	}
+
+	const int32 Height = static_cast<int32>(GetStatsFont()->GetMaxCharHeight() + 2);
+
+	FString TheString;
+	if (UReverbEffect* ReverbEffect = AudioDevice->GetCurrentReverbEffect())
+	{
+		TheString = FString::Printf(TEXT("Active Reverb Effect: %s"), *ReverbEffect->GetName());
+		Canvas->DrawShadowedString(X, Y, *TheString, GetStatsFont(), FLinearColor::White);
+		Y += Height;
+
+		AAudioVolume* CurrentAudioVolume = nullptr;
+		for (const FTransform& Transform : AudioDevice->ListenerTransforms)
+		{
+			AAudioVolume* PlayerAudioVolume = World->GetAudioSettings(Transform.GetLocation(), nullptr, nullptr);
+			if (PlayerAudioVolume && ((CurrentAudioVolume == nullptr) || (PlayerAudioVolume->GetPriority() > CurrentAudioVolume->GetPriority())))
+			{
+				CurrentAudioVolume = PlayerAudioVolume;
+			}
+		}
+		if (CurrentAudioVolume && CurrentAudioVolume->GetReverbSettings().ReverbEffect)
+		{
+			TheString = FString::Printf(TEXT("  Audio Volume Reverb Effect: %s (Priority: %g Volume Name: %s)"), *CurrentAudioVolume->GetReverbSettings().ReverbEffect->GetName(), CurrentAudioVolume->GetPriority(), *CurrentAudioVolume->GetName());
+		}
+		else
+		{
+			TheString = TEXT("  Audio Volume Reverb Effect: None");
+		}
+		Canvas->DrawShadowedString(X, Y, *TheString, GetStatsFont(), FLinearColor::White);
+		Y += Height;
+
+		const TMap<FName, FActivatedReverb>& ActivatedReverbs = AudioDevice->GetActiveReverb();
+		if (ActivatedReverbs.Num() == 0)
+		{
+			TheString = TEXT("  Activated Reverb: None");
+			Canvas->DrawShadowedString(X, Y, *TheString, GetStatsFont(), FLinearColor::White);
+			Y += Height;
+		}
+		else if (ActivatedReverbs.Num() == 1)
+		{
+			auto It = ActivatedReverbs.CreateConstIterator();
+			TheString = FString::Printf(TEXT("  Activated Reverb Effect: %s (Priority: %g Tag: '%s')"), *It.Value().ReverbSettings.ReverbEffect->GetName(), It.Value().Priority, *It.Key().ToString());
+			Canvas->DrawShadowedString(X, Y, *TheString, GetStatsFont(), FLinearColor::White);
+			Y += Height;
+		}
+		else
+		{
+			Canvas->DrawShadowedString(X, Y, TEXT("  Activated Reverb Effects:"), GetStatsFont(), FLinearColor::White);
+			Y += Height;
+			TMap<int32, FString> PrioritySortedActivatedReverbs;
+			for (auto It = ActivatedReverbs.CreateConstIterator(); It; ++It)
+			{
+				TheString = FString::Printf(TEXT("    %s (Priority: %g Tag: '%s')"), *It.Value().ReverbSettings.ReverbEffect->GetName(), It.Value().Priority, *It.Key().ToString());
+				PrioritySortedActivatedReverbs.Add(It.Value().Priority, TheString);
+			}
+			for (auto It = PrioritySortedActivatedReverbs.CreateConstIterator(); It; ++It)
+			{
+				Canvas->DrawShadowedString(X, Y, *It.Value(), GetStatsFont(), FLinearColor::White);
+				Y += Height;
+			}
+		}
+	}
+	else
+	{
+		TheString = TEXT("Active Reverb Effect: None");
+		Canvas->DrawShadowedString(X, Y, *TheString, GetStatsFont(), FLinearColor::White);
+		Y += Height;
+	}
+
+	return Y;
+}
+
+int32 FAudioDebugger::RenderStatSounds(UWorld* World, FViewport* Viewport, FCanvas* Canvas, int32 X, int32 Y, const FVector* ViewLocation, const FRotator* ViewRotation)
+{
+	FAudioDevice* AudioDevice = GetWorldAudio(World);
+	if (!AudioDevice)
+	{
+		return Y;
+	}
+
+	const int32 FontHeight = GetStatsFont()->GetMaxCharHeight() + 2;
+	Y += FontHeight;
+
+	FAudioStats& AudioStats = AudioDeviceStats.FindOrAdd(AudioDevice->DeviceHandle);
+
+	const uint8 bDebug = AudioStats.DisplayFlags & static_cast<uint8>(FAudioStats::EDisplayFlags::Debug);
+
+	FString SortingName = TEXT("disabled");
+
+	// Sort the list.
+	if (AudioStats.DisplayFlags & static_cast<uint8>(FAudioStats::EDisplayFlags::Sort_Name))
+	{
+		AudioStats.StatSoundInfos.Sort([](const FAudioStats::FStatSoundInfo& A, const FAudioStats::FStatSoundInfo& B) { return A.SoundName < B.SoundName; });
+		SortingName = TEXT("pathname");
+	}
+	else if (AudioStats.DisplayFlags & static_cast<uint8>(FAudioStats::EDisplayFlags::Sort_Distance))
+	{
+		AudioStats.StatSoundInfos.Sort([](const FAudioStats::FStatSoundInfo& A, const FAudioStats::FStatSoundInfo& B) { return A.Distance < B.Distance; });
+		SortingName = TEXT("distance");
+	}
+	else if (AudioStats.DisplayFlags & static_cast<uint8>(FAudioStats::EDisplayFlags::Sort_Class))
+	{
+		AudioStats.StatSoundInfos.Sort([](const FAudioStats::FStatSoundInfo& A, const FAudioStats::FStatSoundInfo& B) { return A.SoundClassName < B.SoundClassName; });
+		SortingName = TEXT("class");
+	}
+	else if (AudioStats.DisplayFlags & static_cast<uint8>(FAudioStats::EDisplayFlags::Sort_WavesNum))
+	{
+		AudioStats.StatSoundInfos.Sort([](const FAudioStats::FStatSoundInfo& A, const FAudioStats::FStatSoundInfo& B) { return A.WaveInstanceInfos.Num() > B.WaveInstanceInfos.Num(); });
+		SortingName = TEXT("waves' num");
+	}
+
+	Canvas->DrawShadowedString(X, Y, TEXT("Active Sounds:"), GetStatsFont(), HeaderColor);
+	Y += FontHeight;
+
+	const FString InfoText = FString::Printf(TEXT(" Sorting: %s, Debug: %s"), *SortingName, bDebug ? TEXT("enabled") : TEXT("disabled"));
+	Canvas->DrawShadowedString(X, Y, *InfoText, GetStatsFont(), FColor(128, 255, 128));
+	Y += FontHeight;
+
+	Canvas->DrawShadowedString(X, Y, TEXT("Index Path (Class) Distance"), GetStatsFont(), BodyColor);
+	Y += FontHeight;
+
+	int32 TotalSoundWavesNum = 0;
+	for (int32 SoundIndex = 0; SoundIndex < AudioStats.StatSoundInfos.Num(); ++SoundIndex)
+	{
+		const FAudioStats::FStatSoundInfo& StatSoundInfo = AudioStats.StatSoundInfos[SoundIndex];
+		const int32 WaveInstancesNum = StatSoundInfo.WaveInstanceInfos.Num();
+		if (WaveInstancesNum > 0)
+		{
+			{
+				const FString TheString = FString::Printf(TEXT("%4i. %s (%s) %6.2f"), SoundIndex, *StatSoundInfo.SoundName, *StatSoundInfo.SoundClassName.ToString(), StatSoundInfo.Distance);
+				Canvas->DrawShadowedString(X, Y, *TheString, GetStatsFont(), FColor::White);
+				Y += FontHeight;
+			}
+
+			TotalSoundWavesNum += WaveInstancesNum;
+
+			// Get the active sound waves.
+			for (int32 WaveIndex = 0; WaveIndex < WaveInstancesNum; WaveIndex++)
+			{
+				const FString TheString = *FString::Printf(TEXT("    %4i. %s"), WaveIndex, *StatSoundInfo.WaveInstanceInfos[WaveIndex].Description);
+				Canvas->DrawShadowedString(X, Y, *TheString, GetStatsFont(), FColor(205, 205, 205));
+				Y += FontHeight;
+			}
+		}
+	}
+
+	Canvas->DrawShadowedString(X, Y, *FString::Printf(TEXT("Audio Device ID: %u"), AudioDevice->DeviceHandle), GetStatsFont(), HeaderColor);
+	Y += FontHeight;
+
+	Canvas->DrawShadowedString(X, Y, *FString::Printf(TEXT("Total Sounds: %i, Sound Waves: %i"), AudioStats.StatSoundInfos.Num(), TotalSoundWavesNum), GetStatsFont(), HeaderColor);
+	Y += FontHeight;
+
+	for (int32 i = 0; i < AudioStats.ListenerTransforms.Num(); ++i)
+	{
+		FString LocStr = AudioStats.ListenerTransforms[i].GetLocation().ToString();
+		Canvas->DrawShadowedString(X, Y, *FString::Printf(TEXT("Listener '%d' Position: %s"), i, *LocStr), GetStatsFont(), HeaderColor);
+		Y += FontHeight;
+	}
+
+	if (!bDebug)
+	{
+		return Y;
+	}
+
+	// Draw sound cue's sphere only in debug.
+	for (const FAudioStats::FStatSoundInfo& StatSoundInfo : AudioStats.StatSoundInfos)
+	{
+		const FTransform& SoundTransform = StatSoundInfo.Transform;
+		const int32 WaveInstancesNum = StatSoundInfo.WaveInstanceInfos.Num();
+
+		if (StatSoundInfo.Distance > 100.0f && WaveInstancesNum > 0)
+		{
+			float SphereRadius = 0.f;
+			float SphereInnerRadius = 0.f;
+
+			if (StatSoundInfo.ShapeDetailsMap.Num() > 0)
+			{
+				DrawDebugString(World, SoundTransform.GetTranslation(), StatSoundInfo.SoundName, NULL, FColor::White, 0.01f);
+
+				for (auto ShapeDetailsIt = StatSoundInfo.ShapeDetailsMap.CreateConstIterator(); ShapeDetailsIt; ++ShapeDetailsIt)
+				{
+					const FBaseAttenuationSettings::AttenuationShapeDetails& ShapeDetails = ShapeDetailsIt.Value();
+					switch (ShapeDetailsIt.Key())
+					{
+					case EAttenuationShape::Sphere:
+						if (ShapeDetails.Falloff > 0.f)
+						{
+							DrawDebugSphere(World, SoundTransform.GetTranslation(), ShapeDetails.Extents.X + ShapeDetails.Falloff, 10, FColor(155, 155, 255));
+							DrawDebugSphere(World, SoundTransform.GetTranslation(), ShapeDetails.Extents.X, 10, FColor(55, 55, 255));
+						}
+						else
+						{
+							DrawDebugSphere(World, SoundTransform.GetTranslation(), ShapeDetails.Extents.X, 10, FColor(155, 155, 255));
+						}
+						break;
+
+					case EAttenuationShape::Box:
+						if (ShapeDetails.Falloff > 0.f)
+						{
+							DrawDebugBox(World, SoundTransform.GetTranslation(), ShapeDetails.Extents + FVector(ShapeDetails.Falloff), SoundTransform.GetRotation(), FColor(155, 155, 255));
+							DrawDebugBox(World, SoundTransform.GetTranslation(), ShapeDetails.Extents, SoundTransform.GetRotation(), FColor(55, 55, 255));
+						}
+						else
+						{
+							DrawDebugBox(World, SoundTransform.GetTranslation(), ShapeDetails.Extents, SoundTransform.GetRotation(), FColor(155, 155, 255));
+						}
+						break;
+
+					case EAttenuationShape::Capsule:
+
+						if (ShapeDetails.Falloff > 0.f)
+						{
+							DrawDebugCapsule(World, SoundTransform.GetTranslation(), ShapeDetails.Extents.X + ShapeDetails.Falloff, ShapeDetails.Extents.Y + ShapeDetails.Falloff, SoundTransform.GetRotation(), FColor(155, 155, 255));
+							DrawDebugCapsule(World, SoundTransform.GetTranslation(), ShapeDetails.Extents.X, ShapeDetails.Extents.Y, SoundTransform.GetRotation(), FColor(55, 55, 255));
+						}
+						else
+						{
+							DrawDebugCapsule(World, SoundTransform.GetTranslation(), ShapeDetails.Extents.X, ShapeDetails.Extents.Y, SoundTransform.GetRotation(), FColor(155, 155, 255));
+						}
+						break;
+
+					case EAttenuationShape::Cone:
+					{
+						const FVector Origin = SoundTransform.GetTranslation() - (SoundTransform.GetUnitAxis(EAxis::X) * ShapeDetails.ConeOffset);
+
+						if (ShapeDetails.Falloff > 0.f || ShapeDetails.Extents.Z > 0.f)
+						{
+							const float OuterAngle = FMath::DegreesToRadians(ShapeDetails.Extents.Y + ShapeDetails.Extents.Z);
+							const float InnerAngle = FMath::DegreesToRadians(ShapeDetails.Extents.Y);
+							DrawDebugCone(World, Origin, SoundTransform.GetUnitAxis(EAxis::X), ShapeDetails.Extents.X + ShapeDetails.Falloff + ShapeDetails.ConeOffset, OuterAngle, OuterAngle, 10, FColor(155, 155, 255));
+							DrawDebugCone(World, Origin, SoundTransform.GetUnitAxis(EAxis::X), ShapeDetails.Extents.X + ShapeDetails.ConeOffset, InnerAngle, InnerAngle, 10, FColor(55, 55, 255));
+						}
+						else
+						{
+							const float Angle = FMath::DegreesToRadians(ShapeDetails.Extents.Y);
+							DrawDebugCone(World, Origin, SoundTransform.GetUnitAxis(EAxis::X), ShapeDetails.Extents.X + ShapeDetails.ConeOffset, Angle, Angle, 10, FColor(155, 155, 255));
+						}
+						break;
+					}
+
+					default:
+						check(false);
+					}
+				}
+			}
+		}
+	}
+
+	return Y;
+}
+
+int32 FAudioDebugger::RenderStatWaves(UWorld* World, FViewport* Viewport, FCanvas* Canvas, int32 X, int32 Y, const FVector* ViewLocation, const FRotator* ViewRotation)
+{
+	FAudioDevice* AudioDevice = GetWorldAudio(World);
+	if (!AudioDevice)
+	{
+		return Y;
+	}
+
+	const int32 FontHeight = GetStatsFont()->GetMaxCharHeight() + 2;
+
+	FAudioStats& AudioStats = AudioDeviceStats.FindOrAdd(AudioDevice->DeviceHandle);
+	Canvas->DrawShadowedString(X, Y, TEXT("Active Sound Waves:"), GetStatsFont(), FLinearColor::White);
+	Y += TabWidth;
+
+	using FWaveInstancePair = TPair<const FAudioStats::FStatWaveInstanceInfo*, const FAudioStats::FStatSoundInfo*>;
+	TArray<FWaveInstancePair> WaveInstances;
+	for (const FAudioStats::FStatSoundInfo& StatSoundInfo : AudioStats.StatSoundInfos)
+	{
+		for (const FAudioStats::FStatWaveInstanceInfo& WaveInstanceInfo : StatSoundInfo.WaveInstanceInfos)
+		{
+			if (WaveInstanceInfo.Volume >= MinDisplayVolume || WaveInstanceInfo.bPlayWhenSilent != 0)
+			{
+				WaveInstances.Emplace(&WaveInstanceInfo, &StatSoundInfo);
+			}
+		}
+	}
+
+	WaveInstances.Sort([](const FWaveInstancePair& A, const FWaveInstancePair& B) { return A.Key->InstanceIndex < B.Key->InstanceIndex; });
+
+	for (const FWaveInstancePair& WaveInstanceInfo : WaveInstances)
+	{
+		UAudioComponent* AudioComponent = UAudioComponent::GetAudioComponentFromID(WaveInstanceInfo.Value->AudioComponentID);
+		AActor* SoundOwner = AudioComponent ? AudioComponent->GetOwner() : nullptr;
+
+		FString TheString = *FString::Printf(TEXT("%4i.    %6.2f  %s   Owner: %s   SoundClass: %s"),
+			WaveInstanceInfo.Key->InstanceIndex,
+			WaveInstanceInfo.Key->Volume,
+			*WaveInstanceInfo.Key->WaveInstanceName.ToString(),
+			SoundOwner ? *SoundOwner->GetName() : TEXT("None"),
+			*WaveInstanceInfo.Value->SoundClassName.ToString());
+		Canvas->DrawShadowedString(X, Y, *TheString, GetStatsFont(), WaveInstanceInfo.Key->bPlayWhenSilent == 0 ? FColor::White : FColor::Yellow);
+		Y += FontHeight;
+	}
+
+	const int32 ActiveInstances = WaveInstances.Num();
+
+	const int32 Max = AudioDevice->MaxChannels / 2;
+	float f = FMath::Clamp<float>((float)(ActiveInstances - Max) / (float)Max, 0.f, 1.f);
+	const int32 R = FMath::TruncToInt(f * 255);
+
+	if (ActiveInstances > Max)
+	{
+		f = FMath::Clamp<float>((float)(Max - ActiveInstances) / (float)Max, 0.5f, 1.f);
+	}
+	else
+	{
+		f = 1.0f;
+	}
+	const int32 G = FMath::TruncToInt(f * 255);
+	const int32 B = 0;
+
+	Canvas->DrawShadowedString(X, Y, *FString::Printf(TEXT(" Total: %i"), ActiveInstances), GetStatsFont(), FColor(R, G, B));
+	Y += FontHeight;
+
+	return Y;
+}
+
+void FAudioDebugger::RemoveDevice(const FAudioDevice& AudioDevice)
+{
+	AudioDeviceStats.Remove(AudioDevice.DeviceHandle);
+	AudioDeviceStats_AudioThread.Remove(AudioDevice.DeviceHandle);
+}
+
+bool FAudioDebugger::ToggleStats(UWorld* World, const uint8 StatToToggle)
+{
+	if (!GEngine)
+	{
+		return false;
+	}
+
+	FAudioDevice* AudioDevice = GetWorldAudio(World);
+	if (!AudioDevice)
+	{
+		return false;
+	}
+
+	if (FAudioDeviceManager* DeviceManager = GEngine->GetAudioDeviceManager())
+	{
+		DeviceManager->GetDebugger().ToggleStats(AudioDevice->DeviceHandle, StatToToggle);
+	}
+
+	return true;
+}
+
+void FAudioDebugger::ToggleStats(const uint32 AudioDeviceHandle, const uint8 StatsToToggle)
+{
+	if (!IsInAudioThread())
+	{
+		DECLARE_CYCLE_STAT(TEXT("FAudioThreadTask.ToggleStats"), STAT_AudioToggleStats, STATGROUP_TaskGraphTasks);
+
+		FAudioThread::RunCommandOnAudioThread([this, AudioDeviceHandle, StatsToToggle]()
+		{
+			ToggleStats(AudioDeviceHandle, StatsToToggle);
+		}, GET_STATID(STAT_AudioToggleStats));
+		return;
+	}
+
+	FAudioStats_AudioThread& Stats = AudioDeviceStats_AudioThread.FindOrAdd(AudioDeviceHandle);
+	Stats.RequestedStats ^= StatsToToggle;
+}
+
+bool FAudioDebugger::ToggleStatCues(UWorld* World, FCommonViewportClient* ViewportClient, const TCHAR* Stream)
+{
+	return ToggleStats(World, ERequestedAudioStats::SoundCues);
+}
+
+bool FAudioDebugger::ToggleStatMixes(UWorld* World, FCommonViewportClient* ViewportClient, const TCHAR* Stream)
+{
+	return ToggleStats(World, ERequestedAudioStats::SoundMixes);
+}
+
+bool FAudioDebugger::ToggleStatModulators(UWorld* World, FCommonViewportClient* ViewportClient, const TCHAR* Stream)
+{
+	if (!GEngine)
+	{
+		return false;
+	}
+
+	FAudioDevice* AudioDevice = GetWorldAudio(World);
+	if (!AudioDevice)
+	{
+		return false;
+	}
+
+	if (AudioDevice->IsModulationPluginEnabled())
+	{
+		if (IAudioModulation* Modulation = AudioDevice->ModulationInterface.Get())
+		{
+			if (!Modulation->OnToggleStat(ViewportClient, Stream))
+			{
+				return false;
+			}
+		}
+	}
+
+	return true;
+}
+
+bool FAudioDebugger::PostStatModulatorHelp(UWorld* World, FCommonViewportClient* ViewportClient, const TCHAR* Stream)
+{
+	// Ignore if all Viewports are closed.
+	if (!ViewportClient)
+	{
+		return false;
+	}
+
+	if (FAudioDevice* AudioDevice = World->GetAudioDevice())
+	{
+		if (AudioDevice->IsModulationPluginEnabled())
+		{
+			if (IAudioModulation* Modulation = AudioDevice->ModulationInterface.Get())
+			{
+				if (!Modulation->OnPostHelp(ViewportClient, Stream))
+				{
+					return false;
+				}
+			}
+		}
+	}
+
+	return true;
+}
+
+bool FAudioDebugger::ToggleStatSounds(UWorld* World, FCommonViewportClient* ViewportClient, const TCHAR* Stream)
+{
+	// Ignore if all Viewports are closed.
+	if (!ViewportClient)
+	{
+		return false;
+	}
+
+	if (!ToggleStats(World, ERequestedAudioStats::Sounds))
+	{
+		return false;
+	}
+
+	const bool bHelp = Stream ? FCString::Stristr(Stream, TEXT("?")) != nullptr : false;
+	if (bHelp)
+	{
+		GLog->Logf(TEXT("stat sounds description"));
+		GLog->Logf(TEXT("  stat sounds off - Disables drawing stat sounds"));
+		GLog->Logf(TEXT("  stat sounds sort=distance|class|name|waves|default"));
+		GLog->Logf(TEXT("      distance - sort list by distance to player"));
+		GLog->Logf(TEXT("      class - sort by sound class name"));
+		GLog->Logf(TEXT("      name - sort by cue pathname"));
+		GLog->Logf(TEXT("      waves - sort by waves' num"));
+		GLog->Logf(TEXT("      default - sorting is disabled"));
+		GLog->Logf(TEXT("  stat sounds -debug - enables debugging mode like showing sound radius sphere and names, but only for cues with enabled property bDebug"));
+		GLog->Logf(TEXT("  stat sounds -smalltext - use large text in debug output (default)"));
+		GLog->Logf(TEXT("  stat sounds -largetext - use large text in debug output"));
+		GLog->Logf(TEXT(""));
+		GLog->Logf(TEXT("Ex. stat sounds sort=class -debug"));
+		GLog->Logf(TEXT(" This will show only debug sounds sorted by sound class"));
+	}
+
+	uint8 ShowSounds = 0;
+	if (Stream)
+	{
+		const bool bHide = FParse::Command(&Stream, TEXT("off"));
+		if (!bHide)
+		{
+			const bool bDebug = FParse::Param(Stream, TEXT("debug"));
+			if (bDebug)
+			{
+				ShowSounds |= static_cast<uint8>(FAudioStats::EDisplayFlags::Debug);
+			}
+
+			const bool bLongNames = FParse::Param(Stream, TEXT("longnames"));
+			if (bLongNames)
+			{
+				ShowSounds |= static_cast<uint8>(FAudioStats::EDisplayFlags::Long_Names);
+			}
+
+			FString SortStr;
+			FParse::Value(Stream, TEXT("sort="), SortStr);
+			if (SortStr == TEXT("distance"))
+			{
+				ShowSounds |= static_cast<uint8>(FAudioStats::EDisplayFlags::Sort_Distance);
+			}
+			else if (SortStr == TEXT("class"))
+			{
+				ShowSounds |= static_cast<uint8>(FAudioStats::EDisplayFlags::Sort_Class);
+			}
+			else if (SortStr == TEXT("name"))
+			{
+				ShowSounds |= static_cast<uint8>(FAudioStats::EDisplayFlags::Sort_Name);
+			}
+			else if (SortStr == TEXT("waves"))
+			{
+				ShowSounds |= static_cast<uint8>(FAudioStats::EDisplayFlags::Sort_WavesNum);
+			}
+			else
+			{
+				ShowSounds |= static_cast<uint8>(FAudioStats::EDisplayFlags::Sort_Disabled);
+			}
+		}
+	}
+
+	FAudioDevice* AudioDevice = GetWorldAudio(World);
+	check(AudioDevice);
+	FAudioStats& Stats = AudioDeviceStats.FindOrAdd(AudioDevice->DeviceHandle);
+	Stats.DisplayFlags = ShowSounds;
+
+	ResolveDesiredStats(ViewportClient);
+
+	return true;
+}
+
+bool FAudioDebugger::ToggleStatWaves(UWorld* World, FCommonViewportClient* ViewportClient, const TCHAR* Stream)
+{
+	return ToggleStats(World, ERequestedAudioStats::SoundWaves);
+}
+
+void FAudioDebugger::SendUpdateResultsToGameThread(const FAudioDevice& AudioDevice, const int32 FirstActiveIndex)
+{
+	check(IsInAudioThread());
+
+	FAudioStats_AudioThread* Stats_AudioThread = AudioDeviceStats_AudioThread.Find(AudioDevice.DeviceHandle);
+	if (!Stats_AudioThread)
+	{
+		return;
+	}
+
+	TArray<FAudioStats::FStatSoundInfo> StatSoundInfos;
+	TArray<FAudioStats::FStatSoundMix> StatSoundMixes;
+
+	const uint8 RequestedStats = Stats_AudioThread->RequestedStats;
+	TMap<FActiveSound*, int32> ActiveSoundToInfoIndex;
+
+	const bool bDebug = (RequestedStats & ERequestedAudioStats::DebugSounds) != 0;
+
+	static const uint8 SoundMask = ERequestedAudioStats::Sounds | ERequestedAudioStats::SoundCues | ERequestedAudioStats::SoundWaves;
+	if (RequestedStats & SoundMask)
+	{
+		for (FActiveSound* ActiveSound : AudioDevice.GetActiveSounds())
+		{
+			if (USoundBase* SoundBase = ActiveSound->GetSound())
+			{
+				if (!bDebug || ActiveSound->GetSound()->bDebug)
+				{
+					ActiveSoundToInfoIndex.Add(ActiveSound, StatSoundInfos.AddDefaulted());
+					FAudioStats::FStatSoundInfo& StatSoundInfo = StatSoundInfos.Last();
+					StatSoundInfo.SoundName = ActiveSound->GetSound()->GetPathName();
+					StatSoundInfo.Distance = AudioDevice.GetDistanceToNearestListener(ActiveSound->Transform.GetTranslation());
+
+					if (USoundClass* SoundClass = ActiveSound->GetSoundClass())
+					{
+						StatSoundInfo.SoundClassName = SoundClass->GetFName();
+					}
+					else
+					{
+						StatSoundInfo.SoundClassName = NAME_None;
+					}
+					StatSoundInfo.Transform = ActiveSound->Transform;
+					StatSoundInfo.AudioComponentID = ActiveSound->GetAudioComponentID();
+
+					if (bDebug && ActiveSound->GetSound()->bDebug)
+					{
+						ActiveSound->CollectAttenuationShapesForVisualization(StatSoundInfo.ShapeDetailsMap);
+					}
+				}
+			}
+		}
+
+
+		// Iterate through all wave instances.
+		const TArray<FWaveInstance*>& WaveInstances = AudioDevice.GetActiveWaveInstances();
+		auto WaveInstanceSourceMap = AudioDevice.GetWaveInstanceSourceMap();
+		for (int32 InstanceIndex = FirstActiveIndex; InstanceIndex < WaveInstances.Num(); ++InstanceIndex)
+		{
+			const FWaveInstance* WaveInstance = WaveInstances[InstanceIndex];
+			int32* SoundInfoIndex = ActiveSoundToInfoIndex.Find(WaveInstance->ActiveSound);
+			if (SoundInfoIndex)
+			{
+				FAudioStats::FStatWaveInstanceInfo WaveInstanceInfo;
+				FSoundSource* Source = WaveInstanceSourceMap.FindRef(WaveInstance);
+				WaveInstanceInfo.Description = Source ? Source->Describe((RequestedStats & ERequestedAudioStats::LongSoundNames) != 0) : FString(TEXT("No source"));
+				WaveInstanceInfo.Volume = WaveInstance->GetVolumeWithDistanceAttenuation() * WaveInstance->GetDynamicVolume();
+				WaveInstanceInfo.InstanceIndex = InstanceIndex;
+				WaveInstanceInfo.WaveInstanceName = *WaveInstance->GetName();
+				WaveInstanceInfo.bPlayWhenSilent = WaveInstance->ActiveSound->IsPlayWhenSilent() ? 1 : 0;
+				StatSoundInfos[*SoundInfoIndex].WaveInstanceInfos.Add(MoveTemp(WaveInstanceInfo));
+			}
+		}
+	}
+
+	if (RequestedStats & ERequestedAudioStats::SoundMixes)
+	{
+		if (const FAudioEffectsManager* Effects = AudioDevice.GetEffects())
+		{
+			const USoundMix* CurrentEQMix = Effects->GetCurrentEQMix();
+
+			for (const TPair<USoundMix*, FSoundMixState>& SoundMixPair : AudioDevice.GetSoundMixModifiers())
+			{
+				StatSoundMixes.AddDefaulted();
+				FAudioStats::FStatSoundMix& StatSoundMix = StatSoundMixes.Last();
+				StatSoundMix.MixName = SoundMixPair.Key->GetName();
+				StatSoundMix.InterpValue = SoundMixPair.Value.InterpValue;
+				StatSoundMix.RefCount = SoundMixPair.Value.ActiveRefCount + SoundMixPair.Value.PassiveRefCount;
+				StatSoundMix.bIsCurrentEQ = (SoundMixPair.Key == CurrentEQMix);
+			}
+		}
+	}
+
+	DECLARE_CYCLE_STAT(TEXT("FGameThreadAudioTask.AudioSendResults"), STAT_AudioSendResults, STATGROUP_TaskGraphTasks);
+
+	const uint32 AudioDeviceID = AudioDevice.DeviceHandle;
+
+	TArray<FTransform> ListenerTransforms;
+	for (const FListener& Listener : AudioDevice.GetListeners())
+	{
+		ListenerTransforms.Add(Listener.Transform);
+	}
+	FAudioThread::RunCommandOnGameThread([AudioDeviceID, ListenerTransforms, StatSoundInfos, StatSoundMixes]()
+	{
+		FAudioStats& Stats = AudioDeviceStats.FindOrAdd(AudioDeviceID);
+		Stats.ListenerTransforms = ListenerTransforms;
+		Stats.StatSoundInfos = StatSoundInfos;
+		Stats.StatSoundMixes = StatSoundMixes;
+	}, GET_STATID(STAT_AudioSendResults));
+}
+
+void FAudioDebugger::UpdateAudibleInactiveSounds(const uint32 FirstActiveIndex, const TArray<FWaveInstance*>& WaveInstances)
+{
+#if STATS
+	uint32 AudibleInactiveSounds = 0;
+	// Count how many sounds are not being played but were audible
+	for (uint32 InstanceIndex = 0; InstanceIndex < FirstActiveIndex; ++InstanceIndex)
+	{
+		const FWaveInstance* WaveInstance = WaveInstances[InstanceIndex];
+		const float WaveInstanceVol = WaveInstance->GetVolumeWithDistanceAttenuation() * WaveInstance->GetDynamicVolume();
+		if (WaveInstanceVol > MinDisplayVolume)
+		{
+			AudibleInactiveSounds++;
+		}
+	}
+	SET_DWORD_STAT(STAT_AudibleWavesDroppedDueToPriority, AudibleInactiveSounds);
+#endif
+}
+#endif // ENABLE_AUDIO_DEBUG
