Copying Dev-Audio to UE-Main from CL 6717121

diff --git a/Engine/Source/Runtime/Engine/Private/AudioDevice.cpp b/Engine/Source/Runtime/Engine/Private/AudioDevice.cpp
--- a/Engine/Source/Runtime/Engine/Private/AudioDevice.cpp
+++ b/Engine/Source/Runtime/Engine/Private/AudioDevice.cpp
@@ -1,46 +1,49 @@
 // Copyright 1998-2019 Epic Games, Inc. All Rights Reserved.
-
 #include "AudioDevice.h"
 
-#include "PhysicsEngine/BodyInstance.h"
-#include "Sound/SoundEffectPreset.h"
-#include "Sound/SoundEffectSubmix.h"
-#include "HAL/FileManager.h"
-#include "Misc/Paths.h"
-#include "Misc/OutputDeviceArchiveWrapper.h"
-#include "ProfilingDebugging/ProfilingHelpers.h"
-#include "Misc/ConfigCacheIni.h"
-#include "Misc/App.h"
-#include "UObject/UObjectHash.h"
-#include "UObject/UObjectIterator.h"
-#include "UObject/Package.h"
 #include "ActiveSound.h"
-#include "ContentStreaming.h"
-#include "UnrealEngine.h"
-#include "Sound/SoundGroups.h"
-#include "Sound/SoundEffectSource.h"
-#include "Sound/SoundWave.h"
+#include "AudioCompressionSettingsUtils.h"
 #include "AudioDecompress.h"
 #include "AudioEffect.h"
 #include "AudioPluginUtilities.h"
-#include "Sound/AudioSettings.h"
-#include "Sound/SoundCue.h"
-#include "Sound/SoundNode.h"
-#include "Sound/SoundNodeWavePlayer.h"
+#include "Audio/AudioDebug.h"
+#include "ContentStreaming.h"
+#include "DrawDebugHelpers.h"
 #include "GameFramework/GameUserSettings.h"
 #include "GameFramework/WorldSettings.h"
 #include "GeneralProjectSettings.h"
+#include "HAL/FileManager.h"
 #include "HAL/LowLevelMemTracker.h"
-#include "AudioCompressionSettingsUtils.h"
-#include "DrawDebugHelpers.h"
+#include "IAudioExtensionPlugin.h"
+#include "Misc/ConfigCacheIni.h"
+#include "Misc/App.h"
+#include "Misc/OutputDeviceArchiveWrapper.h"
+#include "Misc/Paths.h"
+#include "PhysicsEngine/BodyInstance.h"
+#include "ProfilingDebugging/ProfilingHelpers.h"
+#include "Sound/AudioSettings.h"
+#include "Sound/ReverbEffect.h"
+#include "Sound/SoundEffectPreset.h"
+#include "Sound/SoundEffectSource.h"
+#include "Sound/SoundEffectSubmix.h"
+#include "Sound/SoundGroups.h"
+#include "Sound/SoundWave.h"
+#include "Sound/SoundCue.h"
+#include "Sound/SoundNode.h"
+#include "Sound/SoundNodeWavePlayer.h"
+#include "UnrealEngine.h"
+#include "UObject/UObjectHash.h"
+#include "UObject/UObjectIterator.h"
+#include "UObject/Package.h"
 
 #if WITH_EDITOR
+#include "AudioEditorModule.h"
 #include "AssetRegistryModule.h"
-#include "Developer/AssetTools/Public/IAssetTools.h"
 #include "Developer/AssetTools/Public/AssetToolsModule.h"
+#include "Developer/AssetTools/Public/IAssetTools.h"
 #include "Editor/EditorEngine.h"
-#include "AudioEditorModule.h"
-#endif
+#endif // WITH_EDITOR
+
 
 static int32 AudioChannelCountCVar = 0;
 FAutoConsoleVariableRef CVarSetAudioChannelCount(
@@ -195,6 +198,7 @@ FAudioDevice::FAudioDevice()
 	, bSpatializationInterfaceEnabled(false)
 	, bOcclusionInterfaceEnabled(false)
 	, bReverbInterfaceEnabled(false)
+	, bModulationInterfaceEnabled(false)
 	, bPluginListenersInitialized(false)
 	, bHRTFEnabledForAll(false)
 	, bIsDeviceMuted(false)
@@ -203,11 +207,8 @@ FAudioDevice::FAudioDevice()
 	, bAllowCenterChannel3DPanning(false)
 	, DeviceDeltaTime(0.0f)
 	, bHasActivatedReverb(false)
-	, bAllowVirtualizedSounds(true)
+	, bAllowPlayWhenSilent(true)
 	, bUseAttenuationForNonGameWorlds(false)
-#if !UE_BUILD_SHIPPING
-	, RequestedAudioStats(0)
-#endif
 	, ConcurrencyManager(this)
 	, OneShotCount(0)
 	, GlobalMinPitch(0.4f)
@@ -285,7 +286,7 @@ bool FAudioDevice::Init(int32 InMaxChannels)
 	GlobalMinPitch = FMath::Clamp(AudioSettings->GlobalMinPitchScale, 0.0001f, 4.0f);
 	GlobalMaxPitch = FMath::Clamp(AudioSettings->GlobalMaxPitchScale, 0.0001f, 4.0f);
 	bAllowCenterChannel3DPanning = AudioSettings->bAllowCenterChannel3DPanning;
-	bAllowVirtualizedSounds = AudioSettings->bAllowVirtualizedSounds;
+	bAllowPlayWhenSilent = AudioSettings->bAllowPlayWhenSilent;
 	DefaultReverbSendLevel = AudioSettings->DefaultReverbSendLevel;
 
 	const FSoftObjectPath DefaultBaseSoundMixName = GetDefault<UAudioSettings>()->DefaultBaseSoundMix;
@@ -362,21 +363,38 @@ bool FAudioDevice::Init(int32 InMaxChannels)
 	//Get the requested occlusion plugin and set it up.
 	IAudioOcclusionFactory* OcclusionPluginFactory = AudioPluginUtilities::GetDesiredOcclusionPlugin(AudioPluginUtilities::CurrentPlatform);
 	if (OcclusionPluginFactory != nullptr)
-		{
+	{
 		OcclusionInterface = OcclusionPluginFactory->CreateNewOcclusionPlugin(this);
 		bOcclusionInterfaceEnabled = true;
 		bOcclusionIsExternalSend = OcclusionPluginFactory->IsExternalSend();
 		UE_LOG(LogAudio, Display, TEXT("Audio Occlusion Plugin: %s"), *(OcclusionPluginFactory->GetDisplayName()));
-		}
+	}
 	else
 	{
 		UE_LOG(LogAudio, Display, TEXT("Using built-in audio occlusion."));
 	}
 
+	//Get the requested modulation plugin and set it up.
+	if (IAudioModulationFactory* ModulationPluginFactory = AudioPluginUtilities::GetDesiredModulationPlugin(AudioPluginUtilities::CurrentPlatform))
+	{
+		ModulationInterface = ModulationPluginFactory->CreateNewModulationPlugin(this);
+
+		//Set up initialization parameters for system level effect plugins:
+		FAudioPluginInitializationParams PluginInitializationParams;
+		PluginInitializationParams.SampleRate = SampleRate;
+		PluginInitializationParams.NumSources = MaxChannels;
+		PluginInitializationParams.BufferLength = PlatformSettings.CallbackBufferFrameSize;
+		PluginInitializationParams.AudioDevicePtr = this;
+		ModulationInterface->Initialize(PluginInitializationParams);
+
+		bModulationInterfaceEnabled = true;
+		UE_LOG(LogAudio, Display, TEXT("Audio Modulation Plugin: %s"), *(ModulationPluginFactory->GetDisplayName()));
+	}
+
 	// allow the platform to startup
-	if (InitializeHardware() == false)
+	if (!InitializeHardware())
 	{
-		// Could not initialize hardware, teardown anything that was set up during initialization
+		// Could not initialize hardware, tear down anything that was set up during initialization
 		Teardown();
 
 		return false;
@@ -576,7 +594,14 @@ void FAudioDevice::Teardown()
 		bOcclusionInterfaceEnabled = false;
 	}
 
+	ModulationInterface.Reset();
+	bModulationInterfaceEnabled = false;
+
 	PluginListeners.Reset();
+
+#if ENABLE_AUDIO_DEBUG
+	FAudioDebugger::RemoveDevice(*this);
+#endif // ENABLE_AUDIO_DEBUG
 }
 
 void FAudioDevice::Suspend(bool bGameTicking)
@@ -604,7 +629,7 @@ void FAudioDevice::UpdateAudioPluginSettingsObjectCache()
 	PluginSettingsObjects.Reset();
 
 	// Make sure we don't GC 3rd party plugin settings since these live on FSoundAttenuationSettings, which may not live in UObject graph due to overrides.
-	// There shouldn't be many of objects these (on the order of 10s not 100s) so if we find any loaded, don't let GC get them.
+	// There shouldn't be many of these objects (on the order of 10s not 100s) so if we find any loaded, don't let GC get them.
 	for (TObjectIterator<USpatializationPluginSourceSettingsBase> It; It; ++It)
 	{
 		PluginSettingsObjects.Add(*It);
@@ -619,6 +644,11 @@ void FAudioDevice::UpdateAudioPluginSettingsObjectCache()
 	{
 		PluginSettingsObjects.Add(*It);
 	}
+
+	for (TObjectIterator<USoundModulationPluginSourceSettingsBase> It; It; ++It)
+	{
+		PluginSettingsObjects.Add(*It);
+	}
 }
 
 void FAudioDevice::AddReferencedObjects(FReferenceCollector& Collector)
@@ -645,7 +675,7 @@ void FAudioDevice::AddReferencedObjects(FReferenceCollector& Collector)
 	for (FActiveSound* ActiveSound : PendingSoundsToDelete)
 	{
 		ActiveSound->AddReferencedObjects(Collector);
-	}	
+	}
 
 	for (FVirtualLoopPair& Pair : VirtualLoops)
 	{
@@ -1127,6 +1157,13 @@ bool FAudioDevice::HandleTestLPFCommand(const TCHAR* Cmd, FOutputDevice& Ar)
 	return true;
 }
 
+bool FAudioDevice::HandleTestHPFCommand(const TCHAR* Cmd, FOutputDevice& Ar)
+{
+	Ar.Logf(TEXT("HPF set to max for all sources"));
+	SetMixDebugState(DEBUGSTATE_TestHPF);
+	return true;
+}
+
 bool FAudioDevice::HandleTestStereoBleedCommand(const TCHAR* Cmd, FOutputDevice& Ar)
 {
 	Ar.Logf(TEXT("StereoBleed set to max for all sources"));
@@ -1148,6 +1185,13 @@ bool FAudioDevice::HandleDisableLPFCommand(const TCHAR* Cmd, FOutputDevice& Ar)
 	return true;
 }
 
+bool FAudioDevice::HandleDisableHPFCommand(const TCHAR* Cmd, FOutputDevice& Ar)
+{
+	Ar.Logf(TEXT("HPF disabled for all sources"));
+	SetMixDebugState(DEBUGSTATE_DisableHPF);
+	return true;
+}
+
 bool FAudioDevice::HandleDisableRadioCommand(const TCHAR* Cmd, FOutputDevice& Ar)
 {
 	EnableRadioEffect(false);
@@ -1651,7 +1695,6 @@ bool FAudioDevice::HandleSetDynamicSoundCommand(const TCHAR* Cmd, FOutputDevice&
 	}
 	return true;
 }
-
 #endif // !UE_BUILD_SHIPPING
 
 bool FAudioDevice::IsHRTFEnabledForAll() const
@@ -1755,6 +1798,10 @@ bool FAudioDevice::Exec(UWorld* InWorld, const TCHAR* Cmd, FOutputDevice& Ar)
 	{
 		return HandleDisableLPFCommand(Cmd, Ar);
 	}
+	else if (FParse::Command(&Cmd, TEXT("DisableHPF")))
+	{
+		return HandleDisableHPFCommand(Cmd, Ar);
+	}
 	else if (FParse::Command(&Cmd, TEXT("DisableRadio")))
 	{
 		return HandleDisableRadioCommand(Cmd, Ar);
@@ -2166,7 +2213,7 @@ void FAudioDevice::UpdatePassiveSoundMixModifiers(TArray<FWaveInstance*>& WaveIn
 			USoundClass* SoundClass = WaveInstance->SoundClass;
 			if (SoundClass)
 			{
-				const float WaveInstanceActualVolume = WaveInstance->GetVolumeWithDistanceAttenuation();
+				const float WaveInstanceActualVolume = WaveInstance->GetVolumeWithDistanceAttenuation() * WaveInstance->GetDynamicVolume();
 				// Check each SoundMix individually for volume levels
 				for (const FPassiveSoundMixModifier& PassiveSoundMixModifier : SoundClass->PassiveSoundMixModifiers)
 				{
@@ -3203,7 +3250,6 @@ void* FAudioDevice::InitEffect(FSoundSource* Source)
 	return nullptr;
 }
 
-
 void* FAudioDevice::UpdateEffect(FSoundSource* Source)
 {
 	SCOPE_CYCLE_COUNTER(STAT_AudioUpdateEffects);
@@ -3216,7 +3262,6 @@ void* FAudioDevice::UpdateEffect(FSoundSource* Source)
 	return nullptr;
 }
 
-
 void FAudioDevice::DestroyEffect(FSoundSource* Source)
 {
 	check(IsInAudioThread());
@@ -3503,7 +3548,7 @@ void FAudioDevice::StopSources(TArray<FWaveInstance*>& WaveInstances, int32 Firs
 			Source->LastUpdate = CurrentTick;
 
 			// If they are still audible, mark them as such
-			float VolumeWeightedPriority = WaveInstance.GetVolumeWithDistanceAttenuation();
+			float VolumeWeightedPriority = WaveInstance.GetVolumeWithDistanceAttenuation() * WaveInstance.GetDynamicVolume();
 			if (VolumeWeightedPriority > 0.0f)
 			{
 				Source->LastHeardUpdate = CurrentTick;
@@ -3542,6 +3587,10 @@ void FAudioDevice::StopSources(TArray<FWaveInstance*>& WaveInstances, int32 Firs
 				Source->NotifyPlaybackData();
 				Source->Update();
 			}
+
+#if ENABLE_AUDIO_DEBUG
+			FAudioDebugger::DrawDebugInfo(*Source);
+#endif // ENABLE_AUDIO_DEBUG
 		}
 	}
 
@@ -3554,19 +3603,9 @@ void FAudioDevice::StopSources(TArray<FWaveInstance*>& WaveInstances, int32 Firs
 		WaveInstance->StopWithoutNotification();
 	}
 
-#if STATS
-	uint32 AudibleInactiveSounds = 0;
-	// Count how many sounds are not being played but were audible
-	for (int32 InstanceIndex = 0; InstanceIndex < FirstActiveIndex; InstanceIndex++)
-	{
-		FWaveInstance* WaveInstance = WaveInstances[InstanceIndex];
-		if (WaveInstance->GetVolumeWithDistanceAttenuation() > KINDA_SMALL_NUMBER)
-		{
-			AudibleInactiveSounds++;
-		}
-	}
-	SET_DWORD_STAT(STAT_AudibleWavesDroppedDueToPriority, AudibleInactiveSounds);
-#endif
+#if ENABLE_AUDIO_DEBUG
+	FAudioDebugger::UpdateAudibleInactiveSounds(FirstActiveIndex, WaveInstances);
+#endif // ENABLE_AUDIO_DEBUG
 }
 
 void FAudioDevice::StartSources(TArray<FWaveInstance*>& WaveInstances, int32 FirstActiveIndex, bool bGameTicking)
@@ -3790,6 +3829,12 @@ void FAudioDevice::Update(bool bGameTicking)
 		HandlePause(bGameTicking);
 	}
 
+	if (bModulationInterfaceEnabled)
+	{
+		SCOPED_NAMED_EVENT(FAudioDevice_ProcessModulators, FColor::Blue);
+		ModulationInterface->ProcessModulators(GetDeviceDeltaTime());
+	}
+
 	bool bHasVolumeSettings = false;
 	float AudioVolumePriority = 0.f;
 	FReverbSettings ReverbSettings;
@@ -3884,8 +3929,6 @@ void FAudioDevice::Update(bool bGameTicking)
 		// active sound is updated while paused (for a long time), most sounds will be stopped when unpaused.
 		UpdateActiveSoundPlaybackTime(bGameTicking);
 
-
-
 		const int32 Channels = GetMaxChannels();
 		SET_DWORD_STAT(STAT_WaveInstances, ActiveWaveInstances.Num());
 		SET_DWORD_STAT(STAT_AudioSources, Sources.Num() - FreeSources.Num());
@@ -3919,100 +3962,11 @@ void FAudioDevice::Update(bool bGameTicking)
 
 void FAudioDevice::SendUpdateResultsToGameThread(const int32 FirstActiveIndex)
 {
-#if !UE_BUILD_SHIPPING
-	if (FirstActiveIndex == INDEX_NONE)
-	{
-		return;
-	}
-
-	TArray<FAudioStats::FStatSoundInfo> StatSoundInfos;
-	TArray<FAudioStats::FStatSoundMix> StatSoundMixes;
-	const FVector ListenerPosition = Listeners[0].Transform.GetTranslation();
-	const bool bStatsStale = (RequestedAudioStats == 0);
-	if (RequestedAudioStats != 0)
-	{
-		TMap<FActiveSound*, int32> ActiveSoundToInfoIndex;
-
-		const bool bDebug = (RequestedAudioStats & ERequestedAudioStats::DebugSounds) != 0;
-
-		for (FActiveSound* ActiveSound : ActiveSounds)
-		{
-			if (ActiveSound->Sound)
-			{
-				if (!bDebug || ActiveSound->GetSound()->bDebug)
-				{
-					ActiveSoundToInfoIndex.Add(ActiveSound, StatSoundInfos.AddDefaulted());
-					FAudioStats::FStatSoundInfo& StatSoundInfo = StatSoundInfos.Last();
-					StatSoundInfo.SoundName = ActiveSound->GetSound()->GetPathName();
-					StatSoundInfo.Distance = (ListenerPosition - ActiveSound->Transform.GetTranslation()).Size();
-
-					if (USoundClass* SoundClass = ActiveSound->GetSoundClass())
-					{
-						StatSoundInfo.SoundClassName = SoundClass->GetFName();
-					}
-					else
-					{
-						StatSoundInfo.SoundClassName = NAME_None;
-					}
-					StatSoundInfo.Transform = ActiveSound->Transform;
-					StatSoundInfo.AudioComponentID = ActiveSound->GetAudioComponentID();
-
-					if (bDebug && ActiveSound->GetSound()->bDebug)
-					{
-						ActiveSound->CollectAttenuationShapesForVisualization(StatSoundInfo.ShapeDetailsMap);
-					}
-				}
-			}
-		}
-
-		// Iterate through all wave instances.
-		for (int32 InstanceIndex = FirstActiveIndex; InstanceIndex < ActiveWaveInstances.Num(); ++InstanceIndex)
-		{
-			FWaveInstance* WaveInstance = ActiveWaveInstances[InstanceIndex];
-			int32* SoundInfoIndex = ActiveSoundToInfoIndex.Find(WaveInstance->ActiveSound);
-			if (SoundInfoIndex)
-			{
-				FAudioStats::FStatWaveInstanceInfo WaveInstanceInfo;
-				FSoundSource* Source = WaveInstanceSourceMap.FindRef(WaveInstance);
-				WaveInstanceInfo.Description = Source ? Source->Describe((RequestedAudioStats & ERequestedAudioStats::LongSoundNames) != 0) : FString(TEXT("No source"));
-				WaveInstanceInfo.ActualVolume = WaveInstance->GetVolumeWithDistanceAttenuation() * WaveInstance->GetDynamicVolume();
-				WaveInstanceInfo.InstanceIndex = InstanceIndex;
-				WaveInstanceInfo.WaveInstanceName = *WaveInstance->GetName();
-				StatSoundInfos[*SoundInfoIndex].WaveInstanceInfos.Add(MoveTemp(WaveInstanceInfo));
-			}
-		}
-
-		if (Effects)
-		{
-			USoundMix* CurrentEQMix = Effects->GetCurrentEQMix();
-
-			for (const TPair<USoundMix*, FSoundMixState>& SoundMixPair : SoundMixModifiers)
-			{
-				StatSoundMixes.AddDefaulted();
-				FAudioStats::FStatSoundMix& StatSoundMix = StatSoundMixes.Last();
-				StatSoundMix.MixName = SoundMixPair.Key->GetName();
-				StatSoundMix.InterpValue = SoundMixPair.Value.InterpValue;
-				StatSoundMix.RefCount = SoundMixPair.Value.ActiveRefCount + SoundMixPair.Value.PassiveRefCount;
-				StatSoundMix.bIsCurrentEQ = (SoundMixPair.Key == CurrentEQMix);
-			}
-		}
-	}
-#endif
-
 	DECLARE_CYCLE_STAT(TEXT("FGameThreadAudioTask.AudioSendResults"), STAT_AudioSendResults, STATGROUP_TaskGraphTasks);
 
 	const uint32 AudioDeviceID = DeviceHandle;
-	UReverbEffect* ReverbEffect = nullptr;
-	if (Effects)
-	{
-		ReverbEffect = Effects->GetCurrentReverbEffect();
-	}
-
-	FAudioThread::RunCommandOnGameThread([AudioDeviceID, ReverbEffect
-#if !UE_BUILD_SHIPPING
-											, ListenerPosition, StatSoundInfos, StatSoundMixes, bStatsStale
-#endif
-													]()
+	UReverbEffect* ReverbEffect = Effects ? Effects->GetCurrentReverbEffect() : nullptr;
+	FAudioThread::RunCommandOnGameThread([AudioDeviceID, ReverbEffect]()
 	{
 		// At shutdown, GEngine may already be null
 		if (GEngine)
@@ -4022,16 +3976,14 @@ void FAudioDevice::SendUpdateResultsToGameThread(const int32 FirstActiveIndex)
 				if (FAudioDevice* AudioDevice = AudioDeviceManager->GetAudioDevice(AudioDeviceID))
 				{
 					AudioDevice->CurrentReverbEffect = ReverbEffect;
-#if !UE_BUILD_SHIPPING
-					AudioDevice->AudioStats.ListenerLocation = ListenerPosition;
-					AudioDevice->AudioStats.StatSoundInfos = StatSoundInfos;
-					AudioDevice->AudioStats.StatSoundMixes = StatSoundMixes;
-					AudioDevice->AudioStats.bStale = bStatsStale;
-#endif
 				}
 			}
 		}
 	}, GET_STATID(STAT_AudioSendResults));
+
+#if ENABLE_AUDIO_DEBUG
+	FAudioDebugger::SendUpdateResultsToGameThread(*this, FirstActiveIndex);
+#endif // ENABLE_AUDIO_DEBUG
 }
 
 void FAudioDevice::StopAllSounds(bool bShouldStopUISounds)
@@ -4246,11 +4198,6 @@ void FAudioDevice::AddNewActiveSoundInternal(const FActiveSound& NewActiveSound,
 		(int32)ActiveSound->Sound->HasAnyFlags(RF_BeginDestroyed)))
 	{
 		static FName InvalidSoundName(TEXT("DESTROYED_Sound"));
-		ActiveSound->DebugOriginalSoundName = InvalidSoundName;
-	}
-	else
-	{
-		ActiveSound->DebugOriginalSoundName = ActiveSound->Sound->GetFName();
 	}
 #endif // !(UE_BUILD_SHIPPING || UE_BUILD_TEST)
 
@@ -4510,7 +4457,6 @@ void FAudioDevice::NotifyActiveSoundOcclusionTraceDone(FActiveSound* InActiveSou
 	}
 }
 
-
 FActiveSound* FAudioDevice::FindActiveSound(const uint64 AudioComponentID)
 {
 	check(IsInAudioThread());
@@ -4645,7 +4591,6 @@ float FAudioDevice::GetSquaredDistanceToListener(const FVector& Location, const
 	return (ListenerTranslation - Location).SizeSquared();
 }
 
-
 void FAudioDevice::GetMaxDistanceAndFocusFactor(USoundBase* Sound, const UWorld* World, const FVector& Location, const FSoundAttenuationSettings* AttenuationSettingsToApply, float& OutMaxDistance, float& OutFocusFactor)
 {
 	check(IsInGameThread());
@@ -4722,12 +4667,12 @@ bool FAudioDevice::SoundIsAudible(const FActiveSound& NewActiveSound)
 		return true;
 	}
 
-	if (VirtualSoundsEnabled() && (NewActiveSound.Sound->IsAllowedVirtual() || (NewActiveSound.bHandleSubtitles && NewActiveSound.bHasExternalSubtitles)))
+	if (PlayWhenSilentEnabled() && (NewActiveSound.Sound->SupportsSubtitles() || (NewActiveSound.bHandleSubtitles && NewActiveSound.bHasExternalSubtitles)))
 	{
 		return true;
 	}
 
-	if (NewActiveSound.Sound->IsVirtualizeWhenSilent())
+	if (NewActiveSound.Sound->IsPlayWhenSilent())
 	{
 		return true;
 	}
@@ -5118,15 +5063,15 @@ void FAudioDevice::PlaySoundAtLocation(USoundBase* Sound, UWorld* World, float V
 
 	GetMaxDistanceAndFocusFactor(Sound, World, Location, AttenuationSettingsToApply, MaxDistance, FocusFactor);
 
-	if (Sound->IsLooping() || Sound->IsVirtualizeWhenSilent() || SoundIsAudible(Sound, World, Location, AttenuationSettingsToApply, MaxDistance, FocusFactor))
+	if (Sound->IsLooping() || Sound->IsPlayWhenSilent() || SoundIsAudible(Sound, World, Location, AttenuationSettingsToApply, MaxDistance, FocusFactor))
 	{
 		const bool bIsInGameWorld = World->IsGameWorld();
 
 		FActiveSound NewActiveSound;
 		NewActiveSound.SetWorld(World);
 		NewActiveSound.SetSound(Sound);
-		NewActiveSound.VolumeMultiplier = VolumeMultiplier;
-		NewActiveSound.PitchMultiplier = PitchMultiplier;
+		NewActiveSound.SetVolume(VolumeMultiplier);
+		NewActiveSound.SetPitch(PitchMultiplier);
 		NewActiveSound.RequestedStartTime = FMath::Max(0.0f, StartTime);
 		NewActiveSound.bLocationDefined = true;
 		NewActiveSound.Transform.SetTranslation(Location);
@@ -5793,7 +5738,7 @@ void FAudioDevice::UpdateVirtualLoops()
 				VirtualLoopsToRetrigger.Add(VirtualLoop);
 			}
 
-			// If signaled to fade out and virtualized, and to pending stop list.
+			// If set to fade out and set to play when silent, add to pending stop list.
 			if (ActiveSound.bFadingOut)
 			{
 				AddSoundToStop(&ActiveSound);
@@ -5908,10 +5853,9 @@ void FAudioDevice::SetGlobalFocusSettings(const FGlobalFocusSettings& NewFocusSe
 	GlobalFocusSettings_GameThread = NewFocusSettings;
 
 	DECLARE_CYCLE_STAT(TEXT("FAudioThreadTask.SetGlobalListenerFocusParameters"), STAT_AudioSetGlobalListenerFocusParameters, STATGROUP_TaskGraphTasks);
-	FAudioDevice* AudioDevice = this;
-	FAudioThread::RunCommandOnAudioThread([AudioDevice, NewFocusSettings]()
+	FAudioThread::RunCommandOnAudioThread([this, NewFocusSettings]()
 	{
-		AudioDevice->GlobalFocusSettings = NewFocusSettings;
+		GlobalFocusSettings = NewFocusSettings;
 	}, GET_STATID(STAT_AudioSetGlobalListenerFocusParameters));
 }
 
@@ -5921,10 +5865,9 @@ void FAudioDevice::SetGlobalPitchModulation(float PitchModulation, float TimeSec
 	{
 		DECLARE_CYCLE_STAT(TEXT("FAudioThreadTask.SetGlobalPitchModulation"), STAT_AudioSetGlobalPitchModulation, STATGROUP_TaskGraphTasks);
 
-		FAudioDevice* AudioDevice = this;
-		FAudioThread::RunCommandOnAudioThread([AudioDevice, PitchModulation, TimeSec]()
+		FAudioThread::RunCommandOnAudioThread([this, PitchModulation, TimeSec]()
 		{
-			AudioDevice->SetGlobalPitchModulation(PitchModulation, TimeSec);
+			SetGlobalPitchModulation(PitchModulation, TimeSec);
 		}, GET_STATID(STAT_AudioSetGlobalPitchModulation));
 
 		return;
@@ -5933,7 +5876,7 @@ void FAudioDevice::SetGlobalPitchModulation(float PitchModulation, float TimeSec
 	GlobalPitchScale.Set(PitchModulation, TimeSec);
 }
 
-float FAudioDevice::ClampPitch(float InPitchScale)
+float FAudioDevice::ClampPitch(float InPitchScale) const
 {
 	return FMath::Clamp(InPitchScale, GlobalMinPitch, GlobalMaxPitch);
 }
@@ -5956,43 +5899,21 @@ void FAudioDevice::SetPlatformAudioHeadroom(const float InPlatformHeadRoom)
 	PlatformAudioHeadroom = InPlatformHeadRoom;
 }
 
-#if !UE_BUILD_SHIPPING
-void HandleDumpActiveSounds(UWorld* World)
+const TArray<FWaveInstance*>& FAudioDevice::GetActiveWaveInstances() const
 {
-	if (FAudioDevice* AudioDevice = (World ? World->GetAudioDevice() : GEngine->GetMainAudioDevice()))
-	{
-		AudioDevice->DumpActiveSounds();
-	}
+	check(IsInAudioThread());
+	return ActiveWaveInstances;
 }
 
-static FAutoConsoleCommandWithWorld DumpActiveSounds(TEXT("Audio.DumpActiveSounds"), TEXT("Outputs data about all the currently active sounds."), FConsoleCommandWithWorldDelegate::CreateStatic(&HandleDumpActiveSounds), ECVF_Cheat);
-
-void FAudioDevice::DumpActiveSounds() const
+const TMap<FName, FActivatedReverb>& FAudioDevice::GetActiveReverb() const
 {
-	check(IsInGameThread());
-
-	FAudioThreadSuspendContext SuspendAudio;
-
-	UE_LOG(LogAudio, Display, TEXT("Active Sound Count: %d"), ActiveSounds.Num());
-	UE_LOG(LogAudio, Display, TEXT("------------------------"), ActiveSounds.Num());
-
-	for (const FActiveSound* ActiveSound : ActiveSounds)
-	{
-		if (ActiveSound)
-		{
-			UE_LOG(LogAudio, Display, TEXT("%s (%.3g) - %s"), *ActiveSound->GetSound()->GetName(), ActiveSound->GetSound()->GetDuration(), *ActiveSound->GetAudioComponentName());
-
-			for (const TPair<UPTRINT, FWaveInstance*>& WaveInstancePair : ActiveSound->WaveInstances)
-			{
-				const FWaveInstance* WaveInstance = WaveInstancePair.Value;
-				UE_LOG(LogAudio, Display, TEXT("   %s (%.3g) (%d) - %.3g"), *WaveInstance->GetName(), WaveInstance->WaveData->GetDuration(), WaveInstance->WaveData->GetResourceSizeBytes(EResourceSizeMode::EstimatedTotal), WaveInstance->GetActualVolume());
-			}
-		}
-	}
-
+	return ActivatedReverbs;
+}
 
+const TMap<FWaveInstance*, FSoundSource*>& FAudioDevice::GetWaveInstanceSourceMap() const
+{
+	return WaveInstanceSourceMap;
 }
-#endif
 
 bool FAudioDevice::ShouldUseAttenuation(const UWorld* World) const
 {
diff --git a/Engine/Source/Runtime/Engine/Private/AudioDeviceManager.cpp b/Engine/Source/Runtime/Engine/Private/AudioDeviceManager.cpp
--- a/Engine/Source/Runtime/Engine/Private/AudioDeviceManager.cpp
+++ b/Engine/Source/Runtime/Engine/Private/AudioDeviceManager.cpp
@@ -1,9 +1,10 @@
 // Copyright 1998-2019 Epic Games, Inc. All Rights Reserved.
-
 #include "AudioDeviceManager.h"
+
+#include "Audio/AudioDebug.h"
 #include "AudioDevice.h"
-#include "Sound/SoundWave.h"
 #include "Sound/AudioSettings.h"
+#include "Sound/SoundWave.h"
 #include "GameFramework/GameUserSettings.h"
 #include "Misc/CommandLine.h"
 #include "Misc/ConfigCacheIni.h"
@@ -75,7 +76,6 @@ FAudioDeviceManager::FAudioDeviceManager()
 	, ActiveAudioDeviceHandle(INDEX_NONE)
 	, bUsingAudioMixer(false)
 	, bPlayAllDeviceAudio(false)
-	, bVisualize3dDebug(false)
 	, bOnlyToggleAudioMixerOnce(false)
 	, bToggledAudioMixer(false)
 {
@@ -933,11 +933,16 @@ void FAudioDeviceManager::TogglePlayAllDeviceAudio()
 
 bool FAudioDeviceManager::IsVisualizeDebug3dEnabled() const
 {
-	return bVisualize3dDebug || CVarIsVisualizeEnabled;
+#if ENABLE_AUDIO_DEBUG
+	return AudioDebugger.IsVisualizeDebug3dEnabled() || CVarIsVisualizeEnabled;
+#else // ENABLE_AUDIO_DEBUG
+	return false;
+#endif // !ENABLE_AUDIO_DEBUG
 }
 
 void FAudioDeviceManager::ToggleVisualize3dDebug()
 {
+#if ENABLE_AUDIO_DEBUG
 	if (!IsInAudioThread())
 	{
 		DECLARE_CYCLE_STAT(TEXT("FAudioThreadTask.ToggleVisualize3dDebug"), STAT_ToggleVisualize3dDebug, STATGROUP_AudioThreadCommands);
@@ -952,33 +957,8 @@ void FAudioDeviceManager::ToggleVisualize3dDebug()
 		return;
 	}
 
-	bVisualize3dDebug = !bVisualize3dDebug;
-}
-
-void FAudioDeviceManager::ToggleDebugStat(const uint8 StatBitMask)
-{
-#if !UE_BUILD_SHIPPING
-	if (!IsInAudioThread())
-	{
-		DECLARE_CYCLE_STAT(TEXT("FAudioThreadTask.ToggleDebugStat"), STAT_ToggleDebugStat, STATGROUP_AudioThreadCommands);
-
-		FAudioDeviceManager* AudioDeviceManager = this;
-		FAudioThread::RunCommandOnAudioThread([AudioDeviceManager, StatBitMask]()
-		{
-			AudioDeviceManager->ToggleDebugStat(StatBitMask);
-		}, GET_STATID(STAT_ToggleDebugStat));
-
-		return;
-	}
-
-	for (FAudioDevice* AudioDevice : Devices)
-	{
-		if (AudioDevice)
-		{
-			AudioDevice->UpdateRequestedStat(StatBitMask);
-		}
-	}
-#endif
+	GetDebugger().ToggleVisualizeDebug3dEnabled();
+#endif // ENABLE_AUDIO_DEBUG
 }
 
 void FAudioDeviceManager::SetDebugSoloSoundClass(const TCHAR* SoundClassName)
@@ -1144,3 +1124,10 @@ void FAudioDeviceManager::SetDynamicSoundVolume(ESoundType SoundType, const FNam
 	TTuple<ESoundType, FName> Key(SoundType, SoundName);
 	DynamicSoundVolumes.FindOrAdd(Key) = Volume;
 }
+
+#if ENABLE_AUDIO_DEBUG
+FAudioDebugger& FAudioDeviceManager::GetDebugger()
+{
+	return AudioDebugger;
+}
+#endif // ENABLE_AUDIO_DEBUG
diff --git a/Engine/Source/Runtime/Engine/Private/AudioPluginUtilities.cpp b/Engine/Source/Runtime/Engine/Private/AudioPluginUtilities.cpp
--- a/Engine/Source/Runtime/Engine/Private/AudioPluginUtilities.cpp
+++ b/Engine/Source/Runtime/Engine/Private/AudioPluginUtilities.cpp
@@ -6,26 +6,6 @@
 #include "CoreGlobals.h"
 
 
-/** Get the target setting name for each platform type. */
-FORCEINLINE const TCHAR* GetPluginConfigName(EAudioPlugin PluginType)
-{
-	switch (PluginType)
-	{
-		case EAudioPlugin::SPATIALIZATION:
-			return TEXT("SpatializationPlugin");
-
-		case EAudioPlugin::REVERB:
-			return TEXT("ReverbPlugin");
-
-		case EAudioPlugin::OCCLUSION:
-			return TEXT("OcclusionPlugin");
-
-		default:
-			checkf(false, TEXT("Undefined audio plugin type."));
-			return TEXT("");
-	}
-}
-
 /************************************************************************/
 /* Plugin Utilities                                                     */
 /************************************************************************/
@@ -111,6 +91,29 @@ const TCHAR* AudioPluginUtilities::GetPlatformConfigSection(EAudioPlatform Audio
 	return *UnknownConfig;
 }
 
+/** Get the target setting name for each platform type. */
+FORCEINLINE const TCHAR* GetPluginConfigName(EAudioPlugin PluginType)
+{
+	switch (PluginType)
+	{
+		case EAudioPlugin::SPATIALIZATION:
+			return TEXT("SpatializationPlugin");
+
+		case EAudioPlugin::REVERB:
+			return TEXT("ReverbPlugin");
+
+		case EAudioPlugin::OCCLUSION:
+			return TEXT("OcclusionPlugin");
+
+		case EAudioPlugin::MODULATION:
+			return TEXT("ModulationPlugin");
+
+		default:
+			checkf(false, TEXT("Undefined audio plugin type."));
+			return TEXT("");
+	}
+}
+
 IAudioSpatializationFactory* AudioPluginUtilities::GetDesiredSpatializationPlugin(EAudioPlatform AudioPlatform)
 {
 	//Get the name of the desired spatialization plugin:
@@ -171,6 +174,26 @@ IAudioOcclusionFactory* AudioPluginUtilities::GetDesiredOcclusionPlugin(EAudioPl
 	return nullptr;
 }
 
+IAudioModulationFactory* AudioPluginUtilities::GetDesiredModulationPlugin(EAudioPlatform AudioPlatform)
+{
+	//Get the name of the desired spatialization plugin:
+	FString DesiredModulationPlugin = GetDesiredPluginName(EAudioPlugin::MODULATION, AudioPlatform);
+
+	TArray<IAudioModulationFactory *> ModulationPluginFactories = IModularFeatures::Get().GetModularFeatureImplementations<IAudioModulationFactory>(IAudioModulationFactory::GetModularFeatureName());
+
+	//Iterate through all of the plugins we've discovered:
+	for (IAudioModulationFactory* PluginFactory : ModulationPluginFactories)
+	{
+		//if this plugin's name matches the name found in the platform settings, use it:
+		if (PluginFactory->GetDisplayName().Equals(DesiredModulationPlugin))
+		{
+			return PluginFactory;
+		}
+	}
+
+	return nullptr;
+}
+
 FString AudioPluginUtilities::GetDesiredPluginName(EAudioPlugin PluginType, EAudioPlatform AudioPlatform)
 {
 	FString PluginName;
diff --git a/Engine/Source/Runtime/Engine/Private/AudioSettings.cpp b/Engine/Source/Runtime/Engine/Private/AudioSettings.cpp
--- a/Engine/Source/Runtime/Engine/Private/AudioSettings.cpp
+++ b/Engine/Source/Runtime/Engine/Private/AudioSettings.cpp
@@ -54,7 +54,7 @@ UAudioSettings::UAudioSettings(const FObjectInitializer& ObjectInitializer)
 	SectionName = TEXT("Audio");
 	AddDefaultSettings();
 
-	bAllowVirtualizedSounds = true;
+	bAllowPlayWhenSilent = true;
 	bIsAudioMixerEnabled = false;
 
 	GlobalMinPitchScale = 0.4F;
@@ -66,7 +66,7 @@ void UAudioSettings::AddDefaultSettings()
 	FAudioQualitySettings DefaultSettings;
 	DefaultSettings.DisplayName = LOCTEXT("DefaultSettingsName", "Default");
 	QualityLevels.Add(DefaultSettings);
-	bAllowVirtualizedSounds = true;
+	bAllowPlayWhenSilent = true;
 	DefaultReverbSendLevel = 0.0f;
 	bEnableLegacyReverb = false;
 	VoiPSampleRate = EVoiceSampleRate::Low16000Hz;
@@ -102,7 +102,7 @@ void UAudioSettings::PostEditChangeChainProperty(FPropertyChangedChainEvent& Pro
 						bool bFoundDuplicate;
 						int32 NewQualityLevelIndex = 0;
 						FText NewLevelName;
-						do 
+						do
 						{
 							bFoundDuplicate = false;
 							NewLevelName = FText::Format(LOCTEXT("NewQualityLevelName","New Level{0}"), (NewQualityLevelIndex > 0 ? FText::FromString(FString::Printf(TEXT(" %d"),NewQualityLevelIndex)) : FText::GetEmpty()));
@@ -139,6 +139,20 @@ void UAudioSettings::PostEditChangeChainProperty(FPropertyChangedChainEvent& Pro
 }
 #endif
 
+void UAudioSettings::Serialize(FArchive& Ar)
+{
+	Super::Serialize(Ar);
+
+	if (Ar.IsLoading() || Ar.IsSaving())
+	{
+		if (bAllowVirtualizedSounds_DEPRECATED)
+		{
+			bAllowVirtualizedSounds_DEPRECATED = 0;
+			bAllowPlayWhenSilent = 1;
+		}
+	}
+}
+
 const FAudioQualitySettings& UAudioSettings::GetQualityLevelSettings(int32 QualityLevel) const
 {
 	check(QualityLevels.Num() > 0);
@@ -158,7 +172,7 @@ const bool UAudioSettings::IsAudioMixerEnabled() const
 int32 UAudioSettings::GetHighestMaxChannels() const
 {
 	check(QualityLevels.Num() > 0);
-	
+
 	int32 HighestMaxChannels = -1;
 	for (const FAudioQualitySettings& Settings : QualityLevels)
 	{
diff --git a/Engine/Source/Runtime/Engine/Private/AudioVirtualLoop.cpp b/Engine/Source/Runtime/Engine/Private/AudioVirtualLoop.cpp
--- a/Engine/Source/Runtime/Engine/Private/AudioVirtualLoop.cpp
+++ b/Engine/Source/Runtime/Engine/Private/AudioVirtualLoop.cpp
@@ -2,8 +2,8 @@
 #include "AudioVirtualLoop.h"
 
 #include "ActiveSound.h"
+#include "Audio/AudioDebug.h"
 #include "AudioDevice.h"
-#include "DrawDebugHelpers.h"
 #include "Sound/SoundBase.h"
 
 
@@ -35,13 +35,6 @@ FAutoConsoleVariableRef CVarVirtualLoopsUpdateRateMax(
 	TEXT("Sets maximum rate to check if sound becomes audible again (at beyond sound's max audible distance + perf scaling distance).\n"),
 	ECVF_Default);
 
-static int32 VirtualLoopsVisualizeEnabledCVar = 1;
-FAutoConsoleVariableRef CVarAudioVisualizeVirtualLoopsEnabled(
-	TEXT("au.3dVisualize.VirtualLoops"),
-	VirtualLoopsVisualizeEnabledCVar,
-	TEXT("Whether or not virtualized loops are visible when 3d visualize is enabled. \n")
-	TEXT("0: Not Enabled, 1: Enabled"),
-	ECVF_Default);
 
 FAudioVirtualLoop::FAudioVirtualLoop()
 	: TimeSinceLastUpdate(0.0f)
@@ -143,7 +136,7 @@ bool FAudioVirtualLoop::IsInAudibleRange(const FActiveSound& InActiveSound, cons
 	}
 	check(AudioDevice);
 
-	if (AudioDevice->VirtualSoundsEnabled() && InActiveSound.IsVirtualizeWhenSilent())
+	if (AudioDevice->PlayWhenSilentEnabled() && InActiveSound.IsPlayWhenSilent())
 	{
 		return true;
 	}
@@ -178,10 +171,9 @@ bool FAudioVirtualLoop::CanRealize(float DeltaTime)
 		TimeSinceLastUpdate = 0.0f;
 	}
 
-	if (VirtualLoopsVisualizeEnabledCVar)
-	{
-		DrawDebugInfo();
-	}
+#if ENABLE_AUDIO_DEBUG
+	FAudioDebugger::DrawDebugInfo(*this);
+#endif // ENABLE_AUDIO_DEBUG
 
 	// If not audible, update when will be checked again and return false
 	if (!IsInAudibleRange(*ActiveSound))
@@ -192,35 +184,3 @@ bool FAudioVirtualLoop::CanRealize(float DeltaTime)
 
 	return true;
 }
-
-void FAudioVirtualLoop::DrawDebugInfo() const
-{
-#if ENABLE_DRAW_DEBUG
-	// Draw 3d Debug information about this source, if enabled
-	FAudioDeviceManager* DeviceManager = GEngine->GetAudioDeviceManager();
-
-	if (DeviceManager && DeviceManager->IsVisualizeDebug3dEnabled())
-	{
-		DECLARE_CYCLE_STAT(TEXT("FAudioThreadTask.DrawVirtualLoopDebugInfo"), STAT_AudioDrawVirtualLoopDebugInfo, STATGROUP_TaskGraphTasks);
-
-		USoundBase* Sound = ActiveSound->GetSound();
-		check(Sound);
-
-		const FTransform Transform = ActiveSound->Transform;
-		const TWeakObjectPtr<UWorld> World = ActiveSound->GetWeakWorld();
-		const FString Name = Sound->GetName();
-		const float DrawInterval = UpdateInterval;
-		FAudioThread::RunCommandOnGameThread([World, Transform, Name, DrawInterval]()
-		{
-			if (World.IsValid())
-			{
-				const FString Description = FString::Printf(TEXT("%s [V]"), *Name);
-				FVector Location = Transform.GetLocation();
-				FRotator Rotation = Transform.GetRotation().Rotator();
-				DrawDebugCrosshairs(World.Get(), Location, Rotation, 20.0f, FColor::Blue, false, DrawInterval, SDPG_Foreground);
-				DrawDebugString(World.Get(), Location + FVector(0, 0, 32), *Description, nullptr, FColor::Blue, DrawInterval, false);
-			}
-		}, GET_STATID(STAT_AudioDrawVirtualLoopDebugInfo));
-	}
-#endif // ENABLE_DRAW_DEBUG
-}
diff --git a/Engine/Source/Runtime/Engine/Private/Components/AudioComponent.cpp b/Engine/Source/Runtime/Engine/Private/Components/AudioComponent.cpp
--- a/Engine/Source/Runtime/Engine/Private/Components/AudioComponent.cpp
+++ b/Engine/Source/Runtime/Engine/Private/Components/AudioComponent.cpp
@@ -393,13 +393,15 @@ void UAudioComponent::PlayInternal(const float StartTime, const float FadeInDura
 			AudioDevice->GetMaxDistanceAndFocusFactor(Sound, World, Location, AttenuationSettingsToApply, MaxDistance, FocusFactor);
 
 			FActiveSound NewActiveSound;
-			NewActiveSound.SetAudioComponent(this);
+			NewActiveSound.SetAudioComponent(*this);
 			NewActiveSound.SetWorld(GetWorld());
 			NewActiveSound.SetSound(Sound);
 			NewActiveSound.SetSoundClass(SoundClassOverride);
 			NewActiveSound.ConcurrencySet = ConcurrencySet;
 
-			NewActiveSound.VolumeMultiplier = (VolumeModulationMax + ((VolumeModulationMin - VolumeModulationMax) * FMath::SRand())) * VolumeMultiplier;
+			const float Volume = (VolumeModulationMax + ((VolumeModulationMin - VolumeModulationMax) * FMath::SRand())) * VolumeMultiplier;
+			NewActiveSound.SetVolume(Volume);
+
 			// The priority used for the active sound is the audio component's priority scaled with the sound's priority
 			if (bOverridePriority)
 			{
@@ -410,7 +412,9 @@ void UAudioComponent::PlayInternal(const float StartTime, const float FadeInDura
 				NewActiveSound.Priority = Sound->Priority;
 			}
 
-			NewActiveSound.PitchMultiplier = (PitchModulationMax + ((PitchModulationMin - PitchModulationMax) * FMath::SRand())) * PitchMultiplier;
+			const float Pitch = (PitchModulationMax + ((PitchModulationMin - PitchModulationMax) * FMath::SRand())) * PitchMultiplier;
+			NewActiveSound.SetPitch(Pitch);
+
 			NewActiveSound.bEnableLowPassFilter = bEnableLowPassFilter;
 			NewActiveSound.LowPassFilterFrequency = LowPassFilterFrequency;
 			NewActiveSound.RequestedStartTime = FMath::Max(0.f, StartTime);
@@ -1046,7 +1050,7 @@ void UAudioComponent::SetVolumeMultiplier(const float NewVolumeMultiplier)
 				FActiveSound* ActiveSound = AudioDevice->FindActiveSound(MyAudioComponentID);
 				if (ActiveSound)
 				{
-					ActiveSound->VolumeMultiplier = NewVolumeMultiplier;
+					ActiveSound->SetVolume(NewVolumeMultiplier);
 				}
 			}, GET_STATID(STAT_AudioSetVolumeMultiplier));
 		}
@@ -1070,7 +1074,7 @@ void UAudioComponent::SetPitchMultiplier(const float NewPitchMultiplier)
 				FActiveSound* ActiveSound = AudioDevice->FindActiveSound(MyAudioComponentID);
 				if (ActiveSound)
 				{
-					ActiveSound->PitchMultiplier = NewPitchMultiplier;
+					ActiveSound->SetPitch(NewPitchMultiplier);
 				}
 			}, GET_STATID(STAT_AudioSetPitchMultiplier));
 		}
diff --git a/Engine/Source/Runtime/Engine/Private/GameViewportClient.cpp b/Engine/Source/Runtime/Engine/Private/GameViewportClient.cpp
--- a/Engine/Source/Runtime/Engine/Private/GameViewportClient.cpp
+++ b/Engine/Source/Runtime/Engine/Private/GameViewportClient.cpp
@@ -39,6 +39,7 @@
 #include "EngineModule.h"
 #include "AudioDeviceManager.h"
 #include "AudioDevice.h"
+#include "Audio/AudioDebug.h"
 #include "Sound/SoundWave.h"
 #include "HighResScreenshot.h"
 #include "BufferVisualizationData.h"
@@ -85,9 +86,6 @@ FSimpleMulticastDelegate UGameViewportClient::CreatedDelegate;
 /** A list of all the stat names which are enabled for this viewport (static so they persist between runs) */
 TArray<FString> UGameViewportClient::EnabledStats;
 
-/** Those sound stat flags which are enabled on this viewport */
-FViewportClient::ESoundShowFlags::Type UGameViewportClient::SoundShowFlags = FViewportClient::ESoundShowFlags::Disabled;
-
 /**
  * UI Stats
  */
@@ -374,15 +372,15 @@ void UGameViewportClient::SetEnabledStats(const TArray<FString>& InEnabledStats)
 		UE_LOG(LogPlayerManagement, Warning, TEXT("WARNING: Stats disabled for non multi-threading platforms"));
 	}
 
-#if !UE_BUILD_SHIPPING
+#if ENABLE_AUDIO_DEBUG
 	if (UWorld* MyWorld = GetWorld())
 	{
-		if (FAudioDevice* AudioDevice = MyWorld->GetAudioDevice())
+		if (FAudioDeviceManager* DeviceManager = GEngine->GetAudioDeviceManager())
 		{
-			AudioDevice->ResolveDesiredStats(this);
+			FAudioDebugger::ResolveDesiredStats(this);
 		}
 	}
-#endif
+#endif // ENABLE_AUDIO_DEBUG
 }
 
 
@@ -433,12 +431,9 @@ void UGameViewportClient::Init(struct FWorldContext& WorldContext, UGameInstance
 			{
 				AudioDeviceHandle = NewDeviceResults.Handle;
 
-#if !UE_BUILD_SHIPPING
-				if (NewDeviceResults.bNewDevice)
-				{
-					NewDeviceResults.AudioDevice->ResolveDesiredStats(this);
-				}
-#endif // UE_BUILD_SHIPPING
+#if ENABLE_AUDIO_DEBUG
+				FAudioDebugger::ResolveDesiredStats(this);
+#endif // ENABLE_AUDIO_DEBUG
 
 				// Set the base mix of the new device based on the world settings of the world
 				if (World)
@@ -972,7 +967,7 @@ void UGameViewportClient::SetDropDetail(float DeltaSeconds)
 		GetWorld()->bAggressiveLOD	= FrameRate < FMath::Clamp(GEngine->MinDesiredFrameRate - 5.f, 1.f, 100.f) && !bTimeIsManipulated;
 
 		// this is slick way to be able to do something based on the frametime and whether we are bound by one thing or another
-#if 0 
+#if 0
 		// so where we check to see if we are above some threshold and below 150 ms (any thing above that is usually blocking loading of some sort)
 		// also we don't want to do the auto trace when we are blocking on async loading
 		if ((0.070 < FrameTime) && (FrameTime < 0.150) && IsAsyncLoading() == false && GetWorld()->bRequestedBlockOnAsyncLoading == false && (GetWorld()->GetTimeSeconds() > 30.0f))
@@ -992,7 +987,7 @@ void UGameViewportClient::SetDropDetail(float DeltaSeconds)
 #endif // WITH_EDITORONLY_DATA
 			}
 		}
-#endif // 0 
+#endif // 0
 	}
 }
 
@@ -1120,7 +1115,7 @@ void UGameViewportClient::Draw(FViewport* InViewport, FCanvas* SceneCanvas)
 	// Create a temporary canvas if there isn't already one.
 	static FName CanvasObjectName(TEXT("CanvasObject"));
 	UCanvas* CanvasObject = GetCanvasByName(CanvasObjectName);
-	CanvasObject->Canvas = SceneCanvas;		
+	CanvasObject->Canvas = SceneCanvas;
 
 	// Create temp debug canvas object
 	FIntPoint DebugCanvasSize = InViewport->GetSizeXY();
@@ -1155,7 +1150,7 @@ void UGameViewportClient::Draw(FViewport* InViewport, FCanvas* SceneCanvas)
 	}
 
 	// create the view family for rendering the world scene to the viewport's render target
-	FSceneViewFamilyContext ViewFamily(FSceneViewFamily::ConstructionValues( 	
+	FSceneViewFamilyContext ViewFamily(FSceneViewFamily::ConstructionValues(
 		InViewport,
 		MyWorld->Scene,
 		EngineShowFlags)
@@ -1405,7 +1400,7 @@ void UGameViewportClient::Draw(FViewport* InViewport, FCanvas* SceneCanvas)
 			}
 		}
 	}
-	
+
 	// If the views don't cover the entire bounding rectangle, clear the entire buffer.
 	bool bBufferCleared = false;
 	bool bStereoscopicPass = (ViewFamily.Views.Num() != 0 && ViewFamily.Views[0]->StereoPass != eSSP_FULL);
@@ -1415,10 +1410,10 @@ void UGameViewportClient::Draw(FViewport* InViewport, FCanvas* SceneCanvas)
 		{
 			SceneCanvas->Clear(FLinearColor::Transparent);
 		}
-		
+
 		bBufferCleared = true;
 	}
-	
+
 	// Force screen percentage show flag to be turned off if not supported.
 	if (!ViewFamily.SupportsScreenPercentage())
 	{
@@ -1520,7 +1515,7 @@ void UGameViewportClient::Draw(FViewport* InViewport, FCanvas* SceneCanvas)
 		// Make sure RHI resources get flushed if we're not using a renderer
 		ENQUEUE_RENDER_COMMAND(UGameViewportClient_FlushRHIResources)(
 			[](FRHICommandListImmediate& RHICmdList)
-			{ 
+			{
 				RHICmdList.ImmediateFlush(EImmediateFlushType::FlushRHIThreadFlushResources);
 			});
 	}
@@ -1552,7 +1547,7 @@ void UGameViewportClient::Draw(FViewport* InViewport, FCanvas* SceneCanvas)
 			SceneCanvas->DrawTile(MinX,MaxY,MaxX,InViewport->GetSizeXY().Y,0.0f,0.0f,1.0f,1.f,FLinearColor::Black,NULL,false);
 		}
 	}
-	
+
 	// Remove temporary debug lines.
 	if (MyWorld->LineBatcher != nullptr)
 	{
@@ -1590,13 +1585,13 @@ void UGameViewportClient::Draw(FViewport* InViewport, FCanvas* SceneCanvas)
 					{
 						// rendering to directly to viewport target
 						FVector CanvasOrigin(FMath::TruncToFloat(View->UnscaledViewRect.Min.X), FMath::TruncToInt(View->UnscaledViewRect.Min.Y), 0.f);
-												
+
 						CanvasObject->Init(View->UnscaledViewRect.Width(), View->UnscaledViewRect.Height(), View, SceneCanvas);
 
 						// Set the canvas transform for the player's view rectangle.
 						check(SceneCanvas);
 						SceneCanvas->PushAbsoluteTransform(FTranslationMatrix(CanvasOrigin));
-						CanvasObject->ApplySafeZoneTransform();						
+						CanvasObject->ApplySafeZoneTransform();
 
 						// Render the player's HUD.
 						if( PlayerController->MyHUD )
@@ -1607,7 +1602,7 @@ void UGameViewportClient::Draw(FViewport* InViewport, FCanvas* SceneCanvas)
 							PlayerController->MyHUD->SetCanvas(CanvasObject, DebugCanvasObject);
 
 							PlayerController->MyHUD->PostRender();
-							
+
 							// Put these pointers back as if a blueprint breakpoint hits during HUD PostRender they can
 							// have been changed
 							CanvasObject->Canvas = SceneCanvas;
@@ -1628,7 +1623,7 @@ void UGameViewportClient::Draw(FViewport* InViewport, FCanvas* SceneCanvas)
 						}
 
 						CanvasObject->PopSafeZoneTransform();
-						SceneCanvas->PopTransform();						
+						SceneCanvas->PopTransform();
 
 						// draw subtitles
 						if (!bDisplayedSubtitles)
@@ -1703,7 +1698,7 @@ void UGameViewportClient::Draw(FViewport* InViewport, FCanvas* SceneCanvas)
 void UGameViewportClient::ProcessScreenShots(FViewport* InViewport)
 {
 	if (GIsDumpingMovie || FScreenshotRequest::IsScreenshotRequested() || GIsHighResScreenshot)
-	{	
+	{
 		TArray<FColor> Bitmap;
 
 		bool bShowUI = false;
@@ -1753,7 +1748,7 @@ void UGameViewportClient::ProcessScreenShots(FViewport* InViewport)
 				}
 
 				GetHighResScreenshotConfig().MergeMaskIntoAlpha(Bitmap);
-				
+
 				FIntRect SourceRect(0, 0, GScreenshotResolutionX, GScreenshotResolutionY);
 				if (GIsHighResScreenshot)
 				{
@@ -1830,7 +1825,7 @@ void UGameViewportClient::Precache()
 TOptional<bool> UGameViewportClient::QueryShowFocus(const EFocusCause InFocusCause) const
 {
 	UUserInterfaceSettings* UISettings = GetMutableDefault<UUserInterfaceSettings>(UUserInterfaceSettings::StaticClass());
-	
+
 	if ( UISettings->RenderFocusRule == ERenderFocusRule::Never ||
 		(UISettings->RenderFocusRule == ERenderFocusRule::NonPointer && InFocusCause == EFocusCause::Mouse) ||
 		(UISettings->RenderFocusRule == ERenderFocusRule::NavigationOnly && InFocusCause != EFocusCause::Navigation))
@@ -1843,7 +1838,7 @@ TOptional<bool> UGameViewportClient::QueryShowFocus(const EFocusCause InFocusCau
 
 void UGameViewportClient::LostFocus(FViewport* InViewport)
 {
-	// We need to reset some key inputs, since keyup events will sometimes not be processed (such as going into immersive/maximized mode).  
+	// We need to reset some key inputs, since keyup events will sometimes not be processed (such as going into immersive/maximized mode).
 	// Resetting them will prevent them from "sticking"
 	UWorld* const ViewportWorld = GetWorld();
 	if (ViewportWorld && !ViewportWorld->bIsTearingDown)
@@ -1874,7 +1869,7 @@ void UGameViewportClient::ReceivedFocus(FViewport* InViewport)
 #endif
 
 	if (GEngine && GEngine->GetAudioDeviceManager())
-	{ 
+	{
 		GEngine->GetAudioDeviceManager()->SetActiveDevice(AudioDeviceHandle);
 		bHasAudioFocus = true;
 	}
@@ -1917,7 +1912,7 @@ void UGameViewportClient::CloseRequested(FViewport* InViewport)
 #if PLATFORM_DESKTOP || PLATFORM_HTML5
 	FSlateApplication::Get().SetGameIsFakingTouchEvents(false);
 #endif
-	
+
 	// broadcast close request to anyone that registered an interest
 	CloseRequestedDelegate.Broadcast(InViewport);
 
@@ -2132,7 +2127,7 @@ void UGameViewportClient::LayoutPlayers()
 {
 	UpdateActiveSplitscreenType();
 	const ESplitScreenType::Type SplitType = GetCurrentSplitscreenConfiguration();
-	
+
 	// Initialize the players
 	const TArray<ULocalPlayer*>& PlayerList = GetOuterUEngine()->GetGamePlayers(this);
 
@@ -2417,17 +2412,17 @@ bool UGameViewportClient::CalculateDeadZoneForAllSides( ULocalPlayer* LPlayer, U
 }
 
 void UGameViewportClient::DrawTitleSafeArea( UCanvas* Canvas )
-{	
+{
 #if WITH_EDITOR
 	FMargin SafeZone;
 	const ULevelEditorPlaySettings* PlayInSettings = GetDefault<ULevelEditorPlaySettings>();
-	
+
 	const float Width = Canvas->UnsafeSizeX;
 	const float Height = Canvas->UnsafeSizeY;
 	const FLinearColor UnsafeZoneColor(1.0f, 0.0f, 0.0f, 0.25f);
 	FCanvasTileItem TileItem(FVector2D::ZeroVector, GWhiteTexture, UnsafeZoneColor);
 	TileItem.BlendMode = SE_BLEND_Translucent;
-	
+
 	// CalculateSafeZoneValues() can be slow, so we only want to run it if we have boundaries to draw
 	if (FDisplayMetrics::GetDebugTitleSafeZoneRatio() < 1.f)
 	{
@@ -2504,7 +2499,7 @@ void UGameViewportClient::DrawTransitionMessage(UCanvas* Canvas,const FString& M
 	TextItem.Text = FText::FromString(Message);
 	float XL, YL;
 	Canvas->StrLen( Font , Message, XL, YL );
-	Canvas->DrawItem( TextItem, 0.5f * (Canvas->ClipX - XL), 0.66f * Canvas->ClipY - YL * 0.5f );	
+	Canvas->DrawItem( TextItem, 0.5f * (Canvas->ClipX - XL), 0.66f * Canvas->ClipY - YL * 0.5f );
 }
 
 void UGameViewportClient::NotifyPlayerAdded( int32 PlayerIndex, ULocalPlayer* AddedPlayer )
@@ -2646,7 +2641,7 @@ bool UGameViewportClient::Exec( UWorld* InWorld, const TCHAR* Cmd,FOutputDevice&
 	else if( FParse::Command(&Cmd,TEXT("TOGGLE_FULLSCREEN")) || FParse::Command(&Cmd,TEXT("FULLSCREEN")) )
 	{
 		return HandleToggleFullscreenCommand();
-	}	
+	}
 	else if( FParse::Command(&Cmd,TEXT("SETRES")) )
 	{
 		return HandleSetResCommand( Cmd, Ar );
@@ -2661,7 +2656,7 @@ bool UGameViewportClient::Exec( UWorld* InWorld, const TCHAR* Cmd,FOutputDevice&
 	}
 	else if( FParse::Command(&Cmd,TEXT("SHOT")) || FParse::Command(&Cmd,TEXT("SCREENSHOT")) )
 	{
-		return HandleScreenshotCommand( Cmd, Ar );	
+		return HandleScreenshotCommand( Cmd, Ar );
 	}
 	else if ( FParse::Command(&Cmd, TEXT("BUGSCREENSHOTWITHHUDINFO")) )
 	{
@@ -2673,7 +2668,7 @@ bool UGameViewportClient::Exec( UWorld* InWorld, const TCHAR* Cmd,FOutputDevice&
 	}
 	else if( FParse::Command(&Cmd,TEXT("KILLPARTICLES")) )
 	{
-		return HandleKillParticlesCommand( Cmd, Ar );	
+		return HandleKillParticlesCommand( Cmd, Ar );
 	}
 	else if( FParse::Command(&Cmd,TEXT("FORCESKELLOD")) )
 	{
@@ -2808,7 +2803,7 @@ bool UGameViewportClient::HandleShowCommand( const TCHAR* Cmd, FOutputDevice& Ar
 				{
 					VerifyPathRenderingComponents();
 				}
-					
+
 				if(FEngineShowFlags::IsNameThere(Cmd, TEXT("Volumes")))
 				{
 					// TODO: Investigate why this is doesn't appear to work
@@ -2855,7 +2850,7 @@ bool UGameViewportClient::HandleShowCommand( const TCHAR* Cmd, FOutputDevice& Ar
 		FIterSink Sink(LinesToSort, EngineShowFlags);
 
 		FEngineShowFlags::IterateAllFlags(Sink);
-	}	
+	}
 
 	LinesToSort.Sort( TLess<FString>() );
 
@@ -3000,7 +2995,7 @@ bool UGameViewportClient::HandleShowLayerCommand( const TCHAR* Cmd, FOutputDevic
 		for (FActorIterator It(InWorld); It; ++It)
 		{
 			AActor* Actor = *It;
-			
+
 			if (Actor->Layers.Contains(LayerFName))
 			{
 				// look for always toggle, or a set when it's unset, etc
@@ -3105,7 +3100,7 @@ bool UGameViewportClient::HandleViewModeCommand( const TCHAR* Cmd, FOutputDevice
 
 	if (FPlatformProperties::SupportsWindowedMode() == false)
 	{
-		if(ViewModeIndex == VMI_Unlit			
+		if(ViewModeIndex == VMI_Unlit
 			|| ViewModeIndex == VMI_StationaryLightOverlap
 			|| ViewModeIndex == VMI_Lit_DetailLighting
 			|| ViewModeIndex == VMI_ReflectionOverride)
@@ -3218,7 +3213,7 @@ bool UGameViewportClient::SetDisplayConfiguration(const FIntPoint* Dimensions, E
 	}
 
 	UGameEngine* GameEngine = Cast<UGameEngine>(GEngine);
-	
+
 	if (GameEngine)
 	{
 		UGameUserSettings* UserSettings = GameEngine->GetGameUserSettings();
@@ -3242,7 +3237,7 @@ bool UGameViewportClient::SetDisplayConfiguration(const FIntPoint* Dimensions, E
 			NewX = Dimensions->X;
 			NewY = Dimensions->Y;
 		}
-	
+
 		FSystemResolution::RequestResolutionChange(NewX, NewY, WindowMode);
 	}
 
@@ -3322,7 +3317,7 @@ bool UGameViewportClient::HandleSetResCommand( const TCHAR* Cmd, FOutputDevice&
 			{
 				WindowMode = EWindowMode::Windowed;
 			}
-			
+
 		}
 		else if(FCString::Strchr(Cmd,'f') || FCString::Strchr(Cmd,'F'))
 		{
@@ -3804,7 +3799,7 @@ bool UGameViewportClient::SetHardwareCursor(EMouseCursor::Type CursorShape, FNam
 	{
 		PlatformCursor->SetTypeShape(CursorShape, HardwareCursor);
 	}
-	
+
 	return true;
 }
 
diff --git a/Engine/Source/Runtime/Engine/Private/GameplayStatics.cpp b/Engine/Source/Runtime/Engine/Private/GameplayStatics.cpp
--- a/Engine/Source/Runtime/Engine/Private/GameplayStatics.cpp
+++ b/Engine/Source/Runtime/Engine/Private/GameplayStatics.cpp
@@ -1179,8 +1179,8 @@ void UGameplayStatics::PlaySound2D(const UObject* WorldContextObject, class USou
 		NewActiveSound.SetSound(Sound);
 		NewActiveSound.SetWorld(ThisWorld);
 
-		NewActiveSound.VolumeMultiplier = VolumeMultiplier;
-		NewActiveSound.PitchMultiplier = PitchMultiplier;
+		NewActiveSound.SetPitch(PitchMultiplier);
+		NewActiveSound.SetVolume(VolumeMultiplier);
 
 		NewActiveSound.RequestedStartTime = FMath::Max(0.f, StartTime);
 
@@ -1214,8 +1214,6 @@ UAudioComponent* UGameplayStatics::CreateSound2D(const UObject* WorldContextObje
 		return nullptr;
 	}
 
-	UAudioComponent* AudioComponent;
-
 	FAudioDevice::FCreateComponentParams Params = bPersistAcrossLevelTransition
 		? FAudioDevice::FCreateComponentParams(ThisWorld->GetAudioDevice())
 		: FAudioDevice::FCreateComponentParams(ThisWorld);
@@ -1225,7 +1223,7 @@ UAudioComponent* UGameplayStatics::CreateSound2D(const UObject* WorldContextObje
 		Params.ConcurrencySet.Add(ConcurrencySettings);
 	}
 
-	AudioComponent = FAudioDevice::CreateComponent(Sound, Params);
+	UAudioComponent* AudioComponent = FAudioDevice::CreateComponent(Sound, Params);
 	if (AudioComponent)
 	{
 		AudioComponent->SetVolumeMultiplier(VolumeMultiplier);
diff --git a/Engine/Source/Runtime/Engine/Private/SoundBase.cpp b/Engine/Source/Runtime/Engine/Private/SoundBase.cpp
--- a/Engine/Source/Runtime/Engine/Private/SoundBase.cpp
+++ b/Engine/Source/Runtime/Engine/Private/SoundBase.cpp
@@ -1,9 +1,11 @@
 // Copyright 1998-2019 Epic Games, Inc. All Rights Reserved.
-
 #include "Sound/SoundBase.h"
-#include "Sound/SoundSubmix.h"
-#include "Sound/AudioSettings.h"
+
 #include "EngineDefines.h"
+#include "IAudioExtensionPlugin.h"
+#include "Sound/AudioSettings.h"
+#include "Sound/SoundSubmix.h"
+
 
 USoundClass* USoundBase::DefaultSoundClassObject = nullptr;
 USoundConcurrency* USoundBase::DefaultSoundConcurrencyObject = nullptr;
@@ -54,9 +56,9 @@ bool USoundBase::IsPlayable() const
 	return false;
 }
 
-bool USoundBase::IsAllowedVirtual() const
-{ 
-	return false; 
+bool USoundBase::SupportsSubtitles() const
+{
+	return false;
 }
 
 bool USoundBase::HasAttenuationNode() const
@@ -102,9 +104,9 @@ bool USoundBase::HasConcatenatorNode() const
 	return bHasConcatenatorNode;
 }
 
-bool USoundBase::IsVirtualizeWhenSilent() const
+bool USoundBase::IsPlayWhenSilent() const
 {
-	return bHasVirtualizeWhenSilent;
+	return bHasPlayWhenSilent;
 }
 
 float USoundBase::GetVolumeMultiplier()
@@ -118,8 +120,8 @@ float USoundBase::GetPitchMultiplier()
 }
 
 bool USoundBase::IsLooping()
-{ 
-	return (GetDuration() >= INDEFINITELY_LOOPING_DURATION); 
+{
+	return (GetDuration() >= INDEFINITELY_LOOPING_DURATION);
 }
 
 bool USoundBase::ShouldApplyInteriorVolumes()
@@ -214,6 +216,15 @@ void USoundBase::Serialize(FArchive& Ar)
 {
 	Super::Serialize(Ar);
 
+	if (Ar.IsLoading() || Ar.IsSaving())
+	{
+		if (bHasVirtualizeWhenSilent_DEPRECATED)
+		{
+			bHasVirtualizeWhenSilent_DEPRECATED = 0;
+			bHasPlayWhenSilent = 1;
+		}
+	}
+
 #if WITH_EDITORONLY_DATA
 	if (Ar.IsLoading())
 	{
diff --git a/Engine/Source/Runtime/Engine/Private/SoundCue.cpp b/Engine/Source/Runtime/Engine/Private/SoundCue.cpp
--- a/Engine/Source/Runtime/Engine/Private/SoundCue.cpp
+++ b/Engine/Source/Runtime/Engine/Private/SoundCue.cpp
@@ -96,7 +96,7 @@ void USoundCue::CacheAggregateValues()
 
 		bHasDelayNode = FirstNode->HasDelayNode();
 		bHasConcatenatorNode = FirstNode->HasConcatenatorNode();
-		bHasVirtualizeWhenSilent = FirstNode->IsVirtualizeWhenSilent();
+		bHasPlayWhenSilent = FirstNode->IsPlayWhenSilent();
 	}
 }
 
@@ -552,7 +552,7 @@ bool USoundCue::HasCookedAmplitudeEnvelopeData() const
 
 #if WITH_EDITOR
 UEdGraph* USoundCue::GetGraph()
-{ 
+{
 	return SoundCueGraph;
 }
 
@@ -610,6 +610,6 @@ TSharedPtr<ISoundCueAudioEditor> USoundCue::GetSoundCueAudioEditor()
 	{
 	return SoundCueAudioEditor;
 }
-			
+
 
 #endif
diff --git a/Engine/Source/Runtime/Engine/Private/SoundNode.cpp b/Engine/Source/Runtime/Engine/Private/SoundNode.cpp
--- a/Engine/Source/Runtime/Engine/Private/SoundNode.cpp
+++ b/Engine/Source/Runtime/Engine/Private/SoundNode.cpp
@@ -212,14 +212,14 @@ bool USoundNode::HasConcatenatorNode() const
 	return false;
 }
 
-bool USoundNode::IsVirtualizeWhenSilent() const
+bool USoundNode::IsPlayWhenSilent() const
 {
 	for (USoundNode* ChildNode : ChildNodes)
 	{
 		if (ChildNode)
 		{
 			ChildNode->ConditionalPostLoad();
-			if (ChildNode->IsVirtualizeWhenSilent())
+			if (ChildNode->IsPlayWhenSilent())
 			{
 				return true;
 			}
diff --git a/Engine/Source/Runtime/Engine/Private/SoundNodeWavePlayer.cpp b/Engine/Source/Runtime/Engine/Private/SoundNodeWavePlayer.cpp
--- a/Engine/Source/Runtime/Engine/Private/SoundNodeWavePlayer.cpp
+++ b/Engine/Source/Runtime/Engine/Private/SoundNodeWavePlayer.cpp
@@ -34,7 +34,7 @@ void USoundNodeWavePlayer::LoadAsset(bool bAddToRoot)
 	{
 		SoundWave = SoundWaveAssetPtr.Get();
 		if (SoundWave && SoundWave->HasAnyFlags(RF_NeedLoad))
-		{			
+		{
 			// This can happen when the owning USoundCue's PostLoad gets called and the SoundWave hasn't been serialized yet
 			// In this case we need to make sure we don't pass the pointer to the SoundNodeWavePlayer too early as the SoundWave
 			// will be serialized on the AsyncLoadingThread shortly and this may lead to strange race conditions / thread safety issues
@@ -147,6 +147,10 @@ void USoundNodeWavePlayer::ParseNodes( FAudioDevice* AudioDevice, const UPTRINT
 float USoundNodeWavePlayer::GetDuration()
 {
 	float Duration = 0.f;
+	if (!SoundWave && !IsAsyncLoading())
+	{
+		SoundWave = SoundWaveAssetPtr.LoadSynchronous();
+	}
 	if (SoundWave)
 	{
 		if (bLooping)
@@ -161,11 +165,11 @@ float USoundNodeWavePlayer::GetDuration()
 	return Duration;
 }
 
-bool USoundNodeWavePlayer::IsVirtualizeWhenSilent() const
+bool USoundNodeWavePlayer::IsPlayWhenSilent() const
 {
 	if (SoundWave)
 	{
-		return SoundWave->bVirtualizeWhenSilent;
+		return SoundWave->bPlayWhenSilent;
 	}
 	return false;
 }
diff --git a/Engine/Source/Runtime/Engine/Private/SoundSourceBus.cpp b/Engine/Source/Runtime/Engine/Private/SoundSourceBus.cpp
--- a/Engine/Source/Runtime/Engine/Private/SoundSourceBus.cpp
+++ b/Engine/Source/Runtime/Engine/Private/SoundSourceBus.cpp
@@ -34,7 +34,7 @@ void USoundSourceBus::Init()
 
 	// Keep playing this bus when the volume is 0
 	// Note source buses can't ever be truly virtual as they are procedurally generated.
-	bVirtualizeWhenSilent = !bAutoDeactivateWhenSilent;
+	bPlayWhenSilent = !bAutoDeactivateWhenSilent;
 
 	// Set the channels equal to the users channel count choice
 	switch (SourceBusChannels)
diff --git a/Engine/Source/Runtime/Engine/Private/SoundWave.cpp b/Engine/Source/Runtime/Engine/Private/SoundWave.cpp
--- a/Engine/Source/Runtime/Engine/Private/SoundWave.cpp
+++ b/Engine/Source/Runtime/Engine/Private/SoundWave.cpp
@@ -31,10 +31,10 @@
 
 #include "Misc/CommandLine.h"
 
-static int32 BypassVirtualizeWhenSilentCVar = 0;
-FAutoConsoleVariableRef CVarBypassVirtualizeWhenSilent(
-	TEXT("au.BypassVirtualizeWhenSilent"),
-	BypassVirtualizeWhenSilentCVar,
+static int32 BypassPlayWhenSilentCVar = 0;
+FAutoConsoleVariableRef CVarBypassPlayWhenSilent(
+	TEXT("au.BypassPlayWhenSilent"),
+	BypassPlayWhenSilentCVar,
 	TEXT("When set to 1, ignores the Play When Silent flag for non-procedural sources.\n")
 	TEXT("0: Honor the Play When Silent flag, 1: stop all silent non-procedural sources."),
 	ECVF_Default);
@@ -300,7 +300,17 @@ void USoundWave::Serialize( FArchive& Ar )
 
 	if (Ar.IsSaving() || Ar.IsCooking())
 	{
-		bHasVirtualizeWhenSilent = bVirtualizeWhenSilent;
+		if (bHasVirtualizeWhenSilent_DEPRECATED)
+		{
+			bHasVirtualizeWhenSilent_DEPRECATED = 0;
+			bHasPlayWhenSilent = 1;
+		}
+
+		if (bVirtualizeWhenSilent_DEPRECATED)
+		{
+			bVirtualizeWhenSilent_DEPRECATED = 0;
+			bPlayWhenSilent = 1;
+		}
 
 #if WITH_ENGINE
 		// If there is an AutoStreamingThreshold set for the platform we're cooking to,
@@ -445,9 +455,9 @@ float USoundWave::GetSubtitlePriority() const
 	return SubtitlePriority;
 };
 
-bool USoundWave::IsAllowedVirtual() const
+bool USoundWave::SupportsSubtitles() const
 {
-	return bVirtualizeWhenSilent || (Subtitles.Num() > 0);
+	return bPlayWhenSilent || (Subtitles.Num() > 0);
 }
 
 void USoundWave::PostInitProperties()
@@ -693,7 +703,7 @@ void USoundWave::PostLoad()
 		return;
 	}
 
-	bHasVirtualizeWhenSilent = bVirtualizeWhenSilent;
+	bHasPlayWhenSilent = bPlayWhenSilent;
 
 #if WITH_EDITORONLY_DATA
 	// Log a warning after loading if the source has effect chains but has channels greater than 2.
@@ -1343,10 +1353,9 @@ void USoundWave::FreeResources()
 	DecompressionType = DTYPE_Setup;
 	bDecompressedFromOgg = false;
 
-	USoundWave* SoundWave = this;
-	if (SoundWave->ResourceState == ESoundWaveResourceState::Freeing)
+	if (ResourceState == ESoundWaveResourceState::Freeing)
 	{
-		SoundWave->ResourceState = ESoundWaveResourceState::Freed;
+		ResourceState = ESoundWaveResourceState::Freed;
 	}
 }
 
@@ -1449,19 +1458,19 @@ void USoundWave::FinishDestroy()
 	IStreamingManager::Get().GetAudioStreamingManager().RemoveStreamingSoundWave(this);
 }
 
-void USoundWave::Parse( FAudioDevice* AudioDevice, const UPTRINT NodeWaveInstanceHash, FActiveSound& ActiveSound, const FSoundParseParameters& ParseParams, TArray<FWaveInstance*>& WaveInstances )
+void USoundWave::Parse(FAudioDevice* AudioDevice, const UPTRINT NodeWaveInstanceHash, FActiveSound& ActiveSound, const FSoundParseParameters& ParseParams, TArray<FWaveInstance*>& WaveInstances)
 {
 	FWaveInstance* WaveInstance = ActiveSound.FindWaveInstance(NodeWaveInstanceHash);
 
 	// Create a new WaveInstance if this SoundWave doesn't already have one associated with it.
-	if( WaveInstance == NULL )
+	if (WaveInstance == NULL)
 	{
-		if( !ActiveSound.bRadioFilterSelected )
+		if (!ActiveSound.bRadioFilterSelected)
 		{
 			ActiveSound.ApplyRadioFilter(ParseParams);
 		}
 
-		WaveInstance = HandleStart( ActiveSound, NodeWaveInstanceHash);
+		WaveInstance = HandleStart(ActiveSound, NodeWaveInstanceHash);
 	}
 
 	// Looping sounds are never actually finished
@@ -1478,14 +1487,13 @@ void USoundWave::Parse( FAudioDevice* AudioDevice, const UPTRINT NodeWaveInstanc
 	}
 
 	// Check for finished paths.
-	if( !WaveInstance->bIsFinished )
+	if (!WaveInstance->bIsFinished)
 	{
 		// Propagate properties and add WaveInstance to outgoing array of FWaveInstances.
 		WaveInstance->SetVolume(ParseParams.Volume * Volume);
 		WaveInstance->SetVolumeMultiplier(ParseParams.VolumeMultiplier);
 		WaveInstance->SetDistanceAttenuation(ParseParams.DistanceAttenuation);
-		WaveInstance->SetVolumeApp(ParseParams.VolumeApp);
-		WaveInstance->Pitch = ParseParams.Pitch * Pitch;
+		WaveInstance->SetPitch(ParseParams.Pitch * Pitch);
 		WaveInstance->bEnableLowPassFilter = ParseParams.bEnableLowPassFilter;
 		WaveInstance->bIsOccluded = ParseParams.bIsOccluded;
 		WaveInstance->LowPassFilterFrequency = ParseParams.LowPassFilterFrequency;
@@ -1518,7 +1526,7 @@ void USoundWave::Parse( FAudioDevice* AudioDevice, const UPTRINT NodeWaveInstanc
 			// Use values from "parsed/ propagated" sound class properties
 			float VolumeMultiplier = WaveInstance->GetVolumeMultiplier();
 			WaveInstance->SetVolumeMultiplier(VolumeMultiplier* SoundClassProperties->Volume);
-			WaveInstance->Pitch *= SoundClassProperties->Pitch;
+			WaveInstance->SetPitch(WaveInstance->Pitch * SoundClassProperties->Pitch);
 			//TODO: Add in HighFrequencyGainMultiplier property to sound classes
 
 			WaveInstance->VoiceCenterChannelVolume = SoundClassProperties->VoiceCenterChannelVolume;
@@ -1626,18 +1634,21 @@ void USoundWave::Parse( FAudioDevice* AudioDevice, const UPTRINT NodeWaveInstanc
 		WaveInstance->SpatializationPluginSettings = ParseParams.SpatializationPluginSettings;
 		WaveInstance->OcclusionPluginSettings = ParseParams.OcclusionPluginSettings;
 		WaveInstance->ReverbPluginSettings = ParseParams.ReverbPluginSettings;
+		WaveInstance->ModulationPluginSettings = ParseParams.ModulationPluginSettings;
 
 		WaveInstance->bIsAmbisonics = bIsAmbisonics;
 
 		bool bAddedWaveInstance = false;
 
-		// Recompute the virtualizability here even though we did it up-front in the active sound parse.
-		// This is because an active sound can generate multiple sound waves, not all of them are necessarily virtualizable.
+		// Determine if can play when silent here even though we did it up-front in the active sound parse.
+		// This is because an active sound can generate multiple sound waves, not all of them are necessarily permitted
+		// to play when silent.
 		bool bHasSubtitles = ActiveSound.bHandleSubtitles && (ActiveSound.bHasExternalSubtitles || (Subtitles.Num() > 0));
 
-		// When the BypassVirtualizeWhenSilent cvar is enabled, we should only honor bVirtualizeWhenSilent for procedural sounds:
-		const bool bShouldVirtualize = bVirtualizeWhenSilent && (!BypassVirtualizeWhenSilentCVar || bProcedural);
-		if (WaveInstance->GetVolumeWithDistanceAttenuation() > KINDA_SMALL_NUMBER || ((bShouldVirtualize || bHasSubtitles) && AudioDevice->VirtualSoundsEnabled()))
+		// When the BypassPlayWhenSilent cvar is enabled, we should only honor bPlayWhenSilent for procedural sounds:
+		const bool bCanPlayWhenSilent = bPlayWhenSilent && (!BypassPlayWhenSilentCVar || bProcedural);
+		const float WaveInstanceVolume = WaveInstance->GetVolumeWithDistanceAttenuation() * WaveInstance->GetDynamicVolume();
+		if (WaveInstanceVolume > KINDA_SMALL_NUMBER || ((bCanPlayWhenSilent || bHasSubtitles) && AudioDevice->PlayWhenSilentEnabled()))
 		{
 			bAddedWaveInstance = true;
 			WaveInstances.Add(WaveInstance);
@@ -1650,7 +1661,7 @@ void USoundWave::Parse( FAudioDevice* AudioDevice, const UPTRINT NodeWaveInstanc
 		}
 
 		// Sanity check
-		if( NumChannels > 2 && WaveInstance->GetUseSpatialization() && !WaveInstance->bReportedSpatializationWarning)
+		if(NumChannels > 2 && WaveInstance->GetUseSpatialization() && !WaveInstance->bReportedSpatializationWarning)
 		{
 			static TSet<USoundWave*> ReportedSounds;
 			if (!ReportedSounds.Contains(this))
@@ -1670,17 +1681,17 @@ void USoundWave::Parse( FAudioDevice* AudioDevice, const UPTRINT NodeWaveInstanc
 						if (UAudioComponent* AudioComponent = UAudioComponent::GetAudioComponentFromID(AudioComponentID))
 						{
 							AActor* SoundOwner = AudioComponent->GetOwner();
-							UE_LOG(LogAudio, Warning, TEXT( "%s Actor: %s AudioComponent: %s" ), *SoundWarningInfo, (SoundOwner ? *SoundOwner->GetName() : TEXT("None")), *AudioComponent->GetName() );
+							UE_LOG(LogAudio, Warning, TEXT("%s Actor: %s AudioComponent: %s"), *SoundWarningInfo, (SoundOwner ? *SoundOwner->GetName() : TEXT("None")), *AudioComponent->GetName());
 						}
 						else
 						{
-							UE_LOG(LogAudio, Warning, TEXT("%s"), *SoundWarningInfo );
+							UE_LOG(LogAudio, Warning, TEXT("%s"), *SoundWarningInfo);
 						}
 					});
 				}
 				else
 				{
-					UE_LOG(LogAudio, Warning, TEXT("%s"), *SoundWarningInfo );
+					UE_LOG(LogAudio, Warning, TEXT("%s"), *SoundWarningInfo);
 				}
 #endif
 
diff --git a/Engine/Source/Runtime/Engine/Private/UnrealEngine.cpp b/Engine/Source/Runtime/Engine/Private/UnrealEngine.cpp
--- a/Engine/Source/Runtime/Engine/Private/UnrealEngine.cpp
+++ b/Engine/Source/Runtime/Engine/Private/UnrealEngine.cpp
@@ -111,10 +111,10 @@ UnrealEngine.cpp: Implements the UEngine class and helpers.
 #include "IXRTrackingSystem.h"
 #include "Stats/StatsData.h"
 #include "Stats/StatsFile.h"
-#include "AudioThread.h"
+#include "Audio/AudioDebug.h"
 #include "AudioDeviceManager.h"
-#include "Sound/ReverbEffect.h"
 #include "AudioDevice.h"
+#include "AudioThread.h"
 #include "Animation/SkeletalMeshActor.h"
 #include "Animation/AnimSequence.h"
 #include "Animation/AnimCompress.h"
@@ -1599,11 +1599,13 @@ void UEngine::Init(IEngineLoop* InEngineLoop)
 	EngineStats.Add(FEngineStatFuncs(TEXT("STAT_ColorList"), TEXT("STATCAT_Engine"), FText::GetEmpty(), &UEngine::RenderStatColorList, NULL));
 	EngineStats.Add(FEngineStatFuncs(TEXT("STAT_Levels"), TEXT("STATCAT_Engine"), FText::GetEmpty(), &UEngine::RenderStatLevels, NULL));
 #if !UE_BUILD_SHIPPING
+	EngineStats.Add(FEngineStatFuncs(TEXT("STAT_Sounds"), TEXT("STATCAT_Engine"), FText::GetEmpty(), &UEngine::RenderStatSounds, &UEngine::ToggleStatSounds));
+	EngineStats.Add(FEngineStatFuncs(TEXT("STAT_SoundCues"), TEXT("STATCAT_Engine"), FText::GetEmpty(), &UEngine::RenderStatSoundCues, &UEngine::ToggleStatSoundCues));
 	EngineStats.Add(FEngineStatFuncs(TEXT("STAT_SoundMixes"), TEXT("STATCAT_Engine"), FText::GetEmpty(), &UEngine::RenderStatSoundMixes, &UEngine::ToggleStatSoundMixes));
-	EngineStats.Add(FEngineStatFuncs(TEXT("STAT_Reverb"), TEXT("STATCAT_Engine"), FText::GetEmpty(), &UEngine::RenderStatReverb, NULL));
+	EngineStats.Add(FEngineStatFuncs(TEXT("STAT_SoundModulators"), TEXT("STATCAT_Engine"), FText::GetEmpty(), &UEngine::RenderStatSoundModulators, &UEngine::ToggleStatSoundModulators));
+	EngineStats.Add(FEngineStatFuncs(TEXT("STAT_SoundModulatorsHelp"), TEXT("STATCAT_Engine"), FText::GetEmpty(), nullptr, &UEngine::PostStatSoundModulatorHelp));
+	EngineStats.Add(FEngineStatFuncs(TEXT("STAT_SoundReverb"), TEXT("STATCAT_Engine"), FText::GetEmpty(), &UEngine::RenderStatSoundReverb, nullptr));
 	EngineStats.Add(FEngineStatFuncs(TEXT("STAT_SoundWaves"), TEXT("STATCAT_Engine"), FText::GetEmpty(), &UEngine::RenderStatSoundWaves, &UEngine::ToggleStatSoundWaves));
-	EngineStats.Add(FEngineStatFuncs(TEXT("STAT_SoundCues"), TEXT("STATCAT_Engine"), FText::GetEmpty(), &UEngine::RenderStatSoundCues, &UEngine::ToggleStatSoundCues));
-	EngineStats.Add(FEngineStatFuncs(TEXT("STAT_Sounds"), TEXT("STATCAT_Engine"), FText::GetEmpty(), &UEngine::RenderStatSounds, &UEngine::ToggleStatSounds));
 #endif // !UE_BUILD_SHIPPING
 	/* @todo UE4 physx fix this once we have convexelem drawing again
 	EngineStats.Add(FEngineStatFuncs(TEXT("STAT_LevelMap"), TEXT("STATCAT_Engine"), FText::GetEmpty(), &UEngine::RenderStatLevelMap, NULL));
@@ -10235,18 +10237,6 @@ DEFINE_STAT(STAT_UnitRender);
 DEFINE_STAT(STAT_UnitRHIT);
 DEFINE_STAT(STAT_UnitGPU);
 
-/*-----------------------------------------------------------------------------
-Lightmass object/actor implementations.
------------------------------------------------------------------------------*/
-
-
-
-
-
-
-/*-----------------------------------------------------------------------------
-ULightmappedSurfaceCollection
------------------------------------------------------------------------------*/
 
 UFont* GetStatsFont()
 {
@@ -15451,680 +15441,123 @@ static FAutoConsoleCommand PakFileTestCmd(
 	FConsoleCommandWithArgsDelegate::CreateStatic(&PakFileTest)
 );
 
-#endif
-
 // REVERB
-#if !UE_BUILD_SHIPPING
-int32 UEngine::RenderStatReverb(UWorld* World, FViewport* Viewport, FCanvas* Canvas, int32 X, int32 Y, const FVector* ViewLocation, const FRotator* ViewRotation)
+int32 UEngine::RenderStatSoundReverb(UWorld* World, FViewport* Viewport, FCanvas* Canvas, int32 X, int32 Y, const FVector* ViewLocation, const FRotator* ViewRotation)
 {
-	if (FAudioDevice* AudioDevice = World->GetAudioDevice())
-	{
-		AudioDevice->RenderStatReverb(World, Viewport, Canvas, X, Y, ViewLocation, ViewRotation);
-	}
-
+#if ENABLE_AUDIO_DEBUG
+	return FAudioDebugger::RenderStatReverb(World, Viewport, Canvas, X, Y, ViewLocation, ViewRotation);
+#else // !ENABLE_AUDIO_DEBUG
 	return Y;
-}
-
-void FAudioDevice::RenderStatReverb(UWorld* World, FViewport* Viewport, FCanvas* Canvas, int32 X, int32& Y, const FVector* ViewLocation, const FRotator* ViewRotation) const
-{
-	UReverbEffect* ReverbEffect = GetCurrentReverbEffect();
-	FString TheString;
-	if (ReverbEffect)
-	{
-		TheString = FString::Printf(TEXT("Active Reverb Effect: %s"), *ReverbEffect->GetName());
-		Canvas->DrawShadowedString(X, Y, *TheString, UEngine::GetSmallFont(), FLinearColor::White);
-		Y += 12;
-
-		AAudioVolume* CurrentAudioVolume = nullptr;
-		for (const FTransform& Transform : ListenerTransforms)
-		{
-			AAudioVolume* PlayerAudioVolume = World->GetAudioSettings(Transform.GetLocation(), nullptr, nullptr);
-			if (PlayerAudioVolume && ((CurrentAudioVolume == nullptr) || (PlayerAudioVolume->GetPriority() > CurrentAudioVolume->GetPriority())))
-			{
-				CurrentAudioVolume = PlayerAudioVolume;
-			}
-		}
-		if (CurrentAudioVolume && CurrentAudioVolume->GetReverbSettings().ReverbEffect)
-		{
-			TheString = FString::Printf(TEXT("  Audio Volume Reverb Effect: %s (Priority: %g Volume Name: %s)"), *CurrentAudioVolume->GetReverbSettings().ReverbEffect->GetName(), CurrentAudioVolume->GetPriority(), *CurrentAudioVolume->GetName());
-		}
-		else
-		{
-			TheString = TEXT("  Audio Volume Reverb Effect: None");
-		}
-		Canvas->DrawShadowedString(X, Y, *TheString, UEngine::GetSmallFont(), FLinearColor::White);
-		Y += 12;
-		if (ActivatedReverbs.Num() == 0)
-		{
-			TheString = TEXT("  Activated Reverb: None");
-			Canvas->DrawShadowedString(X, Y, *TheString, UEngine::GetSmallFont(), FLinearColor::White);
-			Y += 12;
-		}
-		else if (ActivatedReverbs.Num() == 1)
-		{
-			auto It = ActivatedReverbs.CreateConstIterator();
-			TheString = FString::Printf(TEXT("  Activated Reverb Effect: %s (Priority: %g Tag: '%s')"), *It.Value().ReverbSettings.ReverbEffect->GetName(), It.Value().Priority, *It.Key().ToString());
-			Canvas->DrawShadowedString(X, Y, *TheString, UEngine::GetSmallFont(), FLinearColor::White);
-			Y += 12;
-		}
-		else
-		{
-			Canvas->DrawShadowedString(X, Y, TEXT("  Activated Reverb Effects:"), UEngine::GetSmallFont(), FLinearColor::White);
-			Y += 12;
-			TMap<int32, FString> PrioritySortedActivatedReverbs;
-			for (auto It = ActivatedReverbs.CreateConstIterator(); It; ++It)
-			{
-				TheString = FString::Printf(TEXT("    %s (Priority: %g Tag: '%s')"), *It.Value().ReverbSettings.ReverbEffect->GetName(), It.Value().Priority, *It.Key().ToString());
-				PrioritySortedActivatedReverbs.Add(It.Value().Priority, TheString);
-			}
-			for (auto It = PrioritySortedActivatedReverbs.CreateConstIterator(); It; ++It)
-			{
-				Canvas->DrawShadowedString(X, Y, *It.Value(), UEngine::GetSmallFont(), FLinearColor::White);
-				Y += 12;
-			}
-		}
-	}
-	else
-	{
-		TheString = TEXT("Active Reverb Effect: None");
-		Canvas->DrawShadowedString(X, Y, *TheString, UEngine::GetSmallFont(), FLinearColor::White);
-		Y += 12;
-	}
+#endif // !ENABLE_AUDIO_DEBUG
 }
 
 // SOUNDMIXES
 int32 UEngine::RenderStatSoundMixes(UWorld* World, FViewport* Viewport, FCanvas* Canvas, int32 X, int32 Y, const FVector* ViewLocation, const FRotator* ViewRotation)
 {
-	Canvas->DrawShadowedString(X, Y, TEXT("Active Sound Mixes:"), GetSmallFont(), FColor::Green);
-	Y += 12;
-
-	bool bDisplayedSoundMixes = false;
-
-	if (FAudioDevice* AudioDevice = World->GetAudioDevice())
-	{
-		const FAudioStats& AudioStats = AudioDevice->GetAudioStats();
-
-		if (!AudioStats.bStale)
-		{
-			if (AudioStats.StatSoundMixes.Num() > 0)
-			{
-				bDisplayedSoundMixes = true;
-
-				for (const FAudioStats::FStatSoundMix& StatSoundMix : AudioStats.StatSoundMixes)
-				{
-					const FString TheString = FString::Printf(TEXT("%s - Fade Proportion: %1.2f - Total Ref Count: %i"), *StatSoundMix.MixName, StatSoundMix.InterpValue, StatSoundMix.RefCount);
-
-					const FColor& TextColour = (StatSoundMix.bIsCurrentEQ ? FColor::Yellow : FColor::White);
-
-					Canvas->DrawShadowedString(X + 12, Y, *TheString, GetSmallFont(), TextColour);
-					Y += 12;
-				}
-			}
-		}
-	}
-
-	if (!bDisplayedSoundMixes)
-	{
-		Canvas->DrawShadowedString(X + 12, Y, TEXT("None"), GetSmallFont(), FColor::White);
-		Y += 12;
-	}
+#if ENABLE_AUDIO_DEBUG
+	return FAudioDebugger::RenderStatMixes(World, Viewport, Canvas, X, Y, ViewLocation, ViewRotation);
+#else // !ENABLE_AUDIO_DEBUG
 	return Y;
+#endif // !ENABLE_AUDIO_DEBUG
 }
 
-void FAudioDevice::UpdateSoundShowFlags(const uint8 OldSoundShowFlags, const uint8 NewSoundShowFlags)
+// SOUNDMODULATION
+int32 UEngine::RenderStatSoundModulators(UWorld* World, FViewport* Viewport, FCanvas* Canvas, int32 X, int32 Y, const FVector* ViewLocation, const FRotator* ViewRotation)
 {
-	if (NewSoundShowFlags != OldSoundShowFlags)
-	{
-		uint8 RequestedStatChange = 0;
-		if ((NewSoundShowFlags == FViewportClient::ESoundShowFlags::Disabled) || (OldSoundShowFlags == FViewportClient::ESoundShowFlags::Disabled))
-		{
-			RequestedStatChange |= ERequestedAudioStats::Sounds;
-		}
-		if ((NewSoundShowFlags ^ OldSoundShowFlags) & FViewportClient::ESoundShowFlags::Debug)
-		{
-			RequestedStatChange |= ERequestedAudioStats::DebugSounds;
-		}
-		if ((NewSoundShowFlags ^ OldSoundShowFlags) & FViewportClient::ESoundShowFlags::Long_Names)
-		{
-			RequestedStatChange |= ERequestedAudioStats::LongSoundNames;
-		}
-		if (RequestedStatChange != 0)
-		{
-			UpdateRequestedStat(RequestedStatChange);
-		}
-	}
+#if ENABLE_AUDIO_DEBUG
+	return FAudioDebugger::RenderStatModulators(World, Viewport, Canvas, X, Y, ViewLocation, ViewRotation);
+#else // !ENABLE_AUDIO_DEBUG
+	return Y;
+#endif // !ENABLE_AUDIO_DEBUG
 }
 
-void FAudioDevice::ResolveDesiredStats(FViewportClient* ViewportClient)
+// SOUNDWAVES
+bool UEngine::ToggleStatSoundWaves(UWorld* World, FCommonViewportClient* ViewportClient, const TCHAR* Stream)
 {
-	check(IsInGameThread());
-
-	uint8 SetStats = 0;
-	uint8 ClearStats = 0;
-
-	if (ViewportClient->IsStatEnabled(TEXT("SoundCues")))
-	{
-		SetStats |= ERequestedAudioStats::SoundCues;
-	}
-	else
-	{
-		ClearStats |= ERequestedAudioStats::SoundCues;
-	}
-
-	if (ViewportClient->IsStatEnabled(TEXT("SoundWaves")))
-	{
-		SetStats |= ERequestedAudioStats::SoundWaves;
-	}
-	else
-	{
-		ClearStats |= ERequestedAudioStats::SoundWaves;
-	}
-
-	if (ViewportClient->IsStatEnabled(TEXT("SoundMixes")))
-	{
-		SetStats |= ERequestedAudioStats::SoundMixes;
-	}
-	else
-	{
-		ClearStats |= ERequestedAudioStats::SoundMixes;
-	}
-
-	if (ViewportClient->IsStatEnabled(TEXT("Sounds")))
-	{
-		const FViewportClient::ESoundShowFlags::Type SoundShowFlags = ViewportClient->GetSoundShowFlags();
-		SetStats |= ERequestedAudioStats::Sounds;
-
-		if (SoundShowFlags & FViewportClient::ESoundShowFlags::Debug)
-		{
-			SetStats |= ERequestedAudioStats::DebugSounds;
-		}
-		else
-		{
-			ClearStats |= ERequestedAudioStats::DebugSounds;
-		}
-
-		if (SoundShowFlags & FViewportClient::ESoundShowFlags::Long_Names)
-		{
-			SetStats |= ERequestedAudioStats::LongSoundNames;
-		}
-		else
-		{
-			ClearStats |= ERequestedAudioStats::LongSoundNames;
-		}
-	}
-	else
-	{
-		ClearStats |= ERequestedAudioStats::Sounds;
-		ClearStats |= ERequestedAudioStats::DebugSounds;
-		ClearStats |= ERequestedAudioStats::LongSoundNames;
-	}
-
-	DECLARE_CYCLE_STAT(TEXT("FAudioThreadTask.ResolveDesiredStats"), STAT_AudioResolveDesiredStats, STATGROUP_TaskGraphTasks);
-
-	FAudioDevice* AudioDevice = this;
-	FAudioThread::RunCommandOnAudioThread([AudioDevice, SetStats, ClearStats]()
-	{
-		AudioDevice->RequestedAudioStats |= SetStats;
-		AudioDevice->RequestedAudioStats &= ~ClearStats;
-	}, GET_STATID(STAT_AudioResolveDesiredStats));
+#if ENABLE_AUDIO_DEBUG
+	return FAudioDebugger::ToggleStatWaves(World, ViewportClient, Stream);
+#else // !ENABLE_AUDIO_DEBUG
+	return false;
+#endif // !ENABLE_AUDIO_DEBUG
 }
 
-void FAudioDevice::UpdateRequestedStat(const uint8 RequestedStat)
+// SOUNDCUES
+bool UEngine::ToggleStatSoundCues(UWorld* World, FCommonViewportClient* ViewportClient, const TCHAR* Stream)
 {
-	if (!IsInAudioThread())
-	{
-		DECLARE_CYCLE_STAT(TEXT("FAudioThreadTask.UpdateRequestedStat"), STAT_AudioUpdateRequestedStat, STATGROUP_TaskGraphTasks);
-
-		FAudioDevice* AudioDevice = this;
-		FAudioThread::RunCommandOnAudioThread([AudioDevice, RequestedStat]()
-		{
-			AudioDevice->UpdateRequestedStat(RequestedStat);
-		}, GET_STATID(STAT_AudioUpdateRequestedStat));
-		return;
-	}
-
-	RequestedAudioStats ^= RequestedStat;
+#if ENABLE_AUDIO_DEBUG
+	return FAudioDebugger::ToggleStatCues(World, ViewportClient, Stream);
+#else // !ENABLE_AUDIO_DEBUG
+	return false;
+#endif // !ENABLE_AUDIO_DEBUG
 }
 
-bool UEngine::ToggleStatSoundWaves(UWorld* World, FCommonViewportClient* ViewportClient, const TCHAR* Stream)
+// SOUNDMIXES
+bool UEngine::ToggleStatSoundMixes(UWorld* World, FCommonViewportClient* ViewportClient, const TCHAR* Stream)
 {
-	if (AudioDeviceManager)
-	{
-		AudioDeviceManager->ToggleDebugStat(ERequestedAudioStats::SoundWaves);
-	}
-	return true;
+#if ENABLE_AUDIO_DEBUG
+	return FAudioDebugger::ToggleStatMixes(World, ViewportClient, Stream);
+#else // !ENABLE_AUDIO_DEBUG
+	return false;
+#endif // !ENABLE_AUDIO_DEBUG
 }
 
-bool UEngine::ToggleStatSoundCues(UWorld* World, FCommonViewportClient* ViewportClient, const TCHAR* Stream)
+// SOUNDMODULATORS
+bool UEngine::PostStatSoundModulatorHelp(UWorld* World, FCommonViewportClient* ViewportClient, const TCHAR* Stream)
 {
-	if (AudioDeviceManager)
-	{
-		AudioDeviceManager->ToggleDebugStat(ERequestedAudioStats::SoundCues);
-	}
-	return true;
+#if ENABLE_AUDIO_DEBUG
+	return FAudioDebugger::PostStatModulatorHelp(World, ViewportClient, Stream);
+#else // !ENABLE_AUDIO_DEBUG
+	return false;
+#endif // !ENABLE_AUDIO_DEBUG
 }
 
-bool UEngine::ToggleStatSoundMixes(UWorld* World, FCommonViewportClient* ViewportClient, const TCHAR* Stream)
+// SOUNDMODULATORS
+bool UEngine::ToggleStatSoundModulators(UWorld* World, FCommonViewportClient* ViewportClient, const TCHAR* Stream)
 {
-	if (AudioDeviceManager)
-	{
-		AudioDeviceManager->ToggleDebugStat(ERequestedAudioStats::SoundMixes);
-	}
-	return true;
+#if ENABLE_AUDIO_DEBUG
+	return FAudioDebugger::ToggleStatModulators(World, ViewportClient, Stream);
+#else // !ENABLE_AUDIO_DEBUG
+	return false;
+#endif // !ENABLE_AUDIO_DEBUG
 }
 
 // SOUNDWAVES
 int32 UEngine::RenderStatSoundWaves(UWorld* World, FViewport* Viewport, FCanvas* Canvas, int32 X, int32 Y, const FVector* ViewLocation, const FRotator* ViewRotation)
 {
-	if (FAudioDevice* AudioDevice = World->GetAudioDevice())
-	{
-		const FAudioStats& AudioStats = AudioDevice->GetAudioStats();
-
-		if (!AudioStats.bStale)
-		{
-			Canvas->DrawShadowedString(X, Y, TEXT("Active Sound Waves:"), GetSmallFont(), FLinearColor::White);
-			Y += 12;
-
-			typedef TPair<const FAudioStats::FStatWaveInstanceInfo*, const FAudioStats::FStatSoundInfo*> FWaveInstancePair;
-
-			TArray<FWaveInstancePair> WaveInstances;
-
-			for (const FAudioStats::FStatSoundInfo& StatSoundInfo : AudioStats.StatSoundInfos)
-			{
-				for (const FAudioStats::FStatWaveInstanceInfo& WaveInstanceInfo : StatSoundInfo.WaveInstanceInfos)
-				{
-					if (WaveInstanceInfo.ActualVolume >= 0.01f)
-					{
-						bool bShouldPrint = true;
-						FString WaveInstanceName = WaveInstanceInfo.WaveInstanceName.ToString();
-						const FString& DebugSoloSoundName = GEngine->GetAudioDeviceManager()->GetDebugSoloSoundWave();
-						if (DebugSoloSoundName != TEXT("") && !WaveInstanceName.Contains(DebugSoloSoundName))
-						{
-							bShouldPrint = false;
-						}
-
-						if (bShouldPrint == true)
-						{
-							WaveInstances.Emplace(&WaveInstanceInfo, &StatSoundInfo);
-						}
-					}
-				}
-			}
-
-			WaveInstances.Sort([](const FWaveInstancePair& A, const FWaveInstancePair& B) { return A.Key->InstanceIndex < B.Key->InstanceIndex; });
-
-			for (const FWaveInstancePair& WaveInstanceInfo : WaveInstances)
-			{
-				UAudioComponent* AudioComponent = UAudioComponent::GetAudioComponentFromID(WaveInstanceInfo.Value->AudioComponentID);
-				AActor* SoundOwner = AudioComponent ? AudioComponent->GetOwner() : nullptr;
-
-				FString TheString = *FString::Printf(TEXT("%4i.    %6.2f  %s   Owner: %s   SoundClass: %s"),
-					WaveInstanceInfo.Key->InstanceIndex,
-					WaveInstanceInfo.Key->ActualVolume,
-					*WaveInstanceInfo.Key->WaveInstanceName.ToString(),
-					SoundOwner ? *SoundOwner->GetName() : TEXT("None"),
-					*WaveInstanceInfo.Value->SoundClassName.ToString());
-
-				Canvas->DrawShadowedString(X, Y, *TheString, GetSmallFont(), FColor::White);
-				Y += 12;
-			}
-
-			const int32 ActiveInstances = WaveInstances.Num();
-
-			const int32 Max = AudioDevice->MaxChannels / 2;
-			float f = FMath::Clamp<float>((float)(ActiveInstances - Max) / (float)Max, 0.f, 1.f);
-			const int32 R = FMath::TruncToInt(f * 255);
-
-			if (ActiveInstances > Max)
-			{
-				f = FMath::Clamp<float>((float)(Max - ActiveInstances) / (float)Max, 0.5f, 1.f);
-			}
-			else
-			{
-				f = 1.0f;
-			}
-			const int32 G = FMath::TruncToInt(f * 255);
-			const int32 B = 0;
-
-			Canvas->DrawShadowedString(X, Y, *FString::Printf(TEXT(" Total: %i"), ActiveInstances), GetSmallFont(), FColor(R, G, B));
-			Y += 12;
-		}
-	}
-	else
-	{
-		Canvas->DrawShadowedString(X, Y, TEXT("Active Sound Waves:"), GetSmallFont(), FLinearColor::White);
-		Y += 12;
-
-		Canvas->DrawShadowedString(X, Y, TEXT(" Total: 0"), GetSmallFont(), FLinearColor::White);
-		Y += 12;
-	}
+#if ENABLE_AUDIO_DEBUG
+	return FAudioDebugger::RenderStatWaves(World, Viewport, Canvas, X, Y, ViewLocation, ViewRotation);
+#else // !ENABLE_AUDIO_DEBUG
 	return Y;
+#endif // !ENABLE_AUDIO_DEBUG
 }
 
 // SOUNDCUES
 int32 UEngine::RenderStatSoundCues(UWorld* World, FViewport* Viewport, FCanvas* Canvas, int32 X, int32 Y, const FVector* ViewLocation, const FRotator* ViewRotation)
 {
-	Canvas->DrawShadowedString(X, Y, TEXT("Active Sound Cues:"), GetSmallFont(), FColor::Green);
-	Y += 12;
-
-	int32 ActiveSoundCount = 0;
-
-	if (FAudioDevice* AudioDevice = World->GetAudioDevice())
-	{
-		const FAudioStats& AudioStats = AudioDevice->GetAudioStats();
-
-		for (const FAudioStats::FStatSoundInfo& StatSoundInfo : AudioDevice->GetAudioStats().StatSoundInfos)
-		{
-			for (const FAudioStats::FStatWaveInstanceInfo& WaveInstanceInfo : StatSoundInfo.WaveInstanceInfos)
-			{
-				if (WaveInstanceInfo.ActualVolume >= 0.01f)
-				{
-					bool bShouldPrint = true;
-					FString WaveInstanceName = WaveInstanceInfo.WaveInstanceName.ToString();
-					const FString& DebugSoloSoundCue = GEngine->GetAudioDeviceManager()->GetDebugSoloSoundCue();
-					if (DebugSoloSoundCue != TEXT("") && !WaveInstanceName.Contains(DebugSoloSoundCue))
-					{
-						bShouldPrint = false;
-					}
-
-					if (bShouldPrint == true)
-					{
-						const FString TheString = FString::Printf(TEXT("%4i. %s %s"), ActiveSoundCount++, *StatSoundInfo.SoundName, *StatSoundInfo.SoundClassName.ToString());
-						Canvas->DrawShadowedString(X, Y, *TheString, GetSmallFont(), FColor::White);
-						Y += 12;
-						break;
-					}
-				}
-			}
-		}
-	}
-
-	Canvas->DrawShadowedString(X, Y, *FString::Printf(TEXT("Total: %i"), ActiveSoundCount), GetSmallFont(), FColor::Green);
-	Y += 12;
+#if ENABLE_AUDIO_DEBUG
+	return FAudioDebugger::RenderStatCues(World, Viewport, Canvas, X, Y, ViewLocation, ViewRotation);
+#else // !ENABLE_AUDIO_DEBUG
 	return Y;
+#endif // !ENABLE_AUDIO_DEBUG
 }
 
 // SOUNDS
 bool UEngine::ToggleStatSounds(UWorld* World, FCommonViewportClient* ViewportClient, const TCHAR* Stream)
 {
-	if (ViewportClient == nullptr)
-	{
-		// Ignore if all Viewports are closed.
-		return false;
-	}
-	const bool bHelp = Stream ? FCString::Stristr(Stream, TEXT("?")) != NULL : false;
-	if (bHelp)
-	{
-		GLog->Logf(TEXT("stat sounds description"));
-		GLog->Logf(TEXT("  stat sounds off - Disables drawing stat sounds"));
-		GLog->Logf(TEXT("  stat sounds sort=distance|class|name|waves|default"));
-		GLog->Logf(TEXT("      distance - sort list by distance to player"));
-		GLog->Logf(TEXT("      class - sort by sound class name"));
-		GLog->Logf(TEXT("      name - sort by cue pathname"));
-		GLog->Logf(TEXT("      waves - sort by waves' num"));
-		GLog->Logf(TEXT("      default - sorting is no enabled"));
-		GLog->Logf(TEXT("  stat sounds -debug - enables debugging mode like showing sound radius sphere and names, but only for cues with enabled property bDebug"));
-		GLog->Logf(TEXT(""));
-		GLog->Logf(TEXT("Ex. stat sounds sort=class -debug"));
-		GLog->Logf(TEXT(" This will show only debug sounds sorted by sound class"));
-	}
-
-	uint32 OldSoundShowFlags = ViewportClient->GetSoundShowFlags();
-
-	uint32 ShowSounds = FViewportClient::ESoundShowFlags::Disabled;
-
-	if (Stream)
-	{
-		const bool bHide = FParse::Command(&Stream, TEXT("off"));
-		if (bHide)
-		{
-			ShowSounds = FViewportClient::ESoundShowFlags::Disabled;
-		}
-		else
-		{
-			const bool bDebug = FParse::Param(Stream, TEXT("debug"));
-			if (bDebug)
-			{
-				ShowSounds |= FViewportClient::ESoundShowFlags::Debug;
-			}
-
-			const bool bLongNames = FParse::Param(Stream, TEXT("longnames"));
-			if (bLongNames)
-			{
-				ShowSounds |= FViewportClient::ESoundShowFlags::Long_Names;
-			}
-
-			FString SortStr;
-			FParse::Value(Stream, TEXT("sort="), SortStr);
-			if (SortStr == TEXT("distance"))
-			{
-				ShowSounds |= FViewportClient::ESoundShowFlags::Sort_Distance;
-			}
-			else if (SortStr == TEXT("class"))
-			{
-				ShowSounds |= FViewportClient::ESoundShowFlags::Sort_Class;
-			}
-			else if (SortStr == TEXT("name"))
-			{
-				ShowSounds |= FViewportClient::ESoundShowFlags::Sort_Name;
-			}
-			else if (SortStr == TEXT("waves"))
-			{
-				ShowSounds |= FViewportClient::ESoundShowFlags::Sort_WavesNum;
-			}
-			else
-			{
-				ShowSounds |= FViewportClient::ESoundShowFlags::Sort_Disabled;
-			}
-		}
-	}
-
-	if (OldSoundShowFlags != FViewportClient::ESoundShowFlags::Disabled)
-	{
-		if (ShowSounds != FViewportClient::ESoundShowFlags::Disabled && ShowSounds != FViewportClient::ESoundShowFlags::Sort_Disabled)
-		{
-			if (!ViewportClient->IsStatEnabled(TEXT("Sounds")))
-			{
-				if (const TArray<FString>* CurrentStats = ViewportClient->GetEnabledStats())
-				{
-					TArray<FString> NewStats = *CurrentStats;
-					NewStats.Add(TEXT("Sounds"));
-					ViewportClient->SetEnabledStats(NewStats);
-					ViewportClient->SetShowStats(true);
-				}
-			}
-		}
-		else
-		{
-			ShowSounds = FViewportClient::ESoundShowFlags::Disabled;
-		}
-	}
-	else if (ShowSounds == FViewportClient::ESoundShowFlags::Disabled)
-	{
-		if (ViewportClient->IsStatEnabled(TEXT("Sounds")))
-		{
-			if (const TArray<FString>* CurrentStats = ViewportClient->GetEnabledStats())
-			{
-				TArray<FString> NewStats = *CurrentStats;
-				NewStats.Remove(TEXT("Sounds"));
-				ViewportClient->SetEnabledStats(NewStats);
-			}
-		}
-	}
-	ViewportClient->SetSoundShowFlags((FViewportClient::ESoundShowFlags::Type)ShowSounds);
-
-	if (FAudioDevice* AudioDevice = World->GetAudioDevice())
-	{
-		AudioDevice->UpdateSoundShowFlags(OldSoundShowFlags, ShowSounds);
-	}
-
-	return true;
+#if ENABLE_AUDIO_DEBUG
+	return FAudioDebugger::ToggleStatSounds(World, ViewportClient, Stream);
+#else // !ENABLE_AUDIO_DEBUG
+	return false;
+#endif // !ENABLE_AUDIO_DEBUG
 }
 
 int32 UEngine::RenderStatSounds(UWorld* World, FViewport* Viewport, FCanvas* Canvas, int32 X, int32 Y, const FVector* ViewLocation, const FRotator* ViewRotation)
 {
-	const FViewportClient::ESoundShowFlags::Type ShowSounds = Viewport->GetClient() ? Viewport->GetClient()->GetSoundShowFlags() : FViewportClient::ESoundShowFlags::Disabled;
-	const bool bDebug = ShowSounds & FViewportClient::ESoundShowFlags::Debug;
-
-	if (FAudioDevice* AudioDevice = World->GetAudioDevice())
-	{
-		FAudioStats& AudioStats = AudioDevice->GetAudioStats();
-		if (!AudioStats.bStale)
-		{
-			FString SortingName = TEXT("disabled");
-
-			// Sort the list.
-			if (ShowSounds & FViewportClient::ESoundShowFlags::Sort_Name)
-			{
-				AudioStats.StatSoundInfos.Sort([](const FAudioStats::FStatSoundInfo& A, const FAudioStats::FStatSoundInfo& B) { return A.SoundName < B.SoundName; });
-				SortingName = TEXT("pathname");
-			}
-			else if (ShowSounds & FViewportClient::ESoundShowFlags::Sort_Distance)
-			{
-				AudioStats.StatSoundInfos.Sort([](const FAudioStats::FStatSoundInfo& A, const FAudioStats::FStatSoundInfo& B) { return A.Distance < B.Distance; });
-				SortingName = TEXT("distance");
-			}
-			else if (ShowSounds & FViewportClient::ESoundShowFlags::Sort_Class)
-			{
-				AudioStats.StatSoundInfos.Sort([](const FAudioStats::FStatSoundInfo& A, const FAudioStats::FStatSoundInfo& B) { return A.SoundClassName < B.SoundClassName; });
-				SortingName = TEXT("class");
-			}
-			else if (ShowSounds & FViewportClient::ESoundShowFlags::Sort_WavesNum)
-			{
-				AudioStats.StatSoundInfos.Sort([](const FAudioStats::FStatSoundInfo& A, const FAudioStats::FStatSoundInfo& B) { return A.WaveInstanceInfos.Num() > B.WaveInstanceInfos.Num(); });
-				SortingName = TEXT("waves' num");
-			}
-
-			Canvas->DrawShadowedString(X, Y, TEXT("Active Sounds:"), GetSmallFont(), FColor::Green);
-			Y += 12;
-
-			const FString InfoText = FString::Printf(TEXT(" Sorting: %s Debug: %s"), *SortingName, bDebug ? TEXT("enabled") : TEXT("disabled"));
-			Canvas->DrawShadowedString(X, Y, *InfoText, GetSmallFont(), FColor(128, 255, 128));
-			Y += 12;
-
-			Canvas->DrawShadowedString(X, Y, TEXT("Index Path (Class) Distance"), GetSmallFont(), FColor::Green);
-			Y += 12;
-
-			int32 TotalSoundWavesNum = 0;
-			for (int32 SoundIndex = 0 ; SoundIndex < AudioStats.StatSoundInfos.Num(); ++SoundIndex)
-			{
-				const FAudioStats::FStatSoundInfo& StatSoundInfo = AudioStats.StatSoundInfos[SoundIndex];
-				const int32 WaveInstancesNum = StatSoundInfo.WaveInstanceInfos.Num();
-				if (WaveInstancesNum > 0)
-				{
-					{
-						const FString TheString = FString::Printf(TEXT("%4i. %s (%s) %6.2f"), SoundIndex, *StatSoundInfo.SoundName, *StatSoundInfo.SoundClassName.ToString(), StatSoundInfo.Distance);
-						Canvas->DrawShadowedString(X, Y, *TheString, GetSmallFont(), FColor::White);
-						Y += 12;
-					}
-
-					TotalSoundWavesNum += WaveInstancesNum;
-
-					// Get the active sound waves.
-					for (int32 WaveIndex = 0; WaveIndex < WaveInstancesNum; WaveIndex++)
-					{
-						const FString TheString = *FString::Printf(TEXT("    %4i. %s"), WaveIndex, *StatSoundInfo.WaveInstanceInfos[WaveIndex].Description);
-						Canvas->DrawShadowedString(X, Y, *TheString, GetSmallFont(), FColor(205, 205, 205));
-						Y += 12;
-					}
-				}
-			}
-
-			Canvas->DrawShadowedString(X, Y, *FString::Printf(TEXT("Total sounds: %i, sound waves: %i"), AudioStats.StatSoundInfos.Num(), TotalSoundWavesNum), GetSmallFont(), FColor::Green);
-			Y += 12;
-
-			Canvas->DrawShadowedString(X, Y, *FString::Printf(TEXT("Listener position: %s"), *AudioStats.ListenerLocation.ToString()), GetSmallFont(), FColor::Green);
-			Y += 12;
-		}
-
-		// Draw sound cue's sphere.
-		if (bDebug)
-		{
-			for (const FAudioStats::FStatSoundInfo& StatSoundInfo : AudioStats.StatSoundInfos)
-			{
-				const FTransform& SoundTransform = StatSoundInfo.Transform;
-				const int32 WaveInstancesNum = StatSoundInfo.WaveInstanceInfos.Num();
-
-				if (StatSoundInfo.Distance > 100.0f && WaveInstancesNum > 0)
-				{
-					float SphereRadius = 0.f;
-					float SphereInnerRadius = 0.f;
-
-					if (StatSoundInfo.ShapeDetailsMap.Num() > 0)
-					{
-						DrawDebugString(World, SoundTransform.GetTranslation(), StatSoundInfo.SoundName, NULL, FColor::White, 0.01f);
-
-						for (auto ShapeDetailsIt = StatSoundInfo.ShapeDetailsMap.CreateConstIterator(); ShapeDetailsIt; ++ShapeDetailsIt)
-						{
-							const FBaseAttenuationSettings::AttenuationShapeDetails& ShapeDetails = ShapeDetailsIt.Value();
-							switch (ShapeDetailsIt.Key())
-							{
-							case EAttenuationShape::Sphere:
-								if (ShapeDetails.Falloff > 0.f)
-								{
-									DrawDebugSphere(World, SoundTransform.GetTranslation(), ShapeDetails.Extents.X + ShapeDetails.Falloff, 10, FColor(155, 155, 255));
-									DrawDebugSphere(World, SoundTransform.GetTranslation(), ShapeDetails.Extents.X, 10, FColor(55, 55, 255));
-								}
-								else
-								{
-									DrawDebugSphere(World, SoundTransform.GetTranslation(), ShapeDetails.Extents.X, 10, FColor(155, 155, 255));
-								}
-								break;
-
-							case EAttenuationShape::Box:
-								if (ShapeDetails.Falloff > 0.f)
-								{
-									DrawDebugBox(World, SoundTransform.GetTranslation(), ShapeDetails.Extents + FVector(ShapeDetails.Falloff), SoundTransform.GetRotation(), FColor(155, 155, 255));
-									DrawDebugBox(World, SoundTransform.GetTranslation(), ShapeDetails.Extents, SoundTransform.GetRotation(), FColor(55, 55, 255));
-								}
-								else
-								{
-									DrawDebugBox(World, SoundTransform.GetTranslation(), ShapeDetails.Extents, SoundTransform.GetRotation(), FColor(155, 155, 255));
-								}
-								break;
-
-							case EAttenuationShape::Capsule:
-
-								if (ShapeDetails.Falloff > 0.f)
-								{
-									DrawDebugCapsule(World, SoundTransform.GetTranslation(), ShapeDetails.Extents.X + ShapeDetails.Falloff, ShapeDetails.Extents.Y + ShapeDetails.Falloff, SoundTransform.GetRotation(), FColor(155, 155, 255));
-									DrawDebugCapsule(World, SoundTransform.GetTranslation(), ShapeDetails.Extents.X, ShapeDetails.Extents.Y, SoundTransform.GetRotation(), FColor(55, 55, 255));
-								}
-								else
-								{
-									DrawDebugCapsule(World, SoundTransform.GetTranslation(), ShapeDetails.Extents.X, ShapeDetails.Extents.Y, SoundTransform.GetRotation(), FColor(155, 155, 255));
-								}
-								break;
-
-							case EAttenuationShape::Cone:
-							{
-								const FVector Origin = SoundTransform.GetTranslation() - (SoundTransform.GetUnitAxis(EAxis::X) * ShapeDetails.ConeOffset);
-
-								if (ShapeDetails.Falloff > 0.f || ShapeDetails.Extents.Z > 0.f)
-								{
-									const float OuterAngle = FMath::DegreesToRadians(ShapeDetails.Extents.Y + ShapeDetails.Extents.Z);
-									const float InnerAngle = FMath::DegreesToRadians(ShapeDetails.Extents.Y);
-									DrawDebugCone(World, Origin, SoundTransform.GetUnitAxis(EAxis::X), ShapeDetails.Extents.X + ShapeDetails.Falloff + ShapeDetails.ConeOffset, OuterAngle, OuterAngle, 10, FColor(155, 155, 255));
-									DrawDebugCone(World, Origin, SoundTransform.GetUnitAxis(EAxis::X), ShapeDetails.Extents.X + ShapeDetails.ConeOffset, InnerAngle, InnerAngle, 10, FColor(55, 55, 255));
-								}
-								else
-								{
-									const float Angle = FMath::DegreesToRadians(ShapeDetails.Extents.Y);
-									DrawDebugCone(World, Origin, SoundTransform.GetUnitAxis(EAxis::X), ShapeDetails.Extents.X + ShapeDetails.ConeOffset, Angle, Angle, 10, FColor(155, 155, 255));
-								}
-								break;
-							}
-
-							default:
-								check(false);
-							}
-						}
-					}
-				}
-			}
-		}
-	}
+#if ENABLE_AUDIO_DEBUG
+	return FAudioDebugger::RenderStatSounds(World, Viewport, Canvas, X, Y, ViewLocation, ViewRotation);
+#else // !ENABLE_AUDIO_DEBUG
 	return Y;
+#endif // !ENABLE_AUDIO_DEBUG
 }
 #endif // !UE_BUILD_SHIPPING
 
