[NO NON-PII COMMIT MESSAGE]

diff --git a/Engine/Source/Runtime/Engine/Private/LODActor.cpp b/Engine/Source/Runtime/Engine/Private/LODActor.cpp
--- a/Engine/Source/Runtime/Engine/Private/LODActor.cpp
+++ b/Engine/Source/Runtime/Engine/Private/LODActor.cpp
@@ -527,10 +527,13 @@ const bool ALODActor::IsBuilt(bool bInForce/*=false*/) const
 		// Ensure all subactors are linked to a LOD static mesh component.
 		for (AActor* SubActor : SubActors)
 		{
-			UStaticMeshComponent* LODComponent = GetLODComponentForActor(SubActor, false);
-			if (LODComponent == nullptr || LODComponent->GetStaticMesh() == nullptr)
+			if(SubActor)
 			{
-				return false;
+				UStaticMeshComponent* LODComponent = GetLODComponentForActor(SubActor, false);
+				if (LODComponent == nullptr || LODComponent->GetStaticMesh() == nullptr)
+				{
+					return false;
+				}
 			}
 		}
 
@@ -985,12 +988,6 @@ void ALODActor::UpdateSubActorLODParents()
 {
 	for (AActor* Actor : SubActors)
 	{	
-		if (Actor && Actor->IsA<ALODActor>())
-		{
-			UStaticMeshComponent* LODComponent = GetLODComponentForActor(Actor);
-			Actor->SetLODParent(LODComponent, LODComponent->MinDrawDistance);
-		}
-
 		if(Actor)
 		{
 			UStaticMeshComponent* LODComponent = GetLODComponentForActor(Actor);
