Rework log statements for easier searching
#rnx

diff --git a/Engine/Source/Runtime/Engine/Private/KismetSystemLibrary.cpp b/Engine/Source/Runtime/Engine/Private/KismetSystemLibrary.cpp
--- a/Engine/Source/Runtime/Engine/Private/KismetSystemLibrary.cpp
+++ b/Engine/Source/Runtime/Engine/Private/KismetSystemLibrary.cpp
@@ -407,7 +407,7 @@ FTimerHandle UKismetSystemLibrary::K2_SetTimer(UObject* Object, FString Function
 	InitialStartDelay += FMath::RandRange(-InitialStartDelayVariance, InitialStartDelayVariance);
 	if (Time <= 0.f || ((Time + InitialStartDelay) - InitialStartDelayVariance) < 0.f)
 	{
-		FFrame::KismetExecutionMessage(TEXT("SetTimer passed a negative time.  The associated timer may fail to fire!  If using InitialStartDelayVariance, ensure it is smaller than (Time + InitialStartDelay)."), ELogVerbosity::Warning);
+		FFrame::KismetExecutionMessage(TEXT("SetTimer passed a negative time.  The associated timer may fail to fire!  If using InitialStartDelayVariance, be sure it is smaller than (Time + InitialStartDelay)."), ELogVerbosity::Warning);
 	}
 
 	FTimerDynamicDelegate Delegate;
@@ -426,7 +426,7 @@ FTimerHandle UKismetSystemLibrary::K2_SetTimerDelegate(FTimerDynamicDelegate Del
 			InitialStartDelay += FMath::RandRange(-InitialStartDelayVariance, InitialStartDelayVariance);
 			if (Time <= 0.f || ((Time + InitialStartDelay) - InitialStartDelayVariance) < 0.f)
 			{
-				FFrame::KismetExecutionMessage(TEXT("SetTimer passed a negative time or initial start delay.  The associated timer may fail to fire!  If using InitialStartDelayVariance, ensure it is smaller than (Time + InitialStartDelay)."), ELogVerbosity::Warning);
+				FFrame::KismetExecutionMessage(TEXT("SetTimer passed a negative time or initial start delay.  The associated timer may fail to fire!  If using InitialStartDelayVariance, be sure it is smaller than (Time + InitialStartDelay)."), ELogVerbosity::Warning);
 			}
 
 			FTimerManager& TimerManager = World->GetTimerManager();
