- Move PlayWhenSilent to virtualization mode enum
- Add ability to disable virtualization as a user on a sound (defaults to 'restart' which restarts a sound if its a loop)
- General USoundBase Display clean-up
- Clean-up/Delete unused booleans on USoundCue

diff --git a/Engine/Source/Runtime/AudioMixer/Private/Components/SynthComponent.cpp b/Engine/Source/Runtime/AudioMixer/Private/Components/SynthComponent.cpp
--- a/Engine/Source/Runtime/AudioMixer/Private/Components/SynthComponent.cpp
+++ b/Engine/Source/Runtime/AudioMixer/Private/Components/SynthComponent.cpp
@@ -13,7 +13,7 @@ USynthSound::USynthSound(const FObjectInitializer& ObjectInitializer)
 void USynthSound::Init(USynthComponent* InSynthComponent, const int32 InNumChannels, const int32 InSampleRate, const int32 InCallbackSize)
 {
 	OwningSynthComponent = InSynthComponent;
-	bPlayWhenSilent = true;
+	VirtualizationMode = EVirtualizationMode::PlayWhenSilent;
 	NumChannels = InNumChannels;
 	NumSamplesToGeneratePerCallback = InCallbackSize;
 	// Turn off async generation in old audio engine on mac.
diff --git a/Engine/Source/Runtime/Engine/Classes/Sound/SoundBase.h b/Engine/Source/Runtime/Engine/Classes/Sound/SoundBase.h
--- a/Engine/Source/Runtime/Engine/Classes/Sound/SoundBase.h
+++ b/Engine/Source/Runtime/Engine/Classes/Sound/SoundBase.h
@@ -24,6 +24,29 @@ class USoundEffectSourcePresetChain;
 struct FActiveSound;
 struct FSoundParseParameters;
 
+/**
+ * Method of virtualization when a sound is stopped due to playback constraints
+ * (i.e. by concurrency, priority, and/or MaxChannelCount)
+ * for a given sound.
+ */
+UENUM(BlueprintType)
+enum class EVirtualizationMode : uint8
+{
+	/** Virtualization is disabled */
+	Disabled,
+
+	/** Sound continues to play when silent and not virtualize, continuing to use a voice. If
+	 * sound is looping and stopped due to concurrency or channel limit/priority, sound will
+	 * restart on realization. If any SoundWave referenced in a SoundCue's waveplayer is set
+	 * to 'PlayWhenSilent', entire SoundCue will be overridden to 'PlayWhenSilent' (to maintain
+	 * timing over all wave players).
+	 */
+	PlayWhenSilent,
+
+	/** If sound is looping, sound restarts from beginning upon realization from being virtual */
+	Restart
+};
+
 UCLASS(config=Engine, hidecategories=Object, abstract, editinlinenew, BlueprintType)
 class ENGINE_API USoundBase : public UObject
 {
@@ -34,24 +57,21 @@ public:
 	static USoundConcurrency* DefaultSoundConcurrencyObject;
 
 	/** Sound class this sound belongs to */
-	UPROPERTY(EditAnywhere, Category = Sound, meta = (DisplayName = "Sound Class"), AssetRegistrySearchable)
+	UPROPERTY(EditAnywhere, Category = Sound, meta = (DisplayName = "Class"), AssetRegistrySearchable)
 	USoundClass* SoundClassObject;
 
 	/** When "stat sounds -debug" has been specified, draw this sound's attenuation shape when the sound is audible. For debugging purpose only. */
 	UPROPERTY(EditAnywhere, Category = Debug)
-	uint8 bDebug:1;
+	uint8 bDebug : 1;
 
 	/** Whether or not to override the sound concurrency object with local concurrency settings. */
-	UPROPERTY(EditAnywhere, BlueprintReadWrite, Category = Concurrency)
-	uint8 bOverrideConcurrency:1;
+	UPROPERTY(EditAnywhere, BlueprintReadWrite, Category = "Voice Management|Concurrency")
+	uint8 bOverrideConcurrency : 1;
 
 	/** Whether or not to only send this audio's output to a bus. If true, will not be this sound won't be audible except through bus sends. */
 	UPROPERTY(EditAnywhere, BlueprintReadWrite, Category = Effects)
 	uint8 bOutputToBusOnly : 1;
 
-	UPROPERTY()
-	uint8 bIgnoreFocus_DEPRECATED:1;
-
 	/** Whether or not to only send this audio's output to a bus. If true, will not be this sound won't be audible except through bus sends. */
 	UPROPERTY()
 	uint8 bHasDelayNode : 1;
@@ -60,35 +80,39 @@ public:
 	UPROPERTY()
 	uint8 bHasConcatenatorNode : 1;
 
+#if WITH_EDITORONLY_DATA
 	UPROPERTY()
-	uint8 bHasVirtualizeWhenSilent_DEPRECATED:1;
-
-	/** Whether a sound has play when silent enabled (i.e. for a sound cue, if any sound wave player has it enabled). */
-	UPROPERTY()
-	uint8 bHasPlayWhenSilent : 1;
+	uint8 bHasVirtualizeWhenSilent_DEPRECATED : 1;
+#endif // WITH_EDITORONLY_DATA
 
 	/** Bypass volume weighting priority upon evaluating whether sound should remain active when max channel count is met (See platform Audio Settings). */
-	UPROPERTY(EditAnywhere, BlueprintReadWrite, Category = Priority)
+	UPROPERTY(EditAnywhere, BlueprintReadWrite, Category = "Voice Management|Priority")
 	uint8 bBypassVolumeScaleForPriority : 1;
 
+	/** Virtualization behavior, determining if a sound may revive and how it continues playing when culled or evicted (limited to looping sounds). */
+	UPROPERTY(EditAnywhere, BlueprintReadWrite, Category = "Voice Management")
+	EVirtualizationMode VirtualizationMode;
+
 #if WITH_EDITORONLY_DATA
 	UPROPERTY()
-	TEnumAsByte<enum EMaxConcurrentResolutionRule::Type> MaxConcurrentResolutionRule_DEPRECATED;
-#endif
+	TEnumAsByte<EMaxConcurrentResolutionRule::Type> MaxConcurrentResolutionRule_DEPRECATED;
+#endif // WITH_EDITORONLY_DATA
 
 	/** Number of times this cue is currently being played. */
 	int32 CurrentPlayCount;
 
+#if WITH_EDITORONLY_DATA
 	/** If Override Concurrency is false, the sound concurrency settings to use for this sound. */
 	UPROPERTY()
 	USoundConcurrency* SoundConcurrencySettings_DEPRECATED;
+#endif // WITH_EDITORONLY_DATA
 
 	/** Set of concurrency settings to observe (if override is set to false).  Sound must pass all concurrency settings to play. */
-	UPROPERTY(EditAnywhere, BlueprintReadWrite, Category = Concurrency, meta = (EditCondition = "!bOverrideConcurrency"))
+	UPROPERTY(EditAnywhere, BlueprintReadWrite, Category = "Voice Management|Concurrency", meta = (EditCondition = "!bOverrideConcurrency"))
 	TSet<USoundConcurrency*> ConcurrencySet;
 
 	/** If Override Concurrency is true, concurrency settings to use. */
-	UPROPERTY(EditAnywhere, BlueprintReadWrite, Category = Concurrency, meta = (EditCondition = "bOverrideConcurrency"))
+	UPROPERTY(EditAnywhere, BlueprintReadWrite, Category = "Voice Management|Concurrency", meta = (EditCondition = "bOverrideConcurrency"))
 	FSoundConcurrencySettings ConcurrencyOverrides;
 
 #if WITH_EDITORONLY_DATA
@@ -113,11 +137,11 @@ public:
 	  * (see platform's Audio Settings 'Max Channels' property). Unless bypassed, value is weighted with the final volume of the
 	  * sound to produce final runtime priority value.
 	  */
-	UPROPERTY(EditAnywhere, BlueprintReadWrite, Category = Priority, meta = (ClampMin = "0.0", UIMin = "0.0", ClampMax = "100.0", UIMax = "100.0"))
+	UPROPERTY(EditAnywhere, BlueprintReadWrite, Category = "Voice Management|Priority", meta = (ClampMin = "0.0", UIMin = "0.0", ClampMax = "100.0", UIMax = "100.0"))
 	float Priority;
 
 	/** Attenuation settings package for the sound */
-	UPROPERTY(EditAnywhere, Category=Attenuation)
+	UPROPERTY(EditAnywhere, Category = Attenuation)
 	USoundAttenuation* AttenuationSettings;
 
 	/** Modulation for the sound */
@@ -187,7 +211,7 @@ public:
 	bool HasConcatenatorNode() const;
 
 	/** Returns true if any of the sounds in the sound have "play when silent" enabled. */
-	bool IsPlayWhenSilent() const;
+	virtual bool IsPlayWhenSilent() const;
 
 	virtual float GetVolumeMultiplier();
 	virtual float GetPitchMultiplier();
diff --git a/Engine/Source/Runtime/Engine/Classes/Sound/SoundCue.h b/Engine/Source/Runtime/Engine/Classes/Sound/SoundCue.h
--- a/Engine/Source/Runtime/Engine/Classes/Sound/SoundCue.h
+++ b/Engine/Source/Runtime/Engine/Classes/Sound/SoundCue.h
@@ -117,12 +117,14 @@ protected:
 private:
 	float MaxAudibleDistance;
 
-	uint32 bHasVirtualizedSoundWaves:1;
-	uint32 bVirtualizeSoundWavesInitialized:1;
-	uint32 bHasAttenuationNode:1;
-	uint32 bHasAttenuationNodeInitialized:1;
-	uint32 bShouldApplyInteriorVolumes:1;
-	uint32 bShouldApplyInteriorVolumesCached:1;
+	/** Whether a sound has play when silent enabled (i.e. for a sound cue, if any sound wave player has it enabled). */
+	UPROPERTY()
+	uint8 bHasPlayWhenSilent : 1;
+
+	uint8 bHasAttenuationNode : 1;
+	uint8 bHasAttenuationNodeInitialized : 1;
+	uint8 bShouldApplyInteriorVolumes : 1;
+	uint8 bShouldApplyInteriorVolumesCached : 1;
 
 public:
 
@@ -141,6 +143,7 @@ public:
 
 	//~ Begin USoundBase Interface.
 	virtual bool IsPlayable() const override;
+	virtual bool IsPlayWhenSilent() const override;
 	virtual bool ShouldApplyInteriorVolumes() override;
 	virtual void Parse( class FAudioDevice* AudioDevice, const UPTRINT NodeWaveInstanceHash, FActiveSound& ActiveSound, const FSoundParseParameters& ParseParams, TArray<FWaveInstance*>& WaveInstances ) override;
 	virtual float GetVolumeMultiplier() override;
@@ -182,7 +185,7 @@ public:
 	/**
 	 * Recursively finds sound nodes of type T
 	 */
-	template<typename T> 
+	template<typename T>
 	void RecursiveFindNode(USoundNode* Node, TArray<T*>& OutNodes)
 	{
 		if (Node)
@@ -287,6 +290,3 @@ private:
 	static ENGINE_API TSharedPtr<ISoundCueAudioEditor> SoundCueAudioEditor;
 #endif
 };
-
-
-
diff --git a/Engine/Source/Runtime/Engine/Classes/Sound/SoundSourceBus.h b/Engine/Source/Runtime/Engine/Classes/Sound/SoundSourceBus.h
--- a/Engine/Source/Runtime/Engine/Classes/Sound/SoundSourceBus.h
+++ b/Engine/Source/Runtime/Engine/Classes/Sound/SoundSourceBus.h
@@ -37,15 +37,15 @@ public:
 	float SourceBusDuration;
 
 	/** Stop the source bus when the volume goes to zero. */
-	UPROPERTY(EditAnywhere, Category = BusProperties, meta = (UIMin = 0.0, ClampMin = 0.0))
+	UPROPERTY(EditAnywhere, Category = BusProperties)
 	uint32 bAutoDeactivateWhenSilent:1;
 
-	//~ Begin UObject Interface. 
+	//~ Begin UObject Interface.
 	virtual void PostLoad() override;
 #if WITH_EDITOR
 	virtual void PostEditChangeProperty(FPropertyChangedEvent& PropertyChangedEvent) override;
 #endif // WITH_EDITOR
-	//~ End UObject Interface. 
+	//~ End UObject Interface.
 
 	//~ Begin USoundBase Interface.
 	virtual bool IsPlayable() const override;
diff --git a/Engine/Source/Runtime/Engine/Classes/Sound/SoundWave.h b/Engine/Source/Runtime/Engine/Classes/Sound/SoundWave.h
--- a/Engine/Source/Runtime/Engine/Classes/Sound/SoundWave.h
+++ b/Engine/Source/Runtime/Engine/Classes/Sound/SoundWave.h
@@ -134,7 +134,7 @@ struct FSoundWaveSpectralData
 {
 	GENERATED_USTRUCT_BODY()
 
-	// The frequency hz of the spectrum value
+	// The frequency (in Hz) of the spectrum value
 	UPROPERTY(EditAnywhere, BlueprintReadWrite, Category = "SpectralData")
 	float FrequencyHz;
 
@@ -272,36 +272,36 @@ class ENGINE_API USoundWave : public USoundBase
 	GENERATED_UCLASS_BODY()
 
 	/** Platform agnostic compression quality. 1..100 with 1 being best compression and 100 being best quality. */
-	UPROPERTY(EditAnywhere, Category=Compression, meta=(ClampMin = "1", ClampMax = "100"), AssetRegistrySearchable)
+	UPROPERTY(EditAnywhere, Category="Format|Quality", meta=(DisplayName = "Compression", ClampMin = "1", ClampMax = "100"), AssetRegistrySearchable)
 	int32 CompressionQuality;
 
 	/** Priority of this sound when streaming (lower priority streams may not always play) */
-	UPROPERTY(EditAnywhere, Category=Streaming, meta=(ClampMin=0))
+	UPROPERTY(EditAnywhere, Category="Playback|Streaming", meta=(ClampMin=0))
 	int32 StreamingPriority;
 
 	/** Quality of sample rate conversion for platforms that opt into resampling during cook. */
-	UPROPERTY(EditAnywhere, Category = Quality)
+	UPROPERTY(EditAnywhere, Category = "Format|Quality", meta=(DisplayName="Sample Rate"))
 	ESoundwaveSampleRateSettings SampleRateQuality;
 
 	/** Type of buffer this wave uses. Set once on load */
-	TEnumAsByte<enum EDecompressionType> DecompressionType;
+	TEnumAsByte<EDecompressionType> DecompressionType;
 
-	UPROPERTY(EditAnywhere, Category=Sound)
+	UPROPERTY(EditAnywhere, Category=Sound, meta=(DisplayName="Group"))
 	TEnumAsByte<ESoundGroup> SoundGroup;
 
 	/** If set, when played directly (not through a sound cue) the wave will be played looping. */
-	UPROPERTY(EditAnywhere, Category=SoundWave, AssetRegistrySearchable)
+	UPROPERTY(EditAnywhere, Category=Sound, AssetRegistrySearchable)
 	uint8 bLooping:1;
 
 	/** Whether this sound can be streamed to avoid increased memory usage */
-	UPROPERTY(EditAnywhere, Category=Streaming)
+	UPROPERTY(EditAnywhere, Category="Playback|Streaming")
 	uint8 bStreaming:1;
 
 	/** Whether this sound supports seeking. This requires recooking with a codec which supports seekability and streaming. */
-	UPROPERTY(EditAnywhere, Category = Streaming, meta = (DisplayName = "Seekable", EditCondition = "bStreaming"))
+	UPROPERTY(EditAnywhere, Category = "Playback|Streaming", meta = (DisplayName = "Seekable", EditCondition = "bStreaming"))
 	uint8 bSeekableStreaming:1;
 
-	/** Set to true for programmatically-generated, streamed audio. */
+	/** Set to true for programatically-generated, streamed audio. */
 	uint8 bProcedural:1;
 
 	/** Whether this sound wave is beginning to be destroyed by GC. */
@@ -328,15 +328,13 @@ class ENGINE_API USoundWave : public USoundBase
 	UPROPERTY(EditAnywhere, Category=Subtitles )
 	uint8 bSingleLine:1;
 
+#if WITH_EDITORONLY_DATA
 	UPROPERTY()
 	uint8 bVirtualizeWhenSilent_DEPRECATED:1;
-
-	/** Allows sound to continue playing when silent. This prevents issues with sounds restarting when coming back in range, etc. */
-	UPROPERTY(EditAnywhere, Category=Sound, meta = (DisplayName = "Play When Silent"))
-	uint8 bPlayWhenSilent:1;
+#endif // WITH_EDITORONLY_DATA
 
 	/** Whether or not this source is ambisonics file format. */
-	UPROPERTY(EditAnywhere, Category = Sound)
+	UPROPERTY(EditAnywhere, Category = Format)
 	uint8 bIsAmbisonics : 1;
 
 	/** Whether this SoundWave was decompressed from OGG. */
diff --git a/Engine/Source/Runtime/Engine/Private/ActiveSound.cpp b/Engine/Source/Runtime/Engine/Private/ActiveSound.cpp
--- a/Engine/Source/Runtime/Engine/Private/ActiveSound.cpp
+++ b/Engine/Source/Runtime/Engine/Private/ActiveSound.cpp
@@ -196,6 +196,16 @@ void FActiveSound::SetSoundClass(USoundClass* SoundClass)
 							|| (Sound && Sound->ShouldApplyInteriorVolumes());
 }
 
+bool FActiveSound::IsPlayWhenSilent() const
+{
+	if (!AudioDevice || !AudioDevice->PlayWhenSilentEnabled())
+	{
+		return false;
+	}
+
+	return Sound && Sound->IsPlayWhenSilent();
+}
+
 void FActiveSound::ClearAudioComponent()
 {
 	AudioComponentID = 0;
diff --git a/Engine/Source/Runtime/Engine/Private/AudioVirtualLoop.cpp b/Engine/Source/Runtime/Engine/Private/AudioVirtualLoop.cpp
--- a/Engine/Source/Runtime/Engine/Private/AudioVirtualLoop.cpp
+++ b/Engine/Source/Runtime/Engine/Private/AudioVirtualLoop.cpp
@@ -53,6 +53,14 @@ bool FAudioVirtualLoop::Virtualize(const FActiveSound& InActiveSound, bool bDoRa
 
 bool FAudioVirtualLoop::Virtualize(const FActiveSound& InActiveSound, FAudioDevice& AudioDevice, bool bDoRangeCheck, FAudioVirtualLoop& OutVirtualLoop)
 {
+	USoundBase* Sound = InActiveSound.GetSound();
+	check(Sound);
+
+	if (Sound->VirtualizationMode == EVirtualizationMode::Disabled)
+	{
+		return false;
+	}
+
 	if (!bVirtualLoopsEnabledCVar || InActiveSound.bIsPreviewSound || !InActiveSound.IsLooping())
 	{
 		return false;
@@ -136,7 +144,7 @@ bool FAudioVirtualLoop::IsInAudibleRange(const FActiveSound& InActiveSound, cons
 	}
 	check(AudioDevice);
 
-	if (AudioDevice->PlayWhenSilentEnabled() && InActiveSound.IsPlayWhenSilent())
+	if (InActiveSound.IsPlayWhenSilent())
 	{
 		return true;
 	}
diff --git a/Engine/Source/Runtime/Engine/Private/SoundBase.cpp b/Engine/Source/Runtime/Engine/Private/SoundBase.cpp
--- a/Engine/Source/Runtime/Engine/Private/SoundBase.cpp
+++ b/Engine/Source/Runtime/Engine/Private/SoundBase.cpp
@@ -12,7 +12,7 @@ USoundConcurrency* USoundBase::DefaultSoundConcurrencyObject = nullptr;
 
 USoundBase::USoundBase(const FObjectInitializer& ObjectInitializer)
 	: Super(ObjectInitializer)
-	, bIgnoreFocus_DEPRECATED(false)
+	, VirtualizationMode(EVirtualizationMode::Restart)
 	, Priority(1.0f)
 {
 #if WITH_EDITORONLY_DATA
@@ -106,7 +106,7 @@ bool USoundBase::HasConcatenatorNode() const
 
 bool USoundBase::IsPlayWhenSilent() const
 {
-	return bHasPlayWhenSilent;
+	return VirtualizationMode == EVirtualizationMode::PlayWhenSilent;
 }
 
 float USoundBase::GetVolumeMultiplier()
@@ -216,17 +216,8 @@ void USoundBase::Serialize(FArchive& Ar)
 {
 	Super::Serialize(Ar);
 
-	if (Ar.IsLoading() || Ar.IsSaving())
-	{
-		if (bHasVirtualizeWhenSilent_DEPRECATED)
-		{
-			bHasVirtualizeWhenSilent_DEPRECATED = 0;
-			bHasPlayWhenSilent = 1;
-		}
-	}
-
 #if WITH_EDITORONLY_DATA
-	if (Ar.IsLoading())
+	if (Ar.IsLoading() || Ar.IsSaving())
 	{
 		if (SoundConcurrencySettings_DEPRECATED != nullptr)
 		{
diff --git a/Engine/Source/Runtime/Engine/Private/SoundCue.cpp b/Engine/Source/Runtime/Engine/Private/SoundCue.cpp
--- a/Engine/Source/Runtime/Engine/Private/SoundCue.cpp
+++ b/Engine/Source/Runtime/Engine/Private/SoundCue.cpp
@@ -466,11 +466,21 @@ bool USoundCue::IsPlayable() const
 	return FirstNode != nullptr;
 }
 
-void USoundCue::Parse( FAudioDevice* AudioDevice, const UPTRINT NodeWaveInstanceHash, FActiveSound& ActiveSound, const FSoundParseParameters& ParseParams, TArray<FWaveInstance*>& WaveInstances )
+bool USoundCue::IsPlayWhenSilent() const
+{
+	if (VirtualizationMode == EVirtualizationMode::PlayWhenSilent)
+	{
+		return true;
+	}
+
+	return bHasPlayWhenSilent;
+}
+
+void USoundCue::Parse(FAudioDevice* AudioDevice, const UPTRINT NodeWaveInstanceHash, FActiveSound& ActiveSound, const FSoundParseParameters& ParseParams, TArray<FWaveInstance*>& WaveInstances)
 {
 	if (FirstNode)
 	{
-		FirstNode->ParseNodes(AudioDevice,(UPTRINT)FirstNode,ActiveSound,ParseParams,WaveInstances);
+		FirstNode->ParseNodes(AudioDevice, (UPTRINT)FirstNode, ActiveSound, ParseParams, WaveInstances);
 	}
 }
 
diff --git a/Engine/Source/Runtime/Engine/Private/SoundNodeWavePlayer.cpp b/Engine/Source/Runtime/Engine/Private/SoundNodeWavePlayer.cpp
--- a/Engine/Source/Runtime/Engine/Private/SoundNodeWavePlayer.cpp
+++ b/Engine/Source/Runtime/Engine/Private/SoundNodeWavePlayer.cpp
@@ -169,7 +169,7 @@ bool USoundNodeWavePlayer::IsPlayWhenSilent() const
 {
 	if (SoundWave)
 	{
-		return SoundWave->bPlayWhenSilent;
+		return SoundWave->VirtualizationMode == EVirtualizationMode::PlayWhenSilent;
 	}
 	return false;
 }
diff --git a/Engine/Source/Runtime/Engine/Private/SoundSourceBus.cpp b/Engine/Source/Runtime/Engine/Private/SoundSourceBus.cpp
--- a/Engine/Source/Runtime/Engine/Private/SoundSourceBus.cpp
+++ b/Engine/Source/Runtime/Engine/Private/SoundSourceBus.cpp
@@ -34,7 +34,7 @@ void USoundSourceBus::Init()
 
 	// Keep playing this bus when the volume is 0
 	// Note source buses can't ever be truly virtual as they are procedurally generated.
-	bPlayWhenSilent = !bAutoDeactivateWhenSilent;
+	VirtualizationMode = bAutoDeactivateWhenSilent ? EVirtualizationMode::Disabled : EVirtualizationMode::PlayWhenSilent;
 
 	// Set the channels equal to the users channel count choice
 	switch (SourceBusChannels)
diff --git a/Engine/Source/Runtime/Engine/Private/SoundWave.cpp b/Engine/Source/Runtime/Engine/Private/SoundWave.cpp
--- a/Engine/Source/Runtime/Engine/Private/SoundWave.cpp
+++ b/Engine/Source/Runtime/Engine/Private/SoundWave.cpp
@@ -300,17 +300,13 @@ void USoundWave::Serialize( FArchive& Ar )
 
 	if (Ar.IsSaving() || Ar.IsCooking())
 	{
-		if (bHasVirtualizeWhenSilent_DEPRECATED)
-		{
-			bHasVirtualizeWhenSilent_DEPRECATED = 0;
-			bHasPlayWhenSilent = 1;
-		}
-
+#if WITH_EDITORONLY_DATA
 		if (bVirtualizeWhenSilent_DEPRECATED)
 		{
 			bVirtualizeWhenSilent_DEPRECATED = 0;
-			bPlayWhenSilent = 1;
+			VirtualizationMode = EVirtualizationMode::PlayWhenSilent;
 		}
+#endif // WITH_EDITORONLY_DATA
 
 #if WITH_ENGINE
 		// If there is an AutoStreamingThreshold set for the platform we're cooking to,
@@ -457,7 +453,7 @@ float USoundWave::GetSubtitlePriority() const
 
 bool USoundWave::SupportsSubtitles() const
 {
-	return bPlayWhenSilent || (Subtitles.Num() > 0);
+	return VirtualizationMode == EVirtualizationMode::PlayWhenSilent || (Subtitles.Num() > 0);
 }
 
 void USoundWave::PostInitProperties()
@@ -703,8 +699,6 @@ void USoundWave::PostLoad()
 		return;
 	}
 
-	bHasPlayWhenSilent = bPlayWhenSilent;
-
 #if WITH_EDITORONLY_DATA
 	// Log a warning after loading if the source has effect chains but has channels greater than 2.
 	if (SourceEffectChain && SourceEffectChain->Chain.Num() > 0 && NumChannels > 2)
@@ -1646,9 +1640,9 @@ void USoundWave::Parse(FAudioDevice* AudioDevice, const UPTRINT NodeWaveInstance
 		bool bHasSubtitles = ActiveSound.bHandleSubtitles && (ActiveSound.bHasExternalSubtitles || (Subtitles.Num() > 0));
 
 		// When the BypassPlayWhenSilent cvar is enabled, we should only honor bPlayWhenSilent for procedural sounds:
-		const bool bCanPlayWhenSilent = bPlayWhenSilent && (!BypassPlayWhenSilentCVar || bProcedural);
+		const bool bCanPlayWhenSilent = IsPlayWhenSilent() && (!BypassPlayWhenSilentCVar || bProcedural);
 		const float WaveInstanceVolume = WaveInstance->GetVolumeWithDistanceAttenuation() * WaveInstance->GetDynamicVolume();
-		if (WaveInstanceVolume > KINDA_SMALL_NUMBER || ((bCanPlayWhenSilent || bHasSubtitles) && AudioDevice->PlayWhenSilentEnabled()))
+		if (WaveInstanceVolume > KINDA_SMALL_NUMBER || (bCanPlayWhenSilent || bHasSubtitles))
 		{
 			bAddedWaveInstance = true;
 			WaveInstances.Add(WaveInstance);
diff --git a/Engine/Source/Runtime/Engine/Public/ActiveSound.h b/Engine/Source/Runtime/Engine/Public/ActiveSound.h
--- a/Engine/Source/Runtime/Engine/Public/ActiveSound.h
+++ b/Engine/Source/Runtime/Engine/Public/ActiveSound.h
@@ -277,7 +277,7 @@ public:
 	bool IsPlayingAudio() const { return bIsPlayingAudio; }
 
 	/** Whether or not sound reference is valid and set to play when silent. */
-	bool IsPlayWhenSilent() const { return Sound && Sound->IsPlayWhenSilent();  }
+	bool IsPlayWhenSilent() const;
 
 	FAudioDevice* AudioDevice;
 
