- Fix late-triggering issue with oneshots that don't start immediately due to gain stages summing to zero volume on WaveInstance construction/initialization
- Clean-up/clarify WaveInstance lifetime by immediately removing null tpair entries that stick around during stopping phase of ActiveSounds
- Clean-up/clarify requirement of a WaveInstance subscribing to ActiveSound's WaveInstance map on construction

diff --git a/Engine/Source/Runtime/Engine/Private/ActiveSound.cpp b/Engine/Source/Runtime/Engine/Private/ActiveSound.cpp
--- a/Engine/Source/Runtime/Engine/Private/ActiveSound.cpp
+++ b/Engine/Source/Runtime/Engine/Private/ActiveSound.cpp
@@ -123,6 +123,24 @@ FActiveSound::~FActiveSound()
 	check(CanDelete());
 }
 
+FActiveSound* FActiveSound::CreateVirtualCopy(const FActiveSound& InActiveSoundToCopy, FAudioDevice& InAudioDevice)
+{
+	FActiveSound* ActiveSound = new FActiveSound(InActiveSoundToCopy);
+
+	ActiveSound->bAsyncOcclusionPending = false;
+	ActiveSound->bHasVirtualized = true;
+	ActiveSound->bIsPlayingAudio = false;
+
+	ActiveSound->AudioDevice = &InAudioDevice;
+
+	ActiveSound->VolumeConcurrency = 1.0f;
+
+	ActiveSound->ConcurrencyGroupData.Reset();
+	ActiveSound->WaveInstances.Reset();
+
+	return ActiveSound;
+}
+
 FArchive& operator<<( FArchive& Ar, FActiveSound* ActiveSound )
 {
 	if( !Ar.IsLoading() && !Ar.IsSaving() )
@@ -691,12 +709,13 @@ void FActiveSound::Stop(bool bStopNow)
 
 	if (Sound && !bIsStopping)
 	{
-		Sound->CurrentPlayCount = FMath::Max( Sound->CurrentPlayCount - 1, 0 );
+		Sound->CurrentPlayCount = FMath::Max(Sound->CurrentPlayCount - 1, 0);
 	}
 
-	for (auto WaveInstanceIt(WaveInstances.CreateIterator()); WaveInstanceIt; ++WaveInstanceIt)
+	TArray<FWaveInstance*> ToDelete;
+	for (TPair<UPTRINT, FWaveInstance*> WaveInstanceIt : WaveInstances)
 	{
-		FWaveInstance*& WaveInstance = WaveInstanceIt.Value();
+		FWaveInstance* WaveInstance = WaveInstanceIt.Value;
 
 		// Stop the owning sound source
 		FSoundSource* Source = AudioDevice->GetSoundSource(WaveInstance);
@@ -735,8 +754,7 @@ void FActiveSound::Stop(bool bStopNow)
 			{
 				Source->StopNow();
 
-				delete WaveInstance;
-				WaveInstance = nullptr;
+				ToDelete.Add(WaveInstance);
 			}
 			else
 			{
@@ -747,19 +765,18 @@ void FActiveSound::Stop(bool bStopNow)
 		else
 		{
 			// Have a wave instance but no source.
-			delete WaveInstance;
-			WaveInstance = nullptr;
+			ToDelete.Add(WaveInstance);
 		}
 	}
 
-	if (bStopNow)
+	for (FWaveInstance* WaveInstance : ToDelete)
 	{
-		bIsStopping = false;
+		RemoveWaveInstance(WaveInstance->WaveInstanceHash);
 	}
 
-	if (!bIsStopping)
+	if (bStopNow)
 	{
-		WaveInstances.Empty();
+		bIsStopping = false;
 	}
 
 	if (!bWasStopping)
@@ -778,9 +795,10 @@ bool FActiveSound::UpdateStoppingSources(uint64 CurrentTick, bool bEnsureStopped
 
 	bIsStopping = false;
 
-	for (auto WaveInstanceIt(WaveInstances.CreateIterator()); WaveInstanceIt; ++WaveInstanceIt)
+	TArray<FWaveInstance*> ToDelete;
+	for (TPair<UPTRINT, FWaveInstance*> WaveInstanceIt : WaveInstances)
 	{
-		FWaveInstance*& WaveInstance = WaveInstanceIt.Value();
+		FWaveInstance* WaveInstance = WaveInstanceIt.Value;
 
 		// Some wave instances in the list here may be nullptr if some sounds have already stopped or didn't need to do a stop
 		if (WaveInstance)
@@ -795,8 +813,7 @@ bool FActiveSound::UpdateStoppingSources(uint64 CurrentTick, bool bEnsureStopped
 					Source->StopNow();
 
 					// Delete the wave instance
-					delete WaveInstance;
-					WaveInstance = nullptr;
+					ToDelete.Add(WaveInstance);
 				}
 				else
 				{
@@ -811,17 +828,21 @@ bool FActiveSound::UpdateStoppingSources(uint64 CurrentTick, bool bEnsureStopped
 			else
 			{
 				// We have a wave instance but no source for it, so just delete it.
-				delete WaveInstance;
-				WaveInstance = nullptr;
+				ToDelete.Add(WaveInstance);
 			}
 		}
 	}
 
+	for (FWaveInstance* WaveInstance : ToDelete)
+	{
+		RemoveWaveInstance(WaveInstance->WaveInstanceHash);
+	}
+
 	// Return true to indicate this active sound can be cleaned up
 	// If we've reached this point, all sound waves have stopped so we can clear this wave instance out.
 	if (!bIsStopping)
 	{
-		WaveInstances.Reset();
+		check(WaveInstances.Num() == 0);
 		return true;
 	}
 
@@ -829,10 +850,18 @@ bool FActiveSound::UpdateStoppingSources(uint64 CurrentTick, bool bEnsureStopped
 	return false;
 }
 
-FWaveInstance* FActiveSound::FindWaveInstance( const UPTRINT WaveInstanceHash )
+FWaveInstance* FActiveSound::FindWaveInstance(const UPTRINT WaveInstanceHash)
+{
+	return WaveInstances.FindRef(WaveInstanceHash);
+}
+
+void FActiveSound::RemoveWaveInstance(const UPTRINT WaveInstanceHash)
 {
-	FWaveInstance** WaveInstance = WaveInstances.Find(WaveInstanceHash);
-	return (WaveInstance ? *WaveInstance : nullptr);
+	if (FWaveInstance* WaveInstance = WaveInstances.FindRef(WaveInstanceHash))
+	{
+		WaveInstances.Remove(WaveInstanceHash);
+		delete WaveInstance;
+	}
 }
 
 void FActiveSound::UpdateAdjustVolumeMultiplier(const float DeltaTime)
@@ -1061,6 +1090,11 @@ void FActiveSound::HandleInteriorVolumes( const FListener& Listener, FSoundParse
 	}
 }
 
+void FActiveSound::AddWaveInstance(const UPTRINT WaveInstanceHash, FWaveInstance& WaveInstance)
+{
+	WaveInstances.Add(WaveInstanceHash, &WaveInstance);
+}
+
 void FActiveSound::ApplyRadioFilter(const FSoundParseParameters& ParseParams )
 {
 	check(AudioDevice);
diff --git a/Engine/Source/Runtime/Engine/Private/Audio.cpp b/Engine/Source/Runtime/Engine/Private/Audio.cpp
--- a/Engine/Source/Runtime/Engine/Private/Audio.cpp
+++ b/Engine/Source/Runtime/Engine/Private/Audio.cpp
@@ -751,12 +751,12 @@ uint32 FWaveInstance::TypeHashCounter = 0;
  *
  * @param InActiveSound		ActiveSound this wave instance belongs to.
  */
-FWaveInstance::FWaveInstance( FActiveSound* InActiveSound )
+FWaveInstance::FWaveInstance(const UPTRINT InWaveInstanceHash, FActiveSound& InActiveSound)
 	: WaveData(nullptr)
 	, SoundClass(nullptr)
 	, SoundSubmix(nullptr)
 	, SourceEffectChain(nullptr)
-	, ActiveSound(InActiveSound)
+	, ActiveSound(&InActiveSound)
 	, Volume(0.0f)
 	, DistanceAttenuation(1.0f)
 	, VolumeMultiplier(1.0f)
@@ -812,10 +812,11 @@ FWaveInstance::FWaveInstance( FActiveSound* InActiveSound )
 	, ReverbSendLevelDistanceRange(0.0f, 0.0f)
 	, ManualReverbSendLevel(0.0f)
 	, TypeHash(0)
-	, WaveInstanceHash(0)
+	, WaveInstanceHash(InWaveInstanceHash)
 	, UserIndex(0)
 {
 	TypeHash = ++TypeHashCounter;
+	ActiveSound->AddWaveInstance(WaveInstanceHash, *this);
 }
 
 /**
diff --git a/Engine/Source/Runtime/Engine/Private/AudioVirtualLoop.cpp b/Engine/Source/Runtime/Engine/Private/AudioVirtualLoop.cpp
--- a/Engine/Source/Runtime/Engine/Private/AudioVirtualLoop.cpp
+++ b/Engine/Source/Runtime/Engine/Private/AudioVirtualLoop.cpp
@@ -51,7 +51,7 @@ bool FAudioVirtualLoop::Virtualize(const FActiveSound& InActiveSound, bool bDoRa
 	return Virtualize(InActiveSound, *AudioDevice, bDoRangeCheck, OutVirtualLoop);
 }
 
-bool FAudioVirtualLoop::Virtualize(const FActiveSound& InActiveSound, FAudioDevice& AudioDevice, bool bDoRangeCheck, FAudioVirtualLoop& OutVirtualLoop)
+bool FAudioVirtualLoop::Virtualize(const FActiveSound& InActiveSound, FAudioDevice& InAudioDevice, bool bDoRangeCheck, FAudioVirtualLoop& OutVirtualLoop)
 {
 	USoundBase* Sound = InActiveSound.GetSound();
 	check(Sound);
@@ -71,21 +71,12 @@ bool FAudioVirtualLoop::Virtualize(const FActiveSound& InActiveSound, FAudioDevi
 		return false;
 	}
 
-	if (bDoRangeCheck && IsInAudibleRange(InActiveSound, &AudioDevice))
+	if (bDoRangeCheck && IsInAudibleRange(InActiveSound, &InAudioDevice))
 	{
 		return false;
 	}
 
-	FActiveSound* ActiveSound = new FActiveSound(InActiveSound);
-
-	ActiveSound->bAsyncOcclusionPending = false;
-	ActiveSound->AudioDevice = &AudioDevice;
-	ActiveSound->bIsPlayingAudio = false;
-	ActiveSound->ConcurrencyGroupData.Reset();
-	ActiveSound->VolumeConcurrency = 1.0f;
-	ActiveSound->WaveInstances.Reset();
-	ActiveSound->bHasVirtualized = true;
-
+	FActiveSound* ActiveSound = FActiveSound::CreateVirtualCopy(InActiveSound, InAudioDevice);
 	OutVirtualLoop.ActiveSound = ActiveSound;
 	return true;
 }
diff --git a/Engine/Source/Runtime/Engine/Private/DialogueWave.cpp b/Engine/Source/Runtime/Engine/Private/DialogueWave.cpp
--- a/Engine/Source/Runtime/Engine/Private/DialogueWave.cpp
+++ b/Engine/Source/Runtime/Engine/Private/DialogueWave.cpp
@@ -454,7 +454,7 @@ void UDialogueSoundWaveProxy::Parse(class FAudioDevice* AudioDevice, const UPTRI
 	bool bHasSubtitles = (Subtitles.Num() > 0);
 
 	// Check if the wave instance exists before we try to add it...
-	const bool bWaveInstanceAlreadyExisted = (ActiveSound.WaveInstances.Contains(NodeWaveInstanceHash));
+	const bool bWaveInstanceAlreadyExisted = ActiveSound.FindWaveInstance(NodeWaveInstanceHash) != nullptr;
 
 	ActiveSound.bHasExternalSubtitles = bHasSubtitles; // Need to set this so the sound will virtualize when silent if necessary.
 	SoundWave->Parse(AudioDevice, NodeWaveInstanceHash, ActiveSound, ParseParams, WaveInstances);
diff --git a/Engine/Source/Runtime/Engine/Private/SoundNodeWavePlayer.cpp b/Engine/Source/Runtime/Engine/Private/SoundNodeWavePlayer.cpp
--- a/Engine/Source/Runtime/Engine/Private/SoundNodeWavePlayer.cpp
+++ b/Engine/Source/Runtime/Engine/Private/SoundNodeWavePlayer.cpp
@@ -133,11 +133,49 @@ void USoundNodeWavePlayer::ParseNodes( FAudioDevice* AudioDevice, const UPTRINT
 			UpdatedParams.bLooping = true;
 			SoundWave->Parse(AudioDevice, NodeWaveInstanceHash, ActiveSound, UpdatedParams, WaveInstances);
 		}
-		// Don't play non-looping sounds again if this sound has been revived to
-		// avoid re-triggering one shots played adjacent with looping wave instances in cue
-		else if (!ActiveSound.bHasVirtualized || ParseParams.bEnableRetrigger)
+		else
 		{
-			SoundWave->Parse(AudioDevice, NodeWaveInstanceHash, ActiveSound, ParseParams, WaveInstances);
+			// Don't play non-looping sounds again if this sound has been revived to
+			// avoid re-triggering one shots played adjacent with looping wave instances in cue
+			if (ParseParams.bEnableRetrigger)
+			{
+				SoundWave->Parse(AudioDevice, NodeWaveInstanceHash, ActiveSound, ParseParams, WaveInstances);
+			}
+			else
+			{
+				// If sound has been virtualized, don't try to revive one-shots
+				if (!ActiveSound.bHasVirtualized)
+				{
+					RETRIEVE_SOUNDNODE_PAYLOAD(sizeof(int32));
+					DECLARE_SOUNDNODE_ELEMENT(int32, bPlayFailed);
+					if (*RequiresInitialization)
+					{
+						bPlayFailed = 0;
+					}
+
+					const int32 InitActiveSoundWaveInstanceNum = ActiveSound.GetWaveInstanceNum();
+					const int32 InitWaveInstancesNum = WaveInstances.Num();
+
+					// Guard against continual parsing if wave instance was created but not added to transient
+					// wave instance list to avoid inaudible sounds popping back in.
+					if (!bPlayFailed)
+					{
+						SoundWave->Parse(AudioDevice, NodeWaveInstanceHash, ActiveSound, ParseParams, WaveInstances);
+					}
+
+					if (*RequiresInitialization != 0)
+					{
+						if (ActiveSound.GetWaveInstanceNum() == InitActiveSoundWaveInstanceNum)
+						{
+							if (WaveInstances.Num() == InitWaveInstancesNum)
+							{
+								bPlayFailed = 1;
+							}
+						}
+						*RequiresInitialization = 0;
+					}
+				}
+			}
 		}
 
 		SoundWave->bLooping = bWaveIsLooping;
diff --git a/Engine/Source/Runtime/Engine/Private/SoundWave.cpp b/Engine/Source/Runtime/Engine/Private/SoundWave.cpp
--- a/Engine/Source/Runtime/Engine/Private/SoundWave.cpp
+++ b/Engine/Source/Runtime/Engine/Private/SoundWave.cpp
@@ -1386,9 +1386,7 @@ bool USoundWave::CleanupDecompressor(bool bForceWait)
 FWaveInstance* USoundWave::HandleStart( FActiveSound& ActiveSound, const UPTRINT WaveInstanceHash ) const
 {
 	// Create a new wave instance and associate with the ActiveSound
-	FWaveInstance* WaveInstance = new FWaveInstance( &ActiveSound );
-	WaveInstance->WaveInstanceHash = WaveInstanceHash;
-	ActiveSound.WaveInstances.Add( WaveInstanceHash, WaveInstance );
+	FWaveInstance* WaveInstance = new FWaveInstance(WaveInstanceHash, ActiveSound);
 
 	// Add in the subtitle if they exist
 	if (ActiveSound.bHandleSubtitles && Subtitles.Num() > 0)
@@ -1456,8 +1454,10 @@ void USoundWave::Parse(FAudioDevice* AudioDevice, const UPTRINT NodeWaveInstance
 {
 	FWaveInstance* WaveInstance = ActiveSound.FindWaveInstance(NodeWaveInstanceHash);
 
+	const bool bIsNewWave = WaveInstance == nullptr;
+
 	// Create a new WaveInstance if this SoundWave doesn't already have one associated with it.
-	if (WaveInstance == NULL)
+	if(!WaveInstance)
 	{
 		if (!ActiveSound.bRadioFilterSelected)
 		{
@@ -1480,220 +1480,229 @@ void USoundWave::Parse(FAudioDevice* AudioDevice, const UPTRINT NodeWaveInstance
 #endif
 	}
 
-	// Check for finished paths.
-	if (!WaveInstance->bIsFinished)
-	{
-		// Propagate properties and add WaveInstance to outgoing array of FWaveInstances.
-		WaveInstance->SetVolume(ParseParams.Volume * Volume);
-		WaveInstance->SetVolumeMultiplier(ParseParams.VolumeMultiplier);
-		WaveInstance->SetDistanceAttenuation(ParseParams.DistanceAttenuation);
-		WaveInstance->SetPitch(ParseParams.Pitch * Pitch);
-		WaveInstance->bEnableLowPassFilter = ParseParams.bEnableLowPassFilter;
-		WaveInstance->bIsOccluded = ParseParams.bIsOccluded;
-		WaveInstance->LowPassFilterFrequency = ParseParams.LowPassFilterFrequency;
-		WaveInstance->OcclusionFilterFrequency = ParseParams.OcclusionFilterFrequency;
-		WaveInstance->AttenuationLowpassFilterFrequency = ParseParams.AttenuationLowpassFilterFrequency;
-		WaveInstance->AttenuationHighpassFilterFrequency = ParseParams.AttenuationHighpassFilterFrequency;
-		WaveInstance->AmbientZoneFilterFrequency = ParseParams.AmbientZoneFilterFrequency;
-		WaveInstance->bApplyRadioFilter = ActiveSound.bApplyRadioFilter;
-		WaveInstance->StartTime = ParseParams.StartTime;
-		WaveInstance->UserIndex = ActiveSound.UserIndex;
-		WaveInstance->OmniRadius = ParseParams.OmniRadius;
-		WaveInstance->StereoSpread = ParseParams.StereoSpread;
-		WaveInstance->AttenuationDistance = ParseParams.AttenuationDistance;
-		WaveInstance->ListenerToSoundDistance = ParseParams.ListenerToSoundDistance;
-		WaveInstance->ListenerToSoundDistanceForPanning = ParseParams.ListenerToSoundDistanceForPanning;
-		WaveInstance->AbsoluteAzimuth = ParseParams.AbsoluteAzimuth;
-
-		if (NumChannels <= 2)
-		{
-			WaveInstance->SourceEffectChain = ParseParams.SourceEffectChain;
-		}
-
-		bool bAlwaysPlay = false;
-
-		// Properties from the sound class
-		WaveInstance->SoundClass = ParseParams.SoundClass;
-		if (ParseParams.SoundClass)
-		{
-			FSoundClassProperties* SoundClassProperties = AudioDevice->GetSoundClassCurrentProperties(ParseParams.SoundClass);
-			// Use values from "parsed/ propagated" sound class properties
-			float VolumeMultiplier = WaveInstance->GetVolumeMultiplier();
-			WaveInstance->SetVolumeMultiplier(VolumeMultiplier* SoundClassProperties->Volume);
-			WaveInstance->SetPitch(WaveInstance->Pitch * SoundClassProperties->Pitch);
-			//TODO: Add in HighFrequencyGainMultiplier property to sound classes
-
-			WaveInstance->VoiceCenterChannelVolume = SoundClassProperties->VoiceCenterChannelVolume;
-			WaveInstance->RadioFilterVolume = SoundClassProperties->RadioFilterVolume * ParseParams.VolumeMultiplier;
-			WaveInstance->RadioFilterVolumeThreshold = SoundClassProperties->RadioFilterVolumeThreshold * ParseParams.VolumeMultiplier;
-			WaveInstance->StereoBleed = SoundClassProperties->StereoBleed;
-			WaveInstance->LFEBleed = SoundClassProperties->LFEBleed;
-
-			WaveInstance->bIsUISound = ActiveSound.bIsUISound || SoundClassProperties->bIsUISound;
-			WaveInstance->bIsMusic = ActiveSound.bIsMusic || SoundClassProperties->bIsMusic;
-			WaveInstance->bCenterChannelOnly = ActiveSound.bCenterChannelOnly || SoundClassProperties->bCenterChannelOnly;
-			WaveInstance->bEQFilterApplied = ActiveSound.bEQFilterApplied || SoundClassProperties->bApplyEffects;
-			WaveInstance->bReverb = ActiveSound.bReverb || SoundClassProperties->bReverb;
-			WaveInstance->OutputTarget = SoundClassProperties->OutputTarget;
-
-			if (SoundClassProperties->bApplyAmbientVolumes)
-			{
-				VolumeMultiplier = WaveInstance->GetVolumeMultiplier();
-				WaveInstance->SetVolumeMultiplier(VolumeMultiplier * ParseParams.InteriorVolumeMultiplier);
-				WaveInstance->RadioFilterVolume *= ParseParams.InteriorVolumeMultiplier;
-				WaveInstance->RadioFilterVolumeThreshold *= ParseParams.InteriorVolumeMultiplier;
-			}
-
-			bAlwaysPlay = ActiveSound.bAlwaysPlay || SoundClassProperties->bAlwaysPlay;
-		}
-		else
-		{
-			WaveInstance->VoiceCenterChannelVolume = 0.f;
-			WaveInstance->RadioFilterVolume = 0.f;
-			WaveInstance->RadioFilterVolumeThreshold = 0.f;
-			WaveInstance->StereoBleed = 0.f;
-			WaveInstance->LFEBleed = 0.f;
-			WaveInstance->bEQFilterApplied = ActiveSound.bEQFilterApplied;
-			WaveInstance->bIsUISound = ActiveSound.bIsUISound;
-			WaveInstance->bIsMusic = ActiveSound.bIsMusic;
-			WaveInstance->bReverb = ActiveSound.bReverb;
-			WaveInstance->bCenterChannelOnly = ActiveSound.bCenterChannelOnly;
-
-			bAlwaysPlay = ActiveSound.bAlwaysPlay;
-		}
+	// Early out if finished.
+	if (WaveInstance->bIsFinished)
+	{
+		return;
+	}
 
-		// If set to bAlwaysPlay, increase the current sound's priority scale by 10x. This will still result in a possible 0-priority output if the sound has 0 actual volume
-		if (bAlwaysPlay)
-		{
-			WaveInstance->Priority = MAX_FLT;
-		}
-		else
-		{
-			WaveInstance->Priority = ParseParams.Priority;
-		}
+	// Propagate properties and add WaveInstance to outgoing array of FWaveInstances.
+	WaveInstance->SetVolume(ParseParams.Volume * Volume);
+	WaveInstance->SetVolumeMultiplier(ParseParams.VolumeMultiplier);
+	WaveInstance->SetDistanceAttenuation(ParseParams.DistanceAttenuation);
+	WaveInstance->SetPitch(ParseParams.Pitch * Pitch);
+	WaveInstance->bEnableLowPassFilter = ParseParams.bEnableLowPassFilter;
+	WaveInstance->bIsOccluded = ParseParams.bIsOccluded;
+	WaveInstance->LowPassFilterFrequency = ParseParams.LowPassFilterFrequency;
+	WaveInstance->OcclusionFilterFrequency = ParseParams.OcclusionFilterFrequency;
+	WaveInstance->AttenuationLowpassFilterFrequency = ParseParams.AttenuationLowpassFilterFrequency;
+	WaveInstance->AttenuationHighpassFilterFrequency = ParseParams.AttenuationHighpassFilterFrequency;
+	WaveInstance->AmbientZoneFilterFrequency = ParseParams.AmbientZoneFilterFrequency;
+	WaveInstance->bApplyRadioFilter = ActiveSound.bApplyRadioFilter;
+	WaveInstance->StartTime = ParseParams.StartTime;
+	WaveInstance->UserIndex = ActiveSound.UserIndex;
+	WaveInstance->OmniRadius = ParseParams.OmniRadius;
+	WaveInstance->StereoSpread = ParseParams.StereoSpread;
+	WaveInstance->AttenuationDistance = ParseParams.AttenuationDistance;
+	WaveInstance->ListenerToSoundDistance = ParseParams.ListenerToSoundDistance;
+	WaveInstance->ListenerToSoundDistanceForPanning = ParseParams.ListenerToSoundDistanceForPanning;
+	WaveInstance->AbsoluteAzimuth = ParseParams.AbsoluteAzimuth;
+
+	if (NumChannels <= 2)
+	{
+		WaveInstance->SourceEffectChain = ParseParams.SourceEffectChain;
+	}
+
+	bool bAlwaysPlay = false;
+
+	// Properties from the sound class
+	WaveInstance->SoundClass = ParseParams.SoundClass;
+	if (ParseParams.SoundClass)
+	{
+		FSoundClassProperties* SoundClassProperties = AudioDevice->GetSoundClassCurrentProperties(ParseParams.SoundClass);
+		// Use values from "parsed/ propagated" sound class properties
+		float VolumeMultiplier = WaveInstance->GetVolumeMultiplier();
+		WaveInstance->SetVolumeMultiplier(VolumeMultiplier* SoundClassProperties->Volume);
+		WaveInstance->Pitch *= SoundClassProperties->Pitch;
+		WaveInstance->SetPitch(WaveInstance->Pitch * SoundClassProperties->Pitch);
+		//TODO: Add in HighFrequencyGainMultiplier property to sound classes
+
+		WaveInstance->VoiceCenterChannelVolume = SoundClassProperties->VoiceCenterChannelVolume;
+		WaveInstance->RadioFilterVolume = SoundClassProperties->RadioFilterVolume * ParseParams.VolumeMultiplier;
+		WaveInstance->RadioFilterVolumeThreshold = SoundClassProperties->RadioFilterVolumeThreshold * ParseParams.VolumeMultiplier;
+		WaveInstance->StereoBleed = SoundClassProperties->StereoBleed;
+		WaveInstance->LFEBleed = SoundClassProperties->LFEBleed;
+
+		WaveInstance->bIsUISound = ActiveSound.bIsUISound || SoundClassProperties->bIsUISound;
+		WaveInstance->bIsMusic = ActiveSound.bIsMusic || SoundClassProperties->bIsMusic;
+		WaveInstance->bCenterChannelOnly = ActiveSound.bCenterChannelOnly || SoundClassProperties->bCenterChannelOnly;
+		WaveInstance->bEQFilterApplied = ActiveSound.bEQFilterApplied || SoundClassProperties->bApplyEffects;
+		WaveInstance->bReverb = ActiveSound.bReverb || SoundClassProperties->bReverb;
+		WaveInstance->OutputTarget = SoundClassProperties->OutputTarget;
+
+		if (SoundClassProperties->bApplyAmbientVolumes)
+		{
+			VolumeMultiplier = WaveInstance->GetVolumeMultiplier();
+			WaveInstance->SetVolumeMultiplier(VolumeMultiplier * ParseParams.InteriorVolumeMultiplier);
+			WaveInstance->RadioFilterVolume *= ParseParams.InteriorVolumeMultiplier;
+			WaveInstance->RadioFilterVolumeThreshold *= ParseParams.InteriorVolumeMultiplier;
+		}
+
+		bAlwaysPlay = ActiveSound.bAlwaysPlay || SoundClassProperties->bAlwaysPlay;
+	}
+	else
+	{
+		WaveInstance->VoiceCenterChannelVolume = 0.f;
+		WaveInstance->RadioFilterVolume = 0.f;
+		WaveInstance->RadioFilterVolumeThreshold = 0.f;
+		WaveInstance->StereoBleed = 0.f;
+		WaveInstance->LFEBleed = 0.f;
+		WaveInstance->bEQFilterApplied = ActiveSound.bEQFilterApplied;
+		WaveInstance->bIsUISound = ActiveSound.bIsUISound;
+		WaveInstance->bIsMusic = ActiveSound.bIsMusic;
+		WaveInstance->bReverb = ActiveSound.bReverb;
+		WaveInstance->bCenterChannelOnly = ActiveSound.bCenterChannelOnly;
 
-		WaveInstance->Location = ParseParams.Transform.GetTranslation();
-		WaveInstance->bIsStarted = true;
-		WaveInstance->bAlreadyNotifiedHook = false;
-		WaveInstance->SetUseSpatialization(ParseParams.bUseSpatialization);
-		WaveInstance->SpatializationMethod = ParseParams.SpatializationMethod;
-		WaveInstance->WaveData = this;
-		WaveInstance->NotifyBufferFinishedHooks = ParseParams.NotifyBufferFinishedHooks;
-		WaveInstance->LoopingMode = ((bLooping || ParseParams.bLooping) ? LOOP_Forever : LOOP_Never);
-		WaveInstance->bIsPaused = ParseParams.bIsPaused;
+		bAlwaysPlay = ActiveSound.bAlwaysPlay;
+	}
 
-		// If we're normalizing 3d stereo spatialized sounds, we need to scale by -6 dB
-		if (WaveInstance->GetUseSpatialization() && ParseParams.bApplyNormalizationToStereoSounds && NumChannels == 2)
-		{
-			float WaveInstanceVolume = WaveInstance->GetVolume();
-			WaveInstance->SetVolume(WaveInstanceVolume * 0.5f);
-		}
+	// If set to bAlwaysPlay, increase the current sound's priority scale by 10x. This will still result in a possible 0-priority output if the sound has 0 actual volume
+	if (bAlwaysPlay)
+	{
+		WaveInstance->Priority = MAX_FLT;
+	}
+	else
+	{
+		WaveInstance->Priority = ParseParams.Priority;
+	}
 
-		// Copy reverb send settings
-		WaveInstance->ReverbSendMethod = ParseParams.ReverbSendMethod;
-		WaveInstance->ManualReverbSendLevel = ParseParams.ManualReverbSendLevel;
-		WaveInstance->CustomRevebSendCurve = ParseParams.CustomReverbSendCurve;
-		WaveInstance->ReverbSendLevelRange = ParseParams.ReverbSendLevelRange;
-		WaveInstance->ReverbSendLevelDistanceRange = ParseParams.ReverbSendLevelDistanceRange;
+	WaveInstance->Location = ParseParams.Transform.GetTranslation();
+	WaveInstance->bIsStarted = true;
+	WaveInstance->bAlreadyNotifiedHook = false;
+	WaveInstance->SetUseSpatialization(ParseParams.bUseSpatialization);
+	WaveInstance->SpatializationMethod = ParseParams.SpatializationMethod;
+	WaveInstance->WaveData = this;
+	WaveInstance->NotifyBufferFinishedHooks = ParseParams.NotifyBufferFinishedHooks;
+	WaveInstance->LoopingMode = ((bLooping || ParseParams.bLooping) ? LOOP_Forever : LOOP_Never);
+	WaveInstance->bIsPaused = ParseParams.bIsPaused;
 
-		// Get the envelope follower settings
-		WaveInstance->EnvelopeFollowerAttackTime = ParseParams.EnvelopeFollowerAttackTime;
-		WaveInstance->EnvelopeFollowerReleaseTime = ParseParams.EnvelopeFollowerReleaseTime;
+	// If we're normalizing 3d stereo spatialized sounds, we need to scale by -6 dB
+	if (WaveInstance->GetUseSpatialization() && ParseParams.bApplyNormalizationToStereoSounds && NumChannels == 2)
+	{
+		float WaveInstanceVolume = WaveInstance->GetVolume();
+		WaveInstance->SetVolume(WaveInstanceVolume * 0.5f);
+	}
 
-		// Copy over the submix sends.
-		WaveInstance->SoundSubmix = ParseParams.SoundSubmix;
-		WaveInstance->SoundSubmixSends = ParseParams.SoundSubmixSends;
+	// Copy reverb send settings
+	WaveInstance->ReverbSendMethod = ParseParams.ReverbSendMethod;
+	WaveInstance->ManualReverbSendLevel = ParseParams.ManualReverbSendLevel;
+	WaveInstance->CustomRevebSendCurve = ParseParams.CustomReverbSendCurve;
+	WaveInstance->ReverbSendLevelRange = ParseParams.ReverbSendLevelRange;
+	WaveInstance->ReverbSendLevelDistanceRange = ParseParams.ReverbSendLevelDistanceRange;
 
-		// Copy over the source bus send and data
-		if (!WaveInstance->ActiveSound->bIsPreviewSound)
-		{
-			WaveInstance->bOutputToBusOnly = ParseParams.bOutputToBusOnly;
-		}
+	// Get the envelope follower settings
+	WaveInstance->EnvelopeFollowerAttackTime = ParseParams.EnvelopeFollowerAttackTime;
+	WaveInstance->EnvelopeFollowerReleaseTime = ParseParams.EnvelopeFollowerReleaseTime;
 
-		for (int32 BusSendType = 0; BusSendType < (int32)EBusSendType::Count; ++BusSendType)
-		{
-			WaveInstance->SoundSourceBusSends[BusSendType] = ParseParams.SoundSourceBusSends[BusSendType];
-		}
+	// Copy over the submix sends.
+	WaveInstance->SoundSubmix = ParseParams.SoundSubmix;
+	WaveInstance->SoundSubmixSends = ParseParams.SoundSubmixSends;
 
-		if (AudioDevice->IsHRTFEnabledForAll() && ParseParams.SpatializationMethod == ESoundSpatializationAlgorithm::SPATIALIZATION_Default)
-		{
-			WaveInstance->SpatializationMethod = ESoundSpatializationAlgorithm::SPATIALIZATION_HRTF;
-		}
-		else
-		{
-			WaveInstance->SpatializationMethod = ParseParams.SpatializationMethod;
-		}
+	// Copy over the source bus send and data
+	if (!WaveInstance->ActiveSound->bIsPreviewSound)
+	{
+		WaveInstance->bOutputToBusOnly = ParseParams.bOutputToBusOnly;
+	}
 
-		// Pass along plugin settings to the wave instance
-		WaveInstance->SpatializationPluginSettings = ParseParams.SpatializationPluginSettings;
-		WaveInstance->OcclusionPluginSettings = ParseParams.OcclusionPluginSettings;
-		WaveInstance->ReverbPluginSettings = ParseParams.ReverbPluginSettings;
-		WaveInstance->ModulationPluginSettings = ParseParams.ModulationPluginSettings;
+	for (int32 BusSendType = 0; BusSendType < (int32)EBusSendType::Count; ++BusSendType)
+	{
+		WaveInstance->SoundSourceBusSends[BusSendType] = ParseParams.SoundSourceBusSends[BusSendType];
+	}
 
-		WaveInstance->bIsAmbisonics = bIsAmbisonics;
+	if (AudioDevice->IsHRTFEnabledForAll() && ParseParams.SpatializationMethod == ESoundSpatializationAlgorithm::SPATIALIZATION_Default)
+	{
+		WaveInstance->SpatializationMethod = ESoundSpatializationAlgorithm::SPATIALIZATION_HRTF;
+	}
+	else
+	{
+		WaveInstance->SpatializationMethod = ParseParams.SpatializationMethod;
+	}
 
-		bool bAddedWaveInstance = false;
+	// Pass along plugin settings to the wave instance
+	WaveInstance->SpatializationPluginSettings = ParseParams.SpatializationPluginSettings;
+	WaveInstance->OcclusionPluginSettings = ParseParams.OcclusionPluginSettings;
+	WaveInstance->ReverbPluginSettings = ParseParams.ReverbPluginSettings;
+	WaveInstance->ModulationPluginSettings = ParseParams.ModulationPluginSettings;
 
-		// Determine if can play when silent here even though we did it up-front in the active sound parse.
-		// This is because an active sound can generate multiple sound waves, not all of them are necessarily permitted
-		// to play when silent.
-		bool bHasSubtitles = ActiveSound.bHandleSubtitles && (ActiveSound.bHasExternalSubtitles || (Subtitles.Num() > 0));
+	WaveInstance->bIsAmbisonics = bIsAmbisonics;
 
-		// When the BypassPlayWhenSilent cvar is enabled, we should only honor bPlayWhenSilent for procedural sounds:
-		const bool bCanPlayWhenSilent = IsPlayWhenSilent() && (!BypassPlayWhenSilentCVar || bProcedural);
-		const float WaveInstanceVolume = WaveInstance->GetVolumeWithDistanceAttenuation() * WaveInstance->GetDynamicVolume();
-		if (WaveInstanceVolume > KINDA_SMALL_NUMBER || (bCanPlayWhenSilent || bHasSubtitles))
-		{
-			bAddedWaveInstance = true;
-			WaveInstances.Add(WaveInstance);
-		}
+	// Recompute the virtualizability here even though we did it up-front in the active sound parse.
+	// This is because an active sound can generate multiple sound waves, not all of them are necessarily virtualizable.
+	bool bHasSubtitles = ActiveSound.bHandleSubtitles && (ActiveSound.bHasExternalSubtitles || (Subtitles.Num() > 0));
 
-		// We're still alive.
-		if (bAddedWaveInstance || WaveInstance->LoopingMode == LOOP_Forever)
+	// When the BypassVirtualizeWhenSilent cvar is enabled, we should only honor bVirtualizeWhenSilent for procedural sounds:
+	const bool bShouldVirtualize = bVirtualizeWhenSilent && (!BypassVirtualizeWhenSilentCVar || bProcedural);
+	if (WaveInstance->GetVolumeWithDistanceAttenuation() > KINDA_SMALL_NUMBER || ((bShouldVirtualize || bHasSubtitles) && AudioDevice->VirtualSoundsEnabled()))
+	{
+		WaveInstances.Add(WaveInstance);
+		ActiveSound.bFinished = false;
+	}
+	else if (WaveInstance->LoopingMode == LOOP_Forever)
+	{
+		ActiveSound.bFinished = false;
+	}
+	// Not looping, silent, and not set to play when silent
+	else
+	{
+		// If no wave instance added to transient array not looping, and just created, immediately delete
+		// to avoid initializing on a later tick (achieved by adding to active sound's wave instance map
+		// but not the passed transient WaveInstance array)
+		if (bIsNewWave)
 		{
-			ActiveSound.bFinished = false;
+			ActiveSound.RemoveWaveInstance(NodeWaveInstanceHash);
+			return;
 		}
+	}
 
-		// Sanity check
-		if(NumChannels > 2 && WaveInstance->GetUseSpatialization() && !WaveInstance->bReportedSpatializationWarning)
+#if !NO_LOGGING
+	// Sanity check
+	if(NumChannels > 2 && WaveInstance->GetUseSpatialization() && !WaveInstance->bReportedSpatializationWarning)
+	{
+		static TSet<USoundWave*> ReportedSounds;
+		if (!ReportedSounds.Contains(this))
 		{
-			static TSet<USoundWave*> ReportedSounds;
-			if (!ReportedSounds.Contains(this))
+			FString SoundWarningInfo = FString::Printf(TEXT("Spatialization on sounds with channels greater than 2 is not supported. SoundWave: %s"), *GetName());
+			if (ActiveSound.GetSound() != this)
 			{
-				FString SoundWarningInfo = FString::Printf(TEXT("Spatialisation on sounds with channels greater than 2 is not supported. SoundWave: %s"), *GetName());
-				if (ActiveSound.GetSound() != this)
-				{
-					SoundWarningInfo += FString::Printf(TEXT(" SoundCue: %s"), *ActiveSound.GetSound()->GetName());
-				}
+				SoundWarningInfo += FString::Printf(TEXT(" SoundCue: %s"), *ActiveSound.GetSound()->GetName());
+			}
 
-#if !NO_LOGGING
-				const uint64 AudioComponentID = ActiveSound.GetAudioComponentID();
-				if (AudioComponentID > 0)
+			const uint64 AudioComponentID = ActiveSound.GetAudioComponentID();
+			if (AudioComponentID > 0)
+			{
+				FAudioThread::RunCommandOnGameThread([AudioComponentID, SoundWarningInfo]()
 				{
-					FAudioThread::RunCommandOnGameThread([AudioComponentID, SoundWarningInfo]()
+					if (UAudioComponent* AudioComponent = UAudioComponent::GetAudioComponentFromID(AudioComponentID))
 					{
-						if (UAudioComponent* AudioComponent = UAudioComponent::GetAudioComponentFromID(AudioComponentID))
-						{
-							AActor* SoundOwner = AudioComponent->GetOwner();
-							UE_LOG(LogAudio, Warning, TEXT("%s Actor: %s AudioComponent: %s"), *SoundWarningInfo, (SoundOwner ? *SoundOwner->GetName() : TEXT("None")), *AudioComponent->GetName());
-						}
-						else
-						{
-							UE_LOG(LogAudio, Warning, TEXT("%s"), *SoundWarningInfo);
-						}
-					});
-				}
-				else
-				{
-					UE_LOG(LogAudio, Warning, TEXT("%s"), *SoundWarningInfo);
-				}
-#endif
-
-				ReportedSounds.Add(this);
+						AActor* SoundOwner = AudioComponent->GetOwner();
+						UE_LOG(LogAudio, Warning, TEXT( "%s Actor: %s AudioComponent: %s" ), *SoundWarningInfo, (SoundOwner ? *SoundOwner->GetName() : TEXT("None")), *AudioComponent->GetName() );
+					}
+					else
+					{
+						UE_LOG(LogAudio, Warning, TEXT("%s"), *SoundWarningInfo );
+					}
+				});
 			}
-			WaveInstance->bReportedSpatializationWarning = true;
+			else
+			{
+				UE_LOG(LogAudio, Warning, TEXT("%s"), *SoundWarningInfo );
+			}
+
+			ReportedSounds.Add(this);
 		}
+		WaveInstance->bReportedSpatializationWarning = true;
 	}
+#endif // !NO_LOGGING
 }
 
 bool USoundWave::IsPlayable() const
diff --git a/Engine/Source/Runtime/Engine/Public/ActiveSound.h b/Engine/Source/Runtime/Engine/Public/ActiveSound.h
--- a/Engine/Source/Runtime/Engine/Public/ActiveSound.h
+++ b/Engine/Source/Runtime/Engine/Public/ActiveSound.h
@@ -213,6 +213,8 @@ public:
 	FActiveSound();
 	~FActiveSound();
 
+	static FActiveSound* CreateVirtualCopy(const FActiveSound& ActiveSoundToCopy, FAudioDevice& AudioDevice);
+
 private:
 	TWeakObjectPtr<UWorld> World;
 	uint32 WorldID;
@@ -300,6 +302,8 @@ private:
 	/** Optional override for the source bus sends for the sound. */
 	TArray<FSoundSourceBusSendInfo> SoundSourceBusSendsOverride[(int32)EBusSendType::Count];
 
+	TMap<UPTRINT, FWaveInstance*> WaveInstances;
+
 public:
 	/** Whether or not the sound has checked if it was occluded already. Used to initialize a sound as occluded and bypassing occlusion interpolation. */
 	uint8 bHasCheckedOcclusion:1;
@@ -481,8 +485,6 @@ public:
 	int32 EnvelopeFollowerAttackTime;
 	int32 EnvelopeFollowerReleaseTime;
 
-	TMap<UPTRINT, FWaveInstance*> WaveInstances;
-
 	TMap<UPTRINT,uint32> SoundNodeOffsetMap;
 	TArray<uint8> SoundNodeData;
 
@@ -500,10 +502,22 @@ public:
 	 */
 	FWaveInstance* FindWaveInstance(const UPTRINT WaveInstanceHash);
 
+	void RemoveWaveInstance(const UPTRINT WaveInstanceHash);
+
+	int32 GetWaveInstanceNum() const
+	{
+		return WaveInstances.Num();
+	}
+
+	/**
+	 * Add newly created wave instance to active sound
+	 */
+	void AddWaveInstance(const UPTRINT WaveInstanceHash, FWaveInstance& WaveInstance);
+
 	/**
 	 * Check whether to apply the radio filter
 	 */
-	void ApplyRadioFilter(const struct FSoundParseParameters& ParseParams);
+	void ApplyRadioFilter(const FSoundParseParameters& ParseParams);
 
 	/** Gets total concurrency gain stage based on all concurrency memberships of sound */
 	float GetTotalConcurrencyVolumeScale() const;
diff --git a/Engine/Source/Runtime/Engine/Public/Audio.h b/Engine/Source/Runtime/Engine/Public/Audio.h
--- a/Engine/Source/Runtime/Engine/Public/Audio.h
+++ b/Engine/Source/Runtime/Engine/Public/Audio.h
@@ -466,7 +466,7 @@ public:
 	uint8 UserIndex;
 
 	/** Constructor, initializing all member variables. */
-	FWaveInstance(FActiveSound* ActiveSound);
+	FWaveInstance(const UPTRINT InWaveInstanceHash, FActiveSound& ActiveSound);
 
 	/** Stops the wave instance without notifying NotifyWaveInstanceFinishedHook. */
 	void StopWithoutNotification();
