Fix bad merge

diff --git a/Engine/Source/Runtime/Engine/Private/SoundWave.cpp b/Engine/Source/Runtime/Engine/Private/SoundWave.cpp
--- a/Engine/Source/Runtime/Engine/Private/SoundWave.cpp
+++ b/Engine/Source/Runtime/Engine/Private/SoundWave.cpp
@@ -1523,7 +1523,6 @@ void USoundWave::Parse(FAudioDevice* AudioDevice, const UPTRINT NodeWaveInstance
 		// Use values from "parsed/ propagated" sound class properties
 		float VolumeMultiplier = WaveInstance->GetVolumeMultiplier();
 		WaveInstance->SetVolumeMultiplier(VolumeMultiplier* SoundClassProperties->Volume);
-		WaveInstance->Pitch *= SoundClassProperties->Pitch;
 		WaveInstance->SetPitch(WaveInstance->Pitch * SoundClassProperties->Pitch);
 		//TODO: Add in HighFrequencyGainMultiplier property to sound classes
 
