Fix build from merge conflict.  Required exposing WaveInstance getter on ActiveSound (internal member hiding conflicted with newly refactored code for Audio Debugger)

diff --git a/Engine/Source/Runtime/Engine/Private/AudioDebug.cpp b/Engine/Source/Runtime/Engine/Private/AudioDebug.cpp
--- a/Engine/Source/Runtime/Engine/Private/AudioDebug.cpp
+++ b/Engine/Source/Runtime/Engine/Private/AudioDebug.cpp
@@ -415,7 +415,7 @@ void FAudioDebugger::DumpActiveSounds() const
 		{
 			UE_LOG(LogAudio, Display, TEXT("%s (%.3g) - %s"), *ActiveSound->GetSound()->GetName(), ActiveSound->GetSound()->GetDuration(), *ActiveSound->GetAudioComponentName());
 
-			for (const TPair<UPTRINT, FWaveInstance*>& WaveInstancePair : ActiveSound->WaveInstances)
+			for (const TPair<UPTRINT, FWaveInstance*>& WaveInstancePair : ActiveSound->GetWaveInstances())
 			{
 				const FWaveInstance* WaveInstance = WaveInstancePair.Value;
 				UE_LOG(LogAudio, Display, TEXT("   %s (%.3g) (%d) - %.3g"),
diff --git a/Engine/Source/Runtime/Engine/Private/SoundNodeWavePlayer.cpp b/Engine/Source/Runtime/Engine/Private/SoundNodeWavePlayer.cpp
--- a/Engine/Source/Runtime/Engine/Private/SoundNodeWavePlayer.cpp
+++ b/Engine/Source/Runtime/Engine/Private/SoundNodeWavePlayer.cpp
@@ -153,7 +153,7 @@ void USoundNodeWavePlayer::ParseNodes( FAudioDevice* AudioDevice, const UPTRINT
 						bPlayFailed = 0;
 					}
 
-					const int32 InitActiveSoundWaveInstanceNum = ActiveSound.GetWaveInstanceNum();
+					const int32 InitActiveSoundWaveInstanceNum = ActiveSound.GetWaveInstances().Num();
 					const int32 InitWaveInstancesNum = WaveInstances.Num();
 
 					// Guard against continual parsing if wave instance was created but not added to transient
@@ -165,7 +165,7 @@ void USoundNodeWavePlayer::ParseNodes( FAudioDevice* AudioDevice, const UPTRINT
 
 					if (*RequiresInitialization != 0)
 					{
-						if (ActiveSound.GetWaveInstanceNum() == InitActiveSoundWaveInstanceNum)
+						if (ActiveSound.GetWaveInstances().Num() == InitActiveSoundWaveInstanceNum)
 						{
 							if (WaveInstances.Num() == InitWaveInstancesNum)
 							{
diff --git a/Engine/Source/Runtime/Engine/Private/SoundWave.cpp b/Engine/Source/Runtime/Engine/Private/SoundWave.cpp
--- a/Engine/Source/Runtime/Engine/Private/SoundWave.cpp
+++ b/Engine/Source/Runtime/Engine/Private/SoundWave.cpp
@@ -1640,8 +1640,8 @@ void USoundWave::Parse(FAudioDevice* AudioDevice, const UPTRINT NodeWaveInstance
 	bool bHasSubtitles = ActiveSound.bHandleSubtitles && (ActiveSound.bHasExternalSubtitles || (Subtitles.Num() > 0));
 
 	// When the BypassVirtualizeWhenSilent cvar is enabled, we should only honor bVirtualizeWhenSilent for procedural sounds:
-	const bool bShouldVirtualize = bVirtualizeWhenSilent && (!BypassVirtualizeWhenSilentCVar || bProcedural);
-	if (WaveInstance->GetVolumeWithDistanceAttenuation() > KINDA_SMALL_NUMBER || ((bShouldVirtualize || bHasSubtitles) && AudioDevice->VirtualSoundsEnabled()))
+	const bool bCanPlayWhenSilent = IsPlayWhenSilent() && (!BypassPlayWhenSilentCVar || bProcedural);
+	const float WaveInstanceVolume = WaveInstance->GetVolumeWithDistanceAttenuation() * WaveInstance->GetDynamicVolume();
 	{
 		WaveInstances.Add(WaveInstance);
 		ActiveSound.bFinished = false;
diff --git a/Engine/Source/Runtime/Engine/Public/ActiveSound.h b/Engine/Source/Runtime/Engine/Public/ActiveSound.h
--- a/Engine/Source/Runtime/Engine/Public/ActiveSound.h
+++ b/Engine/Source/Runtime/Engine/Public/ActiveSound.h
@@ -504,9 +504,9 @@ public:
 
 	void RemoveWaveInstance(const UPTRINT WaveInstanceHash);
 
-	int32 GetWaveInstanceNum() const
+	const TMap<UPTRINT, FWaveInstance*>& GetWaveInstances() const
 	{
-		return WaveInstances.Num();
+		return WaveInstances;
 	}
 
 	/**
