Add back nuked branch.

diff --git a/Engine/Source/Runtime/Engine/Private/SoundWave.cpp b/Engine/Source/Runtime/Engine/Private/SoundWave.cpp
--- a/Engine/Source/Runtime/Engine/Private/SoundWave.cpp
+++ b/Engine/Source/Runtime/Engine/Private/SoundWave.cpp
@@ -1642,6 +1642,7 @@ void USoundWave::Parse(FAudioDevice* AudioDevice, const UPTRINT NodeWaveInstance
 	// When the BypassVirtualizeWhenSilent cvar is enabled, we should only honor bVirtualizeWhenSilent for procedural sounds:
 	const bool bCanPlayWhenSilent = IsPlayWhenSilent() && (!BypassPlayWhenSilentCVar || bProcedural);
 	const float WaveInstanceVolume = WaveInstance->GetVolumeWithDistanceAttenuation() * WaveInstance->GetDynamicVolume();
+	if (WaveInstanceVolume > KINDA_SMALL_NUMBER || (bCanPlayWhenSilent || bHasSubtitles))
 	{
 		WaveInstances.Add(WaveInstance);
 		ActiveSound.bFinished = false;
