Disable copying of FConstPawnIterator to prevent the possibility of a double deletion

  - This vulnerability is largely unseen due to the prevalence of RVO and copy-elision, but some code (specifically TConstPawnArray) was resulting in a copy of the iterator, and subsequent double deletion of the actor iterator as a result.
  - Ideally FConstPawnIterator would not rely on heap allocated pointers so heavily, but I think the safest fix of least risk is to define move operators, and delete its copy constructors.
  - UE-75216 itself regressed in 6548275

diff --git a/Engine/Source/Runtime/Engine/Classes/Engine/World.h b/Engine/Source/Runtime/Engine/Classes/Engine/World.h
--- a/Engine/Source/Runtime/Engine/Classes/Engine/World.h
+++ b/Engine/Source/Runtime/Engine/Classes/Engine/World.h
@@ -111,7 +111,10 @@ private:
 public:
 	~FConstPawnIterator();
 
-	operator bool() const;
+	FConstPawnIterator(FConstPawnIterator&&);
+	FConstPawnIterator& operator=(FConstPawnIterator&&);
+
+	explicit operator bool() const;
 	FPawnIteratorObject operator*() const;
 	TUniquePtr<FPawnIteratorObject> operator->() const;
 
@@ -123,7 +126,7 @@ public:
 	FConstPawnIterator& operator--(int) { return *this; }
 
 private:
-	TActorIterator<APawn>* Iterator;
+	TUniquePtr<TActorIterator<APawn>> Iterator;
 
 	friend UWorld;
 };
diff --git a/Engine/Source/Runtime/Engine/Private/World.cpp b/Engine/Source/Runtime/Engine/Private/World.cpp
--- a/Engine/Source/Runtime/Engine/Private/World.cpp
+++ b/Engine/Source/Runtime/Engine/Private/World.cpp
@@ -4182,35 +4182,43 @@ FConstPawnIterator::FConstPawnIterator(UWorld* World)
 {
 }
 
-FConstPawnIterator::~FConstPawnIterator()
-{
-	delete Iterator;
-}
+// Operators defined in the cpp to ensure the definition of TActorIterator is known
+FConstPawnIterator::FConstPawnIterator(FConstPawnIterator&&) = default;
+FConstPawnIterator& FConstPawnIterator::operator=(FConstPawnIterator&&) = default;
+FConstPawnIterator::~FConstPawnIterator() = default;
 
  FConstPawnIterator::operator bool() const
 {
-	return (bool)*Iterator;
+	return Iterator.IsValid() && (bool)*Iterator;
 }
 
 FConstPawnIterator& FConstPawnIterator::operator++()
 {
+	checkf(Iterator.IsValid(), TEXT("FConstPawnIterator::operator++() - this iterator has been moved from and is now invalid."));
+
 	++(*Iterator);
 	return *this;
 }
 
 FConstPawnIterator& FConstPawnIterator::operator++(int)
 {
+	checkf(Iterator.IsValid(), TEXT("FConstPawnIterator::operator++(int) - this iterator has been moved from and is now invalid."));
+
 	++(*Iterator);
 	return *this;
 }
 
 FPawnIteratorObject FConstPawnIterator::operator*() const
 {
+	checkf(Iterator.IsValid(), TEXT("FConstPawnIterator::operator*() - this iterator has been moved from and is now invalid."));
+
 	return FPawnIteratorObject(**Iterator);
 }
 
 TUniquePtr<FPawnIteratorObject> FConstPawnIterator::operator->() const
 {
+	checkf(Iterator.IsValid(), TEXT("FConstPawnIterator::operator->() - this iterator has been moved from and is now invalid."));
+
 	return TUniquePtr<FPawnIteratorObject>(new FPawnIteratorObject(**Iterator));
 }
 
