Copying //UE4/Dev-Core to Dev-Main (//UE4/Dev-Main)

diff --git a/Engine/Source/Runtime/Engine/Private/StreamableManager.cpp b/Engine/Source/Runtime/Engine/Private/StreamableManager.cpp
--- a/Engine/Source/Runtime/Engine/Private/StreamableManager.cpp
+++ b/Engine/Source/Runtime/Engine/Private/StreamableManager.cpp
@@ -7,6 +7,7 @@
 #include "UObject/UObjectThreadContext.h"
 #include "HAL/IConsoleManager.h"
 #include "Tickable.h"
+#include "Serialization/LoadTimeTrace.h"
 
 DEFINE_LOG_CATEGORY_STATIC(LogStreamableManager, Log, All);
 
@@ -199,6 +200,8 @@ bool FStreamableHandle::BindUpdateDelegate(FStreamableUpdateDelegate NewDelegate
 
 EAsyncPackageState::Type FStreamableHandle::WaitUntilComplete(float Timeout, bool bStartStalledHandles)
 {
+	TRACE_LOADTIME_WAIT_FOR_STREAMABLE_HANDLE_SCOPE(this);
+
 	if (HasLoadCompleted())
 	{
 		return EAsyncPackageState::Complete;
@@ -220,7 +223,7 @@ EAsyncPackageState::Type FStreamableHandle::WaitUntilComplete(float Timeout, boo
 			Handle->StartStalledHandle();
 		}
 
-		for (TSharedPtr<FStreamableHandle> ChildHandle : Handle->ChildHandles)
+		for (const TSharedPtr<FStreamableHandle>& ChildHandle : Handle->ChildHandles)
 		{
 			if (ChildHandle.IsValid())
 			{
@@ -265,7 +268,7 @@ void FStreamableHandle::GetRequestedAssets(TArray<FSoftObjectPath>& AssetList) c
 	AssetList = RequestedAssets;
 
 	// Check child handles
-	for (TSharedPtr<FStreamableHandle> ChildHandle : ChildHandles)
+	for (const TSharedPtr<FStreamableHandle>& ChildHandle : ChildHandles)
 	{
 		TArray<FSoftObjectPath> ChildAssetList;
 
@@ -310,7 +313,7 @@ void FStreamableHandle::GetLoadedAssets(TArray<UObject *>& LoadedAssets) const
 		}
 
 		// Check child handles
-		for (TSharedPtr<FStreamableHandle> ChildHandle : ChildHandles)
+		for (const TSharedPtr<FStreamableHandle>& ChildHandle : ChildHandles)
 		{
 			for (const FSoftObjectPath& Ref : ChildHandle->RequestedAssets)
 			{
@@ -334,7 +337,7 @@ void FStreamableHandle::GetLoadedCount(int32& LoadedCount, int32& RequestedCount
 	LoadedCount = RequestedCount - StreamablesLoading;
 	
 	// Check child handles
-	for (TSharedPtr<FStreamableHandle> ChildHandle : ChildHandles)
+	for (const TSharedPtr<FStreamableHandle>& ChildHandle : ChildHandles)
 	{
 		int32 ChildRequestedCount = 0;
 		int32 ChildLoadedCount = 0;
@@ -418,7 +421,7 @@ void FStreamableHandle::CancelHandle()
 	OwningManager->ManagedActiveHandles.Remove(SharedThis);
 
 	// Remove child handles
-	for (TSharedPtr<FStreamableHandle> ChildHandle : ChildHandles)
+	for (TSharedPtr<FStreamableHandle>& ChildHandle : ChildHandles)
 	{
 		ChildHandle->ParentHandles.Remove(SharedThis);
 	}
@@ -431,7 +434,7 @@ void FStreamableHandle::CancelHandle()
 	{
 		// Update any meta handles that are still active. Copy the array first as elements may be removed from original while iterating
 		TArray<TWeakPtr<FStreamableHandle>> ParentHandlesCopy = ParentHandles;
-		for (TWeakPtr<FStreamableHandle> WeakHandle : ParentHandlesCopy)
+		for (TWeakPtr<FStreamableHandle>& WeakHandle : ParentHandlesCopy)
 		{
 			TSharedPtr<FStreamableHandle> Handle = WeakHandle.Pin();
 
@@ -503,7 +506,8 @@ void FStreamableHandle::StartStalledHandle()
 
 FStreamableHandle::~FStreamableHandle()
 {
-	check(IsInGameThread());
+	TRACE_LOADTIME_DESTROY_STREAMABLE_HANDLE(this);
+	check(IsInGameThread() || IsInGarbageCollectorThread());
 
 	if (IsActive())
 	{
@@ -516,6 +520,7 @@ FStreamableHandle::~FStreamableHandle()
 
 void FStreamableHandle::CompleteLoad()
 {
+	TRACE_LOADTIME_END_LOAD_STREAMABLE_HANDLE(this);
 	// Only complete if it's still active
 	if (IsActive())
 	{
@@ -528,7 +533,7 @@ void FStreamableHandle::CompleteLoad()
 		{
 			// Update any meta handles that are still active. Copy the array first as elements may be removed from original while iterating
 			TArray<TWeakPtr<FStreamableHandle>> ParentHandlesCopy = ParentHandles;
-			for (TWeakPtr<FStreamableHandle> WeakHandle : ParentHandlesCopy)
+			for (TWeakPtr<FStreamableHandle>& WeakHandle : ParentHandlesCopy)
 			{
 				TSharedPtr<FStreamableHandle> Handle = WeakHandle.Pin();
 
@@ -556,10 +561,15 @@ void FStreamableHandle::UpdateCombinedHandle()
 	// Check all our children, complete if done
 	bool bAllCompleted = true;
 	bool bAllCanceled = true;
-	for (TSharedPtr<FStreamableHandle>& ChildHandle : ChildHandles)
+	for (const TSharedPtr<FStreamableHandle>& ChildHandle : ChildHandles)
 	{
 		bAllCompleted = bAllCompleted && !ChildHandle->IsLoadingInProgress();
 		bAllCanceled = bAllCanceled && ChildHandle->WasCanceled();
+
+		if (!bAllCompleted && !bAllCanceled)
+		{
+			return;
+		}
 	}
 
 	// If all our sub handles were canceled, cancel us. Otherwise complete us if at least one was completed and there are none in progress
@@ -593,7 +603,7 @@ void FStreamableHandle::CallUpdateDelegate()
 	UpdateDelegate.ExecuteIfBound(AsShared());
 
 	// Update any meta handles that are still active
-	for (TWeakPtr<FStreamableHandle> WeakHandle : ParentHandles)
+	for (TWeakPtr<FStreamableHandle>& WeakHandle : ParentHandles)
 	{
 		TSharedPtr<FStreamableHandle> Handle = WeakHandle.Pin();
 
@@ -918,7 +928,8 @@ FStreamable* FStreamableManager::StreamInternal(const FSoftObjectPath& InTargetN
 
 			Existing->bAsyncLoadRequestOutstanding = true;
 			Existing->bLoadFailed = false;
-			LoadPackageAsync(Package, FLoadPackageAsyncDelegate::CreateSP(Handle, &FStreamableHandle::AsyncLoadCallbackWrapper, TargetName), Priority);
+			int32 RequestId = LoadPackageAsync(Package, FLoadPackageAsyncDelegate::CreateSP(Handle, &FStreamableHandle::AsyncLoadCallbackWrapper, TargetName), Priority);
+			TRACE_LOADTIME_STREAMABLE_HANDLE_REQUEST_ASSOCIATION(&Handle.Get(), RequestId);
 		}
 	}
 	return Existing;
@@ -946,6 +957,8 @@ TSharedPtr<FStreamableHandle> FStreamableManager::RequestAsyncLoad(const TArray<
 	NewRequest->DebugName = DebugName;
 	NewRequest->Priority = Priority;
 
+	TRACE_LOADTIME_NEW_STREAMABLE_HANDLE(&NewRequest.Get(), *DebugName, false);
+
 	// Remove null requests
 
 	for (int32 i = NewRequest->RequestedAssets.Num() - 1; i >= 0 ; i--)
@@ -1075,6 +1088,8 @@ TSharedPtr<FStreamableHandle> FStreamableManager::RequestSyncLoad(const FSoftObj
 
 void FStreamableManager::StartHandleRequests(TSharedRef<FStreamableHandle> Handle)
 {
+	TRACE_LOADTIME_BEGIN_LOAD_STREAMABLE_HANDLE(&Handle.Get());
+
 	TArray<FStreamable *> ExistingStreamables;
 	ExistingStreamables.Reserve(Handle->RequestedAssets.Num());
 
@@ -1204,7 +1219,7 @@ void FStreamableManager::CheckCompletedRequests(const FSoftObjectPath& Target, s
 	TArray<TSharedRef<FStreamableHandle>> HandlesToComplete;
 	TArray<TSharedRef<FStreamableHandle>> HandlesToRelease;
 
-	for (TSharedRef<FStreamableHandle> Handle : Existing->LoadingHandles)
+	for (TSharedRef<FStreamableHandle>& Handle : Existing->LoadingHandles)
 	{
 		ensure(Handle->WasCanceled() || Handle->OwningManager == this);
 
@@ -1222,12 +1237,12 @@ void FStreamableManager::CheckCompletedRequests(const FSoftObjectPath& Target, s
 	}
 	Existing->LoadingHandles.Empty();
 
-	for (TSharedRef<FStreamableHandle> Handle : HandlesToComplete)
+	for (TSharedRef<FStreamableHandle>& Handle : HandlesToComplete)
 	{
 		Handle->CompleteLoad();
 	}
 
-	for (TSharedRef<FStreamableHandle> Handle : HandlesToRelease)
+	for (TSharedRef<FStreamableHandle>& Handle : HandlesToRelease)
 	{
 		Handle->ReleaseHandle();
 	}
@@ -1318,6 +1333,7 @@ TSharedPtr<FStreamableHandle> FStreamableManager::CreateCombinedHandle(const TAr
 	NewRequest->OwningManager = this;
 	NewRequest->bIsCombinedHandle = true;
 	NewRequest->DebugName = DebugName;
+	TRACE_LOADTIME_NEW_STREAMABLE_HANDLE(&NewRequest.Get(), *DebugName, true);
 
 	for (TSharedPtr<FStreamableHandle> ChildHandle : ChildHandles)
 	{
diff --git a/Engine/Source/Runtime/Engine/Private/Streaming/StreamingManagerTexture.h b/Engine/Source/Runtime/Engine/Private/Streaming/StreamingManagerTexture.h
--- a/Engine/Source/Runtime/Engine/Private/Streaming/StreamingManagerTexture.h
+++ b/Engine/Source/Runtime/Engine/Private/Streaming/StreamingManagerTexture.h
@@ -107,10 +107,8 @@ struct FRenderAssetStreamingManager final : public IRenderAssetStreamingManager
 #if STATS_FAST
 	bool HandleDumpTextureStreamingStatsCommand( const TCHAR* Cmd, FOutputDevice& Ar );
 #endif // STATS_FAST
-#if STATS
-	bool HandleListStreamingRenderAssetsCommand(const TCHAR* Cmd, FOutputDevice& Ar);
-#endif // STATS
 #if !UE_BUILD_SHIPPING
+	bool HandleListStreamingRenderAssetsCommand(const TCHAR* Cmd, FOutputDevice& Ar);
 	bool HandleResetMaxEverRequiredRenderAssetMemoryCommand(const TCHAR* Cmd, FOutputDevice& Ar);
 	bool HandleLightmapStreamingFactorCommand( const TCHAR* Cmd, FOutputDevice& Ar );
 	bool HandleCancelRenderAssetStreamingCommand( const TCHAR* Cmd, FOutputDevice& Ar );
diff --git a/Engine/Source/Runtime/Engine/Private/UnrealEngine.cpp b/Engine/Source/Runtime/Engine/Private/UnrealEngine.cpp
--- a/Engine/Source/Runtime/Engine/Private/UnrealEngine.cpp
+++ b/Engine/Source/Runtime/Engine/Private/UnrealEngine.cpp
@@ -139,6 +139,7 @@ UnrealEngine.cpp: Implements the UEngine class and helpers.
 #include "Sound/AudioSettings.h"
 #include "Streaming/Texture2DUpdate.h"
 #include "Rendering/SkeletalMeshRenderData.h"
+#include "Serialization/LoadTimeTrace.h"
 
 #if WITH_EDITOR
 #include "Settings/LevelEditorPlaySettings.h"
@@ -4507,17 +4508,19 @@ bool UEngine::HandleCountDisabledParticleItemsCommand( const TCHAR* Cmd, FOutput
 	return true;
 }
 
-// View the last N number of names added to the name table. Useful for tracking down name table bloat
+// View all names or the last N names added to the name table. Useful for tracking down name table bloat
 bool UEngine::HandleViewnamesCommand( const TCHAR* Cmd, FOutputDevice& Ar )
 {
-	int32 NumNames = 0;
-	if (FParse::Value(Cmd,TEXT("NUM="),NumNames))
+	const TArray<const FNameEntry*> Entries = FName::DebugDump();
+
+	int32 NumLast = 0;
+	int32 BeginIdx = FParse::Value(Cmd, TEXT("NUM="), NumLast) ? FMath::Max(Entries.Num() - NumLast, 0) : 0;
+
+	for (int32 I = BeginIdx; I < Entries.Num(); ++I)
 	{
-		for (int32 NameIndex = FMath::Max<int32>(FName::GetMaxNames() - NumNames, 0); NameIndex < FName::GetMaxNames(); NameIndex++)
-		{
-			Ar.Logf(TEXT("%d->%s"), NameIndex, *FName::SafeString(NameIndex));
-		}
+		Ar.Log(*Entries[I]->GetPlainNameString());	
 	}
+
 	return true;
 }
 
@@ -9303,7 +9306,7 @@ struct FSoundInfo
 
 	bool CompareClass( const FSoundInfo& Other ) const
 	{
-		return ClassName < Other.ClassName;
+		return ClassName.LexicalLess(Other.ClassName);
 	}
 
 	bool CompareWaveInstancesNum( const FSoundInfo& Other ) const
@@ -12075,7 +12078,9 @@ void UEngine::TickWorldTravel(FWorldContext& Context, float DeltaSeconds)
 
 bool UEngine::LoadMap( FWorldContext& WorldContext, FURL URL, class UPendingNetGame* Pending, FString& Error )
 {
+	TRACE_LOADTIME_LOAD_MAP_SCOPE(*URL.Map);
 	STAT_ADD_CUSTOMMESSAGE_NAME( STAT_NamedMarker, *(FString( TEXT( "LoadMap - " ) + URL.Map )) );
+	TRACE_BOOKMARK(TEXT("LoadMap - %s"), *URL.Map);
 
 	DECLARE_SCOPE_CYCLE_COUNTER(TEXT("UEngine::LoadMap"), STAT_LoadMap, STATGROUP_LoadTime);
 
@@ -12623,6 +12628,7 @@ bool UEngine::LoadMap( FWorldContext& WorldContext, FURL URL, class UPendingNetG
 	}
 
 	STAT_ADD_CUSTOMMESSAGE_NAME( STAT_NamedMarker, *(FString( TEXT( "LoadMapComplete - " ) + URL.Map )) );
+	TRACE_BOOKMARK(TEXT("LoadMapComplete - %s"), *URL.Map);
 	MALLOC_PROFILER( FMallocProfiler::SnapshotMemoryLoadMapEnd( URL.Map ); )
 
 		double StopTime = FPlatformTime::Seconds();
@@ -13325,6 +13331,7 @@ static void AsyncMapChangeLevelLoadCompletionCallback(const FName& PackageName,
 	}
 
 	STAT_ADD_CUSTOMMESSAGE_NAME( STAT_NamedMarker, *(FString( TEXT( "PrepareMapChangeComplete - " ) + PackageName.ToString() )) );
+	TRACE_BOOKMARK(TEXT("PrepareMapChangeComplete - %s"), *PackageName.ToString());
 }
 
 
@@ -13370,6 +13377,7 @@ bool UEngine::PrepareMapChange(FWorldContext &Context, const TArray<FName>& Leve
 		for (const FName LevelName : Context.LevelsToLoadForPendingMapChange)
 		{
 			STAT_ADD_CUSTOMMESSAGE_NAME( STAT_NamedMarker, *(FString( TEXT( "PrepareMapChange - " ) + LevelName.ToString() )) );
+			TRACE_BOOKMARK(TEXT("PrepareMapChange - %s"), *LevelName.ToString());
 			LoadPackageAsync(LevelName.ToString(),
 				FLoadPackageAsyncDelegate::CreateStatic(&AsyncMapChangeLevelLoadCompletionCallback, Context.ContextHandle)
 			);
diff --git a/Engine/Source/Runtime/Engine/Public/SceneUtils.h b/Engine/Source/Runtime/Engine/Public/SceneUtils.h
--- a/Engine/Source/Runtime/Engine/Public/SceneUtils.h
+++ b/Engine/Source/Runtime/Engine/Public/SceneUtils.h
@@ -14,6 +14,7 @@
 #include "Stats/Stats.h"
 #include "RHI.h"
 #include "ProfilingDebugging/CsvProfiler.h"
+#include "GpuProfilerTrace.h"
 
 enum class EShadingPath
 {
@@ -175,16 +176,16 @@ class FScopedGPUStatEvent;
 #if HAS_GPU_STATS
  CSV_DECLARE_CATEGORY_MODULE_EXTERN(ENGINE_API,GPU);
  // The DECLARE_GPU_STAT macros both declare and define a stat (for use in a single CPP)
- #define DECLARE_GPU_STAT(StatName) DECLARE_FLOAT_COUNTER_STAT(TEXT(#StatName), Stat_GPU_##StatName, STATGROUP_GPU); CSV_DEFINE_STAT(GPU,StatName);
- #define DECLARE_GPU_STAT_NAMED(StatName, NameString) DECLARE_FLOAT_COUNTER_STAT(NameString, Stat_GPU_##StatName, STATGROUP_GPU); CSV_DEFINE_STAT(GPU,StatName);
+ #define DECLARE_GPU_STAT(StatName) DECLARE_FLOAT_COUNTER_STAT(TEXT(#StatName), Stat_GPU_##StatName, STATGROUP_GPU); CSV_DEFINE_STAT(GPU,StatName); TRACE_GPUPROFILER_DEFINE_EVENT_TYPE(StatName);
+ #define DECLARE_GPU_STAT_NAMED(StatName, NameString) DECLARE_FLOAT_COUNTER_STAT(NameString, Stat_GPU_##StatName, STATGROUP_GPU); CSV_DEFINE_STAT(GPU,StatName); TRACE_GPUPROFILER_DEFINE_EVENT_TYPE(StatName);
 
  // Extern GPU stats are needed where a stat is used in multiple CPPs. Use the DECLARE_GPU_STAT_NAMED_EXTERN in the header and DEFINE_GPU_STAT in the CPPs
- #define DECLARE_GPU_STAT_NAMED_EXTERN(StatName, NameString) DECLARE_FLOAT_COUNTER_STAT_EXTERN(NameString, Stat_GPU_##StatName, STATGROUP_GPU, ); CSV_DECLARE_STAT_EXTERN(GPU,StatName);
- #define DEFINE_GPU_STAT(StatName) DEFINE_STAT(Stat_GPU_##StatName); CSV_DEFINE_STAT(GPU,StatName);
+ #define DECLARE_GPU_STAT_NAMED_EXTERN(StatName, NameString) DECLARE_FLOAT_COUNTER_STAT_EXTERN(NameString, Stat_GPU_##StatName, STATGROUP_GPU, ); CSV_DECLARE_STAT_EXTERN(GPU,StatName); TRACE_GPUPROFILER_DECLARE_EVENT_TYPE_EXTERN(StatName);
+ #define DEFINE_GPU_STAT(StatName) DEFINE_STAT(Stat_GPU_##StatName); CSV_DEFINE_STAT(GPU,StatName); TRACE_GPUPROFILER_DEFINE_EVENT_TYPE(StatName);
 #if STATS
-  #define SCOPED_GPU_STAT(RHICmdList, StatName) FScopedGPUStatEvent PREPROCESSOR_JOIN(GPUStatEvent_##StatName,__LINE__); PREPROCESSOR_JOIN(GPUStatEvent_##StatName,__LINE__).Begin(RHICmdList, CSV_STAT_FNAME(StatName), GET_STATID( Stat_GPU_##StatName ).GetName() );
+  #define SCOPED_GPU_STAT(RHICmdList, StatName) FScopedGPUStatEvent PREPROCESSOR_JOIN(GPUStatEvent_##StatName,__LINE__); PREPROCESSOR_JOIN(GPUStatEvent_##StatName,__LINE__).Begin(RHICmdList, CSV_STAT_FNAME(StatName), GET_STATID( Stat_GPU_##StatName ).GetName(), TRACE_GPUPROFILER_EVENT_TYPE(StatName) );
  #else
-  #define SCOPED_GPU_STAT(RHICmdList, StatName) FScopedGPUStatEvent PREPROCESSOR_JOIN(GPUStatEvent_##StatName,__LINE__); PREPROCESSOR_JOIN(GPUStatEvent_##StatName,__LINE__).Begin(RHICmdList, CSV_STAT_FNAME(StatName), FName() );
+  #define SCOPED_GPU_STAT(RHICmdList, StatName) FScopedGPUStatEvent PREPROCESSOR_JOIN(GPUStatEvent_##StatName,__LINE__); PREPROCESSOR_JOIN(GPUStatEvent_##StatName,__LINE__).Begin(RHICmdList, CSV_STAT_FNAME(StatName), FName(), TRACE_GPUPROFILER_EVENT_TYPE(StatName) );
  #endif
  #define GPU_STATS_BEGINFRAME(RHICmdList) FRealtimeGPUProfiler::Get()->BeginFrame(RHICmdList);
  #define GPU_STATS_ENDFRAME(RHICmdList) FRealtimeGPUProfiler::Get()->EndFrame(RHICmdList);
@@ -222,7 +223,7 @@ public:
 	ENGINE_API void Release();
 
 	/** Push/pop events */
-	void PushEvent(FRHICommandListImmediate& RHICmdList, const FName& Name, const FName& StatName);
+	void PushEvent(FRHICommandListImmediate& RHICmdList, const FName& Name, const FName& StatName, const FGpuProfilerTrace::FEventType* TraceEventType);
 	void PopEvent(FRHICommandListImmediate& RHICmdList);
 
 private:
@@ -271,7 +272,7 @@ public:
 	/**
 	* Start/Stop functions for timer stats
 	*/
-	ENGINE_API void Begin(FRHICommandList& InRHICmdList, const FName& Name, const FName& StatName );
+	ENGINE_API void Begin(FRHICommandList& InRHICmdList, const FName& Name, const FName& StatName, const FGpuProfilerTrace::FEventType* TraceEventType );
 	ENGINE_API void End();
 };
 #endif // HAS_GPU_STATS
diff --git a/Engine/Source/Runtime/RHI/Private/GpuProfilerTrace.cpp b/Engine/Source/Runtime/RHI/Private/GpuProfilerTrace.cpp
new file mode 100644
--- /dev/null
+++ b/Engine/Source/Runtime/RHI/Private/GpuProfilerTrace.cpp
@@ -0,0 +1,79 @@
+// Copyright 1998-2019 Epic Games, Inc. All Rights Reserved.
+#include "GpuProfilerTrace.h"
+#include "Trace/Trace.h"
+#include "Misc/CString.h"
+#include "GPUProfiler.h"
+#include "ProfilingDebugging/MiscTrace.h"
+#include "HAL/PlatformMisc.h"
+
+#if GPUPROFILERTRACE_ENABLED
+
+UE_TRACE_EVENT_BEGIN(GpuProfiler, EventSpec, Always)
+	UE_TRACE_EVENT_FIELD(const void*, EventType)
+	UE_TRACE_EVENT_FIELD(uint16, NameLength)
+UE_TRACE_EVENT_END()
+
+UE_TRACE_EVENT_BEGIN(GpuProfiler, Frame)
+	UE_TRACE_EVENT_FIELD(uint64, TimestampBase)
+	UE_TRACE_EVENT_FIELD(uint32, RenderingFrameNumber)
+UE_TRACE_EVENT_END()
+
+FGpuProfilerTrace::FFrame FGpuProfilerTrace::CurrentFrame;
+
+FGpuProfilerTrace::FEventType::FEventType(const TCHAR* Name)
+{
+	uint16 NameLength = FCString::Strlen(Name);
+	uint16 NameSize = (NameLength + 1) * sizeof(TCHAR);
+	UE_TRACE_LOG(GpuProfiler, EventSpec, NameSize) 
+		<< EventSpec.EventType(this)
+		<< EventSpec.NameLength(NameLength)
+		<< EventSpec.Attachment(Name, NameSize);
+}
+
+void FGpuProfilerTrace::BeginFrame()
+{
+	CurrentFrame.TimestampBase = 0;
+	CurrentFrame.EventBufferSize = 0;
+}
+
+void FGpuProfilerTrace::BeginEvent(const FEventType* EventType, uint32 FrameNumber, uint64 TimestampMicroseconds)
+{
+	if (CurrentFrame.EventBufferSize >= MaxEventBufferSize - 17) // 9 + 8
+	{
+		return;
+	}
+	if (CurrentFrame.TimestampBase == 0)
+	{
+		CurrentFrame.TimestampBase = TimestampMicroseconds;
+		CurrentFrame.LastTimestamp = CurrentFrame.TimestampBase;
+		CurrentFrame.RenderingFrameNumber = FrameNumber;
+	}
+	uint8* BufferPtr = CurrentFrame.EventBuffer + CurrentFrame.EventBufferSize;
+	uint64 TimestampDelta = TimestampMicroseconds - CurrentFrame.LastTimestamp;
+	CurrentFrame.LastTimestamp = TimestampMicroseconds;
+	FTraceUtils::Encode7bit((TimestampDelta << 1ull) | 0x1, BufferPtr);
+	*reinterpret_cast<uint64*>(BufferPtr) = uint64(EventType);
+	CurrentFrame.EventBufferSize = BufferPtr - CurrentFrame.EventBuffer + sizeof(uint64);
+}
+
+void FGpuProfilerTrace::EndEvent(uint64 TimestampMicroseconds)
+{
+	uint64 TimestampDelta = TimestampMicroseconds - CurrentFrame.LastTimestamp;
+	CurrentFrame.LastTimestamp = TimestampMicroseconds;
+	uint8* BufferPtr = CurrentFrame.EventBuffer + CurrentFrame.EventBufferSize;
+	FTraceUtils::Encode7bit(TimestampDelta << 1ull, BufferPtr);
+	CurrentFrame.EventBufferSize = BufferPtr - CurrentFrame.EventBuffer;
+}
+
+void FGpuProfilerTrace::EndFrame()
+{
+	if (CurrentFrame.EventBufferSize)
+	{
+		UE_TRACE_LOG(GpuProfiler, Frame, CurrentFrame.EventBufferSize)
+			<< Frame.TimestampBase(CurrentFrame.TimestampBase)
+			<< Frame.RenderingFrameNumber(CurrentFrame.RenderingFrameNumber)
+			<< Frame.Attachment(CurrentFrame.EventBuffer, CurrentFrame.EventBufferSize);
+	}
+}
+
+#endif
diff --git a/Engine/Source/Runtime/RHI/Public/GpuProfilerTrace.h b/Engine/Source/Runtime/RHI/Public/GpuProfilerTrace.h
new file mode 100644
--- /dev/null
+++ b/Engine/Source/Runtime/RHI/Public/GpuProfilerTrace.h
@@ -0,0 +1,84 @@
+// Copyright 1998-2019 Epic Games, Inc. All Rights Reserved.
+
+#pragma once
+
+#include "CoreTypes.h"
+#include "Trace/Trace.h"
+
+#if UE_TRACE_ENABLED && !UE_BUILD_SHIPPING
+#define GPUPROFILERTRACE_ENABLED 0
+#else
+#define GPUPROFILERTRACE_ENABLED 0
+#endif
+
+#if GPUPROFILERTRACE_ENABLED
+
+struct FGpuProfilerTrace
+{
+	struct FEventType
+	{
+		RHI_API FEventType(const TCHAR* Name);
+	};
+
+	RHI_API static void BeginFrame();
+	RHI_API static void BeginEvent(const FEventType* EventType, uint32 FrameNumber, uint64 TimestampMicroseconds);
+	RHI_API static void EndEvent(uint64 TimestampMicroseconds);
+	RHI_API static void EndFrame();
+
+private:
+	enum
+	{
+		MaxEventBufferSize = 32768
+	};
+
+	struct FFrame
+	{
+		uint64 TimestampBase;
+		uint64 LastTimestamp;
+		uint32 RenderingFrameNumber;
+		uint16 EventBufferSize;
+		uint8 EventBuffer[MaxEventBufferSize];
+	};
+	static FFrame CurrentFrame;
+};
+
+#define TRACE_GPUPROFILER_DEFINE_EVENT_TYPE(Name) \
+	FGpuProfilerTrace::FEventType PREPROCESSOR_JOIN(__GGpuProfilerEventType, Name)(TEXT(#Name));
+
+#define TRACE_GPUPROFILER_DECLARE_EVENT_TYPE_EXTERN(Name) \
+	extern FGpuProfilerTrace::FEventType PREPROCESSOR_JOIN(__GGpuProfilerEventType, Name);
+
+#define TRACE_GPUPROFILER_EVENT_TYPE(Name) \
+	&PREPROCESSOR_JOIN(__GGpuProfilerEventType, Name)
+
+#define TRACE_GPUPROFILER_BEGIN_FRAME() \
+	FGpuProfilerTrace::BeginFrame();
+
+#define TRACE_GPUPROFILER_BEGIN_EVENT(EventType, FrameNumber, TimestampMicroseconds) \
+	FGpuProfilerTrace::BeginEvent(EventType, FrameNumber, TimestampMicroseconds);
+
+#define TRACE_GPUPROFILER_END_EVENT(TimestampMicroseconds) \
+	FGpuProfilerTrace::EndEvent(TimestampMicroseconds);
+
+#define TRACE_GPUPROFILER_END_FRAME() \
+	FGpuProfilerTrace::EndFrame();
+
+#else
+
+struct FGpuProfilerTrace
+{
+	struct FEventType
+	{
+		FEventType(const TCHAR* Name) {};
+	};
+};
+
+#define TRACE_GPUPROFILER_DEFINE_EVENT_TYPE(...)
+#define TRACE_GPUPROFILER_DECLARE_EVENT_TYPE_EXTERN(...)
+#define TRACE_GPUPROFILER_EVENT_TYPE(...) nullptr
+#define TRACE_GPUPROFILER_BEGIN_FRAME(...)
+#define TRACE_GPUPROFILER_BEGIN_EVENT(...)
+#define TRACE_GPUPROFILER_END_EVENT(...)
+#define TRACE_GPUPROFILER_END_FRAME(...)
+
+#endif
diff --git a/Engine/Source/Runtime/RenderCore/Private/RenderingThread.cpp b/Engine/Source/Runtime/RenderCore/Private/RenderingThread.cpp
--- a/Engine/Source/Runtime/RenderCore/Private/RenderingThread.cpp
+++ b/Engine/Source/Runtime/RenderCore/Private/RenderingThread.cpp
@@ -23,6 +23,7 @@
 #include "RenderResource.h"
 #include "Misc/ScopeLock.h"
 #include "HAL/LowLevelMemTracker.h"
+#include "ProfilingDebugging/MiscTrace.h"
 
 //
 // Globals
@@ -701,6 +702,7 @@ void StartRenderingThread()
 	GRenderingThreadRunnable = new FRenderingThread();
 
 	GRenderingThread = FRunnableThread::Create(GRenderingThreadRunnable, *BuildRenderingThreadName(ThreadCount), 0, FPlatformAffinity::GetRenderingThreadPriority(), FPlatformAffinity::GetRenderingThreadMask());
+	TRACE_SET_THREAD_GROUP(GRenderingThread->GetThreadID(), "Render");
 
 	// Wait for render thread to have taskgraph bound before we dispatch any tasks for it.
 	((FRenderingThread*)GRenderingThreadRunnable)->TaskGraphBoundSyncEvent->Wait();
diff --git a/Engine/Source/Runtime/Renderer/Private/RendererScene.cpp b/Engine/Source/Runtime/Renderer/Private/RendererScene.cpp
--- a/Engine/Source/Runtime/Renderer/Private/RendererScene.cpp
+++ b/Engine/Source/Runtime/Renderer/Private/RendererScene.cpp
@@ -559,8 +559,13 @@ static TAutoConsoleVariable<int32> CVarDoLazyStaticMeshUpdate(
 
 static void DoLazyStaticMeshUpdateCVarSinkFunction()
 {
-	static bool CachedDoLazyStaticMeshUpdate = CVarDoLazyStaticMeshUpdate.GetValueOnGameThread() && !GIsEditor && FApp::CanEverRender();
-	bool DoLazyStaticMeshUpdate = CVarDoLazyStaticMeshUpdate.GetValueOnGameThread() && !GIsEditor && FApp::CanEverRender();
+	if (!GIsRunning || GIsEditor || !FApp::CanEverRender())
+	{
+		return;
+	}
+
+	static bool CachedDoLazyStaticMeshUpdate = !!CVarDoLazyStaticMeshUpdate.GetValueOnGameThread();
+	bool DoLazyStaticMeshUpdate = !!CVarDoLazyStaticMeshUpdate.GetValueOnGameThread();
 
 	if (DoLazyStaticMeshUpdate != CachedDoLazyStaticMeshUpdate)
 	{
