- Landscape Layer System : Removed UTexture2D::ShouldMipLevelsBeForcedResident dependency on experimental flag. Now relies on its owner LandscapeProxy's content (with or without layers).
#rnx

diff --git a/Engine/Source/Runtime/Engine/Private/Texture2D.cpp b/Engine/Source/Runtime/Engine/Private/Texture2D.cpp
--- a/Engine/Source/Runtime/Engine/Private/Texture2D.cpp
+++ b/Engine/Source/Runtime/Engine/Private/Texture2D.cpp
@@ -35,6 +35,7 @@
 #include "Async/AsyncFileHandle.h"
 #if WITH_EDITOR
 #include "Settings/EditorExperimentalSettings.h"
+#include "Landscape/Classes/LandscapeProxy.h"
 #endif
 
 UTexture2D::UTexture2D(const FObjectInitializer& ObjectInitializer)
@@ -1174,9 +1175,10 @@ bool UTexture2D::ShouldMipLevelsBeForcedResident() const
 	}
 
 #if WITH_EDITOR
-	if (GIsEditor && GetMutableDefault<UEditorExperimentalSettings>()->bLandscapeLayerSystem && (LODGroup == TEXTUREGROUP_Terrain_Heightmap || LODGroup == TEXTUREGROUP_Terrain_Weightmap))
+	if (GIsEditor && (LODGroup == TEXTUREGROUP_Terrain_Heightmap || LODGroup == TEXTUREGROUP_Terrain_Weightmap))
 	{
-		return true;
+		ALandscapeProxy* Proxy = Cast<ALandscapeProxy>(GetOuter());
+		return Proxy && Proxy->HasLayersContent();
 	}
 #endif
 
