Hide Debug Camera Controller HUD when showing  buffer visualization overview mode. Display buffer name in buffer visualization full mode.

diff --git a/Engine/Source/Runtime/Engine/Classes/Engine/DebugCameraController.h b/Engine/Source/Runtime/Engine/Classes/Engine/DebugCameraController.h
--- a/Engine/Source/Runtime/Engine/Classes/Engine/DebugCameraController.h
+++ b/Engine/Source/Runtime/Engine/Classes/Engine/DebugCameraController.h
@@ -53,6 +53,10 @@ class ENGINE_API ADebugCameraController
 	UPROPERTY()
 	uint32 bIsBufferVisualizationInputSetup : 1;
 
+	/** Last display enabled setting before toggling buffer visualization overview */
+	UPROPERTY()
+	uint32 bLastDisplayEnabled : 1;
+
 	/** Visualizes the frustum of the camera */
 	UPROPERTY()
 	class UDrawFrustumComponent* DrawFrustum;
@@ -83,6 +87,12 @@ class ENGINE_API ADebugCameraController
 	UFUNCTION(BlueprintCallable, Category="Debug Camera")
 	void ToggleDisplay();
 
+	/** Sets display of debug info and input commands. */
+	void SetDisplay(bool bEnabled);
+
+	/** Returns whether debug info display is enabled */
+	bool IsDisplayEnabled();
+
 	/**
 	 * function called from key bindings command to save information about
 	 * turning on/off FreezeRendering command.
@@ -150,6 +160,9 @@ class ENGINE_API ADebugCameraController
 
 #endif
 
+	/** Get selected visualization buffer */
+	FString GetSelectedBufferName();
+
 public:
 
 	/** Currently selected actor, may be invalid */
@@ -295,9 +308,9 @@ private:
 
 	/** Last index in settings array for cycle view modes */
 	int32 LastViewModeSettingsIndex;
-	
-	/** Last buffer selected in buffer visualization overview */
-	FString LastSelectedBuffer;
+
+	/** Buffer selected in buffer visualization overview or full screen view */
+	FString CurrSelectedBuffer;
 
 	void OnTouchBegin(ETouchIndex::Type FingerIndex, FVector Location);
 	void OnTouchEnd(ETouchIndex::Type FingerIndex, FVector Location);
diff --git a/Engine/Source/Runtime/Engine/Private/DebugCameraController.cpp b/Engine/Source/Runtime/Engine/Private/DebugCameraController.cpp
--- a/Engine/Source/Runtime/Engine/Private/DebugCameraController.cpp
+++ b/Engine/Source/Runtime/Engine/Private/DebugCameraController.cpp
@@ -748,7 +748,6 @@ void ADebugCameraController::ToggleBufferVisualizationOverviewMode()
 	if (UGameViewportClient* GameViewportClient = GetWorld()->GetGameViewport())
 	{
 		bEnableBufferVisualization = !bEnableBufferVisualization;
-		bEnableBufferVisualizationFullMode = false;
 
 		FString Cmd(TEXT("VIEWMODE "));
 
@@ -759,14 +758,18 @@ void ADebugCameraController::ToggleBufferVisualizationOverviewMode()
 
 			TArray<FString> SelectedBuffers = GetBufferVisualizationOverviewTargets();
 
-			if (LastSelectedBuffer.IsEmpty() || !SelectedBuffers.Contains(LastSelectedBuffer))
+			if (CurrSelectedBuffer.IsEmpty() || !SelectedBuffers.Contains(CurrSelectedBuffer))
 			{
 				GetNextBuffer(SelectedBuffers, 1);
 			}
+
+			bLastDisplayEnabled = IsDisplayEnabled();
+			SetDisplay(false);
 		}
 		else
 		{
 			Cmd += GetViewModeName((EViewModeIndex)LastViewModeIndex);
+			SetDisplay(bLastDisplayEnabled);
 		}
 
 		GameViewportClient->ConsoleCommand(Cmd);
@@ -781,11 +784,11 @@ void ADebugCameraController::UpdateVisualizeBufferPostProcessing(FFinalPostProce
 {
 	if (bEnableBufferVisualization)
 	{
-		FName TargetBufferName = *LastSelectedBuffer;
-		if (UMaterial* Material = GetBufferVisualizationData().GetMaterial(TargetBufferName))
+		FString BufferMaterialName = GetSelectedBufferName();
+		if (!BufferMaterialName.IsEmpty())
 		{
 			InOutPostProcessingSettings.bBufferVisualizationOverviewTargetIsSelected = true;
-			InOutPostProcessingSettings.BufferVisualizationOverviewSelectedTargetMaterialName = Material->GetName();
+			InOutPostProcessingSettings.BufferVisualizationOverviewSelectedTargetMaterialName = BufferMaterialName;
 			return;
 		}
 	}
@@ -811,13 +814,13 @@ void ADebugCameraController::GetNextBuffer(const TArray<FString>& OverviewBuffer
 	{
 		int32 BufferIndex = 0;
 
-		if (!LastSelectedBuffer.IsEmpty())
+		if (!CurrSelectedBuffer.IsEmpty())
 		{
 			bool bFoundIndex = false;
 
 			for (int32 Index = 0; Index < OverviewBuffers.Num(); Index++)
 			{
-				if (OverviewBuffers[Index] == LastSelectedBuffer)
+				if (OverviewBuffers[Index] == CurrSelectedBuffer)
 				{
 					BufferIndex = Index;
 					bFoundIndex = true;
@@ -827,17 +830,17 @@ void ADebugCameraController::GetNextBuffer(const TArray<FString>& OverviewBuffer
 
 			if (!bFoundIndex)
 			{
-				LastSelectedBuffer.Empty();
+				CurrSelectedBuffer.Empty();
 			}
 		}
 
-		if (LastSelectedBuffer.IsEmpty())
+		if (CurrSelectedBuffer.IsEmpty())
 		{
 			for (FString Buffer : OverviewBuffers)
 			{
 				if (!Buffer.IsEmpty())
 				{
-					LastSelectedBuffer = Buffer;
+					CurrSelectedBuffer = Buffer;
 					break;
 				}
 			}
@@ -869,7 +872,7 @@ void ADebugCameraController::GetNextBuffer(const TArray<FString>& OverviewBuffer
 			{
 				if (!OverviewBuffers[NextIndex].IsEmpty())
 				{
-					LastSelectedBuffer = OverviewBuffers[NextIndex];
+					CurrSelectedBuffer = OverviewBuffers[NextIndex];
 					break;
 				}
 				NextIndex = Wrap(NextIndex + Step);
@@ -898,6 +901,16 @@ void ADebugCameraController::BufferVisualizationMoveLeft()
 	GetNextBuffer(-1);
 }
 
+FString ADebugCameraController::GetSelectedBufferName()
+{
+	if (UMaterial* Material = GetBufferVisualizationData().GetMaterial(*CurrSelectedBuffer))
+	{
+		return Material->GetName();
+	}
+
+	return TEXT("");
+}
+
 void ADebugCameraController::ConsumeAxisMotion(float Val)
 {
 	// Just ignore the axis motion.
@@ -906,18 +919,23 @@ void ADebugCameraController::ConsumeAxisMotion(float Val)
 void ADebugCameraController::ToggleBufferVisualizationFullMode()
 {
 	SetBufferVisualizationFullMode(!bEnableBufferVisualizationFullMode);
-	SetupBufferVisualizationOverviewInput();
 }
 
 void ADebugCameraController::SetBufferVisualizationFullMode(bool bFullMode)
 {
-	bEnableBufferVisualizationFullMode = bFullMode;
-
-	static IConsoleVariable* ICVar = IConsoleManager::Get().FindConsoleVariable(FBufferVisualizationData::GetVisualizationTargetConsoleCommandName());
-	if (ICVar)
+	if (bEnableBufferVisualizationFullMode != bFullMode)
 	{
-		static const FName EmptyName = NAME_None;
-		ICVar->Set(bFullMode ? *LastSelectedBuffer : *EmptyName.ToString(), ECVF_SetByCode);
+		bEnableBufferVisualizationFullMode = bFullMode;
+
+		static IConsoleVariable* ICVar = IConsoleManager::Get().FindConsoleVariable(FBufferVisualizationData::GetVisualizationTargetConsoleCommandName());
+		if (ICVar)
+		{
+			static const FName EmptyName = NAME_None;
+			ICVar->Set(bFullMode ? *CurrSelectedBuffer : *EmptyName.ToString(), ECVF_SetByCode);
+		}
+
+		SetupBufferVisualizationOverviewInput();
+		SetDisplay(bEnableBufferVisualizationFullMode);
 	}
 }
 
@@ -996,3 +1014,16 @@ void ADebugCameraController::ToggleDisplay()
 		MyHUD->ShowHUD();
 	}
 }
+
+bool ADebugCameraController::IsDisplayEnabled()
+{
+	return (MyHUD && MyHUD->bShowHUD);
+}
+
+void ADebugCameraController::SetDisplay(bool bEnabled)
+{
+	if (IsDisplayEnabled() != bEnabled)
+	{
+		ToggleDisplay();
+	}
+}
diff --git a/Engine/Source/Runtime/Engine/Private/DebugCameraHUD.cpp b/Engine/Source/Runtime/Engine/Private/DebugCameraHUD.cpp
--- a/Engine/Source/Runtime/Engine/Private/DebugCameraHUD.cpp
+++ b/Engine/Source/Runtime/Engine/Private/DebugCameraHUD.cpp
@@ -66,6 +66,22 @@ void ADebugCameraHUD::PostRender()
 		UFont* RenderFont = GEngine->GetSmallFont();
 		if( DCC != NULL )
 		{
+			if (DCC->bEnableBufferVisualizationFullMode)
+			{
+				RenderFont = GEngine->GetMediumFont();
+				FFontRenderInfo FontRenderInfo = Canvas->CreateFontRenderInfo(false, true);
+
+				Canvas->SetDrawColor(255, 255, 64, 255);
+				FString BufferText = DCC->GetSelectedBufferName();
+				float xl, yl;
+				Canvas->StrLen(RenderFont, BufferText, xl, yl);
+				float X = Canvas->SizeX * 0.05f;
+				float Y = Canvas->SizeY < 3.0f * yl ? 0.0f : Canvas->SizeY - 3.0f * yl;
+				Canvas->DrawText(RenderFont, BufferText, X, Y, 1.f, 1.f, FontRenderInfo);
+
+				return;
+			}
+	
 			FFontRenderInfo FontRenderInfo = Canvas->CreateFontRenderInfo(false, true);
 
 			Canvas->SetDrawColor(64, 64, 255, 255);
@@ -199,7 +215,7 @@ void ADebugCameraHUD::PostRender()
 
 
 			// controls display
-			yl += Y*10;
+			yl += Y*8;
 			
 			Canvas->SetDrawColor(64, 64, 255, 255);
 			Canvas->DrawText(RenderFont, TEXT("Controls"), X, yl, 1.f, 1.f, FontRenderInfo);
