DebugCameraController minor edit to initialize member var.

diff --git a/Engine/Source/Runtime/Engine/Private/DebugCameraController.cpp b/Engine/Source/Runtime/Engine/Private/DebugCameraController.cpp
--- a/Engine/Source/Runtime/Engine/Private/DebugCameraController.cpp
+++ b/Engine/Source/Runtime/Engine/Private/DebugCameraController.cpp
@@ -60,6 +60,7 @@ ADebugCameraController::ADebugCameraController(const FObjectInitializer& ObjectI
 	bEnableBufferVisualization = false;
 	bEnableBufferVisualizationFullMode = false;
 	bIsBufferVisualizationInputSetup = false;
+	bLastDisplayEnabled = true;
 	LastViewModeIndex = VMI_Lit;
 	LastViewModeSettingsIndex = 0;
 }
