Fix static analysis warning, make creation/lifetime of WaveInstance explicit and mirror 'Remove' call implementation in ActiveSound

diff --git a/Engine/Plugins/Runtime/SoundMod/Source/SoundMod/Private/SoundMod.cpp b/Engine/Plugins/Runtime/SoundMod/Source/SoundMod/Private/SoundMod.cpp
--- a/Engine/Plugins/Runtime/SoundMod/Source/SoundMod/Private/SoundMod.cpp
+++ b/Engine/Plugins/Runtime/SoundMod/Source/SoundMod/Private/SoundMod.cpp
@@ -22,12 +22,12 @@ void USoundMod::Parse(class FAudioDevice* AudioDevice, const UPTRINT NodeWaveIns
 	FWaveInstance* WaveInstance = ActiveSound.FindWaveInstance(NodeWaveInstanceHash);
 
 	// Create a new WaveInstance if this SoundWave doesn't already have one associated with it.
-	if (WaveInstance == NULL)
+	if (!WaveInstance)
 	{
 		const int32 SampleRate = 44100;
 
 		// Create a new wave instance and associate with the ActiveSound
-		WaveInstance = new FWaveInstance(NodeWaveInstanceHash, ActiveSound);
+		WaveInstance = &ActiveSound.AddWaveInstance(NodeWaveInstanceHash);
 
 		// Create streaming wave object
 		USoundModWave* ModWave = NewObject<USoundModWave>();
@@ -36,7 +36,7 @@ void USoundMod::Parse(class FAudioDevice* AudioDevice, const UPTRINT NodeWaveIns
 		ModWave->Duration = INDEFINITELY_LOOPING_DURATION;
 		ModWave->bLooping = bLooping;
 
-		if (ResourceData == NULL)
+		if (!ResourceData)
 		{
 			RawData.GetCopy((void**)&ResourceData, true);
 		}
diff --git a/Engine/Source/Runtime/Engine/Classes/Sound/SoundWave.h b/Engine/Source/Runtime/Engine/Classes/Sound/SoundWave.h
--- a/Engine/Source/Runtime/Engine/Classes/Sound/SoundWave.h
+++ b/Engine/Source/Runtime/Engine/Classes/Sound/SoundWave.h
@@ -677,7 +677,7 @@ public:
 	/**
 	 * Handle any special requirements when the sound starts (e.g. subtitles)
 	 */
-	FWaveInstance* HandleStart( FActiveSound& ActiveSound, const UPTRINT WaveInstanceHash ) const;
+	FWaveInstance& HandleStart(FActiveSound& ActiveSound, const UPTRINT WaveInstanceHash) const;
 
 	/**
 	 * This is only used for DTYPE_Procedural audio. It's recommended to use USynthComponent base class
diff --git a/Engine/Source/Runtime/Engine/Private/ActiveSound.cpp b/Engine/Source/Runtime/Engine/Private/ActiveSound.cpp
--- a/Engine/Source/Runtime/Engine/Private/ActiveSound.cpp
+++ b/Engine/Source/Runtime/Engine/Private/ActiveSound.cpp
@@ -1090,9 +1090,11 @@ void FActiveSound::HandleInteriorVolumes( const FListener& Listener, FSoundParse
 	}
 }
 
-void FActiveSound::AddWaveInstance(const UPTRINT WaveInstanceHash, FWaveInstance& WaveInstance)
+FWaveInstance& FActiveSound::AddWaveInstance(const UPTRINT WaveInstanceHash)
 {
-	WaveInstances.Add(WaveInstanceHash, &WaveInstance);
+	FWaveInstance* WaveInstance = new FWaveInstance(WaveInstanceHash, *this);
+	WaveInstances.Add(WaveInstanceHash, WaveInstance);
+	return *WaveInstance;
 }
 
 void FActiveSound::ApplyRadioFilter(const FSoundParseParameters& ParseParams )
diff --git a/Engine/Source/Runtime/Engine/Private/Audio.cpp b/Engine/Source/Runtime/Engine/Private/Audio.cpp
--- a/Engine/Source/Runtime/Engine/Private/Audio.cpp
+++ b/Engine/Source/Runtime/Engine/Private/Audio.cpp
@@ -816,7 +816,6 @@ FWaveInstance::FWaveInstance(const UPTRINT InWaveInstanceHash, FActiveSound& InA
 	, UserIndex(0)
 {
 	TypeHash = ++TypeHashCounter;
-	ActiveSound->AddWaveInstance(WaveInstanceHash, *this);
 }
 
 /**
diff --git a/Engine/Source/Runtime/Engine/Private/SoundWave.cpp b/Engine/Source/Runtime/Engine/Private/SoundWave.cpp
--- a/Engine/Source/Runtime/Engine/Private/SoundWave.cpp
+++ b/Engine/Source/Runtime/Engine/Private/SoundWave.cpp
@@ -1383,10 +1383,10 @@ bool USoundWave::CleanupDecompressor(bool bForceWait)
 	return false;
 }
 
-FWaveInstance* USoundWave::HandleStart( FActiveSound& ActiveSound, const UPTRINT WaveInstanceHash ) const
+FWaveInstance& USoundWave::HandleStart( FActiveSound& ActiveSound, const UPTRINT WaveInstanceHash ) const
 {
 	// Create a new wave instance and associate with the ActiveSound
-	FWaveInstance* WaveInstance = new FWaveInstance(WaveInstanceHash, ActiveSound);
+	FWaveInstance& WaveInstance = ActiveSound.AddWaveInstance(WaveInstanceHash);
 
 	// Add in the subtitle if they exist
 	if (ActiveSound.bHandleSubtitles && Subtitles.Num() > 0)
@@ -1395,7 +1395,7 @@ FWaveInstance* USoundWave::HandleStart( FActiveSound& ActiveSound, const UPTRINT
 		{
 			QueueSubtitleParams.AudioComponentID = ActiveSound.GetAudioComponentID();
 			QueueSubtitleParams.WorldPtr = ActiveSound.GetWeakWorld();
-			QueueSubtitleParams.WaveInstance = (PTRINT)WaveInstance;
+			QueueSubtitleParams.WaveInstance = (PTRINT)&WaveInstance;
 			QueueSubtitleParams.SubtitlePriority = ActiveSound.SubtitlePriority;
 			QueueSubtitleParams.Duration = Duration;
 			QueueSubtitleParams.bManualWordWrap = bManualWordWrap;
@@ -1464,7 +1464,7 @@ void USoundWave::Parse(FAudioDevice* AudioDevice, const UPTRINT NodeWaveInstance
 			ActiveSound.ApplyRadioFilter(ParseParams);
 		}
 
-		WaveInstance = HandleStart(ActiveSound, NodeWaveInstanceHash);
+		WaveInstance = &HandleStart(ActiveSound, NodeWaveInstanceHash);
 	}
 
 	// Looping sounds are never actually finished
diff --git a/Engine/Source/Runtime/Engine/Public/ActiveSound.h b/Engine/Source/Runtime/Engine/Public/ActiveSound.h
--- a/Engine/Source/Runtime/Engine/Public/ActiveSound.h
+++ b/Engine/Source/Runtime/Engine/Public/ActiveSound.h
@@ -512,7 +512,7 @@ public:
 	/**
 	 * Add newly created wave instance to active sound
 	 */
-	void AddWaveInstance(const UPTRINT WaveInstanceHash, FWaveInstance& WaveInstance);
+	FWaveInstance& AddWaveInstance(const UPTRINT WaveInstanceHash);
 
 	/**
 	 * Check whether to apply the radio filter
