Remove guard against virtualization for certain concurrency rules.  No longer necessary as ping-ponging can no longer occur due to age-based concurrency rules now factor in the lifetime of VirtualLoops

diff --git a/Engine/Source/Runtime/Engine/Private/SoundConcurrency.cpp b/Engine/Source/Runtime/Engine/Private/SoundConcurrency.cpp
--- a/Engine/Source/Runtime/Engine/Private/SoundConcurrency.cpp
+++ b/Engine/Source/Runtime/Engine/Private/SoundConcurrency.cpp
@@ -705,42 +705,17 @@ FActiveSound* FSoundConcurrencyManager::CreateAndEvictActiveSounds(const FActive
 		TArray<FConcurrencyHandle> Handles;
 		SoundToEvict->GetConcurrencyHandles(Handles);
 
-		bool bAllowVirtual = true;
-		for (const FConcurrencyHandle& Handle : Handles)
-		{
-			switch (Handle.Settings.ResolutionRule)
-			{
-				// Stop oldest resolution rules will cause an undesired
-				// cycling through virtual loops from this frame to the next,
-				// so don't permit virtualization in this case.
-				case EMaxConcurrentResolutionRule::StopOldest:
-				case EMaxConcurrentResolutionRule::StopFarthestThenOldest:
-				{
-					bAllowVirtual = false;
-					break;
-				}
-				default:
-				{
-					continue;
-				}
-			}
-		}
-
 		AudioDevice->AddSoundToStop(SoundToEvict);
-		// If using a mode where we support loop virtualization, attempt to re-trigger and update boolean state.
-		if (bAllowVirtual)
+		const bool bDoRangeCheck = false;
+		FAudioVirtualLoop VirtualLoop;
+		if (FAudioVirtualLoop::Virtualize(*SoundToEvict, bDoRangeCheck, VirtualLoop))
 		{
-			const bool bDoRangeCheck = true;
-			FAudioVirtualLoop VirtualLoop;
-			if (FAudioVirtualLoop::Virtualize(*SoundToEvict, bDoRangeCheck, VirtualLoop))
+			SoundToEvict->ClearAudioComponent();
+			if (USoundBase* Sound = SoundToEvict->GetSound())
 			{
-				SoundToEvict->ClearAudioComponent();
-				if (USoundBase* Sound = SoundToEvict->GetSound())
-				{
-					UE_LOG(LogAudio, Verbose, TEXT("Playing ActiveSound %s Virtualizing: Evicted due to concurrency rules."), *Sound->GetName());
-				}
-				AudioDevice->AddVirtualLoop(VirtualLoop);
+				UE_LOG(LogAudio, Verbose, TEXT("Playing ActiveSound %s Virtualizing: Evicted due to concurrency rules."), *Sound->GetName());
 			}
+			AudioDevice->AddVirtualLoop(VirtualLoop);
 		}
 	}
 
