Subsystem - Making the GetLocalPlayer templated and adding a checked version.

diff --git a/Engine/Source/Runtime/Engine/Private/Subsystems/LocalPlayerSubsystem.cpp b/Engine/Source/Runtime/Engine/Private/Subsystems/LocalPlayerSubsystem.cpp
--- a/Engine/Source/Runtime/Engine/Private/Subsystems/LocalPlayerSubsystem.cpp
+++ b/Engine/Source/Runtime/Engine/Private/Subsystems/LocalPlayerSubsystem.cpp
@@ -6,11 +6,4 @@
 ULocalPlayerSubsystem::ULocalPlayerSubsystem()
 	: USubsystem()
 {
-
 }
-
-ULocalPlayer* ULocalPlayerSubsystem::GetLocalPlayer() const
-{
-	return Cast<ULocalPlayer>(GetOuter());
-}
-
diff --git a/Engine/Source/Runtime/Engine/Public/Subsystems/LocalPlayerSubsystem.h b/Engine/Source/Runtime/Engine/Public/Subsystems/LocalPlayerSubsystem.h
--- a/Engine/Source/Runtime/Engine/Public/Subsystems/LocalPlayerSubsystem.h
+++ b/Engine/Source/Runtime/Engine/Public/Subsystems/LocalPlayerSubsystem.h
@@ -20,6 +20,15 @@ class ENGINE_API ULocalPlayerSubsystem : public USubsystem
 public:
 	ULocalPlayerSubsystem();
 
-	ULocalPlayer* GetLocalPlayer() const;
-
+	template<typename LocalPlayerType = ULocalPlayer>
+	LocalPlayerType* GetLocalPlayer() const
+	{
+		return Cast<LocalPlayerType>(GetOuter());
+	}
+
+	template<typename LocalPlayerType = ULocalPlayer>
+	LocalPlayerType* GetLocalPlayerChecked() const
+	{
+		return CastChecked<LocalPlayerType>(GetOuter());
+	}
 };
