Add config option to allow demo net driver to filter out recording multicast rpcs

diff --git a/Engine/Source/Runtime/Engine/Classes/Engine/DemoNetDriver.h b/Engine/Source/Runtime/Engine/Classes/Engine/DemoNetDriver.h
--- a/Engine/Source/Runtime/Engine/Classes/Engine/DemoNetDriver.h
+++ b/Engine/Source/Runtime/Engine/Classes/Engine/DemoNetDriver.h
@@ -355,6 +355,21 @@ struct ENGINE_API FDemoSavedRepObjectState
 
 typedef TArray<struct FDemoSavedRepObjectState> FDemoSavedPropertyState;
 
+USTRUCT()
+struct FMulticastRecordOptions
+{
+	GENERATED_BODY()
+
+	UPROPERTY()
+	FString FuncPathName;
+
+	UPROPERTY()
+	bool bServerSkip;
+
+	UPROPERTY()
+	bool bClientSkip;
+};
+
 /**
  * Simulated network driver for recording and playing back game sessions.
  */
@@ -590,6 +605,10 @@ private:
 	/** Called during a normal demoFrame*/
 	void TickDemoRecordFrame(float DeltaSeconds);
 
+	/** Config data for multicast RPCs we might want to skip recording. */
+	UPROPERTY(config)
+	TArray<FMulticastRecordOptions> MulticastRecordOptions;
+
 public:
 
 	// UNetDriver interface.
diff --git a/Engine/Source/Runtime/Engine/Private/DemoNetDriver.cpp b/Engine/Source/Runtime/Engine/Private/DemoNetDriver.cpp
--- a/Engine/Source/Runtime/Engine/Private/DemoNetDriver.cpp
+++ b/Engine/Source/Runtime/Engine/Private/DemoNetDriver.cpp
@@ -6316,8 +6316,25 @@ bool UDemoNetDriver::UpdateExternalDataForActor(AActor* Actor)
 
 bool UDemoNetDriver::ShouldReplicateFunction(AActor* Actor, UFunction* Function) const
 {
-	const bool bRecordingMulticast = (Function && Function->FunctionFlags & FUNC_NetMulticast) && IsRecording();
-	return bRecordingMulticast || Super::ShouldReplicateFunction(Actor, Function);
+	bool bShouldRecordMulticast = (Function && Function->FunctionFlags & FUNC_NetMulticast) && IsRecording();
+	if (bShouldRecordMulticast)
+	{
+		FString FuncPathName = GetPathNameSafe(Function);
+		int32 Idx = MulticastRecordOptions.IndexOfByPredicate([FuncPathName](const FMulticastRecordOptions& Options) { return (Options.FuncPathName == FuncPathName); });
+		if (Idx != INDEX_NONE)
+		{
+			if (World && World->IsRecordingClientReplay())
+			{
+				bShouldRecordMulticast = bShouldRecordMulticast && !MulticastRecordOptions[Idx].bClientSkip;
+			}
+			else
+			{
+				bShouldRecordMulticast = bShouldRecordMulticast && !MulticastRecordOptions[Idx].bServerSkip;
+			}
+		}
+	}
+
+	return bShouldRecordMulticast || Super::ShouldReplicateFunction(Actor, Function);
 }
 
 bool UDemoNetDriver::ShouldReplicateActor(AActor* Actor) const
