#UE4 Upgrade material compile errors to Display verbosity so they will show up in more outputs

diff --git a/Engine/Source/Runtime/Engine/Private/Materials/Material.cpp b/Engine/Source/Runtime/Engine/Private/Materials/Material.cpp
--- a/Engine/Source/Runtime/Engine/Private/Materials/Material.cpp
+++ b/Engine/Source/Runtime/Engine/Private/Materials/Material.cpp
@@ -2949,7 +2949,7 @@ void UMaterial::CacheShadersForResources(EShaderPlatform ShaderPlatform, const T
 			for (int32 ErrorIndex = 0; ErrorIndex < CompileErrors.Num(); ErrorIndex++)
 			{
 				// Always log material errors in an unsuppressed category
-				UE_LOG(LogMaterial, Log, TEXT("	%s"), *CompileErrors[ErrorIndex]);
+				UE_LOG(LogMaterial, Display, TEXT("	%s"), *CompileErrors[ErrorIndex]);
 			}
 #endif
 		}
