Added a clarifying comment to SetTimer to clarify that InRate and InFirstDelay are in seconds

diff --git a/Engine/Source/Runtime/Engine/Public/TimerManager.h b/Engine/Source/Runtime/Engine/Public/TimerManager.h
--- a/Engine/Source/Runtime/Engine/Public/TimerManager.h
+++ b/Engine/Source/Runtime/Engine/Public/TimerManager.h
@@ -183,9 +183,9 @@ public:
 	 * @param InOutHandle			If the passed-in handle refers to an existing timer, it will be cleared before the new timer is added. A new handle to the new timer is returned in either case.
 	 * @param InObj					Object to call the timer function on.
 	 * @param InTimerMethod			Method to call when timer fires.
-	 * @param InRate				The amount of time between set and firing.  If <= 0.f, clears existing timers.
+	 * @param InRate				The amount of time (in seconds) between set and firing.  If <= 0.f, clears existing timers.
 	 * @param InbLoop				true to keep firing at Rate intervals, false to fire only once.
-	 * @param InFirstDelay			The time for the first iteration of a looping timer. If < 0.f inRate will be used.
+	 * @param InFirstDelay			The time (in seconds) for the first iteration of a looping timer. If < 0.f InRate will be used.
 	 */
 	template< class UserClass >
 	FORCEINLINE void SetTimer(FTimerHandle& InOutHandle, UserClass* InObj, typename FTimerDelegate::TUObjectMethodDelegate< UserClass >::FMethodPtr InTimerMethod, float InRate, bool InbLoop = false, float InFirstDelay = -1.f)
