Spline Metadata Support

diff --git a/Engine/Source/Editor/ComponentVisualizers/Private/SplineComponentVisualizer.cpp b/Engine/Source/Editor/ComponentVisualizers/Private/SplineComponentVisualizer.cpp
--- a/Engine/Source/Editor/ComponentVisualizers/Private/SplineComponentVisualizer.cpp
+++ b/Engine/Source/Editor/ComponentVisualizers/Private/SplineComponentVisualizer.cpp
@@ -812,9 +812,10 @@ void FSplineComponentVisualizer::DuplicateKey()
 	SelectedKeysSorted.Sort([](int32 A, int32 B) { return A > B; });
 
 	// Insert duplicates into the list, highest index first, so that the lower indices remain the same
-	FInterpCurveVector& SplinePosition = SplineComp->SplineCurves.Position;
-	FInterpCurveQuat& SplineRotation = SplineComp->SplineCurves.Rotation;
-	FInterpCurveVector& SplineScale = SplineComp->SplineCurves.Scale;
+	FInterpCurveVector& SplinePosition = SplineComp->GetSplinePointsPosition();
+	FInterpCurveQuat& SplineRotation = SplineComp->GetSplinePointsRotation();
+	FInterpCurveVector& SplineScale = SplineComp->GetSplinePointsScale();
+	USplineMetadata* SplineMetadata = SplineComp->GetSplinePointsMetadata();
 
 	for (int32 SelectedKeyIndex : SelectedKeysSorted)
 	{
@@ -824,6 +825,11 @@ void FSplineComponentVisualizer::DuplicateKey()
 		SplineRotation.Points.Insert(FInterpCurvePoint<FQuat>(SplineRotation.Points[SelectedKeyIndex]), SelectedKeyIndex);
 		SplineScale.Points.Insert(FInterpCurvePoint<FVector>(SplineScale.Points[SelectedKeyIndex]), SelectedKeyIndex);
 
+		if (SplineMetadata)
+		{
+			SplineMetadata->DuplicatePoint(SelectedKeyIndex);
+		}
+
 		// Adjust input keys of subsequent points
 		for (int Index = SelectedKeyIndex + 1; Index < SplinePosition.Points.Num(); Index++)
 		{
@@ -891,7 +897,8 @@ void FSplineComponentVisualizer::OnAddKey()
 	FInterpCurveVector& SplinePosition = SplineComp->GetSplinePointsPosition();
 	FInterpCurveQuat& SplineRotation = SplineComp->GetSplinePointsRotation();
 	FInterpCurveVector& SplineScale = SplineComp->GetSplinePointsScale();
-
+	USplineMetadata* SplineMetadata = SplineComp->GetSplinePointsMetadata();
+	
 	FInterpCurvePoint<FVector> NewPoint(
 		SelectedSegmentIndex,
 		SplineComp->GetComponentTransform().InverseTransformPosition(SelectedSplinePosition),
@@ -916,6 +923,10 @@ void FSplineComponentVisualizer::OnAddKey()
 	SplinePosition.Points.Insert(NewPoint, SelectedSegmentIndex + 1);
 	SplineRotation.Points.Insert(NewRotPoint, SelectedSegmentIndex + 1);
 	SplineScale.Points.Insert(NewScalePoint, SelectedSegmentIndex + 1);
+	if (SplineMetadata)
+	{
+		SplineMetadata->InsertPoint(SelectedSegmentIndex, SelectedSegmentIndex + 1);
+	}
 
 	// Adjust input keys of subsequent points
 	for (int Index = SelectedSegmentIndex + 1; Index < SplinePosition.Points.Num(); Index++)
@@ -964,12 +975,18 @@ void FSplineComponentVisualizer::OnDeleteKey()
 	SelectedKeysSorted.Sort([](int32 A, int32 B) { return A > B; });
 
 	// Delete selected keys from list, highest index first
-	FInterpCurveVector& SplinePosition = SplineComp->SplineCurves.Position;
-	FInterpCurveQuat& SplineRotation = SplineComp->SplineCurves.Rotation;
-	FInterpCurveVector& SplineScale = SplineComp->SplineCurves.Scale;
-
+	FInterpCurveVector& SplinePosition = SplineComp->GetSplinePointsPosition();
+	FInterpCurveQuat& SplineRotation = SplineComp->GetSplinePointsRotation();
+	FInterpCurveVector& SplineScale = SplineComp->GetSplinePointsScale();
+	USplineMetadata* SplineMetadata = SplineComp->GetSplinePointsMetadata();
+		
 	for (int32 SelectedKeyIndex : SelectedKeysSorted)
 	{
+		if (SplineMetadata)
+		{
+			SplineMetadata->RemovePoint(SelectedKeyIndex);
+		}
+		
 		SplinePosition.Points.RemoveAt(SelectedKeyIndex);
 		SplineRotation.Points.RemoveAt(SelectedKeyIndex);
 		SplineScale.Points.RemoveAt(SelectedKeyIndex);
diff --git a/Engine/Source/Editor/DetailCustomizations/Private/SplineComponentDetails.cpp b/Engine/Source/Editor/DetailCustomizations/Private/SplineComponentDetails.cpp
--- a/Engine/Source/Editor/DetailCustomizations/Private/SplineComponentDetails.cpp
+++ b/Engine/Source/Editor/DetailCustomizations/Private/SplineComponentDetails.cpp
@@ -1,6 +1,7 @@
 // Copyright 1998-2019 Epic Games, Inc. All Rights Reserved.
 
 #include "SplineComponentDetails.h"
+#include "SplineMetadataDetailsFactory.h"
 #include "Misc/MessageDialog.h"
 #include "UObject/ObjectMacros.h"
 #include "UObject/Class.h"
@@ -28,6 +29,12 @@
 
 #define LOCTEXT_NAMESPACE "SplineComponentDetails"
 
+USplineMetadataDetailsFactoryBase::USplineMetadataDetailsFactoryBase(const FObjectInitializer& ObjectInitializer)
+	: Super(ObjectInitializer)
+{
+
+}
+
 class FSplinePointDetails : public IDetailCustomNodeBuilder, public TSharedFromThis<FSplinePointDetails>
 {
 public:
@@ -185,6 +192,7 @@ private:
 	FSplineComponentVisualizer* SplineVisualizer;
 	UProperty* SplineCurvesProperty;
 	TArray<TSharedPtr<FString>> SplinePointTypes;
+	TSharedPtr<ISplineMetadataDetails> SplineMetaDataDetails;
 };
 
 FSplinePointDetails::FSplinePointDetails()
@@ -407,6 +415,32 @@ void FSplinePointDetails::GenerateChildContent(IDetailChildrenBuilder& ChildrenB
 			.Text(this, &FSplinePointDetails::GetPointType)
 		]
 	];
+
+	TArray<TWeakObjectPtr<UObject>> SelectedObjects;
+	ChildrenBuilder.GetParentCategory().GetParentLayout().GetObjectsBeingCustomized(SelectedObjects);
+	if (SelectedObjects.Num() == 1)
+	{
+		if (USplineComponent* SelectedSplineComp = Cast<USplineComponent>(SelectedObjects[0]))
+		{
+			if (SelectedSplineComp->GetSplinePointsMetadata())
+			{
+				for (TObjectIterator<UClass> ClassIterator; ClassIterator; ++ClassIterator)
+				{
+					if (ClassIterator->IsChildOf(USplineMetadataDetailsFactoryBase::StaticClass()) && !ClassIterator->HasAnyClassFlags(CLASS_Abstract | CLASS_Deprecated | CLASS_NewerVersionExists))
+					{
+						USplineMetadataDetailsFactoryBase* Factory = ClassIterator->GetDefaultObject<USplineMetadataDetailsFactoryBase>();
+						if (Factory->GetMetadataClass() == SelectedSplineComp->GetSplinePointsMetadata()->GetClass())
+						{
+							SplineMetaDataDetails = Factory->Create();
+							IDetailGroup& Group = ChildrenBuilder.AddGroup(SplineMetaDataDetails->GetName(), SplineMetaDataDetails->GetDisplayName());
+							SplineMetaDataDetails->GenerateChildContent(Group);
+							break;
+						}
+					}
+				}
+			}
+		}
+	}
 }
 
 void FSplinePointDetails::Tick(float DeltaTime)
@@ -418,7 +452,7 @@ void FSplinePointDetails::UpdateValues()
 {
 	SplineComp = SplineVisualizer->GetEditedSplineComponent();
 	SelectedKeys = SplineVisualizer->GetSelectedKeys();
-
+		
 	// Cache values to be shown by the details customization.
 	// An unset optional value represents 'multiple values' (in the case where multiple points are selected).
 	InputKey.Reset();
@@ -442,6 +476,11 @@ void FSplinePointDetails::UpdateValues()
 			PointType.Add(ConvertInterpCurveModeToSplinePointType(SplineComp->GetSplinePointsPosition().Points[Index].InterpMode));
 		}
 	}
+
+	if (SplineMetaDataDetails)
+	{
+		SplineMetaDataDetails->Update(SplineComp, SelectedKeys);
+	}
 }
 
 FName FSplinePointDetails::GetName() const
diff --git a/Engine/Source/Editor/DetailCustomizations/Public/SplineMetadataDetailsFactory.h b/Engine/Source/Editor/DetailCustomizations/Public/SplineMetadataDetailsFactory.h
new file mode 100644
--- /dev/null
+++ b/Engine/Source/Editor/DetailCustomizations/Public/SplineMetadataDetailsFactory.h
@@ -0,0 +1,30 @@
+// Copyright 1998-2019 Epic Games, Inc. All Rights Reserved.
+
+#pragma once
+
+#include "CoreMinimal.h"
+#include "UObject/ObjectMacros.h"
+#include "SplineMetadataDetailsFactory.generated.h"
+
+class USplineComponent;
+class IDetailGroup;
+
+class ISplineMetadataDetails
+{
+public:
+	virtual FName GetName() const = 0;
+	virtual FText GetDisplayName() const = 0;
+	virtual void Update(USplineComponent* InSplineComponent, const TSet<int32>& InSelectedKeys) = 0;
+	virtual void GenerateChildContent(IDetailGroup& InGroup) = 0;
+};
+
+UCLASS(Abstract)
+class DETAILCUSTOMIZATIONS_API USplineMetadataDetailsFactoryBase : public UObject
+{
+	GENERATED_UCLASS_BODY()
+
+public:
+	virtual TSharedPtr<ISplineMetadataDetails> Create() PURE_VIRTUAL(USplineMetadataDetailsFactoryBase::Create, return nullptr;);
+	virtual UClass* GetMetadataClass() const PURE_VIRTUAL(USplineMetadataDetailsFactoryBase::GetMetadataClass, return nullptr;);
+};
+
diff --git a/Engine/Source/Runtime/Engine/Classes/Components/SplineComponent.h b/Engine/Source/Runtime/Engine/Classes/Components/SplineComponent.h
--- a/Engine/Source/Runtime/Engine/Classes/Components/SplineComponent.h
+++ b/Engine/Source/Runtime/Engine/Classes/Components/SplineComponent.h
@@ -37,6 +37,19 @@ namespace ESplineCoordinateSpace
 	};
 }
 
+UCLASS(Abstract)
+class ENGINE_API USplineMetadata : public UObject
+{
+	GENERATED_UCLASS_BODY()
+
+public:
+	virtual void InsertPoint(float InputKey, int32 Index) PURE_VIRTUAL(USplineMetadata::InsertPoint, );
+	virtual void AddPoint(float InputKey) PURE_VIRTUAL(USplineMetadata::AddPoint, );
+	virtual void RemovePoint(int32 Index) PURE_VIRTUAL(USplineMetadata::RemovePoint, );
+	virtual void DuplicatePoint(int32 Index) PURE_VIRTUAL(USplineMetadata::DuplicatePoint, );
+	virtual void Reset(int32 NumPoints) PURE_VIRTUAL(USplineMetadata::Reset, );
+};
+
 USTRUCT()
 struct ENGINE_API FSplineCurves
 {
@@ -58,6 +71,9 @@ struct ENGINE_API FSplineCurves
 	UPROPERTY()
 	FInterpCurveFloat ReparamTable;
 
+	UPROPERTY()
+	USplineMetadata* Metadata = nullptr;
+
 	bool operator==(const FSplineCurves& Other) const
 	{
 		return Position == Other.Position && Rotation == Other.Rotation && Scale == Other.Scale;
@@ -170,6 +186,9 @@ UCLASS(ClassGroup=Utility, ShowCategories = (Mobility), HideCategories = (Physic
 class ENGINE_API USplineComponent : public UPrimitiveComponent
 {
 	GENERATED_UCLASS_BODY()
+    
+    /** Child class can optionally provide their metadata object through this constructor */
+	USplineComponent(const FObjectInitializer& ObjectInitializer, USplineMetadata* Metadata);
 
 	UPROPERTY(EditAnywhere, Category=Points)
 	FSplineCurves SplineCurves;
@@ -294,6 +313,8 @@ public:
 	const FInterpCurveQuat& GetSplinePointsRotation() const { return SplineCurves.Rotation; }
 	FInterpCurveVector& GetSplinePointsScale() { return SplineCurves.Scale; }
 	const FInterpCurveVector& GetSplinePointsScale() const { return SplineCurves.Scale; }
+	USplineMetadata* GetSplinePointsMetadata() { return SplineCurves.Metadata; }
+	const USplineMetadata* GetSplinePointsMetadata() const { return SplineCurves.Metadata; }
 
 	void ApplyComponentInstanceData(struct FSplineInstanceData* ComponentInstanceData, const bool bPostUCS);
 
@@ -494,6 +515,10 @@ public:
 	UFUNCTION(BlueprintCallable, Category=Spline)
 	float GetDistanceAlongSplineAtSplinePoint(int32 PointIndex) const;
 
+    /** Get a metadata property float value along the spline at spline point */
+	UFUNCTION(BlueprintCallable, Category = Spline)
+	float GetFloatPropertyAtSplinePoint(int32 Index, FName PropertyName) const;
+
 	/** Returns total length along this spline */
 	UFUNCTION(BlueprintCallable, Category=Spline) 
 	float GetSplineLength() const;
diff --git a/Engine/Source/Runtime/Engine/Private/Components/SplineComponent.cpp b/Engine/Source/Runtime/Engine/Private/Components/SplineComponent.cpp
--- a/Engine/Source/Runtime/Engine/Private/Components/SplineComponent.cpp
+++ b/Engine/Source/Runtime/Engine/Private/Components/SplineComponent.cpp
@@ -16,8 +16,19 @@ const FInterpCurvePointVector USplineComponent::DummyPointPosition(0.0f, FVector
 const FInterpCurvePointQuat USplineComponent::DummyPointRotation(0.0f, FQuat::Identity);
 const FInterpCurvePointVector USplineComponent::DummyPointScale(0.0f, FVector::OneVector);
 
+USplineMetadata::USplineMetadata(const FObjectInitializer& ObjectInitializer)
+	: Super(ObjectInitializer)
+{
+
+}
 
 USplineComponent::USplineComponent(const FObjectInitializer& ObjectInitializer)
+ : USplineComponent(ObjectInitializer, nullptr)
+{
+
+}
+
+USplineComponent::USplineComponent(const FObjectInitializer& ObjectInitializer, USplineMetadata* Metadata)
 	: Super(ObjectInitializer)
 	, bAllowSplineEditingPerInstance_DEPRECATED(true)
 	, ReparamStepsPerSegment(10)
@@ -37,6 +48,8 @@ USplineComponent::USplineComponent(const FObjectInitializer& ObjectInitializer)
 	, ScaleVisualizationWidth(30.0f)
 #endif
 {
+	SplineCurves.Metadata = Metadata;
+
 	SplineCurves.Position.Points.Reset(10);
 	SplineCurves.Rotation.Points.Reset(10);
 	SplineCurves.Scale.Points.Reset(10);
@@ -44,11 +57,18 @@ USplineComponent::USplineComponent(const FObjectInitializer& ObjectInitializer)
 	SplineCurves.Position.Points.Emplace(0.0f, FVector(0, 0, 0), FVector::ZeroVector, FVector::ZeroVector, CIM_CurveAuto);
 	SplineCurves.Rotation.Points.Emplace(0.0f, FQuat::Identity, FQuat::Identity, FQuat::Identity, CIM_CurveAuto);
 	SplineCurves.Scale.Points.Emplace(0.0f, FVector(1.0f), FVector::ZeroVector, FVector::ZeroVector, CIM_CurveAuto);
-
+	
 	SplineCurves.Position.Points.Emplace(1.0f, FVector(100, 0, 0), FVector::ZeroVector, FVector::ZeroVector, CIM_CurveAuto);
 	SplineCurves.Rotation.Points.Emplace(1.0f, FQuat::Identity, FQuat::Identity, FQuat::Identity, CIM_CurveAuto);
 	SplineCurves.Scale.Points.Emplace(1.0f, FVector(1.0f), FVector::ZeroVector, FVector::ZeroVector, CIM_CurveAuto);
 
+	if (SplineCurves.Metadata)
+	{
+		SplineCurves.Metadata->Reset(10);
+		SplineCurves.Metadata->AddPoint(0.0f);
+		SplineCurves.Metadata->AddPoint(1.0f);
+	}
+
 	UpdateSpline();
 
 	// Set these deprecated values up so that old assets with default values load correctly (and are subsequently upgraded during Serialize)
@@ -562,6 +582,12 @@ void USplineComponent::AddPoint(const FSplinePoint& InSplinePoint, bool bUpdateS
 		),
 		Index);
 
+	USplineMetadata* Metadata = GetSplinePointsMetadata();
+	if (Metadata)
+	{
+		Metadata->AddPoint(InSplinePoint.InputKey);
+	}
+
 	if (bLoopPositionOverride && LoopPosition <= SplineCurves.Position.Points.Last().InVal)
 	{
 		bLoopPositionOverride = false;
@@ -603,6 +629,12 @@ void USplineComponent::AddSplinePoint(const FVector& Position, ESplineCoordinate
 	SplineCurves.Position.Points.Emplace(InKey, TransformedPosition, FVector::ZeroVector, FVector::ZeroVector, CIM_CurveAuto);
 	SplineCurves.Rotation.Points.Emplace(InKey, FQuat::Identity, FQuat::Identity, FQuat::Identity, CIM_CurveAuto);
 	SplineCurves.Scale.Points.Emplace(InKey, FVector(1.0f), FVector::ZeroVector, FVector::ZeroVector, CIM_CurveAuto);
+	USplineMetadata* Metadata = GetSplinePointsMetadata();
+	if (Metadata)
+	{
+		Metadata->AddPoint(InKey);
+	}
+
 
 	if (bLoopPositionOverride)
 	{
@@ -629,6 +661,12 @@ void USplineComponent::AddSplinePointAtIndex(const FVector& Position, int32 Inde
 		SplineCurves.Position.Points.Insert(FInterpCurvePoint<FVector>(InKey, TransformedPosition, FVector::ZeroVector, FVector::ZeroVector, CIM_CurveAuto), Index);
 		SplineCurves.Rotation.Points.Insert(FInterpCurvePoint<FQuat>(InKey, FQuat::Identity, FQuat::Identity, FQuat::Identity, CIM_CurveAuto), Index);
 		SplineCurves.Scale.Points.Insert(FInterpCurvePoint<FVector>(InKey, FVector(1.0f), FVector::ZeroVector, FVector::ZeroVector, CIM_CurveAuto), Index);
+		USplineMetadata* Metadata = GetSplinePointsMetadata();
+		if (Metadata)
+		{
+			Metadata->InsertPoint(InKey, Index);
+		}
+		
 		NumPoints++;
 
 		// Adjust subsequent points' input keys to make room for the value just added
@@ -660,6 +698,12 @@ void USplineComponent::RemoveSplinePoint(int32 Index, bool bUpdateSpline)
 		SplineCurves.Position.Points.RemoveAt(Index, 1, false);
 		SplineCurves.Rotation.Points.RemoveAt(Index, 1, false);
 		SplineCurves.Scale.Points.RemoveAt(Index, 1, false);
+		USplineMetadata* Metadata = GetSplinePointsMetadata();
+		if (Metadata)
+		{
+			Metadata->RemovePoint(Index);
+		}
+
 		NumPoints--;
 
 		// Adjust all following spline point input keys to close the gap left by the removed point
@@ -690,6 +734,12 @@ void USplineComponent::SetSplinePoints(const TArray<FVector>& Points, ESplineCoo
 	SplineCurves.Rotation.Points.Reset(NumPoints);
 	SplineCurves.Scale.Points.Reset(NumPoints);
 
+	USplineMetadata* Metadata = GetSplinePointsMetadata();
+	if (Metadata)
+	{
+		Metadata->Reset(NumPoints);
+	}
+
 	float InputKey = 0.0f;
 	for (const auto& Point : Points)
 	{
@@ -700,6 +750,11 @@ void USplineComponent::SetSplinePoints(const TArray<FVector>& Points, ESplineCoo
 		SplineCurves.Rotation.Points.Emplace(InputKey, FQuat::Identity, FQuat::Identity, FQuat::Identity, CIM_CurveAuto);
 		SplineCurves.Scale.Points.Emplace(InputKey, FVector(1.0f), FVector::ZeroVector, FVector::ZeroVector, CIM_CurveAuto);
 
+		if (Metadata)
+		{
+			Metadata->AddPoint(InputKey);
+		}
+
 		InputKey += 1.0f;
 	}
 
@@ -1356,6 +1411,32 @@ FTransform USplineComponent::FindTransformClosestToWorldLocation(const FVector&
 	return GetTransformAtSplineInputKey(Param, CoordinateSpace, bUseScale);
 }
 
+template<class T>
+T GetPropertyValueAtSplinePoint(const USplineMetadata* Metadata, int32 Index, FName PropertyName)
+{
+	if (Metadata)
+	{
+		if (UProperty* Property = Metadata->GetClass()->FindPropertyByName(PropertyName))
+		{
+			const FInterpCurve<T>* Curve = Property->ContainerPtrToValuePtr<FInterpCurve<T>>(Metadata);
+			const TArray<FInterpCurvePoint<T>>& Points = Curve->Points;
+			const int32 NumPoints = Points.Num();
+			if (NumPoints > 0)
+			{
+				const int32 ClampedIndex = FMath::Clamp(Index, 0, NumPoints - 1);
+				return Points[ClampedIndex].OutVal;
+			}
+		}
+	}
+
+	return T();
+}
+
+float USplineComponent::GetFloatPropertyAtSplinePoint(int32 Index, FName PropertyName) const
+{
+	return GetPropertyValueAtSplinePoint<float>(GetSplinePointsMetadata(), Index, PropertyName);
+}
+
 #if !UE_BUILD_SHIPPING
 FPrimitiveSceneProxy* USplineComponent::CreateSceneProxy()
 {
