Selecting Skin X in the locker crashes the client
#fix Replaced FindChecked with a Find and ptr check, Test builds don't have checks so it called a function on an invalid reference
#misc Made sure we only call set static/dynamic override if the CVar is at the correct value, would otherwise log an error depending on calling order

diff --git a/Engine/Source/Runtime/Engine/Private/Animation/SkinWeightProfile.cpp b/Engine/Source/Runtime/Engine/Private/Animation/SkinWeightProfile.cpp
--- a/Engine/Source/Runtime/Engine/Private/Animation/SkinWeightProfile.cpp
+++ b/Engine/Source/Runtime/Engine/Private/Animation/SkinWeightProfile.cpp
@@ -302,14 +302,17 @@ void FSkinWeightProfilesData::ApplyOverrideProfile(FSkinWeightVertexBuffer* Over
 	const bool bExtraWeights = BaseBuffer->HasExtraBoneInfluences();
 	OverrideBuffer->SetHasExtraBoneInfluences(bExtraWeights);
 
-	const FRuntimeSkinWeightProfileData& Profile = OverrideData.FindChecked(ProfileName);
-	if (bExtraWeights)
+	const FRuntimeSkinWeightProfileData* ProfilePtr = OverrideData.Find(ProfileName);
+	if (ProfilePtr)
 	{
-		Profile.ApplyOverrides<true>(OverrideBuffer, BaseBuffer);
-	}
-	else
-	{
-		Profile.ApplyOverrides<false>(OverrideBuffer, BaseBuffer);
+		if (bExtraWeights)
+		{
+			ProfilePtr->ApplyOverrides<true>(OverrideBuffer, BaseBuffer);
+		}
+		else
+		{
+			ProfilePtr->ApplyOverrides<false>(OverrideBuffer, BaseBuffer);
+		}
 	}
 }
 
diff --git a/Engine/Source/Runtime/Engine/Private/SkeletalMeshLODRenderData.cpp b/Engine/Source/Runtime/Engine/Private/SkeletalMeshLODRenderData.cpp
--- a/Engine/Source/Runtime/Engine/Private/SkeletalMeshLODRenderData.cpp
+++ b/Engine/Source/Runtime/Engine/Private/SkeletalMeshLODRenderData.cpp
@@ -765,10 +765,16 @@ void FSkeletalMeshLODRenderData::Serialize(FArchive& Ar, UObject* Owner, int32 I
 	if (Ar.IsLoading() && !!GSkinWeightProfilesLoadByDefaultMode)
 	{
 #if !WITH_EDITOR
-		// Only allow overriding the base buffer in non-editor builds as it could otherwise be serialized into the asset
-		SkinWeightProfilesData.OverrideBaseBufferSkinWeightData(OwnerMesh, Idx);
+		if ( GSkinWeightProfilesLoadByDefaultMode == 1)
+		{
+			// Only allow overriding the base buffer in non-editor builds as it could otherwise be serialized into the asset
+			SkinWeightProfilesData.OverrideBaseBufferSkinWeightData(OwnerMesh, Idx);
+		}
 #endif 	
-		SkinWeightProfilesData.SetDynamicDefaultSkinWeightProfile(OwnerMesh, Idx);
+		if (GSkinWeightProfilesLoadByDefaultMode == 3)
+		{
+			SkinWeightProfilesData.SetDynamicDefaultSkinWeightProfile(OwnerMesh, Idx);
+		}
 	}
 }
 
