The per-object directional shadow map cast distance now increases with the size of the object when large enough. Added CVar PerObjectCastDistanceRadiusScale to control it. Set to 8 by default.

diff --git a/Engine/Source/Runtime/Engine/Private/Components/DirectionalLightComponent.cpp b/Engine/Source/Runtime/Engine/Private/Components/DirectionalLightComponent.cpp
--- a/Engine/Source/Runtime/Engine/Private/Components/DirectionalLightComponent.cpp
+++ b/Engine/Source/Runtime/Engine/Private/Components/DirectionalLightComponent.cpp
@@ -38,6 +38,12 @@ static TAutoConsoleVariable<float> CVarCSMShadowDistanceFadeoutMultiplier(
 	TEXT("Multiplier for the CSM distance fade"),
 	ECVF_RenderThreadSafe | ECVF_Scalability );
 
+static TAutoConsoleVariable<float> CVarPerObjectCastDistanceRadiusScale(
+	TEXT("r.Shadow.PerObjectCastDistanceRadiusScale"),
+	8.0f,
+	TEXT("PerObjectCastDistanceRadiusScale The scale factor multiplied with the radius of the object to calculate the maximum distance a per-object directional shadow can reach. This will only take effect after a certain (large) radius. Default is 8 times the object radius."),
+	ECVF_RenderThreadSafe
+	);
 
 /**
  * The scene info for a directional light.
@@ -333,7 +339,9 @@ public:
 		OutInitializer.MinLightW = -HALF_WORLD_MAX;
 		// Reduce casting distance on a directional light
 		// This is necessary to improve floating point precision in several places, especially when deriving frustum verts from InvReceiverMatrix
-		OutInitializer.MaxDistanceToCastInLightW = HALF_WORLD_MAX / 32.0f;
+		// This takes the object size into account to ensure that large objects get an extended distance
+		OutInitializer.MaxDistanceToCastInLightW = FMath::Clamp(SubjectBounds.SphereRadius * CVarPerObjectCastDistanceRadiusScale.GetValueOnRenderThread(), (float)HALF_WORLD_MAX / 32.0f, (float)WORLD_MAX);
+
 		return true;
 	}
 
