Move Delta Replication out of FObjectReplication and into RepLayout.
This will help reduce code duplication, public api exposure, and memory usage.

diff --git a/Engine/Source/Runtime/CoreUObject/Public/UObject/CoreNet.h b/Engine/Source/Runtime/CoreUObject/Public/UObject/CoreNet.h
--- a/Engine/Source/Runtime/CoreUObject/Public/UObject/CoreNet.h
+++ b/Engine/Source/Runtime/CoreUObject/Public/UObject/CoreNet.h
@@ -522,7 +522,7 @@ struct FNetDeltaSerializeInfo
 	/** When non-null, this indicates the given Guid has become unmapped and any references to it should be updated. */
 	const FNetworkGUID* MoveGuidToUnmapped = nullptr;
 
-	int32 PropertyRepIndex = INDEX_NONE;
+	uint16 CustomDeltaIndex = INDEX_NONE;
 
 	// Debugging variables
 	FString DebugName;
diff --git a/Engine/Source/Runtime/Engine/Private/DataReplication.cpp b/Engine/Source/Runtime/Engine/Private/DataReplication.cpp
--- a/Engine/Source/Runtime/Engine/Private/DataReplication.cpp
+++ b/Engine/Source/Runtime/Engine/Private/DataReplication.cpp
@@ -205,41 +205,18 @@ public:
 	static bool SendCustomDeltaProperty(
 		const FRepLayout& RepLayout,
 		FNetDeltaSerializeInfo& Params,
-		uint16 CustomDeltaProperty)
+		uint16 CustomDeltaIndex)
 	{
-		return RepLayout.SendCustomDeltaProperty(Params, CustomDeltaProperty);
+		return RepLayout.SendCustomDeltaProperty(Params, CustomDeltaIndex);
 	}
 
 	static bool ReceiveCustomDeltaProperty(
 		const FRepLayout& RepLayout,
+		FReceivingRepState* ReceivingRepState,
 		FNetDeltaSerializeInfo& Params,
-		UStructProperty* ReplicatedProp,
-		uint32& StaticArrayIndex,
-		int32& Offset)
+		UStructProperty* ReplicatedProp)
 	{
-		return RepLayout.ReceiveCustomDeltaProperty(Params, ReplicatedProp, StaticArrayIndex, Offset);
-	}
-
-	static void GatherGuidReferencesForCustomDeltaProperties(const FRepLayout& RepLayout, FNetDeltaSerializeInfo& Params)
-	{
-		RepLayout.GatherGuidReferencesForCustomDeltaProperties(Params);
-	}
-
-	static bool MoveMappedObjectToUnmappedForCustomDeltaProperties(
-		const FRepLayout& RepLayout,
-		FNetDeltaSerializeInfo& Params,
-		TMap<int32, UStructProperty*>& UnmappedCustomProperties)
-	{
-		return RepLayout.MoveMappedObjectToUnmappedForCustomDeltaProperties(Params, UnmappedCustomProperties);
-	}
-
-	static void UpdateUnmappedObjectsForCustomDeltaProperties(
-		const FRepLayout& RepLayout,
-		FNetDeltaSerializeInfo& Params,
-		TArray<TPair<int32, UStructProperty*>>& CompletelyMappedProperties,
-		TArray<TPair<int32, UStructProperty*>>& UpdatedProperties)
-	{
-		RepLayout.UpdateUnmappedObjectsForCustomDeltaProperties(Params, CompletelyMappedProperties, UpdatedProperties);
+		return RepLayout.ReceiveCustomDeltaProperty(ReceivingRepState, Params, ReplicatedProp);
 	}
 
 	static void PreSendCustomDeltaProperties(
@@ -262,6 +239,33 @@ public:
 		RepLayout.PostSendCustomDeltaProperties(Object, Connection, ChangelistMgr, CustomDeltaStates);
 	}
 
+	static uint16 GetNumLifetimeCustomDeltaProperties(const FRepLayout& RepLayout)
+	{
+		return RepLayout.GetNumLifetimeCustomDeltaProperties();
+	}
+
+	static UProperty* GetLifetimeCustomDeltaProperty(const FRepLayout& RepLayout, const uint16 CustomDeltaPropertyIndex)
+	{
+		return RepLayout.GetLifetimeCustomDeltaProperty(CustomDeltaPropertyIndex);
+	}
+
+	static void UpdateChangelistMgr(
+		const FRepLayout& RepLayout,
+		FSendingRepState* RESTRICT RepState,
+		FReplicationChangelistMgr& InChangelistMgr,
+		const UObject* InObject,
+		const uint32 ReplicationFrame,
+		const FReplicationFlags& RepFlags,
+		const bool bForceCompare)
+	{
+		RepLayout.UpdateChangelistMgr(RepState, InChangelistMgr, InObject, ReplicationFrame, RepFlags, bForceCompare);
+	}
+
+	static const ELifetimeCondition GetLifetimeCustomDeltaPropertyCondition(const FRepLayout& RepLayout, const uint16 CustomDeltaPropertyIndex)
+	{
+		return RepLayout.GetLifetimeCustomDeltaPropertyCondition(CustomDeltaPropertyIndex);
+	}
+
 private:
 
 	UNetDriver* Driver;
@@ -271,16 +275,17 @@ private:
 PRAGMA_ENABLE_DEPRECATION_WARNINGS
 
 PRAGMA_DISABLE_DEPRECATION_WARNINGS
-FObjectReplicator::FObjectReplicator() :
-	ObjectClass(nullptr),
-	ObjectPtr(nullptr),
-	bLastUpdateEmpty(false),
-	bOpenAckCalled(false),
-	bForceUpdateUnmapped(false),
-	Connection(nullptr),
-	OwningChannel(nullptr),
-	RepState(nullptr),
-	RemoteFunctions(nullptr)
+FObjectReplicator::FObjectReplicator()
+	: bLastUpdateEmpty(false)
+	, bOpenAckCalled(false)
+	, bForceUpdateUnmapped(false)
+	, bHasReplicatedProperties(false)
+	, bSupportsFastArrayDelta(false)
+	, ObjectClass(nullptr)
+	, ObjectPtr(nullptr)
+	, Connection(nullptr)
+	, OwningChannel(nullptr)
+	, RemoteFunctions(nullptr)
 {
 }
 
@@ -291,9 +296,16 @@ FObjectReplicator::~FObjectReplicator()
 PRAGMA_ENABLE_DEPRECATION_WARNINGS
 
 PRAGMA_DISABLE_DEPRECATION_WARNINGS
-bool FObjectReplicator::SerializeCustomDeltaProperty( UNetConnection * Connection, void* Src, UProperty * Property, uint32 ArrayIndex, FNetBitWriter & OutBunch, TSharedPtr<INetDeltaBaseState> &NewFullState, TSharedPtr<INetDeltaBaseState> & OldState )
+bool FObjectReplicator::SerializeCustomDeltaProperty(
+	UNetConnection* Connection,
+	void* Src,
+	UProperty* Property,
+	uint32 ArrayIndex,
+	FNetBitWriter& OutBunch,
+	TSharedPtr<INetDeltaBaseState>& NewFullState,
+	TSharedPtr<INetDeltaBaseState>& OldState)
 {
-	check( NewFullState.IsValid() == false ); // NewState is passed in as NULL and instantiated within this function if necessary
+	check( NewFullState.IsValid() == false ); // NewState is passed in as nullptr and instantiated within this function if necessary
 
 	SCOPE_CYCLE_COUNTER( STAT_NetSerializeItemDeltaTime );
 
@@ -340,9 +352,9 @@ bool FObjectReplicator::SendCustomDeltaProperty(UObject* InObject, UProperty* Pr
 	return SendCustomDeltaProperty(InObject, Property->RepIndex + ArrayIndex, OutBunch, NewFullState, OldState);
 }
 
-bool FObjectReplicator::SendCustomDeltaProperty(UObject* InObject, uint16 CustomDeltaProperty, FNetBitWriter& OutBunch, TSharedPtr<INetDeltaBaseState>& NewFullState, TSharedPtr<INetDeltaBaseState>& OldState)
+bool FObjectReplicator::SendCustomDeltaProperty(UObject* InObject, uint16 CustomDeltaIndex, FNetBitWriter& OutBunch, TSharedPtr<INetDeltaBaseState>& NewFullState, TSharedPtr<INetDeltaBaseState>& OldState)
 {
-	check(!NewFullState.IsValid()); // NewState is passed in as NULL and instantiated within this function if necessary
+	check(!NewFullState.IsValid()); // NewState is passed in as nullptr and instantiated within this function if necessary
 	check(RepLayout);
 
 	SCOPE_CYCLE_COUNTER(STAT_NetSerializeItemDeltaTime);
@@ -359,11 +371,11 @@ bool FObjectReplicator::SendCustomDeltaProperty(UObject* InObject, uint16 Custom
 	Parms.NewState = &NewFullState;
 	Parms.NetSerializeCB = &NetSerializeCB;
 	Parms.bIsWritingOnClient = ConnectionDriver && ConnectionDriver->GetWorld() && ConnectionDriver->GetWorld()->IsRecordingClientReplay();
-	Parms.PropertyRepIndex = CustomDeltaProperty;
+	Parms.CustomDeltaIndex = CustomDeltaIndex;
 	Parms.bSupportsFastArrayDeltaStructSerialization = bSupportsFastArrayDelta;
 	Parms.Connection = Connection;
 
-	return FNetSerializeCB::SendCustomDeltaProperty(*RepLayout, Parms, CustomDeltaProperty);
+	return FNetSerializeCB::SendCustomDeltaProperty(*RepLayout, Parms, CustomDeltaIndex);
 }
 
 /** 
@@ -405,27 +417,27 @@ void FObjectReplicator::InitRecentProperties(uint8* Source)
 
 	bSupportsFastArrayDelta = !!GSupportsFastArrayDelta;
 
-	// TODO: CDOCustomDeltaState, CheckpointCustomDeltaState, RecentCustomDeltaState, and Retirement could all be moved into SendingRepState.
-	//			This would allow us to skip allocating these containers for receivers completely.
-	//			This logic would also be easily moved to FRepLayout::CreateRepState.
-
-	// We should just update this method to accept an object pointer.
-	UObject* UseObject = reinterpret_cast<UObject*>(Source);
-
-	// Init custom delta property state
-	for (const uint16 CustomDeltaProperty : LocalRepLayout.GetLifetimeCustomDeltaProperties())
+	if (FSendingRepState* SendingRepState = RepState->GetSendingRepState())
 	{
-		FOutBunch DeltaState(Connection->PackageMap);
-		TSharedPtr<INetDeltaBaseState>& NewState = RecentCustomDeltaState.FindOrAdd(CustomDeltaProperty);
-		NewState.Reset();					
+		// We should just update this method to accept an object pointer.
+		UObject* UseObject = reinterpret_cast<UObject*>(Source);
+
+		// Init custom delta property state
+		const uint16 NumLifetimeCustomDeltaProperties = FNetSerializeCB::GetNumLifetimeCustomDeltaProperties(LocalRepLayout);
+		for (uint16 CustomDeltaProperty = 0; CustomDeltaProperty < NumLifetimeCustomDeltaProperties; ++CustomDeltaProperty)
+		{
+			FOutBunch DeltaState(Connection->PackageMap);
+			TSharedPtr<INetDeltaBaseState>& NewState = SendingRepState->RecentCustomDeltaState.FindOrAdd(CustomDeltaProperty);
+			NewState.Reset();
 
-		TSharedPtr<INetDeltaBaseState> OldState;
+			TSharedPtr<INetDeltaBaseState> OldState;
 
-		SendCustomDeltaProperty(UseObject, CustomDeltaProperty, DeltaState, NewState, OldState);
+			SendCustomDeltaProperty(UseObject, CustomDeltaProperty, DeltaState, NewState, OldState);
 
-		// Store the initial delta state in case we need it for when we're asked to resend all data since channel was first opened (bResendAllDataSinceOpen)
-		CDOCustomDeltaState.Add(CustomDeltaProperty, NewState);
-		CheckpointCustomDeltaState.Add(CustomDeltaProperty, NewState);
+			// Store the initial delta state in case we need it for when we're asked to resend all data since channel was first opened (bResendAllDataSinceOpen)
+			SendingRepState->CDOCustomDeltaState.Add(CustomDeltaProperty, NewState);
+			SendingRepState->CheckpointCustomDeltaState.Add(CustomDeltaProperty, NewState);
+		}
 	}
 }
 
@@ -471,32 +483,32 @@ bool FObjectReplicator::ValidateAgainstState( const UObject* ObjectState )
 
 void FObjectReplicator::InitWithObject( UObject* InObject, UNetConnection * InConnection, bool bUseDefaultState )
 {
-	check( GetObject() == NULL );
-	check( ObjectClass == NULL );
+	check( GetObject() == nullptr );
+	check( ObjectClass == nullptr );
 	check( bLastUpdateEmpty == false );
-	check( Connection == NULL );
-	check( OwningChannel == NULL );
+	check( Connection == nullptr );
+	check( OwningChannel == nullptr );
 	check( !RepState.IsValid() );
-	check( RemoteFunctions == NULL );
+	check( RemoteFunctions == nullptr );
 	check( !RepLayout.IsValid() );
 
 	SetObject( InObject );
 
-	if ( GetObject() == NULL )
+	if ( GetObject() == nullptr )
 	{
-		// This may seem weird that we're checking for NULL, but the SetObject above will wrap this object with TWeakObjectPtr
-		// If the object is pending kill, it will switch to NULL, we're just making sure we handle this invalid edge case
-		UE_LOG(LogRep, Error, TEXT("InitWithObject: Object == NULL"));
+		// This may seem weird that we're checking for nullptr, but the SetObject above will wrap this object with TWeakObjectPtr
+		// If the object is pending kill, it will switch to nullptr, we're just making sure we handle this invalid edge case
+		UE_LOG(LogRep, Error, TEXT("InitWithObject: Object == nullptr"));
 		return;
 	}
 
 	ObjectClass					= InObject->GetClass();
 	Connection					= InConnection;
-	RemoteFunctions				= NULL;
+	RemoteFunctions				= nullptr;
 	bHasReplicatedProperties	= false;
 	bOpenAckCalled				= false;
-	RepState					= NULL;
-	OwningChannel				= NULL;		// Initially NULL until StartReplicating is called
+	RepState					= nullptr;
+	OwningChannel				= nullptr;		// Initially nullptr until StartReplicating is called
 	TrackedGuidMemoryBytes		= 0;
 
 	RepLayout = Connection->Driver->GetObjectClassRepLayout( ObjectClass );
@@ -506,16 +518,12 @@ void FObjectReplicator::InitWithObject( UObject* InObject, UNetConnection * InCo
 
 	InitRecentProperties( Source );
 
-PRAGMA_DISABLE_DEPRECATION_WARNINGS
-	RepLayout->GetLifetimeCustomDeltaProperties( LifetimeCustomDeltaProperties, LifetimeCustomDeltaPropertyConditions );
-PRAGMA_ENABLE_DEPRECATION_WARNINGS
-
 	Connection->Driver->AllOwnedReplicators.Add(this);
 }
 
 void FObjectReplicator::CleanUp()
 {
-	if ( OwningChannel != NULL )
+	if ( OwningChannel != nullptr )
 	{
 		StopReplicating( OwningChannel );		// We shouldn't get here, but just in case
 	}
@@ -549,23 +557,14 @@ void FObjectReplicator::CleanUp()
 	ReferencedGuids.Empty();
 	TrackedGuidMemoryBytes = 0;
 
-	SetObject( NULL );
+	SetObject( nullptr );
 
-	ObjectClass					= NULL;
-	Connection					= NULL;
-	RemoteFunctions				= NULL;
+	ObjectClass					= nullptr;
+	Connection					= nullptr;
+	RemoteFunctions				= nullptr;
 	bHasReplicatedProperties	= false;
 	bOpenAckCalled				= false;
 
-PRAGMA_DISABLE_DEPRECATION_WARNINGS
-	// Cleanup custom delta state
-	RecentCustomDeltaState.Empty();
-	CheckpointCustomDeltaState.Empty();
-
-	LifetimeCustomDeltaProperties.Empty();
-	LifetimeCustomDeltaPropertyConditions.Empty();
-PRAGMA_ENABLE_DEPRECATION_WARNINGS
-
 	RepState = nullptr;
 	CheckpointRepState = nullptr;
 }
@@ -607,9 +606,13 @@ void FObjectReplicator::StartReplicating(class UActorChannel * InActorChannel)
 
 	if (ConnectionNetDriver->IsServer() || ConnectionNetDriver->MaySendProperties())
 	{
-		// Allocate retirement list.
-		// SetNum now constructs, so this is safe
-		Retirement.SetNum(ObjectClass->ClassReps.Num());
+		if (FSendingRepState* SendingRepState = RepState.IsValid() ? RepState->GetSendingRepState() : nullptr)
+		{
+			// Allocate retirement list.
+			// SetNum now constructs, so this is safe
+
+			SendingRepState->Retirement.SetNum(ObjectClass->ClassReps.Num());
+		}
 
 		const UWorld* const World = ConnectionNetDriver->GetWorld();
 		UNetDriver* const WorldNetDriver = World ? World->GetNetDriver() : nullptr;
@@ -626,7 +629,7 @@ void FObjectReplicator::StartReplicating(class UActorChannel * InActorChannel)
 	}
 }
 
-static FORCEINLINE void ValidateRetirementHistory( const FPropertyRetirement & Retire, const UObject* Object )
+void ValidateRetirementHistory(const FPropertyRetirement & Retire, const UObject* Object)
 {
 #if !UE_BUILD_SHIPPING
 	checkf( Retire.SanityTag == FPropertyRetirement::ExpectedSanityTag, TEXT( "Invalid Retire.SanityTag. Object: %s" ), *GetFullNameSafe(Object) );
@@ -635,7 +638,7 @@ static FORCEINLINE void ValidateRetirementHistory( const FPropertyRetirement & R
 
 	FPacketIdRange LastRange;
 
-	while ( Rec != NULL )
+	while ( Rec != nullptr )
 	{
 		checkf( Rec->SanityTag == FPropertyRetirement::ExpectedSanityTag, TEXT( "Invalid Rec->SanityTag. Object: %s" ), *GetFullNameSafe(Object) );
 		checkf( Rec->OutPacketIdRange.Last >= Rec->OutPacketIdRange.First, TEXT( "Invalid packet id range (Last < First). Object: %s" ), *GetFullNameSafe(Object) );
@@ -650,38 +653,43 @@ static FORCEINLINE void ValidateRetirementHistory( const FPropertyRetirement & R
 
 void FObjectReplicator::StopReplicating( class UActorChannel * InActorChannel )
 {
-	check( OwningChannel != NULL );
-	check( OwningChannel->Connection == Connection );
-	check( OwningChannel == InActorChannel );
+	check(OwningChannel != nullptr);
+	check(OwningChannel->Connection == Connection);
+	check(OwningChannel == InActorChannel);
 
-	OwningChannel = NULL;
+	OwningChannel = nullptr;
 
 	const UObject* Object = GetObject();
 
-	// Cleanup retirement records
-	for ( int32 i = Retirement.Num() - 1; i >= 0; i-- )
+	if (FSendingRepState* SendingRepState = RepState.IsValid() ? RepState->GetSendingRepState() : nullptr)
 	{
-		ValidateRetirementHistory( Retirement[i], Object );
+		// Cleanup retirement records
+		for (int32 i = SendingRepState->Retirement.Num() - 1; i >= 0; i--)
+		{
+			FPropertyRetirement& Retirement = SendingRepState->Retirement[i];
+			ValidateRetirementHistory(Retirement, Object);
 
-		FPropertyRetirement * Rec = Retirement[i].Next;
-		Retirement[i].Next = NULL;
+			FPropertyRetirement* Rec = Retirement.Next;
+			Retirement.Next = nullptr;
 
-		// We dont need to explicitly delete Retirement[i], but anything in the Next chain needs to be.
-		while ( Rec != NULL )
-		{
-			FPropertyRetirement * Next = Rec->Next;
-			delete Rec;
-			Rec = Next;
+			// We dont need to explicitly delete Retirement, but anything in the Next chain needs to be.
+			while (Rec != nullptr)
+			{
+				FPropertyRetirement* Next = Rec->Next;
+				delete Rec;
+				Rec = Next;
+			}
 		}
+
+		SendingRepState->Retirement.Empty();
 	}
 
-	Retirement.Empty();
 	PendingLocalRPCs.Empty();
 
-	if ( RemoteFunctions != NULL )
+	if (RemoteFunctions != nullptr)
 	{
 		delete RemoteFunctions;
-		RemoteFunctions = NULL;
+		RemoteFunctions = nullptr;
 	}
 }
 
@@ -783,75 +791,88 @@ void FObjectReplicator::ReceivedNak( int32 NakPacketId )
 {
 	const UObject* Object = GetObject();
 
-	if ( Object == NULL )
+	if (Object == nullptr)
 	{
-		UE_LOG(LogNet, Verbose, TEXT("ReceivedNak: Object == NULL"));
+		UE_LOG(LogNet, Verbose, TEXT("FObjectReplicator::ReceivedNak: Object == nullptr"));
 		return;
 	}
-
-	if ( Object != NULL && ObjectClass != NULL )
+	else if (ObjectClass == nullptr)
 	{
-		RepLayout->ReceivedNak( RepState.Get(), NakPacketId );
-
-		for ( int32 i = Retirement.Num() - 1; i >= 0; i-- )
+		UE_LOG(LogNet, Verbose, TEXT("FObjectReplicator::ReceivedNak: ObjectClass == nullptr"));
+	}
+	else if (ERepLayoutState::Normal == RepLayout->GetRepLayoutState())
+	{
+		if (FSendingRepState* SendingRepState = RepState.IsValid() ? RepState->GetSendingRepState() : nullptr)
 		{
-			ValidateRetirementHistory( Retirement[i], Object );
-
-			// If this is a dynamic array property, we have to look through the list of retirement records to see if we need to reset the base state
-			FPropertyRetirement * Rec = Retirement[i].Next; // Retirement[i] is head and not actually used in this case
-			while ( Rec != NULL )
+			// Go over properties tracked with histories, and mark them as needing to be resent.
+			for (int32 i = SendingRepState->HistoryStart; i < SendingRepState->HistoryEnd; ++i)
 			{
-				if ( NakPacketId > Rec->OutPacketIdRange.Last )
+				const int32 HistoryIndex = i % FSendingRepState::MAX_CHANGE_HISTORY;
+
+				FRepChangedHistory& HistoryItem = SendingRepState->ChangeHistory[HistoryIndex];
+
+				if (!HistoryItem.Resend && HistoryItem.OutPacketIdRange.InRange(NakPacketId))
 				{
-					// We can assume this means this record's packet was ack'd, so we can get rid of the old state
-					check( Retirement[i].Next == Rec );
-					Retirement[i].Next = Rec->Next;
-					delete Rec;
-					Rec = Retirement[i].Next;
-					continue;
+					check(HistoryItem.Changed.Num() > 0);
+					HistoryItem.Resend = true;
+					++SendingRepState->NumNaks;
 				}
-				else if ( NakPacketId >= Rec->OutPacketIdRange.First && NakPacketId <= Rec->OutPacketIdRange.Last )
-				{
-					UE_LOG(LogNet, Verbose, TEXT("Restoring Previous Base State of dynamic property. Channel: %s, NakId: %d, First: %d, Last: %d, Address: %s)"), *OwningChannel->Describe(), NakPacketId, Rec->OutPacketIdRange.First, Rec->OutPacketIdRange.Last, *Connection->LowLevelGetRemoteAddress(true));
+			}
 
-					// The Nack'd packet did update this property, so we need to replace the buffer in RecentDynamic
-					// with the buffer we used to create this update (which was dropped), so that the update will be recreated on the next replicate actor
-					if ( Rec->DynamicState.IsValid() )
-					{
-						TSharedPtr<INetDeltaBaseState> & RecentState = RecentCustomDeltaState.FindChecked( i );
-						
-						RecentState.Reset();
-						RecentState = Rec->DynamicState;
-					}
+			// Go over our Custom Delta Properties and update their retirements
+			for (int32 i = SendingRepState->Retirement.Num() - 1; i >= 0; i--)
+			{
+				FPropertyRetirement& Retirement = SendingRepState->Retirement[i];
+				ValidateRetirementHistory(Retirement, Object);
 
-					// We can get rid of the rest of the saved off base states since we will be regenerating these updates on the next replicate actor
-					while ( Rec != NULL )
+				// If this is a dynamic array property, we have to look through the list of retirement records to see if we need to reset the base state
+				FPropertyRetirement* Rec = Retirement.Next; // Retirement[i] is head and not actually used in this case
+				while (Rec != nullptr)
+				{
+					if (NakPacketId > Rec->OutPacketIdRange.Last)
 					{
-						FPropertyRetirement * DeleteNext = Rec->Next;
+						// We can assume this means this record's packet was ack'd, so we can get rid of the old state
+						check(Retirement.Next == Rec);
+						Retirement.Next = Rec->Next;
 						delete Rec;
-						Rec = DeleteNext;
+						Rec = Retirement.Next;
+						continue;
 					}
+					else if (NakPacketId >= Rec->OutPacketIdRange.First && NakPacketId <= Rec->OutPacketIdRange.Last)
+					{
+						UE_LOG(LogNet, Verbose, TEXT("Restoring Previous Base State of dynamic property. Channel: %s, NakId: %d, First: %d, Last: %d, Address: %s)"), *OwningChannel->Describe(), NakPacketId, Rec->OutPacketIdRange.First, Rec->OutPacketIdRange.Last, *Connection->LowLevelGetRemoteAddress(true));
+
+						// The Nack'd packet did update this property, so we need to replace the buffer in RecentDynamic
+						// with the buffer we used to create this update (which was dropped), so that the update will be recreated on the next replicate actor
+						if (Rec->DynamicState.IsValid())
+						{
+							TSharedPtr<INetDeltaBaseState> & RecentState = SendingRepState->RecentCustomDeltaState.FindChecked(i);
+
+							RecentState.Reset();
+							RecentState = Rec->DynamicState;
+						}
 
-					// Finished
-					Retirement[i].Next = NULL;
-					break;				
+						// We can get rid of the rest of the saved off base states since we will be regenerating these updates on the next replicate actor
+						while (Rec != nullptr)
+						{
+							FPropertyRetirement * DeleteNext = Rec->Next;
+							delete Rec;
+							Rec = DeleteNext;
+						}
+
+						// Finished
+						Retirement.Next = nullptr;
+						break;
+					}
+					Rec = Rec->Next;
 				}
-				Rec = Rec->Next;
+
+				ValidateRetirementHistory(Retirement, Object);
 			}
-				
-			ValidateRetirementHistory( Retirement[i], Object );
 		}
 	}
 }
 
-#define HANDLE_INCOMPATIBLE_PROP		\
-	if ( bIsServer )					\
-	{									\
-		return false;					\
-	}									\
-	FieldCache->bIncompatible = true;	\
-	continue;							\
-
 bool FObjectReplicator::ReceivedBunch(FNetBitReader& Bunch, const FReplicationFlags& RepFlags, const bool bHasRepLayout, bool& bOutHasUnmapped)
 {
 	check(RepLayout);
@@ -860,7 +881,7 @@ bool FObjectReplicator::ReceivedBunch(FNetBitReader& Bunch, const FReplicationFl
 
 	if (!Object)
 	{
-		UE_LOG(LogNet, Verbose, TEXT("ReceivedBunch: Object == NULL"));
+		UE_LOG(LogNet, Verbose, TEXT("ReceivedBunch: Object == nullptr"));
 		return false;
 	}
 
@@ -874,12 +895,13 @@ bool FObjectReplicator::ReceivedBunch(FNetBitReader& Bunch, const FReplicationFl
 
 	if (!ClassCache)
 	{
-		UE_LOG(LogNet, Error, TEXT("ReceivedBunch: ClassCache == NULL: %s"), *Object->GetFullName());
+		UE_LOG(LogNet, Error, TEXT("ReceivedBunch: ClassCache == nullptr: %s"), *Object->GetFullName());
 		return false;
 	}
 
 	const FRepLayout& LocalRepLayout = *RepLayout;
 	bool bGuidsChanged = false;
+	FReceivingRepState* ReceivingRepState = RepState->GetReceivingRepState();
 
 	// Handle replayout properties
 	if (bHasRepLayout)
@@ -1012,31 +1034,20 @@ bool FObjectReplicator::ReceivedBunch(FNetBitReader& Bunch, const FReplicationFl
 			Parms.Connection = Connection;
 			Parms.Object = Object;
 
-			uint32 StaticArrayIndex = 0;
-			int32 Offset = 0;
-			if (!FNetSerializeCB::ReceiveCustomDeltaProperty(LocalRepLayout, Parms, ReplicatedProp, StaticArrayIndex, Offset))
-			{
-				// RepLayout should have already logged the error.
-				HANDLE_INCOMPATIBLE_PROP;
-			}
-			else if (UNLIKELY(Reader.IsError()))
+			if (!FNetSerializeCB::ReceiveCustomDeltaProperty(LocalRepLayout, ReceivingRepState, Parms, ReplicatedProp))
 			{
-				UE_LOG(LogNet, Error, TEXT("ReceivedBunch: NetDeltaSerialize - Reader.IsError() == true. Property: %s, Object: %s"), *Parms.DebugName, *Object->GetFullName());
-				HANDLE_INCOMPATIBLE_PROP
-			}
+				// Should have already logged the error.
+				if (bIsServer)
+				{
+					return false;
+				}
 
-			else if (UNLIKELY(Reader.GetBitsLeft() != 0))
-			{
-				UE_LOG( LogNet, Error, TEXT( "ReceivedBunch: NetDeltaSerialize - Mismatch read. Property: %s, Object: %s" ), *Parms.DebugName, *Object->GetFullName() );
-				HANDLE_INCOMPATIBLE_PROP
+				FieldCache->bIncompatible = true;
+				continue;
 			}
 
 			if (Parms.bOutHasMoreUnmapped)
 			{
-PRAGMA_DISABLE_DEPRECATION_WARNINGS
-				UnmappedCustomProperties.Add(Offset, ReplicatedProp);
-PRAGMA_ENABLE_DEPRECATION_WARNINGS
-
 				bOutHasUnmapped = true;
 			}
 
@@ -1047,10 +1058,6 @@ PRAGMA_ENABLE_DEPRECATION_WARNINGS
 
 			// Successfully received it.
 			UE_LOG(LogRepTraffic, Log, TEXT(" %s - %s"), *Object->GetName(), *Parms.DebugName);
-
-			// Notify the Object if this var is RepNotify
-			TArray<uint8> MetaData;
-			QueuePropertyRepNotify(Object, ReplicatedProp, StaticArrayIndex, MetaData);
 		}
 		// Handle function call
 		else if ( Cast< UFunction >( FieldCache->Field ) )
@@ -1265,11 +1272,6 @@ void FObjectReplicator::UpdateGuidToReplicatorMap()
 
 	// Gather guids on rep layout
 	if (RepState.IsValid())
-	{
-		LocalRepLayout.GatherGuidReferences(RepState->GetReceivingRepState(), LocalReferencedGuids, LocalTrackedGuidMemoryBytes);
-	}
-
-	if (UObject* Object = GetObject())
 	{
 		FNetSerializeCB NetSerializeCB(Connection->Driver);
 
@@ -1277,9 +1279,9 @@ void FObjectReplicator::UpdateGuidToReplicatorMap()
 		Parms.NetSerializeCB = &NetSerializeCB;
 		Parms.GatherGuidReferences = &LocalReferencedGuids;
 		Parms.TrackedGuidMemoryBytes = &LocalTrackedGuidMemoryBytes;
-		Parms.Object = Object;
+		Parms.Object = GetObject();
 
-		FNetSerializeCB::GatherGuidReferencesForCustomDeltaProperties(LocalRepLayout, Parms);
+		LocalRepLayout.GatherGuidReferences(RepState->GetReceivingRepState(), Parms, LocalReferencedGuids, LocalTrackedGuidMemoryBytes);
 	}
 
 	// Gather RPC guids
@@ -1331,38 +1333,27 @@ bool FObjectReplicator::MoveMappedObjectToUnmapped(const FNetworkGUID& GUID)
 	check(RepLayout);
 	const FRepLayout& LocalRepLayout = *RepLayout;
 
-	bool bFound = LocalRepLayout.MoveMappedObjectToUnmapped(RepState->GetReceivingRepState(), GUID);
-
-	if (UObject* Object = GetObject())
-	{
-		FNetSerializeCB NetSerializeCB(Connection->Driver);
-
-		FNetDeltaSerializeInfo Parms;
-		Parms.Connection = Connection;
-		Parms.Map = Connection->PackageMap;
-		Parms.Object = Object;
-		Parms.NetSerializeCB = &NetSerializeCB;
-		Parms.MoveGuidToUnmapped = &GUID;
-		Parms.Object = Object;
-
-PRAGMA_DISABLE_DEPRECATION_WARNINGS
-		bFound |= FNetSerializeCB::MoveMappedObjectToUnmappedForCustomDeltaProperties(LocalRepLayout, Parms, UnmappedCustomProperties);
-PRAGMA_ENABLE_DEPRECATION_WARNINGS
-	}
+	FNetSerializeCB NetSerializeCB(Connection->Driver);
+	FNetDeltaSerializeInfo Parms;
+	Parms.Connection = Connection;
+	Parms.Map = Connection->PackageMap;
+	Parms.Object = GetObject();
+	Parms.NetSerializeCB = &NetSerializeCB;
+	Parms.MoveGuidToUnmapped = &GUID;
 
-	return bFound;
+	return LocalRepLayout.MoveMappedObjectToUnmapped(RepState->GetReceivingRepState(), Parms, GUID);
 }
 
 void FObjectReplicator::PostReceivedBunch()
 {
-	if ( GetObject() == NULL )
+	if ( GetObject() == nullptr )
 	{
-		UE_LOG(LogNet, Verbose, TEXT("PostReceivedBunch: Object == NULL"));
+		UE_LOG(LogNet, Verbose, TEXT("PostReceivedBunch: Object == nullptr"));
 		return;
 	}
 
 	// Call PostNetReceive
-	const bool bIsServer = (OwningChannel->Connection->Driver->ServerConnection == NULL);
+	const bool bIsServer = (OwningChannel->Connection->Driver->ServerConnection == nullptr);
 	if (!bIsServer && bHasReplicatedProperties)
 	{
 		PostNetReceive();
@@ -1373,21 +1364,24 @@ void FObjectReplicator::PostReceivedBunch()
 	CallRepNotifies(true);
 }
 
-static FORCEINLINE FPropertyRetirement ** UpdateAckedRetirements( FPropertyRetirement &	Retire, const int32 OutAckPacketId, const UObject* Object )
+static FORCEINLINE FPropertyRetirement** UpdateAckedRetirements(
+	FPropertyRetirement& Retire,
+	const int32 OutAckPacketId,
+	const UObject* Object)
 {
-	ValidateRetirementHistory( Retire, Object );
+	ValidateRetirementHistory(Retire, Object);
 
-	FPropertyRetirement ** Rec = &Retire.Next;	// Note the first element is 'head' that we dont actually use
+	FPropertyRetirement** Rec = &Retire.Next; // Note the first element is 'head' that we don't actually use
 
-	while ( *Rec != NULL )
+	while (*Rec != nullptr)
 	{
-		if ( OutAckPacketId >= (*Rec)->OutPacketIdRange.Last )
+		if (OutAckPacketId >= (*Rec)->OutPacketIdRange.Last)
 		{
 			UE_LOG(LogRepTraffic, Verbose, TEXT("Deleting Property Record (%d >= %d)"), OutAckPacketId, (*Rec)->OutPacketIdRange.Last);
 
 			// They've ack'd this packet so we can ditch this record (easier to do it here than look for these every Ack)
-			FPropertyRetirement * ToDelete = *Rec;
-			check( Retire.Next == ToDelete ); // This should only be able to happen to the first record in the list
+			FPropertyRetirement* ToDelete = *Rec;
+			check(Retire.Next == ToDelete); // This should only be able to happen to the first record in the list
 			Retire.Next = ToDelete->Next;
 			Rec = &Retire.Next;
 
@@ -1407,9 +1401,9 @@ void FObjectReplicator::ReplicateCustomDeltaProperties( FNetBitWriter & Bunch, F
 
 	check(RepLayout);
 	const FRepLayout& LocalRepLayout = *RepLayout;
-	const TArrayView<const uint16> LocalLifetimeCustomDeltaProperties = LocalRepLayout.GetLifetimeCustomDeltaProperties();
+	const int32 NumLifetimeCustomDeltaProperties = FNetSerializeCB::GetNumLifetimeCustomDeltaProperties(LocalRepLayout);
 
-	if (!LocalLifetimeCustomDeltaProperties.Num())
+	if (NumLifetimeCustomDeltaProperties <= 0)
 	{
 		// No custom properties
 		return;
@@ -1419,15 +1413,16 @@ void FObjectReplicator::ReplicateCustomDeltaProperties( FNetBitWriter & Bunch, F
 	//			the receiving end, and make things more consistent.
 
 	UObject* Object = GetObject();
+	FSendingRepState* SendingRepState = RepState->GetSendingRepState();
 
 	check(Object);
 	check(OwningChannel);
 	check(Connection == OwningChannel->Connection);
 
 	TMap<int32, TSharedPtr<INetDeltaBaseState>>& UsingCustomDeltaStates =
-		EResendAllDataState::None == Connection->ResendAllDataState ? RecentCustomDeltaState :
-		EResendAllDataState::SinceOpen == Connection->ResendAllDataState ? CDOCustomDeltaState :
-		CheckpointCustomDeltaState;
+		EResendAllDataState::None == Connection->ResendAllDataState ? SendingRepState->RecentCustomDeltaState :
+		EResendAllDataState::SinceOpen == Connection->ResendAllDataState ? SendingRepState->CDOCustomDeltaState :
+		SendingRepState->CheckpointCustomDeltaState;
 
 	FNetSerializeCB::PreSendCustomDeltaProperties(LocalRepLayout, Object, Connection, *ChangelistMgr, UsingCustomDeltaStates);
 
@@ -1445,16 +1440,9 @@ void FObjectReplicator::ReplicateCustomDeltaProperties( FNetBitWriter & Bunch, F
 	FNetBitWriter TempBitWriter( Connection->PackageMap, 1024 );
 
 	// Replicate those properties.
-	for (const uint16 CustomDeltaProperty : LocalLifetimeCustomDeltaProperties)
+	for (uint16 CustomDeltaProperty = 0; CustomDeltaProperty < NumLifetimeCustomDeltaProperties; ++CustomDeltaProperty)
 	{
-		// Get info.
-		FPropertyRetirement& Retire = Retirement[CustomDeltaProperty];
-
-PRAGMA_DISABLE_DEPRECATION_WARNINGS
-		UProperty* It = LocalRepLayout.GetPropertyForRepIndex(CustomDeltaProperty);
-PRAGMA_ENABLE_DEPRECATION_WARNINGS
-
-		const ELifetimeCondition RepCondition = LocalRepLayout.GetPropertyLifetimeCondition(CustomDeltaProperty);
+		const ELifetimeCondition RepCondition = FNetSerializeCB::GetLifetimeCustomDeltaPropertyCondition(LocalRepLayout, CustomDeltaProperty);
 
 		check(RepCondition >= 0 && RepCondition < COND_Max);
 
@@ -1464,6 +1452,8 @@ PRAGMA_ENABLE_DEPRECATION_WARNINGS
 			continue;
 		}
 
+		UProperty* Property = FNetSerializeCB::GetLifetimeCustomDeltaProperty(LocalRepLayout, CustomDeltaProperty);
+
 		// If this is a dynamic array, we do the delta here
 		TSharedPtr<INetDeltaBaseState> NewState;
 
@@ -1473,7 +1463,7 @@ PRAGMA_ENABLE_DEPRECATION_WARNINGS
 		{
 			if (Connection->ResendAllDataState == EResendAllDataState::SinceCheckpoint)
 			{
-				TSharedPtr<INetDeltaBaseState>& OldState = CheckpointCustomDeltaState.FindChecked(CustomDeltaProperty);
+				TSharedPtr<INetDeltaBaseState>& OldState = UsingCustomDeltaStates.FindChecked(CustomDeltaProperty);
 
 				if (!SendCustomDeltaProperty(Object, CustomDeltaProperty, TempBitWriter, NewState, OldState))
 				{
@@ -1487,7 +1477,7 @@ PRAGMA_ENABLE_DEPRECATION_WARNINGS
 			{
 				// If we are resending data since open, we don't want to affect the current state of channel/replication, so just do the minimum and send the data, and return
 				// In this case, we'll send all of the properties since the CDO, so use the initial CDO delta state
-				TSharedPtr<INetDeltaBaseState>& OldState = CDOCustomDeltaState.FindChecked(CustomDeltaProperty);
+				TSharedPtr<INetDeltaBaseState>& OldState = UsingCustomDeltaStates.FindChecked(CustomDeltaProperty);
 
 				if (!SendCustomDeltaProperty(Object, CustomDeltaProperty, TempBitWriter, NewState, OldState))
 				{
@@ -1496,11 +1486,14 @@ PRAGMA_ENABLE_DEPRECATION_WARNINGS
 			}
 
 			// Write property header and payload to the bunch
-			WritePropertyHeaderAndPayload(Object, It, NetFieldExportGroup, Bunch, TempBitWriter);
+			WritePropertyHeaderAndPayload(Object, Property, NetFieldExportGroup, Bunch, TempBitWriter);
 
 			continue;
 		}
 
+		// Get info.
+		FPropertyRetirement& Retire = SendingRepState->Retirement[CustomDeltaProperty];
+
 		// Update Retirement records with this new state so we can handle packet drops.
 		// LastNext will be pointer to the last "Next" pointer in the list (so pointer to a pointer)
 		FPropertyRetirement** LastNext = UpdateAckedRetirements(Retire, Connection->OutAckPacketId, Object);
@@ -1510,7 +1503,7 @@ PRAGMA_ENABLE_DEPRECATION_WARNINGS
 
 		ValidateRetirementHistory(Retire, Object);
 
-		TSharedPtr<INetDeltaBaseState>& OldState = RecentCustomDeltaState.FindOrAdd(CustomDeltaProperty);
+		TSharedPtr<INetDeltaBaseState>& OldState = UsingCustomDeltaStates.FindOrAdd(CustomDeltaProperty);
 
 		//-----------------------------------------
 		//	Do delta serialization on dynamic properties
@@ -1531,9 +1524,9 @@ PRAGMA_ENABLE_DEPRECATION_WARNINGS
 		OldState = NewState; 
 
 		// Write property header and payload to the bunch
-		WritePropertyHeaderAndPayload(Object, It, NetFieldExportGroup, Bunch, TempBitWriter);
+		WritePropertyHeaderAndPayload(Object, Property, NetFieldExportGroup, Bunch, TempBitWriter);
 
-		NETWORK_PROFILER(GNetworkProfiler.TrackReplicateProperty(It, TempBitWriter.GetNumBits(), Connection));
+		NETWORK_PROFILER(GNetworkProfiler.TrackReplicateProperty(Property, TempBitWriter.GetNumBits(), Connection));
 	}
 }
 
@@ -1542,9 +1535,9 @@ bool FObjectReplicator::ReplicateProperties( FOutBunch & Bunch, FReplicationFlag
 {
 	UObject* Object = GetObject();
 
-	if ( Object == NULL )
+	if ( Object == nullptr )
 	{
-		UE_LOG(LogRep, Verbose, TEXT("ReplicateProperties: Object == NULL"));
+		UE_LOG(LogRep, Verbose, TEXT("ReplicateProperties: Object == nullptr"));
 		return false;
 	}
 
@@ -1570,7 +1563,7 @@ bool FObjectReplicator::ReplicateProperties( FOutBunch & Bunch, FReplicationFlag
 
 	// Update change list (this will re-use work done by previous connections)
 	FSendingRepState* SendingRepState = ((Connection->ResendAllDataState == EResendAllDataState::SinceCheckpoint) && CheckpointRepState.IsValid()) ? CheckpointRepState->GetSendingRepState() : RepState->GetSendingRepState();
-	RepLayout->UpdateChangelistMgr(SendingRepState, *ChangelistMgr, Object, Connection->Driver->ReplicationFrame, RepFlags, OwningChannel->bForceCompareProperties);
+	FNetSerializeCB::UpdateChangelistMgr(*RepLayout, SendingRepState, *ChangelistMgr, Object, Connection->Driver->ReplicationFrame, RepFlags, OwningChannel->bForceCompareProperties);
 
 	// Replicate properties in the layout
 	const bool bHasRepLayout = RepLayout->ReplicateProperties(SendingRepState, ChangelistMgr->GetRepChangelistState(), (uint8*)Object, ObjectClass, OwningChannel, Writer, RepFlags);
@@ -1603,7 +1596,7 @@ bool FObjectReplicator::ReplicateProperties( FOutBunch & Bunch, FReplicationFlag
 	bLastUpdateEmpty = Writer.GetNumBits() == 0;
 
 	// Replicate Queued (unreliable functions)
-	if ( RemoteFunctions != NULL && RemoteFunctions->GetNumBits() > 0 )
+	if ( RemoteFunctions != nullptr && RemoteFunctions->GetNumBits() > 0 )
 	{
 		if ( UNLIKELY(GNetRPCDebug == 1) )
 		{
@@ -1632,15 +1625,18 @@ bool FObjectReplicator::ReplicateProperties( FOutBunch & Bunch, FReplicationFlag
 
 void FObjectReplicator::ForceRefreshUnreliableProperties()
 {
-	if ( GetObject() == NULL )
+	if (GetObject() == nullptr)
 	{
-		UE_LOG( LogRep, Verbose, TEXT( "ForceRefreshUnreliableProperties: Object == NULL" ) );
+		UE_LOG(LogRep, Verbose, TEXT("ForceRefreshUnreliableProperties: Object == nullptr"));
 		return;
 	}
 
-	check( !bOpenAckCalled );
+	check(!bOpenAckCalled);
 
-	RepLayout->OpenAcked( RepState->GetSendingRepState() );
+	if (FSendingRepState* SendingRepState = RepState->GetSendingRepState())
+	{
+		SendingRepState->bOpenAckedCalled = true;
+	}
 
 	bOpenAckCalled = true;
 }
@@ -1651,7 +1647,7 @@ void FObjectReplicator::PostSendBunch( FPacketIdRange & PacketRange, uint8 bReli
 
 	if ( Object == nullptr )
 	{
-		UE_LOG(LogNet, Verbose, TEXT("PostSendBunch: Object == NULL"));
+		UE_LOG(LogNet, Verbose, TEXT("PostSendBunch: Object == nullptr"));
 		return;
 	}
 
@@ -1662,49 +1658,67 @@ void FObjectReplicator::PostSendBunch( FPacketIdRange & PacketRange, uint8 bReli
 
 	const FRepLayout& LocalRepLayout = *RepLayout;
 
-	if (!SkipRetirementUpdate)
+	if (FSendingRepState* SendingRepState = RepState.IsValid() ? RepState->GetSendingRepState() : nullptr)
 	{
-		// Don't call if reliable, since the bunch will be resent. We dont want this to end up in the changelist history
-		// But is that enough? How does it know to delta against this latest state?
+		if (!SkipRetirementUpdate)
+		{
+			// Don't call if reliable, since the bunch will be resent. We dont want this to end up in the changelist history
+			// But is that enough? How does it know to delta against this latest state?
 
-		LocalRepLayout.PostReplicate(RepState->GetSendingRepState(), PacketRange, bReliable ? true : false);
-	}
+			for (int32 i = SendingRepState->HistoryStart; i < SendingRepState->HistoryEnd; ++i)
+			{
+				const int32 HistoryIndex = i % FSendingRepState::MAX_CHANGE_HISTORY;
 
-	for (uint16 LifetimePropertyIndex : LocalRepLayout.GetLifetimeCustomDeltaProperties())
-	{
-		FPropertyRetirement & Retire = Retirement[LifetimePropertyIndex];
+				FRepChangedHistory & HistoryItem = SendingRepState->ChangeHistory[HistoryIndex];
+
+				if (HistoryItem.OutPacketIdRange.First == INDEX_NONE)
+				{
+					check(HistoryItem.Changed.Num() > 0);
+					check(!HistoryItem.Resend);
 
-		FPropertyRetirement* Next = Retire.Next;
-		FPropertyRetirement* Prev = &Retire;
+					HistoryItem.OutPacketIdRange = PacketRange;
 
-		while ( Next != nullptr )
+					if (!bReliable && !SendingRepState->bOpenAckedCalled)
+					{
+						SendingRepState->PreOpenAckHistory.Add(HistoryItem);
+					}
+				}
+			}
+		}
+
+		for (FPropertyRetirement& Retirement : SendingRepState->Retirement)
 		{
-			// This is updating the dynamic properties retirement record that was created above during property replication
-			// (we have to wait until we actually send the bunch to know the packetID, which is why we look for .First==INDEX_NONE)
-			if ( Next->OutPacketIdRange.First == INDEX_NONE )
-			{
+			FPropertyRetirement* Next = Retirement.Next;
+			FPropertyRetirement* Prev = &Retirement;
 
-				if (!SkipRetirementUpdate)
+			while (Next != nullptr)
+			{
+				// This is updating the dynamic properties retirement record that was created above during property replication
+				// (we have to wait until we actually send the bunch to know the packetID, which is why we look for .First==INDEX_NONE)
+				if (Next->OutPacketIdRange.First == INDEX_NONE)
 				{
-					Next->OutPacketIdRange = PacketRange;
+					if (!SkipRetirementUpdate)
+					{
+						Next->OutPacketIdRange = PacketRange;
 
-					// Mark the last time on this retirement slot that a property actually changed
-					Retire.OutPacketIdRange = PacketRange;
-				}
-				else
-				{
-					// We need to remove the retirement entry here!
-					Prev->Next = Next->Next;
-					delete Next;
-					Next = Prev;
+						// Mark the last time on this retirement slot that a property actually changed
+						Retirement.OutPacketIdRange = PacketRange;
+					}
+					else
+					{
+						// We need to remove the retirement entry here!
+						Prev->Next = Next->Next;
+						delete Next;
+						Next = Prev;
+					}
 				}
+
+				Prev = Next;
+				Next = Next->Next;
 			}
 
-			Prev = Next;
-			Next = Next->Next;
+			ValidateRetirementHistory(Retirement, Object);
 		}
-
-		ValidateRetirementHistory( Retire, Object );
 	}
 }
 
@@ -1727,48 +1741,6 @@ void FObjectReplicator::CountBytes(FArchive& Ar) const
 {
 	GRANULAR_NETWORK_MEMORY_TRACKING_INIT(Ar, "FObjectReplicator::CountBytes");
 
-	GRANULAR_NETWORK_MEMORY_TRACKING_TRACK("Retirement", Retirement.CountBytes(Ar));
-
-	GRANULAR_NETWORK_MEMORY_TRACKING_TRACK("RecentCustomDeltaState",
-		RecentCustomDeltaState.CountBytes(Ar);
-		for (const auto& RecentCustomDeltaStatePair : RecentCustomDeltaState)
-		{
-			if (INetDeltaBaseState const * const BaseState = RecentCustomDeltaStatePair.Value.Get())
-			{
-				BaseState->CountBytes(Ar);
-			}
-		}
-	);
-
-	GRANULAR_NETWORK_MEMORY_TRACKING_TRACK("CDOCustomDeltaState",
-		CDOCustomDeltaState.CountBytes(Ar);
-		for (const auto& CDOCustomDeltaStatePair : CDOCustomDeltaState)
-		{
-			if (INetDeltaBaseState const * const BaseState = CDOCustomDeltaStatePair.Value.Get())
-			{
-				BaseState->CountBytes(Ar);
-			}
-		}
-	);
-
-PRAGMA_DISABLE_DEPRECATION_WARNINGS
-	GRANULAR_NETWORK_MEMORY_TRACKING_TRACK("LifetimeCustomDeltaProperties", LifetimeCustomDeltaProperties.CountBytes(Ar));
-
-	GRANULAR_NETWORK_MEMORY_TRACKING_TRACK("LifetimeCustomDeltaPropertyConditions", LifetimeCustomDeltaPropertyConditions.CountBytes(Ar));
-
-	GRANULAR_NETWORK_MEMORY_TRACKING_TRACK("UnmappedCustomProperties", UnmappedCustomProperties.CountBytes(Ar));
-PRAGMA_ENABLE_DEPRECATION_WARNINGS
-
-	GRANULAR_NETWORK_MEMORY_TRACKING_TRACK("RepNotifies", RepNotifies.CountBytes(Ar));
-
-	GRANULAR_NETWORK_MEMORY_TRACKING_TRACK("RepNotifyMetaData",
-		RepNotifyMetaData.CountBytes(Ar);
-		for (const auto& MetaDataPair : RepNotifyMetaData)
-		{
-			MetaDataPair.Value.CountBytes(Ar);
-		}
-	);
-
 	// FObjectReplicator has a shared pointer to an FRepLayout, but since it's shared with
 	// the UNetDriver, the memory isn't tracked here.
 
@@ -1815,8 +1787,8 @@ void FObjectReplicator::QueueRemoteFunctionBunch( UFunction* Func, FOutBunch &Bu
 	//
 	// Long term we want to have priorities and stronger cross channel traffic management that
 	// can handle this better
-	int32 InfoIdx=INDEX_NONE;
-	for (int32 i=0; i < RemoteFuncInfo.Num(); ++i)
+	int32 InfoIdx = INDEX_NONE;
+	for (int32 i = 0; i < RemoteFuncInfo.Num(); ++i)
 	{
 		if (RemoteFuncInfo[i].FuncName == Func->GetFName())
 		{
@@ -1824,7 +1796,8 @@ void FObjectReplicator::QueueRemoteFunctionBunch( UFunction* Func, FOutBunch &Bu
 			break;
 		}
 	}
-	if (InfoIdx==INDEX_NONE)
+
+	if (InfoIdx == INDEX_NONE)
 	{
 		InfoIdx = RemoteFuncInfo.AddUninitialized();
 		RemoteFuncInfo[InfoIdx].FuncName = Func->GetFName();
@@ -1837,7 +1810,7 @@ void FObjectReplicator::QueueRemoteFunctionBunch( UFunction* Func, FOutBunch &Bu
 			RemoteFuncInfo[InfoIdx].Calls, *Func->GetName(), *GetPathNameSafe(GetObject()), RemoteFuncInfo[InfoIdx].LastCallTime, OwningChannel->Connection->Driver->Time, OwningChannel->RelevantTime, OwningChannel->LastUpdateTime);
 
 		// The MustBeMappedGuids can just be dropped, because we aren't actually going to send a bunch. If we don't clear it, then we will get warnings when the next channel tries to replicate
-		CastChecked< UPackageMapClient >( Connection->PackageMap )->GetMustBeMappedGuidsInLastBunch().Reset();
+		CastChecked<UPackageMapClient>(Connection->PackageMap)->GetMustBeMappedGuidsInLastBunch().Reset();
 		return;
 	}
 	
@@ -1848,67 +1821,91 @@ void FObjectReplicator::QueueRemoteFunctionBunch( UFunction* Func, FOutBunch &Bu
 		RemoteFunctions = new FOutBunch(OwningChannel, 0);
 	}
 
-	RemoteFunctions->SerializeBits( Bunch.GetData(), Bunch.GetNumBits() );
+	RemoteFunctions->SerializeBits(Bunch.GetData(), Bunch.GetNumBits());
 
-	if ( Connection->PackageMap != nullptr )
+	if (Connection->PackageMap != nullptr)
 	{
-		UPackageMapClient * PackageMapClient = CastChecked< UPackageMapClient >( Connection->PackageMap );
+		UPackageMapClient* PackageMapClient = CastChecked<UPackageMapClient>(Connection->PackageMap);
 
 		// We need to copy over any info that was obtained on the package map during serialization, and remember it until we actually call SendBunch
-		if ( PackageMapClient->GetMustBeMappedGuidsInLastBunch().Num() )
+		if (PackageMapClient->GetMustBeMappedGuidsInLastBunch().Num())
 		{
-			OwningChannel->QueuedMustBeMappedGuidsInLastBunch.Append( PackageMapClient->GetMustBeMappedGuidsInLastBunch() );
+			OwningChannel->QueuedMustBeMappedGuidsInLastBunch.Append(PackageMapClient->GetMustBeMappedGuidsInLastBunch());
 			PackageMapClient->GetMustBeMappedGuidsInLastBunch().Reset();
 		}
 
 		if (!Connection->InternalAck)
 		{
 			// Copy over any exported bunches
-			PackageMapClient->AppendExportBunches( OwningChannel->QueuedExportBunches );
+			PackageMapClient->AppendExportBunches(OwningChannel->QueuedExportBunches);
 		}
 	}
 }
 
-bool FObjectReplicator::ReadyForDormancy(bool suppressLogs)
+bool FObjectReplicator::ReadyForDormancy(bool bSuppressLogs)
 {
-	if ( GetObject() == NULL )
+	if (GetObject() == nullptr)
 	{
-		UE_LOG( LogRep, Verbose, TEXT( "ReadyForDormancy: Object == NULL" ) );
+		UE_LOG(LogRep, Verbose, TEXT("ReadyForDormancy: Object == nullptr"));
 		return true;		// Technically, we don't want to hold up dormancy, but the owner needs to clean us up, so we warn
 	}
 
 	// Can't go dormant until last update produced no new property updates
-	if ( !bLastUpdateEmpty )
+	if (!bLastUpdateEmpty)
 	{
-		if ( !suppressLogs )
+		if (!bSuppressLogs)
 		{
-			UE_LOG( LogRepTraffic, Verbose, TEXT( "    [%d] Not ready for dormancy. bLastUpdateEmpty = false" ), OwningChannel->ChIndex );
+			UE_LOG(LogRepTraffic, Verbose, TEXT("    [%d] Not ready for dormancy. bLastUpdateEmpty = false"), OwningChannel->ChIndex);
 		}
 
 		return false;
 	}
 
-	// Can't go dormant if there are unAckd property updates
-	for ( int32 i = 0; i < Retirement.Num(); ++i )
+	if (FSendingRepState* SendingRepState = RepState.IsValid() ? RepState->GetSendingRepState() : nullptr)
 	{
-		if ( Retirement[i].Next != NULL )
+		if (SendingRepState->HistoryStart != SendingRepState->HistoryEnd)
+		{
+			// We have change lists that haven't been acked
+			return false;
+		}
+
+		if (SendingRepState->NumNaks > 0)
 		{
-			if ( !suppressLogs )
+			return false;
+		}
+
+		if (!SendingRepState->bOpenAckedCalled)
+		{
+			return false;
+		}
+
+		if (SendingRepState->PreOpenAckHistory.Num() > 0)
+		{
+			return false;
+		}
+
+		// Can't go dormant if there are unAckd property updates
+		for (FPropertyRetirement& Retirement : SendingRepState->Retirement)
+		{
+			if (Retirement.Next != nullptr)
 			{
-				UE_LOG( LogRepTraffic, Verbose, TEXT( "    [%d] OutAckPacketId: %d First: %d Last: %d " ), OwningChannel->ChIndex, OwningChannel->Connection->OutAckPacketId, Retirement[i].OutPacketIdRange.First, Retirement[i].OutPacketIdRange.Last );
+				if (!bSuppressLogs)
+				{
+					UE_LOG(LogRepTraffic, Verbose, TEXT("    [%d] OutAckPacketId: %d First: %d Last: %d "), OwningChannel->ChIndex, OwningChannel->Connection->OutAckPacketId, Retirement.OutPacketIdRange.First, Retirement.OutPacketIdRange.Last);
+				}
+				return false;
 			}
-			return false;
 		}
 	}
 
-	return RepLayout->ReadyForDormancy( RepState.Get() );
+	return true;
 }
 
 void FObjectReplicator::StartBecomingDormant()
 {
-	if ( GetObject() == NULL )
+	if ( GetObject() == nullptr )
 	{
-		UE_LOG( LogRep, Verbose, TEXT( "StartBecomingDormant: Object == NULL" ) );
+		UE_LOG( LogRep, Verbose, TEXT( "StartBecomingDormant: Object == nullptr" ) );
 		return;
 	}
 
@@ -1940,69 +1937,6 @@ void FObjectReplicator::CallRepNotifies(bool bSkipIfChannelHasQueuedBunches)
 	FReceivingRepState* ReceivingRepState = RepState->GetReceivingRepState();
 	RepLayout->CallRepNotifies(ReceivingRepState, Object);
 
-	if (RepNotifies.Num() > 0)
-	{
-		for (UProperty* RepProperty : RepNotifies)
-		{
-			//UE_LOG(LogRep, Log,  TEXT("Calling Object->%s with %s"), *RepNotifies(RepNotifyIdx)->RepNotifyFunc.ToString(), *RepNotifies(RepNotifyIdx)->GetName()); 						
-			UFunction* RepNotifyFunc = Object->FindFunction(RepProperty->RepNotifyFunc);
-
-			if (RepNotifyFunc == nullptr)
-			{
-				UE_LOG(LogRep, Warning, TEXT("FObjectReplicator::CallRepNotifies: Can't find RepNotify function %s for property %s on object %s."),
-					*RepProperty->RepNotifyFunc.ToString(), *RepProperty->GetName(), *Object->GetName());
-				continue;
-			}
-
-			if (RepNotifyFunc->NumParms == 0)
-			{
-				Object->ProcessEvent(RepNotifyFunc, NULL);
-			}
-			else if (RepNotifyFunc->NumParms == 1)
-			{
-				Object->ProcessEvent(RepNotifyFunc, RepProperty->ContainerPtrToValuePtr<uint8>(ReceivingRepState->StaticBuffer.GetData()) );
-			}
-			else if (RepNotifyFunc->NumParms == 2)
-			{
-				// Fixme: this isn't as safe as it could be. Right now we have two types of parameters: MetaData (a TArray<uint8>)
-				// and the last local value (pointer into the Recent[] array).
-				//
-				// Arrays always expect MetaData. Everything else, including structs, expect last value.
-				// This is enforced with UHT only. If a ::NetSerialize function ever starts producing a MetaData array thats not in UArrayProperty,
-				// we have no static way of catching this and the replication system could pass the wrong thing into ProcessEvent here.
-				//
-				// But this is all sort of an edge case feature anyways, so its not worth tearing things up too much over.
-
-				FMemMark Mark(FMemStack::Get());
-				uint8* Parms = new(FMemStack::Get(), MEM_Zeroed, RepNotifyFunc->ParmsSize)uint8;
-
-				TFieldIterator<UProperty> Itr(RepNotifyFunc);
-				check(Itr);
-
-				Itr->CopyCompleteValue(Itr->ContainerPtrToValuePtr<void>(Parms), RepProperty->ContainerPtrToValuePtr<uint8>(ReceivingRepState->StaticBuffer.GetData()));
-				++Itr;
-				check(Itr);
-
-				TArray<uint8> *NotifyMetaData = RepNotifyMetaData.Find(RepProperty);
-				check(NotifyMetaData);
-				Itr->CopyCompleteValue(Itr->ContainerPtrToValuePtr<void>(Parms), NotifyMetaData);
-
-				Object->ProcessEvent(RepNotifyFunc, Parms);
-
-				Mark.Pop();
-			}
-
-			if (Object->IsPendingKill())
-			{
-				// script event destroyed Object
-				break;
-			}
-		}
-	}
-
-	RepNotifies.Reset();
-	RepNotifyMetaData.Empty();
-
 	if (!Object->IsPendingKill())
 	{
 		Object->PostRepNotifies();
@@ -2028,8 +1962,10 @@ void FObjectReplicator::UpdateUnmappedObjects(bool & bOutHasMoreUnmapped)
 	// Since RepNotifies aren't processed while a channel has queued bunches, don't assert in that case.
 	FReceivingRepState* ReceivingRepState = RepState->GetReceivingRepState();
 	const bool bHasQueuedBunches = OwningChannel && OwningChannel->QueuedBunches.Num() > 0;
-	checkf( bHasQueuedBunches || ReceivingRepState->RepNotifies.Num() == 0, TEXT("Failed RepState RepNotifies check. Num=%d. Object=%s. Channel QueuedBunches=%d"), ReceivingRepState->RepNotifies.Num(), *Object->GetFullName(), OwningChannel ? OwningChannel->QueuedBunches.Num() : 0 );
-	checkf( bHasQueuedBunches || RepNotifies.Num() == 0, TEXT("Failed replicator RepNotifies check. Num=%d. Object=%s. Channel QueuedBunches=%d"), RepNotifies.Num(), *Object->GetFullName(), OwningChannel ? OwningChannel->QueuedBunches.Num() : 0 );
+
+	checkf(bHasQueuedBunches || ReceivingRepState->RepNotifies.Num() == 0,
+		TEXT("Failed RepState RepNotifies check. Num=%d. Object=%s. Channel QueuedBunches=%d"),
+		ReceivingRepState->RepNotifies.Num(), *Object->GetFullName(), OwningChannel ? OwningChannel->QueuedBunches.Num() : 0);
 
 	bool bCalledPreNetReceive = false;
 	bool bSomeObjectsWereMapped = false;
@@ -2038,9 +1974,6 @@ void FObjectReplicator::UpdateUnmappedObjects(bool & bOutHasMoreUnmapped)
 
 	const FRepLayout& LocalRepLayout = *RepLayout;
 
-	// Let the rep layout update any unmapped properties
-	LocalRepLayout.UpdateUnmappedObjects(ReceivingRepState, Connection->PackageMap, Object, bCalledPreNetReceive, bSomeObjectsWereMapped, bOutHasMoreUnmapped);
-
 	FNetSerializeCB NetSerializeCB(Connection->Driver);
 
 	FNetDeltaSerializeInfo Parms;
@@ -2050,34 +1983,14 @@ void FObjectReplicator::UpdateUnmappedObjects(bool & bOutHasMoreUnmapped)
 	Parms.NetSerializeCB = &NetSerializeCB;
 
 	Parms.bUpdateUnmappedObjects = true;
-	Parms.bCalledPreNetReceive = bCalledPreNetReceive;
-	
 
-	TArray<TPair<int32, UStructProperty*>> CompletelyMappedProperties;
-	TArray<TPair<int32, UStructProperty*>> UpdatedProperties;
-	FNetSerializeCB::UpdateUnmappedObjectsForCustomDeltaProperties(LocalRepLayout, Parms, CompletelyMappedProperties, UpdatedProperties);
+	// Let the rep layout update any unmapped properties
+	LocalRepLayout.UpdateUnmappedObjects(ReceivingRepState, Connection->PackageMap, Object, Parms, bCalledPreNetReceive, bSomeObjectsWereMapped, bOutHasMoreUnmapped);
 
 	bSomeObjectsWereMapped |= Parms.bOutSomeObjectsWereMapped;
 	bOutHasMoreUnmapped |= Parms.bOutHasMoreUnmapped;
 	bCalledPreNetReceive |= Parms.bCalledPreNetReceive;
 
-	// This should go away when UnmappedCustomProperties goes away, and when RepNotifies
-	// are merged with RepState RepNotifies.
-PRAGMA_DISABLE_DEPRECATION_WARNINGS
-	for (const TPair<int32, UStructProperty*>& KVP : UpdatedProperties)
-	{
-		TArray<uint8> MetaData;
-		QueuePropertyRepNotify(Object, KVP.Value, 0, MetaData);
-	}
-
-	// This is just for the sake of keeping UnmappedCustomProperties up to date.
-	// Remove this when that property goes away.
-	for (const TPair<int32, UStructProperty*>& KVP : CompletelyMappedProperties)
-	{
-		UnmappedCustomProperties.Remove(KVP.Key);
-	}
-PRAGMA_ENABLE_DEPRECATION_WARNINGS
-
 	if (bCalledPreNetReceive)
 	{
 		// If we mapped some objects, make sure to call PostNetReceive (some game code will need to think this was actually replicated to work)
@@ -2091,18 +2004,18 @@ PRAGMA_ENABLE_DEPRECATION_WARNINGS
 	// the RepNotifies will remain in the list and the check for 0 RepNotifies above will fail next time.
 	CallRepNotifies(false);
 
-	UPackageMapClient * PackageMapClient = Cast< UPackageMapClient >(Connection->PackageMap);
+	UPackageMapClient* PackageMapClient = Cast<UPackageMapClient>(Connection->PackageMap);
 
 	if (PackageMapClient && OwningChannel)
 	{
 		const bool bIsServer = Connection->Driver->IsServer();
-		const FClassNetCache * ClassCache = Connection->Driver->NetCache->GetClassNetCache(ObjectClass);
+		const FClassNetCache* ClassCache = Connection->Driver->NetCache->GetClassNetCache(ObjectClass);
 
 		// Handle pending RPCs, in order
 		for (int32 RPCIndex = 0; RPCIndex < PendingLocalRPCs.Num(); RPCIndex++)
 		{
 			FRPCPendingLocalCall& Pending = PendingLocalRPCs[RPCIndex];
-			const FFieldNetCache * FieldCache = ClassCache->GetFromIndex(Pending.RPCFieldIndex);
+			const FFieldNetCache* FieldCache = ClassCache->GetFromIndex(Pending.RPCFieldIndex);
 
 			FNetBitReader Reader(Connection->PackageMap, Pending.Buffer.GetData(), Pending.NumBits);
 
@@ -2174,72 +2087,81 @@ PRAGMA_ENABLE_DEPRECATION_WARNINGS
 	}
 }
 
-void FObjectReplicator::QueuePropertyRepNotify( UObject* Object, UProperty * Property, const int32 ElementIndex, TArray< uint8 > & MetaData )
+void FObjectReplicator::QueuePropertyRepNotify(
+	UObject* Object,
+	UProperty * Property,
+	const int32 ElementIndex,
+	TArray<uint8>& MetaData)
 {
-	if ( !Property->HasAnyPropertyFlags( CPF_RepNotify ) )
+	if (!Property->HasAnyPropertyFlags(CPF_RepNotify))
 	{
 		return;
 	}
 
-	//@note: AddUniqueItem() here for static arrays since RepNotify() currently doesn't indicate index,
-	//			so reporting the same property multiple times is not useful and wastes CPU
-	//			were that changed, this should go back to AddItem() for efficiency
-	// @todo UE4 - not checking if replicated value is changed from old.  Either fix or document, as may get multiple repnotifies of unacked properties.
-	RepNotifies.AddUnique( Property );
+	FReceivingRepState* ReceivingRepState = RepState.IsValid() ? RepState->GetReceivingRepState() : nullptr;
+	if (ensureMsgf(ReceivingRepState, TEXT("FObjectReplicator::QueuePropertyRepNotifiy: No receiving RepState. Object=%s, Property=%s"),
+		*GetPathNameSafe(Object), *Property->GetName()))
+	{
+		//@note: AddUniqueItem() here for static arrays since RepNotify() currently doesn't indicate index,
+		//			so reporting the same property multiple times is not useful and wastes CPU
+		//			were that changed, this should go back to AddItem() for efficiency
+		// @todo UE4 - not checking if replicated value is changed from old.  Either fix or document, as may get multiple repnotifies of unacked properties.
+		ReceivingRepState->RepNotifies.AddUnique(Property);
 
-	UFunction * RepNotifyFunc = Object->FindFunctionChecked( Property->RepNotifyFunc );
+		UFunction* RepNotifyFunc = Object->FindFunctionChecked(Property->RepNotifyFunc);
 
-	if ( RepNotifyFunc->NumParms > 0 )
-	{
-		if ( Property->ArrayDim != 1 )
-		{
-			// For static arrays, we build the meta data here, but adding the Element index that was just read into the PropMetaData array.
-			UE_LOG( LogRepTraffic, Verbose, TEXT("Property %s had ArrayDim: %d change"), *Property->GetName(), ElementIndex );
-
-			// Property is multi dimensional, keep track of what elements changed
-			TArray< uint8 > & PropMetaData = RepNotifyMetaData.FindOrAdd( Property );
-			PropMetaData.Add( ElementIndex );
-		} 
-		else if ( MetaData.Num() > 0 )
+		if (RepNotifyFunc->NumParms > 0)
 		{
-			// For other properties (TArrays only now) the MetaData array is build within ::NetSerialize. Just add it to the RepNotifyMetaData map here.
+			if (Property->ArrayDim != 1)
+			{
+				// For static arrays, we build the meta data here, but adding the Element index that was just read into the PropMetaData array.
+				UE_LOG(LogRepTraffic, Verbose, TEXT("Property %s had ArrayDim: %d change"), *Property->GetName(), ElementIndex);
 
-			//UE_LOG(LogRepTraffic, Verbose, TEXT("Property %s had MetaData: "), *Property->GetName() );
-			//for (auto MetaIt = MetaData.CreateIterator(); MetaIt; ++MetaIt)
-			//	UE_LOG(LogRepTraffic, Verbose, TEXT("   %d"), *MetaIt );
+				// Property is multi dimensional, keep track of what elements changed
+				TArray< uint8 > & PropMetaData = ReceivingRepState->RepNotifyMetaData.FindOrAdd(Property);
+				PropMetaData.Add(ElementIndex);
+			}
+			else if (MetaData.Num() > 0)
+			{
+				// For other properties (TArrays only now) the MetaData array is build within ::NetSerialize. Just add it to the RepNotifyMetaData map here.
+
+				//UE_LOG(LogRepTraffic, Verbose, TEXT("Property %s had MetaData: "), *Property->GetName() );
+				//for (auto MetaIt = MetaData.CreateIterator(); MetaIt; ++MetaIt)
+				//	UE_LOG(LogRepTraffic, Verbose, TEXT("   %d"), *MetaIt );
 
-			// Property included some meta data about what was serialized. 
-			TArray< uint8 > & PropMetaData = RepNotifyMetaData.FindOrAdd( Property );
-			PropMetaData = MetaData;
+				// Property included some meta data about what was serialized. 
+				TArray< uint8 > & PropMetaData = ReceivingRepState->RepNotifyMetaData.FindOrAdd(Property);
+				PropMetaData = MetaData;
+			}
 		}
 	}
 }
 
 void FObjectReplicator::WritePropertyHeaderAndPayload(
-	UObject*				Object,
-	UProperty*				Property,
-	FNetFieldExportGroup*	NetFieldExportGroup,
-	FNetBitWriter&			Bunch,
-	FNetBitWriter&			Payload ) const
+	UObject* Object,
+	UProperty* Property,
+	FNetFieldExportGroup* NetFieldExportGroup,
+	FNetBitWriter& Bunch,
+	FNetBitWriter& Payload) const
 {
 	// Get class network info cache.
-	const FClassNetCache* ClassCache = Connection->Driver->NetCache->GetClassNetCache( ObjectClass );
+	const FClassNetCache* ClassCache = Connection->Driver->NetCache->GetClassNetCache(ObjectClass);
 
-	check( ClassCache );
+	check(ClassCache);
 
 	// Get the network friend property index to replicate
-	const FFieldNetCache * FieldCache = ClassCache->GetFromField( Property );
+	const FFieldNetCache * FieldCache = ClassCache->GetFromField(Property);
 
-	checkSlow( FieldCache );
+	checkSlow(FieldCache);
 
 	// Send property name and optional array index.
-	check( FieldCache->FieldNetIndex <= ClassCache->GetMaxIndex() );
+	check(FieldCache->FieldNetIndex <= ClassCache->GetMaxIndex());
 
 	// WriteFieldHeaderAndPayload will return the total number of bits written.
 	// So, we subtract out the Payload size to get the actual number of header bits.
-	const int32 HeaderBits = static_cast<int64>(OwningChannel->WriteFieldHeaderAndPayload( Bunch, ClassCache, FieldCache, NetFieldExportGroup, Payload )) - Payload.GetNumBits();
+	const int32 HeaderBits = static_cast<int64>(OwningChannel->WriteFieldHeaderAndPayload(Bunch, ClassCache, FieldCache, NetFieldExportGroup, Payload)) - Payload.GetNumBits();
 
-	NETWORK_PROFILER( GNetworkProfiler.TrackWritePropertyHeader( Property, HeaderBits, nullptr ) );
+	NETWORK_PROFILER(GNetworkProfiler.TrackWritePropertyHeader(Property, HeaderBits, nullptr));
 }
 
 void FObjectReplicator::UpdateCheckpoint()
diff --git a/Engine/Source/Runtime/Engine/Private/RepLayout.cpp b/Engine/Source/Runtime/Engine/Private/RepLayout.cpp
--- a/Engine/Source/Runtime/Engine/Private/RepLayout.cpp
+++ b/Engine/Source/Runtime/Engine/Private/RepLayout.cpp
@@ -20,7 +20,6 @@
 #include "Algo/Sort.h"
 #include "Net/NetworkGranularMemoryLogging.h"
 #include "Serialization/ArchiveCountMem.h"
-#include "Containers/SortedMap.h"
 #include "Templates/AndOrNot.h"
 
 DECLARE_CYCLE_STAT(TEXT("RepLayout AddPropertyCmd"), STAT_RepLayout_AddPropertyCmd, STATGROUP_Game);
@@ -128,10 +127,10 @@ FConsoleVariableSinkHandle MaxRepArrayMemorySink = CreateMaxArrayMemoryCVarAndRe
 namespace UE4_RepLayout_Private
 {
 	template<typename OutputType, typename CommandType, typename BufferType>
-	static typename TTranslateConstVolatile<BufferType, OutputType>::Type*
+	static typename TCopyQualifiersFromTo<BufferType, OutputType>::Type*
 	GetTypedProperty(const BufferType& Buffer, const CommandType& Cmd)
 	{
-		using ConstOrNotOutputType = typename TTranslateConstVolatile<BufferType, OutputType>::Type;
+		using ConstOrNotOutputType = typename TCopyQualifiersFromTo<BufferType, OutputType>::Type;
 
 		using BaseBufferType = typename TRemovePointer<typename TDecay<BufferType>::Type>::Type;
 
@@ -143,16 +142,74 @@ namespace UE4_RepLayout_Private
 	}
 
 	template<typename OutputType, typename CommandType, enum ERepDataBufferType BufferDataType, typename BufferUnderlyingType>
-	static typename TTranslateConstVolatile<typename TRepDataBufferBase<BufferDataType, BufferUnderlyingType>::ConstOrNotVoid, OutputType>::Type*
+	static typename TCopyQualifiersFromTo<typename TRepDataBufferBase<BufferDataType, BufferUnderlyingType>::ConstOrNotVoid, OutputType>::Type*
 	GetTypedProperty(const TRepDataBufferBase<BufferDataType, BufferUnderlyingType>& Buffer, const CommandType& Cmd)
 	{
-		using ConstOrNotOutputType = typename TTranslateConstVolatile<typename TRepDataBufferBase<BufferDataType, BufferUnderlyingType>::ConstOrNotVoid, OutputType>::Type;
+		using ConstOrNotOutputType = typename TCopyQualifiersFromTo<typename TRepDataBufferBase<BufferDataType, BufferUnderlyingType>::ConstOrNotVoid, OutputType>::Type;
 
 		static_assert(!TIsPointer<OutputType>::Value, "GetTypedProperty invalid OutputType! Don't specify output as a pointer.");
 
 		// TODO: Conditionally compilable runtime type validation.
 		return reinterpret_cast<ConstOrNotOutputType *>((Buffer + Cmd).Data);
 	}
+
+	static void QueueRepNotifyForCustomDeltaProperty(
+		FReceivingRepState* RESTRICT ReceivingRepState,
+		FNetDeltaSerializeInfo& Params,
+		UProperty* Property,
+		uint32 StaticArrayIndex)
+	{
+		//@note: AddUniqueItem() here for static arrays since RepNotify() currently doesn't indicate index,
+		//			so reporting the same property multiple times is not useful and wastes CPU
+		//			were that changed, this should go back to AddItem() for efficiency
+		// @todo UE4 - not checking if replicated value is changed from old.  Either fix or document, as may get multiple repnotifies of unacked properties.
+		ReceivingRepState->RepNotifies.AddUnique(Property);
+
+		UFunction* RepNotifyFunc = Params.Object->FindFunctionChecked(Property->RepNotifyFunc);
+
+		if (RepNotifyFunc->NumParms > 0)
+		{
+			if (Property->ArrayDim != 1)
+			{
+				// For static arrays, we build the meta data here, but adding the Element index that was just read into the PropMetaData array.
+				UE_LOG(LogRepTraffic, Verbose, TEXT("Property %s had ArrayDim: %d change"), *Property->GetName(), StaticArrayIndex);
+
+				// Property is multi dimensional, keep track of what elements changed
+				TArray<uint8>& PropMetaData = ReceivingRepState->RepNotifyMetaData.FindOrAdd(Property);
+				PropMetaData.Add(StaticArrayIndex);
+			}
+		}
+	}
+
+	static void WritePropertyHeaderAndPayload(
+		UObject* Object,
+		UClass* ObjectClass,
+		UProperty* Property,
+		UNetConnection* Connection,
+		UActorChannel* OwningChannel,
+		FNetFieldExportGroup* NetFieldExportGroup,
+		FNetBitWriter& Bunch,
+		FNetBitWriter& Payload)
+	{
+		// Get class network info cache.
+		const FClassNetCache* ClassCache = Connection->Driver->NetCache->GetClassNetCache(ObjectClass);
+
+		check(ClassCache);
+
+		// Get the network friend property index to replicate
+		const FFieldNetCache * FieldCache = ClassCache->GetFromField(Property);
+
+		checkSlow(FieldCache);
+
+		// Send property name and optional array index.
+		check(FieldCache->FieldNetIndex <= ClassCache->GetMaxIndex());
+
+		// WriteFieldHeaderAndPayload will return the total number of bits written.
+		// So, we subtract out the Payload size to get the actual number of header bits.
+		const int32 HeaderBits = static_cast<int64>(OwningChannel->WriteFieldHeaderAndPayload(Bunch, ClassCache, FieldCache, NetFieldExportGroup, Payload)) - Payload.GetNumBits();
+
+		NETWORK_PROFILER(GNetworkProfiler.TrackWritePropertyHeader(Property, HeaderBits, nullptr));
+	}
 }
 
 //~ TODO: Consider moving the FastArray members into their own sub-struct to save memory for non fast array
@@ -161,32 +218,57 @@ namespace UE4_RepLayout_Private
 
 struct FLifetimeCustomDeltaProperty
 {
+	FLifetimeCustomDeltaProperty(const uint16 InPropertyRepIndex)
+		: PropertyRepIndex(InPropertyRepIndex)
+	{
+	}
+
+	FLifetimeCustomDeltaProperty(
+		const uint16 InPropertyRepIndex,
+		const int32 InFastArrayItemsCommand,
+		const int32 InFastArrayNumber,
+		const int32 InFastArrayDeltaFlagsOffset,
+		const int32 InFastArrayArrayReplicationKeyOffset,
+		const int32 InFastArrayItemReplicationIdOffset
+	)
+		: PropertyRepIndex(InPropertyRepIndex)
+		, FastArrayItemsCommand(InFastArrayItemsCommand)
+		, FastArrayNumber(InFastArrayNumber)
+		, FastArrayDeltaFlagsOffset(InFastArrayDeltaFlagsOffset)
+		, FastArrayArrayReplicationKeyOffset(InFastArrayArrayReplicationKeyOffset)
+		, FastArrayItemReplicationIdOffset(InFastArrayItemReplicationIdOffset)
+	{
+	}
+
+	/** The RepIndex of the corresponding Property. This can be used as an index into FRepLayout::Parents. */
+	const uint16 PropertyRepIndex = INDEX_NONE;
+
 	/** If this is a Fast Array Serializer property, this will be the command index for the Fast Array Item array. */
-	int32 FastArrayItemsCommand = INDEX_NONE;
+	const int32 FastArrayItemsCommand = INDEX_NONE;
 
 	/**
 	 * If this is a Fast Array Serializer property, this will be the instance number in the class.
 	 * This is used to lookup Changelists.
 	 */
-	int32 FastArrayNumber = INDEX_NONE;
+	const int32 FastArrayNumber = INDEX_NONE;
 
 	/**
 	 * If this is a Fast Array Serializer property (and it is set up correctly for Delta Serialization), this will be an
 	 * offset from to the property.
 	 */
-	int32 FastArrayDeltaFlagsOffset = INDEX_NONE;
+	const int32 FastArrayDeltaFlagsOffset = INDEX_NONE;
 
 	/**
 	 * If this is a Fast Array Serializer property (and it is set up correctly for Delta Serialization), this will be a pointer
 	 * to the FFastArraySerializer::ArrayReplicationKey property.
 	 */
-	int32 FastArrayArrayReplicationKeyOffset = INDEX_NONE;
+	const int32 FastArrayArrayReplicationKeyOffset = INDEX_NONE;
 
 	/**
 	 * If this is a Fast Array Serializer property (and it is set up correctly for Delta Serialization), this will be a pointer
 	 * to the FFastArraySerializerItem::ReplicationID property.
 	 */
-	int32 FastArrayItemReplicationIdOffset = INDEX_NONE;
+	const int32 FastArrayItemReplicationIdOffset = INDEX_NONE;
 
 	const EFastArraySerializerDeltaFlags GetFastArrayDeltaFlags(void const * const FastArray) const
 	{
@@ -230,24 +312,79 @@ private:
  */
 struct FLifetimeCustomDeltaState
 {
-	TSortedMap<uint16, FLifetimeCustomDeltaProperty> LifetimeCustomDeltaProperties;
-	
-	/** The number of valid FFastArraySerializer properties we found. */
-	int32 NumFastArrayItems = 0;
+public:
 
-	/**
-	 * This is less than ideal.
-	 * For now, we want a fast way to let users query the list of lifetime custom delta property indices.
-	 * This is mostly to support FObjectReplicator since we're going to remove its (per object per connection!!) cache.
-	 * However, we wouldn't even need this if Custom Delta Properties were more completely handled by RepLayout.
-	 */
-	TArray<uint16> LifetimeCustomDeltaPropertiesParentCache;
+	static constexpr uint16 INVALID_DELTA_INDEX = INDEX_NONE;
+
+	FLifetimeCustomDeltaState(uint16 TotalNumberOfLifetimeProperties)
+	{
+		LifetimeCustomDeltaIndexLookup.Init(INVALID_DELTA_INDEX, TotalNumberOfLifetimeProperties);
+	}
 
 	void CountBytes(FArchive& Ar) const
 	{
 		LifetimeCustomDeltaProperties.CountBytes(Ar);
-		LifetimeCustomDeltaPropertiesParentCache.CountBytes(Ar);
+		LifetimeCustomDeltaIndexLookup.CountBytes(Ar);
+	}
+
+	const uint16 GetNumCustomDeltaProperties() const
+	{
+		return LifetimeCustomDeltaProperties.Num();
+	}
+
+	const uint16 GetNumFastArrayProperties() const
+	{
+		return NumFastArrayProperties;
 	}
+
+	const FLifetimeCustomDeltaProperty& GetCustomDeltaProperty(const uint16 CustomDeltaIndex) const
+	{
+		return LifetimeCustomDeltaProperties[CustomDeltaIndex];
+	}
+
+	const uint16 GetCustomDeltaIndexFromPropertyRepIndex(const uint16 PropertyRepIndex) const
+	{
+		const uint16 CustomDeltaIndex = LifetimeCustomDeltaIndexLookup[PropertyRepIndex];
+		check(INVALID_DELTA_INDEX != CustomDeltaIndex);
+		return CustomDeltaIndex;
+	}
+
+	void Add(FLifetimeCustomDeltaProperty&& ToAdd)
+	{
+		check(INVALID_DELTA_INDEX == LifetimeCustomDeltaIndexLookup[ToAdd.PropertyRepIndex]);
+
+		if (ToAdd.FastArrayNumber != INDEX_NONE)
+		{
+			++NumFastArrayProperties;
+		}
+
+		LifetimeCustomDeltaIndexLookup[ToAdd.PropertyRepIndex] = LifetimeCustomDeltaProperties.Num();
+		LifetimeCustomDeltaProperties.Emplace(MoveTemp(ToAdd));
+	}
+
+	void CompactMemory()
+	{
+		LifetimeCustomDeltaProperties.Shrink();
+		LifetimeCustomDeltaIndexLookup.Shrink();
+	}
+
+private:
+
+	//~ Since there is only 1 RepLayout per class, and we will only create a FLifetimeCustomDeltaState for a RepLayout whose owning class
+	//~ has Custom Delta Properties, using 2 arrays here seems like a good trade off for performance, memory, and convenience as opposed
+	//~ to a TMap or TSortedMap.
+	//~
+	//~ Having just a map alone makes it harder for external code to iterate over custom delta properties without exposing these internal
+	//~ classes.
+	//~
+	//~ However, maintaining just an array of CustomDeltaProperties makes it less efficient to perform lookups (we either need to keep
+	//~ the list sorted or do linear searches).
+
+	TArray<FLifetimeCustomDeltaProperty> LifetimeCustomDeltaProperties;
+	TArray<uint16> LifetimeCustomDeltaIndexLookup;
+
+	/** The number of valid FFastArraySerializer properties we found. */
+	uint16 NumFastArrayProperties = 0;
 };
 
 //~ Some of this complexity could go away if we introduced a new Compare step to Custom Delta Serializers
@@ -703,7 +840,7 @@ void FRepChangelistState::CountBytes(FArchive& Ar) const
 	{
 		Ar.CountBytes(sizeof(FCustomDeltaChangelistState), sizeof(FCustomDeltaChangelistState));
 		CustomDeltaChangelistState->CountBytes(Ar);
-}
+	}
 }
 
 FReplicationChangelistMgr::FReplicationChangelistMgr(
@@ -881,10 +1018,10 @@ static void CompareParentProperties(
 
 		const bool bShouldSkip = !bIsLifetime || !bIsActive || (Parent.Condition == COND_InitialOnly && !SharedParams.bIsInitial);
 
-			if (bShouldSkip)
-			{
-				continue;
-			}
+		if (bShouldSkip)
+		{
+			continue;
+		}
 
 		const FRepLayoutCmd& Cmd = SharedParams.Cmds[Parent.CmdStart];
 		const uint16 Handle = Cmd.RelativeHandle;
@@ -983,9 +1120,9 @@ static void CompareProperties_Array_r(
 		bool& bForceFail = const_cast<bool&>(SharedParams.bForceFail);
 		TGuardValue<bool> ForceFailGuard(bForceFail, bForceFail);
 
-	for (int32 i = 0; i < ArrayNum; i++)
-	{
-		const int32 ArrayElementOffset = i * Cmd.ElementSize;
+		for (int32 i = 0; i < ArrayNum; i++)
+		{
+			const int32 ArrayElementOffset = i * Cmd.ElementSize;
 			bForceFail = bOldForceFail || i >= ShadowArrayNum;
 			LocalHandle = CompareProperties_r(SharedParams, CmdIndex + 1, Cmd.EndCmd - 1, ShadowArrayData + ArrayElementOffset, ArrayData + ArrayElementOffset, ChangedLocal, LocalHandle);
 		}
@@ -1386,109 +1523,6 @@ void FRepLayout::UpdateChangelistHistory(
 	check(RepState->NumNaks == 0);
 }
 
-void FRepLayout::OpenAcked(FSendingRepState* RepState) const
-{
-	check(RepState);
-	RepState->bOpenAckedCalled = true;
-}
-
-void FRepLayout::PostReplicate(
-	FSendingRepState* RepState,
-	FPacketIdRange& PacketRange,
-	bool bReliable) const
-{
-	if (LayoutState == ERepLayoutState::Normal)
-	{
-		for (int32 i = RepState->HistoryStart; i < RepState->HistoryEnd; ++i)
-		{
-			const int32 HistoryIndex = i % FSendingRepState::MAX_CHANGE_HISTORY;
-
-			FRepChangedHistory & HistoryItem = RepState->ChangeHistory[HistoryIndex];
-
-			if (HistoryItem.OutPacketIdRange.First == INDEX_NONE)
-			{
-				check(HistoryItem.Changed.Num() > 0);
-				check(!HistoryItem.Resend);
-
-				HistoryItem.OutPacketIdRange = PacketRange;
-
-				if (!bReliable && !RepState->bOpenAckedCalled)
-				{
-					RepState->PreOpenAckHistory.Add(HistoryItem);
-				}
-			}
-		}
-	}
-}
-
-void FRepLayout::ReceivedNak(FRepState* RepState, int32 NakPacketId) const
-{
-	if (RepState == nullptr)
-	{
-		return;		// I'm not 100% certain why this happens, the only think I can think of is this is a bNetTemporary?
-	}
-
-	if (LayoutState == ERepLayoutState::Normal)
-	{
-		if (FSendingRepState* SendingRepState = RepState->GetSendingRepState())
-		{
-			for (int32 i = SendingRepState->HistoryStart; i < SendingRepState->HistoryEnd; ++i)
-			{
-				const int32 HistoryIndex = i % FSendingRepState::MAX_CHANGE_HISTORY;
-
-				FRepChangedHistory & HistoryItem = SendingRepState->ChangeHistory[HistoryIndex];
-
-				if (!HistoryItem.Resend && HistoryItem.OutPacketIdRange.InRange(NakPacketId))
-				{
-					check(HistoryItem.Changed.Num() > 0);
-					HistoryItem.Resend = true;
-					++SendingRepState->NumNaks;
-				}
-			}
-		}
-	}
-}
-
-bool FRepLayout::AllAcked(FRepState* RepState) const
-{
-	if (FSendingRepState* SendingRepState = RepState->GetSendingRepState())
-	{
-		if (SendingRepState->HistoryStart != SendingRepState->HistoryEnd)
-		{
-			// We have change lists that haven't been acked
-			return false;
-		}
-
-		if (SendingRepState->NumNaks > 0)
-		{
-			return false;
-		}
-
-		if (!SendingRepState->bOpenAckedCalled)
-		{
-			return false;
-		}
-
-		if (SendingRepState->PreOpenAckHistory.Num() > 0)
-		{
-			return false;
-		}
-	}
-
-	return true;
-}
-
-bool FRepLayout::ReadyForDormancy(FRepState* RepState) const
-{
-	// Clients should never go dormant.
-	if (RepState == nullptr || RepState->GetSendingRepState() == nullptr)
-	{
-		return false;
-	}
-
-	return AllAcked(RepState);
-}
-
 void FRepLayout::SerializeObjectReplicatedProperties(UObject* Object, FBitArchive & Ar) const
 {
 	static FRepSerializationSharedInfo Empty;
@@ -3262,34 +3296,36 @@ void FRepLayout::GatherGuidReferences_r(
 
 void FRepLayout::GatherGuidReferences(
 	FReceivingRepState* RESTRICT RepState,
+	FNetDeltaSerializeInfo& Params,
 	TSet<FNetworkGUID>& OutReferencedGuids,
 	int32& OutTrackedGuidMemoryBytes) const
 {
 	if (LayoutState == ERepLayoutState::Normal)
 	{
 		GatherGuidReferences_r(&RepState->GuidReferencesMap, OutReferencedGuids, OutTrackedGuidMemoryBytes);
-	}
-}
 
-void FRepLayout::GatherGuidReferencesForCustomDeltaProperties(FNetDeltaSerializeInfo& Params) const
-{
-	if (LifetimeCustomPropertyState)
-	{
-		FRepObjectDataBuffer ObjectData(Params.Object);
-		for (const auto& KVP : LifetimeCustomPropertyState->LifetimeCustomDeltaProperties)
+		// Custom Delta Properties
+		if (LifetimeCustomPropertyState)
 		{
-			const FRepParentCmd& Parent = Parents[KVP.Key];
+			FRepObjectDataBuffer ObjectData(Params.Object);
+			const int32 NumLifetimeCustomDeltaProperties = LifetimeCustomPropertyState->GetNumCustomDeltaProperties();
+
+			for (int32 CustomDeltaIndex = 0; CustomDeltaIndex < NumLifetimeCustomDeltaProperties; ++CustomDeltaIndex)
+			{
+				const FLifetimeCustomDeltaProperty& CustomDeltaProperty = LifetimeCustomPropertyState->GetCustomDeltaProperty(CustomDeltaIndex);
+				const FRepParentCmd& Parent = Parents[CustomDeltaProperty.PropertyRepIndex];
 
-			// Static cast is safe here, because this property wouldn't have been marked CustomDelta otherwise.
-			UStructProperty* StructProperty = static_cast<UStructProperty*>(Parent.Property);
-			UScriptStruct::ICppStructOps* CppStructOps = StructProperty->Struct->GetCppStructOps();
+				// Static cast is safe here, because this property wouldn't have been marked CustomDelta otherwise.
+				UStructProperty* StructProperty = static_cast<UStructProperty*>(Parent.Property);
+				UScriptStruct::ICppStructOps* CppStructOps = StructProperty->Struct->GetCppStructOps();
 
-			FNetDeltaSerializeInfo TempParams = Params;
-			TempParams.Struct = StructProperty->Struct;
-			TempParams.PropertyRepIndex = KVP.Key;
-			TempParams.Data = ObjectData + Parent;
+				FNetDeltaSerializeInfo TempParams = Params;
+				TempParams.Struct = StructProperty->Struct;
+				TempParams.CustomDeltaIndex = CustomDeltaIndex;
+				TempParams.Data = ObjectData + Parent;
 
-			CppStructOps->NetDeltaSerialize(TempParams, TempParams.Data);
+				CppStructOps->NetDeltaSerialize(TempParams, TempParams.Data);
+			}
 		}
 	}
 }
@@ -3324,55 +3360,50 @@ bool FRepLayout::MoveMappedObjectToUnmapped_r(FGuidReferencesMap* GuidReferences
 	return bFoundGUID;
 }
 
-bool FRepLayout::MoveMappedObjectToUnmapped(FReceivingRepState* RESTRICT RepState, const FNetworkGUID& GUID) const
-{
-	return ERepLayoutState::Normal == LayoutState && MoveMappedObjectToUnmapped_r(&RepState->GuidReferencesMap, GUID);
-}
-
-PRAGMA_DISABLE_DEPRECATION_WARNINGS
-bool FRepLayout::MoveMappedObjectToUnmappedForCustomDeltaProperties(FNetDeltaSerializeInfo& Params) const
-{
-	TMap<int32, UStructProperty*> Dummy;
-	return MoveMappedObjectToUnmappedForCustomDeltaProperties(Params, Dummy);
-}
-
-bool FRepLayout::MoveMappedObjectToUnmappedForCustomDeltaProperties(FNetDeltaSerializeInfo& Params, TMap<int32, UStructProperty*>& UnmappedCustomProperties) const
+bool FRepLayout::MoveMappedObjectToUnmapped(FReceivingRepState* RESTRICT RepState, FNetDeltaSerializeInfo& Params, const FNetworkGUID& GUID) const
 {
 	bool bFound = false;
 
-	if (LifetimeCustomPropertyState)
+	if (ERepLayoutState::Normal == LayoutState)
 	{
-		FRepObjectDataBuffer ObjectData(Params.Object);
-		for (const auto& KVP : LifetimeCustomPropertyState->LifetimeCustomDeltaProperties)
+		bFound = MoveMappedObjectToUnmapped_r(&RepState->GuidReferencesMap, GUID);
+
+		// Custom Delta Properties
+		if (LifetimeCustomPropertyState && Params.Object)
 		{
-			const FRepParentCmd& Parent = Parents[KVP.Key];
+			FRepObjectDataBuffer ObjectData(Params.Object);
+			const int32 NumLifetimeCustomDeltaProperties = LifetimeCustomPropertyState->GetNumCustomDeltaProperties();
 
-			// Static cast is safe here, because this property wouldn't have been marked CustomDelta otherwise.
-			UStructProperty* StructProperty = static_cast<UStructProperty*>(Parent.Property);
-			UScriptStruct::ICppStructOps* CppStructOps = StructProperty->Struct->GetCppStructOps();
+			for (int32 CustomDeltaIndex = 0; CustomDeltaIndex < NumLifetimeCustomDeltaProperties; ++CustomDeltaIndex)
+			{
+				const FLifetimeCustomDeltaProperty& CustomDeltaProperty = LifetimeCustomPropertyState->GetCustomDeltaProperty(CustomDeltaIndex);
+				const FRepParentCmd& Parent = Parents[CustomDeltaProperty.PropertyRepIndex];
 
-			FNetDeltaSerializeInfo TempParams = Params;
+				// Static cast is safe here, because this property wouldn't have been marked CustomDelta otherwise.
+				UStructProperty* StructProperty = static_cast<UStructProperty*>(Parent.Property);
+				UScriptStruct::ICppStructOps* CppStructOps = StructProperty->Struct->GetCppStructOps();
 
-			TempParams.Struct = StructProperty->Struct;
-			TempParams.Data = ObjectData + Parent;
-			TempParams.PropertyRepIndex = KVP.Key;
-			TempParams.bOutHasMoreUnmapped = false;
-			TempParams.bOutSomeObjectsWereMapped = false;
+				FNetDeltaSerializeInfo TempParams = Params;
 
-			if (CppStructOps->NetDeltaSerialize(TempParams, TempParams.Data))
-			{
-				UnmappedCustomProperties.Add(Parent.Offset, StructProperty);
-				bFound = true;
-			}
+				TempParams.Struct = StructProperty->Struct;
+				TempParams.Data = ObjectData + Parent;
+				TempParams.CustomDeltaIndex = CustomDeltaIndex;
+				TempParams.bOutHasMoreUnmapped = false;
+				TempParams.bOutSomeObjectsWereMapped = false;
+
+				if (CppStructOps->NetDeltaSerialize(TempParams, TempParams.Data))
+				{
+					bFound = true;
+				}
 
-			Params.bOutHasMoreUnmapped |= TempParams.bOutHasMoreUnmapped;
-			Params.bOutSomeObjectsWereMapped |= TempParams.bOutSomeObjectsWereMapped;
+				Params.bOutHasMoreUnmapped |= TempParams.bOutHasMoreUnmapped;
+				Params.bOutSomeObjectsWereMapped |= TempParams.bOutSomeObjectsWereMapped;
+			}
 		}
-}
+	}
 
 	return bFound;
 }
-PRAGMA_ENABLE_DEPRECATION_WARNINGS
 
 void FRepLayout::UpdateUnmappedObjects_r(
 	FReceivingRepState* RESTRICT RepState, 
@@ -3408,13 +3439,13 @@ void FRepLayout::UpdateUnmappedObjects_r(
 			
 			if (ShadowData)
 			{
-			FScriptArray* ShadowArray = (FScriptArray*)(ShadowData + Cmd).Data;
-			FScriptArray* Array = (FScriptArray*)(Data + AbsOffset).Data;
-			
-			FRepShadowDataBuffer ShadowArrayData(ShadowArray->GetData());
-			FRepObjectDataBuffer ArrayData(Array->GetData());
-			
-			const int32 NewMaxOffset = FMath::Min(ShadowArray->Num() * Cmd.ElementSize, Array->Num() * Cmd.ElementSize);
+				FScriptArray* ShadowArray = (FScriptArray*)(ShadowData + Cmd).Data;
+				FScriptArray* Array = (FScriptArray*)(Data + AbsOffset).Data;
+
+				FRepShadowDataBuffer ShadowArrayData(ShadowArray->GetData());
+				FRepObjectDataBuffer ArrayData(Array->GetData());
+
+				const int32 NewMaxOffset = FMath::Min(ShadowArray->Num() * Cmd.ElementSize, Array->Num() * Cmd.ElementSize);
 
 				UpdateUnmappedObjects_r(RepState, GuidReferences.Array, OriginalObject, PackageMap, ShadowArrayData, ArrayData, NewMaxOffset, bCalledPreNetReceive, bOutSomeObjectsWereMapped, bOutHasMoreUnmapped);
 			}
@@ -3470,7 +3501,7 @@ void FRepLayout::UpdateUnmappedObjects_r(
 				bCalledPreNetReceive = true;
 			}
 
-				bOutSomeObjectsWereMapped = true;
+			bOutSomeObjectsWereMapped = true;
 			const bool bUpdateShadowState = (ShadowData && INDEX_NONE != Parent.RepNotifyNumParams);
 
 			// Copy current value over so we can check to see if it changed
@@ -3522,6 +3553,7 @@ void FRepLayout::UpdateUnmappedObjects(
 	FReceivingRepState* RESTRICT RepState,
 	UPackageMap* PackageMap,
 	UObject* OriginalObject,
+	FNetDeltaSerializeInfo& Params,
 	bool& bCalledPreNetReceive,
 	bool& bOutSomeObjectsWereMapped,
 	bool& bOutHasMoreUnmapped) const
@@ -3543,71 +3575,54 @@ void FRepLayout::UpdateUnmappedObjects(
 			bCalledPreNetReceive,
 			bOutSomeObjectsWereMapped,
 			bOutHasMoreUnmapped);
-	}
-}
 
-PRAGMA_DISABLE_DEPRECATION_WARNINGS
-void FRepLayout::UpdateUnmappedObjectsForCustomDeltaProperties(FNetDeltaSerializeInfo& Params) const
-{
-	TArray<TPair<int32, UStructProperty*>> Dummy;
-	UpdateUnmappedObjectsForCustomDeltaProperties(Params, Dummy, Dummy);
-}
-PRAGMA_ENABLE_DEPRECATION_WARNINGS
+		Params.bCalledPreNetReceive = bCalledPreNetReceive;
 
-void FRepLayout::UpdateUnmappedObjectsForCustomDeltaProperties(FNetDeltaSerializeInfo& Params, TArray<TPair<int32, UStructProperty*>>& CompletelyUpdated, TArray<TPair<int32, UStructProperty*>>& Updated) const
-{
-	if (LifetimeCustomPropertyState)
-	{
-		FRepObjectDataBuffer ObjectData(Params.Object);
-		for (const auto& KVP : LifetimeCustomPropertyState->LifetimeCustomDeltaProperties)
+		// Custom Delta Properties
+		if (LifetimeCustomPropertyState)
 		{
-			const FRepParentCmd& Parent = Parents[KVP.Key];
+			FRepObjectDataBuffer ObjectData(Params.Object);
+			const int32 NumLifetimeCustomDeltaProperties = LifetimeCustomPropertyState->GetNumCustomDeltaProperties();
 
-			// Static cast is safe here, because this property wouldn't have been marked CustomDelta otherwise.
-			UStructProperty* StructProperty = static_cast<UStructProperty*>(Parent.Property);
-			UScriptStruct::ICppStructOps* CppStructOps = StructProperty->Struct->GetCppStructOps();
+			for (int32 CustomDeltaIndex = 0; CustomDeltaIndex < NumLifetimeCustomDeltaProperties; ++CustomDeltaIndex)
+			{
+				const FLifetimeCustomDeltaProperty& CustomDeltaProperty = LifetimeCustomPropertyState->GetCustomDeltaProperty(CustomDeltaIndex);
+				const FRepParentCmd& Parent = Parents[CustomDeltaProperty.PropertyRepIndex];
 
-			FNetDeltaSerializeInfo TempParams = Params;
+				// Static cast is safe here, because this property wouldn't have been marked CustomDelta otherwise.
+				UStructProperty* StructProperty = static_cast<UStructProperty*>(Parent.Property);
+				UScriptStruct::ICppStructOps* CppStructOps = StructProperty->Struct->GetCppStructOps();
 
-			TempParams.DebugName = Parent.CachedPropertyName.ToString();
-			TempParams.Struct = StructProperty->Struct;
-			TempParams.bOutSomeObjectsWereMapped = false;
-			TempParams.bOutHasMoreUnmapped = false;
-			TempParams.PropertyRepIndex = KVP.Key;
-			TempParams.Data = ObjectData + Parent;
+				FNetDeltaSerializeInfo TempParams = Params;
 
-			// Call the custom delta serialize function to handle it
-			CppStructOps->NetDeltaSerialize(TempParams, TempParams.Data);
+				TempParams.DebugName = Parent.CachedPropertyName.ToString();
+				TempParams.Struct = StructProperty->Struct;
+				TempParams.bOutSomeObjectsWereMapped = false;
+				TempParams.bOutHasMoreUnmapped = false;
+				TempParams.CustomDeltaIndex = CustomDeltaIndex;
+				TempParams.Data = ObjectData + Parent;
 
-			if (TempParams.bOutSomeObjectsWereMapped)
-			{
-				Updated.Emplace(Parent.Offset, StructProperty);
-			}
+				// Call the custom delta serialize function to handle it
+				CppStructOps->NetDeltaSerialize(TempParams, TempParams.Data);
 
-			if (!TempParams.bOutHasMoreUnmapped)
-			{
-				CompletelyUpdated.Emplace(Parent.Offset, StructProperty);
-			}
+				if (TempParams.bOutSomeObjectsWereMapped && INDEX_NONE != Parent.RepNotifyNumParams)
+				{
+					UE4_RepLayout_Private::QueueRepNotifyForCustomDeltaProperty(RepState, Params, StructProperty, Parent.ArrayIndex);
+				}
 
-			Params.bOutSomeObjectsWereMapped |= TempParams.bOutSomeObjectsWereMapped;
-			Params.bOutHasMoreUnmapped |= TempParams.bOutHasMoreUnmapped;
-			Params.bCalledPreNetReceive |= TempParams.bCalledPreNetReceive;
+				Params.bOutSomeObjectsWereMapped |= TempParams.bOutSomeObjectsWereMapped;
+				Params.bOutHasMoreUnmapped |= TempParams.bOutHasMoreUnmapped;
+				Params.bCalledPreNetReceive |= TempParams.bCalledPreNetReceive;
+			}
 		}
 	}
 }
 
-bool FRepLayout::SendCustomDeltaProperty(
-	FNetDeltaSerializeInfo& Params,
-	UProperty* Property,
-	int32 StaticArrayIndex) const
+bool FRepLayout::SendCustomDeltaProperty(FNetDeltaSerializeInfo& Params, uint16 CustomDeltaIndex) const
 {
-	check(Property->ArrayDim > StaticArrayIndex);
-	return SendCustomDeltaProperty(Params, Property->RepIndex + StaticArrayIndex);
-}
+	const FLifetimeCustomDeltaProperty& CustomDeltaProperty = LifetimeCustomPropertyState->GetCustomDeltaProperty(CustomDeltaIndex);
+	const FRepParentCmd& Parent = Parents[CustomDeltaProperty.PropertyRepIndex];
 
-bool FRepLayout::SendCustomDeltaProperty(FNetDeltaSerializeInfo& Params, uint16 CustomDeltaRepIndex) const
-{
-	const FRepParentCmd& Parent = Parents[CustomDeltaRepIndex];
 	if (!ensure(EnumHasAnyFlags(Parent.Flags, ERepParentFlags::IsCustomDelta)))
 	{
 		return false;
@@ -3620,16 +3635,15 @@ bool FRepLayout::SendCustomDeltaProperty(FNetDeltaSerializeInfo& Params, uint16
 
 	Params.DebugName = Parent.CachedPropertyName.ToString();
 	Params.Struct = StructProperty->Struct;
-	Params.PropertyRepIndex = CustomDeltaRepIndex;
+	Params.CustomDeltaIndex = CustomDeltaIndex;
 	Params.Data = FRepObjectDataBuffer(Params.Object) + Parent;
 
 	bool bSupportsFastArrayDelta = Params.bSupportsFastArrayDeltaStructSerialization;
 
 	if (Params.bSupportsFastArrayDeltaStructSerialization &&
 		EnumHasAnyFlags(Parent.Flags, ERepParentFlags::IsFastArray) &&
-		!!LifetimeCustomPropertyState->NumFastArrayItems)
+		!!LifetimeCustomPropertyState->GetNumFastArrayProperties())
 	{
-		const FLifetimeCustomDeltaProperty& CustomDeltaProperty = LifetimeCustomPropertyState->LifetimeCustomDeltaProperties.FindChecked(CustomDeltaRepIndex);
 		bSupportsFastArrayDelta = CustomDeltaProperty.FastArrayNumber != INDEX_NONE;
 	}
 
@@ -3645,16 +3659,10 @@ bool FRepLayout::SendCustomDeltaProperty(FNetDeltaSerializeInfo& Params, uint16
 	return CppStructOps->NetDeltaSerialize(Params, Params.Data);
 }
 
-PRAGMA_DISABLE_DEPRECATION_WARNINGS
-bool FRepLayout::ReceiveCustomDeltaProperty(FNetDeltaSerializeInfo& Params, UStructProperty* Property) const
-{
-	uint32 StaticArrayIndex = 0;
-	int32 Offset = 0;
-	return ReceiveCustomDeltaProperty(Params, Property, StaticArrayIndex, Offset);
-}
-PRAGMA_ENABLE_DEPRECATION_WARNINGS
-
-bool FRepLayout::ReceiveCustomDeltaProperty(FNetDeltaSerializeInfo& Params, UStructProperty* Property, uint32& StaticArrayIndex, int32& Offset) const
+bool FRepLayout::ReceiveCustomDeltaProperty(
+	FReceivingRepState* RESTRICT ReceivingRepState,
+	FNetDeltaSerializeInfo& Params,
+	UStructProperty* Property) const
 {
 	if (Params.Connection->EngineNetworkProtocolVersion >= EEngineNetworkVersionHistory::HISTORY_FAST_ARRAY_DELTA_STRUCT)
 	{
@@ -3665,6 +3673,8 @@ bool FRepLayout::ReceiveCustomDeltaProperty(FNetDeltaSerializeInfo& Params, UStr
 		Params.bSupportsFastArrayDeltaStructSerialization = false;
 	}
 
+	uint32 StaticArrayIndex = 0;
+
 	// Receive array index (static sized array, i.e. MemberVariable[4])
 	if (Property->ArrayDim != 1)
 	{
@@ -3688,7 +3698,6 @@ bool FRepLayout::ReceiveCustomDeltaProperty(FNetDeltaSerializeInfo& Params, UStr
 		return false;
 	}
 
-	Offset = Parent.Offset;
 
 	UScriptStruct* InnerStruct = Property->Struct;
 	UScriptStruct::ICppStructOps* CppStructOps = InnerStruct->GetCppStructOps();
@@ -3697,10 +3706,32 @@ bool FRepLayout::ReceiveCustomDeltaProperty(FNetDeltaSerializeInfo& Params, UStr
 
 	Params.DebugName = Parent.CachedPropertyName.ToString();
 	Params.Struct = InnerStruct;
-	Params.PropertyRepIndex = Property->RepIndex + StaticArrayIndex;
+	Params.CustomDeltaIndex = LifetimeCustomPropertyState->GetCustomDeltaIndexFromPropertyRepIndex(Property->RepIndex + StaticArrayIndex);
 	Params.Data = FRepObjectDataBuffer(Params.Object) + Parent;
 
-	return CppStructOps->NetDeltaSerialize(Params, Params.Data);
+	if (CppStructOps->NetDeltaSerialize(Params, Params.Data))
+	{
+		if (UNLIKELY(Params.Reader->IsError()))
+		{
+			UE_LOG(LogNet, Error, TEXT("FRepLayout::ReceiveCustomDeltaProperty: NetDeltaSerialize - Reader.IsError() == true. Property: %s, Object: %s"), *Params.DebugName, *Params.Object->GetFullName());
+			return false;
+		}
+		if (UNLIKELY(Params.Reader->GetBitsLeft() != 0))
+		{
+			UE_LOG(LogNet, Error, TEXT("FRepLayout::ReceiveCustomDeltaProperty: NetDeltaSerialize - Mismatch read. Property: %s, Object: %s"), *Params.DebugName, *Params.Object->GetFullName());
+			return false;
+		}
+
+		// Successfully received it.
+		if (INDEX_NONE != Parent.RepNotifyNumParams)
+		{
+			UE4_RepLayout_Private::QueueRepNotifyForCustomDeltaProperty(ReceivingRepState, Params, Property, StaticArrayIndex);
+		}
+
+		return true;
+	}
+
+	return false;
 }
 
 void FRepLayout::CallRepNotifies(FReceivingRepState* RepState, UObject* Object) const
@@ -3723,6 +3754,13 @@ void FRepLayout::CallRepNotifies(FReceivingRepState* RepState, UObject* Object)
 
 	for (UProperty* RepProperty : RepState->RepNotifies)
 	{
+		if (!Parents.IsValidIndex(RepProperty->RepIndex))
+		{
+			UE_LOG(LogRep, Warning, TEXT("FRepLayout::CallRepNotifies: Called with invalid property %s on object %s."),
+				*RepProperty->GetName(), *Object->GetName());
+				continue;
+		}
+
 		UFunction* RepNotifyFunc = Object->FindFunction(RepProperty->RepNotifyFunc);
 
 		if (RepNotifyFunc == nullptr)
@@ -3732,46 +3770,90 @@ void FRepLayout::CallRepNotifies(FReceivingRepState* RepState, UObject* Object)
 			continue;
 		}
 
-		check(RepNotifyFunc->NumParms <= 1);	// 2 parms not supported yet
+		const FRepParentCmd& Parent = Parents[RepProperty->RepIndex];
+		const int32 NumParms = RepNotifyFunc->NumParms;
 
-		if (RepNotifyFunc->NumParms == 0)
-		{
-			Object->ProcessEvent(RepNotifyFunc, nullptr);
-		}
-		else if (RepNotifyFunc->NumParms == 1)
+		switch (NumParms)
 		{
-			// TODO_JDN: If this turns out to be too slow to be practical,
-			//	we should consider a TSortedMap of *just* these RepNotify properties, or even just
-			//	an Array of indices of RepNotify properties.
-			const FRepParentCmd* Parent = Parents.FindByPredicate([RepProperty](const FRepParentCmd& InParent)
+			case 0:
 			{
-				return InParent.Property == RepProperty;
-			});
+				Object->ProcessEvent(RepNotifyFunc, nullptr);
+				break;
+			}
+			case 1:
+			{
+				FRepShadowDataBuffer PropertyData = ShadowData + Parent;
 
-			check(Parent);
-			
-			FRepShadowDataBuffer PropertyData = ShadowData + (*Parent);
+				if (EnumHasAnyFlags(Parent.Flags, ERepParentFlags::IsCustomDelta))
+				{
+					Object->ProcessEvent(RepNotifyFunc, PropertyData);
+				}
+				else
+				{
+					// This could be cached off as a Parent flag, to avoid touching the Commands array.
+					if (ERepLayoutCmdType::PropertyBool == Cmds[Parent.CmdStart].Type)
+					{
+						bool BoolPropertyValue = !!static_cast<const UBoolProperty*>(Parent.Property)->GetPropertyValue(PropertyData);
+						Object->ProcessEvent(RepNotifyFunc, &BoolPropertyValue);
+					}
+					else
+					{
+						Object->ProcessEvent(RepNotifyFunc, PropertyData);
+					}
 
-			// This could be cached off as a Parent flag, to avoid touching the Commands array.
-			if (ERepLayoutCmdType::PropertyBool == Cmds[Parent->CmdStart].Type)
-			{
-				bool BoolPropertyValue = !!static_cast<const UBoolProperty*>(Parent->Property)->GetPropertyValue(PropertyData);
-				Object->ProcessEvent(RepNotifyFunc, &BoolPropertyValue);
+					// now store the complete value in the shadow buffer
+					if (!EnumHasAnyFlags(Parent.Flags, ERepParentFlags::IsNetSerialize))
+					{
+						RepProperty->CopyCompleteValue(ShadowData + Parent, ObjectData + Parent);
+					}
+				}
+				break;
 			}
-			else
+			case 2:
 			{
-				Object->ProcessEvent(RepNotifyFunc, PropertyData);
+				check(EnumHasAnyFlags(Parent.Flags, ERepParentFlags::IsCustomDelta));
+
+				// Fixme: this isn't as safe as it could be. Right now we have two types of parameters: MetaData (a TArray<uint8>)
+				// and the last local value (pointer into the Recent[] array).
+				//
+				// Arrays always expect MetaData. Everything else, including structs, expect last value.
+				// This is enforced with UHT only. If a ::NetSerialize function ever starts producing a MetaData array thats not in UArrayProperty,
+				// we have no static way of catching this and the replication system could pass the wrong thing into ProcessEvent here.
+				//
+				// But this is all sort of an edge case feature anyways, so its not worth tearing things up too much over.
+
+				FMemMark Mark(FMemStack::Get());
+				uint8* Parms = new(FMemStack::Get(), MEM_Zeroed, RepNotifyFunc->ParmsSize)uint8;
+
+				TFieldIterator<UProperty> Itr(RepNotifyFunc);
+				check(Itr);
+
+				FRepShadowDataBuffer PropertyData = ShadowData + Parent;
+
+				Itr->CopyCompleteValue(Itr->ContainerPtrToValuePtr<void>(Parms), PropertyData);
+				++Itr;
+				check(Itr);
+
+				TArray<uint8> *NotifyMetaData = RepState->RepNotifyMetaData.Find(RepProperty);
+				check(NotifyMetaData);
+				Itr->CopyCompleteValue(Itr->ContainerPtrToValuePtr<void>(Parms), NotifyMetaData);
+
+				Object->ProcessEvent(RepNotifyFunc, Parms);
+
+				Mark.Pop();
+				break;
 			}
-			
-			// now store the complete value in the shadow buffer
-			if (!EnumHasAnyFlags(Parent->Flags, ERepParentFlags::IsNetSerialize | ERepParentFlags::IsCustomDelta))
+			default:
 			{
-				RepProperty->CopyCompleteValue(ShadowData + (*Parent), ObjectData + (*Parent));
+				checkf(false, TEXT("FRepLayout::CallRepNotifies: Invalid number of parameters for property %s on object %s. NumParms=%d, CustomDelta=%d"),
+					*RepProperty->GetName(), *Object->GetName(), NumParms, !!EnumHasAnyFlags(Parent.Flags, ERepParentFlags::IsCustomDelta));
+				break;
 			}
 		}
 	}
 
 	RepState->RepNotifies.Empty();
+	RepState->RepNotifyMetaData.Empty();
 }
 
 template<ERepDataBufferType DataType>
@@ -4522,7 +4604,7 @@ static int32 InitFromProperty_r(
 		const uint32 ArrayChecksum = AddArrayCmd(Cmds, ArrayProp, Offset + GetOffsetForProperty<BuildType>(*ArrayProp), RelativeHandle, ParentIndex, ParentChecksum, StaticArrayIndex, ServerConnection);
 
 		InitFromProperty_r<BuildType>(Cmds, ArrayProp->Inner, 0, 0, ParentIndex, ArrayChecksum, 0, ServerConnection);
-		
+
 		AddReturnCmd(Cmds);
 
 		Cmds[CmdStart].EndCmd = Cmds.Num();		// Patch in the offset to jump over our array inner elements
@@ -4965,12 +5047,9 @@ void FRepLayout::InitFromClass(UClass* InObjectClass, const UNetConnection* Serv
 
 			if (!LifetimeCustomPropertyState)
 			{
-				LifetimeCustomPropertyState.Reset(new FLifetimeCustomDeltaState());
+				LifetimeCustomPropertyState.Reset(new FLifetimeCustomDeltaState(LifetimeProps.Num()));
 			}
 
-			const int32 CustomIndex = LifetimeCustomPropertyState->LifetimeCustomDeltaPropertiesParentCache.Add(ParentIndex);
-			FLifetimeCustomDeltaProperty& CustomDeltaProperty = LifetimeCustomPropertyState->LifetimeCustomDeltaProperties.Add(ParentIndex);
-
 			// If we're a FastArraySerializer, we'll look for our replicated item type.
 			// We do this by looking for an array property whose inner type is an FFastArraySerializerItem.
 			// Note, this isn't perfect. With the way the interface is set up now, there's no technically
@@ -4981,6 +5060,8 @@ void FRepLayout::InitFromClass(UClass* InObjectClass, const UNetConnection* Serv
 			// However, comments imply these, and typically they are true (certainly, any engine cases follow this).
 			// Further, these layouts are only needed for the new Delta Struct Serialization feature, so this won't break backwards compat.
 
+			bool bAddedFastArray = false;
+
 			if (EnumHasAnyFlags(Parents[ParentIndex].Flags, ERepParentFlags::IsFastArray))
 			{
 				int32 FastArrayItemArrayCmd = INDEX_NONE;
@@ -5003,13 +5084,16 @@ void FRepLayout::InitFromClass(UClass* InObjectClass, const UNetConnection* Serv
 								// This better be a script struct, otherwise our flags aren't set up correctly!
 								UScriptStruct* FastArray = CastChecked<UScriptStruct>(MaybeFastArrayItemsArray->GetOwnerStruct());
 
-								CustomDeltaProperty.FastArrayItemsCommand = CmdIndex;
-								CustomDeltaProperty.FastArrayItemReplicationIdOffset = MaybeFastArrayItem->FindPropertyByName(FastArrayItemReplicationIDName)->GetOffset_ForGC();
-								CustomDeltaProperty.FastArrayArrayReplicationKeyOffset = FastArray->FindPropertyByName(FastArrayArrayReplicationKeyName)->GetOffset_ForGC();
-								CustomDeltaProperty.FastArrayDeltaFlagsOffset = FastArray->FindPropertyByName(FastArrayDeltaFlagsName)->GetOffset_ForGC();
+								LifetimeCustomPropertyState->Add(FLifetimeCustomDeltaProperty(
+									/*PropertyRepIndex=*/ParentIndex,
+									/*FastArrayItemsCommand=*/CmdIndex,
+									/*FastArrayNumber=*/LifetimeCustomPropertyState->GetNumFastArrayProperties(),
+									/*FastArrayDeltaFlagsOffset=*/FastArray->FindPropertyByName(FastArrayDeltaFlagsName)->GetOffset_ForGC(),
+									/*FastArrayReplicationKeyOffset=*/FastArray->FindPropertyByName(FastArrayArrayReplicationKeyName)->GetOffset_ForGC(),
+									/*FastArrayItemReplicationIdOffset=*/MaybeFastArrayItem->FindPropertyByName(FastArrayItemReplicationIDName)->GetOffset_ForGC()
+								));
 
-								CustomDeltaProperty.FastArrayNumber = LifetimeCustomPropertyState->NumFastArrayItems;
-								++LifetimeCustomPropertyState->NumFastArrayItems;
+								bAddedFastArray = true;
 								break;
 							}
 						}
@@ -5018,12 +5102,17 @@ void FRepLayout::InitFromClass(UClass* InObjectClass, const UNetConnection* Serv
 					}
 				}
 
-				if (CustomDeltaProperty.FastArrayItemsCommand == INDEX_NONE)
+				if (!bAddedFastArray)
 		{
 					UE_LOG(LogRep, Warning, TEXT("FRepLayout::InitFromClass: Unable to find Fast Array Item array in Fast Array Serializer: %s"), *Parents[ParentIndex].CachedPropertyName.ToString());
 				}
 			}
 
+			if (!bAddedFastArray)
+			{
+				LifetimeCustomPropertyState->Add(FLifetimeCustomDeltaProperty(ParentIndex));
+			}
+
 			continue;
 		}
 
@@ -5049,19 +5138,6 @@ void FRepLayout::InitFromClass(UClass* InObjectClass, const UNetConnection* Serv
 		BuildHandleToCmdIndexTable_r(0, Cmds.Num() - 1, BaseHandleToCmdIndex);
 	}
 
-	if (LifetimeCustomPropertyState && !!LifetimeCustomPropertyState->NumFastArrayItems)
-	{
-		int32 Counter = 0;
-		for (auto& KVP : LifetimeCustomPropertyState->LifetimeCustomDeltaProperties)
-		{
-			if (EnumHasAnyFlags(Parents[KVP.Key].Flags, ERepParentFlags::IsFastArray) &&
-				INDEX_NONE != KVP.Value.FastArrayItemsCommand)
-			{
-				KVP.Value.FastArrayNumber = Counter++;
-			}
-		}
-	}
-
 	BuildShadowOffsets<ERepBuildType::Class>(InObjectClass, Parents, Cmds, ShadowDataBufferSize, LayoutState);
 
 	Owner = InObjectClass;
@@ -5802,9 +5878,9 @@ TSharedPtr<FReplicationChangelistMgr> FRepLayout::CreateReplicationChangelistMgr
 	// so no need to worry about deleting it here.
 
 	FCustomDeltaChangelistState* DeltaChangelistState = nullptr;
-	if (LifetimeCustomPropertyState && !!LifetimeCustomPropertyState->NumFastArrayItems)
+	if (LifetimeCustomPropertyState && !!LifetimeCustomPropertyState->GetNumFastArrayProperties())
 	{
-		DeltaChangelistState = new FCustomDeltaChangelistState(LifetimeCustomPropertyState->NumFastArrayItems);
+		DeltaChangelistState = new FCustomDeltaChangelistState(LifetimeCustomPropertyState->GetNumFastArrayProperties());
 	}
 
 	const uint8* ShadowStateSource = (const uint8*)InObject->GetArchetype();
@@ -5917,37 +5993,6 @@ void FRepLayout::DestructProperties(FRepStateStaticBuffer& InShadowData) const
 	InShadowData.Buffer.Empty();
 }
 
-PRAGMA_DISABLE_DEPRECATION_WARNINGS
-
-void FRepLayout::GetLifetimeCustomDeltaProperties(TArray<int32>& OutCustom, TArray<ELifetimeCondition>& OutConditions) const
-{
-	OutCustom.Empty();
-	OutConditions.Empty();
-
-	if (LifetimeCustomPropertyState)
-	{
-		for (const auto& KVP : LifetimeCustomPropertyState->LifetimeCustomDeltaProperties)
-		{
-			OutCustom.Add(KVP.Key);
-			OutConditions.Add(Parents[KVP.Key].Condition);
-		}
-	}
-}
-
-PRAGMA_ENABLE_DEPRECATION_WARNINGS
-
-const TArrayView<const uint16> FRepLayout::GetLifetimeCustomDeltaProperties() const
-{
-	if (LifetimeCustomPropertyState)
-	{
-		return LifetimeCustomPropertyState->LifetimeCustomDeltaPropertiesParentCache;
-		}
-	else
-	{
-		return TArrayView<const uint16>();
-	}
-}
-
 void FRepLayout::AddReferencedObjects(FReferenceCollector& Collector)
 {
 	UProperty* Current = nullptr;
@@ -6028,7 +6073,7 @@ void FRepLayout::PreSendCustomDeltaProperties(
 	{
 		const FLifetimeCustomDeltaState& LocalLifetimeCustomPropertyState = *LifetimeCustomPropertyState;
 
-		if (LocalLifetimeCustomPropertyState.NumFastArrayItems)
+		if (LocalLifetimeCustomPropertyState.GetNumFastArrayProperties())
 		{
 			FRepChangelistState& ChangelistState = *ChangelistMgr.GetRepChangelistState();
 			FCustomDeltaChangelistState& CustomDeltaChangelistState = *ChangelistState.CustomDeltaChangelistState;
@@ -6041,11 +6086,12 @@ void FRepLayout::PreSendCustomDeltaProperties(
 				CustomDeltaChangelistState.CompareIndex = GFrameCounter;
 
 				const FConstRepObjectDataBuffer ObjectData(Object);
+				const uint16 NumLifetimeCustomDeltaProperties = LocalLifetimeCustomPropertyState.GetNumCustomDeltaProperties();
 
-				for (const auto& KVP : LocalLifetimeCustomPropertyState.LifetimeCustomDeltaProperties)
+				for (uint16 CustomDeltaIndex = 0; CustomDeltaIndex < NumLifetimeCustomDeltaProperties; ++CustomDeltaIndex)
 				{
-					const uint16 RepIndex = KVP.Key;
-					const FLifetimeCustomDeltaProperty& CustomDeltaProperty = KVP.Value;
+					const FLifetimeCustomDeltaProperty& CustomDeltaProperty = LocalLifetimeCustomPropertyState.GetCustomDeltaProperty(CustomDeltaIndex);
+					const uint16 RepIndex = CustomDeltaProperty.PropertyRepIndex;
 
 					// If our Fast Array Items Command is invalid, we can't do anything.
 					// This should have been logged on RepLayout creation.
@@ -6126,8 +6172,10 @@ bool FRepLayout::DeltaSerializeFastArrayProperty(FFastArrayDeltaSerializeParams&
 	check(LifetimeCustomPropertyState);
 
 	FNetDeltaSerializeInfo& DeltaSerializeInfo = Params.DeltaSerializeInfo;
-	const int32 ParentIndex = DeltaSerializeInfo.PropertyRepIndex;
-	const FLifetimeCustomDeltaProperty& CustomDeltaProperty = LifetimeCustomPropertyState->LifetimeCustomDeltaProperties.FindChecked(ParentIndex);
+
+	const FLifetimeCustomDeltaProperty& CustomDeltaProperty = LifetimeCustomPropertyState->GetCustomDeltaProperty(DeltaSerializeInfo.CustomDeltaIndex);
+	const uint16 ParentIndex = CustomDeltaProperty.PropertyRepIndex;
+
 	const FRepParentCmd& Parent = Parents[ParentIndex];
 	const int32 CmdIndex = CustomDeltaProperty.FastArrayItemsCommand;
 
@@ -6716,7 +6764,8 @@ void FRepLayout::GatherGuidReferencesForFastArray(FFastArrayDeltaSerializeParams
 	using namespace UE4_RepLayout_Private;
 
 	const FConstRepObjectDataBuffer ObjectData(Params.DeltaSerializeInfo.Object);
-	const FRepParentCmd& Parent = Parents[Params.DeltaSerializeInfo.PropertyRepIndex];
+	const FLifetimeCustomDeltaProperty& CustomDeltaProperty = LifetimeCustomPropertyState->GetCustomDeltaProperty(Params.DeltaSerializeInfo.CustomDeltaIndex);
+	const FRepParentCmd& Parent = Parents[CustomDeltaProperty.PropertyRepIndex];
 
 	const FFastArraySerializer& ArraySerializer = Params.ArraySerializer;
 	TSet<FNetworkGUID>& GatherGuids = *Params.DeltaSerializeInfo.GatherGuidReferences;
@@ -6738,7 +6787,8 @@ bool FRepLayout::MoveMappedObjectToUnmappedForFastArray(FFastArrayDeltaSerialize
 	using namespace UE4_RepLayout_Private;
 
 	const FRepObjectDataBuffer ObjectData(Params.DeltaSerializeInfo.Object);
-	const FRepParentCmd& Parent = Parents[Params.DeltaSerializeInfo.PropertyRepIndex];
+	const FLifetimeCustomDeltaProperty& CustomDeltaProperty = LifetimeCustomPropertyState->GetCustomDeltaProperty(Params.DeltaSerializeInfo.CustomDeltaIndex);
+	const FRepParentCmd& Parent = Parents[CustomDeltaProperty.PropertyRepIndex];
 
 	FFastArraySerializer& ArraySerializer = Params.ArraySerializer;
 	const FNetworkGUID& MoveToUnmapped = *Params.DeltaSerializeInfo.MoveGuidToUnmapped;
@@ -6758,8 +6808,9 @@ void FRepLayout::UpdateUnmappedGuidsForFastArray(FFastArrayDeltaSerializeParams&
 	check(LifetimeCustomPropertyState);
 
 	FNetDeltaSerializeInfo& DeltaSerializeInfo = Params.DeltaSerializeInfo;
-	const int32 ParentIndex = DeltaSerializeInfo.PropertyRepIndex;
-	const FLifetimeCustomDeltaProperty& CustomDeltaProperty = LifetimeCustomPropertyState->LifetimeCustomDeltaProperties.FindChecked(ParentIndex);
+
+	const FLifetimeCustomDeltaProperty& CustomDeltaProperty = LifetimeCustomPropertyState->GetCustomDeltaProperty(DeltaSerializeInfo.CustomDeltaIndex);
+	const int32 ParentIndex = CustomDeltaProperty.PropertyRepIndex;
 	const FRepParentCmd& Parent = Parents[ParentIndex];
 	const int32 CmdIndex = CustomDeltaProperty.FastArrayItemsCommand;
 	
@@ -6829,6 +6880,23 @@ void FRepLayout::CountBytes(FArchive& Ar) const
 	);
 }
 
+const uint16 FRepLayout::GetNumLifetimeCustomDeltaProperties() const
+{
+	return LifetimeCustomPropertyState.IsValid() ? LifetimeCustomPropertyState->GetNumCustomDeltaProperties() : 0;
+}
+
+UProperty* FRepLayout::GetLifetimeCustomDeltaProperty(const uint16 CustomDeltaPropertyIndex) const
+{
+	const FLifetimeCustomDeltaProperty& CustomDeltaProperty = LifetimeCustomPropertyState->GetCustomDeltaProperty(CustomDeltaPropertyIndex);
+	return Parents[CustomDeltaProperty.PropertyRepIndex].Property;
+}
+
+const ELifetimeCondition FRepLayout::GetLifetimeCustomDeltaPropertyCondition(const uint16 CustomDeltaPropertyIndex) const
+{
+	const FLifetimeCustomDeltaProperty& CustomDeltaProperty = LifetimeCustomPropertyState->GetCustomDeltaProperty(CustomDeltaPropertyIndex);
+	return Parents[CustomDeltaProperty.PropertyRepIndex].Condition;
+}
+
 void FReceivingRepState::CountBytes(FArchive& Ar) const
 {
 	GRANULAR_NETWORK_MEMORY_TRACKING_INIT(Ar, "FReceivingRepState::CountBytes");
@@ -6842,8 +6910,16 @@ void FReceivingRepState::CountBytes(FArchive& Ar) const
 			GuidRefPair.Value.CountBytes(Ar);
 		}
 	);
-	
+
 	GRANULAR_NETWORK_MEMORY_TRACKING_TRACK("RepNotifies", RepNotifies.CountBytes(Ar));
+
+	GRANULAR_NETWORK_MEMORY_TRACKING_TRACK("RepNotifyMetaData",
+		RepNotifyMetaData.CountBytes(Ar);
+		for (const auto& MetaDataPair : RepNotifyMetaData)
+		{
+			MetaDataPair.Value.CountBytes(Ar);
+		}
+	);
 }
 
 void FSendingRepState::CountBytes(FArchive& Ar) const
@@ -6869,6 +6945,30 @@ void FSendingRepState::CountBytes(FArchive& Ar) const
 	GRANULAR_NETWORK_MEMORY_TRACKING_TRACK("LifetimeChangelist", LifetimeChangelist.CountBytes(Ar));
 	GRANULAR_NETWORK_MEMORY_TRACKING_TRACK("InactiveChangelist", InactiveChangelist.CountBytes(Ar));
 	GRANULAR_NETWORK_MEMORY_TRACKING_TRACK("InactiveParents", InactiveParents.CountBytes(Ar));
+
+	GRANULAR_NETWORK_MEMORY_TRACKING_TRACK("Retirement", Retirement.CountBytes(Ar));
+
+	GRANULAR_NETWORK_MEMORY_TRACKING_TRACK("RecentCustomDeltaState",
+		RecentCustomDeltaState.CountBytes(Ar);
+	for (const auto& RecentCustomDeltaStatePair : RecentCustomDeltaState)
+	{
+		if (INetDeltaBaseState const * const BaseState = RecentCustomDeltaStatePair.Value.Get())
+		{
+			BaseState->CountBytes(Ar);
+		}
+	}
+	);
+
+	GRANULAR_NETWORK_MEMORY_TRACKING_TRACK("CDOCustomDeltaState",
+		CDOCustomDeltaState.CountBytes(Ar);
+		for (const auto& CDOCustomDeltaStatePair : CDOCustomDeltaState)
+		{
+			if (INetDeltaBaseState const * const BaseState = CDOCustomDeltaStatePair.Value.Get())
+			{
+				BaseState->CountBytes(Ar);
+			}
+		}
+	);
 }
 
 void FRepState::CountBytes(FArchive& Ar) const
diff --git a/Engine/Source/Runtime/Engine/Public/Net/DataReplication.h b/Engine/Source/Runtime/Engine/Public/Net/DataReplication.h
--- a/Engine/Source/Runtime/Engine/Public/Net/DataReplication.h
+++ b/Engine/Source/Runtime/Engine/Public/Net/DataReplication.h
@@ -44,13 +44,15 @@ struct FReplicatedActorProperty
 	const class UObjectPropertyBase* Property;
 
 	FReplicatedActorProperty(int32 InOffset, const UObjectPropertyBase* InProperty)
-		: Offset(InOffset), Property(InProperty)
+		: Offset(InOffset)
+		, Property(InProperty)
 	{}
 };
 
-/** FObjectReplicator
- *   Generic class that replicates properties for an object.
- *	 Its primary job is to produce and consume chunks of properties/RPCs:
+/**
+ * Represents an object that is currently being replicated or handling RPCs.
+ *
+ *
  *
  *		|----------------|
  *		| NetGUID ObjRef |
@@ -68,88 +70,17 @@ struct FReplicatedActorProperty
 class ENGINE_API FObjectReplicator
 {
 public:
+
 	FObjectReplicator();
 	~FObjectReplicator();
 
-	UClass *										ObjectClass;
-	FNetworkGUID									ObjectNetGUID;
-	UObject*										ObjectPtr;
-
-	//~ Retirement, RecentCustomDeltaState, CDOCustomDeltaState, and CheckpointCustomDeltaState
-	//~ could all be moved to FSendingRepState, as they are not needed when solely receiving.
-
-	TArray<FPropertyRetirement>						Retirement;					// Property retransmission.
-	TMap<int32, TSharedPtr<INetDeltaBaseState> >	RecentCustomDeltaState;		// This is the delta state we need to compare with when determining what to send to a client for custom delta properties
-	TMap<int32, TSharedPtr<INetDeltaBaseState> >	CDOCustomDeltaState;		// Same as RecentCustomDeltaState, but this will always remain as the initial CDO version. We use this to send all properties since channel was first opened (for bResendAllDataSinceOpen)
-	TMap<int32, TSharedPtr<INetDeltaBaseState> >	CheckpointCustomDeltaState;	// Same as RecentCustomDeltaState, but will represent the state at the last checkpoint
-
-	UE_DEPRECATED(4.23, "This property will be removed in future versions. Use FRepLayout::GetLifetimeCustomDeltaProperties instead.")
-	TArray< int32 >									LifetimeCustomDeltaProperties;
-
-	UE_DEPRECATED(4.23, "This property will be removed in future versions. Use FRepLayout::GetRepPropertyCondition instead.")
-	TArray< ELifetimeCondition >					LifetimeCustomDeltaPropertyConditions;
-
-	/** True if last update (ReplicateActor) produced no replicated properties */
-	uint32 bLastUpdateEmpty : 1;
-
-	/** Whether or not the Actor Channel on which we're replicating has been Opened / Acked by the receiver. */
-	uint32 bOpenAckCalled : 1;
-
-	/** True if we need to do an unmapped check next frame. */
-	uint32 bForceUpdateUnmapped: 1;
-
-	/** Whether or not we've already replicated properties this frame. */
-	uint32 bHasReplicatedProperties : 1;
-	
-private:
-
-	uint32 bSupportsFastArrayDelta : 1;
-
-public:	
-
-	/** Connection this replicator was created on. */
-	UNetConnection* Connection;
-
-	/** The Actor Channel that we're replicating on. This expected to be owned by Connection. */
-	class UActorChannel* OwningChannel;
-
-	//~ This property is only used as an acceleration to avoid GUID tracking / updates when we know there are
-	//~ none unmapped. This will just be removed in future versions, but if the performance hit is unacceptable
-	//~ it could be replaced by a TBitArray in FReceivingRepState, where each bit flags whether or not
-	//~ a given index in FRepLayout::LifetimeCustomDeltaProperties has unmapped guids.
-	//~ That should be both more CPU and Memory efficient than the current approach.
-	//~ TODO: Remove FRepLayout::MoveMappedToUnmappedForCustomDeltaProperties overload when this goes away.
-	//~ TODO: Remove FRepLayout::UpdateUnmappedObjectsForCustomDeltaProperties overload when this goes away.
-	UE_DEPRECATED(4.23, "This property will be removed in future versions.")
-	TMap< int32, UStructProperty* >					UnmappedCustomProperties;
-
-	//~ TODO: RepNotifies here could likely just be removed in favor of using the existing array on FReceivingRepState.
-	//~			In most cases, RepNotifyMetaData isn't really used, but if it is useful, then that should also
-	//~			be moved to FReceivingRepState and we should combine the RepNotify logic here with that in
-	//~			FRepLayout.
-
-	TArray< UProperty*,TInlineAllocator< 32 > >		RepNotifies;
-	TMap< UProperty*, TArray<uint8> >				RepNotifyMetaData;
-
-	TSharedPtr< FRepLayout >						RepLayout;
-	TUniquePtr< FRepState > 						RepState;
-	TUniquePtr< FRepState >							CheckpointRepState;
-
-	TSet< FNetworkGUID >							ReferencedGuids;
-	int32											TrackedGuidMemoryBytes;
-
-	TSharedPtr<class FReplicationChangelistMgr> ChangelistMgr;
-
 	struct FRPCCallInfo 
 	{
-		FName	FuncName;
-		int32	Calls;
-		float	LastCallTime;
+		FName FuncName;
+		int32 Calls;
+		float LastCallTime;
 	};
 
-	TArray< FRPCCallInfo >							RemoteFuncInfo;				// Meta information on pending net RPCs (to be sent)
-	FOutBunch *										RemoteFunctions;
-
 	struct FRPCPendingLocalCall
 	{
 		/** Index to the RPC that was delayed */
@@ -167,38 +98,66 @@ public:
 		/** Guids being waited on */
 		TSet<FNetworkGUID> UnmappedGuids;
 
-		FRPCPendingLocalCall(const FFieldNetCache* InRPCField, const FReplicationFlags& InRepFlags, FNetBitReader& InReader, const TSet<FNetworkGUID>& InUnmappedGuids)
-			: RPCFieldIndex(InRPCField->FieldNetIndex), RepFlags(InRepFlags), Buffer(InReader.GetBuffer()), NumBits(InReader.GetNumBits()), UnmappedGuids(InUnmappedGuids)
+		FRPCPendingLocalCall(
+			const FFieldNetCache* InRPCField,
+			const FReplicationFlags& InRepFlags,
+			FNetBitReader& InReader,
+			const TSet<FNetworkGUID>& InUnmappedGuids)
+			: RPCFieldIndex(InRPCField->FieldNetIndex)
+			, RepFlags(InRepFlags)
+			, Buffer(InReader.GetBuffer())
+			, NumBits(InReader.GetNumBits())
+			, UnmappedGuids(InUnmappedGuids)
 		{}
 
 		void CountBytes(FArchive& Ar) const;
 	};
 
-	TArray< FRPCPendingLocalCall >					PendingLocalRPCs;			// Information on RPCs that have been received but not yet executed
+	void InitWithObject(
+		UObject* InObject,
+		UNetConnection* InConnection,
+		bool bUseDefaultState = true);
 
-	void InitWithObject( UObject* InObject, UNetConnection * InConnection, bool bUseDefaultState = true );
 	void CleanUp();
 
-	void StartReplicating( class UActorChannel * InActorChannel );
-	void StopReplicating( class UActorChannel * InActorChannel );
+	void StartReplicating(class UActorChannel* InActorChannel);
+	void StopReplicating(class UActorChannel* InActorChannel);
 
 	/** Recent/dirty related functions */
-	void InitRecentProperties( uint8* Source );
+	void InitRecentProperties(uint8* Source);
 
 	/** Takes Data, and compares against shadow state to log differences */
-	bool ValidateAgainstState( const UObject* ObjectState );
+	bool ValidateAgainstState(const UObject* ObjectState);
 
 	UE_DEPRECATED(4.23, "This method will be removed in future versions. Please use SendCustomDeltaProperty instead.")
-	static bool SerializeCustomDeltaProperty( UNetConnection * Connection, void* Src, UProperty * Property, uint32 ArrayIndex, FNetBitWriter & OutBunch, TSharedPtr<INetDeltaBaseState> & NewFullState, TSharedPtr<INetDeltaBaseState> & OldState );
+	static bool SerializeCustomDeltaProperty(
+		UNetConnection* Connection,
+		void* Src, UProperty* Property,
+		uint32 ArrayIndex,
+		FNetBitWriter& OutBunch,
+		TSharedPtr<INetDeltaBaseState>& NewFullState,
+		TSharedPtr<INetDeltaBaseState> & OldState);
 
 	//~ Both of these should be private, IMO, but we'll leave them public for now for back compat
 	//~ in case anyone was using SerializeCustomDeltaProperty already.
 
-	bool SendCustomDeltaProperty(UObject* InObject, UProperty* Property, uint32 ArrayIndex, FNetBitWriter& OutBunch, TSharedPtr<INetDeltaBaseState> & NewFullState, TSharedPtr<INetDeltaBaseState> & OldState);
-	bool SendCustomDeltaProperty(UObject* InObject, uint16 CustomDeltaProperty, FNetBitWriter& OutBunch, TSharedPtr<INetDeltaBaseState>& NewFullState, TSharedPtr<INetDeltaBaseState>& OldState);
+	bool SendCustomDeltaProperty(
+		UObject* InObject,
+		UProperty* Property,
+		uint32 ArrayIndex,
+		FNetBitWriter& OutBunch,
+		TSharedPtr<INetDeltaBaseState>& NewFullState,
+		TSharedPtr<INetDeltaBaseState> & OldState);
+
+	bool SendCustomDeltaProperty(
+		UObject* InObject,
+		uint16 CustomDeltaProperty,
+		FNetBitWriter& OutBunch,
+		TSharedPtr<INetDeltaBaseState>& NewFullState,
+		TSharedPtr<INetDeltaBaseState>& OldState);
 
 	/** Packet was dropped */
-	void	ReceivedNak( int32 NakPacketId );
+	void ReceivedNak(int32 NakPacketId);
 
 	UE_DEPRECATED(4.23, "Use CountBytes instead")
 	void Serialize(FArchive& Ar);
@@ -206,61 +165,145 @@ public:
 	void CountBytes(FArchive& Ar) const;
 
 	/** Writes dirty properties to bunch */
-	void	ReplicateCustomDeltaProperties( FNetBitWriter & Bunch, FReplicationFlags RepFlags );
-	bool	ReplicateProperties( FOutBunch & Bunch, FReplicationFlags RepFlags );
-	void	PostSendBunch(FPacketIdRange & PacketRange, uint8 bReliable);
+	void ReplicateCustomDeltaProperties(FNetBitWriter& Bunch, FReplicationFlags RepFlags);
+	bool ReplicateProperties(FOutBunch& Bunch, FReplicationFlags RepFlags);
+	void PostSendBunch(FPacketIdRange& PacketRange, uint8 bReliable);
 
 	/** Updates the custom delta state for a replay delta checkpoint */
-	void	UpdateCheckpoint();
+	void UpdateCheckpoint();
 	
-	bool	ReceivedBunch( FNetBitReader& Bunch, const FReplicationFlags& RepFlags, const bool bHasRepLayout, bool& bOutHasUnmapped );
-	bool	ReceivedRPC(FNetBitReader& Reader, const FReplicationFlags& RepFlags, const FFieldNetCache* FieldCache, const bool bCanDelayRPC, bool& bOutDelayRPC, TSet<FNetworkGUID>& OutUnmappedGuids);
-	void	UpdateGuidToReplicatorMap();
-	bool	MoveMappedObjectToUnmapped( const FNetworkGUID& GUID );
-	void	PostReceivedBunch();
+	bool ReceivedBunch(
+		FNetBitReader& Bunch,
+		const FReplicationFlags& RepFlags,
+		const bool bHasRepLayout,
+		bool& bOutHasUnmapped);
 
-	void	ForceRefreshUnreliableProperties();
+	bool ReceivedRPC(
+		FNetBitReader& Reader,
+		const FReplicationFlags& RepFlags,
+		const FFieldNetCache* FieldCache,
+		const bool bCanDelayRPC,
+		bool& bOutDelayRPC,
+		TSet<FNetworkGUID>& OutUnmappedGuids);
 
-	void QueueRemoteFunctionBunch( UFunction* Func, FOutBunch &Bunch );
+	void UpdateGuidToReplicatorMap();
+	bool MoveMappedObjectToUnmapped(const FNetworkGUID& GUID);
+	void PostReceivedBunch();
 
-	bool ReadyForDormancy(bool debug=false);
+	void ForceRefreshUnreliableProperties();
+
+	void QueueRemoteFunctionBunch(UFunction* Func, FOutBunch &Bunch);
+
+	bool ReadyForDormancy(bool bDebug=false);
 
 	void StartBecomingDormant();
 
 	void CallRepNotifies(bool bSkipIfChannelHasQueuedBunches);
 
-	void UpdateUnmappedObjects( bool & bOutHasMoreUnmapped );
+	void UpdateUnmappedObjects(bool& bOutHasMoreUnmapped);
 
-	FORCEINLINE TWeakObjectPtr<UObject>	GetWeakObjectPtr() const { return WeakObjectPtr; }
-	FORCEINLINE UObject *	GetObject() const { return ObjectPtr; }
-	FORCEINLINE void		SetObject( UObject* NewObj ) { ObjectPtr = NewObj; WeakObjectPtr = NewObj; }
+	FORCEINLINE TWeakObjectPtr<UObject>	GetWeakObjectPtr() const
+	{
+		return WeakObjectPtr;
+	}
 
-	FORCEINLINE void PreNetReceive()		
+	FORCEINLINE UObject* GetObject() const
+	{
+		return ObjectPtr;
+	}
+
+	FORCEINLINE void SetObject(UObject* NewObj)
+	{
+		ObjectPtr = NewObj;
+		WeakObjectPtr = NewObj;
+	}
+
+	FORCEINLINE void PreNetReceive()
 	{ 
 		UObject* Object = GetObject();
 		if ( Object != NULL )
 		{
-			Object->PreNetReceive(); 
+			Object->PreNetReceive();
 		}
 	}
 
-	FORCEINLINE void PostNetReceive()	
+	FORCEINLINE void PostNetReceive()
 	{ 
 		UObject* Object = GetObject();
 		if ( Object != NULL )
 		{
-			Object->PostNetReceive(); 
+			Object->PostNetReceive();
 		}
 	}
 
-	void QueuePropertyRepNotify( UObject* Object, UProperty * Property, const int32 ElementIndex, TArray< uint8 > & MetaData );
-
+	void QueuePropertyRepNotify(
+		UObject* Object,
+		UProperty* Property,
+		const int32 ElementIndex,
+		TArray<uint8>& MetaData);
+		
 	void WritePropertyHeaderAndPayload(
-		UObject*				Object,
-		UProperty*				Property,
-		FNetFieldExportGroup*	NetFieldExportGroup,
-		FNetBitWriter&			Bunch,
-		FNetBitWriter&			Payload ) const;
+		UObject* Object,
+		UProperty* Property,
+		FNetFieldExportGroup* NetFieldExportGroup,
+		FNetBitWriter& Bunch,
+		FNetBitWriter& Payload) const;	
+
+public:
+
+	/** Net GUID for the object we're replicating. */
+	FNetworkGUID ObjectNetGUID;
+
+	/** The amount of memory (in bytes) that we're using to track Unmapped GUIDs. */
+	int32 TrackedGuidMemoryBytes;
+
+	/** True if last update (ReplicateActor) produced no replicated properties */
+	uint32 bLastUpdateEmpty : 1;
+
+	/** Whether or not the Actor Channel on which we're replicating has been Opened / Acked by the receiver. */
+	uint32 bOpenAckCalled : 1;
+
+	/** True if we need to do an unmapped check next frame. */
+	uint32 bForceUpdateUnmapped : 1;
+
+	/** Whether or not we've already replicated properties this frame. */
+	uint32 bHasReplicatedProperties : 1;
+
+private:
+
+	/** Whether or not we are going to use Fast Array Delta Struct Delta Serialization. See FFastArraySerializer::FastArrayDeltaSerialize_DeltaSerializeStructs. */
+	uint32 bSupportsFastArrayDelta : 1;
+
+public:
+	
+	TSharedPtr<class FReplicationChangelistMgr> ChangelistMgr;
+	TSharedPtr<FRepLayout> RepLayout;
+	TUniquePtr<FRepState>  RepState;
+	TUniquePtr<FRepState> CheckpointRepState;
+
+	UClass* ObjectClass;
+
+	UObject* ObjectPtr;
+
+	/** Connection this replicator was created on. */
+	UNetConnection* Connection;
+
+	/** The Actor Channel that we're replicating on. This expected to be owned by Connection. */
+	class UActorChannel* OwningChannel;
+
+	FOutBunch* RemoteFunctions;
+
+	/** Meta information on pending net RPCs (to be sent) */
+	TArray<FRPCCallInfo> RemoteFuncInfo;
+
+	/** Information on RPCs that have been received but not yet executed */
+	TArray<FRPCPendingLocalCall> PendingLocalRPCs;
+
+	//~ RepNotify properties were moved to FReceivingRepState.
+	//~ CustomDeltaState properties were moved to FSendingRepState.
+	//~ Retirement properties were moved to FSendingRepState.
+
+	TSet<FNetworkGUID> ReferencedGuids;
 
 private:
 	TWeakObjectPtr<UObject> WeakObjectPtr;
diff --git a/Engine/Source/Runtime/Engine/Public/Net/GuidReferences.h b/Engine/Source/Runtime/Engine/Public/Net/GuidReferences.h
--- a/Engine/Source/Runtime/Engine/Public/Net/GuidReferences.h
+++ b/Engine/Source/Runtime/Engine/Public/Net/GuidReferences.h
@@ -16,6 +16,15 @@ using FGuidReferencesMap = TMap<int32, FGuidReferences>;
  * replicated data for the Object will be handled elsewhere (either on its own Actor channel,
  * or on its Owning Actor's channel, as a replicated subobject).
  *
+ * As Objects are replicated and received, these GUID References will become "mapped", and the
+ * networking system will update the corresponding properties with pointers to the appropriate objects.
+ *
+ * As Objects are destroyed (due to game play, actor channels closing, relevancy, etc.), GUID References
+ * will be "unmapped", and their corresponding properties will nulled out.
+ *
+ * The process of Mapping and Unmapping can happen numerous times for the same object (e.g., if an Actor
+ * goes in and out of relevancy repeatedly).
+ *
  * This class helps manage those references for specific replicated properties.
  * A FGuidReferences instance will be created for each Replicated Property that is a reference to an object.
  *
diff --git a/Engine/Source/Runtime/Engine/Public/Net/RepLayout.h b/Engine/Source/Runtime/Engine/Public/Net/RepLayout.h
--- a/Engine/Source/Runtime/Engine/Public/Net/RepLayout.h
+++ b/Engine/Source/Runtime/Engine/Public/Net/RepLayout.h
@@ -55,12 +55,6 @@ enum class ERepDataBufferType
 
 namespace UE4_RepLayout_Private
 {
-	// The purpose of this class is to allow us to translate or preserve const volatileness between disparate classes.
-	template<typename T, typename U> struct TTranslateConstVolatile { typedef typename TRemoveCV<U>::Type Type; };
-	template<typename T, typename U> struct TTranslateConstVolatile<const T, U> { typedef const typename TRemoveCV<U>::Type Type; };
-	template<typename T, typename U> struct TTranslateConstVolatile<volatile T, U> { typedef volatile typename TRemoveCV<U>::Type Type; };
-	template<typename T, typename U> struct TTranslateConstVolatile<const volatile T, U> { typedef const volatile typename TRemoveCV<U>::Type Type;  };
-
 	/**
 	 * TRepDataBuffer and TConstRepDataBuffer act as wrapper around internal data
 	 * buffers that FRepLayout may use. This allows FRepLayout to properly interact
@@ -70,7 +64,7 @@ namespace UE4_RepLayout_Private
 	struct TRepDataBufferBase
 	{
 		static constexpr ERepDataBufferType Type = DataType;
-		using ConstOrNotVoid = typename TTranslateConstVolatile<ConstOrNotType, void>::Type;
+		using ConstOrNotVoid = typename TCopyQualifiersFromTo<ConstOrNotType, void>::Type;
 
 	public:
 
@@ -515,6 +509,12 @@ public:
 
 	/** List of properties that have RepNotifies that we will need to call on Clients. */
 	TArray<UProperty*> RepNotifies;
+
+	/**
+	 * Holds MetaData (such as array index) for RepNotifies.
+	 * Only used for CustomDeltaProperties.
+	 */
+	TMap<UProperty*, TArray<uint8>> RepNotifyMetaData;
 };
 
 /** Replication State that is only needed when sending properties. */
@@ -592,6 +592,9 @@ public:
 	/** Circular buffer of changelists. */
 	FRepChangedHistory ChangeHistory[MAX_CHANGE_HISTORY];
 
+	/** Array of property retirements that we'll use to track retransmission for Custom Delta Properties. */
+	TArray<FPropertyRetirement> Retirement;
+
 	/** List of changelists that were generated before the channel was fully opened.*/
 	TArray<FRepChangedHistory> PreOpenAckHistory;
 
@@ -608,6 +611,15 @@ public:
 
 	/** Cached set of inactive parent commands. */
 	TBitArray<> InactiveParents;
+
+	/** This is the delta state we need to compare with when determining what to send to a client for custom delta properties. */
+	TMap<int32, TSharedPtr<INetDeltaBaseState>> RecentCustomDeltaState;
+
+	/** Same as RecentCustomDeltaState, but this will always remain as the initial CDO version. We use this to send all properties since channel was first opened (for bResendAllDataSinceOpen). */
+	TMap<int32, TSharedPtr<INetDeltaBaseState>>	CDOCustomDeltaState;
+
+	/** Same as RecentCustomDeltaState, but will represent the state at the last checkpoint. */
+	TMap<int32, TSharedPtr<INetDeltaBaseState>>	CheckpointCustomDeltaState;
 };
 
 /** Replication State that is unique Per Object Per Net Connection. */
@@ -1076,6 +1088,7 @@ private:
 	friend struct FRepStateStaticBuffer;
 	friend class UPackageMapClient;
 	friend class FNetSerializeCB;
+	friend struct FCustomDeltaPropertyIterator;
 
 	FRepLayout();
 
@@ -1092,22 +1105,6 @@ public:
 	/** Creates a new FRepLayout for the given function. */
 	static TSharedPtr<FRepLayout> CreateFromFunction(UFunction* InFunction, const UNetConnection* ServerConnection = nullptr, const ECreateRepLayoutFlags Flags = ECreateRepLayoutFlags::None);
 
-	/**
-	 * Used to signal that the channel that owns a given object has been opened and acknowledged
-	 * by a client.
-	 *
-	 * @param RepState	RepState for the Object whose channel was acked.
-	 *					This is expected to be valid.
-	 */
-	void OpenAcked(FSendingRepState* RepState) const;
-
-	UE_DEPRECATED(4.23, "Use OpenAcked that accepts a FSendingRepState")
-	void OpenAcked(FRepState* RepState) const
-	{
-		check(RepState);
-		OpenAcked(RepState->GetSendingRepState());
-	}
-
 	/**
 	 * Creates and initialize a new Shadow Buffer.
 	 *
@@ -1178,32 +1175,6 @@ public:
 		FNetBitWriter& Writer,
 		const FReplicationFlags& RepFlags) const;
 
-	UE_DEPRECATED(4.23, "Use ReplicateProperties that accepts a FSendingRepState")
-	bool ReplicateProperties(
-		FRepState* RESTRICT RepState,
-		FRepChangelistState* RESTRICT RepChangelistState,
-		const uint8* RESTRICT Data,
-		UClass* ObjectClass,
-		UActorChannel* OwningChannel,
-		FNetBitWriter& Writer,
-		const FReplicationFlags& RepFlags) const
-	{
-		return ReplicateProperties(RepState->GetSendingRepState(), RepChangelistState, Data, ObjectClass, OwningChannel, Writer, RepFlags);
-	}
-
-	UE_DEPRECATED(4.23, "This method will be made private in future versions. Use ReplicateProperties (or SendProperties_BackwardsCompatible for replays).")
-	void SendProperties(
-		FRepState* RESTRICT RepState,
-		FRepChangedPropertyTracker* ChangedTracker,
-		const uint8* RESTRICT Data,
-		UClass* ObjectClass,
-		FNetBitWriter& Writer,
-		TArray<uint16>& Changed,
-		const FRepSerializationSharedInfo& SharedInfo) const
-	{
-		SendProperties(RepState->GetSendingRepState(), ChangedTracker, Data, ObjectClass, Writer, Changed, SharedInfo);
-	}
-
 	/**
 	 * Reads all property values from the received buffer, and applies them to the
 	 * property memory.
@@ -1228,20 +1199,6 @@ public:
 		bool& bOutGuidsChanged,
 		const EReceivePropertiesFlags Flags) const;
 
-	UE_DEPRECATED(4.23, "Use ReceiveProperties that accepts a FReceivingRepState")
-	bool ReceiveProperties(
-		UActorChannel* OwningChannel,
-		UClass* InObjectClass,
-		FRepState* RESTRICT RepState,
-		void* RESTRICT Data,
-		FNetBitReader& InBunch,
-		bool& bOutHasUnmapped,
-		bool& bOutGuidsChanged,
-		const EReceivePropertiesFlags Flags) const
-	{
-		return ReceiveProperties(OwningChannel, InObjectClass, RepState->GetReceivingRepState(), Data, InBunch, bOutHasUnmapped, bOutGuidsChanged, Flags);
-	}
-
 	/**
 	 * Finds any properties in the Shadow Buffer of the given Rep State that are currently valid
 	 * (mapped or unmapped) references to other network objects, and retrieves the associated
@@ -1254,18 +1211,10 @@ public:
 	 */
 	void GatherGuidReferences(
 		FReceivingRepState* RESTRICT RepState,
+		struct FNetDeltaSerializeInfo& Params,
 		TSet<FNetworkGUID>& OutReferencedGuids,
 		int32& OutTrackedGuidMemoryBytes) const;
 
-	UE_DEPRECATED(4.23, "Use GatherGuidReferences that accepts a FReceivingRepState")
-	void GatherGuidReferences(
-		FRepState* RepState,
-		TSet<FNetworkGUID>& OutReferencedGuids,
-		int32& OutTrackedGuidMemoryBytes) const
-	{
-		GatherGuidReferences(RepState->GetReceivingRepState(), OutReferencedGuids, OutTrackedGuidMemoryBytes);
-	}
-
 	/**
 	 * Called to indicate that the object referenced by the FNetworkGUID is no longer mapped.
 	 * This can happen if the Object was destroyed, if its level was streamed out, or for any
@@ -1275,14 +1224,12 @@ public:
 	 * @param RepState	The RepState that holds a reference to the object.
 	 *					This is expected to be valid.
 	 * @param GUID		The Network GUID of the object to unmap.
+	 * @param Params	Delta Serialization Params used for Custom Delta Properties.
 	 */
-	bool MoveMappedObjectToUnmapped(FReceivingRepState* RESTRICT RepState, const FNetworkGUID& GUID) const;
-
-	UE_DEPRECATED(4.23, "Use MoveMappedObjectToUnmapped that accepts a FReceivingRepState")
-	bool MoveMappedObjectToUnmapped(FRepState* RepState, const FNetworkGUID& GUID) const
-	{
-		return MoveMappedObjectToUnmapped(RepState->GetReceivingRepState(), GUID);
-	}
+	bool MoveMappedObjectToUnmapped(
+		FReceivingRepState* RESTRICT RepState,
+		struct FNetDeltaSerializeInfo& Params,
+		const FNetworkGUID& GUID) const;
 
 	/**
 	 * Attempts to update any unmapped network guids referenced by the RepState.
@@ -1300,22 +1247,11 @@ public:
 		FReceivingRepState* RESTRICT RepState,
 		UPackageMap* PackageMap,
 		UObject* Object,
+		struct FNetDeltaSerializeInfo& Params,
 		bool& bCalledPreNetReceive,
 		bool& bOutSomeObjectsWereMapped,
 		bool& bOutHasMoreUnmapped) const;
 
-	UE_DEPRECATED(4.23, "Use UpdateUnmappedObjects that accepts a FReceivingRepState")
-	void UpdateUnmappedObjects(
-		FRepState* RepState,
-		UPackageMap* PackageMap,
-		UObject* Object,
-		bool& bOutSomeObjectsWereMapped,
-		bool& bOutHasMoreUnmapped) const
-	{
-		bool bCalledPreNetReceive = false;
-		UpdateUnmappedObjects(RepState->GetReceivingRepState(), PackageMap, Object, bCalledPreNetReceive, bOutSomeObjectsWereMapped, bOutHasMoreUnmapped);
-	}
-
 	/**
 	 * Fire any RepNotifies that have been queued for an object while receiving properties.
 	 *
@@ -1325,12 +1261,6 @@ public:
 	 */
 	void CallRepNotifies(FReceivingRepState* RepState, UObject* Object) const;
 
-	UE_DEPRECATED(4.23, "Use CallRepNotifies that accepts a FReceivingRepState")
-	void CallRepNotifies(FRepState* RepState, UObject* Object) const
-	{
-		return CallRepNotifies(RepState->GetReceivingRepState(), Object);
-	}
-
 	/**
 	 * Called after an object has finished replicating its properties.
 	 *
@@ -1344,44 +1274,11 @@ public:
 		FPacketIdRange& PacketRange,
 		bool bReliable) const;
 
-	UE_DEPRECATED(4.23, "Use PostReplicate that accepts a FSendingRepState")
-	void PostReplicate(
-		FRepState* RepState,
-		FPacketIdRange& PacketRange,
-		bool bReliable) const
-	{
-		PostReplicate(RepState->GetSendingRepState(), PacketRange, bReliable);
-	}
-
-	void ReceivedNak(FRepState* RepState, int32 NakPacketId) const;
-	bool AllAcked(FRepState* RepState) const;
-	bool ReadyForDormancy(FRepState* RepState) const;
-
 	template<ERepDataBufferType DataType>
 	void ValidateWithChecksum(TConstRepDataBuffer<DataType> Data, FBitArchive & Ar) const;
 
 	uint32 GenerateChecksum(const FRepState* RepState) const;
 
-	UE_DEPRECATED(4.23, "This method will be made private in future versions.")
-	void PruneChangeList(
-		FRepState* RepState,
-		const void* RESTRICT Data,
-		const TArray<uint16>& Changed,
-		TArray<uint16>& PrunedChanged) const
-	{
-		PruneChangeList(Data, Changed, PrunedChanged);
-	}
-
-	UE_DEPRECATED(4.23, "This method will be made private in future versions.")
-	void MergeChangeList(
-		const uint8* RESTRICT Data,
-		const TArray<uint16>& Dirty1,
-		const TArray<uint16>& Dirty2,
-		TArray<uint16>& MergedDirty) const
-	{
-		MergeChangeList(FConstRepObjectDataBuffer(Data), Dirty1, Dirty2, MergedDirty);
-	}
-
 	/**
 	 * Compare all properties between source and destination buffer, and optionally update the destination
 	 * buffer to match the state of the source buffer if they don't match.
@@ -1424,18 +1321,6 @@ public:
 		TRepDataBuffer<DstType> Destination,
 		TConstRepDataBuffer<SrcType> Source) const;
 
-	UE_DEPRECATED(4.23, "This method is deprecated. Please use GetLifetimeCustomDeltaProperties that returns an ArrayView and GetPropertyLifetimeCondition instead.")
-	void GetLifetimeCustomDeltaProperties(TArray<int32>& OutCustom, TArray<ELifetimeCondition>& OutConditions) const;
-
-	/** Returns the RepIndices of any Lifetime Custom Delta properties. */
-	const TArrayView<const uint16> GetLifetimeCustomDeltaProperties() const;
-
-	/** Returns the Lifetime Condition for the Replicated property associated with the given RepIndex. */
-	const ELifetimeCondition GetPropertyLifetimeCondition(uint16 RepIndex) const
-	{
-		return Parents[RepIndex].Condition;
-	}
-
 	/** @see SendProperties. */
 	void ENGINE_API SendPropertiesForRPC(
 		UFunction* Function,
@@ -1490,14 +1375,6 @@ public:
 		const bool bEnableRepNotifies,
 		bool& bOutGuidsChanged) const;
 
-	void UpdateChangelistMgr(
-		FSendingRepState* RESTRICT RepState,
-		FReplicationChangelistMgr& InChangelistMgr,
-		const UObject* InObject,
-		const uint32 ReplicationFrame,
-		const FReplicationFlags& RepFlags,
-		const bool bForceCompare) const;
-
 	//~ Begin FGCObject Interface
 	ENGINE_API virtual void AddReferencedObjects(FReferenceCollector& Collector) override;
 	ENGINE_API virtual FString GetReferencerName() const override;
@@ -1543,12 +1420,6 @@ public:
 
 	void CountBytes(FArchive& Ar) const;
 
-	UE_DEPRECATED(4.23, "This method is only temporary to support existing behavior, and will be removed in future versions.")
-	UProperty* GetPropertyForRepIndex(uint32 RepIndex) const
-	{
-		return Parents[RepIndex].Property;
-	}
-
 private:
 
 	void InitFromClass(UClass* InObjectClass, const UNetConnection* ServerConnection, const ECreateRepLayoutFlags Flags);
@@ -1803,60 +1674,30 @@ private:
 		const int32 CmdEnd,
 		TArray<FHandleToCmdIndex>& HandleToCmdIndex);
 
+	void UpdateChangelistMgr(
+		FSendingRepState* RESTRICT RepState,
+		FReplicationChangelistMgr& InChangelistMgr,
+		const UObject* InObject,
+		const uint32 ReplicationFrame,
+		const FReplicationFlags& RepFlags,
+		const bool bForceCompare) const;
+
 	void InitRepStateStaticBuffer(FRepStateStaticBuffer& ShadowData, const FConstRepObjectDataBuffer Source) const;
 	void ConstructProperties(FRepStateStaticBuffer& ShadowData) const;
 	void CopyProperties(FRepStateStaticBuffer& ShadowData, const FConstRepObjectDataBuffer Source) const;
 	void DestructProperties(FRepStateStaticBuffer& RepStateStaticBuffer) const;
 
 	/**
-	 * Similar to GatherGuidReferences, except works for CustomDeltaProperties.
-	 *
-	 * @param Params	Params that we'll use to gather guids.
-	 *					The GatherGuidReferences member must not be null.
-	 *					The Object member must not be null.
+	 * @param CustomDeltaIndex	The index of the Custom Delta Property.
+	 *							This is not the same as UProperty::RepIndex!
+	 *							@see FLifetimeCustomDeltaState.
 	 */
-	void GatherGuidReferencesForCustomDeltaProperties(FNetDeltaSerializeInfo& Params) const;
+	bool SendCustomDeltaProperty(FNetDeltaSerializeInfo& Params, const uint16 CustomDeltaIndex) const;
 
 	/**
-	 * Similar to MoveMappedObjectToUnmapped, except works for CustomDeltaProperties.
+	 * Attempts to receive the custom delta property.
 	 *
-	 * @param Params	Params that we'll use to gather guids.
-	 *					The MoveGuidToUnmapped member must not be null.
-	 *					The Object member must not be null.
-	 */
-	bool MoveMappedObjectToUnmappedForCustomDeltaProperties(FNetDeltaSerializeInfo& Params) const;
-
-	//~ Remove this when FObjectReplicator::UnmappedCustomProperties goes away.
-	UE_DEPRECATED(4.23, "This method is only temporary to support existing behavior, and will be removed in future versions.")
-	bool MoveMappedObjectToUnmappedForCustomDeltaProperties(FNetDeltaSerializeInfo& Params, TMap<int32, UStructProperty*>& UnmappedProperties) const;
-
-	/**
-	 * Similar to UpdateUnmappedObjects, except works for CustomDeltaProperties.
 	 *
-	 * @param Params					Params that we'll use to gather guids.
-	 *									The bUpdateUnmappedObjects member must be true.
-	 *									The PackageMap member must be non null.
-	 *									The Object member must be non null.
-	 */
-	void UpdateUnmappedObjectsForCustomDeltaProperties(FNetDeltaSerializeInfo& Params) const;
-
-	//~ Remove this when FObjectReplicator::UnmappedCustomProperties goes away.
-	UE_DEPRECATED(4.23, "This method is only temporary to support existing behavior, and will be removed in future versions.")
-	void UpdateUnmappedObjectsForCustomDeltaProperties(
-		FNetDeltaSerializeInfo& Params,
-		TArray<TPair<int32, UStructProperty*>>& CompletelyMapped,
-		TArray<TPair<int32, UStructProperty*>>& Updated) const;
-
-	bool SendCustomDeltaProperty(FNetDeltaSerializeInfo& Params, uint16 CustomDeltaRepIndex) const;
-
-	UE_DEPRECATED(4.23, "This method is only temporary to support existing behavior, and will be removed in future versions.")
-	bool SendCustomDeltaProperty(
-		FNetDeltaSerializeInfo& Params,
-		UProperty* Property,
-		int32 StaticArrayIndex) const;
-
-	/**
-	 * Attempts to receive the custom delta property.
 	 *
 	 * @param Params			Params that we'll use to receive.
 	 *							The PackageMap member must be non null.
@@ -1865,10 +1706,10 @@ private:
 	 * @param Property			The Property that we're trying to received.
 	 *							This is expected to be a valid Replicated property that is a CustomDelta type.
 	 */
-	bool ReceiveCustomDeltaProperty(FNetDeltaSerializeInfo& Params, UStructProperty* Property) const;
-
-	UE_DEPRECATED(4.23, "This method is only temporary to support existing behavior, and will be removed in future versions.")
-	bool ReceiveCustomDeltaProperty(FNetDeltaSerializeInfo& Params, UStructProperty* Property, uint32& StaticArrayIndex, int32& Offset) const;
+	bool ReceiveCustomDeltaProperty(
+		FReceivingRepState* RESTRICT ReceivingRepState,
+		FNetDeltaSerializeInfo& Params,
+		UStructProperty* Property) const;
 
 	bool DeltaSerializeFastArrayProperty(struct FFastArrayDeltaSerializeParams& Params, FReplicationChangelistMgr* ChangelistMgr) const;
 
@@ -1890,6 +1731,12 @@ private:
 		FReplicationChangelistMgr& ChangelistMgr,
 		TMap<int32, TSharedPtr<INetDeltaBaseState>>& CustomDeltaStates) const;
 
+	const uint16 GetNumLifetimeCustomDeltaProperties() const;
+
+	UProperty* GetLifetimeCustomDeltaProperty(const uint16 CustomDeltaPropertyIndex) const;
+
+	const ELifetimeCondition GetLifetimeCustomDeltaPropertyCondition(const uint16 RepIndCustomDeltaPropertyIndexex) const;
+
 	ERepLayoutState LayoutState;
 
 	/** Index of the Role property in the Parents list. May be INDEX_NONE if Owner doesn't have the property. */
@@ -1923,5 +1770,5 @@ private:
 	FRepSerializationSharedInfo SharedInfoRPC;
 
 	/** Shared comparison to default state for multicast rpc */
-	TBitArray<> SharedInfoRPCParentsChanged;	
+	TBitArray<> SharedInfoRPCParentsChanged;
 };
