Fix platform specific link error introduced by CL-6837569.

diff --git a/Engine/Source/Runtime/Engine/Private/RepLayout.cpp b/Engine/Source/Runtime/Engine/Private/RepLayout.cpp
--- a/Engine/Source/Runtime/Engine/Private/RepLayout.cpp
+++ b/Engine/Source/Runtime/Engine/Private/RepLayout.cpp
@@ -314,11 +314,9 @@ struct FLifetimeCustomDeltaState
 {
 public:
 
-	static constexpr uint16 INVALID_DELTA_INDEX = INDEX_NONE;
-
 	FLifetimeCustomDeltaState(uint16 TotalNumberOfLifetimeProperties)
 	{
-		LifetimeCustomDeltaIndexLookup.Init(INVALID_DELTA_INDEX, TotalNumberOfLifetimeProperties);
+		LifetimeCustomDeltaIndexLookup.Init(static_cast<uint16>(INDEX_NONE), TotalNumberOfLifetimeProperties);
 	}
 
 	void CountBytes(FArchive& Ar) const
@@ -345,13 +343,13 @@ public:
 	const uint16 GetCustomDeltaIndexFromPropertyRepIndex(const uint16 PropertyRepIndex) const
 	{
 		const uint16 CustomDeltaIndex = LifetimeCustomDeltaIndexLookup[PropertyRepIndex];
-		check(INVALID_DELTA_INDEX != CustomDeltaIndex);
+		check(static_cast<uint16>(INDEX_NONE) != CustomDeltaIndex);
 		return CustomDeltaIndex;
 	}
 
 	void Add(FLifetimeCustomDeltaProperty&& ToAdd)
 	{
-		check(INVALID_DELTA_INDEX == LifetimeCustomDeltaIndexLookup[ToAdd.PropertyRepIndex]);
+		check(static_cast<uint16>(INDEX_NONE) == LifetimeCustomDeltaIndexLookup[ToAdd.PropertyRepIndex]);
 
 		if (ToAdd.FastArrayNumber != INDEX_NONE)
 		{
