Refactor UAnimSequence out of compression code

diff --git a/Engine/Source/Editor/SkeletonEditor/Private/EditableSkeleton.cpp b/Engine/Source/Editor/SkeletonEditor/Private/EditableSkeleton.cpp
--- a/Engine/Source/Editor/SkeletonEditor/Private/EditableSkeleton.cpp
+++ b/Engine/Source/Editor/SkeletonEditor/Private/EditableSkeleton.cpp
@@ -298,8 +298,8 @@ void FEditableSkeleton::RenameSmartname(const FName InContainerName, SmartName::
 						{
 							SequencesToRecompress.Add(Seq);
 
-							Seq->CompressedCurveByteStream.Empty();
-							Seq->CurveCompressionCodec = nullptr;
+							Seq->CompressedData.CompressedCurveByteStream.Empty();
+							Seq->CompressedData.CurveCompressionCodec = nullptr;
 						}
 					}
 				}
diff --git a/Engine/Source/Editor/UnrealEd/Private/Commandlets/PackageUtilities.cpp b/Engine/Source/Editor/UnrealEd/Private/Commandlets/PackageUtilities.cpp
--- a/Engine/Source/Editor/UnrealEd/Private/Commandlets/PackageUtilities.cpp
+++ b/Engine/Source/Editor/UnrealEd/Private/Commandlets/PackageUtilities.cpp
@@ -1843,8 +1843,10 @@ struct CompressAnimationsFunctor
 
 				NumTotalSize += ResourceSize;
 				
+				FUECompressedAnimData& CompressedData = AnimSeq->CompressedData.CompressedDataStructure;
+
 				// Looking for PerTrackCompression using 96bit translation compression.
-				if( AnimSeq->KeyEncodingFormat == AKF_PerTrackCompression && AnimSeq->CompressedByteStream.Num() > 0 )
+				if( CompressedData.KeyEncodingFormat == AKF_PerTrackCompression && CompressedData.CompressedByteStream.Num() > 0 )
 				{
 					bool bCandidate = false;
 
@@ -1857,12 +1859,12 @@ struct CompressAnimationsFunctor
  						// Translation
 						{
 							// Use the CompressedTrackOffsets stream to find the data addresses
-							const int32* RESTRICT TrackDataForTransKey = AnimSeq->CompressedTrackOffsets.GetData() + (TrackIndex * 2);
+							const int32* RESTRICT TrackDataForTransKey = CompressedData.CompressedTrackOffsets.GetData() + (TrackIndex * 2);
 							const int32 TransKeysOffset = TrackDataForTransKey[0];
  							if( TransKeysOffset != INDEX_NONE )
  							{
-								const uint8* RESTRICT TrackData = AnimSeq->CompressedByteStream.GetData() + TransKeysOffset + 4;
-								const int32 Header = *((int32*)(AnimSeq->CompressedByteStream.GetData() + TransKeysOffset));
+								const uint8* RESTRICT TrackData = CompressedData.CompressedByteStream.GetData() + TransKeysOffset + 4;
+								const int32 Header = *((int32*)(CompressedData.CompressedByteStream.GetData() + TransKeysOffset));
 	 
  								int32 KeyFormat;
  								int32 NumKeys;
@@ -1962,12 +1964,12 @@ struct CompressAnimationsFunctor
 						// Rotation
 						{
 							// Use the CompressedTrackOffsets stream to find the data addresses
-							const int32* RESTRICT TrackDataForRotKey = AnimSeq->CompressedTrackOffsets.GetData() + (TrackIndex * 2);
+							const int32* RESTRICT TrackDataForRotKey = CompressedData.CompressedTrackOffsets.GetData() + (TrackIndex * 2);
 							const int32 RotKeysOffset = TrackDataForRotKey[1];
 							if( RotKeysOffset != INDEX_NONE )
 							{
-								const uint8* RESTRICT TrackData = AnimSeq->CompressedByteStream.GetData() + RotKeysOffset + 4;
-								const int32 Header = *((int32*)(AnimSeq->CompressedByteStream.GetData() + RotKeysOffset));
+								const uint8* RESTRICT TrackData = CompressedData.CompressedByteStream.GetData() + RotKeysOffset + 4;
+								const int32 Header = *((int32*)(CompressedData.CompressedByteStream.GetData() + RotKeysOffset));
 
 								int32 KeyFormat;
 								int32 NumKeys;
@@ -2023,11 +2025,11 @@ struct CompressAnimationsFunctor
 						// Scale
 						{
 							// Use the CompressedTrackOffsets stream to find the data addresses
-							const int32 ScaleKeysOffset = AnimSeq->CompressedScaleOffsets.GetOffsetData(TrackIndex, 0);
+							const int32 ScaleKeysOffset = CompressedData.CompressedScaleOffsets.GetOffsetData(TrackIndex, 0);
 							if( ScaleKeysOffset != INDEX_NONE )
 							{
-								const uint8* RESTRICT TrackData = AnimSeq->CompressedByteStream.GetData() + ScaleKeysOffset + 4;
-								const int32 Header = *((int32*)(AnimSeq->CompressedByteStream.GetData() + ScaleKeysOffset));
+								const uint8* RESTRICT TrackData = CompressedData.CompressedByteStream.GetData() + ScaleKeysOffset + 4;
+								const int32 Header = *((int32*)(CompressedData.CompressedByteStream.GetData() + ScaleKeysOffset));
 
 								int32 KeyFormat;
 								int32 NumKeys;
diff --git a/Engine/Source/Editor/UnrealEd/Private/Fbx/FbxAnimationExport.cpp b/Engine/Source/Editor/UnrealEd/Private/Fbx/FbxAnimationExport.cpp
--- a/Engine/Source/Editor/UnrealEd/Private/Fbx/FbxAnimationExport.cpp
+++ b/Engine/Source/Editor/UnrealEd/Private/Fbx/FbxAnimationExport.cpp
@@ -137,7 +137,7 @@ void FFbxExporter::ExportAnimSequenceToFbx(const UAnimSequence* AnimSeq,
 		ExportTimeIncrement.SetSecondDouble( TimePerKey );
 
 		int32 BoneTreeIndex = Skeleton->GetSkeletonBoneIndexFromMeshBoneIndex(SkelMesh, BoneIndex);
-		int32 BoneTrackIndex = Skeleton->GetAnimationTrackIndex(BoneTreeIndex, AnimSeq, true);
+		int32 BoneTrackIndex = Skeleton->GetRawAnimationTrackIndex(BoneTreeIndex, AnimSeq);
 		if(BoneTrackIndex == INDEX_NONE)
 		{
 			// If this sequence does not have a track for the current bone, then skip it
diff --git a/Engine/Source/Runtime/Engine/Classes/Animation/AnimCompress.h b/Engine/Source/Runtime/Engine/Classes/Animation/AnimCompress.h
--- a/Engine/Source/Runtime/Engine/Classes/Animation/AnimCompress.h
+++ b/Engine/Source/Runtime/Engine/Classes/Animation/AnimCompress.h
@@ -17,6 +17,8 @@
 #include "AnimationCompression.h"
 #include "AnimCompress.generated.h"
 
+struct FCompressibleAnimData;
+
 //Helper function for ddc key generation
 uint8 MakeBitForFlag(uint32 Item, uint32 Position);
 
@@ -183,9 +185,9 @@ class ENGINE_API FCompressionMemorySummary
 public:
 	FCompressionMemorySummary(bool bInEnabled);
 
-	void GatherPreCompressionStats(UAnimSequence* Seq, int32 ProgressNumerator, int32 ProgressDenominator);
+	void GatherPreCompressionStats(const FCompressibleAnimData& CompressibleAnimData, int32 PreCompressedSize, int32 ProgressNumerator, int32 ProgressDenominator);
 
-	void GatherPostCompressionStats(UAnimSequence* Seq, const TArray<FBoneData>& BoneData, double CompressionTime);
+	void GatherPostCompressionStats(const FCompressibleAnimData& CompressibleAnimData, FCompressibleAnimDataResult& CompressedData, double CompressionTime);
 
 	~FCompressionMemorySummary();
 
@@ -220,9 +222,9 @@ struct ENGINE_API FAnimCompressContext
 private:
 	FCompressionMemorySummary	CompressionSummary;
 
-	void GatherPreCompressionStats(UAnimSequence* Seq);
+	void GatherPreCompressionStats(const FCompressibleAnimData& CompressibleAnimData, int32 PreviousCompressionSize);
 
-	void GatherPostCompressionStats(UAnimSequence* Seq, const TArray<FBoneData>& BoneData, double CompressionTime);
+	void GatherPostCompressionStats(const FCompressibleAnimData& CompressibleAnimData, FCompressibleAnimDataResult& CompressedData, double CompressionTime);
 
 
 public:
@@ -231,14 +233,28 @@ public:
 	bool						bAllowAlternateCompressor;
 	bool						bOutput;
 
-	FAnimCompressContext(bool bInAllowAlternateCompressor, bool bInOutput, uint32 InMaxAnimations = 1) : CompressionSummary(bInOutput), AnimIndex(0), MaxAnimations(InMaxAnimations), bAllowAlternateCompressor(bInAllowAlternateCompressor), bOutput(bInOutput) {}
+	FAnimCompressContext(bool bInAllowAlternateCompressor, bool bInOutput, uint32 InMaxAnimations = 1)
+		: CompressionSummary(bInOutput)
+		, AnimIndex(0)
+		, MaxAnimations(InMaxAnimations)
+		, bAllowAlternateCompressor(bInAllowAlternateCompressor)
+		, bOutput(bInOutput)
+	{}
 
 	// If we are duping a compression context we don't want the CompressionSummary to output
-	FAnimCompressContext(const FAnimCompressContext& Rhs) : CompressionSummary(false), AnimIndex(Rhs.AnimIndex), MaxAnimations(Rhs.MaxAnimations), bAllowAlternateCompressor(Rhs.bAllowAlternateCompressor), bOutput(Rhs.bOutput) {}
+	FAnimCompressContext(const FAnimCompressContext& Rhs)
+		: CompressionSummary(false)
+		, AnimIndex(Rhs.AnimIndex)
+		, MaxAnimations(Rhs.MaxAnimations)
+		, bAllowAlternateCompressor(Rhs.bAllowAlternateCompressor)
+		, bOutput(Rhs.bOutput)
+	{}
 
 	friend class FAnimationUtils;
+	friend class FDerivedDataAnimationCompression;
 };
 
+#if USE_SEGMENTING_CONTEXT
 //////////////////////////////////////////////////////////////////////////
 // FAnimSegmentContext - This holds the relevant intermediate information
 // when compressing animation sequence segments.
@@ -260,6 +276,7 @@ struct FAnimSegmentContext
 	TArray<uint8> CompressedByteStream;
 	TArray<uint8> CompressedTrivialTracksByteStream;
 };
+#endif
 
 UCLASS(abstract, hidecategories=Object, MinimalAPI, EditInlineNew)
 class UAnimCompress : public UObject
@@ -305,15 +322,6 @@ class UAnimCompress : public UObject
 
 #if WITH_EDITOR
 public:
-	/**
-	 * Reduce the number of keyframes and bitwise compress the specified sequence.
-	 *
-	 * @param	AnimSeq		The animation sequence to compress.
-	 * @param	bOutput		If false don't generate output or compute memory savings.
-	 * @return				false if a skeleton was needed by the algorithm but not provided.
-	 */
-	ENGINE_API bool Reduce(class UAnimSequence* AnimSeq, bool bOutput, const TArray<FBoneData>& BoneData);
-
 	/**
 	 * Reduce the number of keyframes and bitwise compress all sequences in the specified array.
 	 *
@@ -321,7 +329,8 @@ public:
 	 * @param	bOutput			If false don't generate output or compute memory savings.
 	 * @return					false if a skeleton was needed by the algorithm but not provided.
 	 */
-	ENGINE_API bool Reduce(class UAnimSequence* AnimSeq, FAnimCompressContext& Context, const TArray<FBoneData>& BoneData);
+	ENGINE_API bool Reduce(const FCompressibleAnimData& CompressibleAnimData, FCompressibleAnimDataResult& OutResult);
+
 #endif // WITH_EDITOR
 protected:
 #if WITH_EDITOR
@@ -331,7 +340,7 @@ protected:
 	 *
 	 * @return		true if the keyframe reduction was successful.
 	 */
-	virtual void DoReduction(class UAnimSequence* AnimSeq, const TArray<class FBoneData>& BoneData) PURE_VIRTUAL(UAnimCompress::DoReduction,);
+	virtual void DoReduction(const FCompressibleAnimData& CompressibleAnimData, FCompressibleAnimDataResult& OutResult) PURE_VIRTUAL(UAnimCompress::DoReduction,);
 #endif // WITH_EDITOR
 	/**
 	 * Common compression utility to remove 'redundant' position keys based on the provided delta threshold
@@ -420,11 +429,13 @@ protected:
 	 * @param	MaxRotDelta		Maximum angle threshold to consider stationary motion
 	 * @param	MaxScaleDelta	Maximum scale threshold to consider stationary motion
 	 */
+#if USE_SEGMENTING_CONTEXT
 	static void FilterTrivialKeys(
 		TArray<FAnimSegmentContext>& RawSegments,
 		float MaxPosDelta,
 		float MaxRotDelta,
 		float MaxScaleDelta);
+#endif
 
 	/**
 	 * Common compression utility to retain only intermittent position keys. For example,
@@ -521,6 +532,7 @@ protected:
 	 * @param	MaxNumFramesPerSegment		The maximum number of frames within a segment
 	 * @param	OutRawSegments				The array of segments to populate
 	 */
+#if USE_SEGMENTING_CONTEXT
 	static void SeparateRawDataIntoTracks(
 		const UAnimSequence& AnimSeq,
 		const TArray<struct FTranslationTrack>& TranslationData,
@@ -529,7 +541,7 @@ protected:
 		int32 IdealNumFramesPerSegment,
 		int32 MaxNumFramesPerSegment,
 		TArray<FAnimSegmentContext>& OutRawSegments);
-
+#endif
 	/**
 	 * Common compression utility to walk an array of rotation tracks and enforce
 	 * that all adjacent rotation keys are represented by shortest-arc quaternion pairs.
@@ -553,7 +565,8 @@ public:
 	 * @param	IncludeKeyTable				true if the compressed data should also contain a table of frame indices for each key. (required by some codecs)
 	 */
 	static void BitwiseCompressAnimationTracks(
-		UAnimSequence* Seq, 
+		const FCompressibleAnimData& CompressibleAnimData,
+		FCompressibleAnimDataResult& OutCompressedData,
 		AnimationCompressionFormat TargetTranslationFormat, 
 		AnimationCompressionFormat TargetRotationFormat,
 		AnimationCompressionFormat TargetScaleFormat,
@@ -573,6 +586,7 @@ public:
 	 * @param	RawSegments					The array of segments to compress
 	 * @param	bIsSorted					For variable interpolation, is the compressed data sorted or not?
 	 */
+#if USE_SEGMENTING_CONTEXT
 	static void BitwiseCompressAnimationTracks(
 		UAnimSequence& AnimSeq,
 		AnimationCompressionFormat TargetTranslationFormat,
@@ -616,6 +630,7 @@ public:
 	 * @param	bIsSorted					For variable interpolation, is the compressed data sorted or not?
 	 */
 	static void CoalesceCompressedSegments(UAnimSequence& AnimSeq, const TArray<FAnimSegmentContext>& RawSegments, bool bIsSorted = false);
+#endif
 
 #if WITH_EDITOR
 	FString MakeDDCKey();
@@ -688,6 +703,7 @@ protected:
 		const float* Mins,
 		const float* Ranges);
 
+#if USE_SEGMENTING_CONTEXT
 	/**
 	 * Utility function that performs minimal sanity checks.
 	 *
@@ -695,6 +711,7 @@ protected:
 	 * @param	Segment						The segment to check
 	 */
 	static void SanityCheckTrackData(const UAnimSequence& AnimSeq, const FAnimSegmentContext& Segment);
+#endif
 
 	/**
 	 * Calculates the translation track range.
@@ -726,6 +743,7 @@ protected:
 	 */
 	static void CalculateTrackRange(const FScaleTrack& ScaleData, AnimationCompressionFormat Format, FVector& OutMin, FVector& OutExtent);
 
+#if USE_SEGMENTING_CONTEXT
 	/**
 	 * Calculates the track ranges within a segment.
 	 *
@@ -741,6 +759,7 @@ protected:
 		AnimationCompressionFormat TargetScaleFormat,
 		const FAnimSegmentContext& Segment,
 		TArray<FAnimTrackRange>& TrackRanges);
+#endif
 
 	/**
 	 * Structure to wrap and represent track key flags.
@@ -759,6 +778,7 @@ protected:
 		uint8 Flags;
 	};
 
+#if USE_SEGMENTING_CONTEXT
 	/**
 	 * Writes the necessary track ranges to a byte stream.
 	 *
@@ -881,6 +901,7 @@ protected:
 		TFunction<void(TArray<uint8>& ByteStream, AnimationCompressionFormat Format, const FVector& Key, const float* Mins, const float* Ranges, int32 TrackIndex)> PackScaleKeyFun,
 		const FAnimSegmentContext& Segment,
 		const TArray<FAnimTrackRange>& TrackRanges);
+#endif
 
 	/**
 	 * Pads a byte stream to force a particular alignment for the data to follow.
diff --git a/Engine/Source/Runtime/Engine/Classes/Animation/AnimCompress_Automatic.h b/Engine/Source/Runtime/Engine/Classes/Animation/AnimCompress_Automatic.h
--- a/Engine/Source/Runtime/Engine/Classes/Animation/AnimCompress_Automatic.h
+++ b/Engine/Source/Runtime/Engine/Classes/Animation/AnimCompress_Automatic.h
@@ -51,7 +51,7 @@ class UAnimCompress_Automatic : public UAnimCompress
 protected:
 	//~ Begin UAnimCompress Interface
 #if WITH_EDITOR
-	virtual void DoReduction(class UAnimSequence* AnimSeq, const TArray<class FBoneData>& BoneData) override;
+	virtual void DoReduction(const FCompressibleAnimData& CompressibleAnimData, FCompressibleAnimDataResult& OutResult) override;
 	virtual void PopulateDDCKey(FArchive& Ar) override;
 #endif // WITH_EDITOR
 	//~ Begin UAnimCompress Interface
diff --git a/Engine/Source/Runtime/Engine/Classes/Animation/AnimCompress_BitwiseCompressOnly.h b/Engine/Source/Runtime/Engine/Classes/Animation/AnimCompress_BitwiseCompressOnly.h
--- a/Engine/Source/Runtime/Engine/Classes/Animation/AnimCompress_BitwiseCompressOnly.h
+++ b/Engine/Source/Runtime/Engine/Classes/Animation/AnimCompress_BitwiseCompressOnly.h
@@ -21,7 +21,7 @@ class UAnimCompress_BitwiseCompressOnly : public UAnimCompress
 protected:
 	//~ Begin UAnimCompress Interface
 #if WITH_EDITOR
-	virtual void DoReduction(class UAnimSequence* AnimSeq, const TArray<class FBoneData>& BoneData) override;
+	virtual void DoReduction(const FCompressibleAnimData& CompressibleAnimData, FCompressibleAnimDataResult& OutResult) override;
 #endif // WITH_EDITOR
 	//~ Begin UAnimCompress Interface
 };
diff --git a/Engine/Source/Runtime/Engine/Classes/Animation/AnimCompress_PerTrackCompression.h b/Engine/Source/Runtime/Engine/Classes/Animation/AnimCompress_PerTrackCompression.h
--- a/Engine/Source/Runtime/Engine/Classes/Animation/AnimCompress_PerTrackCompression.h
+++ b/Engine/Source/Runtime/Engine/Classes/Animation/AnimCompress_PerTrackCompression.h
@@ -139,7 +139,7 @@ protected:
 	//~ Begin UAnimCompress Interface
 #if WITH_EDITOR
 	virtual void PostEditChangeProperty(struct FPropertyChangedEvent& PropertyChangedEvent) override;
-	virtual void DoReduction(class UAnimSequence* AnimSeq, const TArray<class FBoneData>& BoneData) override;
+	virtual void DoReduction(const FCompressibleAnimData& CompressibleAnimData, FCompressibleAnimDataResult& OutResult) override;
 	virtual void PopulateDDCKey(FArchive& Ar) override;
 #endif // WITH_EDITOR
 	//~ Begin UAnimCompress Interface
@@ -147,18 +147,20 @@ protected:
 #if WITH_EDITOR
 	//~ Begin UAnimCompress_RemoveLinearKeys Interface
 	virtual void CompressUsingUnderlyingCompressor(
-		UAnimSequence* AnimSeq,
-		const TArray<FBoneData>& BoneData, 
+		const FCompressibleAnimData& CompressibleAnimData,
+		FCompressibleAnimDataResult& OutCompressedData,
 		const TArray<FTranslationTrack>& TranslationData,
 		const TArray<FRotationTrack>& RotationData,
 		const TArray<FScaleTrack>& ScaleData,
 		const bool bFinalPass) override;
 
+#if USE_SEGMENTING_CONTEXT
 	virtual void CompressUsingUnderlyingCompressor(
 		UAnimSequence& AnimSeq,
 		const TArray<FBoneData>& BoneData,
 		TArray<FAnimSegmentContext>& RawSegments,
 		const bool bFinalPass) override;
+#endif
 
 	/**
 	 * Performs the per track compression optimization for a single segment.
@@ -167,8 +169,7 @@ protected:
 	void OptimizeSegmentTracks(struct FOptimizeSegmentTracksContext& Context) const;
 
 	virtual void FilterBeforeMainKeyRemoval(
-		UAnimSequence* AnimSeq, 
-		const TArray<FBoneData>& BoneData, 
+		const FCompressibleAnimData& CompressibleAnimData,
 		TArray<FTranslationTrack>& TranslationData,
 		TArray<FRotationTrack>& RotationData,
 		TArray<FScaleTrack>& ScaleData) override;
diff --git a/Engine/Source/Runtime/Engine/Classes/Animation/AnimCompress_RemoveLinearKeys.h b/Engine/Source/Runtime/Engine/Classes/Animation/AnimCompress_RemoveLinearKeys.h
--- a/Engine/Source/Runtime/Engine/Classes/Animation/AnimCompress_RemoveLinearKeys.h
+++ b/Engine/Source/Runtime/Engine/Classes/Animation/AnimCompress_RemoveLinearKeys.h
@@ -106,7 +106,7 @@ class UAnimCompress_RemoveLinearKeys : public UAnimCompress
 protected:
 	//~ Begin UAnimCompress Interface
 #if WITH_EDITOR
-	virtual void DoReduction(class UAnimSequence* AnimSeq, const TArray<class FBoneData>& BoneData) override;
+	virtual void DoReduction(const FCompressibleAnimData& CompressibleAnimData, FCompressibleAnimDataResult& OutResult) override;
 	virtual void PopulateDDCKey(FArchive& Ar) override;
 #endif // WITH_EDITOR
 	//~ Begin UAnimCompress Interface
@@ -116,8 +116,7 @@ protected:
 	 * Pre-filters the tracks before running the main key removal algorithm
 	 */
 	virtual void FilterBeforeMainKeyRemoval(
-		UAnimSequence* AnimSeq, 
-		const TArray<FBoneData>& BoneData, 
+		const FCompressibleAnimData& CompressibleAnimData,
 		TArray<FTranslationTrack>& TranslationData,
 		TArray<FRotationTrack>& RotationData, 
 		TArray<FScaleTrack>& ScaleData);
@@ -126,13 +125,14 @@ protected:
 	 * Compresses the tracks passed in using the underlying compressor for this key removal codec
 	 */
 	virtual void CompressUsingUnderlyingCompressor(
-		UAnimSequence* AnimSeq, 
-		const TArray<FBoneData>& BoneData, 
+		const FCompressibleAnimData& CompressibleAnimData,
+		FCompressibleAnimDataResult& OutCompressedData,
 		const TArray<FTranslationTrack>& TranslationData,
 		const TArray<FRotationTrack>& RotationData,
 		const TArray<FScaleTrack>& ScaleData,
 		const bool bFinalPass);
 
+#if USE_SEGMENTING_CONTEXT
 	/**
 	 * Compresses the tracks passed in using the underlying compressor for this key removal codec for a single segment
 	 */
@@ -141,13 +141,14 @@ protected:
 		const TArray<FBoneData>& BoneData,
 		TArray<FAnimSegmentContext>& RawSegments,
 		const bool bFinalPass);
+#endif
 
 	/**
 	 * Updates the world bone transforms for a range of bone indices
 	 */
 	void UpdateWorldBoneTransformRange(
-		UAnimSequence* AnimSeq, 
-		const TArray<FBoneData>& BoneData, 
+		const FCompressibleAnimData& CompressibleAnimData,
+		FCompressibleAnimDataResult& OutCompressedData,
 		const TArray<FTransform>& RefPose,
 		const TArray<FTranslationTrack>& PositionTracks,
 		const TArray<FRotationTrack>& RotationTracks,
@@ -157,27 +158,40 @@ protected:
 		bool UseRaw,
 		TArray<FTransform>& OutputWorldBones);
 
-	/**
-	 * Updates the world bone transforms for a range of bone indices for a single segment
-	 */
-	void UpdateWorldBoneTransformRange(
-		FProcessAnimationTracksContext& Context,
-		int32 StartingBoneIndex,
-		int32 EndingBoneIndex);
-
 	/**
 	 * To guide the key removal process, we need to maintain a table of world transforms
 	 * for the bones we are investigating. This helper function fills a row of the 
 	 * table for a specified bone.
 	 */
 	void UpdateWorldBoneTransformTable(
-		UAnimSequence* AnimSeq,
-		const TArray<FBoneData>& BoneData,
+		const FCompressibleAnimData& CompressibleAnimData,
+		FCompressibleAnimDataResult& OutCompressedData,
 		const TArray<FTransform>& RefPose,
 		int32 BoneIndex,
 		bool UseRaw,
 		TArray<FTransform>& OutputWorldBones);
 
+	/**
+	 * Creates a list of the bone atom result for every frame of a given track
+	 */
+	static void UpdateBoneAtomList(
+		const FCompressibleAnimData& CompressibleAnimData,
+		FCompressibleAnimDataResult& OutCompressedData,
+		int32 BoneIndex,
+		int32 TrackIndex,
+		int32 NumFrames,
+		float TimePerFrame,
+		TArray<FTransform>& BoneAtoms);
+
+#if USE_SEGMENTING_CONTEXT
+	/**
+	 * Updates the world bone transforms for a range of bone indices for a single segment
+	 */
+	void UpdateWorldBoneTransformRange(
+		FProcessAnimationTracksContext& Context,
+		int32 StartingBoneIndex,
+		int32 EndingBoneIndex);
+
 	/**
 	 * To guide the key removal process, we need to maintain a table of world transforms
 	 * for the bones we are investigating. This helper function fills a row of the
@@ -189,17 +203,6 @@ protected:
 		bool UseRaw,
 		TArray<FTransform>& OutputWorldBones);
 
-	/**
-	 * Creates a list of the bone atom result for every frame of a given track
-	 */
-	static void UpdateBoneAtomList(
-		UAnimSequence* AnimSeq, 
-		int32 BoneIndex,
-		int32 TrackIndex,
-		int32 NumFrames,
-		float TimePerFrame,
-		TArray<FTransform>& BoneAtoms);
-
 	/**
 	 * Creates a list of the bone atom result for every frame of a given track for a single segment
 	 */
@@ -207,6 +210,7 @@ protected:
 		const FProcessAnimationTracksContext& Context,
 		int32 TrackIndex,
 		TArray<FTransform>& BoneAtoms) const;
+#endif
 
 	/**
 	 * If the passed in animation sequence is additive, converts it to absolute (using the frame 0 pose) and returns true
@@ -214,9 +218,8 @@ protected:
 	 *
 	 * @param AnimSeq			The animation sequence being compressed
 	 *
-	 * @return true if the animation was additive and has been converted to absolute space.
 	 */
-	bool ConvertFromRelativeSpace(UAnimSequence* AnimSeq);
+	void ConvertFromRelativeSpace(FCompressibleAnimData& CompressibleAnimData);
 
 	/**
 	 * Converts an absolute animation sequence and matching track data to a relative (additive) one.
@@ -227,7 +230,7 @@ protected:
 	 * @param ScaleData			Scale Tracks to convert to relative space
 	 *
 	 */
-	void ConvertToRelativeSpace(UAnimSequence* AnimSeq, TArray<FTranslationTrack>& TranslationData, TArray<FRotationTrack>& RotationData, TArray<FScaleTrack>& ScaleData);
+	void ConvertToRelativeSpaceBoth(FCompressibleAnimData& CompressibleAnimData, TArray<FTranslationTrack>& TranslationData, TArray<FRotationTrack>& RotationData, TArray<FScaleTrack>& ScaleData);
 
 	/**
 	 * Converts an absolute animation sequence to a relative (additive) one.
@@ -235,7 +238,7 @@ protected:
 	 * @param AnimSeq			The animation sequence being compressed and converted
 	 *
 	 */
-	void ConvertToRelativeSpace(UAnimSequence& AnimSeq) const;
+	void ConvertToRelativeSpace(FCompressibleAnimData& CompressibleAnimData) const;
 
 	/**
 	 * Converts track data to relative (additive) space.
@@ -246,7 +249,7 @@ protected:
 	 * @param ScaleData			Scale Tracks to convert to relative space
 	 *
 	 */
-	void ConvertToRelativeSpace(const UAnimSequence& AnimSeq, TArray<FTranslationTrack>& TranslationData, TArray<FRotationTrack>& RotationData, TArray<FScaleTrack>& ScaleData) const;
+	void ConvertToRelativeSpace(const FCompressibleAnimData& CompressibleAnimData, TArray<FTranslationTrack>& TranslationData, TArray<FRotationTrack>& RotationData, TArray<FScaleTrack>& ScaleData) const;
 
 
 	/**
@@ -262,12 +265,13 @@ protected:
 	 * @return				None.
 	 */
 	void ProcessAnimationTracks(
-		UAnimSequence* AnimSeq, 
-		const TArray<FBoneData>& BoneData, 
+		const FCompressibleAnimData& CompressibleAnimData,
+		FCompressibleAnimDataResult& OutCompressedData,
 		TArray<FTranslationTrack>& PositionTracks,
 		TArray<FRotationTrack>& RotationTracks, 
 		TArray<FScaleTrack>& ScaleTracks);
 
+#if USE_SEGMENTING_CONTEXT
 	/**
 	 * Locates spans of keys within the position and rotation tracks provided which can be estimated
 	 * through linear interpolation of the surrounding keys. The remaining key values are bit packed into
@@ -295,8 +299,10 @@ protected:
 	 * Samples a track within a segment at a specified time with the current track format.
 	 */
 	static FTransform SampleSegment(const FProcessAnimationTracksContext& Context, int32 TrackIndex, float Time);
+#endif
 
 private:
+#if USE_SEGMENTING_CONTEXT
 	/**
 	 * Performs the linear key reduction and retargetting for a single segment.
 	 * This can be called from multiple threads concurrently.
@@ -304,6 +310,7 @@ private:
 	void ProcessAnimationTracks(FProcessAnimationTracksContext& Context);
 
 	friend struct FAsyncProcessAnimationTracksTaskGroupContext;
+#endif
 
 #endif // WITH_EDITOR
 };
diff --git a/Engine/Source/Runtime/Engine/Classes/Animation/AnimCurveCompressionCodec.h b/Engine/Source/Runtime/Engine/Classes/Animation/AnimCurveCompressionCodec.h
--- a/Engine/Source/Runtime/Engine/Classes/Animation/AnimCurveCompressionCodec.h
+++ b/Engine/Source/Runtime/Engine/Classes/Animation/AnimCurveCompressionCodec.h
@@ -53,7 +53,7 @@ class ENGINE_API UAnimCurveCompressionCodec : public UObject
 	virtual bool IsCodecValid() const { return true; }
 
 	/** Compresses the curve data from an animation sequence. */
-	virtual bool Compress(const UAnimSequence& AnimSeq, FAnimCurveCompressionResult& OutResult) PURE_VIRTUAL(UAnimCurveCompressionCodec::Compress, return false;);
+	virtual bool Compress(const FCompressibleAnimData& AnimSeq, FAnimCurveCompressionResult& OutResult) PURE_VIRTUAL(UAnimCurveCompressionCodec::Compress, return false;);
 
 	/*
 	 * Called to generate a unique DDC key for this codec instance.
@@ -67,12 +67,12 @@ class ENGINE_API UAnimCurveCompressionCodec : public UObject
 	 * Note: Codecs should _NOT_ rely on any member properties during decompression. Decompression
 	 * behavior should entirely be driven by code and the compressed data.
 	 */
-	virtual void DecompressCurves(const UAnimSequence& AnimSeq, FBlendedCurve& Curves, float CurrentTime) const PURE_VIRTUAL(UAnimCurveCompressionCodec::DecompressCurves, );
+	virtual void DecompressCurves(const FCompressedAnimSequence& AnimSeq, FBlendedCurve& Curves, float CurrentTime) const PURE_VIRTUAL(UAnimCurveCompressionCodec::DecompressCurves, );
 
 	/*
 	 * Decompress a single curve.
 	 * Note: Codecs should _NOT_ rely on any member properties during decompression. Decompression
 	 * behavior should entirely be driven by code and the compressed data.
 	 */
-	virtual float DecompressCurve(const UAnimSequence& AnimSeq, SmartName::UID_Type CurveUID, float CurrentTime) const PURE_VIRTUAL(UAnimCurveCompressionCodec::DecompressCurve, return 0.0f;);
+	virtual float DecompressCurve(const FCompressedAnimSequence& AnimSeq, SmartName::UID_Type CurveUID, float CurrentTime) const PURE_VIRTUAL(UAnimCurveCompressionCodec::DecompressCurve, return 0.0f;);
 };
diff --git a/Engine/Source/Runtime/Engine/Classes/Animation/AnimCurveCompressionCodec_CompressedRichCurve.h b/Engine/Source/Runtime/Engine/Classes/Animation/AnimCurveCompressionCodec_CompressedRichCurve.h
--- a/Engine/Source/Runtime/Engine/Classes/Animation/AnimCurveCompressionCodec_CompressedRichCurve.h
+++ b/Engine/Source/Runtime/Engine/Classes/Animation/AnimCurveCompressionCodec_CompressedRichCurve.h
@@ -33,10 +33,10 @@ class ENGINE_API UAnimCurveCompressionCodec_CompressedRichCurve : public UAnimCu
 
 #if WITH_EDITORONLY_DATA
 	// UAnimCurveCompressionCodec overrides
-	virtual bool Compress(const UAnimSequence& AnimSeq, FAnimCurveCompressionResult& OutResult) override;
+	virtual bool Compress(const FCompressibleAnimData& AnimSeq, FAnimCurveCompressionResult& OutResult) override;
 	virtual void PopulateDDCKey(FArchive& Ar) override;
 #endif
 
-	virtual void DecompressCurves(const UAnimSequence& AnimSeq, FBlendedCurve& Curves, float CurrentTime) const override;
-	virtual float DecompressCurve(const UAnimSequence& AnimSeq, SmartName::UID_Type CurveUID, float CurrentTime) const override;
+	virtual void DecompressCurves(const FCompressedAnimSequence& AnimSeq, FBlendedCurve& Curves, float CurrentTime) const override;
+	virtual float DecompressCurve(const FCompressedAnimSequence& AnimSeq, SmartName::UID_Type CurveUID, float CurrentTime) const override;
 };
diff --git a/Engine/Source/Runtime/Engine/Classes/Animation/AnimCurveCompressionCodec_UniformlySampled.h b/Engine/Source/Runtime/Engine/Classes/Animation/AnimCurveCompressionCodec_UniformlySampled.h
--- a/Engine/Source/Runtime/Engine/Classes/Animation/AnimCurveCompressionCodec_UniformlySampled.h
+++ b/Engine/Source/Runtime/Engine/Classes/Animation/AnimCurveCompressionCodec_UniformlySampled.h
@@ -29,10 +29,10 @@ class ENGINE_API UAnimCurveCompressionCodec_UniformlySampled : public UAnimCurve
 
 #if WITH_EDITORONLY_DATA
 	// UAnimCurveCompressionCodec overrides
-	virtual bool Compress(const UAnimSequence& AnimSeq, FAnimCurveCompressionResult& OutResult) override;
+	virtual bool Compress(const FCompressibleAnimData& AnimSeq, FAnimCurveCompressionResult& OutResult) override;
 	virtual void PopulateDDCKey(FArchive& Ar) override;
 #endif
 
-	virtual void DecompressCurves(const UAnimSequence& AnimSeq, FBlendedCurve& Curves, float CurrentTime) const override;
-	virtual float DecompressCurve(const UAnimSequence& AnimSeq, SmartName::UID_Type CurveUID, float CurrentTime) const override;
+	virtual void DecompressCurves(const FCompressedAnimSequence& AnimSeq, FBlendedCurve& Curves, float CurrentTime) const override;
+	virtual float DecompressCurve(const FCompressedAnimSequence& AnimSeq, SmartName::UID_Type CurveUID, float CurrentTime) const override;
 };
diff --git a/Engine/Source/Runtime/Engine/Classes/Animation/AnimCurveCompressionSettings.h b/Engine/Source/Runtime/Engine/Classes/Animation/AnimCurveCompressionSettings.h
--- a/Engine/Source/Runtime/Engine/Classes/Animation/AnimCurveCompressionSettings.h
+++ b/Engine/Source/Runtime/Engine/Classes/Animation/AnimCurveCompressionSettings.h
@@ -31,11 +31,10 @@ class ENGINE_API UAnimCurveCompressionSettings : public UObject
 	bool AreSettingsValid() const;
 
 	/*
-	 * Compresses the animation curves inside the supplied sequence.
-	 * Note that this will modify the animation sequence by populating the compressed bytes
-	 * and the codec used but it is left unchanged if compression fails.
+	 * Compresses the animation curves inside the supplied sequence data.
+	 * The resultant compressed data is applied to the OutCompressedData structure.
 	 */
-	bool Compress(UAnimSequence& AnimSeq) const;
+	bool Compress(const FCompressibleAnimData& AnimSeq, FCompressedAnimSequence& OutCompressedData) const;
 
 	/** Generates a DDC key that takes into account the current settings and selected codec. */
 	FString MakeDDCKey() const;
diff --git a/Engine/Source/Runtime/Engine/Classes/Animation/AnimSequence.h b/Engine/Source/Runtime/Engine/Classes/Animation/AnimSequence.h
--- a/Engine/Source/Runtime/Engine/Classes/Animation/AnimSequence.h
+++ b/Engine/Source/Runtime/Engine/Classes/Animation/AnimSequence.h
@@ -15,15 +15,11 @@
 #include "Animation/AnimationAsset.h"
 #include "Animation/AnimCurveTypes.h"
 #include "Animation/AnimSequenceBase.h"
-
-#include "Async/MappedFileHandle.h"
-#include "HAL/PlatformFilemanager.h"
-#include "GenericPlatform/GenericPlatformFile.h"
-#include "Misc/Paths.h"
-#include "Serialization/BulkData.h"
+#include "Animation/AnimCompressionTypes.h"
 
 #include "AnimSequence.generated.h"
 
+#define USE_SEGMENTING_CONTEXT 0 // Uses segmenting in anim compression + context in decompression
 
 typedef TArray<FTransform> FTransformArrayA2;
 
@@ -32,54 +28,6 @@ struct FAnimCompressContext;
 struct FAnimSequenceDecompressionContext;
 struct FCompactPose;
 
-/**
- * Indicates animation data key format.
- */
-UENUM()
-enum AnimationKeyFormat
-{
-	AKF_ConstantKeyLerp,
-	AKF_VariableKeyLerp,
-	AKF_PerTrackCompression,
-	AKF_MAX,
-};
-
-/**
- * Raw keyframe data for one track.  Each array will contain either NumFrames elements or 1 element.
- * One element is used as a simple compression scheme where if all keys are the same, they'll be
- * reduced to 1 key that is constant over the entire sequence.
- */
-USTRUCT()
-struct ENGINE_API FRawAnimSequenceTrack
-{
-	GENERATED_USTRUCT_BODY()
-
-	/** Position keys. */
-	UPROPERTY()
-	TArray<FVector> PosKeys;
-
-	/** Rotation keys. */
-	UPROPERTY()
-	TArray<FQuat> RotKeys;
-
-	/** Scale keys. */
-	UPROPERTY()
-	TArray<FVector> ScaleKeys;
-
-	// Serializer.
-	friend FArchive& operator<<( FArchive& Ar, FRawAnimSequenceTrack& T )
-	{
-		T.PosKeys.BulkSerialize(Ar);
-		T.RotKeys.BulkSerialize(Ar);
-
-		if (Ar.UE4Ver() >= VER_UE4_ANIM_SUPPORT_NONUNIFORM_SCALE_ANIMATION)
-		{
-			T.ScaleKeys.BulkSerialize(Ar);
-		}
-
-		return Ar;
-	}
-};
 
 // These two always should go together, but it is not right now. 
 // I wonder in the future, we change all compressed to be inside as well, so they all stay together
@@ -119,33 +67,6 @@ struct ENGINE_API FAnimSequenceTrackContainer
 	}
 };
 
-// @note We have a plan to support skeletal hierarchy. When that happens, we'd like to keep skeleton indexing.
-USTRUCT()
-struct ENGINE_API FTrackToSkeletonMap
-{
-	GENERATED_USTRUCT_BODY()
-
-	// Index of Skeleton.BoneTree this Track belongs to.
-	UPROPERTY()
-	int32 BoneTreeIndex;
-
-
-	FTrackToSkeletonMap()
-		: BoneTreeIndex(0)
-	{
-	}
-
-	FTrackToSkeletonMap(int32 InBoneTreeIndex)
-		: BoneTreeIndex(InBoneTreeIndex)
-	{
-	}
-
-	friend FArchive& operator<<(FArchive& Ar, FTrackToSkeletonMap &Item)
-	{
-		return Ar << Item.BoneTreeIndex;
-	}
-};
-
 /**
  * Keyframe position data for one track.  Pos(i) occurs at Time(i).  Pos.Num() always equals Time.Num().
  */
@@ -249,61 +170,6 @@ struct ENGINE_API FCompressedTrack
 
 };
 
-struct ENGINE_API FCompressedOffsetData
-{
-	TArray<int32> OffsetData;
-
-	int32 StripSize;
-
-	FCompressedOffsetData( int32 InStripSize=2 )
-		: StripSize (InStripSize)
-		{}
-
-	void SetStripSize(int32 InStripSize)
-	{
-		ensure (InStripSize > 0 );
-		StripSize = InStripSize;
-	}
-
-	const int32 GetOffsetData(int32 StripIndex, int32 Offset) const
-	{
-		checkSlow(OffsetData.IsValidIndex(StripIndex * StripSize + Offset));
-
-		return OffsetData[StripIndex * StripSize + Offset];
-	}
-
-	void SetOffsetData(int32 StripIndex, int32 Offset, int32 Value)
-	{
-		checkSlow(OffsetData.IsValidIndex(StripIndex * StripSize + Offset));
-		OffsetData[StripIndex * StripSize + Offset] = Value;
-	}
-
-	void AddUninitialized(int32 NumOfTracks)
-	{
-		OffsetData.AddUninitialized(NumOfTracks*StripSize);
-	}
-
-	void Empty(int32 NumOfTracks=0)
-	{
-		OffsetData.Empty(NumOfTracks*StripSize);
-	}
-
-	int32 GetMemorySize() const
-	{
-		return sizeof(int32)*OffsetData.Num() + sizeof(int32);
-	}
-
-	int32 GetNumTracks() const
-	{
-		return OffsetData.Num()/StripSize;
-	}
-
-	bool IsValid() const
-	{
-		return (OffsetData.Num() > 0);
-	}
-};
-
 // Param structure for UAnimSequence::RequestAnimCompressionParams
 struct ENGINE_API FRequestAnimCompressionParams
 {
@@ -377,229 +243,6 @@ private:
 
 FArchive& operator<<(FArchive& Ar, FCompressedOffsetData& D);
 
-/**
- * Represents a segment of the anim sequence that is compressed.
- */
-USTRUCT()
-struct ENGINE_API FCompressedSegment
-{
-	GENERATED_USTRUCT_BODY()
-
-	// Frame where the segment begins in the anim sequence
-	int32 StartFrame;
-
-	// Num of frames contained in the segment
-	int32 NumFrames;
-
-	// Segment data offset in CompressedByteStream
-	int32 ByteStreamOffset;
-
-	/** The compression format that was used to compress translation tracks. */
-	TEnumAsByte<enum AnimationCompressionFormat> TranslationCompressionFormat;
-
-	/** The compression format that was used to compress rotation tracks. */
-	TEnumAsByte<enum AnimationCompressionFormat> RotationCompressionFormat;
-
-	/** The compression format that was used to compress rotation tracks. */
-	TEnumAsByte<enum AnimationCompressionFormat> ScaleCompressionFormat;
-
-	FCompressedSegment()
-		: StartFrame(0)
-		, NumFrames(0)
-		, ByteStreamOffset(0)
-		, TranslationCompressionFormat(ACF_None)
-		, RotationCompressionFormat(ACF_None)
-		, ScaleCompressionFormat(ACF_None)
-	{
-	}
-
-	friend FArchive& operator<<(FArchive& Ar, FCompressedSegment &Segment)
-	{
-		return Ar << Segment.StartFrame << Segment.NumFrames << Segment.ByteStreamOffset
-			<< Segment.TranslationCompressionFormat << Segment.RotationCompressionFormat << Segment.ScaleCompressionFormat;
-	}
-};
-
-
-template<uint32 Alignment = DEFAULT_ALIGNMENT>
-class TMaybeMappedAllocator
-{
-public:
-	using SizeType = int32;
-
-	enum { NeedsElementType = false };
-	enum { RequireRangeCheck = true };
-
-	class ForAnyElementType
-	{
-	public:
-
-		/** Default constructor. */
-		ForAnyElementType()
-			: Data(nullptr)
-			, MappedHandle(nullptr)
-			, MappedRegion(nullptr)
-		{}
-
-		/**
-		 * Moves the state of another allocator into this one.
-		 * Assumes that the allocator is currently empty, i.e. memory may be allocated but any existing elements have already been destructed (if necessary).
-		 * @param Other - The allocator to move the state from.  This allocator should be left in a valid empty state.
-		 */
-		void MoveToEmpty(ForAnyElementType& Other)
-		{
-			checkSlow(this != &Other);
-
-			Reset();
-
-			Data = Other.Data;
-			Other.Data = nullptr;
-
-			MappedRegion = Other.MappedRegion;
-			Other->MappedRegion = nullptr;
-
-			MappedHandle = Other.MappedHandle;
-			Other->MappedHandle = nullptr;
-		}
-
-		/** Destructor. */
-		~ForAnyElementType()
-		{
-			Reset();
-		}
-
-		// FContainerAllocatorInterface
-		FScriptContainerElement* GetAllocation() const
-		{
-			return Data;
-		}
-		void ResizeAllocation(
-			SizeType PreviousNumElements,
-			SizeType NumElements,
-			SIZE_T NumBytesPerElement
-		)
-		{
-			// Avoid calling FMemory::Realloc( nullptr, 0 ) as ANSI C mandates returning a valid pointer which is not what we want.
-			if (Data || NumElements)
-			{
-				check(!MappedHandle && !MappedRegion); // this could be supported, but it probably is never what you want, so we will just assert.
-					//checkSlow(((uint64)NumElements*(uint64)ElementTypeInfo.GetSize() < (uint64)INT_MAX));
-				Data = (FScriptContainerElement*)FMemory::Realloc(Data, NumElements*NumBytesPerElement, Alignment);
-			}
-		}
-		SizeType CalculateSlackReserve(SizeType NumElements, SIZE_T NumBytesPerElement) const
-		{
-			check(!MappedHandle && !MappedRegion); // this could be supported, but it probably is never what you want, so we will just assert.
-			return DefaultCalculateSlackReserve(NumElements, NumBytesPerElement, true, Alignment);
-		}
-		SizeType CalculateSlackShrink(SizeType NumElements, SizeType NumAllocatedElements, SIZE_T NumBytesPerElement) const
-		{
-			check(!MappedHandle && !MappedRegion); // this could be supported, but it probably is never what you want, so we will just assert.
-			return DefaultCalculateSlackShrink(NumElements, NumAllocatedElements, NumBytesPerElement, true, Alignment);
-		}
-		SizeType CalculateSlackGrow(SizeType NumElements, SizeType NumAllocatedElements, SIZE_T NumBytesPerElement) const
-		{
-			check(!MappedHandle && !MappedRegion); // this could be supported, but it probably is never what you want, so we will just assert.
-			return DefaultCalculateSlackGrow(NumElements, NumAllocatedElements, NumBytesPerElement, true, Alignment);
-		}
-
-		SIZE_T GetAllocatedSize(SizeType NumAllocatedElements, SIZE_T NumBytesPerElement) const
-		{
-			return NumAllocatedElements * NumBytesPerElement;
-		}
-
-		bool HasAllocation() const
-		{
-			return !!Data;
-		}
-
-		void AcceptFileMapping(IMappedFileHandle* InMappedHandle, IMappedFileRegion* InMappedRegion, void *MallocPtr)
-		{
-			check(!MappedHandle && !Data); // we could support stuff like this, but that usually isn't what we want for streamlined loading
-			Reset(); // just in case
-			if (InMappedHandle || InMappedRegion)
-			{
-				MappedHandle = InMappedHandle;
-				MappedRegion = InMappedRegion;
-				Data = (FScriptContainerElement*)MappedRegion->GetMappedPtr(); //@todo mapped files should probably be const-correct
-				check(IsAligned(Data, FPlatformProperties::GetMemoryMappingAlignment()));
-			}
-			else
-			{
-				Data = (FScriptContainerElement*)MallocPtr;
-			}
-		}
-
-		bool IsMapped() const
-		{
-			return MappedRegion || MappedHandle;
-		}
-	private:
-
-		FScriptContainerElement* Data;
-		IMappedFileHandle* MappedHandle;
-		IMappedFileRegion* MappedRegion;
-
-		void Reset()
-		{
-			if (MappedRegion || MappedHandle)
-			{
-				delete MappedRegion;
-				delete MappedHandle;
-				MappedRegion = nullptr;
-				MappedHandle = nullptr;
-				Data = nullptr; // make sure we don't try to free this pointer
-			}
-			if (Data)
-			{
-				FMemory::Free(Data);
-				Data = nullptr;
-			}
-		}
-
-
-		ForAnyElementType(const ForAnyElementType&);
-		ForAnyElementType& operator=(const ForAnyElementType&);
-	};
-
-	template<typename ElementType>
-	class ForElementType : public ForAnyElementType
-	{
-	public:
-
-		/** Default constructor. */
-		ForElementType()
-		{}
-
-		ElementType* GetAllocation() const
-		{
-			return (ElementType*)ForAnyElementType::GetAllocation();
-		}
-	};
-};
-
-template<typename T, uint32 Alignment = DEFAULT_ALIGNMENT>
-class TMaybeMappedArray : public TArray<T, TMaybeMappedAllocator<Alignment>>
-{
-public:
-	TMaybeMappedArray()
-	{
-	}
-	TMaybeMappedArray(TMaybeMappedArray&&) = default;
-	TMaybeMappedArray(const TMaybeMappedArray&) = default;
-	TMaybeMappedArray& operator=(TMaybeMappedArray&&) = default;
-	TMaybeMappedArray& operator=(const TMaybeMappedArray&) = default;
-
-	void AcceptOwnedBulkDataPtr(FOwnedBulkDataPtr* OwnedPtr, int32 Num)
-	{
-		this->ArrayNum = Num;
-		this->ArrayMax = Num;
-		this->AllocatorInstance.AcceptFileMapping(OwnedPtr->GetMappedHandle(), OwnedPtr->GetMappedRegion(), (void*)OwnedPtr->GetPointer());
-		OwnedPtr->RelinquishOwnership();
-	}
-};
-
-
 UCLASS(config=Engine, hidecategories=(UObject, Length), BlueprintType)
 class ENGINE_API UAnimSequence : public UAnimSequenceBase
 {
@@ -622,9 +265,6 @@ protected:
 	UPROPERTY(AssetRegistrySearchable, meta = (DisplayName = "Number of Keys"))
 	int32 NumFrames;
 
-	// The number of frames that the animation had when compressed data was created (may have been resampled for instance)
-	int32 CompressedNumFrames;
-
 	/**
 	 * In the future, maybe keeping RawAnimSequenceTrack + TrackMap as one would be good idea to avoid inconsistent array size
 	 * TrackToSkeletonMapTable(i) should contains  track mapping data for RawAnimationData(i). 
@@ -632,19 +272,6 @@ protected:
 	UPROPERTY()
 	TArray<struct FTrackToSkeletonMap> TrackToSkeletonMapTable;
 
-	/**
-	 * Version of TrackToSkeletonMapTable for the compressed tracks. Due to baking additive data
-	 * we can end up with a different amount of tracks to the original raw animation and so we must index into the 
-	 * compressed data using this 
-	 */
-	TArray<struct FTrackToSkeletonMap> CompressedTrackToSkeletonMapTable;
-
-	/**
-	 * Much like track indices above, we need to be able to remap curve names. The FName inside FSmartName
-	 * is serialized and the UID is generated at runtime. Stored in the DDC.
-	 */
-	TArray<struct FSmartName> CompressedCurveNames;
-
 	/**
 	 * Raw uncompressed keyframe data. 
 	 */
@@ -665,12 +292,6 @@ protected:
 	 * Source RawAnimationData. Only can be overridden by when transform curves are added first time OR imported
 	 */
 	TArray<struct FRawAnimSequenceTrack> SourceRawAnimationData;
-
-	/**
-	* Temporary base pose buffer for additive animation that is used by compression. Do not use this unless within compression. It is not available.
-	* This is used by remove linear key that has to rebuild to full transform in order to compress
-	*/
-	TArray<struct FRawAnimSequenceTrack> TemporaryAdditiveBaseAnimationData;
 #endif
 
 public:
@@ -695,80 +316,11 @@ public:
 	UPROPERTY(Category=Compression, EditAnywhere)
 	class UAnimCurveCompressionSettings* CurveCompressionSettings;
 
-	/** The codec used by the compressed data as determined by the compression settings. */
-	class UAnimCurveCompressionCodec* CurveCompressionCodec;
-
-	/** The compression format that was used to compress translation tracks. */
-	TEnumAsByte<enum AnimationCompressionFormat> TranslationCompressionFormat;
-
-	/** The compression format that was used to compress rotation tracks. */
-	TEnumAsByte<enum AnimationCompressionFormat> RotationCompressionFormat;
-
-	/** The compression format that was used to compress rotation tracks. */
-	TEnumAsByte<enum AnimationCompressionFormat> ScaleCompressionFormat;
-
-	/**
-	 * An array of 4*NumTrack ints, arranged as follows: - PerTrack is 2*NumTrack, so this isn't true any more
-	 *   [0] Trans0.Offset
-	 *   [1] Trans0.NumKeys
-	 *   [2] Rot0.Offset
-	 *   [3] Rot0.NumKeys
-	 *   [4] Trans1.Offset
-	 *   . . .
-	 */
-	TArray<int32> CompressedTrackOffsets;
-
-	/**
-	 * An array of 2*NumTrack ints, arranged as follows: 
-		if identity, it is offset
-		if not, it is num of keys 
-	 *   [0] Scale0.Offset or NumKeys
-	 *   [1] Scale1.Offset or NumKeys
-
-	 * @TODO NOTE: first implementation is offset is [0], numkeys [1]
-	 *   . . .
-	 */
-	FCompressedOffsetData CompressedScaleOffsets;
-
-	/**
-	 * ByteStream for compressed animation data.
-	 * The memory layout is dependent on the algorithm used to compress the anim sequence.
-	 */
-#if WITH_EDITOR
-	TArray<uint8> CompressedByteStream;
-	FByteBulkData OptionalBulk;
-#else
-	TMaybeMappedArray<uint8> CompressedByteStream;
-#endif
-
-
-	/**
-	 * Array of segment descriptors for this compressed anim sequence.
-	 */
-	TArray<FCompressedSegment> CompressedSegments;
-
-	TEnumAsByte<enum AnimationKeyFormat> KeyEncodingFormat;
-
-	/**
-	 * The runtime interface to decode and byte swap the compressed animation
-	 * May be NULL. Set at runtime - does not exist in editor
-	 */
-	class AnimEncoding* TranslationCodec;
-
-	class AnimEncoding* RotationCodec;
-
-	class AnimEncoding* ScaleCodec;
-
-	/* Compressed curve data stream used by AnimCurveCompressionCodec */
-	TArray<uint8> CompressedCurveByteStream;
-
-	// The size of the raw data used to create the compressed data
-	int32 CompressedRawDataSize;
+	FCompressedAnimSequence CompressedData;
 
 	// Accessors for animation frame count
 	int32 GetRawNumberOfFrames() const { return NumFrames; }
 	void SetRawNumberOfFrame(int32 InNumFrames) { NumFrames = InNumFrames; }
-	int32 GetCompressedNumberOfFrames() const { return CompressedNumFrames; }
 
 	/** Additive animation type. **/
 	UPROPERTY(EditAnywhere, Category=AdditiveSettings, AssetRegistrySearchable)
@@ -786,11 +338,6 @@ public:
 	UPROPERTY(EditAnywhere, Category=AdditiveSettings)
 	int32 RefFrameIndex;
 
-	// Versioning Support
-	/** The version of the global encoding package used at the time of import */
-	UPROPERTY()
-	int32 EncodingPkgVersion;
-
 	/** Base pose to use when retargeting */
 	UPROPERTY(EditAnywhere, AssetRegistrySearchable, Category=Animation)
 	FName RetargetSource;
@@ -831,12 +378,6 @@ public:
 	UPROPERTY(EditAnywhere, Category=Compression)
 	uint32 bDoNotOverrideCompression:1;
 
-	/**
-	 * Used to track whether, or not, this sequence was compressed with it's full translation tracks
-	 */
-	UPROPERTY()
-	uint32 bWasCompressedWithoutTranslations:1;
-
 	/** Importing data and options used for this mesh */
 	UPROPERTY(VisibleAnywhere, Instanced, Category=ImportSettings)
 	class UAssetImportData* AssetImportData;
@@ -950,10 +491,7 @@ public:
 #if WITH_EDITORONLY_DATA
 	bool  HasSourceRawData() const { return SourceRawAnimationData.Num() > 0; }
 	const TArray<FName>& GetAnimationTrackNames() const { return AnimationTrackNames; }
-	const TArray<FRawAnimSequenceTrack>& GetAdditiveBaseAnimationData() const { return TemporaryAdditiveBaseAnimationData; }
-	void  UpdateCompressedTrackMapFromRaw() { CompressedTrackToSkeletonMapTable = TrackToSkeletonMapTable; }
-	void  UpdateCompressedNumFramesFromRaw() { CompressedNumFrames = NumFrames; }
-	void  UpdateCompressedCurveNames();
+	
 	void  UpdateCompressedCurveName(SmartName::UID_Type CurveUID, const struct FSmartName& NewCurveName);
 	
 	// Adds a new track (if no track of the supplied name is found) to the raw animation data, optionally setting it to TrackData.
@@ -961,8 +499,8 @@ public:
 #endif
 
 	const TArray<FTrackToSkeletonMap>& GetRawTrackToSkeletonMapTable() const { return TrackToSkeletonMapTable; }
-	const TArray<FTrackToSkeletonMap>& GetCompressedTrackToSkeletonMapTable() const { return CompressedTrackToSkeletonMapTable; }
-	const TArray<struct FSmartName>& GetCompressedCurveNames() const { return CompressedCurveNames; }
+	const TArray<FTrackToSkeletonMap>& GetCompressedTrackToSkeletonMapTable() const { return CompressedData.CompressedTrackToSkeletonMapTable; }
+	const TArray<struct FSmartName>& GetCompressedCurveNames() const { return CompressedData.CompressedCurveNames; }
 
 	FRawAnimSequenceTrack& GetRawAnimationTrack(int32 TrackIndex) { return RawAnimationData[TrackIndex]; }
 	const FRawAnimSequenceTrack& GetRawAnimationTrack(int32 TrackIndex) const { return RawAnimationData[TrackIndex]; }
@@ -1061,6 +599,11 @@ public:
 	 */
 	int32 GetApproxRawSize() const;
 
+	/**
+	 * @return		The approximate size of compressed animation data for only bones.
+	 */
+	int32 GetApproxBoneCompressedSize() const;
+	
 	/**
 	 * @return		The approximate size of compressed animation data.
 	 */
@@ -1069,12 +612,6 @@ public:
 	// Initialize curve compression settings, does nothing if scheme already valid
 	void InitCurveCompressionScheme();
 
-	/**
-	 *  Utility function for lossless compression of a FRawAnimSequenceTrack 
-	 *  @return true if keys were removed.
-	 **/
-	bool CompressRawAnimSequenceTrack(FRawAnimSequenceTrack& RawTrack, float MaxPosDiff, float MaxAngleDiff);
-
 	/**
 	 * Removes trivial frames -- frames of tracks when position or orientation is constant
 	 * over the entire animation -- from the raw animation data.  If both position and rotation
@@ -1120,7 +657,7 @@ public:
 	*/
 	int32 GetSkeletonIndexFromCompressedDataTrackIndex(const int32 TrackIndex) const
 	{
-		return CompressedTrackToSkeletonMapTable[TrackIndex].BoneTreeIndex;
+		return GetCompressedTrackToSkeletonMapTable()[TrackIndex].BoneTreeIndex;
 	}
 
 	/** Clears any data in the AnimSequence */
@@ -1131,12 +668,6 @@ public:
 	void CleanAnimSequenceForImport();
 #endif
 
-	/** 
-	 * Utility function to copy all UAnimSequence properties from Source to Destination.
-	 * Does not copy however RawAnimData, CompressedAnimData and AdditiveBasePose.
-	 */
-	static bool CopyAnimSequenceProperties(UAnimSequence* SourceAnimSeq, UAnimSequence* DestAnimSeq, bool bSkipCopyingNotifies=false);
-
 	/** 
 	 * Copy AnimNotifies from one UAnimSequence to another.
 	 */
@@ -1275,13 +806,13 @@ public:
 	bool CanBakeAdditive() const;
 
 	// Bakes out track data for the skeletons virtual bones into the raw data
-	void BakeOutVirtualBoneTracks();
+	void BakeOutVirtualBoneTracks(TArray<FRawAnimSequenceTrack>& NewRawTracks, TArray<FName>& NewAnimationTrackNames, TArray<FTrackToSkeletonMap>& NewTrackToSkeletonMapTable);
 
 	// Performs multiple evaluations of the animation as a test of compressed data validatity
 	void TestEvalauteAnimation() const;
 
 	// Bakes out the additive version of this animation into the raw data.
-	void BakeOutAdditiveIntoRawData();
+	void BakeOutAdditiveIntoRawData(TArray<FRawAnimSequenceTrack>& NewRawTracks, TArray<FName>& NewAnimationTrackNames, TArray<FTrackToSkeletonMap>& NewTrackToSkeletonMapTable, FRawCurveTracks& NewCurveTracks, TArray<FRawAnimSequenceTrack>& AdditiveBaseAnimationData);
 
 	// Test whether at any point we will scale a bone to 0 (needed for validating additive anims)
 	bool DoesSequenceContainZeroScale();
@@ -1391,54 +922,4 @@ public:
 	friend class UAnimationBlueprintLibrary;
 };
 
-struct FScopedAnimSequenceRawDataCache
-{
-	UAnimSequence* SrcAnim;
-	TArray<FRawAnimSequenceTrack> RawAnimationData;
-	TArray<FRawAnimSequenceTrack> TemporaryAdditiveBaseAnimationData;
-	TArray<FName> AnimationTrackNames;
-	TArray<FTrackToSkeletonMap> TrackToSkeletonMapTable;
-	FRawCurveTracks RawCurveData;
-	bool bWasEmpty;
-	int32 NumFrames;
-
-	FScopedAnimSequenceRawDataCache() : SrcAnim(nullptr), bWasEmpty(false) {}
-	~FScopedAnimSequenceRawDataCache()
-	{
-		if (SrcAnim)
-		{
-			RestoreTo(SrcAnim);
-		}
-	}
-
-	void InitFrom(UAnimSequence* Src)
-	{
-#if WITH_EDITORONLY_DATA
-		check(!SrcAnim);
-		SrcAnim = Src;
-		RawAnimationData = Src->RawAnimationData;
-		TemporaryAdditiveBaseAnimationData = Src->TemporaryAdditiveBaseAnimationData;
-		bWasEmpty = RawAnimationData.Num() == 0;
-		AnimationTrackNames = Src->AnimationTrackNames;
-		TrackToSkeletonMapTable = Src->TrackToSkeletonMapTable;
-		RawCurveData = Src->RawCurveData;
-		NumFrames = Src->NumFrames;
-#endif
-	}
-
-	void RestoreTo(UAnimSequence* Src)
-	{
-#if WITH_EDITORONLY_DATA
-		Src->RawAnimationData = MoveTemp(RawAnimationData);
-		Src->TemporaryAdditiveBaseAnimationData = MoveTemp(TemporaryAdditiveBaseAnimationData);
-		check(bWasEmpty || Src->RawAnimationData.Num() > 0);
-		Src->AnimationTrackNames = MoveTemp(AnimationTrackNames);
-		Src->TrackToSkeletonMapTable = MoveTemp(TrackToSkeletonMapTable);
-		Src->RawCurveData = RawCurveData;
-		Src->NumFrames = NumFrames;
-#endif
-	}
-
-};
-
 
diff --git a/Engine/Source/Runtime/Engine/Classes/Animation/Skeleton.h b/Engine/Source/Runtime/Engine/Classes/Animation/Skeleton.h
--- a/Engine/Source/Runtime/Engine/Classes/Animation/Skeleton.h
+++ b/Engine/Source/Runtime/Engine/Classes/Animation/Skeleton.h
@@ -711,7 +711,7 @@ public:
 	 *
 	 * @return	Index of Track of Animation Sequence
 	 */
-	ENGINE_API int32 GetAnimationTrackIndex(const int32 InSkeletonBoneIndex, const UAnimSequence* InAnimSeq, const bool bUseRawData);
+	ENGINE_API int32 GetRawAnimationTrackIndex(const int32 InSkeletonBoneIndex, const UAnimSequence* InAnimSeq);
 
 	/** 
 	 * Get Bone Tree Index from Reference Bone Index
diff --git a/Engine/Source/Runtime/Engine/Private/Animation/AnimCompress.cpp b/Engine/Source/Runtime/Engine/Private/Animation/AnimCompress.cpp
--- a/Engine/Source/Runtime/Engine/Private/Animation/AnimCompress.cpp
+++ b/Engine/Source/Runtime/Engine/Private/Animation/AnimCompress.cpp
@@ -28,33 +28,32 @@ UAnimCompress::UAnimCompress(const FObjectInitializer& ObjectInitializer)
 	MaxNumFramesPerSegment = (IdealNumFramesPerSegment * 2) - 1;
 }
 
-void FCompressionMemorySummary::GatherPostCompressionStats(UAnimSequence* Seq, const TArray<FBoneData>& BoneData, double CompressionTime)
+void FCompressionMemorySummary::GatherPostCompressionStats(const FCompressibleAnimData& CompressibleAnimData, FCompressibleAnimDataResult& CompressedData, double CompressionTime)
 {
+#if WITH_EDITOR
 	if (bEnabled)
 	{
-		TotalAfterCompressed += Seq->GetApproxCompressedSize();
+		TotalAfterCompressed += CompressedData.GetApproxBoneCompressedSize();
 		TotalCompressionExecutionTime += CompressionTime;
 
-		if (Seq->GetSkeleton() != NULL)
-		{
-			// determine the error added by the compression
-			AnimationErrorStats ErrorStats;
-			FAnimationUtils::ComputeCompressionError(Seq, BoneData, ErrorStats);
+		// determine the error added by the compression
+		AnimationErrorStats ErrorStats;
+		FAnimationUtils::ComputeCompressionError(CompressibleAnimData, CompressedData, ErrorStats);
 
-			ErrorTotal += ErrorStats.AverageError;
-			ErrorCount += 1.0f;
-			AverageError = ErrorTotal / ErrorCount;
+		ErrorTotal += ErrorStats.AverageError;
+		ErrorCount += 1.0f;
+		AverageError = ErrorTotal / ErrorCount;
 
-			WorstBoneError.StoreErrorStat(ErrorStats.MaxError, ErrorStats.MaxError, ErrorStats.MaxErrorTime, ErrorStats.MaxErrorBone, BoneData[ErrorStats.MaxErrorBone].Name, Seq->GetFName());
+		WorstBoneError.StoreErrorStat(ErrorStats.MaxError, ErrorStats.MaxError, ErrorStats.MaxErrorTime, ErrorStats.MaxErrorBone, CompressibleAnimData.BoneData[ErrorStats.MaxErrorBone].Name, CompressibleAnimData.AnimFName);
 
-			WorstAnimationError.StoreErrorStat(ErrorStats.AverageError, ErrorStats.AverageError, Seq->GetFName());
-		}
+		WorstAnimationError.StoreErrorStat(ErrorStats.AverageError, ErrorStats.AverageError, CompressibleAnimData.AnimFName);
 	}
+#endif
 }
 
-void FAnimCompressContext::GatherPostCompressionStats(UAnimSequence* Seq, const TArray<FBoneData>& BoneData, double CompressionTime)
+void FAnimCompressContext::GatherPostCompressionStats(const FCompressibleAnimData& CompressibleAnimData, FCompressibleAnimDataResult& CompressedData, double CompressionTime)
 {
-	CompressionSummary.GatherPostCompressionStats(Seq, BoneData, CompressionTime);
+	CompressionSummary.GatherPostCompressionStats(CompressibleAnimData, CompressedData, CompressionTime);
 }
 
 FCompressionMemorySummary::FCompressionMemorySummary(bool bInEnabled)
@@ -204,13 +203,13 @@ uint8 MakeBitForFlag(uint32 Item, uint32 Position)
 // FCompressionMemorySummary
 
 
-void FCompressionMemorySummary::GatherPreCompressionStats(UAnimSequence* Seq, int32 ProgressNumerator, int32 ProgressDenominator)
+void FCompressionMemorySummary::GatherPreCompressionStats(const FCompressibleAnimData& CompressibleAnimData, int32 PreCompressedSize, int32 ProgressNumerator, int32 ProgressDenominator)
 {
 	if (bEnabled)
 	{
 		bUsed = true;
 		FFormatNamedArguments Args;
-		Args.Add(TEXT("AnimSequenceName"), FText::FromString(Seq->GetName()));
+		Args.Add(TEXT("AnimSequenceName"), FText::FromString(CompressibleAnimData.Name));
 		Args.Add(TEXT("ProgressNumerator"), ProgressNumerator);
 		Args.Add(TEXT("ProgressDenominator"), ProgressDenominator);
 
@@ -218,17 +217,17 @@ void FCompressionMemorySummary::GatherPreCompressionStats(UAnimSequence* Seq, in
 			ProgressDenominator,
 			FText::Format(NSLOCTEXT("CompressionMemorySummary", "CompressingTaskStatusMessageFormat", "Compressing {AnimSequenceName} ({ProgressNumerator}/{ProgressDenominator})"), Args));
 
-		TotalRaw += Seq->GetApproxRawSize();
-		TotalBeforeCompressed += Seq->GetApproxCompressedSize();
+		TotalRaw += CompressibleAnimData.GetApproxRawBoneSize();
+		TotalBeforeCompressed += PreCompressedSize;
 		++NumberOfAnimations;
 	}
 }
 
 
 //////////////////////////////////////////////////////////////////////////////////////
-void FAnimCompressContext::GatherPreCompressionStats(UAnimSequence* Seq)
+void FAnimCompressContext::GatherPreCompressionStats(const FCompressibleAnimData& CompressibleAnimData, int32 PreviousCompressionSize)
 {
-	CompressionSummary.GatherPreCompressionStats(Seq, AnimIndex, MaxAnimations);
+	CompressionSummary.GatherPreCompressionStats(CompressibleAnimData, PreviousCompressionSize, AnimIndex, MaxAnimations);
 }
 
 
@@ -270,7 +269,8 @@ void UAnimCompress::PadByteStream(TArray<uint8>& ByteStream, const int32 Alignme
 
 
 void UAnimCompress::BitwiseCompressAnimationTracks(
-	class UAnimSequence* Seq,
+	const FCompressibleAnimData& CompressibleAnimData,
+	FCompressibleAnimDataResult& OutCompressedData,
 	AnimationCompressionFormat TargetTranslationFormat,
 	AnimationCompressionFormat TargetRotationFormat,
 	AnimationCompressionFormat TargetScaleFormat,
@@ -298,19 +298,18 @@ void UAnimCompress::BitwiseCompressAnimationTracks(
 	}
 	if (bInvalidCompressionFormat)
 	{
-		Seq->TranslationCompressionFormat = ACF_None;
-		Seq->RotationCompressionFormat = ACF_None;
-		Seq->ScaleCompressionFormat = ACF_None;
-		Seq->CompressedTrackOffsets.Empty();
-		Seq->CompressedScaleOffsets.Empty();
-		Seq->CompressedByteStream.Empty();
-		Seq->CompressedSegments.Empty();
+		OutCompressedData.TranslationCompressionFormat = ACF_None;
+		OutCompressedData.RotationCompressionFormat = ACF_None;
+		OutCompressedData.ScaleCompressionFormat = ACF_None;
+		OutCompressedData.CompressedTrackOffsets.Empty();
+		OutCompressedData.CompressedScaleOffsets.Empty();
+		OutCompressedData.CompressedByteStream.Empty();
 	}
 	else
 	{
-		Seq->RotationCompressionFormat = TargetRotationFormat;
-		Seq->TranslationCompressionFormat = TargetTranslationFormat;
-		Seq->ScaleCompressionFormat = TargetScaleFormat;
+		OutCompressedData.RotationCompressionFormat = TargetRotationFormat;
+		OutCompressedData.TranslationCompressionFormat = TargetTranslationFormat;
+		OutCompressedData.ScaleCompressionFormat = TargetScaleFormat;
 
 		check(TranslationData.Num() == RotationData.Num());
 		const int32 NumTracks = RotationData.Num();
@@ -318,41 +317,40 @@ void UAnimCompress::BitwiseCompressAnimationTracks(
 
 		if (NumTracks == 0)
 		{
-			UE_LOG(LogAnimationCompression, Warning, TEXT("When compressing %s: no key-reduced data"), *Seq->GetName());
+			UE_LOG(LogAnimationCompression, Warning, TEXT("When compressing %s: no key-reduced data"), *CompressibleAnimData.Name);
 		}
 
-		Seq->CompressedTrackOffsets.Empty(NumTracks * 4);
-		Seq->CompressedTrackOffsets.AddUninitialized(NumTracks * 4);
+		OutCompressedData.CompressedTrackOffsets.Empty(NumTracks * 4);
+		OutCompressedData.CompressedTrackOffsets.AddUninitialized(NumTracks * 4);
 
 		// just empty it since there is chance this can be 0
-		Seq->CompressedScaleOffsets.Empty();
+		OutCompressedData.CompressedScaleOffsets.Empty();
 		// only do this if Scale exists;
 		if (bHasScale)
 		{
-			Seq->CompressedScaleOffsets.SetStripSize(2);
-			Seq->CompressedScaleOffsets.AddUninitialized(NumTracks);
+			OutCompressedData.CompressedScaleOffsets.SetStripSize(2);
+			OutCompressedData.CompressedScaleOffsets.AddUninitialized(NumTracks);
 		}
 
-		Seq->CompressedByteStream.Empty();
-		Seq->CompressedSegments.Empty();
+		OutCompressedData.CompressedByteStream.Empty();
 
 		for (int32 TrackIndex = 0; TrackIndex < NumTracks; ++TrackIndex)
 		{
 			// Translation data.
 			const FTranslationTrack& SrcTrans = TranslationData[TrackIndex];
 
-			const int32 OffsetTrans = Seq->CompressedByteStream.Num();
+			const int32 OffsetTrans = OutCompressedData.CompressedByteStream.Num();
 			const int32 NumKeysTrans = SrcTrans.PosKeys.Num();
 
 			// Warn on empty data.
 			if (NumKeysTrans == 0)
 			{
-				UE_LOG(LogAnimationCompression, Warning, TEXT("When compressing %s track %i: no translation keys"), *Seq->GetName(), TrackIndex);
+				UE_LOG(LogAnimationCompression, Warning, TEXT("When compressing %s track %i: no translation keys"), *CompressibleAnimData.Name, TrackIndex);
 			}
 
 			checkf((OffsetTrans % 4) == 0, TEXT("CompressedByteStream not aligned to four bytes"));
-			Seq->CompressedTrackOffsets[TrackIndex * 4] = OffsetTrans;
-			Seq->CompressedTrackOffsets[TrackIndex * 4 + 1] = NumKeysTrans;
+			OutCompressedData.CompressedTrackOffsets[TrackIndex * 4] = OffsetTrans;
+			OutCompressedData.CompressedTrackOffsets[TrackIndex * 4 + 1] = NumKeysTrans;
 
 			// Calculate the bounding box of the translation keys
 			FBox PositionBounds(SrcTrans.PosKeys);
@@ -368,31 +366,31 @@ void UAnimCompress::BitwiseCompressAnimationTracks(
 				// Write the mins and ranges if they'll be used on the other side
 				if (TargetTranslationFormat == ACF_IntervalFixed32NoW)
 				{
-					UnalignedWriteToStream(Seq->CompressedByteStream, TransMins, sizeof(float) * 3);
-					UnalignedWriteToStream(Seq->CompressedByteStream, TransRanges, sizeof(float) * 3);
+					UnalignedWriteToStream(OutCompressedData.CompressedByteStream, TransMins, sizeof(float) * 3);
+					UnalignedWriteToStream(OutCompressedData.CompressedByteStream, TransRanges, sizeof(float) * 3);
 				}
 
 				// Pack the positions into the stream
 				for (int32 KeyIndex = 0; KeyIndex < NumKeysTrans; ++KeyIndex)
 				{
 					const FVector& Vec = SrcTrans.PosKeys[KeyIndex];
-					PackVectorToStream(Seq->CompressedByteStream, TargetTranslationFormat, Vec, TransMins, TransRanges);
+					PackVectorToStream(OutCompressedData.CompressedByteStream, TargetTranslationFormat, Vec, TransMins, TransRanges);
 				}
 
 				if (IncludeKeyTable)
 				{
 					// Align to four bytes.
-					PadByteStream(Seq->CompressedByteStream, 4, AnimationPadSentinel);
+					PadByteStream(OutCompressedData.CompressedByteStream, 4, AnimationPadSentinel);
 
 					// write the key table
-					const int32 NumFrames = Seq->GetRawNumberOfFrames();
-					const int32 LastFrame = Seq->GetRawNumberOfFrames()-1;
-					const size_t FrameSize = Seq->GetRawNumberOfFrames() > 0xff ? sizeof(uint16) : sizeof(uint8);
-					const float FrameRate = LastFrame / Seq->SequenceLength;
+					const int32 NumFrames = CompressibleAnimData.NumFrames;
+					const int32 LastFrame = NumFrames-1;
+					const size_t FrameSize = NumFrames > 0xff ? sizeof(uint16) : sizeof(uint8);
+					const float FrameRate = LastFrame / CompressibleAnimData.SequenceLength;
 
 					const int32 TableSize = NumKeysTrans*FrameSize;
 					const int32 TableDwords = (TableSize + 3) >> 2;
-					const int32 StartingOffset = Seq->CompressedByteStream.Num();
+					const int32 StartingOffset = OutCompressedData.CompressedByteStream.Num();
 
 					for (int32 KeyIndex = 0; KeyIndex < NumKeysTrans; ++KeyIndex)
 					{
@@ -400,37 +398,37 @@ void UAnimCompress::BitwiseCompressAnimationTracks(
 						float KeyTime = SrcTrans.Times[KeyIndex];
 						float FrameTime = KeyTime * FrameRate;
 						int32 FrameIndex = FMath::Clamp(FMath::TruncToInt(FrameTime + 0.5f), 0, LastFrame);
-						UnalignedWriteToStream(Seq->CompressedByteStream, &FrameIndex, FrameSize);
+						UnalignedWriteToStream(OutCompressedData.CompressedByteStream, &FrameIndex, FrameSize);
 					}
 
 					// Align to four bytes. Padding with 0's to round out the key table
-					PadByteStream(Seq->CompressedByteStream, 4, 0);
+					PadByteStream(OutCompressedData.CompressedByteStream, 4, 0);
 
-					const int32 EndingOffset = Seq->CompressedByteStream.Num();
+					const int32 EndingOffset = OutCompressedData.CompressedByteStream.Num();
 					check((EndingOffset - StartingOffset) == (TableDwords * 4));
 				}
 			}
 			else if (NumKeysTrans == 1)
 			{
 				// A single translation key gets written out a single uncompressed float[3].
-				UnalignedWriteToStream(Seq->CompressedByteStream, &(SrcTrans.PosKeys[0]), sizeof(FVector));
+				UnalignedWriteToStream(OutCompressedData.CompressedByteStream, &(SrcTrans.PosKeys[0]), sizeof(FVector));
 			}
 			else
 			{
-				UE_LOG(LogAnimationCompression, Warning, TEXT("When compressing %s track %i: no translation keys"), *Seq->GetName(), TrackIndex);
+				UE_LOG(LogAnimationCompression, Warning, TEXT("When compressing %s track %i: no translation keys"), *CompressibleAnimData.Name, TrackIndex);
 			}
 
 			// Align to four bytes.
-			PadByteStream(Seq->CompressedByteStream, 4, AnimationPadSentinel);
+			PadByteStream(OutCompressedData.CompressedByteStream, 4, AnimationPadSentinel);
 
 			// Compress rotation data.
 			const FRotationTrack& SrcRot = RotationData[TrackIndex];
-			const int32 OffsetRot = Seq->CompressedByteStream.Num();
+			const int32 OffsetRot = OutCompressedData.CompressedByteStream.Num();
 			const int32 NumKeysRot = SrcRot.RotKeys.Num();
 
 			checkf((OffsetRot % 4) == 0, TEXT("CompressedByteStream not aligned to four bytes"));
-			Seq->CompressedTrackOffsets[TrackIndex * 4 + 2] = OffsetRot;
-			Seq->CompressedTrackOffsets[TrackIndex * 4 + 3] = NumKeysRot;
+			OutCompressedData.CompressedTrackOffsets[TrackIndex * 4 + 2] = OffsetRot;
+			OutCompressedData.CompressedTrackOffsets[TrackIndex * 4 + 3] = NumKeysRot;
 
 			if (NumKeysRot > 1)
 			{
@@ -469,32 +467,32 @@ void UAnimCompress::BitwiseCompressAnimationTracks(
 				// Write the mins and ranges if they'll be used on the other side
 				if (TargetRotationFormat == ACF_IntervalFixed32NoW)
 				{
-					UnalignedWriteToStream(Seq->CompressedByteStream, Mins, sizeof(float) * 3);
-					UnalignedWriteToStream(Seq->CompressedByteStream, Ranges, sizeof(float) * 3);
+					UnalignedWriteToStream(OutCompressedData.CompressedByteStream, Mins, sizeof(float) * 3);
+					UnalignedWriteToStream(OutCompressedData.CompressedByteStream, Ranges, sizeof(float) * 3);
 				}
 
 				// n elements of the compressed type.
 				for (int32 KeyIndex = 0; KeyIndex < SrcRot.RotKeys.Num(); ++KeyIndex)
 				{
 					const FQuat& Quat = SrcRot.RotKeys[KeyIndex];
-					PackQuaternionToStream(Seq->CompressedByteStream, TargetRotationFormat, Quat, Mins, Ranges);
+					PackQuaternionToStream(OutCompressedData.CompressedByteStream, TargetRotationFormat, Quat, Mins, Ranges);
 				}
 
 				// n elements of frame indices
 				if (IncludeKeyTable)
 				{
 					// Align to four bytes.
-					PadByteStream(Seq->CompressedByteStream, 4, AnimationPadSentinel);
+					PadByteStream(OutCompressedData.CompressedByteStream, 4, AnimationPadSentinel);
 
 					// write the key table
-					const int32 NumFrames = Seq->GetRawNumberOfFrames();
-					const int32 LastFrame= Seq->GetRawNumberOfFrames()-1;
-					const size_t FrameSize= Seq->GetRawNumberOfFrames() > 0xff ? sizeof(uint16) : sizeof(uint8);
-					const float FrameRate = LastFrame / Seq->SequenceLength;
+					const int32 NumFrames = CompressibleAnimData.NumFrames;
+					const int32 LastFrame= NumFrames-1;
+					const size_t FrameSize= NumFrames > 0xff ? sizeof(uint16) : sizeof(uint8);
+					const float FrameRate = LastFrame / CompressibleAnimData.SequenceLength;
 
 					const int32 TableSize = NumKeysRot*FrameSize;
 					const int32 TableDwords = (TableSize + 3) >> 2;
-					const int32 StartingOffset = Seq->CompressedByteStream.Num();
+					const int32 StartingOffset = OutCompressedData.CompressedByteStream.Num();
 
 					for (int32 KeyIndex = 0; KeyIndex < NumKeysRot; ++KeyIndex)
 					{
@@ -502,13 +500,13 @@ void UAnimCompress::BitwiseCompressAnimationTracks(
 						float KeyTime = SrcRot.Times[KeyIndex];
 						float FrameTime = KeyTime * FrameRate;
 						int32 FrameIndex = FMath::Clamp(FMath::TruncToInt(FrameTime + 0.5f), 0, LastFrame);
-						UnalignedWriteToStream(Seq->CompressedByteStream, &FrameIndex, FrameSize);
+						UnalignedWriteToStream(OutCompressedData.CompressedByteStream, &FrameIndex, FrameSize);
 					}
 
 					// Align to four bytes. Padding with 0's to round out the key table
-					PadByteStream(Seq->CompressedByteStream, 4, 0);
+					PadByteStream(OutCompressedData.CompressedByteStream, 4, 0);
 
-					const int32 EndingOffset = Seq->CompressedByteStream.Num();
+					const int32 EndingOffset = OutCompressedData.CompressedByteStream.Num();
 					check((EndingOffset - StartingOffset) == (TableDwords * 4));
 
 				}
@@ -518,28 +516,28 @@ void UAnimCompress::BitwiseCompressAnimationTracks(
 				// For a rotation track of n=1 keys, the single key is packed as an FQuatFloat96NoW.
 				const FQuat& Quat = SrcRot.RotKeys[0];
 				const FQuatFloat96NoW QuatFloat96NoW(Quat);
-				UnalignedWriteToStream(Seq->CompressedByteStream, &QuatFloat96NoW, sizeof(FQuatFloat96NoW));
+				UnalignedWriteToStream(OutCompressedData.CompressedByteStream, &QuatFloat96NoW, sizeof(FQuatFloat96NoW));
 			}
 			else
 			{
-				UE_LOG(LogAnimationCompression, Warning, TEXT("When compressing %s track %i: no rotation keys"), *Seq->GetName(), TrackIndex);
+				UE_LOG(LogAnimationCompression, Warning, TEXT("When compressing %s track %i: no rotation keys"), *CompressibleAnimData.Name, TrackIndex);
 			}
 
 
 			// Align to four bytes.
-			PadByteStream(Seq->CompressedByteStream, 4, AnimationPadSentinel);
+			PadByteStream(OutCompressedData.CompressedByteStream, 4, AnimationPadSentinel);
 
 			// we also should do this only when scale exists. 
 			if (bHasScale)
 			{
 				const FScaleTrack& SrcScale = ScaleData[TrackIndex];
 
-				const int32 OffsetScale = Seq->CompressedByteStream.Num();
+				const int32 OffsetScale = OutCompressedData.CompressedByteStream.Num();
 				const int32 NumKeysScale = SrcScale.ScaleKeys.Num();
 
 				checkf((OffsetScale % 4) == 0, TEXT("CompressedByteStream not aligned to four bytes"));
-				Seq->CompressedScaleOffsets.SetOffsetData(TrackIndex, 0, OffsetScale);
-				Seq->CompressedScaleOffsets.SetOffsetData(TrackIndex, 1, NumKeysScale);
+				OutCompressedData.CompressedScaleOffsets.SetOffsetData(TrackIndex, 0, OffsetScale);
+				OutCompressedData.CompressedScaleOffsets.SetOffsetData(TrackIndex, 1, NumKeysScale);
 
 				// Calculate the bounding box of the Scalelation keys
 				FBox ScaleBoundsBounds(SrcScale.ScaleKeys);
@@ -556,31 +554,31 @@ void UAnimCompress::BitwiseCompressAnimationTracks(
 					// Write the mins and ranges if they'll be used on the other side
 					if (TargetScaleFormat == ACF_IntervalFixed32NoW)
 					{
-						UnalignedWriteToStream(Seq->CompressedByteStream, ScaleMins, sizeof(float) * 3);
-						UnalignedWriteToStream(Seq->CompressedByteStream, ScaleRanges, sizeof(float) * 3);
+						UnalignedWriteToStream(OutCompressedData.CompressedByteStream, ScaleMins, sizeof(float) * 3);
+						UnalignedWriteToStream(OutCompressedData.CompressedByteStream, ScaleRanges, sizeof(float) * 3);
 					}
 
 					// Pack the positions into the stream
 					for (int32 KeyIndex = 0; KeyIndex < NumKeysScale; ++KeyIndex)
 					{
 						const FVector& Vec = SrcScale.ScaleKeys[KeyIndex];
-						PackVectorToStream(Seq->CompressedByteStream, TargetScaleFormat, Vec, ScaleMins, ScaleRanges);
+						PackVectorToStream(OutCompressedData.CompressedByteStream, TargetScaleFormat, Vec, ScaleMins, ScaleRanges);
 					}
 
 					if (IncludeKeyTable)
 					{
 						// Align to four bytes.
-						PadByteStream(Seq->CompressedByteStream, 4, AnimationPadSentinel);
+						PadByteStream(OutCompressedData.CompressedByteStream, 4, AnimationPadSentinel);
 
 						// write the key table
-						const int32 NumFrames = Seq->GetRawNumberOfFrames();
-						const int32 LastFrame = Seq->GetRawNumberOfFrames()-1;
-						const size_t FrameSize = Seq->GetRawNumberOfFrames() > 0xff ? sizeof(uint16) : sizeof(uint8);
-						const float FrameRate = LastFrame / Seq->SequenceLength;
+						const int32 NumFrames = CompressibleAnimData.NumFrames;
+						const int32 LastFrame = NumFrames-1;
+						const size_t FrameSize = NumFrames > 0xff ? sizeof(uint16) : sizeof(uint8);
+						const float FrameRate = LastFrame / CompressibleAnimData.SequenceLength;
 
 						const int32 TableSize = NumKeysScale*FrameSize;
 						const int32 TableDwords = (TableSize + 3) >> 2;
-						const int32 StartingOffset = Seq->CompressedByteStream.Num();
+						const int32 StartingOffset = OutCompressedData.CompressedByteStream.Num();
 
 						for (int32 KeyIndex = 0; KeyIndex < NumKeysScale; ++KeyIndex)
 						{
@@ -588,36 +586,37 @@ void UAnimCompress::BitwiseCompressAnimationTracks(
 							float KeyTime = SrcScale.Times[KeyIndex];
 							float FrameTime = KeyTime * FrameRate;
 							int32 FrameIndex = FMath::Clamp(FMath::TruncToInt(FrameTime + 0.5f), 0, LastFrame);
-							UnalignedWriteToStream(Seq->CompressedByteStream, &FrameIndex, FrameSize);
+							UnalignedWriteToStream(OutCompressedData.CompressedByteStream, &FrameIndex, FrameSize);
 						}
 
 						// Align to four bytes. Padding with 0's to round out the key table
-						PadByteStream(Seq->CompressedByteStream, 4, 0);
+						PadByteStream(OutCompressedData.CompressedByteStream, 4, 0);
 
-						const int32 EndingOffset = Seq->CompressedByteStream.Num();
+						const int32 EndingOffset = OutCompressedData.CompressedByteStream.Num();
 						check((EndingOffset - StartingOffset) == (TableDwords * 4));
 					}
 				}
 				else if (NumKeysScale == 1)
 				{
 					// A single Scalelation key gets written out a single uncompressed float[3].
-					UnalignedWriteToStream(Seq->CompressedByteStream, &(SrcScale.ScaleKeys[0]), sizeof(FVector));
+					UnalignedWriteToStream(OutCompressedData.CompressedByteStream, &(SrcScale.ScaleKeys[0]), sizeof(FVector));
 				}
 				else
 				{
-					UE_LOG(LogAnimationCompression, Warning, TEXT("When compressing %s track %i: no Scale keys"), *Seq->GetName(), TrackIndex);
+					UE_LOG(LogAnimationCompression, Warning, TEXT("When compressing %s track %i: no Scale keys"), *CompressibleAnimData.Name, TrackIndex);
 				}
 
 				// Align to four bytes.
-				PadByteStream(Seq->CompressedByteStream, 4, AnimationPadSentinel);
+				PadByteStream(OutCompressedData.CompressedByteStream, 4, AnimationPadSentinel);
 			}
 		}
 
 		// Trim unused memory.
-		Seq->CompressedByteStream.Shrink();
+		OutCompressedData.CompressedByteStream.Shrink();
 	}
 }
 
+#if USE_SEGMENTING_CONTEXT
 void UAnimCompress::BitwiseCompressAnimationTracks(
 	class UAnimSequence& AnimSeq,
 	AnimationCompressionFormat TargetTranslationFormat,
@@ -705,6 +704,7 @@ void UAnimCompress::SanityCheckTrackData(const UAnimSequence& AnimSeq, const FAn
 		}
 	}
 }
+#endif
 
 void UAnimCompress::CalculateTrackRange(const FTranslationTrack& TranslationData, AnimationCompressionFormat Format, FVector& OutMin, FVector& OutExtent)
 {
@@ -806,6 +806,7 @@ void UAnimCompress::CalculateTrackRange(const FScaleTrack& ScaleData, AnimationC
 	OutExtent = Ranges;
 }
 
+#if USE_SEGMENTING_CONTEXT
 void UAnimCompress::CalculateTrackRanges(
 	AnimationCompressionFormat TargetTranslationFormat,
 	AnimationCompressionFormat TargetRotationFormat,
@@ -1576,6 +1577,7 @@ void UAnimCompress::CoalesceCompressedSegments(UAnimSequence& AnimSeq, const TAr
 	Header.SequenceCRC = FCrc::MemCrc32(AnimSeq.CompressedByteStream.GetData(), AnimSeq.CompressedByteStream.Num());
 	memcpy(AnimSeq.CompressedByteStream.GetData(), &Header, sizeof(FAnimSequenceCompressionHeader));
 }
+#endif
 
 #if WITH_EDITOR
 
@@ -1621,42 +1623,24 @@ void UAnimCompress::PopulateDDCKey(FArchive& Ar)
  * Tracks
  */
 
-bool UAnimCompress::Reduce(UAnimSequence* AnimSeq, bool bOutput, const TArray<FBoneData>& BoneData)
+bool UAnimCompress::Reduce(const FCompressibleAnimData& CompressibleAnimData, FCompressibleAnimDataResult& OutResult)
 {
-	bool bResult = false;
 #if WITH_EDITORONLY_DATA
-	USkeleton* AnimSkeleton = AnimSeq->GetSkeleton();
-	const bool bSkeletonExistsIfNeeded = ( AnimSkeleton || !bNeedsSkeleton);
+	const bool bSkeletonExistsIfNeeded = ( CompressibleAnimData.Skeleton || !bNeedsSkeleton);
 	if ( bSkeletonExistsIfNeeded )
 	{
-		FAnimCompressContext CompressContext(false, bOutput);
-		Reduce(AnimSeq, CompressContext, BoneData);
+		OutResult.CompressedNumberOfFrames = CompressibleAnimData.NumFrames;
 
-		bResult = true;
+		DoReduction(CompressibleAnimData, OutResult);
+		
+		OutResult.CompressionScheme = this;
+		return true;
 	}
 #endif // WITH_EDITORONLY_DATA
 
-	return bResult;
+	return false;
 }
 
-bool UAnimCompress::Reduce(class UAnimSequence* AnimSeq, FAnimCompressContext& Context, const TArray<FBoneData>& BoneData)
-{
-	bool bResult = false;
-
-#if WITH_EDITORONLY_DATA
-	
-	DoReduction(AnimSeq, BoneData);
-
-	AnimSeq->bWasCompressedWithoutTranslations = false; // @fixmelh : bAnimRotationOnly
-
-	AnimSeq->EncodingPkgVersion = CURRENT_ANIMATION_ENCODING_PACKAGE_VERSION;
-	AnimSeq->MarkPackageDirty();
-
-	bResult = true;
-#endif // WITH_EDITORONLY_DATA
-
-	return bResult;
-}
 #endif // WITH_EDITOR
 
 void UAnimCompress::FilterTrivialPositionKeys(
@@ -1819,7 +1803,7 @@ void UAnimCompress::FilterTrivialKeys(
 	FilterTrivialScaleKeys(ScaleTracks, MaxScaleDelta);
 }
 
-
+#if USE_SEGMENTING_CONTEXT
 void UAnimCompress::FilterTrivialKeys(
 	TArray<FAnimSegmentContext>& RawSegments,
 	float MaxPosDelta,
@@ -1833,7 +1817,7 @@ void UAnimCompress::FilterTrivialKeys(
 		FilterTrivialScaleKeys(Segment.ScaleData, MaxScaleDelta);
 	}
 }
-
+#endif
 
 void UAnimCompress::FilterIntermittentPositionKeys(
 	FTranslationTrack& Track, 
@@ -2058,6 +2042,7 @@ void UAnimCompress::SeparateRawDataIntoTracks(
 	}
 }
 
+#if USE_SEGMENTING_CONTEXT
 static void GenerateAnimSequenceSegments(const UAnimSequence& AnimSeq, int32 NumFrames, int32 IdealNumFramesPerSegment, int32 MaxNumFramesPerSegment, TArray<FAnimSegmentContext>& OutRawSegments)
 {
 	int32 NumSegments;
@@ -2118,6 +2103,7 @@ static void GenerateAnimSequenceSegments(const UAnimSequence& AnimSeq, int32 Num
 		NumPreviousFrames += NumFramesInSegment;
 	}
 }
+#endif
 
 static int32 GetNumFrames(const UAnimSequence& AnimSeq, const TArray<FTranslationTrack>& TranslationData, const TArray<FRotationTrack>& RotationData, const TArray<FScaleTrack>& ScaleData)
 {
@@ -2149,6 +2135,7 @@ static int32 GetNumFrames(const UAnimSequence& AnimSeq, const TArray<FTranslatio
 	return AnimSeq.GetRawNumberOfFrames();
 }
 
+#if USE_SEGMENTING_CONTEXT
 void UAnimCompress::SeparateRawDataIntoTracks(
 	const UAnimSequence& AnimSeq,
 	const TArray<FTranslationTrack>& TranslationData,
@@ -2283,4 +2270,6 @@ void UAnimCompress::SeparateRawDataIntoTracks(
 		}
 	}
 }
+#endif
+
 #endif
\ No newline at end of file
diff --git a/Engine/Source/Runtime/Engine/Private/Animation/AnimCompress_Automatic.cpp b/Engine/Source/Runtime/Engine/Private/Animation/AnimCompress_Automatic.cpp
--- a/Engine/Source/Runtime/Engine/Private/Animation/AnimCompress_Automatic.cpp
+++ b/Engine/Source/Runtime/Engine/Private/Animation/AnimCompress_Automatic.cpp
@@ -17,12 +17,13 @@ UAnimCompress_Automatic::UAnimCompress_Automatic(const FObjectInitializer& Objec
 }
 
 #if WITH_EDITOR
-void UAnimCompress_Automatic::DoReduction(UAnimSequence* AnimSeq, const TArray<FBoneData>& BoneData)
+void UAnimCompress_Automatic::DoReduction(const FCompressibleAnimData& CompressibleAnimData, FCompressibleAnimDataResult& OutResult)
 {
 	FAnimCompressContext CompressContext(MaxEndEffectorError > 0.0f, false);
 #if WITH_EDITORONLY_DATA
 	FAnimationUtils::CompressAnimSequenceExplicit(
-		AnimSeq,
+		CompressibleAnimData,
+		OutResult,
 		CompressContext,
 		MaxEndEffectorError,
 		bRunCurrentDefaultCompressor,
@@ -31,8 +32,7 @@ void UAnimCompress_Automatic::DoReduction(UAnimSequence* AnimSeq, const TArray<F
 		bTryExhaustiveSearch,
 		bEnableSegmenting,
 		IdealNumFramesPerSegment,
-		MaxNumFramesPerSegment,
-		BoneData);
+		MaxNumFramesPerSegment);
 #endif // WITH_EDITORONLY_DATA
 }
 
diff --git a/Engine/Source/Runtime/Engine/Private/Animation/AnimCompress_BitwiseCompressOnly.cpp b/Engine/Source/Runtime/Engine/Private/Animation/AnimCompress_BitwiseCompressOnly.cpp
--- a/Engine/Source/Runtime/Engine/Private/Animation/AnimCompress_BitwiseCompressOnly.cpp
+++ b/Engine/Source/Runtime/Engine/Private/Animation/AnimCompress_BitwiseCompressOnly.cpp
@@ -15,24 +15,24 @@ UAnimCompress_BitwiseCompressOnly::UAnimCompress_BitwiseCompressOnly(const FObje
 }
 
 #if WITH_EDITOR
-void UAnimCompress_BitwiseCompressOnly::DoReduction(UAnimSequence* AnimSeq, const TArray<FBoneData>& BoneData)
+void UAnimCompress_BitwiseCompressOnly::DoReduction(const FCompressibleAnimData& CompressibleAnimData, FCompressibleAnimDataResult& OutResult)
 {
 #if WITH_EDITORONLY_DATA
 	// split the raw data into tracks
 	TArray<FTranslationTrack> TranslationData;
 	TArray<FRotationTrack> RotationData;
 	TArray<FScaleTrack> ScaleData;
-	SeparateRawDataIntoTracks( AnimSeq->GetRawAnimationData(), AnimSeq->SequenceLength, TranslationData, RotationData, ScaleData );
+	SeparateRawDataIntoTracks( CompressibleAnimData.RawAnimationData, CompressibleAnimData.SequenceLength, TranslationData, RotationData, ScaleData );
 
 	// remove obviously redundant keys from the source data
 	FilterTrivialKeys(TranslationData, RotationData, ScaleData, TRANSLATION_ZEROING_THRESHOLD, QUATERNION_ZEROING_THRESHOLD, SCALE_ZEROING_THRESHOLD);
 
 	// record the proper runtime decompressor to use
-	AnimSeq->KeyEncodingFormat = AKF_ConstantKeyLerp;
-	AnimSeq->RotationCompressionFormat = RotationCompressionFormat;
-	AnimSeq->TranslationCompressionFormat = TranslationCompressionFormat;
-	AnimSeq->ScaleCompressionFormat = ScaleCompressionFormat;
-	AnimationFormat_SetInterfaceLinks(*AnimSeq);
+	OutResult.KeyEncodingFormat = AKF_ConstantKeyLerp;
+	OutResult.RotationCompressionFormat = RotationCompressionFormat;
+	OutResult.TranslationCompressionFormat = TranslationCompressionFormat;
+	OutResult.ScaleCompressionFormat = ScaleCompressionFormat;
+	AnimationFormat_SetInterfaceLinks(OutResult);
 
 #if USE_SEGMENTING_CONTEXT
 	if (bEnableSegmenting)
@@ -58,7 +58,8 @@ void UAnimCompress_BitwiseCompressOnly::DoReduction(UAnimSequence* AnimSeq, cons
 	{
 		// bitwise compress the tracks into the anim sequence buffers
 		BitwiseCompressAnimationTracks(
-			AnimSeq,
+			CompressibleAnimData,
+			OutResult,
 			static_cast<AnimationCompressionFormat>(TranslationCompressionFormat),
 			static_cast<AnimationCompressionFormat>(RotationCompressionFormat),
 			static_cast<AnimationCompressionFormat>(ScaleCompressionFormat),
@@ -68,7 +69,7 @@ void UAnimCompress_BitwiseCompressOnly::DoReduction(UAnimSequence* AnimSeq, cons
 	}
 
 	// We could be invalid, set the links again
-	AnimationFormat_SetInterfaceLinks(*AnimSeq);
+	AnimationFormat_SetInterfaceLinks(OutResult);
 #endif // WITH_EDITORONLY_DATA
 }
 
diff --git a/Engine/Source/Runtime/Engine/Private/Animation/AnimCompress_LeastDestructive.cpp b/Engine/Source/Runtime/Engine/Private/Animation/AnimCompress_LeastDestructive.cpp
--- a/Engine/Source/Runtime/Engine/Private/Animation/AnimCompress_LeastDestructive.cpp
+++ b/Engine/Source/Runtime/Engine/Private/Animation/AnimCompress_LeastDestructive.cpp
@@ -17,13 +17,11 @@ UAnimCompress_LeastDestructive::UAnimCompress_LeastDestructive(const FObjectInit
 
 
 #if WITH_EDITOR
-void UAnimCompress_LeastDestructive::DoReduction(UAnimSequence* AnimSeq, const TArray<FBoneData>& BoneData)
+void UAnimCompress_LeastDestructive::DoReduction(const FCompressibleAnimData& CompressibleAnimData, FCompressibleAnimDataResult& OutResult)
 {
 	UAnimCompress* BitwiseCompressor = NewObject<UAnimCompress_BitwiseCompressOnly>();
 	BitwiseCompressor->RotationCompressionFormat = ACF_Float96NoW;
 	BitwiseCompressor->TranslationCompressionFormat = ACF_None;
-	BitwiseCompressor->Reduce(AnimSeq, false, BoneData);
-
-	AnimSeq->CompressionScheme = BitwiseCompressor;
+	BitwiseCompressor->Reduce(CompressibleAnimData, OutResult);
 }
 #endif // WITH_EDITOR
