Refactor UAnimSequence out of compression code

diff --git a/Engine/Source/Runtime/Engine/Private/Animation/AnimEncoding.cpp b/Engine/Source/Runtime/Engine/Private/Animation/AnimEncoding.cpp
--- a/Engine/Source/Runtime/Engine/Private/Animation/AnimEncoding.cpp
+++ b/Engine/Source/Runtime/Engine/Private/Animation/AnimEncoding.cpp
@@ -128,26 +128,26 @@ inline int32 GetCompressedScaleStride(AnimationCompressionFormat ScaleCompressio
 /**
  * Compressed translation data will be byte swapped in chunks of this size.
  */
-inline int32 GetCompressedTranslationStride(const UAnimSequence* Seq)
+inline int32 GetCompressedTranslationStride(const FUECompressedAnimData& CompressedData)
 {
-	return CompressedTranslationStrides[static_cast<AnimationCompressionFormat>(Seq->TranslationCompressionFormat)];
+	return CompressedTranslationStrides[static_cast<AnimationCompressionFormat>(CompressedData.TranslationCompressionFormat)];
 }
 
 /**
  * Compressed rotation data will be byte swapped in chunks of this size.
  */
-inline int32 GetCompressedRotationStride(const UAnimSequence* Seq)
+inline int32 GetCompressedRotationStride(const FUECompressedAnimData& CompressedData)
 {
-	return CompressedRotationStrides[static_cast<AnimationCompressionFormat>(Seq->RotationCompressionFormat)];
+	return CompressedRotationStrides[static_cast<AnimationCompressionFormat>(CompressedData.RotationCompressionFormat)];
 }
 
 /**
  * Compressed Scale data will be byte swapped in chunks of this size.
  */
-inline int32 GetCompressedScaleStride(const UAnimSequence* Seq)
+inline int32 GetCompressedScaleStride(const FUECompressedAnimData& CompressedData)
 {
 	// @todo fixme change this?
-	return CompressedScaleStrides[static_cast<AnimationCompressionFormat>(Seq->ScaleCompressionFormat)];
+	return CompressedScaleStrides[static_cast<AnimationCompressionFormat>(CompressedData.ScaleCompressionFormat)];
 }
 
 
@@ -264,7 +264,7 @@ void AnimEncodingLegacyBase::GetBoneAtom(FTransform& OutAtom, FAnimSequenceDecom
 /**
  * Handles Byte-swapping incoming animation data from a MemoryReader
  *
- * @param	Seq					An Animation Sequence to contain the read data.
+ * @param	CompressedData		The compressed animation data being operated on.
  * @param	MemoryReader		The MemoryReader object to read from.
  */
 [REDACTED_EMAIL]: Apparent VC10 compiler bug here causes an access violation in optimized builds
@@ -272,42 +272,32 @@ void AnimEncodingLegacyBase::GetBoneAtom(FTransform& OutAtom, FAnimSequenceDecom
 	PRAGMA_DISABLE_OPTIMIZATION
 #endif
 void AnimEncodingLegacyBase::ByteSwapIn(
-	UAnimSequence& Seq, 
+	FUECompressedAnimData& CompressedData,
 	FMemoryReader& MemoryReader)
 {
-	const int32 NumTracks = Seq.CompressedTrackOffsets.Num() / 4;
+	const int32 BufferStart = MemoryReader.Tell();
 
-	int32 OriginalNumBytes = MemoryReader.TotalSize();
-	Seq.CompressedByteStream.Empty(OriginalNumBytes);
-	Seq.CompressedByteStream.AddUninitialized(OriginalNumBytes);
+	const int32 NumTracks = CompressedData.CompressedTrackOffsets.Num() / 4;
 
-	if (Seq.CompressedSegments.Num() != 0)
-	{
-#if !PLATFORM_LITTLE_ENDIAN
-#error "Byte swapping needs to be implemented here to support big-endian platforms"
-#endif
-
-		// TODO: Byte swap the new format
-		MemoryReader.Serialize(Seq.CompressedByteStream.GetData(), Seq.CompressedByteStream.Num());
-		return;
-	}
+	int32 OriginalNumBytes = MemoryReader.TotalSize() - BufferStart;
+	check(CompressedData.CompressedByteStream.Num() == OriginalNumBytes);
 
 	// Read and swap
-	uint8* StreamBase = Seq.CompressedByteStream.GetData();
-	bool bHasValidScale = Seq.CompressedScaleOffsets.IsValid();
+	uint8* StreamBase = CompressedData.CompressedByteStream.GetData();
+	bool bHasValidScale = CompressedData.CompressedScaleOffsets.IsValid();
 
 	for ( int32 TrackIndex = 0; TrackIndex < NumTracks; ++TrackIndex )
 	{
-		const int32 OffsetTrans	= Seq.CompressedTrackOffsets[TrackIndex*4+0];
-		const int32 NumKeysTrans	= Seq.CompressedTrackOffsets[TrackIndex*4+1];
-		const int32 OffsetRot		= Seq.CompressedTrackOffsets[TrackIndex*4+2];
-		const int32 NumKeysRot	= Seq.CompressedTrackOffsets[TrackIndex*4+3];
+		const int32 OffsetTrans	= CompressedData.CompressedTrackOffsets[TrackIndex*4+0];
+		const int32 NumKeysTrans	= CompressedData.CompressedTrackOffsets[TrackIndex*4+1];
+		const int32 OffsetRot		= CompressedData.CompressedTrackOffsets[TrackIndex*4+2];
+		const int32 NumKeysRot	= CompressedData.CompressedTrackOffsets[TrackIndex*4+3];
 
 		// Translation data.
 		checkSlow( (OffsetTrans % 4) == 0 && "CompressedByteStream not aligned to four bytes" );
 		uint8* TransTrackData = StreamBase + OffsetTrans;
-		checkSlow(Seq.TranslationCodec != NULL);
-		((AnimEncodingLegacyBase*)Seq.TranslationCodec)->ByteSwapTranslationIn(Seq, MemoryReader, TransTrackData, NumKeysTrans);
+		checkSlow(CompressedData.TranslationCodec != NULL);
+		((AnimEncodingLegacyBase*)CompressedData.TranslationCodec)->ByteSwapTranslationIn(CompressedData, MemoryReader, TransTrackData, NumKeysTrans);
 
 		// Like the compressed byte stream, pad the serialization stream to four bytes.
 		// As a sanity check, each pad byte can be checked to be the PadSentinel.
@@ -316,8 +306,8 @@ void AnimEncodingLegacyBase::ByteSwapIn(
 		// Rotation data.
 		checkSlow( (OffsetRot % 4) == 0 && "CompressedByteStream not aligned to four bytes" );
 		uint8* RotTrackData = StreamBase + OffsetRot;
-		checkSlow(Seq.RotationCodec != NULL);
-		((AnimEncodingLegacyBase*)Seq.RotationCodec)->ByteSwapRotationIn(Seq, MemoryReader, RotTrackData, NumKeysRot);
+		checkSlow(CompressedData.RotationCodec != NULL);
+		((AnimEncodingLegacyBase*)CompressedData.RotationCodec)->ByteSwapRotationIn(CompressedData, MemoryReader, RotTrackData, NumKeysRot);
 
 		// Like the compressed byte stream, pad the serialization stream to four bytes.
 		// As a sanity check, each pad byte can be checked to be the PadSentinel.
@@ -325,14 +315,14 @@ void AnimEncodingLegacyBase::ByteSwapIn(
 
 		if (bHasValidScale)
 		{
-			const int32 OffsetScale		= Seq.CompressedScaleOffsets.GetOffsetData(TrackIndex, 0);
-			const int32 NumKeysScale		= Seq.CompressedScaleOffsets.GetOffsetData(TrackIndex, 1);
+			const int32 OffsetScale		= CompressedData.CompressedScaleOffsets.GetOffsetData(TrackIndex, 0);
+			const int32 NumKeysScale		= CompressedData.CompressedScaleOffsets.GetOffsetData(TrackIndex, 1);
 
 			// Scale data.
 			checkSlow( (OffsetScale % 4) == 0 && "CompressedByteStream not aligned to four bytes" );
 			uint8* ScaleTrackData = StreamBase + OffsetScale;
-			checkSlow(Seq.ScaleCodec != NULL);
-			((AnimEncodingLegacyBase*)Seq.ScaleCodec)->ByteSwapScaleIn(Seq, MemoryReader, ScaleTrackData, NumKeysScale);
+			checkSlow(CompressedData.ScaleCodec != NULL);
+			((AnimEncodingLegacyBase*)CompressedData.ScaleCodec)->ByteSwapScaleIn(CompressedData, MemoryReader, ScaleTrackData, NumKeysScale);
 
 			// Like the compressed byte stream, pad the serialization stream to four bytes.
 			// As a sanity check, each pad byte can be checked to be the PadSentinel.
@@ -354,85 +344,68 @@ PRAGMA_ENABLE_OPTIMIZATION
  * @param	ForceByteSwapping	true is byte swapping is not optional.
  */
 void AnimEncodingLegacyBase::ByteSwapOut(
-	UAnimSequence& Seq,
-	TArray<uint8>& SerializedData, 
-	bool ForceByteSwapping,
-	bool bMaintainComponentOrder)
+	FUECompressedAnimData& CompressedData,
+	FMemoryWriter& MemoryWriter)
 {
-	FMemoryWriter MemoryWriter( SerializedData, true );
-	MemoryWriter.SetByteSwapping( ForceByteSwapping );
-
-	if (Seq.CompressedSegments.Num() != 0)
-	{
-		// TODO: Byte swap the new format
-		MemoryWriter.Serialize(Seq.CompressedByteStream.GetData(), Seq.CompressedByteStream.Num());
-		return;
-	}
+	const int32 BufferStart = MemoryWriter.Tell();
 
-	uint8* StreamBase		= Seq.CompressedByteStream.GetData();
-	const int32 NumTracks		= Seq.CompressedTrackOffsets.Num()/4;
+	uint8* StreamBase		= CompressedData.CompressedByteStream.GetData();
+	const int32 NumTracks		= CompressedData.CompressedTrackOffsets.Num()/4;
 
-	bool bHasValidScale = Seq.CompressedScaleOffsets.IsValid();
+	bool bHasValidScale = CompressedData.CompressedScaleOffsets.IsValid();
 
 	for ( int32 TrackIndex = 0; TrackIndex < NumTracks; ++TrackIndex )
 	{
-		const int32 OffsetTrans	= Seq.CompressedTrackOffsets[TrackIndex*4];
-		const int32 NumKeysTrans	= Seq.CompressedTrackOffsets[TrackIndex*4+1];
-		const int32 OffsetRot		= Seq.CompressedTrackOffsets[TrackIndex*4+2];
-		const int32 NumKeysRot	= Seq.CompressedTrackOffsets[TrackIndex*4+3];
+		const int32 OffsetTrans	= CompressedData.CompressedTrackOffsets[TrackIndex*4];
+		const int32 NumKeysTrans	= CompressedData.CompressedTrackOffsets[TrackIndex*4+1];
+		const int32 OffsetRot		= CompressedData.CompressedTrackOffsets[TrackIndex*4+2];
+		const int32 NumKeysRot	= CompressedData.CompressedTrackOffsets[TrackIndex*4+3];
+
+		MemoryWriter.Seek(BufferStart + OffsetTrans);
 
-		if (bMaintainComponentOrder)
-		{
-			MemoryWriter.Seek(OffsetTrans);
-		}
 		// Translation data.
 		checkSlow( (OffsetTrans % 4) == 0 && "CompressedByteStream not aligned to four bytes" );
 		uint8* TransTrackData = StreamBase + OffsetTrans;
-		if (Seq.TranslationCodec != NULL)
+		if (CompressedData.TranslationCodec != NULL)
 		{
-			((AnimEncodingLegacyBase*)Seq.TranslationCodec)->ByteSwapTranslationOut(Seq, MemoryWriter, TransTrackData, NumKeysTrans);
+			((AnimEncodingLegacyBase*)CompressedData.TranslationCodec)->ByteSwapTranslationOut(CompressedData, MemoryWriter, TransTrackData, NumKeysTrans);
 		}
 		else
 		{
-			UE_LOG(LogAnimationCompression, Fatal, TEXT("%i: unknown or unsupported animation format"), (int32)Seq.KeyEncodingFormat );
+			UE_LOG(LogAnimationCompression, Fatal, TEXT("%i: unknown or unsupported animation format"), (int32)CompressedData.KeyEncodingFormat );
 		};
 
 		// Like the compressed byte stream, pad the serialization stream to four bytes.
 		PadMemoryWriter(&MemoryWriter, TransTrackData, 4);
 
-		if (bMaintainComponentOrder)
-		{
-			MemoryWriter.Seek(OffsetRot);
-		}
+		MemoryWriter.Seek(BufferStart + OffsetRot);
+
 		// Rotation data.
 		checkSlow( (OffsetRot % 4) == 0 && "CompressedByteStream not aligned to four bytes" );
 		uint8* RotTrackData = StreamBase + OffsetRot;
-		checkSlow(Seq.RotationCodec != NULL);
-		((AnimEncodingLegacyBase*)Seq.RotationCodec)->ByteSwapRotationOut(Seq, MemoryWriter, RotTrackData, NumKeysRot);
+		checkSlow(CompressedData.RotationCodec != NULL);
+		((AnimEncodingLegacyBase*)CompressedData.RotationCodec)->ByteSwapRotationOut(CompressedData, MemoryWriter, RotTrackData, NumKeysRot);
 
 		// Like the compressed byte stream, pad the serialization stream to four bytes.
 		PadMemoryWriter(&MemoryWriter, RotTrackData, 4);
 
 		if (bHasValidScale)
 		{
-			const int32 OffsetScale	= Seq.CompressedScaleOffsets.GetOffsetData(TrackIndex, 0);
-			const int32 NumKeysScale	= Seq.CompressedScaleOffsets.GetOffsetData(TrackIndex, 1);
+			const int32 OffsetScale	= CompressedData.CompressedScaleOffsets.GetOffsetData(TrackIndex, 0);
+			const int32 NumKeysScale	= CompressedData.CompressedScaleOffsets.GetOffsetData(TrackIndex, 1);
 
-			if (bMaintainComponentOrder)
-			{
-				MemoryWriter.Seek(OffsetScale);
-			}
+			MemoryWriter.Seek(BufferStart + OffsetScale);
 
 			// Scale data.
 			checkSlow( (OffsetScale % 4) == 0 && "CompressedByteStream not aligned to four bytes" );
 			uint8* ScaleTrackData = StreamBase + OffsetScale;
-			if (Seq.ScaleCodec != NULL)
+			if (CompressedData.ScaleCodec != NULL)
 			{
-				((AnimEncodingLegacyBase*)Seq.ScaleCodec)->ByteSwapScaleOut(Seq, MemoryWriter, ScaleTrackData, NumKeysScale);
+				((AnimEncodingLegacyBase*)CompressedData.ScaleCodec)->ByteSwapScaleOut(CompressedData, MemoryWriter, ScaleTrackData, NumKeysScale);
 			}
 			else
 			{
-				UE_LOG(LogAnimationCompression, Fatal, TEXT("%i: unknown or unsupported animation format"), (int32)Seq.KeyEncodingFormat );
+				UE_LOG(LogAnimationCompression, Fatal, TEXT("%i: unknown or unsupported animation format"), (int32)CompressedData.KeyEncodingFormat );
 			};
 
 			// Like the compressed byte stream, pad the serialization stream to four bytes.
@@ -456,7 +429,7 @@ void AnimEncodingLegacyBase::ByteSwapOut(
  * @param	NumRotTracksWithOneKey		The total number of Rotation Tracks found containing a single key.
 */
 void AnimationFormat_GetStats(	
-	const UAnimSequence* Seq, 
+	const FUECompressedAnimData& CompressedData, 
 	int32& NumTransTracks,
 	int32& NumRotTracks,
 	int32& NumScaleTracks,
@@ -471,245 +444,243 @@ void AnimationFormat_GetStats(
 	int32& NumRotTracksWithOneKey, 
 	int32& NumScaleTracksWithOneKey)
 {
-	if (Seq)
-	{
-		OverheadSize = Seq->CompressedTrackOffsets.Num() * sizeof(int32);
-		const size_t KeyFrameLookupSize = (Seq->GetCompressedNumberOfFrames() > 0xFF) ? sizeof(uint16) : sizeof(uint8);
+	OverheadSize = CompressedData.CompressedTrackOffsets.Num() * sizeof(int32);
+	const size_t KeyFrameLookupSize = (CompressedData.CompressedNumberOfFrames > 0xFF) ? sizeof(uint16) : sizeof(uint8);
 
-		if (Seq->KeyEncodingFormat != AKF_PerTrackCompression)
+	if (CompressedData.KeyEncodingFormat != AKF_PerTrackCompression)
+	{
+		const int32 TransStride	= GetCompressedTranslationStride(CompressedData);
+		const int32 RotStride		= GetCompressedRotationStride(CompressedData);
+		const int32 ScaleStride		= GetCompressedScaleStride(CompressedData);
+		const int32 TransNum		= CompressedTranslationNum[CompressedData.TranslationCompressionFormat];
+		const int32 RotNum		= CompressedRotationNum[CompressedData.RotationCompressionFormat];
+		const int32 ScaleNum		= CompressedScaleNum[CompressedData.ScaleCompressionFormat];
+
+		TranslationKeySize = TransStride * TransNum;
+		RotationKeySize = RotStride * RotNum;
+		ScaleKeySize = ScaleStride * ScaleNum;
+
+		// Track number of tracks.
+		NumTransTracks	= CompressedData.CompressedTrackOffsets.Num()/4;
+		NumRotTracks	= CompressedData.CompressedTrackOffsets.Num()/4;
+		NumScaleTracks	= CompressedData.CompressedScaleOffsets.GetNumTracks();
+
+		// Track total number of keys.
+		TotalNumTransKeys = 0;
+		TotalNumRotKeys = 0;
+		TotalNumScaleKeys = 0;
+
+		// Track number of tracks with a single key.
+		NumTransTracksWithOneKey = 0;
+		NumRotTracksWithOneKey = 0;
+		NumScaleTracksWithOneKey = 0;
+
+		// Translation.
+		for ( int32 TrackIndex = 0; TrackIndex < NumTransTracks; ++TrackIndex )
 		{
-			const int32 TransStride	= GetCompressedTranslationStride(Seq);
-			const int32 RotStride		= GetCompressedRotationStride(Seq);
-			const int32 ScaleStride		= GetCompressedScaleStride(Seq);
-			const int32 TransNum		= CompressedTranslationNum[Seq->TranslationCompressionFormat];
-			const int32 RotNum		= CompressedRotationNum[Seq->RotationCompressionFormat];
-			const int32 ScaleNum		= CompressedScaleNum[Seq->ScaleCompressionFormat];
-
-			TranslationKeySize = TransStride * TransNum;
-			RotationKeySize = RotStride * RotNum;
-			ScaleKeySize = ScaleStride * ScaleNum;
-
-			// Track number of tracks.
-			NumTransTracks	= Seq->CompressedTrackOffsets.Num()/4;
-			NumRotTracks	= Seq->CompressedTrackOffsets.Num()/4;
-			NumScaleTracks	= Seq->CompressedScaleOffsets.GetNumTracks();
-
-			// Track total number of keys.
-			TotalNumTransKeys = 0;
-			TotalNumRotKeys = 0;
-			TotalNumScaleKeys = 0;
-
-			// Track number of tracks with a single key.
-			NumTransTracksWithOneKey = 0;
-			NumRotTracksWithOneKey = 0;
-			NumScaleTracksWithOneKey = 0;
-
-			// Translation.
-			for ( int32 TrackIndex = 0; TrackIndex < NumTransTracks; ++TrackIndex )
+			const int32 NumKeys = CompressedData.CompressedTrackOffsets[TrackIndex*4+1];
+			TotalNumTransKeys += NumKeys;
+			if ( NumKeys == 1 )
 			{
-				const int32 NumKeys = Seq->CompressedTrackOffsets[TrackIndex*4+1];
-				TotalNumTransKeys += NumKeys;
-				if ( NumKeys == 1 )
-				{
-					++NumTransTracksWithOneKey;
-				}
-				else
-				{
-					OverheadSize += (Seq->KeyEncodingFormat == AKF_VariableKeyLerp) ? NumKeys * KeyFrameLookupSize : 0;
-				}
+				++NumTransTracksWithOneKey;
 			}
-
-			// Rotation.
-			for ( int32 TrackIndex = 0; TrackIndex < NumRotTracks; ++TrackIndex )
+			else
 			{
-				const int32 NumKeys = Seq->CompressedTrackOffsets[TrackIndex*4+3];
-				TotalNumRotKeys += NumKeys;
-				if ( NumKeys == 1 )
-				{
-					++NumRotTracksWithOneKey;
-				}
-				else
-				{
-					OverheadSize += (Seq->KeyEncodingFormat == AKF_VariableKeyLerp) ? NumKeys * KeyFrameLookupSize : 0;
-				}
+				OverheadSize += (CompressedData.KeyEncodingFormat == AKF_VariableKeyLerp) ? NumKeys * KeyFrameLookupSize : 0;
 			}
+		}
 
-			// Scale.
-			for ( int32 ScaleIndex = 0; ScaleIndex < NumScaleTracks; ++ScaleIndex )
+		// Rotation.
+		for ( int32 TrackIndex = 0; TrackIndex < NumRotTracks; ++TrackIndex )
+		{
+			const int32 NumKeys = CompressedData.CompressedTrackOffsets[TrackIndex*4+3];
+			TotalNumRotKeys += NumKeys;
+			if ( NumKeys == 1 )
 			{
-				const int32 NumKeys = Seq->CompressedScaleOffsets.GetOffsetData(ScaleIndex, 1);
-				TotalNumScaleKeys += NumKeys;
-				if ( NumKeys == 1 )
-				{
-					++NumScaleTracksWithOneKey;
-				}
-				else
-				{
-					OverheadSize += (Seq->KeyEncodingFormat == AKF_VariableKeyLerp) ? NumKeys * KeyFrameLookupSize : 0;
-				}
+				++NumRotTracksWithOneKey;
+			}
+			else
+			{
+				OverheadSize += (CompressedData.KeyEncodingFormat == AKF_VariableKeyLerp) ? NumKeys * KeyFrameLookupSize : 0;
 			}
+		}
 
-			// Add in scaling values (min+range for interval encoding)
-			OverheadSize += (Seq->RotationCompressionFormat == ACF_IntervalFixed32NoW) ? (NumRotTracks - NumRotTracksWithOneKey) * sizeof(float) * 6 : 0;
-			OverheadSize += (Seq->TranslationCompressionFormat == ACF_IntervalFixed32NoW) ? (NumTransTracks - NumTransTracksWithOneKey) * sizeof(float) * 6 : 0;
-			OverheadSize += (Seq->ScaleCompressionFormat == ACF_IntervalFixed32NoW) ? (NumScaleTracks - NumScaleTracksWithOneKey) * sizeof(float) * 6 : 0;
+		// Scale.
+		for ( int32 ScaleIndex = 0; ScaleIndex < NumScaleTracks; ++ScaleIndex )
+		{
+			const int32 NumKeys = CompressedData.CompressedScaleOffsets.GetOffsetData(ScaleIndex, 1);
+			TotalNumScaleKeys += NumKeys;
+			if ( NumKeys == 1 )
+			{
+				++NumScaleTracksWithOneKey;
+			}
+			else
+			{
+				OverheadSize += (CompressedData.KeyEncodingFormat == AKF_VariableKeyLerp) ? NumKeys * KeyFrameLookupSize : 0;
+			}
 		}
-		else
+
+		// Add in scaling values (min+range for interval encoding)
+		OverheadSize += (CompressedData.RotationCompressionFormat == ACF_IntervalFixed32NoW) ? (NumRotTracks - NumRotTracksWithOneKey) * sizeof(float) * 6 : 0;
+		OverheadSize += (CompressedData.TranslationCompressionFormat == ACF_IntervalFixed32NoW) ? (NumTransTracks - NumTransTracksWithOneKey) * sizeof(float) * 6 : 0;
+		OverheadSize += (CompressedData.ScaleCompressionFormat == ACF_IntervalFixed32NoW) ? (NumScaleTracks - NumScaleTracksWithOneKey) * sizeof(float) * 6 : 0;
+	}
+	else
+	{
+		int32 TotalTransKeysThatContributedSize = 0;
+		int32 TotalRotKeysThatContributedSize = 0;
+		int32 TotalScaleKeysThatContributedSize = 0;
+
+		TranslationKeySize = 0;
+		RotationKeySize = 0;
+		ScaleKeySize = 0;
+
+		// Track number of tracks.
+		NumTransTracks = CompressedData.CompressedTrackOffsets.Num() / 2;
+		NumRotTracks = CompressedData.CompressedTrackOffsets.Num() / 2;
+		// should be without divided by 2
+		NumScaleTracks = CompressedData.CompressedScaleOffsets.GetNumTracks();
+
+		// Track total number of keys.
+		TotalNumTransKeys = 0;
+		TotalNumRotKeys = 0;
+		TotalNumScaleKeys = 0;
+
+		// Track number of tracks with a single key.
+		NumTransTracksWithOneKey = 0;
+		NumRotTracksWithOneKey = 0;
+		NumScaleTracksWithOneKey = 0;
+
+		// Translation.
+		for (int32 TrackIndex = 0; TrackIndex < NumTransTracks; ++TrackIndex)
 		{
-			int32 TotalTransKeysThatContributedSize = 0;
-			int32 TotalRotKeysThatContributedSize = 0;
-			int32 TotalScaleKeysThatContributedSize = 0;
-
-			TranslationKeySize = 0;
-			RotationKeySize = 0;
-			ScaleKeySize = 0;
-
-			// Track number of tracks.
-			NumTransTracks = Seq->CompressedTrackOffsets.Num() / 2;
-			NumRotTracks = Seq->CompressedTrackOffsets.Num() / 2;
-			// should be without divided by 2
-			NumScaleTracks = Seq->CompressedScaleOffsets.GetNumTracks();
-
-			// Track total number of keys.
-			TotalNumTransKeys = 0;
-			TotalNumRotKeys = 0;
-			TotalNumScaleKeys = 0;
-
-			// Track number of tracks with a single key.
-			NumTransTracksWithOneKey = 0;
-			NumRotTracksWithOneKey = 0;
-			NumScaleTracksWithOneKey = 0;
-
-			// Translation.
-			for (int32 TrackIndex = 0; TrackIndex < NumTransTracks; ++TrackIndex)
+			const int32 TransOffset = CompressedData.CompressedTrackOffsets[TrackIndex*2+0];
+			if (TransOffset == INDEX_NONE)
+			{
+				++TotalNumTransKeys;
+				++NumTransTracksWithOneKey;
+			}
+			else
 			{
-				const int32 TransOffset = Seq->CompressedTrackOffsets[TrackIndex*2+0];
-				if (TransOffset == INDEX_NONE)
+				const int32 Header = *((int32*)(CompressedData.CompressedByteStream.GetData() + TransOffset));
+
+				int32 KeyFormat;
+				int32 FormatFlags;
+				int32 TrackBytesPerKey;
+				int32 TrackFixedBytes;
+				int32 NumKeys;
+
+				FAnimationCompression_PerTrackUtils::DecomposeHeader(Header, /*OUT*/ KeyFormat, /*OUT*/ NumKeys, /*OUT*/ FormatFlags, /*OUT*/ TrackBytesPerKey, /*OUT*/ TrackFixedBytes);
+				TranslationKeySize += TrackBytesPerKey * NumKeys;
+				TotalTransKeysThatContributedSize += NumKeys;
+				OverheadSize += TrackFixedBytes;
+				OverheadSize += ((FormatFlags & 0x08) != 0) ? NumKeys * KeyFrameLookupSize : 0;
+					
+				TotalNumTransKeys += NumKeys;
+				if (NumKeys <= 1)
 				{
-					++TotalNumTransKeys;
 					++NumTransTracksWithOneKey;
 				}
-				else
-				{
-					const int32 Header = *((int32*)(Seq->CompressedByteStream.GetData() + TransOffset));
-
-					int32 KeyFormat;
-					int32 FormatFlags;
-					int32 TrackBytesPerKey;
-					int32 TrackFixedBytes;
-					int32 NumKeys;
-
-					FAnimationCompression_PerTrackUtils::DecomposeHeader(Header, /*OUT*/ KeyFormat, /*OUT*/ NumKeys, /*OUT*/ FormatFlags, /*OUT*/ TrackBytesPerKey, /*OUT*/ TrackFixedBytes);
-					TranslationKeySize += TrackBytesPerKey * NumKeys;
-					TotalTransKeysThatContributedSize += NumKeys;
-					OverheadSize += TrackFixedBytes;
-					OverheadSize += ((FormatFlags & 0x08) != 0) ? NumKeys * KeyFrameLookupSize : 0;
-					
-					TotalNumTransKeys += NumKeys;
-					if (NumKeys <= 1)
-					{
-						++NumTransTracksWithOneKey;
-					}
-				}
 			}
+		}
 
-			// Rotation.
-			for (int32 TrackIndex = 0; TrackIndex < NumRotTracks; ++TrackIndex)
+		// Rotation.
+		for (int32 TrackIndex = 0; TrackIndex < NumRotTracks; ++TrackIndex)
+		{
+			const int32 RotOffset = CompressedData.CompressedTrackOffsets[TrackIndex*2+1];
+			if (RotOffset == INDEX_NONE)
+			{
+				++TotalNumRotKeys;
+				++NumRotTracksWithOneKey;
+			}
+			else
 			{
-				const int32 RotOffset = Seq->CompressedTrackOffsets[TrackIndex*2+1];
-				if (RotOffset == INDEX_NONE)
+				const int32 Header = *((int32*)(CompressedData.CompressedByteStream.GetData() + RotOffset));
+
+				int32 KeyFormat;
+				int32 FormatFlags;
+				int32 TrackBytesPerKey;
+				int32 TrackFixedBytes;
+				int32 NumKeys;
+
+				FAnimationCompression_PerTrackUtils::DecomposeHeader(Header, /*OUT*/ KeyFormat, /*OUT*/ NumKeys, /*OUT*/ FormatFlags, /*OUT*/ TrackBytesPerKey, /*OUT*/ TrackFixedBytes);
+				RotationKeySize += TrackBytesPerKey * NumKeys;
+				TotalRotKeysThatContributedSize += NumKeys;
+				OverheadSize += TrackFixedBytes;
+				OverheadSize += ((FormatFlags & 0x08) != 0) ? NumKeys * KeyFrameLookupSize : 0;
+
+				TotalNumRotKeys += NumKeys;
+				if (NumKeys <= 1)
 				{
-					++TotalNumRotKeys;
 					++NumRotTracksWithOneKey;
 				}
-				else
-				{
-					const int32 Header = *((int32*)(Seq->CompressedByteStream.GetData() + RotOffset));
-
-					int32 KeyFormat;
-					int32 FormatFlags;
-					int32 TrackBytesPerKey;
-					int32 TrackFixedBytes;
-					int32 NumKeys;
-
-					FAnimationCompression_PerTrackUtils::DecomposeHeader(Header, /*OUT*/ KeyFormat, /*OUT*/ NumKeys, /*OUT*/ FormatFlags, /*OUT*/ TrackBytesPerKey, /*OUT*/ TrackFixedBytes);
-					RotationKeySize += TrackBytesPerKey * NumKeys;
-					TotalRotKeysThatContributedSize += NumKeys;
-					OverheadSize += TrackFixedBytes;
-					OverheadSize += ((FormatFlags & 0x08) != 0) ? NumKeys * KeyFrameLookupSize : 0;
-
-					TotalNumRotKeys += NumKeys;
-					if (NumKeys <= 1)
-					{
-						++NumRotTracksWithOneKey;
-					}
-				}
 			}
+		}
 
-			// Scale.
-			for (int32 TrackIndex = 0; TrackIndex < NumScaleTracks; ++TrackIndex)
+		// Scale.
+		for (int32 TrackIndex = 0; TrackIndex < NumScaleTracks; ++TrackIndex)
+		{
+			const int32 ScaleOffset = CompressedData.CompressedScaleOffsets.GetOffsetData(TrackIndex, 0);
+			if (ScaleOffset == INDEX_NONE)
+			{
+				++TotalNumScaleKeys;
+				++NumScaleTracksWithOneKey;
+			}
+			else
 			{
-				const int32 ScaleOffset = Seq->CompressedScaleOffsets.GetOffsetData(TrackIndex, 0);
-				if (ScaleOffset == INDEX_NONE)
+				const int32 Header = *((int32*)(CompressedData.CompressedByteStream.GetData() + ScaleOffset));
+
+				int32 KeyFormat;
+				int32 FormatFlags;
+				int32 TrackBytesPerKey;
+				int32 TrackFixedBytes;
+				int32 NumKeys;
+
+				FAnimationCompression_PerTrackUtils::DecomposeHeader(Header, /*OUT*/ KeyFormat, /*OUT*/ NumKeys, /*OUT*/ FormatFlags, /*OUT*/ TrackBytesPerKey, /*OUT*/ TrackFixedBytes);
+				ScaleKeySize += TrackBytesPerKey * NumKeys;
+				TotalScaleKeysThatContributedSize += NumKeys;
+				OverheadSize += TrackFixedBytes;
+				OverheadSize += ((FormatFlags & 0x08) != 0) ? NumKeys * KeyFrameLookupSize : 0;
+
+				TotalNumScaleKeys += NumKeys;
+				if (NumKeys <= 1)
 				{
-					++TotalNumScaleKeys;
 					++NumScaleTracksWithOneKey;
 				}
-				else
-				{
-					const int32 Header = *((int32*)(Seq->CompressedByteStream.GetData() + ScaleOffset));
-
-					int32 KeyFormat;
-					int32 FormatFlags;
-					int32 TrackBytesPerKey;
-					int32 TrackFixedBytes;
-					int32 NumKeys;
-
-					FAnimationCompression_PerTrackUtils::DecomposeHeader(Header, /*OUT*/ KeyFormat, /*OUT*/ NumKeys, /*OUT*/ FormatFlags, /*OUT*/ TrackBytesPerKey, /*OUT*/ TrackFixedBytes);
-					ScaleKeySize += TrackBytesPerKey * NumKeys;
-					TotalScaleKeysThatContributedSize += NumKeys;
-					OverheadSize += TrackFixedBytes;
-					OverheadSize += ((FormatFlags & 0x08) != 0) ? NumKeys * KeyFrameLookupSize : 0;
-
-					TotalNumScaleKeys += NumKeys;
-					if (NumKeys <= 1)
-					{
-						++NumScaleTracksWithOneKey;
-					}
-				}
 			}
+		}
 
-			// Average key sizes
-			if (TotalRotKeysThatContributedSize > 0)
-			{
-				RotationKeySize = RotationKeySize / TotalRotKeysThatContributedSize;
-			}
+		// Average key sizes
+		if (TotalRotKeysThatContributedSize > 0)
+		{
+			RotationKeySize = RotationKeySize / TotalRotKeysThatContributedSize;
+		}
 
-			if (TotalTransKeysThatContributedSize > 0)
-			{
-				TranslationKeySize = TranslationKeySize / TotalTransKeysThatContributedSize;
-			}
+		if (TotalTransKeysThatContributedSize > 0)
+		{
+			TranslationKeySize = TranslationKeySize / TotalTransKeysThatContributedSize;
+		}
 
-			if (TotalScaleKeysThatContributedSize > 0)
-			{
-				ScaleKeySize = ScaleKeySize / TotalScaleKeysThatContributedSize;
-			}
+		if (TotalScaleKeysThatContributedSize > 0)
+		{
+			ScaleKeySize = ScaleKeySize / TotalScaleKeysThatContributedSize;
 		}
 	}
 }
 
 /**
- * Sets the internal Animation Codec Interface Links within an Animation Sequence
+ * Sets the internal Animation Codec Interface Links within an Animation CompressedDatauence
  *
  * @param	Seq					An Animation Sequence to setup links within.
 */
-void AnimationFormat_SetInterfaceLinks(UAnimSequence& Seq)
+template<typename CompressedDataType>
+void AnimationFormat_SetInterfaceLinks(CompressedDataType& CompressedData)
 {
-	Seq.TranslationCodec = NULL;
-	Seq.RotationCodec = NULL;
-	Seq.ScaleCodec = NULL;
+	CompressedData.TranslationCodec = NULL;
+	CompressedData.RotationCodec = NULL;
+	CompressedData.ScaleCodec = NULL;
 
-	if (Seq.KeyEncodingFormat == AKF_ConstantKeyLerp)
+	if (CompressedData.KeyEncodingFormat == AKF_ConstantKeyLerp)
 	{
 		static AEFConstantKeyLerp<ACF_None>					AEFConstantKeyLerp_None;
 		static AEFConstantKeyLerp<ACF_Float96NoW>			AEFConstantKeyLerp_Float96NoW;
@@ -720,75 +691,75 @@ void AnimationFormat_SetInterfaceLinks(UAnimSequence& Seq)
 		static AEFConstantKeyLerp<ACF_Identity>				AEFConstantKeyLerp_Identity;
 
 		// setup translation codec
-		switch(Seq.TranslationCompressionFormat)
+		switch(CompressedData.TranslationCompressionFormat)
 		{
 			case ACF_None:
-				Seq.TranslationCodec = &AEFConstantKeyLerp_None;
+				CompressedData.TranslationCodec = &AEFConstantKeyLerp_None;
 				break;
 			case ACF_Float96NoW:
-				Seq.TranslationCodec = &AEFConstantKeyLerp_Float96NoW;
+				CompressedData.TranslationCodec = &AEFConstantKeyLerp_Float96NoW;
 				break;
 			case ACF_IntervalFixed32NoW:
-				Seq.TranslationCodec = &AEFConstantKeyLerp_IntervalFixed32NoW;
+				CompressedData.TranslationCodec = &AEFConstantKeyLerp_IntervalFixed32NoW;
 				break;
 			case ACF_Identity:
-				Seq.TranslationCodec = &AEFConstantKeyLerp_Identity;
+				CompressedData.TranslationCodec = &AEFConstantKeyLerp_Identity;
 				break;
 
 			default:
-				UE_LOG(LogAnimationCompression, Fatal, TEXT("%i: unknown or unsupported translation compression"), (int32)Seq.TranslationCompressionFormat );
+				UE_LOG(LogAnimationCompression, Fatal, TEXT("%i: unknown or unsupported translation compression"), (int32)CompressedData.TranslationCompressionFormat );
 		};
 
 		// setup rotation codec
-		switch(Seq.RotationCompressionFormat)
+		switch(CompressedData.RotationCompressionFormat)
 		{
 			case ACF_None:
-				Seq.RotationCodec = &AEFConstantKeyLerp_None;
+				CompressedData.RotationCodec = &AEFConstantKeyLerp_None;
 				break;
 			case ACF_Float96NoW:
-				Seq.RotationCodec = &AEFConstantKeyLerp_Float96NoW;
+				CompressedData.RotationCodec = &AEFConstantKeyLerp_Float96NoW;
 				break;
 			case ACF_Fixed48NoW:
-				Seq.RotationCodec = &AEFConstantKeyLerp_Fixed48NoW;
+				CompressedData.RotationCodec = &AEFConstantKeyLerp_Fixed48NoW;
 				break;
 			case ACF_IntervalFixed32NoW:
-				Seq.RotationCodec = &AEFConstantKeyLerp_IntervalFixed32NoW;
+				CompressedData.RotationCodec = &AEFConstantKeyLerp_IntervalFixed32NoW;
 				break;
 			case ACF_Fixed32NoW:
-				Seq.RotationCodec = &AEFConstantKeyLerp_Fixed32NoW;
+				CompressedData.RotationCodec = &AEFConstantKeyLerp_Fixed32NoW;
 				break;
 			case ACF_Float32NoW:
-				Seq.RotationCodec = &AEFConstantKeyLerp_Float32NoW;
+				CompressedData.RotationCodec = &AEFConstantKeyLerp_Float32NoW;
 				break;
 			case ACF_Identity:
-				Seq.RotationCodec = &AEFConstantKeyLerp_Identity;
+				CompressedData.RotationCodec = &AEFConstantKeyLerp_Identity;
 				break;
 
 			default:
-				UE_LOG(LogAnimationCompression, Fatal, TEXT("%i: unknown or unsupported rotation compression"), (int32)Seq.RotationCompressionFormat );
+				UE_LOG(LogAnimationCompression, Fatal, TEXT("%i: unknown or unsupported rotation compression"), (int32)CompressedData.RotationCompressionFormat );
 		};
 
 		// setup Scale codec
-		switch(Seq.ScaleCompressionFormat)
+		switch(CompressedData.ScaleCompressionFormat)
 		{
 		case ACF_None:
-			Seq.ScaleCodec = &AEFConstantKeyLerp_None;
+			CompressedData.ScaleCodec = &AEFConstantKeyLerp_None;
 			break;
 		case ACF_Float96NoW:
-			Seq.ScaleCodec = &AEFConstantKeyLerp_Float96NoW;
+			CompressedData.ScaleCodec = &AEFConstantKeyLerp_Float96NoW;
 			break;
 		case ACF_IntervalFixed32NoW:
-			Seq.ScaleCodec = &AEFConstantKeyLerp_IntervalFixed32NoW;
+			CompressedData.ScaleCodec = &AEFConstantKeyLerp_IntervalFixed32NoW;
 			break;
 		case ACF_Identity:
-			Seq.ScaleCodec = &AEFConstantKeyLerp_Identity;
+			CompressedData.ScaleCodec = &AEFConstantKeyLerp_Identity;
 			break;
 
 		default:
-			UE_LOG(LogAnimationCompression, Fatal, TEXT("%i: unknown or unsupported Scale compression"), (int32)Seq.ScaleCompressionFormat );
+			UE_LOG(LogAnimationCompression, Fatal, TEXT("%i: unknown or unsupported Scale compression"), (int32)CompressedData.ScaleCompressionFormat );
 		};
 	}
-	else if (Seq.KeyEncodingFormat == AKF_VariableKeyLerp)
+	else if (CompressedData.KeyEncodingFormat == AKF_VariableKeyLerp)
 	{
 		static AEFVariableKeyLerp<ACF_None>					AEFVariableKeyLerp_None;
 		static AEFVariableKeyLerp<ACF_Float96NoW>			AEFVariableKeyLerp_Float96NoW;
@@ -799,91 +770,94 @@ void AnimationFormat_SetInterfaceLinks(UAnimSequence& Seq)
 		static AEFVariableKeyLerp<ACF_Identity>				AEFVariableKeyLerp_Identity;
 
 		// setup translation codec
-		switch(Seq.TranslationCompressionFormat)
+		switch(CompressedData.TranslationCompressionFormat)
 		{
 			case ACF_None:
-				Seq.TranslationCodec = &AEFVariableKeyLerp_None;
+				CompressedData.TranslationCodec = &AEFVariableKeyLerp_None;
 				break;
 			case ACF_Float96NoW:
-				Seq.TranslationCodec = &AEFVariableKeyLerp_Float96NoW;
+				CompressedData.TranslationCodec = &AEFVariableKeyLerp_Float96NoW;
 				break;
 			case ACF_IntervalFixed32NoW:
-				Seq.TranslationCodec = &AEFVariableKeyLerp_IntervalFixed32NoW;
+				CompressedData.TranslationCodec = &AEFVariableKeyLerp_IntervalFixed32NoW;
 				break;
 			case ACF_Identity:
-				Seq.TranslationCodec = &AEFVariableKeyLerp_Identity;
+				CompressedData.TranslationCodec = &AEFVariableKeyLerp_Identity;
 				break;
 
 			default:
-				UE_LOG(LogAnimationCompression, Fatal, TEXT("%i: unknown or unsupported translation compression"), (int32)Seq.TranslationCompressionFormat );
+				UE_LOG(LogAnimationCompression, Fatal, TEXT("%i: unknown or unsupported translation compression"), (int32)CompressedData.TranslationCompressionFormat );
 		};
 
 		// setup rotation codec
-		switch(Seq.RotationCompressionFormat)
+		switch(CompressedData.RotationCompressionFormat)
 		{
 			case ACF_None:
-				Seq.RotationCodec = &AEFVariableKeyLerp_None;
+				CompressedData.RotationCodec = &AEFVariableKeyLerp_None;
 				break;
 			case ACF_Float96NoW:
-				Seq.RotationCodec = &AEFVariableKeyLerp_Float96NoW;
+				CompressedData.RotationCodec = &AEFVariableKeyLerp_Float96NoW;
 				break;
 			case ACF_Fixed48NoW:
-				Seq.RotationCodec = &AEFVariableKeyLerp_Fixed48NoW;
+				CompressedData.RotationCodec = &AEFVariableKeyLerp_Fixed48NoW;
 				break;
 			case ACF_IntervalFixed32NoW:
-				Seq.RotationCodec = &AEFVariableKeyLerp_IntervalFixed32NoW;
+				CompressedData.RotationCodec = &AEFVariableKeyLerp_IntervalFixed32NoW;
 				break;
 			case ACF_Fixed32NoW:
-				Seq.RotationCodec = &AEFVariableKeyLerp_Fixed32NoW;
+				CompressedData.RotationCodec = &AEFVariableKeyLerp_Fixed32NoW;
 				break;
 			case ACF_Float32NoW:
-				Seq.RotationCodec = &AEFVariableKeyLerp_Float32NoW;
+				CompressedData.RotationCodec = &AEFVariableKeyLerp_Float32NoW;
 				break;
 			case ACF_Identity:
-				Seq.RotationCodec = &AEFVariableKeyLerp_Identity;
+				CompressedData.RotationCodec = &AEFVariableKeyLerp_Identity;
 				break;
 
 			default:
-				UE_LOG(LogAnimationCompression, Fatal, TEXT("%i: unknown or unsupported rotation compression"), (int32)Seq.RotationCompressionFormat );
+				UE_LOG(LogAnimationCompression, Fatal, TEXT("%i: unknown or unsupported rotation compression"), (int32)CompressedData.RotationCompressionFormat );
 		};
 
 		// setup Scale codec
-		switch(Seq.ScaleCompressionFormat)
+		switch(CompressedData.ScaleCompressionFormat)
 		{
 		case ACF_None:
-			Seq.ScaleCodec = &AEFVariableKeyLerp_None;
+			CompressedData.ScaleCodec = &AEFVariableKeyLerp_None;
 			break;
 		case ACF_Float96NoW:
-			Seq.ScaleCodec = &AEFVariableKeyLerp_Float96NoW;
+			CompressedData.ScaleCodec = &AEFVariableKeyLerp_Float96NoW;
 			break;
 		case ACF_IntervalFixed32NoW:
-			Seq.ScaleCodec = &AEFVariableKeyLerp_IntervalFixed32NoW;
+			CompressedData.ScaleCodec = &AEFVariableKeyLerp_IntervalFixed32NoW;
 			break;
 		case ACF_Identity:
-			Seq.ScaleCodec = &AEFVariableKeyLerp_Identity;
+			CompressedData.ScaleCodec = &AEFVariableKeyLerp_Identity;
 			break;
 
 		default:
-			UE_LOG(LogAnimationCompression, Fatal, TEXT("%i: unknown or unsupported Scale compression"), (int32)Seq.ScaleCompressionFormat );
+			UE_LOG(LogAnimationCompression, Fatal, TEXT("%i: unknown or unsupported Scale compression"), (int32)CompressedData.ScaleCompressionFormat );
 		};
 	}
-	else if (Seq.KeyEncodingFormat == AKF_PerTrackCompression)
+	else if (CompressedData.KeyEncodingFormat == AKF_PerTrackCompression)
 	{
 		static AEFPerTrackCompressionCodec StaticCodec;
 
-		Seq.RotationCodec = &StaticCodec;
-		Seq.TranslationCodec = &StaticCodec;
-		Seq.ScaleCodec = &StaticCodec;
+		CompressedData.RotationCodec = &StaticCodec;
+		CompressedData.TranslationCodec = &StaticCodec;
+		CompressedData.ScaleCodec = &StaticCodec;
 
-		check(Seq.RotationCompressionFormat == ACF_Identity);
-		check(Seq.TranslationCompressionFormat == ACF_Identity);
+		check(CompressedData.RotationCompressionFormat == ACF_Identity);
+		check(CompressedData.TranslationCompressionFormat == ACF_Identity);
 		// commenting out scale check because the older versions won't have this set correctly
 		// and I can't get the version VER_UE4_ANIM_SUPPORT_NONUNIFORM_SCALE_ANIMATION here because this function
 		// is called in Serialize where GetLinker is too early to call
-		//checkf(Seq.ScaleCompressionFormat == ACF_Identity);
+		//checkf(CompressedData.ScaleCompressionFormat == ACF_Identity);
 	}
 	else
 	{
-		UE_LOG(LogAnimationCompression, Fatal, TEXT("%i: unknown or unsupported animation format"), (int32)Seq.KeyEncodingFormat );
+		UE_LOG(LogAnimationCompression, Fatal, TEXT("%i: unknown or unsupported animation format"), (int32)CompressedData.KeyEncodingFormat );
 	}
 }
+
+template void AnimationFormat_SetInterfaceLinks(FUECompressedAnimData& CompressedData);
+template void AnimationFormat_SetInterfaceLinks(FCompressibleAnimDataResult& CompressedData);
\ No newline at end of file
diff --git a/Engine/Source/Runtime/Engine/Private/Animation/AnimEncoding_ConstantKeyLerp.cpp b/Engine/Source/Runtime/Engine/Private/Animation/AnimEncoding_ConstantKeyLerp.cpp
--- a/Engine/Source/Runtime/Engine/Private/Animation/AnimEncoding_ConstantKeyLerp.cpp
+++ b/Engine/Source/Runtime/Engine/Private/Animation/AnimEncoding_ConstantKeyLerp.cpp
@@ -11,19 +11,19 @@
 /**
  * Handles the ByteSwap of compressed rotation data on import
  *
- * @param	Seq				The UAnimSequence container.
+ * @param	CompressedData		The compressed animation data being operated on.
  * @param	MemoryReader	The FMemoryReader to read from.
  * @param	TrackData		The compressed data stream.
  * @param	NumKeys			The number of keys present in the stream.
  */
 void AEFConstantKeyLerpShared::ByteSwapRotationIn(
-	UAnimSequence& Seq, 
+	FUECompressedAnimData& CompressedData,
 	FMemoryReader& MemoryReader,
 	uint8*& TrackData,
 	int32 NumKeys)
 {
 	// Calculate the effective compression (in a track with only one key, it's always stored lossless)
-	const int32 EffectiveFormat = (NumKeys == 1) ? ACF_Float96NoW : (int32) Seq.RotationCompressionFormat;
+	const int32 EffectiveFormat = (NumKeys == 1) ? ACF_Float96NoW : (int32)CompressedData.RotationCompressionFormat;
 	const int32 KeyComponentSize = CompressedRotationStrides[EffectiveFormat];
 	const int32 KeyNumComponents = CompressedRotationNum[EffectiveFormat];
 
@@ -49,19 +49,19 @@ void AEFConstantKeyLerpShared::ByteSwapRotationIn(
 /**
  * Handles the ByteSwap of compressed translation data on import
  *
- * @param	Seq				The UAnimSequence container.
+ * @param	CompressedData	The compressed animation data being operated on.
  * @param	MemoryReader	The FMemoryReader to read from.
  * @param	TrackData		The compressed data stream.
  * @param	NumKeys			The number of keys present in the stream.
  */
 void AEFConstantKeyLerpShared::ByteSwapTranslationIn(
-	UAnimSequence& Seq, 
+	FUECompressedAnimData& CompressedData,
 	FMemoryReader& MemoryReader,
 	uint8*& TrackData,
 	int32 NumKeys)
 {
 	// Calculate the effective compression (in a track with only one key, it's always stored lossless)
-	const int32 EffectiveFormat = (NumKeys == 1) ? ACF_None : (int32) Seq.TranslationCompressionFormat;
+	const int32 EffectiveFormat = (NumKeys == 1) ? ACF_None : (int32) CompressedData.TranslationCompressionFormat;
 	const int32 KeyComponentSize = CompressedTranslationStrides[EffectiveFormat];
 	const int32 KeyNumComponents = CompressedTranslationNum[EffectiveFormat];
 
@@ -87,19 +87,19 @@ void AEFConstantKeyLerpShared::ByteSwapTranslationIn(
 /**
  * Handles the ByteSwap of compressed Scale data on import
  *
- * @param	Seq				The UAnimSequence container.
+ * @param	CompressedData	The compressed animation data being operated on.
  * @param	MemoryReader	The FMemoryReader to read from.
  * @param	TrackData		The compressed data stream.
  * @param	NumKeys			The number of keys present in the stream.
  */
 void AEFConstantKeyLerpShared::ByteSwapScaleIn(
-	UAnimSequence& Seq, 
+	FUECompressedAnimData& CompressedData,
 	FMemoryReader& MemoryReader,
 	uint8*& TrackData,
 	int32 NumKeys)
 {
 	// Calculate the effective compression (in a track with only one key, it's always stored lossless)
-	const int32 EffectiveFormat = (NumKeys == 1) ? ACF_None : (int32) Seq.ScaleCompressionFormat;
+	const int32 EffectiveFormat = (NumKeys == 1) ? ACF_None : (int32)CompressedData.ScaleCompressionFormat;
 	const int32 KeyComponentSize = CompressedScaleStrides[EffectiveFormat];
 	const int32 KeyNumComponents = CompressedScaleNum[EffectiveFormat];
 
@@ -124,19 +124,19 @@ void AEFConstantKeyLerpShared::ByteSwapScaleIn(
 /**
  * Handles the ByteSwap of compressed rotation data on export
  *
- * @param	Seq				The UAnimSequence container.
+ * @param	CompressedData	The compressed animation data being operated on.
  * @param	MemoryWriter	The FMemoryWriter to write to.
  * @param	TrackData		The compressed data stream.
  * @param	NumKeys			The number of keys to write to the stream.
  */
 void AEFConstantKeyLerpShared::ByteSwapRotationOut(
-	UAnimSequence& Seq, 
+	FUECompressedAnimData& CompressedData,
 	FMemoryWriter& MemoryWriter,
 	uint8*& TrackData,
 	int32 NumKeys)
 {
 	// Calculate the effective compression (in a track with only one key, it's always stored lossless)
-	const int32 EffectiveFormat = (NumKeys == 1) ? ACF_Float96NoW : (int32) Seq.RotationCompressionFormat;
+	const int32 EffectiveFormat = (NumKeys == 1) ? ACF_Float96NoW : (int32) CompressedData.RotationCompressionFormat;
 	const int32 KeyComponentSize = CompressedRotationStrides[EffectiveFormat];
 	const int32 KeyNumComponents = CompressedRotationNum[EffectiveFormat];
 
@@ -162,19 +162,19 @@ void AEFConstantKeyLerpShared::ByteSwapRotationOut(
 /**
  * Handles the ByteSwap of compressed translation data on export
  *
- * @param	Seq				The UAnimSequence container.
+ * @param	CompressedData	The compressed animation data being operated on.
  * @param	MemoryWriter	The FMemoryWriter to write to.
  * @param	TrackData		The compressed data stream.
  * @param	NumKeys			The number of keys to write to the stream.
  */
 void AEFConstantKeyLerpShared::ByteSwapTranslationOut(
-	UAnimSequence& Seq, 
+	FUECompressedAnimData& CompressedData,
 	FMemoryWriter& MemoryWriter,
 	uint8*& TrackData,
 	int32 NumKeys)
 {
 	// Calculate the effective compression (in a track with only one key, it's always stored lossless)
-	const int32 EffectiveFormat = (NumKeys == 1) ? ACF_None : (int32) Seq.TranslationCompressionFormat;
+	const int32 EffectiveFormat = (NumKeys == 1) ? ACF_None : (int32) CompressedData.TranslationCompressionFormat;
 	const int32 KeyComponentSize = CompressedTranslationStrides[EffectiveFormat];
 	const int32 KeyNumComponents = CompressedTranslationNum[EffectiveFormat];
 
@@ -200,19 +200,19 @@ void AEFConstantKeyLerpShared::ByteSwapTranslationOut(
 /**
  * Handles the ByteSwap of compressed Scale data on export
  *
- * @param	Seq				The UAnimSequence container.
+ * @param	CompressedData	The compressed animation data being operated on.
  * @param	MemoryWriter	The FMemoryWriter to write to.
  * @param	TrackData		The compressed data stream.
  * @param	NumKeys			The number of keys to write to the stream.
  */
 void AEFConstantKeyLerpShared::ByteSwapScaleOut(
-	UAnimSequence& Seq, 
+	FUECompressedAnimData& CompressedData,
 	FMemoryWriter& MemoryWriter,
 	uint8*& TrackData,
 	int32 NumKeys)
 {
 	// Calculate the effective compression (in a track with only one key, it's always stored lossless)
-	const int32 EffectiveFormat = (NumKeys == 1) ? ACF_None : (int32) Seq.ScaleCompressionFormat;
+	const int32 EffectiveFormat = (NumKeys == 1) ? ACF_None : (int32) CompressedData.ScaleCompressionFormat;
 	const int32 KeyComponentSize = CompressedScaleStrides[EffectiveFormat];
 	const int32 KeyNumComponents = CompressedScaleNum[EffectiveFormat];
 
diff --git a/Engine/Source/Runtime/Engine/Private/Animation/AnimEncoding_PerTrackCompression.cpp b/Engine/Source/Runtime/Engine/Private/Animation/AnimEncoding_PerTrackCompression.cpp
--- a/Engine/Source/Runtime/Engine/Private/Animation/AnimEncoding_PerTrackCompression.cpp
+++ b/Engine/Source/Runtime/Engine/Private/Animation/AnimEncoding_PerTrackCompression.cpp
@@ -345,26 +345,17 @@ static FORCEINLINE_DEBUGGABLE VectorRegister DecompressSingleTrackRotationVector
 
 #endif // USE_VECTOR_PTC_DECOMPRESSOR
 
-/**
- * Handles Byte-swapping a single track of animation data from a MemoryReader or to a MemoryWriter
- *
- * @param	Seq					The Animation Sequence being operated on.
- * @param	MemoryStream		The MemoryReader or MemoryWriter object to read from/write to.
- * @param	Offset				The starting offset into the compressed byte stream for this track (can be INDEX_NONE to indicate an identity track)
- */
 template<class TArchive>
-void AEFPerTrackCompressionCodec::ByteSwapOneTrack(UAnimSequence& Seq, TArchive& MemoryStream, int32 Offset, bool bMaintainComponentOrder)
+void AEFPerTrackCompressionCodec::ByteSwapOneTrack(FUECompressedAnimData& CompressedData, TArchive& MemoryStream, int32 BufferStart, int32 Offset)
 {
 	// Translation data.
 	if (Offset != INDEX_NONE)
 	{
 		checkSlow( (Offset % 4) == 0 && "CompressedByteStream not aligned to four bytes" );
 
-		if (bMaintainComponentOrder)
-		{
-			MemoryStream.Seek(Offset);
-		}
-		uint8* TrackData = Seq.CompressedByteStream.GetData() + Offset;
+		MemoryStream.Seek(BufferStart+Offset);
+
+		uint8* TrackData = CompressedData.CompressedByteStream.GetData() + Offset;
 
 		// Read the header
 		AC_UnalignedSwap(MemoryStream, TrackData, sizeof(int32));
@@ -383,6 +374,9 @@ void AEFPerTrackCompressionCodec::ByteSwapOneTrack(UAnimSequence& Seq, TArchive&
 		int32 KeyComponentCount = 0;
 		FAnimationCompression_PerTrackUtils::GetAllSizesFromFormat(KeyFormat, FormatFlags, /*OUT*/ KeyComponentCount, /*OUT*/ KeyComponentSize, /*OUT*/ FixedComponentCount, /*OUT*/ FixedComponentSize);
 
+		check(FixedComponentSize != 0);
+		check(KeyComponentSize != 0);
+
 		// Handle per-track metadata
 		for (int32 i = 0; i < FixedComponentCount; ++i)
 		{
@@ -404,7 +398,7 @@ void AEFPerTrackCompressionCodec::ByteSwapOneTrack(UAnimSequence& Seq, TArchive&
 			// Make sure the key->frame table is 4 byte aligned
 			PreservePadding(TrackData, MemoryStream);
 
-			const int32 FrameTableEntrySize = (Seq.GetCompressedNumberOfFrames() <= 0xFF) ? sizeof(uint8) : sizeof(uint16);
+			const int32 FrameTableEntrySize = (CompressedData.CompressedNumberOfFrames <= 0xFF) ? sizeof(uint8) : sizeof(uint16);
 			for (int32 i = 0; i < NumKeys; ++i)
 			{
 				AC_UnalignedSwap(MemoryStream, TrackData, FrameTableEntrySize);
@@ -416,8 +410,8 @@ void AEFPerTrackCompressionCodec::ByteSwapOneTrack(UAnimSequence& Seq, TArchive&
 	}
 }
 
-template void AEFPerTrackCompressionCodec::ByteSwapOneTrack(UAnimSequence& Seq, FMemoryReader& MemoryStream, int32 Offset, bool bMaintainComponentOrder);
-template void AEFPerTrackCompressionCodec::ByteSwapOneTrack(UAnimSequence& Seq, FMemoryWriter& MemoryStream, int32 Offset, bool bMaintainComponentOrder);
+template void AEFPerTrackCompressionCodec::ByteSwapOneTrack(FUECompressedAnimData& CompressedData, FMemoryReader& MemoryStream, int32 BufferStart, int32 Offset);
+template void AEFPerTrackCompressionCodec::ByteSwapOneTrack(FUECompressedAnimData& CompressedData, FMemoryWriter& MemoryStream, int32 BufferStart, int32 Offset);
 
 
 /**
@@ -451,42 +445,33 @@ void AEFPerTrackCompressionCodec::PreservePadding(uint8*& TrackData, FMemoryArch
 /**
  * Handles Byte-swapping incoming animation data from a MemoryReader
  *
- * @param	Seq					An Animation Sequence to contain the read data.
+ * @param	CompressedData		The compressed animation data being operated on.
  * @param	MemoryReader		The MemoryReader object to read from.
  */
 void AEFPerTrackCompressionCodec::ByteSwapIn(
-	UAnimSequence& Seq, 
+	FUECompressedAnimData& CompressedData,
 	FMemoryReader& MemoryReader)
 {
-	int32 OriginalNumBytes = MemoryReader.TotalSize();
-	Seq.CompressedByteStream.Empty(OriginalNumBytes);
-	Seq.CompressedByteStream.AddUninitialized(OriginalNumBytes);
-
-	if (Seq.CompressedSegments.Num() != 0)
-	{
-#if !PLATFORM_LITTLE_ENDIAN
-#error "Byte swapping needs to be implemented here to support big-endian platforms"
-#endif
-
-		// TODO: Byte swap the new format
-		MemoryReader.Serialize(Seq.CompressedByteStream.GetData(), Seq.CompressedByteStream.Num());
-		return;
-	}
+	const int64 BufferStart = MemoryReader.Tell();
+	int64 OriginalNumBytes = MemoryReader.TotalSize() - BufferStart;
+	check(OriginalNumBytes < INT_MAX);
+	int32 OriginalNumBytes32 = (int32)OriginalNumBytes;
+	check(CompressedData.CompressedByteStream.Num() == OriginalNumBytes);
 
-	const int32 NumTracks = Seq.CompressedTrackOffsets.Num() / 2;
-	const bool bHasScaleData = Seq.CompressedScaleOffsets.IsValid();
+	const int32 NumTracks = CompressedData.CompressedTrackOffsets.Num() / 2;
+	const bool bHasScaleData = CompressedData.CompressedScaleOffsets.IsValid();
 
 	for ( int32 TrackIndex = 0; TrackIndex < NumTracks; ++TrackIndex )
 	{
-		const int32 OffsetTrans = Seq.CompressedTrackOffsets[TrackIndex*2+0];
-		ByteSwapOneTrack(Seq, MemoryReader, OffsetTrans);
+		const int32 OffsetTrans = CompressedData.CompressedTrackOffsets[TrackIndex*2+0];
+		ByteSwapOneTrack(CompressedData, MemoryReader, BufferStart, OffsetTrans);
 
-		const int32 OffsetRot = Seq.CompressedTrackOffsets[TrackIndex*2+1];
-		ByteSwapOneTrack(Seq, MemoryReader, OffsetRot);
+		const int32 OffsetRot = CompressedData.CompressedTrackOffsets[TrackIndex*2+1];
+		ByteSwapOneTrack(CompressedData, MemoryReader, BufferStart, OffsetRot);
 		if (bHasScaleData)
 		{
-			const int32 OffsetScale = Seq.CompressedScaleOffsets.GetOffsetData(TrackIndex, 0);
-			ByteSwapOneTrack(Seq, MemoryReader, OffsetScale);
+			const int32 OffsetScale = CompressedData.CompressedScaleOffsets.GetOffsetData(TrackIndex, 0);
+			ByteSwapOneTrack(CompressedData, MemoryReader, BufferStart, OffsetScale);
 		}
 	}
 }
@@ -495,40 +480,30 @@ void AEFPerTrackCompressionCodec::ByteSwapIn(
 /**
  * Handles Byte-swapping outgoing animation data to an array of BYTEs
  *
- * @param	Seq					An Animation Sequence to write.
+ * @param	CompressedData		The compressed animation data being operated on.
  * @param	SerializedData		The output buffer.
  * @param	ForceByteSwapping	true is byte swapping is not optional.
  */
 void AEFPerTrackCompressionCodec::ByteSwapOut(
-	UAnimSequence& Seq,
-	TArray<uint8>& SerializedData, 
-	bool ForceByteSwapping,
-	bool bMaintainComponentOrder)
+	FUECompressedAnimData& CompressedData,
+	FMemoryWriter& MemoryWriter)
 {
-	FMemoryWriter MemoryWriter(SerializedData, true);
-	MemoryWriter.SetByteSwapping(ForceByteSwapping);
-
-	if (Seq.CompressedSegments.Num() != 0)
-	{
-		// TODO: Byte swap the new format
-		MemoryWriter.Serialize(Seq.CompressedByteStream.GetData(), Seq.CompressedByteStream.Num());
-		return;
-	}
+	const int32 BufferStart = MemoryWriter.Tell();
 
-	const int32 NumTracks = Seq.CompressedTrackOffsets.Num() / 2;
-	const bool bHasScaleData = Seq.CompressedScaleOffsets.IsValid();
+	const int32 NumTracks = CompressedData.CompressedTrackOffsets.Num() / 2;
+	const bool bHasScaleData = CompressedData.CompressedScaleOffsets.IsValid();
 	for ( int32 TrackIndex = 0; TrackIndex < NumTracks; ++TrackIndex )
 	{
-		const int32 OffsetTrans = Seq.CompressedTrackOffsets[TrackIndex*2+0];
-		ByteSwapOneTrack(Seq, MemoryWriter, OffsetTrans, bMaintainComponentOrder);
+		const int32 OffsetTrans = CompressedData.CompressedTrackOffsets[TrackIndex*2+0];
+		ByteSwapOneTrack(CompressedData, MemoryWriter, BufferStart, OffsetTrans);
 
-		const int32 OffsetRot = Seq.CompressedTrackOffsets[TrackIndex*2+1];
-		ByteSwapOneTrack(Seq, MemoryWriter, OffsetRot, bMaintainComponentOrder);
+		const int32 OffsetRot = CompressedData.CompressedTrackOffsets[TrackIndex*2+1];
+		ByteSwapOneTrack(CompressedData, MemoryWriter, BufferStart, OffsetRot);
 
 		if (bHasScaleData)
 		{
-			const int32 OffsetScale = Seq.CompressedScaleOffsets.GetOffsetData(TrackIndex, 0);
-			ByteSwapOneTrack(Seq, MemoryWriter, OffsetScale, bMaintainComponentOrder);
+			const int32 OffsetScale = CompressedData.CompressedScaleOffsets.GetOffsetData(TrackIndex, 0);
+			ByteSwapOneTrack(CompressedData, MemoryWriter, BufferStart, OffsetScale);
 		}
 	}
 }
@@ -608,12 +583,12 @@ void AEFPerTrackCompressionCodec::GetBoneAtomRotation(
 		{
 			if ((FormatFlags & 0x8) == 0)
 			{
-				Alpha = TimeToIndex(*DecompContext.AnimSeq, DecompContext.RelativePos, NumKeys, Index0, Index1);
+				Alpha = TimeToIndex(DecompContext.GetSequenceLength(), DecompContext.RelativePos, NumKeys, DecompContext.GetInterpolation(), Index0, Index1);
 			}
 			else
 			{
 				const uint8* RESTRICT FrameTable = Align(TrackData + FixedBytes + BytesPerKey * NumKeys, 4);
-				Alpha = TimeToIndex(*DecompContext.AnimSeq, FrameTable, DecompContext.RelativePos, NumKeys, Index0, Index1);
+				Alpha = TimeToIndex(DecompContext.GetInterpolation(), DecompContext.GetCompressedNumberOfFrames(), FrameTable, DecompContext.RelativePos, NumKeys, Index0, Index1);
 			}
 		}
 
@@ -702,8 +677,9 @@ void AEFPerTrackCompressionCodec::GetBoneAtomTranslation(
 		int32 FixedBytes;
 		FAnimationCompression_PerTrackUtils::DecomposeHeader(Header, /*OUT*/ KeyFormat, /*OUT*/ NumKeys, /*OUT*/ FormatFlags, /*OUT*/BytesPerKey, /*OUT*/ FixedBytes);
 
-		checkf(KeyFormat != ACF_None, TEXT("[%s] contians invalid keyformat. NumKeys (%d), FormatFlags (%d), BytesPerKeys (%d), FixedBytes (%d)"), *DecompContext.AnimSeq->GetName(), NumKeys, FormatFlags, BytesPerKey, FixedBytes);
+		checkf(KeyFormat != ACF_None, TEXT("[%s] contians invalid keyformat. NumKeys (%d), FormatFlags (%d), BytesPerKeys (%d), FixedBytes (%d)"), *DecompContext.GetAnimFName().ToString(), NumKeys, FormatFlags, BytesPerKey, FixedBytes);
 
+		checkf(KeyFormat < ACF_MAX, TEXT("[%s] contians invalid keyformat. NumKeys (%d), FormatFlags (%d), BytesPerKeys (%d), FixedBytes (%d), PosKeysOffset (%d), TrackIndex (%d)"), *DecompContext.GetAnimFName().ToString(), NumKeys, FormatFlags, BytesPerKey, FixedBytes, PosKeysOffset, TrackIndex)
 		// Figure out the key indexes
 		int32 Index0 = 0;
 		int32 Index1 = 0;
@@ -715,12 +691,12 @@ void AEFPerTrackCompressionCodec::GetBoneAtomTranslation(
 		{
 			if ((FormatFlags & 0x8) == 0)
 			{
-				Alpha = TimeToIndex(*DecompContext.AnimSeq, DecompContext.RelativePos, NumKeys, Index0, Index1);
+				Alpha = TimeToIndex(DecompContext.GetSequenceLength(), DecompContext.RelativePos, NumKeys, DecompContext.GetInterpolation(), Index0, Index1);
 			}
 			else
 			{
 				const uint8* RESTRICT FrameTable = Align(TrackData + FixedBytes + BytesPerKey * NumKeys, 4);
-				Alpha = TimeToIndex(*DecompContext.AnimSeq, FrameTable, DecompContext.RelativePos, NumKeys, Index0, Index1);
+				Alpha = TimeToIndex(DecompContext.GetInterpolation(), DecompContext.GetCompressedNumberOfFrames(), FrameTable, DecompContext.RelativePos, NumKeys, Index0, Index1);
 			}
 		}
 
@@ -819,12 +795,12 @@ void AEFPerTrackCompressionCodec::GetBoneAtomScale(
 		{
 			if ((FormatFlags & 0x8) == 0)
 			{
-				Alpha = TimeToIndex(*DecompContext.AnimSeq, DecompContext.RelativePos, NumKeys, Index0, Index1);
+				Alpha = TimeToIndex(DecompContext.GetSequenceLength(), DecompContext.RelativePos, NumKeys, DecompContext.GetInterpolation(), Index0, Index1);
 			}
 			else
 			{
 				const uint8* RESTRICT FrameTable = Align(TrackData + FixedBytes + BytesPerKey * NumKeys, 4);
-				Alpha = TimeToIndex(*DecompContext.AnimSeq, FrameTable, DecompContext.RelativePos, NumKeys, Index0, Index1);
+				Alpha = TimeToIndex(DecompContext.GetInterpolation(), DecompContext.GetCompressedNumberOfFrames(), FrameTable, DecompContext.RelativePos, NumKeys, Index0, Index1);
 			}
 		}
 
diff --git a/Engine/Source/Runtime/Engine/Private/Animation/AnimEncoding_VariableKeyLerp.cpp b/Engine/Source/Runtime/Engine/Private/Animation/AnimEncoding_VariableKeyLerp.cpp
--- a/Engine/Source/Runtime/Engine/Private/Animation/AnimEncoding_VariableKeyLerp.cpp
+++ b/Engine/Source/Runtime/Engine/Private/Animation/AnimEncoding_VariableKeyLerp.cpp
@@ -11,18 +11,18 @@
 /**
  * Handles the ByteSwap of compressed rotation data on import
  *
- * @param	Seq				The UAnimSequence container.
+ * @param	CompressedData		The compressed animation data being operated on.
  * @param	MemoryReader	The FMemoryReader to read from.
  * @param	TrackData		The compressed data stream.
  * @param	NumKeys			The number of keys present in the stream.
  */
 void AEFVariableKeyLerpShared::ByteSwapRotationIn(
-	UAnimSequence& Seq, 
+	FUECompressedAnimData& CompressedData,
 	FMemoryReader& MemoryReader,
 	uint8*& TrackData,
 	int32 NumKeys)
 {
-	AEFConstantKeyLerpShared::ByteSwapRotationIn(Seq, MemoryReader, TrackData, NumKeys);
+	AEFConstantKeyLerpShared::ByteSwapRotationIn(CompressedData, MemoryReader, TrackData, NumKeys);
 
 	// Load the track table if present
 	if (NumKeys > 1)
@@ -32,7 +32,7 @@ void AEFVariableKeyLerpShared::ByteSwapRotationIn(
 		PadMemoryReader(&MemoryReader, TrackData, 4); 
 
 		// swap the track table
-		const size_t EntryStride = (Seq.GetCompressedNumberOfFrames() > 0xFF) ? sizeof(uint16) : sizeof(uint8);
+		const size_t EntryStride = (CompressedData.CompressedNumberOfFrames > 0xFF) ? sizeof(uint16) : sizeof(uint8);
 		for (int32 KeyIndex = 0; KeyIndex < NumKeys; ++KeyIndex)
 		{
 			AC_UnalignedSwap(MemoryReader, TrackData, EntryStride);
@@ -43,18 +43,18 @@ void AEFVariableKeyLerpShared::ByteSwapRotationIn(
 /**
  * Handles the ByteSwap of compressed translation data on import
  *
- * @param	Seq				The UAnimSequence container.
+ * @param	CompressedData	The compressed animation data being operated on.
  * @param	MemoryReader	The FMemoryReader to read from.
  * @param	TrackData		The compressed data stream.
  * @param	NumKeys			The number of keys present in the stream.
  */
 void AEFVariableKeyLerpShared::ByteSwapTranslationIn(
-	UAnimSequence& Seq, 
+	FUECompressedAnimData& CompressedData,
 	FMemoryReader& MemoryReader,
 	uint8*& TrackData,
 	int32 NumKeys)
 {
-	AEFConstantKeyLerpShared::ByteSwapTranslationIn(Seq, MemoryReader, TrackData, NumKeys);
+	AEFConstantKeyLerpShared::ByteSwapTranslationIn(CompressedData, MemoryReader, TrackData, NumKeys);
 
 	// Load the track table if present
 	if (NumKeys > 1)
@@ -64,7 +64,7 @@ void AEFVariableKeyLerpShared::ByteSwapTranslationIn(
 		PadMemoryReader(&MemoryReader, TrackData, 4); 
 
 		// swap the track table
-		const size_t EntryStride = (Seq.GetCompressedNumberOfFrames() > 0xFF) ? sizeof(uint16) : sizeof(uint8);
+		const size_t EntryStride = (CompressedData.CompressedNumberOfFrames > 0xFF) ? sizeof(uint16) : sizeof(uint8);
 		for (int32 KeyIndex = 0; KeyIndex < NumKeys; ++KeyIndex)
 		{
 			AC_UnalignedSwap(MemoryReader, TrackData, EntryStride);
@@ -75,18 +75,18 @@ void AEFVariableKeyLerpShared::ByteSwapTranslationIn(
 /**
  * Handles the ByteSwap of compressed Scale data on import
  *
- * @param	Seq				The UAnimSequence container.
+ * @param	CompressedData	The compressed animation data being operated on.
  * @param	MemoryReader	The FMemoryReader to read from.
  * @param	TrackData		The compressed data stream.
  * @param	NumKeys			The number of keys present in the stream.
  */
 void AEFVariableKeyLerpShared::ByteSwapScaleIn(
-	UAnimSequence& Seq, 
+	FUECompressedAnimData& CompressedData,
 	FMemoryReader& MemoryReader,
 	uint8*& TrackData,
 	int32 NumKeys)
 {
-	AEFConstantKeyLerpShared::ByteSwapScaleIn(Seq, MemoryReader, TrackData, NumKeys);
+	AEFConstantKeyLerpShared::ByteSwapScaleIn(CompressedData, MemoryReader, TrackData, NumKeys);
 
 	// Load the track table if present
 	if (NumKeys > 1)
@@ -96,7 +96,7 @@ void AEFVariableKeyLerpShared::ByteSwapScaleIn(
 		PadMemoryReader(&MemoryReader, TrackData, 4); 
 
 		// swap the track table
-		const size_t EntryStride = (Seq.GetCompressedNumberOfFrames() > 0xFF) ? sizeof(uint16) : sizeof(uint8);
+		const size_t EntryStride = (CompressedData.CompressedNumberOfFrames > 0xFF) ? sizeof(uint16) : sizeof(uint8);
 		for (int32 KeyIndex = 0; KeyIndex < NumKeys; ++KeyIndex)
 		{
 			AC_UnalignedSwap(MemoryReader, TrackData, EntryStride);
@@ -106,18 +106,18 @@ void AEFVariableKeyLerpShared::ByteSwapScaleIn(
 /**
  * Handles the ByteSwap of compressed rotation data on export
  *
- * @param	Seq				The UAnimSequence container.
+ * @param	CompressedData	The compressed animation data being operated on.
  * @param	MemoryWriter	The FMemoryWriter to write to.
  * @param	TrackData		The compressed data stream.
  * @param	NumKeys			The number of keys to write to the stream.
  */
 void AEFVariableKeyLerpShared::ByteSwapRotationOut(
-	UAnimSequence& Seq, 
+	FUECompressedAnimData& CompressedData,
 	FMemoryWriter& MemoryWriter,
 	uint8*& TrackData,
 	int32 NumKeys)
 {
-	AEFConstantKeyLerpShared::ByteSwapRotationOut(Seq, MemoryWriter, TrackData, NumKeys);
+	AEFConstantKeyLerpShared::ByteSwapRotationOut(CompressedData, MemoryWriter, TrackData, NumKeys);
 
 	// Store the track table if needed
 	if (NumKeys > 1)
@@ -127,7 +127,7 @@ void AEFVariableKeyLerpShared::ByteSwapRotationOut(
 		PadMemoryWriter(&MemoryWriter, TrackData, 4);
 
 		// swap the track table
-		const size_t EntryStride = (Seq.GetCompressedNumberOfFrames() > 0xFF) ? sizeof(uint16) : sizeof(uint8);
+		const size_t EntryStride = (CompressedData.CompressedNumberOfFrames > 0xFF) ? sizeof(uint16) : sizeof(uint8);
 		for (int32 KeyIndex = 0; KeyIndex < NumKeys; ++KeyIndex)
 		{
 			AC_UnalignedSwap(MemoryWriter, TrackData, EntryStride);
@@ -138,18 +138,18 @@ void AEFVariableKeyLerpShared::ByteSwapRotationOut(
 /**
  * Handles the ByteSwap of compressed translation data on export
  *
- * @param	Seq				The UAnimSequence container.
+ * @param	CompressedData	The compressed animation data being operated on.
  * @param	MemoryWriter	The FMemoryWriter to write to.
  * @param	TrackData		The compressed data stream.
  * @param	NumKeys			The number of keys to write to the stream.
  */
 void AEFVariableKeyLerpShared::ByteSwapTranslationOut(
-	UAnimSequence& Seq, 
+	FUECompressedAnimData& CompressedData,
 	FMemoryWriter& MemoryWriter,
 	uint8*& TrackData,
 	int32 NumKeys)
 {
-	AEFConstantKeyLerpShared::ByteSwapTranslationOut(Seq, MemoryWriter, TrackData, NumKeys);
+	AEFConstantKeyLerpShared::ByteSwapTranslationOut(CompressedData, MemoryWriter, TrackData, NumKeys);
 
 	// Store the track table if needed
 	if (NumKeys > 1)
@@ -159,7 +159,7 @@ void AEFVariableKeyLerpShared::ByteSwapTranslationOut(
 		PadMemoryWriter(&MemoryWriter, TrackData, 4);
 
 		// swap the track table
-		const size_t EntryStride = (Seq.GetCompressedNumberOfFrames() > 0xFF) ? sizeof(uint16) : sizeof(uint8);
+		const size_t EntryStride = (CompressedData.CompressedNumberOfFrames > 0xFF) ? sizeof(uint16) : sizeof(uint8);
 		for (int32 KeyIndex = 0; KeyIndex < NumKeys; ++KeyIndex)
 		{
 			AC_UnalignedSwap(MemoryWriter, TrackData, EntryStride);
@@ -171,18 +171,18 @@ void AEFVariableKeyLerpShared::ByteSwapTranslationOut(
 /**
  * Handles the ByteSwap of compressed Scale data on export
  *
- * @param	Seq				The UAnimSequence container.
+ * @param	CompressedData	The compressed animation data being operated on.
  * @param	MemoryWriter	The FMemoryWriter to write to.
  * @param	TrackData		The compressed data stream.
  * @param	NumKeys			The number of keys to write to the stream.
  */
 void AEFVariableKeyLerpShared::ByteSwapScaleOut(
-	UAnimSequence& Seq, 
+	FUECompressedAnimData& CompressedData,
 	FMemoryWriter& MemoryWriter,
 	uint8*& TrackData,
 	int32 NumKeys)
 {
-	AEFConstantKeyLerpShared::ByteSwapScaleOut(Seq, MemoryWriter, TrackData, NumKeys);
+	AEFConstantKeyLerpShared::ByteSwapScaleOut(CompressedData, MemoryWriter, TrackData, NumKeys);
 
 	// Store the track table if needed
 	if (NumKeys > 1)
@@ -192,7 +192,7 @@ void AEFVariableKeyLerpShared::ByteSwapScaleOut(
 		PadMemoryWriter(&MemoryWriter, TrackData, 4);
 
 		// swap the track table
-		const size_t EntryStride = (Seq.GetCompressedNumberOfFrames() > 0xFF) ? sizeof(uint16) : sizeof(uint8);
+		const size_t EntryStride = (CompressedData.CompressedNumberOfFrames > 0xFF) ? sizeof(uint16) : sizeof(uint8);
 		for (int32 KeyIndex = 0; KeyIndex < NumKeys; ++KeyIndex)
 		{
 			AC_UnalignedSwap(MemoryWriter, TrackData, EntryStride);
