Refactor UAnimSequence out of compression code

diff --git a/Engine/Source/Runtime/Engine/Private/Animation/AnimationUtils.cpp b/Engine/Source/Runtime/Engine/Private/Animation/AnimationUtils.cpp
--- a/Engine/Source/Runtime/Engine/Private/Animation/AnimationUtils.cpp
+++ b/Engine/Source/Runtime/Engine/Private/Animation/AnimationUtils.cpp
@@ -7,6 +7,7 @@
 #include "AnimationUtils.h"
 #include "Misc/ConfigCacheIni.h"
 #include "UObject/Package.h"
+#include "AnimationRuntime.h"
 #include "Animation/AnimCompress.h"
 #include "Animation/AnimCompress_BitwiseCompressOnly.h"
 #include "Animation/AnimCompress_RemoveLinearKeys.h"
@@ -213,6 +214,15 @@ void FAnimationUtils::BuildComponentSpaceTransform(FTransform& OutTransform,
 	}
 }
 
+int32 FAnimationUtils::GetAnimTrackIndexForSkeletonBone(const int32 InSkeletonBoneIndex, const TArray<FTrackToSkeletonMap>& TrackToSkelMap)
+{
+	return TrackToSkelMap.IndexOfByPredicate([&](const FTrackToSkeletonMap& TrackToSkel)
+	{ 
+		return TrackToSkel.BoneTreeIndex == InSkeletonBoneIndex;
+	});
+}
+
+#if WITH_EDITOR
 /**
  * Utility function to measure the accuracy of a compressed animation. Each end-effector is checked for 
  * world-space movement as a result of compression.
@@ -222,7 +232,7 @@ void FAnimationUtils::BuildComponentSpaceTransform(FTransform& OutTransform,
  * @param	ErrorStats	Output structure containing the final compression error values
  * @return				None.
  */
-void FAnimationUtils::ComputeCompressionError(const UAnimSequence* AnimSeq, const TArray<FBoneData>& BoneData, AnimationErrorStats& ErrorStats)
+void FAnimationUtils::ComputeCompressionError(const FCompressibleAnimData& CompressibleAnimData, FCompressibleAnimDataResult& CompressedData, AnimationErrorStats& ErrorStats)
 {
 	ErrorStats.AverageError = 0.0f;
 	ErrorStats.MaxError = 0.0f;
@@ -230,21 +240,21 @@ void FAnimationUtils::ComputeCompressionError(const UAnimSequence* AnimSeq, cons
 	ErrorStats.MaxErrorTime = 0.0f;
 	int32 MaxErrorTrack = -1;
 
-	if (AnimSeq->GetCompressedNumberOfFrames() > 0)
+	if (CompressedData.CompressedNumberOfFrames > 0)
 	{
-		const bool bCanUseCompressedData = (AnimSeq->CompressedByteStream.Num() > 0);
+		const bool bCanUseCompressedData = (CompressedData.CompressedByteStream.Num() > 0);
 		if (!bCanUseCompressedData)
 		{
 			// If we can't use CompressedData, there's not much point in being here.
 			return;
 		}
 
-		const int32 NumBones = BoneData.Num();
+		const int32 NumBones = CompressibleAnimData.BoneData.Num();
 		
 		float ErrorCount = 0.0f;
 		float ErrorTotal = 0.0f;
 
-		USkeleton* Skeleton = AnimSeq->GetSkeleton();
+		USkeleton* Skeleton = CompressibleAnimData.Skeleton;
 		check ( Skeleton );
 
 		const TArray<FTransform>& RefPose = Skeleton->GetRefLocalPoses();
@@ -258,8 +268,7 @@ void FAnimationUtils::ComputeCompressionError(const UAnimSequence* AnimSeq, cons
 		// We do this only once, instead of every frame.
 		struct FCachedBoneIndexData
 		{
-			int32 RawTrackIndex;
-			int32 CompressedTrackIndex;
+			int32 TrackIndex;
 			int32 ParentIndex;
 		};
 		TArray<FCachedBoneIndexData> CachedBoneIndexData;
@@ -267,8 +276,7 @@ void FAnimationUtils::ComputeCompressionError(const UAnimSequence* AnimSeq, cons
 		for (int32 BoneIndex = 0; BoneIndex < NumBones; ++BoneIndex)
 		{
 			FCachedBoneIndexData& BoneIndexData = CachedBoneIndexData[BoneIndex];
-			BoneIndexData.RawTrackIndex = Skeleton->GetAnimationTrackIndex(BoneIndex, AnimSeq, true);
-			BoneIndexData.CompressedTrackIndex = Skeleton->GetAnimationTrackIndex(BoneIndex, AnimSeq, false);
+			BoneIndexData.TrackIndex = GetAnimTrackIndexForSkeletonBone(BoneIndex, CompressibleAnimData.TrackToSkeletonMapTable);
 			BoneIndexData.ParentIndex = Skeleton->GetReferenceSkeleton().GetParentIndex(BoneIndex);
 		}
 
@@ -282,12 +290,15 @@ void FAnimationUtils::ComputeCompressionError(const UAnimSequence* AnimSeq, cons
 
 		const FTransform EndEffectorDummyBoneSocket(FQuat::Identity, FVector(END_EFFECTOR_DUMMY_BONE_LENGTH_SOCKET));
 		const FTransform EndEffectorDummyBone(FQuat::Identity, FVector(END_EFFECTOR_DUMMY_BONE_LENGTH));
-		const FAnimKeyHelper Helper(AnimSeq->SequenceLength, AnimSeq->GetCompressedNumberOfFrames());
+		const FAnimKeyHelper Helper(CompressibleAnimData.SequenceLength, CompressedData.CompressedNumberOfFrames);
 		const float KeyLength = Helper.TimePerKey() + SMALL_NUMBER;
 
-		FAnimSequenceDecompressionContext DecompContext(AnimSeq);
+		const FUECompressedAnimData CompressedDataWrapper(CompressedData);
+		FAnimSequenceDecompressionContext DecompContext(CompressibleAnimData, CompressedDataWrapper);
+
+		const TArray<FBoneData>& BoneData = CompressibleAnimData.BoneData;
 
-		for (int32 FrameIndex = 0; FrameIndex<AnimSeq->GetCompressedNumberOfFrames(); FrameIndex++)
+		for (int32 FrameIndex = 0; FrameIndex< CompressedData.CompressedNumberOfFrames; FrameIndex++)
 		{
 			const float Time = (float)FrameIndex * KeyLength;
 			DecompContext.Seek(Time);
@@ -296,7 +307,7 @@ void FAnimationUtils::ComputeCompressionError(const UAnimSequence* AnimSeq, cons
 			for (int32 BoneIndex = 0; BoneIndex < NumBones; ++BoneIndex)
 			{
 				const FCachedBoneIndexData& BoneIndexData = CachedBoneIndexData[BoneIndex];
-				if (BoneIndexData.RawTrackIndex == INDEX_NONE)
+				if (BoneIndexData.TrackIndex == INDEX_NONE)
 				{
 					// No track for the bone was found, use default transform
 					const FTransform& RefPoseTransform = RefPose[BoneIndex];
@@ -309,7 +320,7 @@ void FAnimationUtils::ComputeCompressionError(const UAnimSequence* AnimSeq, cons
 					// This is because additive animations are mostly rotation.
 					// And for the error metric we measure distance between end effectors.
 					// So that means additive animations by default will all be balled up at the origin and not show any error.
-					if (AnimSeq->IsValidAdditive())
+					if (CompressibleAnimData.bIsValidAdditive)
 					{
 						const FTransform& RefPoseTransform = RefPose[BoneIndex];
 						RawTransforms[BoneIndex] = RefPoseTransform;
@@ -317,8 +328,9 @@ void FAnimationUtils::ComputeCompressionError(const UAnimSequence* AnimSeq, cons
 
 						FTransform AdditiveRawTransform;
 						FTransform AdditiveNewTransform;
-						AnimSeq->GetBoneTransform(AdditiveRawTransform, BoneIndexData.RawTrackIndex, DecompContext, true);
-						AnimSeq->GetBoneTransform(AdditiveNewTransform, BoneIndexData.CompressedTrackIndex, DecompContext, false);
+						FAnimationUtils::ExtractTransformFromTrack(Time, CompressibleAnimData.NumFrames, CompressibleAnimData.SequenceLength, CompressibleAnimData.RawAnimationData[BoneIndexData.TrackIndex], CompressibleAnimData.Interpolation, AdditiveRawTransform);
+
+						AnimationFormat_GetBoneAtom(AdditiveNewTransform, DecompContext, BoneIndexData.TrackIndex);
 
 						const ScalarRegister VBlendWeight(1.f);
 						RawTransforms[BoneIndex].AccumulateWithAdditiveScale(AdditiveRawTransform, VBlendWeight);
@@ -326,8 +338,8 @@ void FAnimationUtils::ComputeCompressionError(const UAnimSequence* AnimSeq, cons
 					}
 					else
 					{
-						AnimSeq->GetBoneTransform(RawTransforms[BoneIndex], BoneIndexData.RawTrackIndex, DecompContext, true);
-						AnimSeq->GetBoneTransform(NewTransforms[BoneIndex], BoneIndexData.CompressedTrackIndex, DecompContext, false);
+						FAnimationUtils::ExtractTransformFromTrack(Time, CompressibleAnimData.NumFrames, CompressibleAnimData.SequenceLength, CompressibleAnimData.RawAnimationData[BoneIndexData.TrackIndex], CompressibleAnimData.Interpolation, RawTransforms[BoneIndex]);
+						AnimationFormat_GetBoneAtom(NewTransforms[BoneIndex], DecompContext, BoneIndexData.TrackIndex);
 					}
 				}
 
@@ -372,7 +384,7 @@ void FAnimationUtils::ComputeCompressionError(const UAnimSequence* AnimSeq, cons
 					{
 						ErrorStats.MaxError	= Error;
 						ErrorStats.MaxErrorBone = BoneIndex;
-						MaxErrorTrack = BoneIndexData.RawTrackIndex;
+						MaxErrorTrack = BoneIndexData.TrackIndex;
 						ErrorStats.MaxErrorTime = Time;
 					}
 				}
@@ -383,26 +395,9 @@ void FAnimationUtils::ComputeCompressionError(const UAnimSequence* AnimSeq, cons
 		{
 			ErrorStats.AverageError = ErrorTotal / ErrorCount;
 		}
-
-#if 0		
-		// That's a big error, log out some information!
- 		if( ErrorStats.MaxError > 10.f )
- 		{
-			UE_LOG(LogAnimation, Log, TEXT("!!! Big error found: %f, Time: %f, BoneIndex: %d, Track: %d, CompressionScheme: %s, additive: %d"), 
- 				ErrorStats.MaxError,
-				ErrorStats.MaxErrorTime,
-				ErrorStats.MaxErrorBone,
-				MaxErrorTrack,
-				AnimSeq->CompressionScheme ? *AnimSeq->CompressionScheme->GetFName().ToString() : TEXT("NULL"),
-				AnimSeq->bIsAdditive );
- 			UE_LOG(LogAnimation, Log, TEXT("   RawOrigin: %s, NormalOrigin: %s"), *RawTransforms(ErrorStats.MaxErrorBone).GetOrigin().ToString(), *NewTransforms(ErrorStats.MaxErrorBone).GetOrigin().ToString());
- 
- 			// We shouldn't have a big error with no compression.
- 			check( AnimSeq->CompressionScheme != NULL );
- 		}	
-#endif
 	}
 }
+#endif
 
 ///////////////////////////////////////////////////////////////////////////////////////////////////////////////////////
 //
@@ -483,51 +478,31 @@ bool FAnimationUtils::GetForcedRecompressionSetting()
 struct FAnimCompressionJobContext
 {
 	// Inputs
-	int64 OriginalSize;
-	float MasterTolerance;
 	bool bForceBelowThreshold;
-	const TArray<FBoneData>* BoneData;
 	FAnimCompressContext* CompressContext;
 	UAnimCompress* CompressionAlgorithm;
-	UAnimSequence* AnimSeq;
+	
+	const FCompressibleAnimData* CompressibleAnimData;
+	FCompressibleAnimDataResult CompressionResult;
+
 	const TCHAR* CompressionName;
 	int32* WinningCompressor_Count;
 	float* WinningCompressor_Error;
 	int64* WinningCompressor_Margin;
 
-	// Input/Output
-	int64 CurrentSize;
-	bool bSavedUseRawDataOnly;
-
-	int64 WinningCompressorMarginalSavings;
-	int32* WinningCompressorCounter;
-	float* WinningCompressorErrorSum;
-	int64* WinningCompressorMarginalSavingsSum;
-	FString WinningCompressorName;
-	int32 WinningCompressorSavings;
-	float WinningCompressorError;
-
-	UAnimCompress* SavedCompressionScheme;
-	AnimationCompressionFormat SavedTranslationCompressionFormat;
-	AnimationCompressionFormat SavedRotationCompressionFormat;
-	AnimationCompressionFormat SavedScaleCompressionFormat;
-	AnimationKeyFormat SavedKeyEncodingFormat;
-	TArray<int32> SavedCompressedTrackOffsets;
-	FCompressedOffsetData SavedCompressedScaleOffsets;
-	TArray<FCompressedSegment> SavedCompressedSegments;
-	TArray<uint8> SavedCompressedByteStream;
-	AnimEncoding* SavedTranslationCodec;
-	AnimEncoding* SavedRotationCodec;
-	AnimEncoding* SavedScaleCodec;
-
 	// Outputs
-	AnimationErrorStats NewErrorStats;
+	AnimationErrorStats ErrorStats;
 	float PctSaving;
 
 	FAnimCompressionJobContext()
 	{
 		memset(this, 0, sizeof(FAnimCompressionJobContext));
 	}
+
+	void UpdatePctSaving(int64 OriginalSize)
+	{
+		PctSaving = (OriginalSize > 0) ? (100.f - (100.f * float(CompressionResult.GetApproxBoneCompressedSize()) / float(OriginalSize))) : 0.f;
+	}
 };
 
 static void TryCompressionInner(FAnimCompressionJobContext& JobContext, bool is_async);
@@ -554,116 +529,107 @@ public:
 	FAnimCompressionJobContext* JobContext;
 };
 
-static void TryCompressionInner(FAnimCompressionJobContext& JobContext, bool is_async)
+bool ShouldKeepNewCompressionMethod(const FAnimCompressionJobContext& JobContext, SIZE_T OriginalSize, SIZE_T CurrentSize, float WinningCompressorError, float MasterTolerance)
 {
-#if WITH_EDITOR
-	UAnimSequence* AnimSeq = JobContext.AnimSeq;
+	const SIZE_T NewSize = JobContext.CompressionResult.GetApproxBoneCompressedSize();
+
+	/* compute the savings and compression error*/
+	const int64 MemorySavingsFromOriginal = OriginalSize - NewSize;
+	const int64 MemorySavingsFromPrevious = CurrentSize - NewSize;
+
+	/* figure out our new compression error*/
+	const AnimationErrorStats& ErrorStats = JobContext.ErrorStats;
+
+	const bool bLowersError = ErrorStats.MaxError < WinningCompressorError;
+	const bool bErrorUnderThreshold = ErrorStats.MaxError <= MasterTolerance;
+
+	/* keep it if it we want to force the error below the threshold and it reduces error */
+	bool bKeepNewCompressionMethod = false;
+	const bool bReducesErrorBelowThreshold = (bLowersError && (WinningCompressorError > MasterTolerance) && JobContext.bForceBelowThreshold);
+	bKeepNewCompressionMethod |= bReducesErrorBelowThreshold;
+	/* or if has an acceptable error and saves space  */
+	const bool bHasAcceptableErrorAndSavesSpace = bErrorUnderThreshold && (MemorySavingsFromPrevious > 0);
+	bKeepNewCompressionMethod |= bHasAcceptableErrorAndSavesSpace;
+	/* or if saves the same amount and an acceptable error that is lower than the previous best */
+	const bool bLowersErrorAndSavesSameOrBetter = bErrorUnderThreshold && bLowersError && (MemorySavingsFromPrevious >= 0);
+	bKeepNewCompressionMethod |= bLowersErrorAndSavesSameOrBetter;
+
+	UE_LOG(LogAnimationCompression, Verbose, TEXT("- %s - bytes saved(%i) (%.1f%%) from previous(%i) MaxError(%.2f) bLowersError(%d) %s"),
+		JobContext.CompressionName, MemorySavingsFromOriginal, JobContext.PctSaving, MemorySavingsFromPrevious, ErrorStats.MaxError, bLowersError, bKeepNewCompressionMethod ? TEXT("(**Best so far**)") : TEXT(""));
+
+	UE_LOG(LogAnimationCompression, Verbose, TEXT("    bReducesErrorBelowThreshold(%d) bHasAcceptableErrorAndSavesSpace(%d) bLowersErrorAndSavesSameOrBetter(%d)"),
+		bReducesErrorBelowThreshold, bHasAcceptableErrorAndSavesSpace, bLowersErrorAndSavesSameOrBetter);
+
+	UE_LOG(LogAnimationCompression, Verbose, TEXT("    WinningCompressorError(%f) MasterTolerance(%f) bForceBelowThreshold(%d) bErrorUnderThreshold(%d)"),
+		WinningCompressorError, MasterTolerance, JobContext.bForceBelowThreshold, bErrorUnderThreshold);
+	return bKeepNewCompressionMethod;
+}
 
+static void TryCompressionInner(FAnimCompressionJobContext& JobContext, bool bIsAsync)
+{
+#if WITH_EDITOR
 	/* try the alternative compressor	*/
-	AnimSeq->CompressionScheme = JobContext.CompressionAlgorithm;
-	JobContext.CompressionAlgorithm->Reduce(AnimSeq, *JobContext.CompressContext, *JobContext.BoneData);
-	AnimSeq->SetUseRawDataOnly(false);
-	const SIZE_T NewSize = AnimSeq->GetApproxCompressedSize();
-
-	if (is_async)
-	{
-		JobContext.CurrentSize = NewSize;
-	}
-	else
-	{
-		/* compute the savings and compression error*/
-		const int64 MemorySavingsFromOriginal = JobContext.OriginalSize - NewSize;
-		const int64 MemorySavingsFromPrevious = JobContext.CurrentSize - NewSize;
+	JobContext.CompressionAlgorithm->Reduce(*JobContext.CompressibleAnimData, JobContext.CompressionResult);
+	FAnimationUtils::ComputeCompressionError(*JobContext.CompressibleAnimData, JobContext.CompressionResult, JobContext.ErrorStats);
+#endif	// WITH_EDITOR
+}
 
-		/* figure out our new compression error*/
-		FAnimationUtils::ComputeCompressionError(AnimSeq, *JobContext.BoneData, JobContext.NewErrorStats);
+#if WITH_EDITORONLY_DATA
+struct WinningCompressorStats
+{
+	SIZE_T CurrentSize;
 
-		const bool bLowersError = JobContext.NewErrorStats.MaxError < JobContext.WinningCompressorError;
-		const bool bErrorUnderThreshold = JobContext.NewErrorStats.MaxError <= JobContext.MasterTolerance;
+	int32* WinningCompressorCounter;
+	float* WinningCompressorErrorSum;
+	int64* WinningCompressorMarginalSavingsSum;
+	int64 WinningCompressorMarginalSavings;
+	FString WinningCompressorName;
+	int32 WinningCompressorSavings;
+	float PctSaving;
+	float WinningCompressorError;
 
-		/* keep it if it we want to force the error below the threshold and it reduces error */
-		bool bKeepNewCompressionMethod = false;
-		const bool bReducesErrorBelowThreshold = (bLowersError && (JobContext.WinningCompressorError > JobContext.MasterTolerance) && JobContext.bForceBelowThreshold);
-		bKeepNewCompressionMethod |= bReducesErrorBelowThreshold;
-		/* or if has an acceptable error and saves space  */
-		const bool bHasAcceptableErrorAndSavesSpace = bErrorUnderThreshold && (MemorySavingsFromPrevious > 0);
-		bKeepNewCompressionMethod |= bHasAcceptableErrorAndSavesSpace;
-		/* or if saves the same amount and an acceptable error that is lower than the previous best */
-		const bool bLowersErrorAndSavesSameOrBetter = bErrorUnderThreshold && bLowersError && (MemorySavingsFromPrevious >= 0);
-		bKeepNewCompressionMethod |= bLowersErrorAndSavesSameOrBetter;
+	WinningCompressorStats()
+	{
+		memset(this, 0, sizeof(WinningCompressorStats));
+	}
+};
 
-		JobContext.PctSaving = (JobContext.OriginalSize > 0) ? (100.f - (100.f * float(NewSize) / float(JobContext.OriginalSize))) : 0.f;
-		UE_LOG(LogAnimationCompression, Verbose, TEXT("- %s - bytes saved(%i) (%.1f%%) from previous(%i) MaxError(%.2f) bLowersError(%d) %s"),
-			JobContext.CompressionName, MemorySavingsFromOriginal, JobContext.PctSaving, MemorySavingsFromPrevious, JobContext.NewErrorStats.MaxError, bLowersError, bKeepNewCompressionMethod ? TEXT("(**Best so far**)") : TEXT(""));
+void HandlePostTryCompression(FAnimCompressionJobContext& JobContext, int64 OriginalSize, float MasterTolerance, WinningCompressorStats& OutCompressorStats, FCompressibleAnimDataResult& OutCompressedData, AnimationErrorStats& OutNewErrorStats)
+{
+	JobContext.UpdatePctSaving(OriginalSize);
 
-		UE_LOG(LogAnimationCompression, Verbose, TEXT("    bReducesErrorBelowThreshold(%d) bHasAcceptableErrorAndSavesSpace(%d) bLowersErrorAndSavesSameOrBetter(%d)"),
-			bReducesErrorBelowThreshold, bHasAcceptableErrorAndSavesSpace, bLowersErrorAndSavesSameOrBetter);
+	bool bKeepNewCompressionMethod = ShouldKeepNewCompressionMethod(JobContext, OriginalSize, OutCompressorStats.CurrentSize, OutCompressorStats.WinningCompressorError, MasterTolerance);
+	if (bKeepNewCompressionMethod)
+	{
+		OutCompressedData = JobContext.CompressionResult;
 
-		UE_LOG(LogAnimationCompression, Verbose, TEXT("    WinningCompressorError(%f) MasterTolerance(%f) bForceBelowThreshold(%d) bErrorUnderThreshold(%d)"),
-			JobContext.WinningCompressorError, JobContext.MasterTolerance, JobContext.bForceBelowThreshold, bErrorUnderThreshold);
+		OutNewErrorStats = JobContext.ErrorStats;
 
-		if (bKeepNewCompressionMethod)
-		{
-			JobContext.WinningCompressorMarginalSavings = MemorySavingsFromPrevious;
-			JobContext.WinningCompressorCounter = JobContext.WinningCompressor_Count;
-			JobContext.WinningCompressorErrorSum = JobContext.WinningCompressor_Error;
-			JobContext.WinningCompressorMarginalSavingsSum = JobContext.WinningCompressor_Margin;
-			JobContext.WinningCompressorName = JobContext.CompressionName;
-			JobContext.CurrentSize = NewSize;
-			JobContext.WinningCompressorSavings = MemorySavingsFromOriginal;
-			JobContext.WinningCompressorError = JobContext.NewErrorStats.MaxError;
-
-			/* backup key information from the sequence */
-			JobContext.SavedCompressionScheme = AnimSeq->CompressionScheme;
-			JobContext.SavedTranslationCompressionFormat = AnimSeq->TranslationCompressionFormat;
-			JobContext.SavedRotationCompressionFormat = AnimSeq->RotationCompressionFormat;
-			JobContext.SavedScaleCompressionFormat = AnimSeq->ScaleCompressionFormat;
-			JobContext.SavedKeyEncodingFormat = AnimSeq->KeyEncodingFormat;
-			JobContext.SavedCompressedTrackOffsets = AnimSeq->CompressedTrackOffsets;
-			JobContext.SavedCompressedScaleOffsets = AnimSeq->CompressedScaleOffsets;
-			JobContext.SavedCompressedSegments = AnimSeq->CompressedSegments;
-			JobContext.SavedCompressedByteStream = AnimSeq->CompressedByteStream;
-			JobContext.SavedTranslationCodec = AnimSeq->TranslationCodec;
-			JobContext.SavedRotationCodec = AnimSeq->RotationCodec;
-			JobContext.SavedScaleCodec = AnimSeq->ScaleCodec;
-			JobContext.bSavedUseRawDataOnly = false;
-		}
-		else
-		{
-			/* revert back to the old method by copying back the data we cached */
-			AnimSeq->CompressionScheme = JobContext.SavedCompressionScheme;
-			AnimSeq->TranslationCompressionFormat = JobContext.SavedTranslationCompressionFormat;
-			AnimSeq->RotationCompressionFormat = JobContext.SavedRotationCompressionFormat;
-			AnimSeq->ScaleCompressionFormat = JobContext.SavedScaleCompressionFormat;
-			AnimSeq->KeyEncodingFormat = JobContext.SavedKeyEncodingFormat;
-			AnimSeq->CompressedTrackOffsets = JobContext.SavedCompressedTrackOffsets;
-			AnimSeq->CompressedByteStream = JobContext.SavedCompressedByteStream;
-			AnimSeq->CompressedScaleOffsets = JobContext.SavedCompressedScaleOffsets;
-			AnimSeq->CompressedSegments = JobContext.SavedCompressedSegments;
-			AnimSeq->TranslationCodec = JobContext.SavedTranslationCodec;
-			AnimSeq->RotationCodec = JobContext.SavedRotationCodec;
-			AnimSeq->ScaleCodec = JobContext.SavedScaleCodec;
-			AnimSeq->SetUseRawDataOnly(JobContext.bSavedUseRawDataOnly);
-			AnimationFormat_SetInterfaceLinks(*AnimSeq);
-
-			const SIZE_T RestoredSize = AnimSeq->GetApproxCompressedSize();
-			check(RestoredSize == JobContext.CurrentSize);
-		}
+		int64 NewSize = JobContext.CompressionResult.GetApproxBoneCompressedSize();
+		const int64 MemorySavingsFromOriginal = OriginalSize - NewSize;
+		const int64 MemorySavingsFromPrevious = OutCompressorStats.CurrentSize - NewSize;
+
+		OutCompressorStats.CurrentSize = NewSize;
+		OutCompressorStats.WinningCompressorMarginalSavings = MemorySavingsFromPrevious;
+		OutCompressorStats.WinningCompressorCounter = JobContext.WinningCompressor_Count; 
+		OutCompressorStats.WinningCompressorErrorSum = JobContext.WinningCompressor_Error; 
+		OutCompressorStats.WinningCompressorMarginalSavingsSum = JobContext.WinningCompressor_Margin; 
+		OutCompressorStats.WinningCompressorName = JobContext.CompressionName; 
+		OutCompressorStats.WinningCompressorSavings = MemorySavingsFromOriginal;
+		OutCompressorStats.WinningCompressorError = JobContext.ErrorStats.MaxError;
+		OutCompressorStats.PctSaving = JobContext.PctSaving;
 	}
-#endif	// WITH_EDITOR
 }
+#endif
 
-#define POPULATE_ANIM_COMPRESSION_JOB_CONTEXT(Name_, CompressionAlgorithm_, AnimSeq_, ParentAnimSeq_, CompressContext_, JobContext_) \
-	(JobContext_).OriginalSize = OriginalSize; \
-	(JobContext_).MasterTolerance = MasterTolerance; \
+#define POPULATE_ANIM_COMPRESSION_JOB_CONTEXT(Name_, CompressionAlgorithm_, CompressContext_, JobContext_) \
 	(JobContext_).bForceBelowThreshold = bForceBelowThreshold; \
-	(JobContext_).BoneData = &BoneData; \
 	(JobContext_).CompressContext = &CompressContext_; \
-	(JobContext_).CompressionAlgorithm = DuplicateObject<UAnimCompress>((CompressionAlgorithm_), ParentAnimSeq_); \
+	(JobContext_).CompressibleAnimData = &CompressibleAnimData; \
+	(JobContext_).CompressionAlgorithm = DuplicateObject<UAnimCompress>((CompressionAlgorithm_), GetTransientPackage()); \
 	(JobContext_).CompressionAlgorithm->bEnableSegmenting = bEnableSegmenting; \
 	(JobContext_).CompressionAlgorithm->IdealNumFramesPerSegment = IdealNumFramesPerSegment; \
 	(JobContext_).CompressionAlgorithm->MaxNumFramesPerSegment = MaxNumFramesPerSegment; \
-	(JobContext_).AnimSeq = AnimSeq_; \
 	(JobContext_).CompressionName = TEXT(#Name_); \
 	(JobContext_).WinningCompressor_Count = &(Name_ ## CompressorWins); \
 	(JobContext_).WinningCompressor_Error = &(Name_ ## CompressorSumError); \
@@ -673,56 +639,11 @@ static void TryCompressionInner(FAnimCompressionJobContext& JobContext, bool is_
 	do \
 	{ \
 		FAnimCompressionJobContext JobContext; \
-		POPULATE_ANIM_COMPRESSION_JOB_CONTEXT(Name, CompressionAlgorithm_, AnimSeq, AnimSeq, CompressContext, JobContext); \
-		\
-		JobContext.CurrentSize = CompressorStats.CurrentSize; \
-		JobContext.bSavedUseRawDataOnly = bSavedUseRawDataOnly; \
-		JobContext.WinningCompressorMarginalSavings = CompressorStats.WinningCompressorMarginalSavings; \
-		JobContext.WinningCompressorCounter = CompressorStats.WinningCompressorCounter; \
-		JobContext.WinningCompressorErrorSum = CompressorStats.WinningCompressorErrorSum; \
-		JobContext.WinningCompressorMarginalSavingsSum = CompressorStats.WinningCompressorMarginalSavingsSum; \
-		JobContext.WinningCompressorName = CompressorStats.WinningCompressorName; \
-		JobContext.WinningCompressorSavings = CompressorStats.WinningCompressorSavings; \
-		JobContext.WinningCompressorError = CompressorStats.WinningCompressorError; \
-		JobContext.SavedCompressionScheme = SavedCompressionScheme; \
-		JobContext.SavedTranslationCompressionFormat = SavedTranslationCompressionFormat; \
-		JobContext.SavedRotationCompressionFormat = SavedRotationCompressionFormat; \
-		JobContext.SavedScaleCompressionFormat = SavedScaleCompressionFormat; \
-		JobContext.SavedKeyEncodingFormat = SavedKeyEncodingFormat; \
-		JobContext.SavedCompressedTrackOffsets = SavedCompressedTrackOffsets; \
-		JobContext.SavedCompressedScaleOffsets = SavedCompressedScaleOffsets; \
-		JobContext.SavedCompressedByteStream = SavedCompressedByteStream; \
-		JobContext.SavedCompressedSegments = SavedCompressedSegments; \
-		JobContext.SavedTranslationCodec = SavedTranslationCodec; \
-		JobContext.SavedRotationCodec = SavedRotationCodec; \
-		JobContext.SavedScaleCodec = SavedScaleCodec; \
+		POPULATE_ANIM_COMPRESSION_JOB_CONTEXT(Name, CompressionAlgorithm_, CompressContext, JobContext); \
 		\
 		TryCompressionInner(JobContext, false); \
 		\
-		CompressorStats.CurrentSize = JobContext.CurrentSize; \
-		bSavedUseRawDataOnly = JobContext.bSavedUseRawDataOnly; \
-		CompressorStats.WinningCompressorMarginalSavings = JobContext.WinningCompressorMarginalSavings; \
-		CompressorStats.WinningCompressorCounter = JobContext.WinningCompressorCounter; \
-		CompressorStats.WinningCompressorErrorSum = JobContext.WinningCompressorErrorSum; \
-		CompressorStats.WinningCompressorMarginalSavingsSum = JobContext.WinningCompressorMarginalSavingsSum; \
-		CompressorStats.WinningCompressorName = JobContext.WinningCompressorName; \
-		CompressorStats.WinningCompressorSavings = JobContext.WinningCompressorSavings; \
-		CompressorStats.WinningCompressorError = JobContext.WinningCompressorError; \
-		SavedCompressionScheme = JobContext.SavedCompressionScheme; \
-		SavedTranslationCompressionFormat = JobContext.SavedTranslationCompressionFormat; \
-		SavedRotationCompressionFormat = JobContext.SavedRotationCompressionFormat; \
-		SavedScaleCompressionFormat = JobContext.SavedScaleCompressionFormat; \
-		SavedKeyEncodingFormat = JobContext.SavedKeyEncodingFormat; \
-		SavedCompressedTrackOffsets = JobContext.SavedCompressedTrackOffsets; \
-		SavedCompressedScaleOffsets = JobContext.SavedCompressedScaleOffsets; \
-		SavedCompressedByteStream = JobContext.SavedCompressedByteStream; \
-		SavedCompressedSegments = JobContext.SavedCompressedSegments; \
-		SavedTranslationCodec = JobContext.SavedTranslationCodec; \
-		SavedRotationCodec = JobContext.SavedRotationCodec; \
-		SavedScaleCodec = JobContext.SavedScaleCodec; \
-		\
-		NewErrorStats = JobContext.NewErrorStats; \
-		CompressorStats.PctSaving = JobContext.PctSaving; \
+		HandlePostTryCompression(JobContext, OriginalSize, MasterTolerance, CompressorStats, OutCompressedData, NewErrorStats); \
 	} \
 	while (false)
 
@@ -734,13 +655,12 @@ static void TryCompressionInner(FAnimCompressionJobContext& JobContext, bool is_
 #define TRYCOMPRESSION_ASYNC(Name, CompressionAlgorithm_) \
 	do \
 	{ \
-		if (!AnimSeq->IsValidAdditive()) \
+		if (!CompressibleAnimData.bIsValidAdditive) \
 		{ \
 			FAnimCompressionJobContext* JobContext = new FAnimCompressionJobContext(); \
 			FAnimCompressContext CompressContextCopy(CompressContext); \
 			CompressContextCopy.CompressionSummary = FCompressionMemorySummary(false); \
-			UAnimSequence* AnimSeqCopy = DuplicateObject<UAnimSequence>(AnimSeq, GetTransientPackage()); \
-			POPULATE_ANIM_COMPRESSION_JOB_CONTEXT(Name, CompressionAlgorithm_, AnimSeqCopy, AnimSeq, CompressContextCopy, *JobContext); \
+			POPULATE_ANIM_COMPRESSION_JOB_CONTEXT(Name, CompressionAlgorithm_, CompressContextCopy, *JobContext); \
 			AnimCompressionJobContexes.Add(JobContext); \
 			\
 			AnimCompressionTask_CompletionEvents.Add(TGraphTask<FAsyncAnimCompressionTask>::CreateTask(NULL, ENamedThreads::GameThread).ConstructAndDispatchWhenReady(JobContext)); \
@@ -761,24 +681,6 @@ static void TryCompressionInner(FAnimCompressionJobContext& JobContext, bool is_
 	static int64 Algorithm ## CompressorWinMargin = 0
 
 #if WITH_EDITORONLY_DATA
-struct WinningCompressorStats
-{
-	SIZE_T CurrentSize;
-
-	int32* WinningCompressorCounter;
-	float* WinningCompressorErrorSum;
-	int64* WinningCompressorMarginalSavingsSum;
-	int64 WinningCompressorMarginalSavings;
-	FString WinningCompressorName;
-	int32 WinningCompressorSavings;
-	float PctSaving;
-	float WinningCompressorError;
-
-	WinningCompressorStats()
-	{
-		memset(this, 0, sizeof(WinningCompressorStats));
-	}
-};
 
 static void WaitForAnimCompressionJobs(const FGraphEventArray& AnimCompressionTask_CompletionEvents)
 {
@@ -789,9 +691,6 @@ static void ClearAnimCompressionJobs(FGraphEventArray& AnimCompressionTask_Compl
 {
 	for (FAnimCompressionJobContext* Context : AnimCompressionJobContexes)
 	{
-		Context->AnimSeq->RecycleAnimSequence();
-		Context->AnimSeq->ClearFlags(RF_Standalone | RF_Public);
-		Context->AnimSeq->MarkPendingKill();
 		delete Context;
 	}
 	AnimCompressionJobContexes.Reset();
@@ -805,52 +704,22 @@ static FAnimCompressionJobContext* FindBestAnimCompression(TArray<FAnimCompressi
 	for (FAnimCompressionJobContext* context : AnimCompressionJobContexes)
 	{
 		FAnimCompressionJobContext& JobContext = *context;
+		JobContext.UpdatePctSaving(OriginalSize);
 
-		const SIZE_T NewSize = JobContext.AnimSeq->GetApproxCompressedSize();
-
-		/* compute the savings and compression error*/
-		const int64 MemorySavingsFromOriginal = OriginalSize - NewSize;
-		const int64 MemorySavingsFromPrevious = CurrentSize - NewSize;
-
-		/* figure out our new compression error*/
-		FAnimationUtils::ComputeCompressionError(JobContext.AnimSeq, *JobContext.BoneData, JobContext.NewErrorStats);
-
-		const bool bLowersError = JobContext.NewErrorStats.MaxError < WinningCompressorError;
-		const bool bErrorUnderThreshold = JobContext.NewErrorStats.MaxError <= MasterTolerance;
-
-		/* keep it if it we want to force the error below the threshold and it reduces error */
-		bool bKeepNewCompressionMethod = false;
-		const bool bReducesErrorBelowThreshold = (bLowersError && (WinningCompressorError > MasterTolerance) && JobContext.bForceBelowThreshold);
-		bKeepNewCompressionMethod |= bReducesErrorBelowThreshold;
-		/* or if has an acceptable error and saves space  */
-		const bool bHasAcceptableErrorAndSavesSpace = bErrorUnderThreshold && (MemorySavingsFromPrevious > 0);
-		bKeepNewCompressionMethod |= bHasAcceptableErrorAndSavesSpace;
-		/* or if saves the same amount and an acceptable error that is lower than the previous best */
-		const bool bLowersErrorAndSavesSameOrBetter = bErrorUnderThreshold && bLowersError && (MemorySavingsFromPrevious >= 0);
-		bKeepNewCompressionMethod |= bLowersErrorAndSavesSameOrBetter;
-
-		JobContext.PctSaving = (OriginalSize > 0) ? (100.f - (100.f * float(NewSize) / float(OriginalSize))) : 0.f;
-		UE_LOG(LogAnimationCompression, Verbose, TEXT("- %s - bytes saved(%i) (%.1f%%) from previous(%i) MaxError(%.2f) bLowersError(%d) %s"),
-			JobContext.CompressionName, MemorySavingsFromOriginal, JobContext.PctSaving, MemorySavingsFromPrevious, JobContext.NewErrorStats.MaxError, bLowersError, bKeepNewCompressionMethod ? TEXT("(**Best so far**)") : TEXT(""));
-
-		UE_LOG(LogAnimationCompression, Verbose, TEXT("    bReducesErrorBelowThreshold(%d) bHasAcceptableErrorAndSavesSpace(%d) bLowersErrorAndSavesSameOrBetter(%d)"),
-			bReducesErrorBelowThreshold, bHasAcceptableErrorAndSavesSpace, bLowersErrorAndSavesSameOrBetter);
-
-		UE_LOG(LogAnimationCompression, Verbose, TEXT("    WinningCompressorError(%f) MasterTolerance(%f) bForceBelowThreshold(%d) bErrorUnderThreshold(%d)"),
-			WinningCompressorError, JobContext.MasterTolerance, JobContext.bForceBelowThreshold, bErrorUnderThreshold);
+		bool bKeepNewCompressionMethod = ShouldKeepNewCompressionMethod(JobContext, OriginalSize, CurrentSize, WinningCompressorError, MasterTolerance);
 
 		if (bKeepNewCompressionMethod)
 		{
 			BestJobContext = context;
-			WinningCompressorError = JobContext.NewErrorStats.MaxError;
-			CurrentSize = NewSize;
+			WinningCompressorError = JobContext.ErrorStats.MaxError;
+			CurrentSize = JobContext.CompressionResult.GetApproxBoneCompressedSize();
 		}
 	}
 
 	return BestJobContext;
 }
 
-static void UpdateAnimCompressionFromAsyncJobs(UAnimSequence* AnimSeq, FGraphEventArray& AnimCompressionTask_CompletionEvents, TArray<FAnimCompressionJobContext*>& AnimCompressionJobContexes, SIZE_T OriginalSize, WinningCompressorStats& CompressorStats, float MasterTolerance)
+static void UpdateAnimCompressionFromAsyncJobs(FCompressibleAnimDataResult& OutCompressedData, FGraphEventArray& AnimCompressionTask_CompletionEvents, TArray<FAnimCompressionJobContext*>& AnimCompressionJobContexes, SIZE_T OriginalSize, WinningCompressorStats& CompressorStats, float MasterTolerance)
 {
 	// Pick the best
 	const FAnimCompressionJobContext* BestJobContext = FindBestAnimCompression(AnimCompressionJobContexes, OriginalSize, CompressorStats.CurrentSize, CompressorStats.WinningCompressorError, MasterTolerance);
@@ -860,34 +729,21 @@ static void UpdateAnimCompressionFromAsyncJobs(UAnimSequence* AnimSeq, FGraphEve
 		const FAnimCompressionJobContext& JobContext = *BestJobContext;
 
 		/* Copy our data */
-		AnimSeq->CompressionScheme = JobContext.AnimSeq->CompressionScheme;
-		AnimSeq->TranslationCompressionFormat = JobContext.AnimSeq->TranslationCompressionFormat;
-		AnimSeq->RotationCompressionFormat = JobContext.AnimSeq->RotationCompressionFormat;
-		AnimSeq->KeyEncodingFormat = JobContext.AnimSeq->KeyEncodingFormat;
-		AnimSeq->CompressedTrackOffsets = JobContext.AnimSeq->CompressedTrackOffsets;
-		AnimSeq->CompressedByteStream = JobContext.AnimSeq->CompressedByteStream;
-		AnimSeq->CompressedScaleOffsets = JobContext.AnimSeq->CompressedScaleOffsets;
-		AnimSeq->CompressedSegments = JobContext.AnimSeq->CompressedSegments;
-		AnimSeq->TranslationCodec = JobContext.AnimSeq->TranslationCodec;
-		AnimSeq->RotationCodec = JobContext.AnimSeq->RotationCodec;
-		AnimSeq->ScaleCodec = JobContext.AnimSeq->ScaleCodec;
-		AnimSeq->SetUseRawDataOnly(false);
-		AnimationFormat_SetInterfaceLinks(*AnimSeq);
-
-		const SIZE_T RestoredSize = AnimSeq->GetApproxCompressedSize();
-		check(RestoredSize == JobContext.CurrentSize);
-
-		const int64 MemorySavingsFromOriginal = OriginalSize - RestoredSize;
-		const SIZE_T MemorySavingsFromPrevious = CompressorStats.CurrentSize - RestoredSize;
+		OutCompressedData = JobContext.CompressionResult;
+
+		const SIZE_T NewSize = OutCompressedData.GetApproxBoneCompressedSize();
+
+		const int64 MemorySavingsFromOriginal = OriginalSize - NewSize;
+		const SIZE_T MemorySavingsFromPrevious = CompressorStats.CurrentSize - NewSize;
 
 		CompressorStats.WinningCompressorMarginalSavings = MemorySavingsFromPrevious;
 		CompressorStats.WinningCompressorCounter = JobContext.WinningCompressor_Count;
 		CompressorStats.WinningCompressorErrorSum = JobContext.WinningCompressor_Error;
 		CompressorStats.WinningCompressorMarginalSavingsSum = JobContext.WinningCompressor_Margin;
 		CompressorStats.WinningCompressorName = JobContext.CompressionName;
-		CompressorStats.CurrentSize = RestoredSize;
+		CompressorStats.CurrentSize = NewSize;
 		CompressorStats.WinningCompressorSavings = MemorySavingsFromOriginal;
-		CompressorStats.WinningCompressorError = JobContext.NewErrorStats.MaxError;
+		CompressorStats.WinningCompressorError = JobContext.ErrorStats.MaxError;
 	}
 
 	ClearAnimCompressionJobs(AnimCompressionTask_CompletionEvents, AnimCompressionJobContexes);
@@ -929,25 +785,15 @@ public:
  * @param	bOutput		If false don't generate output or compute memory savings.
  * @return				None.
  */
-void FAnimationUtils::CompressAnimSequence(UAnimSequence* AnimSeq, FAnimCompressContext& CompressContext)
+void FAnimationUtils::CompressAnimSequence(const FCompressibleAnimData& CompressibleAnimData, FCompressibleAnimDataResult& OutCompressedData, FAnimCompressContext& CompressContext)
 {
 	if (FPlatformProperties::HasEditorOnlyData())
 	{
-		// the underlying code won't work right without skeleton. 
-		if ( !AnimSeq->GetSkeleton() )
-		{
-			return;
-		}
-
 		// get the master tolerance we will use to guide recompression
 		float MasterTolerance = GetAlternativeCompressionThreshold(); 
 
 		bool bOnlyCheckForMissingSkeletalMeshes = UAnimationSettings::Get()->bOnlyCheckForMissingSkeletalMeshes;
-		if (bOnlyCheckForMissingSkeletalMeshes)
-		{
-			TestForMissingMeshes(AnimSeq);
-		}
-		else
+		if (!bOnlyCheckForMissingSkeletalMeshes)
 		{
 			UAnimationSettings* AnimSetting = UAnimationSettings::Get();
 			bool bForceBelowThreshold = AnimSetting->bForceBelowThreshold;
@@ -964,19 +810,8 @@ void FAnimationUtils::CompressAnimSequence(UAnimSequence* AnimSeq, FAnimCompress
 			int32 IdealNumFramesPerSegment = 64;
 			int32 MaxNumFramesPerSegment = (IdealNumFramesPerSegment * 2) - 1;
 
-			// Build skeleton metadata to use during the key reduction.
-			TArray<FBoneData> BoneData;
-			FAnimationUtils::BuildSkeletonMetaData(AnimSeq->GetSkeleton(), BoneData);
-
 #if WITH_EDITORONLY_DATA
-			UAnimCompress_Automatic* AutoCompressionScheme = Cast<UAnimCompress_Automatic>(AnimSeq->CompressionScheme);
-
-			CompressContext.GatherPreCompressionStats(AnimSeq);
-
-			AnimSeq->CompressedByteStream.Reset();
-			AnimSeq->CompressedTrackOffsets.Reset();
-			AnimSeq->CompressedScaleOffsets.OffsetData.Reset();
-			AnimSeq->CompressedSegments.Reset();
+			UAnimCompress_Automatic* AutoCompressionScheme = Cast<UAnimCompress_Automatic>(CompressibleAnimData.RequestedCompressionScheme);
 
 			if (AutoCompressionScheme != nullptr)
 			{
@@ -985,7 +820,7 @@ void FAnimationUtils::CompressAnimSequence(UAnimSequence* AnimSeq, FAnimCompress
 				IdealNumFramesPerSegment = AutoCompressionScheme->IdealNumFramesPerSegment;
 				MaxNumFramesPerSegment = AutoCompressionScheme->MaxNumFramesPerSegment;
 			}
-			else if (AnimSeq->CompressionScheme != nullptr)
+			else if (CompressibleAnimData.RequestedCompressionScheme != nullptr)
 			{
 				bTryExhaustiveSearch = AnimSetting->bTryExhaustiveSearch;
 			}
@@ -996,7 +831,8 @@ void FAnimationUtils::CompressAnimSequence(UAnimSequence* AnimSeq, FAnimCompress
 				//Scoped timing of compression, make sure nothing else is added to this scope
 				FCompressionTimeElapsed TimeTracker(CompressionTime);
 				CompressAnimSequenceExplicit(
-					AnimSeq,
+					CompressibleAnimData,
+					OutCompressedData,
 					CompressContext,
 					CompressContext.bAllowAlternateCompressor ? MasterTolerance : 0.0f,
 					bFirstRecompressUsingCurrentOrDefault,
@@ -1005,11 +841,10 @@ void FAnimationUtils::CompressAnimSequence(UAnimSequence* AnimSeq, FAnimCompress
 					bTryExhaustiveSearch,
 					bEnableSegmenting,
 					IdealNumFramesPerSegment,
-					MaxNumFramesPerSegment,
-					BoneData);
+					MaxNumFramesPerSegment);
 			}
 
-			CompressContext.GatherPostCompressionStats(AnimSeq, BoneData, CompressionTime);
+			CompressContext.GatherPostCompressionStats(CompressibleAnimData, OutCompressedData, CompressionTime);
 		}
 	}
 }
@@ -1027,7 +862,8 @@ void FAnimationUtils::CompressAnimSequence(UAnimSequence* AnimSeq, FAnimCompress
  * @return				None.
  */
 void FAnimationUtils::CompressAnimSequenceExplicit(
-	UAnimSequence* AnimSeq,
+	const FCompressibleAnimData& CompressibleAnimData,
+	FCompressibleAnimDataResult& OutCompressedData,
 	FAnimCompressContext& CompressContext,
 	float MasterTolerance,
 	bool bFirstRecompressUsingCurrentOrDefault,
@@ -1036,8 +872,7 @@ void FAnimationUtils::CompressAnimSequenceExplicit(
 	bool bTryExhaustiveSearch,
 	bool bEnableSegmenting,
 	int32 IdealNumFramesPerSegment,
-	int32 MaxNumFramesPerSegment,
-	const TArray<FBoneData>& BoneData)
+	int32 MaxNumFramesPerSegment)
 {
 #if WITH_EDITORONLY_DATA
 	DECLARE_ANIM_COMP_ALGORITHM(BitwiseACF_Float96);
@@ -1085,14 +920,7 @@ void FAnimationUtils::CompressAnimSequenceExplicit(
 	DECLARE_ANIM_COMP_ALGORITHM(Linear_PerTrackExp1);
 	DECLARE_ANIM_COMP_ALGORITHM(Linear_PerTrackExp2);
 
-	check(AnimSeq != nullptr);
-	if (AnimSeq->HasAnyFlags(RF_NeedLoad))
-	{
-		AnimSeq->GetLinker()->Preload(AnimSeq);
-	}
-
-	// attempt to find the default skeletal mesh associated with this sequence
-	USkeleton* Skeleton = AnimSeq->GetSkeleton();
+	USkeleton* Skeleton = CompressibleAnimData.Skeleton;
 	check (Skeleton);
 	if (Skeleton->HasAnyFlags(RF_NeedLoad))
 	{
@@ -1110,13 +938,13 @@ void FAnimationUtils::CompressAnimSequenceExplicit(
 	static int64 TotalSizeNow = 0;
 	static int64 TotalUncompressed = 0;
 
-	const int32 NumRawDataTracks = AnimSeq->GetRawAnimationData().Num();
+	const int32 NumRawDataTracks = CompressibleAnimData.RawAnimationData.Num();
 
 	// we must have raw data to continue
 	if(NumRawDataTracks > 0 )
 	{
 		// If compression Scheme is automatic, then we definitely want to 'TryAlternateCompressor'.
-		if (AnimSeq->CompressionScheme && AnimSeq->CompressionScheme->IsA(UAnimCompress_Automatic::StaticClass()))
+		if (CompressibleAnimData.RequestedCompressionScheme && CompressibleAnimData.RequestedCompressionScheme->IsA(UAnimCompress_Automatic::StaticClass()))
 		{
 			MasterTolerance = GetAlternativeCompressionThreshold();
 		}
@@ -1124,24 +952,20 @@ void FAnimationUtils::CompressAnimSequenceExplicit(
 		// See if we're trying alternate compressors
 		const bool bTryAlternateCompressor = MasterTolerance > 0.0f;
 
-		// Filter RAW data to get rid of mismatched tracks (translation/rotation data with a different number of keys than there are frames)
-		// No trivial key removal is done at this point (impossible error metrics of -1), since all of the techniques will perform it themselves
-		AnimSeq->CompressRawAnimData(-1.0f, -1.0f);
-
-		AnimationErrorStats OriginalErrorStats;
 		AnimationErrorStats TrueOriginalErrorStats;
+		FAnimationUtils::ComputeCompressionError(CompressibleAnimData, OutCompressedData, TrueOriginalErrorStats);
 
-		FAnimationUtils::ComputeCompressionError(AnimSeq, BoneData, TrueOriginalErrorStats);
+		AnimationErrorStats OriginalErrorStats;
 
 		int32 AfterOriginalRecompression = 0;
-		if( (bFirstRecompressUsingCurrentOrDefault && !bTryAlternateCompressor) || (AnimSeq->CompressedByteStream.Num() == 0))
+		if( (bFirstRecompressUsingCurrentOrDefault && !bTryAlternateCompressor) || (OutCompressedData.CompressedByteStream.Num() == 0))
 		{
-			UAnimCompress* OriginalCompressionAlgorithm = AnimSeq->CompressionScheme ? AnimSeq->CompressionScheme : FAnimationUtils::GetDefaultAnimationCompressionAlgorithm();
+			UAnimCompress* OriginalCompressionAlgorithm = CompressibleAnimData.RequestedCompressionScheme ? CompressibleAnimData.RequestedCompressionScheme : FAnimationUtils::GetDefaultAnimationCompressionAlgorithm();
 
 			// Automatic compression brings us back here, so don't create an infinite loop and pick bitwise compress instead.
 			if (!OriginalCompressionAlgorithm || OriginalCompressionAlgorithm->IsA(UAnimCompress_Automatic::StaticClass()))
 			{
-				UAnimCompress* CompressionAlgorithm = NewObject<UAnimCompress_BitwiseCompressOnly>(AnimSeq);
+				UAnimCompress* CompressionAlgorithm = NewObject<UAnimCompress_BitwiseCompressOnly>();
 				if (OriginalCompressionAlgorithm != nullptr)
 				{
 					// Keep the same segmenting settings
@@ -1154,45 +978,42 @@ void FAnimationUtils::CompressAnimSequenceExplicit(
 			}
 
 			UE_LOG(LogAnimationCompression, Log, TEXT("Recompressing (%s) using current/default (%s) bFirstRecompressUsingCurrentOrDefault(%d) bTryAlternateCompressor(%d) IsCompressedDataValid(%d)"),
-				*AnimSeq->GetFullName(),
+				*CompressibleAnimData.FullName,
 				*OriginalCompressionAlgorithm->GetName(),
 				bFirstRecompressUsingCurrentOrDefault,
 				bTryAlternateCompressor,
-				AnimSeq->IsCompressedDataValid());
+				OutCompressedData.IsCompressedDataValid());
 
-			AnimSeq->CompressionScheme = DuplicateObject<UAnimCompress>(OriginalCompressionAlgorithm, AnimSeq);
-			OriginalCompressionAlgorithm->Reduce(AnimSeq, CompressContext, BoneData);
-			AnimSeq->SetUseRawDataOnly(false);
-			AfterOriginalRecompression = AnimSeq->GetApproxCompressedSize();
+			OriginalCompressionAlgorithm->Reduce(CompressibleAnimData, OutCompressedData);
+			
+			AfterOriginalRecompression = OutCompressedData.GetApproxBoneCompressedSize();
 
 			// figure out our current compression error
-			FAnimationUtils::ComputeCompressionError(AnimSeq, BoneData, OriginalErrorStats);
+			FAnimationUtils::ComputeCompressionError(CompressibleAnimData, OutCompressedData, OriginalErrorStats);
 		}
 		else
 		{
-			AfterOriginalRecompression = AnimSeq->GetApproxCompressedSize();
+			AfterOriginalRecompression = OutCompressedData.GetApproxBoneCompressedSize();
 			OriginalErrorStats = TrueOriginalErrorStats;
 		}
 
+		//For logging
+		FString OriginalKeyEncodingFormat = GetAnimationKeyFormatString(static_cast<AnimationKeyFormat>(OutCompressedData.KeyEncodingFormat));
+		FString OriginalRotationFormat = FAnimationUtils::GetAnimationCompressionFormatString(static_cast<AnimationCompressionFormat>(OutCompressedData.TranslationCompressionFormat));
+		FString OriginalTranslationFormat = FAnimationUtils::GetAnimationCompressionFormatString(static_cast<AnimationCompressionFormat>(OutCompressedData.RotationCompressionFormat));
+
 		// Get the current size
-		SIZE_T OriginalSize = AnimSeq->GetApproxCompressedSize();
+		SIZE_T OriginalSize = AfterOriginalRecompression;
 		TotalSizeBefore += OriginalSize;
 
 		// Estimate total uncompressed
-		TotalUncompressed += ((sizeof(FVector) + sizeof(FQuat) + sizeof(FVector)) * NumRawDataTracks * AnimSeq->GetRawNumberOfFrames());
-
-		// start with the current technique, or the default if none exists.
-		// this will serve as our fallback if no better technique can be found
-		const int32 OriginalKeyEncodingFormat = AnimSeq->KeyEncodingFormat;
-		const int32 OriginalTranslationFormat = AnimSeq->TranslationCompressionFormat;
-		const int32 OriginalRotationFormat = AnimSeq->RotationCompressionFormat;
+		TotalUncompressed += ((sizeof(FVector) + sizeof(FQuat) + sizeof(FVector)) * NumRawDataTracks * CompressibleAnimData.NumFrames);
 
 		// Check for global permission to try an alternative compressor
-		// we don't check for bDoNotOverrideCompression here, as that is now used as part of the UAnimCompress_Automatic compressor
 		// And it's valid to manually recompress animations
-		if( bTryAlternateCompressor /*&& (!AnimSeq->bDoNotOverrideCompression */)
+		if( bTryAlternateCompressor )
 		{
-			ensure(AnimSeq->CompressedByteStream.Num() > 0);
+			ensure(OutCompressedData.CompressedByteStream.Num() > 0);
 
 			AnimationErrorStats NewErrorStats = OriginalErrorStats;
 			if (bRaiseMaxErrorToExisting)
@@ -1214,21 +1035,6 @@ void FAnimationUtils::CompressAnimSequenceExplicit(
 			FGraphEventArray AnimCompressionTask_CompletionEvents;
 			TArray<FAnimCompressionJobContext*> AnimCompressionJobContexes;
 
-			// backup key information from the sequence
-			class UAnimCompress* SavedCompressionScheme = AnimSeq->CompressionScheme;
-			AnimationCompressionFormat SavedTranslationCompressionFormat = AnimSeq->TranslationCompressionFormat;
-			AnimationCompressionFormat SavedRotationCompressionFormat = AnimSeq->RotationCompressionFormat;
-			AnimationCompressionFormat SavedScaleCompressionFormat = AnimSeq->ScaleCompressionFormat;
-			AnimationKeyFormat SavedKeyEncodingFormat = AnimSeq->KeyEncodingFormat;
-			TArray<int32> SavedCompressedTrackOffsets = AnimSeq->CompressedTrackOffsets;
-			TArray<uint8> SavedCompressedByteStream = AnimSeq->CompressedByteStream;
-			FCompressedOffsetData SavedCompressedScaleOffsets = AnimSeq->CompressedScaleOffsets;
-			TArray<FCompressedSegment> SavedCompressedSegments = AnimSeq->CompressedSegments;
-			AnimEncoding* SavedTranslationCodec = AnimSeq->TranslationCodec;
-			AnimEncoding* SavedRotationCodec = AnimSeq->RotationCodec;
-			AnimEncoding* SavedScaleCodec = AnimSeq->ScaleCodec;
-			bool bSavedUseRawDataOnly = AnimSeq->OnlyUseRawData();
-
 			if (!bTryExhaustiveSearch)
 			{
 				// Dispatch our async compression
@@ -1256,7 +1062,7 @@ void FAnimationUtils::CompressAnimSequenceExplicit(
 						UAnimCompress_PerTrackCompression* PerTrackCompressor = NewObject<UAnimCompress_PerTrackCompression>();
 						PerTrackCompressor->bUseAdaptiveError = true;
 
-						if (AnimSeq->GetRawNumberOfFrames() > 1)
+						if (CompressibleAnimData.NumFrames > 1)
 						{
 							PerTrackCompressor->bActuallyFilterLinearKeys = true;
 							PerTrackCompressor->bRetarget = true;
@@ -1295,7 +1101,7 @@ void FAnimationUtils::CompressAnimSequenceExplicit(
 						PerTrackCompressor->bUseAdaptiveError = true;
 
 						// Try the decimation algorithms
-						if (AnimSeq->GetRawNumberOfFrames() >= PerTrackCompressor->MinKeysForResampling)
+						if (CompressibleAnimData.NumFrames >= PerTrackCompressor->MinKeysForResampling)
 						{
 							PerTrackCompressor->bActuallyFilterLinearKeys = false;
 							PerTrackCompressor->bRetarget = false;
@@ -1308,7 +1114,7 @@ void FAnimationUtils::CompressAnimSequenceExplicit(
 						}
 					}
 
-					if (AnimSeq->GetRawNumberOfFrames() > 1)
+					if (CompressibleAnimData.NumFrames > 1)
 					{
 						UAnimCompress_RemoveLinearKeys* LinearKeyRemover = NewObject<UAnimCompress_RemoveLinearKeys>();
 						{
@@ -1335,32 +1141,32 @@ void FAnimationUtils::CompressAnimSequenceExplicit(
 				}
 
 				WaitForAnimCompressionJobs(AnimCompressionTask_CompletionEvents);
-				UpdateAnimCompressionFromAsyncJobs(AnimSeq, AnimCompressionTask_CompletionEvents, AnimCompressionJobContexes, OriginalSize, CompressorStats, MasterTolerance);
+				UpdateAnimCompressionFromAsyncJobs(OutCompressedData, AnimCompressionTask_CompletionEvents, AnimCompressionJobContexes, OriginalSize, CompressorStats, MasterTolerance);
 			}
 			else
 			{
 				// Prepare to compress
 				UE_LOG(LogAnimationCompression, Log, TEXT("Compressing %s (%s)\n\tSkeleton: %s\n\tOriginal Size: %i   MaxDiff: %f"),
-					*AnimSeq->GetName(),
-					*AnimSeq->GetFullName(),
+					*CompressibleAnimData.Name,
+					*CompressibleAnimData.FullName,
 					Skeleton ? *Skeleton->GetFName().ToString() : TEXT("NULL - Not all compression techniques can be used!"),
 					OriginalSize,
 					TrueOriginalErrorStats.MaxError);
 
 				UE_LOG(LogAnimationCompression, Log, TEXT("Original Key Encoding: %s\n\tOriginal Rotation Format: %s\n\tOriginal Translation Format: %s\n\tNumFrames: %i\n\tSequenceLength: %f (%2.1f fps)"),
-					*GetAnimationKeyFormatString(static_cast<AnimationKeyFormat>(OriginalKeyEncodingFormat)),
-					*FAnimationUtils::GetAnimationCompressionFormatString(static_cast<AnimationCompressionFormat>(OriginalRotationFormat)),
-					*FAnimationUtils::GetAnimationCompressionFormatString(static_cast<AnimationCompressionFormat>(OriginalTranslationFormat)),
-					AnimSeq->GetCompressedNumberOfFrames(),
-					AnimSeq->SequenceLength,
-					(AnimSeq->GetCompressedNumberOfFrames() > 1) ? (AnimSeq->GetCompressedNumberOfFrames()-1) / AnimSeq->SequenceLength : DEFAULT_SAMPLERATE);
+					*OriginalKeyEncodingFormat,
+					*OriginalRotationFormat,
+					*OriginalTranslationFormat,
+					OutCompressedData.CompressedNumberOfFrames,
+					CompressibleAnimData.SequenceLength,
+					(OutCompressedData.CompressedNumberOfFrames > 1) ? (OutCompressedData.CompressedNumberOfFrames -1) / CompressibleAnimData.SequenceLength : DEFAULT_SAMPLERATE);
 
 				if (bFirstRecompressUsingCurrentOrDefault)
 				{
 					UE_LOG(LogAnimationCompression, Log, TEXT("Recompressed using current/default\n\tRecompress Size: %i   MaxDiff: %f\n\tRecompress Scheme: %s"),
 						AfterOriginalRecompression,
 						OriginalErrorStats.MaxError,
-						AnimSeq->CompressionScheme ? *AnimSeq->CompressionScheme->GetClass()->GetName() : TEXT("NULL"));
+						OutCompressedData.CompressionScheme ? *OutCompressedData.CompressionScheme->GetClass()->GetName() : TEXT("NULL"));
 				}
 
 				// Progressive Algorithm
@@ -1379,14 +1185,11 @@ void FAnimationUtils::CompressAnimSequenceExplicit(
 					if( NewErrorStats.MaxError >= MasterTolerance )
 					{
 						UE_LOG(LogAnimationCompression, Log, TEXT("\tStandard bitwise compressor too aggressive, lower default settings."));
-
-						AnimationErrorStats TestErrorStats;
-						FAnimationUtils::ComputeCompressionError(AnimSeq, BoneData, TestErrorStats);
 					}
 					else
 					{
 						// First, start by finding most downsampling factor.
-						if( AnimSeq->GetRawNumberOfFrames() >= PerTrackCompressor->MinKeysForResampling )
+						if (CompressibleAnimData.NumFrames >= PerTrackCompressor->MinKeysForResampling)
 						{
 							PerTrackCompressor->bResampleAnimation = true;
 				
@@ -1425,7 +1228,7 @@ void FAnimationUtils::CompressAnimSequenceExplicit(
 						}
 						
 						// Now do Linear Key Removal
-						if (AnimSeq->GetRawNumberOfFrames() > 1)
+						if (CompressibleAnimData.NumFrames > 1)
 						{
 							PerTrackCompressor->bActuallyFilterLinearKeys = true;
 							PerTrackCompressor->bRetarget = true;
@@ -1534,7 +1337,7 @@ void FAnimationUtils::CompressAnimSequenceExplicit(
 				// this compressor has a minimum number of frames requirement. So no need to go there if we don't meet that...
 				{
 					UAnimCompress_RemoveEverySecondKey* RemoveEveryOtherKeyCompressor = NewObject<UAnimCompress_RemoveEverySecondKey>();
-					if( AnimSeq->GetRawNumberOfFrames() > RemoveEveryOtherKeyCompressor->MinKeys )
+					if(CompressibleAnimData.NumFrames > RemoveEveryOtherKeyCompressor->MinKeys )
 					{
 						RemoveEveryOtherKeyCompressor->bStartAtSecondKey = false;
 						{
@@ -1586,7 +1389,7 @@ void FAnimationUtils::CompressAnimSequenceExplicit(
 				}
 
 				// construct the proposed compressor		
-				if( AnimSeq->GetRawNumberOfFrames() > 1 )
+				if(CompressibleAnimData.NumFrames > 1 )
 				{
 					UAnimCompress_RemoveLinearKeys* LinearKeyRemover = NewObject<UAnimCompress_RemoveLinearKeys>();
 					{
@@ -1648,7 +1451,7 @@ void FAnimationUtils::CompressAnimSequenceExplicit(
 					UAnimCompress_PerTrackCompression* PerTrackCompressor = NewObject<UAnimCompress_PerTrackCompression>();
 					PerTrackCompressor->bUseAdaptiveError = true;
 
-					if ( AnimSeq->GetRawNumberOfFrames() > 1 )
+					if (CompressibleAnimData.NumFrames > 1 )
 					{
 						PerTrackCompressor->bActuallyFilterLinearKeys = true;
 						PerTrackCompressor->bRetarget = true;
@@ -1681,7 +1484,7 @@ void FAnimationUtils::CompressAnimSequenceExplicit(
 					PerTrackCompressor->bUseAdaptiveError = true;
 
 					// Try the decimation algorithms
-					if (AnimSeq->GetRawNumberOfFrames() >= PerTrackCompressor->MinKeysForResampling)
+					if (CompressibleAnimData.NumFrames >= PerTrackCompressor->MinKeysForResampling)
 					{
 						PerTrackCompressor->bActuallyFilterLinearKeys = false;
 						PerTrackCompressor->bRetarget = false;
@@ -1726,7 +1529,7 @@ void FAnimationUtils::CompressAnimSequenceExplicit(
 
 				{
 					// Try the decimation algorithms
-					if (AnimSeq->GetRawNumberOfFrames() >= 3)
+					if (CompressibleAnimData.NumFrames >= 3)
 					{
 						UAnimCompress_PerTrackCompression* NewPerTrackCompressor = NewObject<UAnimCompress_PerTrackCompression>();
 
@@ -1768,7 +1571,7 @@ void FAnimationUtils::CompressAnimSequenceExplicit(
 				}
 
 				WaitForAnimCompressionJobs(AnimCompressionTask_CompletionEvents);
-				UpdateAnimCompressionFromAsyncJobs(AnimSeq, AnimCompressionTask_CompletionEvents, AnimCompressionJobContexes, OriginalSize, CompressorStats, MasterTolerance);
+				UpdateAnimCompressionFromAsyncJobs(OutCompressedData, AnimCompressionTask_CompletionEvents, AnimCompressionJobContexes, OriginalSize, CompressorStats, MasterTolerance);
 			}
 
 			// Increase winning compressor.
@@ -1783,7 +1586,7 @@ void FAnimationUtils::CompressAnimSequenceExplicit(
 					check(CompressorStats.WinningCompressorSavings == SizeDecrease);
 
 					UE_LOG(LogAnimationCompression, Log, TEXT("  Recompressing(%s) with compressor('%s') saved %i bytes (%i -> %i -> %i) (max diff=%f)\n"),
-						*AnimSeq->GetName(),
+						*CompressibleAnimData.Name,
 						*CompressorStats.WinningCompressorName,
 						SizeDecrease,
 						OriginalSize, AfterOriginalRecompression,CompressorStats. CurrentSize,
@@ -1792,24 +1595,24 @@ void FAnimationUtils::CompressAnimSequenceExplicit(
 				else
 				{
 					UE_LOG(LogAnimationCompression, Log, TEXT("  No compressor suitable! Recompressing(%s) with original/default compressor(%s) saved %i bytes (%i -> %i -> %i) (max diff=%f)\n"), 
-						*AnimSeq->GetName(),
-						*AnimSeq->CompressionScheme->GetName(),
+						*CompressibleAnimData.Name,
+						*OutCompressedData.CompressionScheme->GetName(),
 						SizeDecrease,
 						OriginalSize, AfterOriginalRecompression, CompressorStats.CurrentSize,
 						CompressorStats.WinningCompressorError);
 
 					UE_LOG(LogAnimation, Warning, TEXT("  CompressedTrackOffsets(%d) CompressedByteStream(%d) CompressedScaleOffsets(%d) CompressedSegments(%d)"),
-						AnimSeq->CompressedTrackOffsets.Num(),
-						AnimSeq->CompressedByteStream.Num(),
-						AnimSeq->CompressedScaleOffsets.GetMemorySize(),
-						AnimSeq->CompressedSegments.Num());
+						OutCompressedData.CompressedTrackOffsets.Num(),
+						OutCompressedData.CompressedByteStream.Num(),
+						OutCompressedData.CompressedScaleOffsets.GetMemorySize(),
+						0);
 
 					TotalNoWinnerRounds++;
 				}
 			}
 
 			// Make sure we got that right.
-			check(CompressorStats.CurrentSize == AnimSeq->GetApproxCompressedSize());
+			check(CompressorStats.CurrentSize == OutCompressedData.GetApproxBoneCompressedSize());
 			TotalSizeNow += CompressorStats.CurrentSize;
 
 			CompressorStats.PctSaving = TotalSizeBefore > 0 ? 100.f - (100.f * float(TotalSizeNow) / float(TotalSizeBefore)) : 0.f;
@@ -1884,32 +1687,17 @@ void FAnimationUtils::CompressAnimSequenceExplicit(
 		// Do not recompress - Still take into account size for stats.
 		else
 		{
-			TotalSizeNow += AnimSeq->GetApproxCompressedSize();
+			TotalSizeNow += OutCompressedData.GetApproxBoneCompressedSize();
 		}
 	}
 	else
 	{
 		// this can happen if the animation only contains curve - i.e. blendshape curves
-		UE_LOG(LogAnimationCompression, Log, TEXT("Compression Requested for Empty Animation %s"), *AnimSeq->GetName() );
+		UE_LOG(LogAnimationCompression, Log, TEXT("Compression Requested for Empty Animation %s"), *CompressibleAnimData.Name );
 	}
 #endif // WITH_EDITORONLY_DATA
 }
 
-
-void FAnimationUtils::TestForMissingMeshes(UAnimSequence* AnimSeq)
-{
-	if (FPlatformProperties::HasEditorOnlyData())
-	{
-		check(AnimSeq != NULL);
-
-		USkeleton* Skeleton = AnimSeq->GetSkeleton();
-		check( Skeleton );
-
-		static int32 MissingSkeletonCount = 0;
-		static TArray<FString> MissingSkeletonArray;
-	}
-}
-
 static void GetBindPoseAtom(FTransform &OutBoneAtom, int32 BoneIndex, USkeleton *Skeleton)
 {
 	OutBoneAtom = Skeleton->GetRefLocalPoses()[BoneIndex];
@@ -2026,7 +1814,7 @@ FString FAnimationUtils::GetAnimationKeyFormatString(AnimationKeyFormat InFormat
  * @param TrackHeights [OUT]	The computed track heights
  *
  */
-void FAnimationUtils::CalculateTrackHeights(UAnimSequence* AnimSeq, const TArray<FBoneData>& BoneData, int32 NumTracks, TArray<int32>& TrackHeights)
+void FAnimationUtils::CalculateTrackHeights(const FCompressibleAnimData& CompressibleAnimData, int32 NumTracks, TArray<int32>& TrackHeights)
 {
 	TrackHeights.Empty();
 	TrackHeights.AddUninitialized(NumTracks);
@@ -2035,8 +1823,7 @@ void FAnimationUtils::CalculateTrackHeights(UAnimSequence* AnimSeq, const TArray
 		TrackHeights[TrackIndex] = 0;
 	}
 
-	USkeleton* Skeleton = AnimSeq->GetSkeleton();
-	check(Skeleton);
+	const TArray<FBoneData>& BoneData = CompressibleAnimData.BoneData;
 
 	// Populate the bone 'height' table (distance from closest end effector, with 0 indicating an end effector)
 	// setup the raw bone transformation and find all end effectors
@@ -2051,7 +1838,7 @@ void FAnimationUtils::CalculateTrackHeights(UAnimSequence* AnimSeq, const TArray
 			for (int32 FamilyIndex = 0; FamilyIndex < EffectorBoneData.BonesToRoot.Num(); ++FamilyIndex)
 			{
 				const int32 NextParentBoneIndex = EffectorBoneData.BonesToRoot[FamilyIndex];
-				const int32 NextParentTrackIndex = Skeleton->GetAnimationTrackIndex(NextParentBoneIndex, AnimSeq, true);
+				const int32 NextParentTrackIndex = GetAnimTrackIndexForSkeletonBone(NextParentBoneIndex, CompressibleAnimData.TrackToSkeletonMapTable);
 				if (NextParentTrackIndex != INDEX_NONE)
 				{
 					const int32 CurHeight = TrackHeights[NextParentTrackIndex];
@@ -2072,9 +1859,9 @@ void FAnimationUtils::CalculateTrackHeights(UAnimSequence* AnimSeq, const TArray
  *
  * @return				true if the keys are uniformly spaced (or one of the trivial conditions is detected).  false if any key spacing is greater than 1e-4 off.
  */
-bool FAnimationUtils::HasUniformKeySpacing(UAnimSequence* AnimSeq, const TArray<float>& Times)
+bool FAnimationUtils::HasUniformKeySpacing(int32 NumFrames, const TArray<float>& Times)
 {
-	if ((Times.Num() <= 2) || (Times.Num() == AnimSeq->GetRawNumberOfFrames()))
+	if ((Times.Num() <= 2) || (Times.Num() == NumFrames))
 	{
 		return true;
 	}
@@ -2097,19 +1884,18 @@ bool FAnimationUtils::HasUniformKeySpacing(UAnimSequence* AnimSeq, const TArray<
  * Perturbs the bone(s) associated with each track in turn, measuring the maximum error introduced in end effectors as a result
  */
 void FAnimationUtils::TallyErrorsFromPerturbation(
-	const UAnimSequence* AnimSeq,
+	const FCompressibleAnimData& CompressibleAnimData,
 	int32 NumTracks,
-	const TArray<FBoneData>& BoneData,
 	const FVector& PositionNudge,
 	const FQuat& RotationNudge,
 	const FVector& ScaleNudge,
 	TArray<FAnimPerturbationError>& InducedErrors)
 {
-	const float TimeStep = (float)AnimSeq->SequenceLength / (float)AnimSeq->GetRawNumberOfFrames();
-	const int32 NumBones = BoneData.Num();
+	const float TimeStep = (float)CompressibleAnimData.SequenceLength / CompressibleAnimData.NumFrames;
+	const int32 NumBones = CompressibleAnimData.BoneData.Num();
 
 
-	USkeleton* Skeleton = AnimSeq->GetSkeleton();
+	USkeleton* Skeleton = CompressibleAnimData.Skeleton;
 	check ( Skeleton );
 
 	const TArray<FTransform>& RefPose = Skeleton->GetRefLocalPoses();
@@ -2149,12 +1935,12 @@ void FAnimationUtils::TallyErrorsFromPerturbation(
 		float MaxErrorS_DueToS = 0.0f;
 		
 		// for each whole increment of time (frame stepping)
-		for (float Time = 0.0f; Time < AnimSeq->SequenceLength; Time += TimeStep)
+		for (float Time = 0.0f; Time < CompressibleAnimData.SequenceLength; Time += TimeStep)
 		{
 			// get the raw and compressed atom for each bone
 			for (int32 BoneIndex = 0; BoneIndex < NumBones; ++BoneIndex)
 			{
-				const int32 TrackIndex = Skeleton->GetAnimationTrackIndex(BoneIndex, AnimSeq, true);
+				const int32 TrackIndex = GetAnimTrackIndexForSkeletonBone(BoneIndex, CompressibleAnimData.TrackToSkeletonMapTable);
 
 				if (TrackIndex == INDEX_NONE)
 				{
@@ -2166,7 +1952,7 @@ void FAnimationUtils::TallyErrorsFromPerturbation(
 				}
 				else
 				{
-					AnimSeq->GetBoneTransform(RawAtoms[BoneIndex], TrackIndex, Time, true);
+					ExtractTransformFromTrack(Time, CompressibleAnimData.NumFrames, CompressibleAnimData.SequenceLength, CompressibleAnimData.RawAnimationData[TrackIndex], CompressibleAnimData.Interpolation, RawAtoms[BoneIndex]);
 
 					NewAtomsT[BoneIndex] = RawAtoms[BoneIndex];
 					NewAtomsR[BoneIndex] = RawAtoms[BoneIndex];
@@ -2208,7 +1994,7 @@ void FAnimationUtils::TallyErrorsFromPerturbation(
 				}
 
 				// Only look at the error that occurs in end effectors
-				if (BoneData[BoneIndex].IsEndEffector())
+				if (CompressibleAnimData.BoneData[BoneIndex].IsEndEffector())
 				{
 					MaxErrorT_DueToT = FMath::Max(MaxErrorT_DueToT, (RawTransforms[BoneIndex].GetLocation() - NewTransformsT[BoneIndex].GetLocation()).Size());
 					MaxErrorT_DueToR = FMath::Max(MaxErrorT_DueToR, (RawTransforms[BoneIndex].GetLocation() - NewTransformsR[BoneIndex].GetLocation()).Size());
@@ -2296,13 +2082,113 @@ UAnimCurveCompressionSettings* FAnimationUtils::GetDefaultAnimationCurveCompress
 	return DefaultCurveCompressionSettings;
 }
 
+void FAnimationUtils::ExtractTransformFromTrack(float Time, int32 NumFrames, float SequenceLength, const struct FRawAnimSequenceTrack& RawTrack, EAnimInterpolationType Interpolation, FTransform &OutAtom)
+{
+	// Bail out (with rather wacky data) if data is empty for some reason.
+	if (RawTrack.PosKeys.Num() == 0 || RawTrack.RotKeys.Num() == 0)
+	{
+		OutAtom.SetIdentity();
+		return;
+	}
+
+	int32 KeyIndex1, KeyIndex2;
+	float Alpha;
+	FAnimationRuntime::GetKeyIndicesFromTime(KeyIndex1, KeyIndex2, Alpha, Time, NumFrames, SequenceLength);
+	// @Todo fix me: this change is not good, it has lots of branches. But we'd like to save memory for not saving scale if no scale change exists
+	const bool bHasScaleKey = (RawTrack.ScaleKeys.Num() > 0);
+	static const FVector DefaultScale3D = FVector(1.f);
+
+	if (Interpolation == EAnimInterpolationType::Step)
+	{
+		Alpha = 0.f;
+	}
+
+	if (Alpha <= 0.f)
+	{
+		const int32 PosKeyIndex1 = FMath::Min(KeyIndex1, RawTrack.PosKeys.Num() - 1);
+		const int32 RotKeyIndex1 = FMath::Min(KeyIndex1, RawTrack.RotKeys.Num() - 1);
+		if (bHasScaleKey)
+		{
+			const int32 ScaleKeyIndex1 = FMath::Min(KeyIndex1, RawTrack.ScaleKeys.Num() - 1);
+			OutAtom = FTransform(RawTrack.RotKeys[RotKeyIndex1], RawTrack.PosKeys[PosKeyIndex1], RawTrack.ScaleKeys[ScaleKeyIndex1]);
+		}
+		else
+		{
+			OutAtom = FTransform(RawTrack.RotKeys[RotKeyIndex1], RawTrack.PosKeys[PosKeyIndex1], DefaultScale3D);
+		}
+		return;
+	}
+	else if (Alpha >= 1.f)
+	{
+		const int32 PosKeyIndex2 = FMath::Min(KeyIndex2, RawTrack.PosKeys.Num() - 1);
+		const int32 RotKeyIndex2 = FMath::Min(KeyIndex2, RawTrack.RotKeys.Num() - 1);
+		if (bHasScaleKey)
+		{
+			const int32 ScaleKeyIndex2 = FMath::Min(KeyIndex2, RawTrack.ScaleKeys.Num() - 1);
+			OutAtom = FTransform(RawTrack.RotKeys[RotKeyIndex2], RawTrack.PosKeys[PosKeyIndex2], RawTrack.ScaleKeys[ScaleKeyIndex2]);
+		}
+		else
+		{
+			OutAtom = FTransform(RawTrack.RotKeys[RotKeyIndex2], RawTrack.PosKeys[PosKeyIndex2], DefaultScale3D);
+		}
+		return;
+	}
+
+	const int32 PosKeyIndex1 = FMath::Min(KeyIndex1, RawTrack.PosKeys.Num() - 1);
+	const int32 RotKeyIndex1 = FMath::Min(KeyIndex1, RawTrack.RotKeys.Num() - 1);
+
+	const int32 PosKeyIndex2 = FMath::Min(KeyIndex2, RawTrack.PosKeys.Num() - 1);
+	const int32 RotKeyIndex2 = FMath::Min(KeyIndex2, RawTrack.RotKeys.Num() - 1);
+
+	FTransform KeyAtom1, KeyAtom2;
+
+	if (bHasScaleKey)
+	{
+		const int32 ScaleKeyIndex1 = FMath::Min(KeyIndex1, RawTrack.ScaleKeys.Num() - 1);
+		const int32 ScaleKeyIndex2 = FMath::Min(KeyIndex2, RawTrack.ScaleKeys.Num() - 1);
+
+		KeyAtom1 = FTransform(RawTrack.RotKeys[RotKeyIndex1], RawTrack.PosKeys[PosKeyIndex1], RawTrack.ScaleKeys[ScaleKeyIndex1]);
+		KeyAtom2 = FTransform(RawTrack.RotKeys[RotKeyIndex2], RawTrack.PosKeys[PosKeyIndex2], RawTrack.ScaleKeys[ScaleKeyIndex2]);
+	}
+	else
+	{
+		KeyAtom1 = FTransform(RawTrack.RotKeys[RotKeyIndex1], RawTrack.PosKeys[PosKeyIndex1], DefaultScale3D);
+		KeyAtom2 = FTransform(RawTrack.RotKeys[RotKeyIndex2], RawTrack.PosKeys[PosKeyIndex2], DefaultScale3D);
+	}
+
+	// 	UE_LOG(LogAnimation, Log, TEXT(" *  *  *  Position. PosKeyIndex1: %3d, PosKeyIndex2: %3d, Alpha: %f"), PosKeyIndex1, PosKeyIndex2, Alpha);
+	// 	UE_LOG(LogAnimation, Log, TEXT(" *  *  *  Rotation. RotKeyIndex1: %3d, RotKeyIndex2: %3d, Alpha: %f"), RotKeyIndex1, RotKeyIndex2, Alpha);
+
+	// Ensure rotations are normalized (Added for Jira UE-53971)
+	KeyAtom1.NormalizeRotation();
+	KeyAtom2.NormalizeRotation();
+
+	OutAtom.Blend(KeyAtom1, KeyAtom2, Alpha);
+	OutAtom.NormalizeRotation();
+}
+
 #if WITH_EDITOR
 
-bool FAnimationUtils::CompressAnimCurves(UAnimSequence& AnimSeq)
+void FAnimationUtils::ExtractTransformFromCompressionData(const FCompressibleAnimData& CompressibleAnimData, FCompressibleAnimDataResult& CompressedAnimData, float Time, int32 TrackIndex, bool bUseRawData, FTransform& OutBoneTransform)
+{
+	// If the caller didn't request that raw animation data be used . . .
+	if (!bUseRawData && CompressedAnimData.IsCompressedDataValid())
+	{
+		const FUECompressedAnimData CompressedDataWrapper(CompressedAnimData);
+		FAnimSequenceDecompressionContext DecompContext(CompressibleAnimData, CompressedDataWrapper);
+		DecompContext.Seek(Time);
+		AnimationFormat_GetBoneAtom(OutBoneTransform, DecompContext, TrackIndex);
+		return;
+	}
+
+	FAnimationUtils::ExtractTransformFromTrack(Time, CompressibleAnimData.NumFrames, CompressibleAnimData.SequenceLength, CompressibleAnimData.RawAnimationData[TrackIndex], CompressibleAnimData.Interpolation, OutBoneTransform);
+}
+
+bool FAnimationUtils::CompressAnimCurves(FCompressibleAnimData& AnimSeq, FCompressedAnimSequence& Target)
 {
 	// Clear any previous data we might have even if we end up failing to compress
-	AnimSeq.CompressedCurveByteStream.Empty();
-	AnimSeq.CurveCompressionCodec = nullptr;
+	Target.CompressedCurveByteStream.Empty();
+	Target.CurveCompressionCodec = nullptr;
 
 	if (AnimSeq.CurveCompressionSettings == nullptr || !AnimSeq.CurveCompressionSettings->AreSettingsValid())
 	{
@@ -2310,6 +2196,6 @@ bool FAnimationUtils::CompressAnimCurves(UAnimSequence& AnimSeq)
 	}
 
 	check(AnimSeq.CurveCompressionSettings->AreSettingsValid());
-	return AnimSeq.CurveCompressionSettings->Compress(AnimSeq);
+	return AnimSeq.CurveCompressionSettings->Compress(AnimSeq, Target);
 }
 #endif
diff --git a/Engine/Source/Runtime/Engine/Private/Animation/Skeleton.cpp b/Engine/Source/Runtime/Engine/Private/Animation/Skeleton.cpp
--- a/Engine/Source/Runtime/Engine/Private/Animation/Skeleton.cpp
+++ b/Engine/Source/Runtime/Engine/Private/Animation/Skeleton.cpp
@@ -744,19 +744,14 @@ FName USkeleton::GetRetargetSourceForMesh(USkeletalMesh* InMesh) const
 
 #endif
 
-int32 USkeleton::GetAnimationTrackIndex(const int32 InSkeletonBoneIndex, const UAnimSequence* InAnimSeq, const bool bUseRawData)
+int32 USkeleton::GetRawAnimationTrackIndex(const int32 InSkeletonBoneIndex, const UAnimSequence* InAnimSeq)
 {
-	const TArray<FTrackToSkeletonMap>& TrackToSkelMap = bUseRawData ? InAnimSeq->GetRawTrackToSkeletonMapTable() : InAnimSeq->GetCompressedTrackToSkeletonMapTable();
 	if( InSkeletonBoneIndex != INDEX_NONE )
 	{
-		for (int32 TrackIndex = 0; TrackIndex<TrackToSkelMap.Num(); ++TrackIndex)
+		return InAnimSeq->GetRawTrackToSkeletonMapTable().IndexOfByPredicate([&](const FTrackToSkeletonMap& TrackToSkel)
 		{
-			const FTrackToSkeletonMap& TrackToSkeleton = TrackToSkelMap[TrackIndex];
-			if( TrackToSkeleton.BoneTreeIndex == InSkeletonBoneIndex )
-			{
-				return TrackIndex;
-			}
-		}
+			return TrackToSkel.BoneTreeIndex == InSkeletonBoneIndex;
+		});
 	}
 
 	return INDEX_NONE;
diff --git a/Engine/Source/Runtime/Engine/Private/UnrealEngine.cpp b/Engine/Source/Runtime/Engine/Private/UnrealEngine.cpp
--- a/Engine/Source/Runtime/Engine/Private/UnrealEngine.cpp
+++ b/Engine/Source/Runtime/Engine/Private/UnrealEngine.cpp
@@ -5452,9 +5452,10 @@ bool UEngine::HandleListAnimsCommand(const TCHAR* Cmd, FOutputDevice& Ar)
 			RateScale = AnimSeq->RateScale;
 			NumCurves = AnimSeq->RawCurveData.FloatCurves.Num();
 
-			TranslationFormat = FAnimationUtils::GetAnimationCompressionFormatString(AnimSeq->TranslationCompressionFormat);
-			RotationFormat = FAnimationUtils::GetAnimationCompressionFormatString(AnimSeq->RotationCompressionFormat);
-			ScaleFormat = FAnimationUtils::GetAnimationCompressionFormatString(AnimSeq->ScaleCompressionFormat);
+			const FUECompressedAnimData& CompressedData = AnimSeq->CompressedData.CompressedDataStructure;
+			TranslationFormat = FAnimationUtils::GetAnimationCompressionFormatString(CompressedData.TranslationCompressionFormat);
+			RotationFormat = FAnimationUtils::GetAnimationCompressionFormatString(CompressedData.RotationCompressionFormat);
+			ScaleFormat = FAnimationUtils::GetAnimationCompressionFormatString(CompressedData.ScaleCompressionFormat);
 		}
 
 		new(SortedAnimAssets) FSortedAnimAsset(
diff --git a/Engine/Source/Runtime/Engine/Public/AnimEncoding.h b/Engine/Source/Runtime/Engine/Public/AnimEncoding.h
--- a/Engine/Source/Runtime/Engine/Public/AnimEncoding.h
+++ b/Engine/Source/Runtime/Engine/Public/AnimEncoding.h
@@ -219,7 +219,7 @@ void AnimationFormat_GetAnimationPose(
  * @param	NumTransTracksWithOneKey	The total number of Translation Tracks found containing a single key.
  * @param	NumRotTracksWithOneKey		The total number of Rotation Tracks found containing a single key.
 */
-ENGINE_API void AnimationFormat_GetStats(const UAnimSequence* Seq, 
+ENGINE_API void AnimationFormat_GetStats(const FUECompressedAnimData& CompressedData, 
   								int32& NumTransTracks,
 								int32& NumRotTracks,
 								int32& NumScaleTracks,
@@ -239,7 +239,8 @@ ENGINE_API void AnimationFormat_GetStats(const UAnimSequence* Seq,
  *
  * @param	Seq					An Animation Sequence to setup links within.
 */
-void AnimationFormat_SetInterfaceLinks(UAnimSequence& Seq);
+template<typename CompressedDataType>
+void AnimationFormat_SetInterfaceLinks(CompressedDataType& CompressedData);
 
 #if WITH_EDITORONLY_DATA
 #define AC_UnalignedSwap( MemoryArchive, Data, Len )		\
@@ -273,26 +274,24 @@ public:
 	/**
 	 * Handles Byte-swapping incoming animation data from a MemoryReader
 	 *
-	 * @param	Seq					An Animation Sequence to contain the read data.
+	 * @param	CompressedData		The compressed animation data being operated on.
 	 * @param	MemoryReader		The MemoryReader object to read from.
 	 * @param	SourceArVersion		The version of the archive that the data is coming from.
 	 */
 	virtual void ByteSwapIn(
-		UAnimSequence& Seq, 
+		FUECompressedAnimData& CompressedData,
 		FMemoryReader& MemoryReader) PURE_VIRTUAL(AnimEncoding::ByteSwapIn,);
 
 	/**
 	 * Handles Byte-swapping outgoing animation data to an array of BYTEs
 	 *
-	 * @param	Seq					An Animation Sequence to write.
+	 * @param	CompressedData		The compressed animation data being operated on.
 	 * @param	SerializedData		The output buffer.
 	 * @param	ForceByteSwapping	true is byte swapping is not optional.
 	 */
 	virtual void ByteSwapOut(
-		UAnimSequence& Seq,
-		TArray<uint8>& SerializedData, 
-		bool ForceByteSwapping,
-		bool bMaintainComponentOrder=false) PURE_VIRTUAL(AnimEncoding::ByteSwapOut, );
+		FUECompressedAnimData& CompressedData,
+		FMemoryWriter& MemoryWriter) PURE_VIRTUAL(AnimEncoding::ByteSwapOut, );
 
 	/**
 	 * Extracts a single BoneAtom from an Animation Sequence.
@@ -350,6 +349,7 @@ public:
 		FAnimSequenceDecompressionContext& DecompContext) PURE_VIRTUAL(AnimEncoding::GetPoseScales,);
 #endif
 
+#if USE_SEGMENTING_CONTEXT
 	static FORCEINLINE float TimeToIndex(
 		const FAnimSequenceDecompressionContext& DecompContext,
 		const uint8* TimeMarkers,
@@ -359,30 +359,34 @@ public:
 		float SegmentRelativePos,
 		int32& FrameIndex0Out,
 		int32& FrameIndex1Out);
+#endif
 
 protected:
 
 	/**
 	 * Utility function to determine the two key indices to interpolate given a relative position in the animation
 	 *
-	 * @param	Seq				The UAnimSequence container.
+	 * @param	SequenceLength	The length of the anim sequence 
 	 * @param	RelativePos		The relative position to solve in the range [0,1] inclusive.
 	 * @param	NumKeys			The number of keys present in the track being solved.
+	 * @param	Interpolation	The Interpolation type of the sequence
 	 * @param	PosIndex0Out	Output value for the closest key index before the RelativePos specified.
 	 * @param	PosIndex1Out	Output value for the closest key index after the RelativePos specified.
 	 * @return	The rate at which to interpolate the two keys returned to obtain the final result.
 	 */
 	static float TimeToIndex(
-		const UAnimSequence& Seq,
+		float SequenceLength,
 		float RelativePos,
 		int32 NumKeys,
+		EAnimInterpolationType Interpolation,
 		int32 &PosIndex0Out,
 		int32 &PosIndex1Out);
 
 	/**
 	 * Utility function to determine the two key indices to interpolate given a relative position in the animation
 	 *
-	 * @param	Seq				The UAnimSequence container.
+	 * @param	Interpolation	The Interpolation type of the sequence
+	 * @param	NumberOfFrames	The number of frames in the original sequence
 	 * @param	FrameTable		The frame table containing a frame index for each key.
 	 * @param	RelativePos		The relative position to solve in the range [0,1] inclusive.
 	 * @param	NumKeys			The number of keys present in the track being solved.
@@ -391,7 +395,8 @@ protected:
 	 * @return	The rate at which to interpolate the two keys returned to obtain the final result.
 	 */
 	static float TimeToIndex(
-		const UAnimSequence& Seq,
+		EAnimInterpolationType Interpolation,
+		int32 NumberOfFrames,
 		const uint8* FrameTable,
 		float RelativePos,
 		int32 NumKeys,
@@ -461,30 +466,28 @@ public:
 	/**
 	 * Handles Byte-swapping incoming animation data from a MemoryReader
 	 *
-	 * @param	Seq					An Animation Sequence to contain the read data.
+	 * @param	CompressedData		The compressed animation data being operated on.
 	 * @param	MemoryReader		The MemoryReader object to read from.
 	 */
 	virtual void ByteSwapIn(
-		UAnimSequence& Seq, 
+		FUECompressedAnimData& CompressedData,
 		FMemoryReader& MemoryReader) override;
 
 	/**
 	 * Handles Byte-swapping outgoing animation data to an array of BYTEs
 	 *
-	 * @param	Seq					An Animation Sequence to write.
+	 * @param	CompressedData		The compressed animation data being operated on.
 	 * @param	SerializedData		The output buffer.
 	 * @param	ForceByteSwapping	true is byte swapping is not optional.
 	 */
 	virtual void ByteSwapOut(
-		UAnimSequence& Seq,
-		TArray<uint8>& SerializedData,
-		bool ForceByteSwapping,
-		bool bMaintainComponentOrder=false) override;
+		FUECompressedAnimData& CompressedData,
+		FMemoryWriter& MemoryWriter) override;
 
 	/**
 	 * Handles the ByteSwap of compressed animation data on import
 	 *
-	 * @param	Seq				The UAnimSequence container.
+	 * @param	CompressedData		The compressed animation data being operated on.
 	 * @param	MemoryReader	The FMemoryReader to read from.
 	 * @param	Stream			The compressed animation data.
 	 * @param	NumKeys			The number of keys present in Stream.
@@ -492,7 +495,7 @@ public:
 	 * @return					The adjusted Stream position after import. 
 	 */
 	virtual void ByteSwapRotationIn(
-		UAnimSequence& Seq, 
+		FUECompressedAnimData& CompressedData,
 		FMemoryReader& MemoryReader,
 		uint8*& Stream,
 		int32 NumKeys) PURE_VIRTUAL(AnimEncoding::ByteSwapRotationIn,);
@@ -500,7 +503,7 @@ public:
 	/**
 	 * Handles the ByteSwap of compressed animation data on import
 	 *
-	 * @param	Seq				The UAnimSequence container.
+	 * @param	CompressedData	The compressed animation data being operated on.
 	 * @param	MemoryReader	The FMemoryReader to read from.
 	 * @param	Stream			The compressed animation data.
 	 * @param	NumKeys			The number of keys present in Stream.
@@ -508,7 +511,7 @@ public:
 	 * @return					The adjusted Stream position after import. 
 	 */
 	virtual void ByteSwapTranslationIn(
-		UAnimSequence& Seq, 
+		FUECompressedAnimData& CompressedData,
 		FMemoryReader& MemoryReader,
 		uint8*& Stream,
 		int32 NumKeys) PURE_VIRTUAL(AnimEncoding::ByteSwapTranslationIn,);
@@ -516,7 +519,7 @@ public:
 	/**
 	 * Handles the ByteSwap of compressed animation data on import
 	 *
-	 * @param	Seq				The UAnimSequence container.
+	 * @param	CompressedData	The compressed animation data being operated on.
 	 * @param	MemoryReader	The FMemoryReader to read from.
 	 * @param	Stream			The compressed animation data.
 	 * @param	NumKeys			The number of keys present in Stream.
@@ -524,7 +527,7 @@ public:
 	 * @return					The adjusted Stream position after import. 
 	 */
 	virtual void ByteSwapScaleIn(
-		UAnimSequence& Seq, 
+		FUECompressedAnimData& CompressedData,
 		FMemoryReader& MemoryReader,
 		uint8*& Stream,
 		int32 NumKeys) PURE_VIRTUAL(AnimEncoding::ByteSwapScaleIn,);
@@ -532,14 +535,14 @@ public:
 	/**
 	 * Handles the ByteSwap of compressed animation data on export
 	 *
-	 * @param	Seq				The UAnimSequence container.
+	 * @param	CompressedData	The compressed animation data being operated on.
 	 * @param	MemoryWriter	The FMemoryReader to write to.
 	 * @param	Stream			The compressed animation data.
 	 * @param	NumKeys			The number of keys present in Stream.
 	 * @return					The adjusted Stream position after export. 
 	 */
 	virtual void ByteSwapRotationOut(
-		UAnimSequence& Seq, 
+		FUECompressedAnimData& CompressedData,
 		FMemoryWriter& MemoryWriter,
 		uint8*& Stream,
 		int32 NumKeys) PURE_VIRTUAL(AnimEncoding::ByteSwapRotationOut,);
@@ -547,14 +550,14 @@ public:
 	/**
 	 * Handles the ByteSwap of compressed animation data on export
 	 *
-	 * @param	Seq				The UAnimSequence container.
+	 * @param	CompressedData	The compressed animation data being operated on.
 	 * @param	MemoryWriter	The FMemoryReader to write to.
 	 * @param	Stream			The compressed animation data.
 	 * @param	NumKeys			The number of keys present in Stream.
 	 * @return					The adjusted Stream position after export. 
 	 */
 	virtual void ByteSwapTranslationOut(
-		UAnimSequence& Seq, 
+		FUECompressedAnimData& CompressedData,
 		FMemoryWriter& MemoryWriter,
 		uint8*& Stream,
 		int32 NumKeys) PURE_VIRTUAL(AnimEncoding::ByteSwapTranslationOut,);
@@ -562,14 +565,14 @@ public:
 	/**
 	 * Handles the ByteSwap of compressed animation data on export
 	 *
-	 * @param	Seq				The UAnimSequence container.
+	 * @param	CompressedData	The compressed animation data being operated on.
 	 * @param	MemoryWriter	The FMemoryReader to write to.
 	 * @param	Stream			The compressed animation data.
 	 * @param	NumKeys			The number of keys present in Stream.
 	 * @return					The adjusted Stream position after export. 
 	 */
 	virtual void ByteSwapScaleOut(
-		UAnimSequence& Seq, 
+		FUECompressedAnimData& CompressedData,
 		FMemoryWriter& MemoryWriter,
 		uint8*& Stream,
 		int32 NumKeys) PURE_VIRTUAL(AnimEncoding::ByteSwapScaleOut,);
@@ -587,13 +590,13 @@ public:
  * @return	The rate at which to interpolate the two keys returned to obtain the final result.
  */
 FORCEINLINE float AnimEncoding::TimeToIndex(
-	const UAnimSequence& Seq,
+	float SequenceLength,
 	float RelativePos,
 	int32 NumKeys,
+	EAnimInterpolationType Interpolation,
 	int32 &PosIndex0Out,
 	int32 &PosIndex1Out)
 {
-	const float SequenceLength= Seq.SequenceLength;
 	float Alpha;
 
 	if (NumKeys < 2)
@@ -628,7 +631,7 @@ FORCEINLINE float AnimEncoding::TimeToIndex(
 			checkSlow(KeyPos >= 0.0f);
 			const float KeyPosFloor = floorf(KeyPos);
 			PosIndex0Out = FMath::Min( FMath::TruncToInt(KeyPosFloor), NumKeys );
-			Alpha = Seq.Interpolation == EAnimInterpolationType::Step ? 0.0f : KeyPos - KeyPosFloor;
+			Alpha = (Interpolation == EAnimInterpolationType::Step) ? 0.0f : KeyPos - KeyPosFloor;
 			PosIndex1Out = FMath::Min( PosIndex0Out + 1, NumKeys );
 		}
 	}
@@ -702,7 +705,8 @@ FORCEINLINE_DEBUGGABLE int32 FindLowKeyIndex(
  * @return	The rate at which to interpolate the two keys returned to obtain the final result.
  */
 FORCEINLINE float AnimEncoding::TimeToIndex(
-	const UAnimSequence& Seq,
+	EAnimInterpolationType Interpolation,
+	int32 NumberOfFrames,
 	const uint8* FrameTable,
 	float RelativePos,
 	int32 NumKeys,
@@ -715,7 +719,7 @@ FORCEINLINE float AnimEncoding::TimeToIndex(
 	
 	const int32 LastKey= NumKeys-1;
 	
-	int32 TotalFrames = Seq.GetCompressedNumberOfFrames()-1;
+	int32 TotalFrames = NumberOfFrames -1;
 	int32 EndingKey = LastKey;
 
 	if (NumKeys < 2 || RelativePos <= 0.f)
@@ -745,7 +749,7 @@ FORCEINLINE float AnimEncoding::TimeToIndex(
 		int32 HighFrame = 0;
 		
 		// find the pair of keys which surround our target frame index
-		if (Seq.GetCompressedNumberOfFrames() > 0xFF)
+		if (NumberOfFrames > 0xFF)
 		{
 			const uint16* Frames= (uint16*)FrameTable;
 			PosIndex0Out = FindLowKeyIndex<uint16>(Frames, NumKeys, FramePosFloor, KeyEstimate);
@@ -775,12 +779,13 @@ FORCEINLINE float AnimEncoding::TimeToIndex(
 		// compute the blend parameters for the keys we have found
 		int32 Delta= FMath::Max(HighFrame - LowFrame, 1);
 		const float Remainder = (FramePos - (float)LowFrame);
-		Alpha = Seq.Interpolation == EAnimInterpolationType::Step ? 0.f : (Remainder / (float)Delta);
+		Alpha = Interpolation == EAnimInterpolationType::Step ? 0.f : (Remainder / (float)Delta);
 	}
 	
 	return Alpha;
 }
 
+#if USE_SEGMENTING_CONTEXT
 float AnimEncoding::TimeToIndex(
 	const FAnimSequenceDecompressionContext& DecompContext,
 	const uint8* TimeMarkers,
@@ -845,9 +850,9 @@ float AnimEncoding::TimeToIndex(
 		// compute the blend parameters for the keys we have found
 		const int32 Delta = FMath::Max(HighFrame - LowFrame, 1);
 		const float Remainder = FramePos - (float)LowFrame;
-		Alpha = DecompContext.AnimSeq->Interpolation == EAnimInterpolationType::Step ? 0.f : (Remainder / (float)Delta);
+		Alpha = DecompContext.GetInterpolation() == EAnimInterpolationType::Step ? 0.f : (Remainder / (float)Delta);
 	}
 
 	return Alpha;
 }
-
+#endif
diff --git a/Engine/Source/Runtime/Engine/Public/AnimEncoding_ConstantKeyLerp.h b/Engine/Source/Runtime/Engine/Public/AnimEncoding_ConstantKeyLerp.h
--- a/Engine/Source/Runtime/Engine/Public/AnimEncoding_ConstantKeyLerp.h
+++ b/Engine/Source/Runtime/Engine/Public/AnimEncoding_ConstantKeyLerp.h
@@ -40,13 +40,13 @@ public:
 	/**
 	 * Handles the ByteSwap of compressed rotation data on import
 	 *
-	 * @param	Seq				The UAnimSequence container.
+	 * @param	CompressedData		The compressed animation data being operated on.
 	 * @param	MemoryReader	The FMemoryReader to read from.
 	 * @param	RotTrackData	The compressed rotation data stream.
 	 * @param	NumKeysRot		The number of keys present in the stream.
 	 */
 	virtual void ByteSwapRotationIn(
-		UAnimSequence& Seq, 
+		FUECompressedAnimData& CompressedData,
 		FMemoryReader& MemoryReader,
 		uint8*& RotTrackData,
 		int32 NumKeysRot) override;
@@ -54,13 +54,13 @@ public:
 	/**
 	 * Handles the ByteSwap of compressed translation data on import
 	 *
-	 * @param	Seq				The UAnimSequence container.
+	 * @param	CompressedData	The compressed animation data being operated on.
 	 * @param	MemoryReader	The FMemoryReader to read from.
 	 * @param	TransTrackData	The compressed translation data stream.
 	 * @param	NumKeysTrans	The number of keys present in the stream.
 	 */
 	virtual void ByteSwapTranslationIn(
-		UAnimSequence& Seq, 
+		FUECompressedAnimData& CompressedData,
 		FMemoryReader& MemoryReader,
 		uint8*& TransTrackData,
 		int32 NumKeysTrans) override;
@@ -68,13 +68,13 @@ public:
 	/**
 	 * Handles the ByteSwap of compressed Scale data on import
 	 *
-	 * @param	Seq				The UAnimSequence container.
+	 * @param	CompressedData	The compressed animation data being operated on.
 	 * @param	MemoryReader	The FMemoryReader to read from.
 	 * @param	ScaleTrackData	The compressed Scale data stream.
 	 * @param	NumKeysScale	The number of keys present in the stream.
 	 */
 	virtual void ByteSwapScaleIn(
-		UAnimSequence& Seq, 
+		FUECompressedAnimData& CompressedData,
 		FMemoryReader& MemoryReader,
 		uint8*& ScaleTrackData,
 		int32 NumKeysScale) override;
@@ -83,13 +83,13 @@ public:
 	/**
 	 * Handles the ByteSwap of compressed rotation data on export
 	 *
-	 * @param	Seq				The UAnimSequence container.
+	 * @param	CompressedData	The compressed animation data being operated on.
 	 * @param	MemoryWriter	The FMemoryWriter to write to.
 	 * @param	RotTrackData	The compressed rotation data stream.
 	 * @param	NumKeysRot		The number of keys to write to the stream.
 	 */
 	virtual void ByteSwapRotationOut(
-		UAnimSequence& Seq, 
+		FUECompressedAnimData& CompressedData,
 		FMemoryWriter& MemoryWriter,
 		uint8*& RotTrackData,
 		int32 NumKeysRot) override;
@@ -97,13 +97,13 @@ public:
 	/**
 	 * Handles the ByteSwap of compressed translation data on export
 	 *
-	 * @param	Seq				The UAnimSequence container.
+	 * @param	CompressedData	The compressed animation data being operated on.
 	 * @param	MemoryWriter	The FMemoryWriter to write to.
 	 * @param	TransTrackData	The compressed translation data stream.
 	 * @param	NumKeysTrans	The number of keys to write to the stream.
 	 */
 	virtual void ByteSwapTranslationOut(
-		UAnimSequence& Seq, 
+		FUECompressedAnimData& CompressedData,
 		FMemoryWriter& MemoryWriter,
 		uint8*& TransTrackData,
 		int32 NumKeysTrans) override;
@@ -111,13 +111,13 @@ public:
 	/**
 	 * Handles the ByteSwap of compressed Scale data on export
 	 *
-	 * @param	Seq				The UAnimSequence container.
+	 * @param	CompressedData	The compressed animation data being operated on.
 	 * @param	MemoryWriter	The FMemoryWriter to write to.
 	 * @param	TransTrackData	The compressed Scale data stream.
 	 * @param	NumKeysTrans	The number of keys to write to the stream.
 	 */
 	virtual void ByteSwapScaleOut(
-		UAnimSequence& Seq, 
+		FUECompressedAnimData& CompressedData,
 		FMemoryWriter& MemoryWriter,
 		uint8*& ScaleTrackData,
 		int32 NumKeysScale) override;
@@ -261,7 +261,7 @@ FORCEINLINE void AEFConstantKeyLerp<FORMAT>::GetBoneAtomRotation(FTransform& Out
 	{
 		int32 Index0;
 		int32 Index1;
-		float Alpha = TimeToIndex(*DecompContext.AnimSeq, DecompContext.RelativePos, NumRotKeys, Index0, Index1);
+		float Alpha = TimeToIndex(DecompContext.GetSequenceLength(), DecompContext.RelativePos, NumRotKeys, DecompContext.GetInterpolation(), Index0, Index1);
 
 		const int32 RotationStreamOffset = (FORMAT == ACF_IntervalFixed32NoW) ? (sizeof(float)*6) : 0; // offset past Min and Range data
 
@@ -340,7 +340,7 @@ FORCEINLINE void AEFConstantKeyLerp<FORMAT>::GetBoneAtomTranslation(FTransform&
 
 	int32 Index0;
 	int32 Index1;
-	float Alpha = TimeToIndex(*DecompContext.AnimSeq, DecompContext.RelativePos, NumTransKeys, Index0, Index1);
+	float Alpha = TimeToIndex(DecompContext.GetSequenceLength(), DecompContext.RelativePos, NumTransKeys, DecompContext.GetInterpolation(), Index0, Index1);
 
 	const int32 TransStreamOffset = ((FORMAT == ACF_IntervalFixed32NoW) && NumTransKeys > 1) ? (sizeof(float)*6) : 0; // offset past Min and Range data
 
@@ -410,7 +410,7 @@ FORCEINLINE void AEFConstantKeyLerp<FORMAT>::GetBoneAtomScale(FTransform& OutAto
 
 	int32 Index0;
 	int32 Index1;
-	float Alpha = TimeToIndex(*DecompContext.AnimSeq, DecompContext.RelativePos, NumScaleKeys, Index0, Index1);
+	float Alpha = TimeToIndex(DecompContext.GetSequenceLength(), DecompContext.RelativePos, NumScaleKeys, DecompContext.GetInterpolation(), Index0, Index1);
 
 	const int32 ScaleStreamOffset = ((FORMAT == ACF_IntervalFixed32NoW) && NumScaleKeys > 1) ? (sizeof(float)*6) : 0; // offset past Min and Range data
 
diff --git a/Engine/Source/Runtime/Engine/Public/AnimEncoding_PerTrackCompression.h b/Engine/Source/Runtime/Engine/Public/AnimEncoding_PerTrackCompression.h
--- a/Engine/Source/Runtime/Engine/Public/AnimEncoding_PerTrackCompression.h
+++ b/Engine/Source/Runtime/Engine/Public/AnimEncoding_PerTrackCompression.h
@@ -20,24 +20,22 @@ public:
 	/**
 	 * Handles Byte-swapping incoming animation data from a MemoryReader
 	 *
-	 * @param	Seq					An Animation Sequence to contain the read data.
+	 * @param	CompressedData		The compressed animation data being operated on.
 	 * @param	MemoryReader		The MemoryReader object to read from.
 	 * @param	SourceArVersion		The version of the archive that the data is coming from.
 	 */
-	virtual void ByteSwapIn(UAnimSequence& Seq, FMemoryReader& MemoryReader) override;
+	virtual void ByteSwapIn(FUECompressedAnimData& CompressedData, FMemoryReader& MemoryReader) override;
 
 	/**
 	 * Handles Byte-swapping outgoing animation data to an array of BYTEs
 	 *
-	 * @param	Seq					An Animation Sequence to write.
+	 * @param	CompressedData		The compressed animation data being operated on.
 	 * @param	SerializedData		The output buffer.
 	 * @param	ForceByteSwapping	true is byte swapping is not optional.
 	 */
 	virtual void ByteSwapOut(
-		UAnimSequence& Seq,
-		TArray<uint8>& SerializedData, 
-		bool ForceByteSwapping,
-		bool bMaintainComponentOrder=false) override;
+		FUECompressedAnimData& CompressedData,
+		FMemoryWriter& MemoryWriter) override;
 
 	/**
 	 * Extracts a single BoneAtom from an Animation Sequence.
@@ -99,13 +97,12 @@ protected:
 	/**
 	 * Handles Byte-swapping a single track of animation data from a MemoryReader or to a MemoryWriter
 	 *
-	 * @param	Seq					The Animation Sequence being operated on.
+	 * @param	CompressedData		The compressed animation data being operated on.
 	 * @param	MemoryStream		The MemoryReader or MemoryWriter object to read from/write to.
 	 * @param	Offset				The starting offset into the compressed byte stream for this track (can be INDEX_NONE to indicate an identity track)
-	 * @param	bMaintainComponentOrder Should we maintain the order of track components (trans/rot/scale) in the byte stream (if false the output will always be T/R/S regardless of input order)
 	 */
 	template<class TArchive>
-	static void ByteSwapOneTrack(UAnimSequence& Seq, TArchive& MemoryStream, int32 Offset, bool bMaintainComponentOrder = false);
+	static void ByteSwapOneTrack(FUECompressedAnimData& CompressedData, TArchive& MemoryStream, int32 BufferStart, int32 Offset);
 
 	/**
 	 * Preserves 4 byte alignment within a stream
diff --git a/Engine/Source/Runtime/Engine/Public/AnimEncoding_VariableKeyLerp.h b/Engine/Source/Runtime/Engine/Public/AnimEncoding_VariableKeyLerp.h
--- a/Engine/Source/Runtime/Engine/Public/AnimEncoding_VariableKeyLerp.h
+++ b/Engine/Source/Runtime/Engine/Public/AnimEncoding_VariableKeyLerp.h
@@ -82,13 +82,13 @@ public:
 	/**
 	 * Handles the ByteSwap of compressed rotation data on import
 	 *
-	 * @param	Seq				The UAnimSequence container.
+	 * @param	CompressedData	The compressed animation data being operated on.
 	 * @param	MemoryReader	The FMemoryReader to read from.
 	 * @param	RotTrackData	The compressed rotation data stream.
 	 * @param	NumKeysRot		The number of keys present in the stream.
 	 */
 	virtual void ByteSwapRotationIn(
-		UAnimSequence& Seq, 
+		FUECompressedAnimData& CompressedData,
 		FMemoryReader& MemoryReader,
 		uint8*& RotTrackData,
 		int32 NumKeysRot) override;
@@ -96,13 +96,13 @@ public:
 	/**
 	 * Handles the ByteSwap of compressed translation data on import
 	 *
-	 * @param	Seq				The UAnimSequence container.
+	 * @param	CompressedData	The compressed animation data being operated on.
 	 * @param	MemoryReader	The FMemoryReader to read from.
 	 * @param	TransTrackData	The compressed translation data stream.
 	 * @param	NumKeysTrans	The number of keys present in the stream.
 	 */
 	virtual void ByteSwapTranslationIn(
-		UAnimSequence& Seq, 
+		FUECompressedAnimData& CompressedData,
 		FMemoryReader& MemoryReader,
 		uint8*& TransTrackData,
 		int32 NumKeysTransn) override;
@@ -110,13 +110,13 @@ public:
 	/**
 	 * Handles the ByteSwap of compressed Scale data on import
 	 *
-	 * @param	Seq				The UAnimSequence container.
+	 * @param	CompressedData	The compressed animation data being operated on.
 	 * @param	MemoryReader	The FMemoryReader to read from.
 	 * @param	ScaleTrackData	The compressed Scale data stream.
 	 * @param	NumKeysScale	The number of keys present in the stream.
 	 */
 	virtual void ByteSwapScaleIn(
-		UAnimSequence& Seq, 
+		FUECompressedAnimData& CompressedData,
 		FMemoryReader& MemoryReader,
 		uint8*& ScaleTrackData,
 		int32 NumKeysScale) override;
@@ -124,13 +124,13 @@ public:
 	/**
 	 * Handles the ByteSwap of compressed rotation data on export
 	 *
-	 * @param	Seq				The UAnimSequence container.
+	 * @param	CompressedData	The compressed animation data being operated on.
 	 * @param	MemoryWriter	The FMemoryWriter to write to.
 	 * @param	RotTrackData	The compressed rotation data stream.
 	 * @param	NumKeysRot		The number of keys to write to the stream.
 	 */
 	virtual void ByteSwapRotationOut(
-		UAnimSequence& Seq, 
+		FUECompressedAnimData& CompressedData,
 		FMemoryWriter& MemoryWriter,
 		uint8*& RotTrackData,
 		int32 NumKeysRot) override;
@@ -138,13 +138,13 @@ public:
 	/**
 	 * Handles the ByteSwap of compressed translation data on export
 	 *
-	 * @param	Seq				The UAnimSequence container.
+	 * @param	CompressedData	The compressed animation data being operated on.
 	 * @param	MemoryWriter	The FMemoryWriter to write to.
 	 * @param	TransTrackData	The compressed translation data stream.
 	 * @param	NumKeysTrans	The number of keys to write to the stream.
 	 */
 	virtual void ByteSwapTranslationOut(
-		UAnimSequence& Seq, 
+		FUECompressedAnimData& CompressedData,
 		FMemoryWriter& MemoryWriter,
 		uint8*& TransTrackData,
 		int32 NumKeysTrans) override;
@@ -152,13 +152,13 @@ public:
 	/**
 	 * Handles the ByteSwap of compressed Scale data on export
 	 *
-	 * @param	Seq				The UAnimSequence container.
+	 * @param	CompressedData	The compressed animation data being operated on.
 	 * @param	MemoryWriter	The FMemoryWriter to write to.
 	 * @param	ScaleTrackData	The compressed Scale data stream.
 	 * @param	NumKeysScale	The number of keys to write to the stream.
 	 */
 	virtual void ByteSwapScaleOut(
-		UAnimSequence& Seq, 
+		FUECompressedAnimData& CompressedData,
 		FMemoryWriter& MemoryWriter,
 		uint8*& ScaleTrackData,
 		int32 NumKeysScale) override;
@@ -302,7 +302,7 @@ FORCEINLINE_DEBUGGABLE void AEFVariableKeyLerp<FORMAT>::GetBoneAtomRotation(FTra
 
 		int32 Index0;
 		int32 Index1;
-		float Alpha = TimeToIndex(*DecompContext.AnimSeq, FrameTable, DecompContext.RelativePos, NumRotKeys, Index0, Index1);
+		float Alpha = TimeToIndex(DecompContext.GetInterpolation(), DecompContext.GetCompressedNumberOfFrames(), FrameTable, DecompContext.RelativePos, NumRotKeys, Index0, Index1);
 
 
 		if (Index0 != Index1)
@@ -384,7 +384,7 @@ FORCEINLINE_DEBUGGABLE void AEFVariableKeyLerp<FORMAT>::GetBoneAtomTranslation(F
 
 	int32 Index0;
 	int32 Index1;
-	float Alpha = TimeToIndex(*DecompContext.AnimSeq, FrameTable, DecompContext.RelativePos, NumTransKeys, Index0, Index1);
+	float Alpha = TimeToIndex(DecompContext.GetInterpolation(), DecompContext.GetCompressedNumberOfFrames(), FrameTable, DecompContext.RelativePos, NumTransKeys, Index0, Index1);
 	const int32 TransStreamOffset = ((FORMAT == ACF_IntervalFixed32NoW) && NumTransKeys > 1) ? (sizeof(float)*6) : 0; // offset past Min and Range data
 
 	if (Index0 != Index1)
@@ -455,7 +455,7 @@ FORCEINLINE_DEBUGGABLE void AEFVariableKeyLerp<FORMAT>::GetBoneAtomScale(FTransf
 
 	int32 Index0;
 	int32 Index1;
-	float Alpha = TimeToIndex(*DecompContext.AnimSeq, FrameTable, DecompContext.RelativePos, NumScaleKeys, Index0, Index1);
+	float Alpha = TimeToIndex(DecompContext.GetInterpolation(), DecompContext.GetCompressedNumberOfFrames(), FrameTable, DecompContext.RelativePos, NumScaleKeys, Index0, Index1);
 	const int32 ScaleStreamOffset = ((FORMAT == ACF_IntervalFixed32NoW) && NumScaleKeys > 1) ? (sizeof(float)*6) : 0; // offset past Min and Range data
 
 	if (Index0 != Index1)
