Refactor UAnimSequence out of compression code

diff --git a/Engine/Source/Runtime/Engine/Public/Animation/AnimCompressionTypes.h b/Engine/Source/Runtime/Engine/Public/Animation/AnimCompressionTypes.h
new file mode 100644
--- /dev/null
+++ b/Engine/Source/Runtime/Engine/Public/Animation/AnimCompressionTypes.h
@@ -0,0 +1,671 @@
+// Copyright 1998-2019 Epic Games, Inc. All Rights Reserved.
+
+#pragma once
+
+#include "CoreMinimal.h"
+#include "UObject/ObjectMacros.h"
+#include "Misc/MemStack.h"
+#include "Animation/AnimTypes.h"
+#include "Animation/AnimCurveTypes.h"
+
+#include "Async/MappedFileHandle.h"
+#include "HAL/PlatformFilemanager.h"
+#include "GenericPlatform/GenericPlatformFile.h"
+#include "Misc/Paths.h"
+#include "Serialization/BulkData.h"
+
+#include "Serialization/MemoryReader.h"
+#include "Serialization/MemoryWriter.h"
+
+#include "AnimCompressionTypes.generated.h"
+
+/**
+ * Indicates animation data key format.
+ */
+UENUM()
+enum AnimationKeyFormat
+{
+	AKF_ConstantKeyLerp,
+	AKF_VariableKeyLerp,
+	AKF_PerTrackCompression,
+	AKF_MAX,
+};
+
+class FMemoryReader;
+class FMemoryWriter;
+
+template<typename ArrayClass>
+struct ENGINE_API FCompressedOffsetDataBase
+{
+	ArrayClass OffsetData;
+
+	int32 StripSize;
+
+	FCompressedOffsetDataBase(int32 InStripSize = 2)
+		: StripSize(InStripSize)
+	{}
+
+	void SetStripSize(int32 InStripSize)
+	{
+		ensure(InStripSize > 0);
+		StripSize = InStripSize;
+	}
+
+	const int32 GetOffsetData(int32 StripIndex, int32 Offset) const
+	{
+		checkSlow(OffsetData.IsValidIndex(StripIndex * StripSize + Offset));
+
+		return OffsetData[StripIndex * StripSize + Offset];
+	}
+
+	void SetOffsetData(int32 StripIndex, int32 Offset, int32 Value)
+	{
+		checkSlow(OffsetData.IsValidIndex(StripIndex * StripSize + Offset));
+		OffsetData[StripIndex * StripSize + Offset] = Value;
+	}
+
+	void AddUninitialized(int32 NumOfTracks)
+	{
+		OffsetData.AddUninitialized(NumOfTracks*StripSize);
+	}
+
+	void Empty(int32 NumOfTracks = 0)
+	{
+		OffsetData.Empty(NumOfTracks*StripSize);
+	}
+
+	int32 GetMemorySize() const
+	{
+		return sizeof(int32)*OffsetData.Num() + sizeof(int32);
+	}
+
+	int32 GetNumTracks() const
+	{
+		return OffsetData.Num() / StripSize;
+	}
+
+	bool IsValid() const
+	{
+		return (OffsetData.Num() > 0);
+	}
+};
+
+struct FCompressedOffsetData : public FCompressedOffsetDataBase<TArray<int32>>
+{
+
+};
+
+
+/**
+ * Represents a segment of the anim sequence that is compressed.
+ */
+USTRUCT()
+struct ENGINE_API FCompressedSegment
+{
+	GENERATED_USTRUCT_BODY()
+
+	// Frame where the segment begins in the anim sequence
+	int32 StartFrame;
+
+	// Num of frames contained in the segment
+	int32 NumFrames;
+
+	// Segment data offset in CompressedByteStream
+	int32 ByteStreamOffset;
+
+	/** The compression format that was used to compress translation tracks. */
+	TEnumAsByte<enum AnimationCompressionFormat> TranslationCompressionFormat;
+
+	/** The compression format that was used to compress rotation tracks. */
+	TEnumAsByte<enum AnimationCompressionFormat> RotationCompressionFormat;
+
+	/** The compression format that was used to compress rotation tracks. */
+	TEnumAsByte<enum AnimationCompressionFormat> ScaleCompressionFormat;
+
+	FCompressedSegment()
+		: StartFrame(0)
+		, NumFrames(0)
+		, ByteStreamOffset(0)
+		, TranslationCompressionFormat(ACF_None)
+		, RotationCompressionFormat(ACF_None)
+		, ScaleCompressionFormat(ACF_None)
+	{
+	}
+
+	friend FArchive& operator<<(FArchive& Ar, FCompressedSegment &Segment)
+	{
+		return Ar << Segment.StartFrame << Segment.NumFrames << Segment.ByteStreamOffset
+			<< Segment.TranslationCompressionFormat << Segment.RotationCompressionFormat << Segment.ScaleCompressionFormat;
+	}
+};
+
+struct ENGINE_API FCompressibleAnimData
+{
+public:
+	FCompressibleAnimData()
+		: RequestedCompressionScheme(nullptr)
+		, CurveCompressionSettings(nullptr)
+		, Skeleton(nullptr)
+		, Interpolation((EAnimInterpolationType)0)
+		, SequenceLength(0.f)
+		, NumFrames(0)
+		, bIsValidAdditive(false)
+		, CompressCommandletVersion(0)
+		, RefFrameIndex(0)
+		, RefPoseType((EAdditiveBasePoseType)0)
+		, AdditiveAnimType((EAdditiveAnimationType)0)
+	{
+	}
+
+	FCompressibleAnimData(class UAnimSequence* InSeq);
+
+	class UAnimCompress* RequestedCompressionScheme;
+
+	class UAnimCurveCompressionSettings* CurveCompressionSettings;
+
+	class USkeleton* Skeleton;
+
+	TArray<FTrackToSkeletonMap> TrackToSkeletonMapTable;
+
+	TArray<FRawAnimSequenceTrack> RawAnimationData;
+
+	TArray<FRawAnimSequenceTrack> AdditiveBaseAnimationData;
+
+	EAnimInterpolationType Interpolation;
+
+	TArray<FBoneData> BoneData;
+
+	FRawCurveTracks RawCurveData;
+
+	float SequenceLength;
+
+	int32 NumFrames;
+
+	bool bIsValidAdditive;
+
+	//For DDC
+	FString TypeName;
+	FString DDCKey;
+
+	int32 CompressCommandletVersion;
+	FGuid RawDataGuid;
+
+	//Additive
+	FGuid AdditiveDataGuid;
+	int32 RefFrameIndex;
+	EAdditiveBasePoseType RefPoseType;
+	EAdditiveAnimationType AdditiveAnimType;
+
+	//For Logging
+	FString Name;
+	FString FullName;
+	FName   AnimFName;
+
+	int32 GetApproxRawBoneSize() const
+	{
+		int32 Total = sizeof(FRawAnimSequenceTrack) * RawAnimationData.Num();
+		for (int32 i = 0; i < RawAnimationData.Num(); ++i)
+		{
+			const FRawAnimSequenceTrack& RawTrack = RawAnimationData[i];
+			Total +=
+				sizeof(FVector) * RawTrack.PosKeys.Num() +
+				sizeof(FQuat) * RawTrack.RotKeys.Num() +
+				sizeof(FVector) * RawTrack.ScaleKeys.Num();
+		}
+		return Total;
+	}
+
+	int32 GetApproxRawCurveSize() const
+	{
+		int32 Total = 0;
+		for (const FFloatCurve& Curve : RawCurveData.FloatCurves)
+		{
+			Total += sizeof(FFloatCurve);
+			Total += sizeof(FRichCurveKey) * Curve.FloatCurve.Keys.Num();
+		}
+		return Total;
+	}
+
+	int32 GetApproxRawSize() const
+	{
+		return GetApproxRawBoneSize() + GetApproxRawCurveSize();
+	}
+
+	void Update(struct FCompressedAnimSequence& CompressedData) const;
+
+private:
+
+};
+
+// Wrapper Code
+template <typename T>
+struct TArrayMaker
+{
+	using Type = TArray<T>;
+};
+
+
+template <typename T>
+struct TNonConstArrayViewMaker
+{
+	using Type = TArrayView<T>;
+};
+
+template <typename T>
+struct TArrayViewMaker
+{
+	using Type = TArrayView<const T>;
+};
+
+template <template <typename> class ContainerTypeMakerTemplate>
+struct FCompressedAnimDataBase
+{
+	/**
+	 * An array of 4*NumTrack ints, arranged as follows: - PerTrack is 2*NumTrack, so this isn't true any more
+	 *   [0] Trans0.Offset
+	 *   [1] Trans0.NumKeys
+	 *   [2] Rot0.Offset
+	 *   [3] Rot0.NumKeys
+	 *   [4] Trans1.Offset
+	 *   . . .
+	 */
+	typename ContainerTypeMakerTemplate<int32>::Type CompressedTrackOffsets;
+
+	/**
+	 * An array of 2*NumTrack ints, arranged as follows:
+		if identity, it is offset
+		if not, it is num of keys
+	 *   [0] Scale0.Offset or NumKeys
+	 *   [1] Scale1.Offset or NumKeys
+
+	 * @TODO NOTE: first implementation is offset is [0], numkeys [1]
+	 *   . . .
+	 */
+	FCompressedOffsetDataBase<typename ContainerTypeMakerTemplate<int32>::Type>  CompressedScaleOffsets;
+
+	/**
+	 * ByteStream for compressed animation data.
+	 * The memory layout is dependent on the algorithm used to compress the anim sequence.
+	 */
+	typename ContainerTypeMakerTemplate<uint8>::Type CompressedByteStream;
+
+	/**
+	 * The runtime interface to decode and byte swap the compressed animation
+	 * May be NULL. Set at runtime - does not exist in editor
+	 */
+	class AnimEncoding* TranslationCodec;
+	class AnimEncoding* RotationCodec;
+	class AnimEncoding* ScaleCodec;
+
+	enum AnimationKeyFormat KeyEncodingFormat;
+
+	/** The compression format that was used to compress tracks parts. */
+	AnimationCompressionFormat TranslationCompressionFormat;
+	AnimationCompressionFormat RotationCompressionFormat;
+	AnimationCompressionFormat ScaleCompressionFormat;
+
+	int32 CompressedNumberOfFrames;
+	
+	template <template <typename> class OtherContainerMaker>
+	explicit FCompressedAnimDataBase(FCompressedAnimDataBase<OtherContainerMaker>& InCompressedData)
+		: CompressedTrackOffsets(InCompressedData.CompressedTrackOffsets)
+		, CompressedByteStream(InCompressedData.CompressedByteStream)
+
+		, TranslationCodec(InCompressedData.TranslationCodec)
+		, RotationCodec(InCompressedData.RotationCodec)
+		, ScaleCodec(InCompressedData.ScaleCodec)
+
+		, KeyEncodingFormat(InCompressedData.KeyEncodingFormat)
+		, TranslationCompressionFormat(InCompressedData.TranslationCompressionFormat)
+		, RotationCompressionFormat(InCompressedData.RotationCompressionFormat)
+		, ScaleCompressionFormat(InCompressedData.ScaleCompressionFormat)
+
+		, CompressedNumberOfFrames(InCompressedData.CompressedNumberOfFrames)
+	{
+		CompressedScaleOffsets.OffsetData = InCompressedData.CompressedScaleOffsets.OffsetData;
+		CompressedScaleOffsets.StripSize = InCompressedData.CompressedScaleOffsets.StripSize;
+	}
+
+	FCompressedAnimDataBase()
+		: TranslationCodec(nullptr)
+		, RotationCodec(nullptr)
+		, ScaleCodec(nullptr)
+
+		, KeyEncodingFormat((AnimationKeyFormat)0)
+		, TranslationCompressionFormat((AnimationCompressionFormat)0)
+		, RotationCompressionFormat((AnimationCompressionFormat)0)
+		, ScaleCompressionFormat((AnimationCompressionFormat)0)
+
+		, CompressedNumberOfFrames(0)
+	{
+
+	}
+
+	template <template <typename> class OtherContainerMaker>
+	void CopyFromSettings(const FCompressedAnimDataBase<OtherContainerMaker>& Other)
+	{
+		TranslationCodec = Other.TranslationCodec;
+		RotationCodec = Other.RotationCodec;
+		ScaleCodec = Other.ScaleCodec;
+
+		KeyEncodingFormat = Other.KeyEncodingFormat;
+		TranslationCompressionFormat = Other.TranslationCompressionFormat;
+		RotationCompressionFormat = Other.RotationCompressionFormat;
+		ScaleCompressionFormat = Other.ScaleCompressionFormat;
+
+		CompressedNumberOfFrames = Other.CompressedNumberOfFrames;
+	}
+
+	int64 GetApproxBoneCompressedSize() const
+	{
+		return (int64)CompressedTrackOffsets.GetTypeSize()*(int64)CompressedTrackOffsets.Num() + (int64)CompressedByteStream.Num() + (int64)CompressedScaleOffsets.GetMemorySize();
+	}
+
+	bool IsCompressedDataValid() const
+	{
+		return CompressedByteStream.Num() > 0 || (TranslationCompressionFormat == ACF_Identity && RotationCompressionFormat == ACF_Identity && ScaleCompressionFormat == ACF_Identity);
+	}
+};
+
+struct ENGINE_API FCompressibleAnimDataResult : public FCompressedAnimDataBase<TArrayMaker>
+{
+	class UAnimCompress* CompressionScheme;
+
+	FCompressibleAnimDataResult()
+		: CompressionScheme(nullptr)
+	{
+
+	}
+
+	template <template <typename> class OtherContainerMaker>
+	explicit FCompressibleAnimDataResult(FCompressedAnimDataBase<OtherContainerMaker>& InCompressedData)
+		: FCompressedAnimDataBase(InCompressedData)
+	{}
+
+	void CopyFrom(const FCompressibleAnimDataResult& Other)
+	{
+		CompressedTrackOffsets = Other.CompressedTrackOffsets;
+		CompressedByteStream = Other.CompressedByteStream;
+		CompressedScaleOffsets.OffsetData = Other.CompressedScaleOffsets.OffsetData;
+		CompressedScaleOffsets.StripSize = Other.CompressedScaleOffsets.StripSize;
+
+		CopyFromSettings(Other);
+	}
+
+	FCompressibleAnimDataResult& operator=(const FCompressibleAnimDataResult& Other)
+	{
+		CopyFrom(Other);
+		CompressionScheme = Other.CompressionScheme;
+		return *this;
+	}
+
+	void BuildFinalBuffer(TArray<uint8>& OutBuffer);
+};
+
+struct ICompressedAnimData
+{
+public:
+	virtual ~ICompressedAnimData() {}
+
+	virtual void SerializeCompressedData(class FArchive& Ar) = 0;
+
+	virtual void ByteSwapIn(TArrayView<uint8> CompressedData, FMemoryReader& MemoryStream) = 0;
+	virtual void ByteSwapOut(TArrayView<uint8> CompressedData, FMemoryWriter& MemoryStream) = 0;
+};
+
+template<typename T>
+TArrayView<T> RebaseTArrayView(const TArrayView<T>& ArrayView, const uint8* OriginalBase, const uint8* NewBase)
+{
+	if (ArrayView.GetData() != nullptr)
+	{
+		uint32 Offset = ArrayView.GetData() - OriginalBase;
+		T* NewData = (T*)(NewBase + Offset);
+		return TArrayView<T>(NewData, ArrayView.Num());
+	}
+	return ArrayView;
+}
+
+struct FUECompressedAnimData : public ICompressedAnimData, public FCompressedAnimDataBase<TNonConstArrayViewMaker>
+{
+	FUECompressedAnimData() = default;
+
+	template <template <typename> class OtherContainerMaker>
+	explicit FUECompressedAnimData(FCompressedAnimDataBase<OtherContainerMaker>& InCompressedData)
+		: FCompressedAnimDataBase(InCompressedData)
+	{}
+
+	void Reset();
+
+	void InitViewsFromBuffer(const TArrayView<uint8> BulkData);
+
+#if WITH_EDITOR
+	void CopyFrom(const FCompressibleAnimDataResult& Other);
+#endif
+
+	virtual void SerializeCompressedData(class FArchive& Ar);
+
+	template<typename TArchive>
+	void ByteSwapData(TArrayView<uint8> CompresedData, TArchive& MemoryStream);
+
+	virtual void ByteSwapIn(TArrayView<uint8> CompressedData, FMemoryReader& MemoryStream) override { ByteSwapData(CompressedData, MemoryStream); }
+	virtual void ByteSwapOut(TArrayView<uint8> CompressedData, FMemoryWriter& MemoryStream) override { ByteSwapData(CompressedData, MemoryStream); }
+};
+
+template<uint32 Alignment = DEFAULT_ALIGNMENT>
+class TMaybeMappedAllocator
+{
+public:
+
+	enum { NeedsElementType = false };
+	enum { RequireRangeCheck = true };
+
+	class ForAnyElementType
+	{
+	public:
+
+		/** Default constructor. */
+		ForAnyElementType()
+			: Data(nullptr)
+			, MappedHandle(nullptr)
+			, MappedRegion(nullptr)
+		{}
+
+		/**
+		 * Moves the state of another allocator into this one.
+		 * Assumes that the allocator is currently empty, i.e. memory may be allocated but any existing elements have already been destructed (if necessary).
+		 * @param Other - The allocator to move the state from.  This allocator should be left in a valid empty state.
+		 */
+		void MoveToEmpty(ForAnyElementType& Other)
+		{
+			checkSlow(this != &Other);
+
+			Reset();
+
+			Data = Other.Data;
+			Other.Data = nullptr;
+
+			MappedRegion = Other.MappedRegion;
+			Other->MappedRegion = nullptr;
+
+			MappedHandle = Other.MappedHandle;
+			Other->MappedHandle = nullptr;
+		}
+
+		/** Destructor. */
+		~ForAnyElementType()
+		{
+			Reset();
+		}
+
+		// FContainerAllocatorInterface
+		FScriptContainerElement* GetAllocation() const
+		{
+			return Data;
+		}
+		void ResizeAllocation(
+			int32 PreviousNumElements,
+			int32 NumElements,
+			SIZE_T NumBytesPerElement
+		)
+		{
+			// Avoid calling FMemory::Realloc( nullptr, 0 ) as ANSI C mandates returning a valid pointer which is not what we want.
+			if (Data || NumElements)
+			{
+				check(!MappedHandle && !MappedRegion); // this could be supported, but it probably is never what you want, so we will just assert.
+					//checkSlow(((uint64)NumElements*(uint64)ElementTypeInfo.GetSize() < (uint64)INT_MAX));
+				Data = (FScriptContainerElement*)FMemory::Realloc(Data, NumElements*NumBytesPerElement, Alignment);
+			}
+		}
+		int32 CalculateSlackReserve(int32 NumElements, int32 NumBytesPerElement) const
+		{
+			check(!MappedHandle && !MappedRegion); // this could be supported, but it probably is never what you want, so we will just assert.
+			return DefaultCalculateSlackReserve(NumElements, NumBytesPerElement, true, Alignment);
+		}
+		int32 CalculateSlackShrink(int32 NumElements, int32 NumAllocatedElements, int32 NumBytesPerElement) const
+		{
+			check(!MappedHandle && !MappedRegion); // this could be supported, but it probably is never what you want, so we will just assert.
+			return DefaultCalculateSlackShrink(NumElements, NumAllocatedElements, NumBytesPerElement, true, Alignment);
+		}
+		int32 CalculateSlackGrow(int32 NumElements, int32 NumAllocatedElements, int32 NumBytesPerElement) const
+		{
+			check(!MappedHandle && !MappedRegion); // this could be supported, but it probably is never what you want, so we will just assert.
+			return DefaultCalculateSlackGrow(NumElements, NumAllocatedElements, NumBytesPerElement, true, Alignment);
+		}
+
+		SIZE_T GetAllocatedSize(int32 NumAllocatedElements, SIZE_T NumBytesPerElement) const
+		{
+			return NumAllocatedElements * NumBytesPerElement;
+		}
+
+		bool HasAllocation()
+		{
+			return !!Data;
+		}
+
+		void AcceptFileMapping(IMappedFileHandle* InMappedHandle, IMappedFileRegion* InMappedRegion, void *MallocPtr)
+		{
+			check(!MappedHandle && !Data); // we could support stuff like this, but that usually isn't what we want for streamlined loading
+			Reset(); // just in case
+			if (InMappedHandle || InMappedRegion)
+			{
+				MappedHandle = InMappedHandle;
+				MappedRegion = InMappedRegion;
+				Data = (FScriptContainerElement*)MappedRegion->GetMappedPtr(); //@todo mapped files should probably be const-correct
+				check(IsAligned(Data, FPlatformProperties::GetMemoryMappingAlignment()));
+			}
+			else
+			{
+				Data = (FScriptContainerElement*)MallocPtr;
+			}
+		}
+
+		bool IsMapped() const
+		{
+			return MappedRegion || MappedHandle;
+		}
+	private:
+
+		FScriptContainerElement* Data;
+		IMappedFileHandle* MappedHandle;
+		IMappedFileRegion* MappedRegion;
+
+		void Reset()
+		{
+			if (MappedRegion || MappedHandle)
+			{
+				delete MappedRegion;
+				delete MappedHandle;
+				MappedRegion = nullptr;
+				MappedHandle = nullptr;
+				Data = nullptr; // make sure we don't try to free this pointer
+			}
+			if (Data)
+			{
+				FMemory::Free(Data);
+				Data = nullptr;
+			}
+		}
+
+
+		ForAnyElementType(const ForAnyElementType&);
+		ForAnyElementType& operator=(const ForAnyElementType&);
+	};
+
+	template<typename ElementType>
+	class ForElementType : public ForAnyElementType
+	{
+	public:
+
+		/** Default constructor. */
+		ForElementType()
+		{}
+
+		ElementType* GetAllocation() const
+		{
+			return (ElementType*)ForAnyElementType::GetAllocation();
+		}
+	};
+};
+
+template<typename T, uint32 Alignment = DEFAULT_ALIGNMENT>
+class TMaybeMappedArray : public TArray<T, TMaybeMappedAllocator<Alignment>>
+{
+public:
+	TMaybeMappedArray()
+	{
+	}
+	TMaybeMappedArray(TMaybeMappedArray&&) = default;
+	TMaybeMappedArray(const TMaybeMappedArray&) = default;
+	TMaybeMappedArray& operator=(TMaybeMappedArray&&) = default;
+	TMaybeMappedArray& operator=(const TMaybeMappedArray&) = default;
+
+	void AcceptOwnedBulkDataPtr(FOwnedBulkDataPtr* OwnedPtr, int32 Num)
+	{
+		this->ArrayNum = Num;
+		this->ArrayMax = Num;
+		this->AllocatorInstance.AcceptFileMapping(OwnedPtr->GetMappedHandle(), OwnedPtr->GetMappedRegion(), (void*)OwnedPtr->GetPointer());
+		OwnedPtr->RelinquishOwnership();
+	}
+};
+
+struct ENGINE_API FCompressedAnimSequence
+{
+public:
+
+	/**
+	 * Version of TrackToSkeletonMapTable for the compressed tracks. Due to baking additive data
+	 * we can end up with a different amount of tracks to the original raw animation and so we must index into the
+	 * compressed data using this
+	 */
+	TArray<struct FTrackToSkeletonMap> CompressedTrackToSkeletonMapTable;
+
+	/**
+	 * Much like track indices above, we need to be able to remap curve names. The FName inside FSmartName
+	 * is serialized and the UID is generated at runtime. Stored in the DDC.
+	 */
+	TArray<struct FSmartName> CompressedCurveNames;
+
+	/**
+	 * ByteStream for compressed animation data.
+	 * The memory layout is dependent on the algorithm used to compress the anim sequence.
+	 */
+#if WITH_EDITOR
+	TArray<uint8> CompressedByteStream;
+	FByteBulkData OptionalBulk;
+#else
+	TMaybeMappedArray<uint8> CompressedByteStream;
+#endif
+
+	/* Compressed curve data stream used by AnimCurveCompressionCodec */
+	TArray<uint8> CompressedCurveByteStream;
+
+	FUECompressedAnimData CompressedDataStructure;
+
+	/** The codec used by the compressed data as determined by the compression settings. */
+	class UAnimCurveCompressionCodec* CurveCompressionCodec;
+
+	// The size of the raw data used to create the compressed data
+	int32 CompressedRawDataSize;
+
+	void SerializeCompressedData(FArchive& Ar, bool bDDCData, UObject* DataOwner, class UAnimCurveCompressionSettings* CurveCompressionSettings);
+};
\ No newline at end of file
diff --git a/Engine/Source/Runtime/Engine/Public/Animation/AnimSequenceDecompressionContext.h b/Engine/Source/Runtime/Engine/Public/Animation/AnimSequenceDecompressionContext.h
--- a/Engine/Source/Runtime/Engine/Public/Animation/AnimSequenceDecompressionContext.h
+++ b/Engine/Source/Runtime/Engine/Public/Animation/AnimSequenceDecompressionContext.h
@@ -13,8 +13,6 @@ struct FCompressedSegment;
 class FAnimEncodingDecompressionContext;
 class AnimEncoding;
 
-#define USE_SEGMENTING_CONTEXT 0 // Uses segmenting in anim compression + context in decompression
-
 /**
  * Structure to wrap a trivial key handle. Used by the decompression context and encoders.
  */
@@ -33,17 +31,22 @@ struct FTrivialAnimKeyHandle
 struct ENGINE_API FAnimSequenceDecompressionContext
 {
 public:
-	FAnimSequenceDecompressionContext() : AnimSeq(nullptr), Time(0.f) {}
-	FAnimSequenceDecompressionContext(const UAnimSequence* AnimSeq_) :AnimSeq(AnimSeq_), Time(0.f) {}
-
+#if WITH_EDITOR
+	FAnimSequenceDecompressionContext(const FCompressibleAnimData& CompressibleAnimData_, const FUECompressedAnimData& CompressedAnimData_) : SequenceLength(CompressibleAnimData_.SequenceLength), Interpolation(CompressibleAnimData_.Interpolation), AnimName(CompressibleAnimData_.AnimFName), CompressedAnimData(CompressedAnimData_), Time(0.f) {}
+#endif
+	FAnimSequenceDecompressionContext(float SequenceLength_, EAnimInterpolationType Interpolation_, const FName& AnimName_, const FUECompressedAnimData& CompressedAnimData_) : SequenceLength(SequenceLength_), Interpolation(Interpolation_), AnimName(AnimName_), CompressedAnimData(CompressedAnimData_), Time(0.f) {}
 	/**
 	* We don't disable copying but we also don't retain any state. Copying is required by the UE serialization
 	* but retaining state is not necessary and can cause issues.
 	*/
-	FAnimSequenceDecompressionContext(const FAnimSequenceDecompressionContext& Other) : AnimSeq(Other.AnimSeq), Time(Other.Time) {}
-	FAnimSequenceDecompressionContext& operator=(const FAnimSequenceDecompressionContext& Other) { AnimSeq = Other.AnimSeq; Time = Other.Time; return *this; }
+	FAnimSequenceDecompressionContext(const FAnimSequenceDecompressionContext& Other) : SequenceLength(Other.SequenceLength), Interpolation(Other.Interpolation), AnimName(Other.AnimName), CompressedAnimData(Other.CompressedAnimData), Time(Other.Time) {}
 
-	const UAnimSequence* AnimSeq;
+	// Anim info
+	float SequenceLength;
+	EAnimInterpolationType Interpolation;
+	FName AnimName;
+
+	const FUECompressedAnimData& CompressedAnimData;
 	float Time;
 	float RelativePos;
 	bool bHasScale;
@@ -51,18 +54,26 @@ public:
 	void Seek(float SampleAtTime)
 	{
 		Time = SampleAtTime;
-		RelativePos = SampleAtTime / AnimSeq->SequenceLength;
-		bHasScale = AnimSeq->CompressedScaleOffsets.IsValid();
+		RelativePos = SampleAtTime / SequenceLength;
+		bHasScale = CompressedAnimData.CompressedScaleOffsets.IsValid();
 	}
 
+	float GetSequenceLength() const { return SequenceLength; }
+
+	int32 GetCompressedNumberOfFrames() const { return CompressedAnimData.CompressedNumberOfFrames; }
+
 	// Following just expose the properties of the anim sequence, 
-	AnimEncoding* GetRotationCodec() const { return AnimSeq->RotationCodec; }
-	AnimEncoding* GetTranslationCodec() const { return AnimSeq->TranslationCodec; }
-	AnimEncoding* GetScaleCodec() const { return AnimSeq->ScaleCodec; }
+	AnimEncoding* GetRotationCodec() const { return CompressedAnimData.RotationCodec; }
+	AnimEncoding* GetTranslationCodec() const { return CompressedAnimData.TranslationCodec; }
+	AnimEncoding* GetScaleCodec() const { return CompressedAnimData.ScaleCodec; }
+
+	const int32* GetCompressedTrackOffsets() const { return CompressedAnimData.CompressedTrackOffsets.GetData(); }
+	const uint8* GetCompressedByteStream() const { return CompressedAnimData.CompressedByteStream.GetData(); }
+	const FCompressedOffsetDataBase<TArrayView<int32>>* GetCompressedScaleOffsets() const { return &CompressedAnimData.CompressedScaleOffsets; }
+
+	EAnimInterpolationType GetInterpolation() const { return Interpolation; }
 
-	const int32* GetCompressedTrackOffsets() const { return AnimSeq->CompressedTrackOffsets.GetData(); }
-	const uint8* GetCompressedByteStream() const { return AnimSeq->CompressedByteStream.GetData(); }
-	const FCompressedOffsetData* GetCompressedScaleOffsets() const { return &AnimSeq->CompressedScaleOffsets; }
+	FName GetAnimFName() const { return AnimName; }
 };
 
 #else
diff --git a/Engine/Source/Runtime/Engine/Public/Animation/AnimTypes.h b/Engine/Source/Runtime/Engine/Public/Animation/AnimTypes.h
--- a/Engine/Source/Runtime/Engine/Public/Animation/AnimTypes.h
+++ b/Engine/Source/Runtime/Engine/Public/Animation/AnimTypes.h
@@ -7,12 +7,17 @@
 #include "Misc/MemStack.h"
 //#include "Animation/AnimationAsset.h"
 #include "Animation/AnimLinkableElement.h"
+#include "Animation/AnimEnums.h"
+#include "Misc/SecureHash.h"
 #include "AnimTypes.generated.h"
 
 struct FMarkerPair;
 struct FMarkerSyncAnimPosition;
 struct FPassedMarker;
 
+class FMemoryReader;
+class FMemoryWriter;
+
 // Disable debugging information for shipping and test builds.
 #define ENABLE_ANIM_DEBUG (1 && !(UE_BUILD_SHIPPING || UE_BUILD_TEST))
 
@@ -738,3 +743,135 @@ namespace EComponentType
 		ScaleZ
 	};
 }
+
+// @note We have a plan to support skeletal hierarchy. When that happens, we'd like to keep skeleton indexing.
+USTRUCT()
+struct ENGINE_API FTrackToSkeletonMap
+{
+	GENERATED_USTRUCT_BODY()
+
+	// Index of Skeleton.BoneTree this Track belongs to.
+	UPROPERTY()
+	int32 BoneTreeIndex;
+
+
+	FTrackToSkeletonMap()
+		: BoneTreeIndex(0)
+	{
+	}
+
+	FTrackToSkeletonMap(int32 InBoneTreeIndex)
+		: BoneTreeIndex(InBoneTreeIndex)
+	{
+	}
+
+	friend FArchive& operator<<(FArchive& Ar, FTrackToSkeletonMap &Item)
+	{
+		return Ar << Item.BoneTreeIndex;
+	}
+};
+
+/**
+* Raw keyframe data for one track.Each array will contain either NumFrames elements or 1 element.
+* One element is used as a simple compression scheme where if all keys are the same, they'll be
+* reduced to 1 key that is constant over the entire sequence.
+*/
+USTRUCT()
+struct ENGINE_API FRawAnimSequenceTrack
+{
+	GENERATED_USTRUCT_BODY()
+
+	/** Position keys. */
+	UPROPERTY()
+	TArray<FVector> PosKeys;
+
+	/** Rotation keys. */
+	UPROPERTY()
+	TArray<FQuat> RotKeys;
+
+	/** Scale keys. */
+	UPROPERTY()
+	TArray<FVector> ScaleKeys;
+
+	// Serializer.
+	friend FArchive& operator<<(FArchive& Ar, FRawAnimSequenceTrack& T)
+	{
+		T.PosKeys.BulkSerialize(Ar);
+		T.RotKeys.BulkSerialize(Ar);
+
+		if (Ar.UE4Ver() >= VER_UE4_ANIM_SUPPORT_NONUNIFORM_SCALE_ANIMATION)
+		{
+			T.ScaleKeys.BulkSerialize(Ar);
+		}
+
+		return Ar;
+	}
+};
+
+
+/**
+ * Encapsulates commonly useful data about bones.
+ */
+class FBoneData
+{
+public:
+	FQuat		Orientation;
+	FVector		Position;
+	/** Bone name. */
+	FName		Name;
+	/** Direct descendants.  Empty for end effectors. */
+	TArray<int32> Children;
+	/** List of bone indices from parent up to root. */
+	TArray<int32>	BonesToRoot;
+	/** List of end effectors for which this bone is an ancestor.  End effectors have only one element in this list, themselves. */
+	TArray<int32>	EndEffectors;
+	/** If a Socket is attached to that bone */
+	bool		bHasSocket;
+	/** If matched as a Key end effector */
+	bool		bKeyEndEffector;
+
+	/**	@return		Index of parent bone; -1 for the root. */
+	int32 GetParent() const
+	{
+		return GetDepth() ? BonesToRoot[0] : -1;
+	}
+	/**	@return		Distance to root; 0 for the root. */
+	int32 GetDepth() const
+	{
+		return BonesToRoot.Num();
+	}
+	/** @return		true if this bone is an end effector (has no children). */
+	bool IsEndEffector() const
+	{
+		return Children.Num() == 0;
+	}
+};
+
+#if 0
+template <typename ArrayType>
+FGuid GetArrayGuid(const TArray<ArrayType>& Array)
+{
+	FSHA1 Sha;
+	Sha.Update((uint8*)Array.GetData(), Array.Num() * Array.GetTypeSize());
+
+	Sha.Final();
+
+	uint32 Hash[5];
+	Sha.GetHash((uint8*)Hash);
+	FGuid Guid(Hash[0] ^ Hash[4], Hash[1], Hash[2], Hash[3]);
+	return Guid;
+}
+
+template <typename ArrayType>
+void DebugLogArray(const TArray<ArrayType>& Array)
+{
+	FPlatformMisc::LowLevelOutputDebugStringf(TEXT("Array: %i %s\n"), Array.Num(), *GetArrayGuid(Array).ToString());
+	for (int32 i = 0; i < Array.Num(); ++i)
+	{
+		FPlatformMisc::LowLevelOutputDebugStringf(TEXT("\t%i %s\n"), i, *Array[i].ToString());
+	}
+}
+
+template <>
+void DebugLogArray(const TArray<FRawAnimSequenceTrack>& RawData);
+#endif
diff --git a/Engine/Source/Runtime/Engine/Public/AnimationUtils.h b/Engine/Source/Runtime/Engine/Public/AnimationUtils.h
--- a/Engine/Source/Runtime/Engine/Public/AnimationUtils.h
+++ b/Engine/Source/Runtime/Engine/Public/AnimationUtils.h
@@ -23,44 +23,6 @@ class UAnimCompress;
 class UAnimSet;
 struct FTransform;
 
-/**
- * Encapsulates commonly useful data about bones.
- */
-class FBoneData
-{
-public:
-	FQuat		Orientation;
-	FVector		Position;
-	/** Bone name. */
-	FName		Name;
-	/** Direct descendants.  Empty for end effectors. */
-	TArray<int32> Children;
-	/** List of bone indices from parent up to root. */
-	TArray<int32>	BonesToRoot;
-	/** List of end effectors for which this bone is an ancestor.  End effectors have only one element in this list, themselves. */
-	TArray<int32>	EndEffectors;
-	/** If a Socket is attached to that bone */
-	bool		bHasSocket;
-	/** If matched as a Key end effector */
-	bool		bKeyEndEffector;
-
-	/**	@return		Index of parent bone; -1 for the root. */
-	int32 GetParent() const
-	{
-		return GetDepth() ? BonesToRoot[0] : -1;
-	}
-	/**	@return		Distance to root; 0 for the root. */
-	int32 GetDepth() const
-	{
-		return BonesToRoot.Num();
-	}
-	/** @return		true if this bone is an end effector (has no children). */
-	bool IsEndEffector() const
-	{
-		return Children.Num() == 0;
-	}
-};
-
 /**
  * Error statistics for an animation generated by TallyErrorsFromPerturbation, representing errors in translation or rotation due to a small incremental translation or rotation
  */
@@ -129,7 +91,9 @@ public:
 
 	static void BuildSkeletonMetaData(USkeleton* Skeleton, TArray<FBoneData>& OutBoneData);
 
+	static int32 GetAnimTrackIndexForSkeletonBone(const int32 InSkeletonBoneIndex, const TArray<FTrackToSkeletonMap>& TrackToSkelMap);
 
+#if WITH_EDITOR
 	/**
 	 * Utility function to measure the accuracy of a compressed animation. Each end-effector is checked for 
 	 * world-space movement as a result of compression.
@@ -139,7 +103,8 @@ public:
 	 * @param	ErrorStats	Output structure containing the final compression error values
 	 * @return				None.
 	 */
-	static void ComputeCompressionError(const UAnimSequence* AnimSeq, const TArray<FBoneData>& BoneData, AnimationErrorStats& ErrorStats);
+	static void ComputeCompressionError(const FCompressibleAnimData& CompressibleAnimData, FCompressibleAnimDataResult& CompressedData, AnimationErrorStats& ErrorStats);
+#endif
 
 	/**
 	 * Utility function to compress an animation. If the animation is currently associated with a codec, it will be used to 
@@ -148,11 +113,11 @@ public:
 	 * of the compressed animation remains within tolerances, the alternative codec will be used. 
 	 * See GetAlternativeCompressionThreshold for information on the tolerance value used.
 	 *
-	 * @param	AnimSeq				The anim sequence to compress.
-	 * @param	CompressContext		Context for applying compression, tracking memory saved etc
-	 * @return						None.
+	 * @param	CompressibleAnimData	The animation data to compress.
+	 * @param	CompressContext			Context for applying compression, tracking memory saved etc
+	 * @return							The compressed animation data
 	 */
-	ENGINE_API static void CompressAnimSequence(UAnimSequence* AnimSeq, FAnimCompressContext& CompressContext);
+	ENGINE_API static void CompressAnimSequence(const FCompressibleAnimData& CompressibleAnimData, FCompressibleAnimDataResult& OutCompressedData, FAnimCompressContext& CompressContext);
 
 	/**
 	 * Utility function to compress an animation. If the animation is currently associated with a codec, it will be used to
@@ -161,7 +126,8 @@ public:
 	 * of the compressed animation remains within tolerances, the alternative codec will be used.
 	 * See GetAlternativeCompressionThreshold for information on the tolerance value used.
 	 *
-	 * @param	AnimSeq									The anim sequence to compress.
+	 * @param	CompressibleAnimData					The animation data to compress.
+	 * @param	OriginalCompressedData					The original compressed data (if it exists).
 	 * @param	CompressContext							Context for applying compression, tracking memory saved etc*
 	 * @param	MasterTolerance							The alternate error threshold (0.0 means don't try anything other than the current / default scheme)
 	 * @param	bFirstRecompressUsingCurrentOrDefault	If true, then the animation will be first recompressed with it's current compressor if non-NULL, or with the global default compressor (specified in the engine ini)
@@ -169,10 +135,11 @@ public:
 	 * @param	bRaiseMaxErrorToExisting				If true and the existing compression error is greater than MasterTolerance, then MasterTolerance will be effectively raised to the existing error level
 	 * @param	bTryExhaustiveSearch					If true, then an exhaustive search is used otherwise only a short list of the best methods is tried
 	 *
-	 * @return	None.
+	 * @return											The compressed animation data
 	 */
 	static void CompressAnimSequenceExplicit(
-		UAnimSequence* AnimSeq,
+		const FCompressibleAnimData& CompressibleAnimData,
+		FCompressibleAnimDataResult& OutCompressedData,
 		FAnimCompressContext& CompressContext,
 		float MasterTolerance,
 		bool bFirstRecompressUsingCurrentOrDefault,
@@ -181,17 +148,7 @@ public:
 		bool bTryExhaustiveSearch,
 		bool bEnableSegmenting,
 		int32 IdealNumFramesPerSegment,
-		int32 MaxNumFramesPerSegment,
-		const TArray<FBoneData>& BoneData);
-
-	/**
-	 * Tests for a missing or invalid mesh on the animation sequence, warning if one or more were found
-	 *
-	 * @param AnimSeq	The anim sequence to check.
-	 *
-	 * @return	None.
-	 */
-	static void TestForMissingMeshes(UAnimSequence* AnimSeq);
+		int32 MaxNumFramesPerSegment);
 
 	/**
 	 * Determines the current setting for world-space error tolerance in the animation compressor.
@@ -240,32 +197,30 @@ public:
 	 * The track height is defined as the minimal number of bones away from an end effector (end effectors are 0, their parents are 1, etc...)
 	 *
 	 * @param AnimLinkup			The animation linkup
-	 * @param BoneData				The bone data to check
 	 * @param NumTracks				The number of tracks
 	 * @param TrackHeights [OUT]	The computed track heights
 	 *
 	 */
-	static void CalculateTrackHeights(UAnimSequence* AnimSeq, const TArray<FBoneData>& BoneData, int NumTracks, TArray<int32>& TrackHeights);
+	static void CalculateTrackHeights(const FCompressibleAnimData& CompressibleAnimData, int NumTracks, TArray<int32>& TrackHeights);
 
 	/**
 	 * Checks a set of key times to see if the spacing is uniform or non-uniform.
 	 * Note: If there are as many times as frames, they are automatically assumed to be uniformly spaced.
 	 * Note: If there are two or fewer times, they are automatically assumed to be uniformly spaced.
 	 *
-	 * @param AnimSeq		The animation sequence the Times array is associated with
+	 * @param AnimSeq		The number of frames in the source animation data
 	 * @param Times			The array of key times
 	 *
 	 * @return				true if the keys are uniformly spaced (or one of the trivial conditions is detected).  false if any key spacing is greater than 1e-4 off.
 	 */
-	static bool HasUniformKeySpacing(UAnimSequence* AnimSeq, const TArray<float>& Times);
+	static bool HasUniformKeySpacing(int32 NumFrames, const TArray<float>& Times);
 
 	/**
 	 * Perturbs the bone(s) associated with each track in turn, measuring the maximum error introduced in end effectors as a result
 	 */
 	static void TallyErrorsFromPerturbation(
-		const UAnimSequence* AnimSeq,
+		const FCompressibleAnimData& CompressibleAnimData,
 		int32 NumTracks,
-		const TArray<FBoneData>& BoneData,
 		const FVector& PositionNudge,
 		const FQuat& RotationNudge,
 		const FVector& ScaleNudge,
@@ -279,13 +234,16 @@ public:
 	/** Returns the default animation curve compression settings, can never by null. */
 	ENGINE_API static UAnimCurveCompressionSettings* GetDefaultAnimationCurveCompressionSettings();
 
+	static void ExtractTransformFromTrack(float Time, int32 NumFrames, float SequenceLength, const struct FRawAnimSequenceTrack& RawTrack, EAnimInterpolationType Interpolation, FTransform &OutAtom);
+
 #if WITH_EDITOR
+	static void ExtractTransformFromCompressionData(const FCompressibleAnimData& CompressibleAnimData, FCompressibleAnimDataResult& CompressedAnimData, float Time, int32 TrackIndex, bool bUseRawData, FTransform& OutBoneTransform);
 	/**
 	 * Compresses the animation curves within a sequence with the chosen settings.
 	 * Note: This modifies the sequence.
 	 *
 	 * @return		Returns true on success, false it we fail to compress
 	 */
-	ENGINE_API static bool CompressAnimCurves(UAnimSequence& AnimSeq);
+	ENGINE_API static bool CompressAnimCurves(FCompressibleAnimData& AnimSeq, FCompressedAnimSequence& Target);
 #endif
 };
