Non-unity CIS fixes

diff --git a/Engine/Source/Runtime/Engine/Classes/Animation/AnimCurveCompressionCodec.h b/Engine/Source/Runtime/Engine/Classes/Animation/AnimCurveCompressionCodec.h
--- a/Engine/Source/Runtime/Engine/Classes/Animation/AnimCurveCompressionCodec.h
+++ b/Engine/Source/Runtime/Engine/Classes/Animation/AnimCurveCompressionCodec.h
@@ -4,6 +4,7 @@
 
 #include "CoreMinimal.h"
 #include "Animation/AnimTypes.h"
+#include "Animation/AnimCompressionTypes.h"
 #include "AnimCurveCompressionCodec.generated.h"
 
 class UAnimCurveCompressionCodec;
diff --git a/Engine/Source/Runtime/Engine/Classes/Animation/AnimCurveCompressionSettings.h b/Engine/Source/Runtime/Engine/Classes/Animation/AnimCurveCompressionSettings.h
--- a/Engine/Source/Runtime/Engine/Classes/Animation/AnimCurveCompressionSettings.h
+++ b/Engine/Source/Runtime/Engine/Classes/Animation/AnimCurveCompressionSettings.h
@@ -3,6 +3,7 @@
 #pragma once
 
 #include "CoreMinimal.h"
+#include "Animation/AnimCompressionTypes.h"
 #include "AnimCurveCompressionSettings.generated.h"
 
 class UAnimCurveCompressionCodec;
diff --git a/Engine/Source/Runtime/Engine/Private/Animation/AnimCompressionDerivedData.h b/Engine/Source/Runtime/Engine/Private/Animation/AnimCompressionDerivedData.h
--- a/Engine/Source/Runtime/Engine/Private/Animation/AnimCompressionDerivedData.h
+++ b/Engine/Source/Runtime/Engine/Private/Animation/AnimCompressionDerivedData.h
@@ -8,6 +8,8 @@
 #include "DerivedDataPluginInterface.h"
 #endif
 
+#include "Animation/AnimCompressionTypes.h"
+
 class UAnimSequence;
 struct FAnimCompressContext;
 
diff --git a/Engine/Source/Runtime/Engine/Private/Animation/AnimCompressionTypes.cpp b/Engine/Source/Runtime/Engine/Private/Animation/AnimCompressionTypes.cpp
--- a/Engine/Source/Runtime/Engine/Private/Animation/AnimCompressionTypes.cpp
+++ b/Engine/Source/Runtime/Engine/Private/Animation/AnimCompressionTypes.cpp
@@ -6,6 +6,7 @@
 #include "Interfaces/ITargetPlatform.h"
 #include "Animation/AnimCurveCompressionCodec.h"
 #include "Animation/AnimCurveCompressionSettings.h"
+#include "UObject/FortniteMainBranchObjectVersion.h"
 
 FCompressibleAnimData::FCompressibleAnimData(class UAnimSequence* InSeq) :
 #if WITH_EDITOR
