Fix CIS issue for CL 6846811

diff --git a/Engine/Source/Runtime/Engine/Public/Animation/AnimCompressionTypes.h b/Engine/Source/Runtime/Engine/Public/Animation/AnimCompressionTypes.h
--- a/Engine/Source/Runtime/Engine/Public/Animation/AnimCompressionTypes.h
+++ b/Engine/Source/Runtime/Engine/Public/Animation/AnimCompressionTypes.h
@@ -455,6 +455,7 @@ template<uint32 Alignment = DEFAULT_ALIGNMENT>
 class TMaybeMappedAllocator
 {
 public:
+	using SizeType = int32;
 
 	enum { NeedsElementType = false };
 	enum { RequireRangeCheck = true };
@@ -503,8 +504,8 @@ public:
 			return Data;
 		}
 		void ResizeAllocation(
-			int32 PreviousNumElements,
-			int32 NumElements,
+			SizeType PreviousNumElements,
+			SizeType NumElements,
 			SIZE_T NumBytesPerElement
 		)
 		{
@@ -516,28 +517,28 @@ public:
 				Data = (FScriptContainerElement*)FMemory::Realloc(Data, NumElements*NumBytesPerElement, Alignment);
 			}
 		}
-		int32 CalculateSlackReserve(int32 NumElements, int32 NumBytesPerElement) const
+		SizeType CalculateSlackReserve(SizeType NumElements, SIZE_T NumBytesPerElement) const
 		{
 			check(!MappedHandle && !MappedRegion); // this could be supported, but it probably is never what you want, so we will just assert.
 			return DefaultCalculateSlackReserve(NumElements, NumBytesPerElement, true, Alignment);
 		}
-		int32 CalculateSlackShrink(int32 NumElements, int32 NumAllocatedElements, int32 NumBytesPerElement) const
+		SizeType CalculateSlackShrink(SizeType NumElements, SizeType NumAllocatedElements, SIZE_T NumBytesPerElement) const
 		{
 			check(!MappedHandle && !MappedRegion); // this could be supported, but it probably is never what you want, so we will just assert.
 			return DefaultCalculateSlackShrink(NumElements, NumAllocatedElements, NumBytesPerElement, true, Alignment);
 		}
-		int32 CalculateSlackGrow(int32 NumElements, int32 NumAllocatedElements, int32 NumBytesPerElement) const
+		SizeType CalculateSlackGrow(SizeType NumElements, SizeType NumAllocatedElements, SIZE_T NumBytesPerElement) const
 		{
 			check(!MappedHandle && !MappedRegion); // this could be supported, but it probably is never what you want, so we will just assert.
 			return DefaultCalculateSlackGrow(NumElements, NumAllocatedElements, NumBytesPerElement, true, Alignment);
 		}
 
-		SIZE_T GetAllocatedSize(int32 NumAllocatedElements, SIZE_T NumBytesPerElement) const
+		SIZE_T GetAllocatedSize(SizeType NumAllocatedElements, SIZE_T NumBytesPerElement) const
 		{
 			return NumAllocatedElements * NumBytesPerElement;
 		}
 
-		bool HasAllocation()
+		bool HasAllocation() const
 		{
 			return !!Data;
 		}
