Fixed crash in FSceneView::StartFinalPostprocessSettings due to Family->Scene being nullptr and added code to prevent crash if Family is nullptr as an additional precaution.

#UE4 #ReleaseNote

diff --git a/Engine/Source/Runtime/Engine/Private/SceneView.cpp b/Engine/Source/Runtime/Engine/Private/SceneView.cpp
--- a/Engine/Source/Runtime/Engine/Private/SceneView.cpp
+++ b/Engine/Source/Runtime/Engine/Private/SceneView.cpp
@@ -1715,15 +1715,15 @@ void FSceneView::StartFinalPostprocessSettings(FVector InViewLocation)
 		}
 	}
 
-	if(State)
+	if (State != nullptr)
 	{
 		State->OnStartPostProcessing(*this);
 	}
 
-	UWorld* World = Family->Scene->GetWorld();
+	UWorld* World = ((Family != nullptr) && (Family->Scene != nullptr)) ? Family->Scene->GetWorld() : nullptr;
 
 	// Some views have no world (e.g. material preview)
-	if (World)
+	if (World != nullptr)
 	{
 		for (auto VolumeIt = World->PostProcessVolumes.CreateIterator(); VolumeIt; ++VolumeIt)
 		{
