Make sure check is done when a sound is stopping to avoid realizing when virtual voice is pending stop. Repurposes bIsStopping which was flag only used by playing ActiveSounds, not the intended check to see if the Virtualized ActiveSound is pending stop.
Remove minor cruft

diff --git a/Engine/Source/Runtime/Engine/Private/ActiveSound.cpp b/Engine/Source/Runtime/Engine/Private/ActiveSound.cpp
--- a/Engine/Source/Runtime/Engine/Private/ActiveSound.cpp
+++ b/Engine/Source/Runtime/Engine/Private/ActiveSound.cpp
@@ -125,6 +125,8 @@ FActiveSound::~FActiveSound()
 
 FActiveSound* FActiveSound::CreateVirtualCopy(const FActiveSound& InActiveSoundToCopy, FAudioDevice& InAudioDevice)
 {
+	check(!InActiveSoundToCopy.bIsStopping);
+
 	FActiveSound* ActiveSound = new FActiveSound(InActiveSoundToCopy);
 
 	ActiveSound->bAsyncOcclusionPending = false;
diff --git a/Engine/Source/Runtime/Engine/Private/AudioDevice.cpp b/Engine/Source/Runtime/Engine/Private/AudioDevice.cpp
--- a/Engine/Source/Runtime/Engine/Private/AudioDevice.cpp
+++ b/Engine/Source/Runtime/Engine/Private/AudioDevice.cpp
@@ -4264,6 +4264,8 @@ bool FAudioDevice::RemoveVirtualLoop(FActiveSound& InActiveSound)
 
 	if (FAudioVirtualLoop* VirtualLoop = VirtualLoops.Find(&InActiveSound))
 	{
+		check(InActiveSound.bIsStopping);
+
 		const uint64 ComponentID = InActiveSound.GetAudioComponentID();
 		if (ComponentID > 0)
 		{
@@ -4394,7 +4396,11 @@ void FAudioDevice::AddSoundToStop(FActiveSound* SoundToStop)
 			AudioComponentIDToActiveSoundMap.Remove(AudioComponentID);
 		}
 
-		if (!VirtualLoops.Contains(SoundToStop))
+		if (VirtualLoops.Contains(SoundToStop))
+		{
+			SoundToStop->bIsStopping = true;
+		}
+		else
 		{
 			ConcurrencyManager.RemoveActiveSound(*SoundToStop);
 		}
@@ -5737,20 +5743,12 @@ void FAudioDevice::UpdateVirtualLoops()
 			{
 				VirtualLoopsToRetrigger.Add(VirtualLoop);
 			}
-
-			// If set to fade out and set to play when silent, add to pending stop list.
-			if (ActiveSound.bFadingOut)
-			{
-				AddSoundToStop(&ActiveSound);
-			}
 		}
 
 		for (FAudioVirtualLoop& RetriggerLoop : VirtualLoopsToRetrigger)
 		{
 			RetriggerVirtualLoop(RetriggerLoop);
 		}
-
-		VirtualLoopsToRetrigger.Reset();
 	}
 
 	// if !FAudioVirtualLoop::IsEnabled(), attempt to realize/re-trigger
