#UE4 Adding more context to a warning log message about DrawMaterialToRenderTarget

diff --git a/Engine/Source/Runtime/Engine/Private/KismetRenderingLibrary.cpp b/Engine/Source/Runtime/Engine/Private/KismetRenderingLibrary.cpp
--- a/Engine/Source/Runtime/Engine/Private/KismetRenderingLibrary.cpp
+++ b/Engine/Source/Runtime/Engine/Private/KismetRenderingLibrary.cpp
@@ -99,15 +99,15 @@ void UKismetRenderingLibrary::DrawMaterialToRenderTarget(UObject* WorldContextOb
 	}
 	else if (!Material)
 	{
-		FMessageLog("Blueprint").Warning(LOCTEXT("DrawMaterialToRenderTarget_InvalidMaterial", "DrawMaterialToRenderTarget: Material must be non-null."));
+		FMessageLog("Blueprint").Warning(FText::Format(LOCTEXT("DrawMaterialToRenderTarget_InvalidMaterial", "DrawMaterialToRenderTarget[{0}]: Material must be non-null."), FText::FromString(GetPathNameSafe(WorldContextObject))));
 	}
 	else if (!TextureRenderTarget)
 	{
-		FMessageLog("Blueprint").Warning(LOCTEXT("DrawMaterialToRenderTarget_InvalidTextureRenderTarget", "DrawMaterialToRenderTarget: TextureRenderTarget must be non-null."));
+		FMessageLog("Blueprint").Warning(FText::Format(LOCTEXT("DrawMaterialToRenderTarget_InvalidTextureRenderTarget", "DrawMaterialToRenderTarget[{0}]: TextureRenderTarget must be non-null."), FText::FromString(GetPathNameSafe(WorldContextObject))));
 	}
 	else if (!TextureRenderTarget->Resource)
 	{
-		FMessageLog("Blueprint").Warning(LOCTEXT("DrawMaterialToRenderTarget_ReleasedTextureRenderTarget", "DrawMaterialToRenderTarget: render target has been released."));
+		FMessageLog("Blueprint").Warning(FText::Format(LOCTEXT("DrawMaterialToRenderTarget_ReleasedTextureRenderTarget", "DrawMaterialToRenderTarget[{0}]: render target has been released."), FText::FromString(GetPathNameSafe(WorldContextObject))));
 	}
 	else
 	{
