Don't look for implicit NAKs for Fast Array Delta Serializers on completely reliable connections.

diff --git a/Engine/Source/Runtime/CoreUObject/Public/UObject/CoreNet.h b/Engine/Source/Runtime/CoreUObject/Public/UObject/CoreNet.h
--- a/Engine/Source/Runtime/CoreUObject/Public/UObject/CoreNet.h
+++ b/Engine/Source/Runtime/CoreUObject/Public/UObject/CoreNet.h
@@ -502,6 +502,12 @@ struct FNetDeltaSerializeInfo
 	/** Whether or not we support FFastArraySerializer::FastArrayDeltaSerialize_DeltaSerializeStructs */
 	bool bSupportsFastArrayDeltaStructSerialization = false;
 
+	/**
+	 * Whether or the connection is completely reliable.
+	 * We cache this off separate from UNetConnection so we can limit usage.
+	 */
+	bool bInternalAck = false;
+
 	/** The object that owns the struct we're serializing. */
 	UObject* Object = nullptr;
 
diff --git a/Engine/Source/Runtime/Engine/Classes/Engine/NetSerialization.h b/Engine/Source/Runtime/Engine/Classes/Engine/NetSerialization.h
--- a/Engine/Source/Runtime/Engine/Classes/Engine/NetSerialization.h
+++ b/Engine/Source/Runtime/Engine/Classes/Engine/NetSerialization.h
@@ -963,19 +963,23 @@ void FFastArraySerializer::TFastArraySerializeHelper<Type, SerializerType>::Post
 	// Look for implicit deletes that would happen due to Naks
 	// ---------------------------------------------------------
 
-	for (int32 idx = 0; idx < Items.Num(); ++idx)
+	// If we're sending data completely reliably, there's no need to do this.
+	if (!Parms.bInternalAck)
 	{
-		Type& Item = Items[idx];
-		if (Item.MostRecentArrayReplicationKey < Header.ArrayReplicationKey && Item.MostRecentArrayReplicationKey > Header.BaseReplicationKey)
+		for (int32 idx = 0; idx < Items.Num(); ++idx)
 		{
-			// Make sure this wasn't an explicit delete in this bunch (otherwise we end up deleting an extra element!)
-			if (!Header.DeletedIndices.Contains(idx))
+			Type& Item = Items[idx];
+			if (Item.MostRecentArrayReplicationKey < Header.ArrayReplicationKey && Item.MostRecentArrayReplicationKey > Header.BaseReplicationKey)
 			{
-				// This will happen in normal conditions in network replays.
-				UE_LOG(LogNetFastTArray, Log, TEXT("Adding implicit delete for ElementID: %d. MostRecentArrayReplicationKey: %d. Current Payload: [%d/%d]"),
-					Item.ReplicationID, Item.MostRecentArrayReplicationKey, Header.ArrayReplicationKey, Header.BaseReplicationKey);
+				// Make sure this wasn't an explicit delete in this bunch (otherwise we end up deleting an extra element!)
+				if (!Header.DeletedIndices.Contains(idx))
+				{
+					// This will happen in normal conditions in network replays.
+					UE_LOG(LogNetFastTArray, Log, TEXT("Adding implicit delete for ElementID: %d. MostRecentArrayReplicationKey: %d. Current Payload: [%d/%d]"),
+						Item.ReplicationID, Item.MostRecentArrayReplicationKey, Header.ArrayReplicationKey, Header.BaseReplicationKey);
 
-				Header.DeletedIndices.Add(idx);
+					Header.DeletedIndices.Add(idx);
+				}
 			}
 		}
 	}
diff --git a/Engine/Source/Runtime/Engine/Private/DataReplication.cpp b/Engine/Source/Runtime/Engine/Private/DataReplication.cpp
--- a/Engine/Source/Runtime/Engine/Private/DataReplication.cpp
+++ b/Engine/Source/Runtime/Engine/Private/DataReplication.cpp
@@ -325,6 +325,7 @@ bool FObjectReplicator::SerializeCustomDeltaProperty(
 	Parms.Data = Property->ContainerPtrToValuePtr<void>(Src, ArrayIndex);
 	Parms.Object = reinterpret_cast<UObject*>(Src);
 	Parms.Connection = Connection;
+	Parms.bInternalAck = Connection->InternalAck;
 	Parms.Writer = &OutBunch;
 	Parms.Map = Connection->PackageMap;
 	Parms.OldState = OldState.Get();
@@ -374,6 +375,7 @@ bool FObjectReplicator::SendCustomDeltaProperty(UObject* InObject, uint16 Custom
 	Parms.CustomDeltaIndex = CustomDeltaIndex;
 	Parms.bSupportsFastArrayDeltaStructSerialization = bSupportsFastArrayDelta;
 	Parms.Connection = Connection;
+	Parms.bInternalAck = Connection->InternalAck;
 
 	return FNetSerializeCB::SendCustomDeltaProperty(*RepLayout, Parms, CustomDeltaIndex);
 }
@@ -1032,6 +1034,7 @@ bool FObjectReplicator::ReceivedBunch(FNetBitReader& Bunch, const FReplicationFl
 			Parms.Reader = &Reader;
 			Parms.NetSerializeCB = &NetSerializeCB;
 			Parms.Connection = Connection;
+			Parms.bInternalAck = Connection->InternalAck;
 			Parms.Object = Object;
 
 			if (!FNetSerializeCB::ReceiveCustomDeltaProperty(LocalRepLayout, ReceivingRepState, Parms, ReplicatedProp))
@@ -1280,6 +1283,7 @@ void FObjectReplicator::UpdateGuidToReplicatorMap()
 		Parms.GatherGuidReferences = &LocalReferencedGuids;
 		Parms.TrackedGuidMemoryBytes = &LocalTrackedGuidMemoryBytes;
 		Parms.Object = GetObject();
+		Parms.bInternalAck = Connection->InternalAck;
 
 		LocalRepLayout.GatherGuidReferences(RepState->GetReceivingRepState(), Parms, LocalReferencedGuids, LocalTrackedGuidMemoryBytes);
 	}
@@ -1336,6 +1340,7 @@ bool FObjectReplicator::MoveMappedObjectToUnmapped(const FNetworkGUID& GUID)
 	FNetSerializeCB NetSerializeCB(Connection->Driver);
 	FNetDeltaSerializeInfo Parms;
 	Parms.Connection = Connection;
+	Parms.bInternalAck = Connection->InternalAck;
 	Parms.Map = Connection->PackageMap;
 	Parms.Object = GetObject();
 	Parms.NetSerializeCB = &NetSerializeCB;
@@ -1979,6 +1984,7 @@ void FObjectReplicator::UpdateUnmappedObjects(bool & bOutHasMoreUnmapped)
 	FNetDeltaSerializeInfo Parms;
 	Parms.Object = Object;
 	Parms.Connection = Connection;
+	Parms.bInternalAck = Connection->InternalAck;
 	Parms.Map = Connection->PackageMap;
 	Parms.NetSerializeCB = &NetSerializeCB;
 
