#UE4 Quick fix for CIS. Adding NetworkReplayStreaming to the CircularlyReferencedDependentModules list in Engine.Build.cs

diff --git a/Engine/Source/Runtime/Engine/Engine.Build.cs b/Engine/Source/Runtime/Engine/Engine.Build.cs
--- a/Engine/Source/Runtime/Engine/Engine.Build.cs
+++ b/Engine/Source/Runtime/Engine/Engine.Build.cs
@@ -187,6 +187,7 @@ public class Engine : ModuleRules
 		CircularlyReferencedDependentModules.Add("MaterialShaderQualitySettings");
 		CircularlyReferencedDependentModules.Add("CinematicCamera");
 		CircularlyReferencedDependentModules.Add("AudioMixer");
+		CircularlyReferencedDependentModules.Add("NetworkReplayStreaming");
 
 		// The AnimGraphRuntime module is not needed by Engine proper, but it is loaded in LaunchEngineLoop.cpp,
 		// and needs to be listed in an always-included module in order to be compiled into standalone games
