Back out changelist 6858310

diff --git a/Engine/Source/Runtime/Engine/Engine.Build.cs b/Engine/Source/Runtime/Engine/Engine.Build.cs
--- a/Engine/Source/Runtime/Engine/Engine.Build.cs
+++ b/Engine/Source/Runtime/Engine/Engine.Build.cs
@@ -187,7 +187,6 @@ public class Engine : ModuleRules
 		CircularlyReferencedDependentModules.Add("MaterialShaderQualitySettings");
 		CircularlyReferencedDependentModules.Add("CinematicCamera");
 		CircularlyReferencedDependentModules.Add("AudioMixer");
-		CircularlyReferencedDependentModules.Add("NetworkReplayStreaming");
 
 		// The AnimGraphRuntime module is not needed by Engine proper, but it is loaded in LaunchEngineLoop.cpp,
 		// and needs to be listed in an always-included module in order to be compiled into standalone games
