Move behaviour for generating AnimationBlueprint functions out to separate function, call only from Link() in Editor and PostLoad in non-editor instead, this ensures all UFunctions will be loaded

diff --git a/Engine/Source/Runtime/Engine/Classes/Animation/AnimBlueprintGeneratedClass.h b/Engine/Source/Runtime/Engine/Classes/Animation/AnimBlueprintGeneratedClass.h
--- a/Engine/Source/Runtime/Engine/Classes/Animation/AnimBlueprintGeneratedClass.h
+++ b/Engine/Source/Runtime/Engine/Classes/Animation/AnimBlueprintGeneratedClass.h
@@ -194,7 +194,7 @@ class ENGINE_API UAnimBlueprintGeneratedClass : public UBlueprintGeneratedClass,
 	UPROPERTY()
 	TMap<FName, FCachedPoseIndices> OrderedSavedPoseIndicesMap;
 
-	// The various anim functions that this class holds (created during Link)
+	// The various anim functions that this class holds (created during GenerateAnimationBlueprintFunctions)
 	TArray<FAnimBlueprintFunction> AnimBlueprintFunctions;
 
 	// The arrays of anim nodes; this is transient generated data (created during Link)
@@ -223,7 +223,7 @@ public:
 	virtual int32 GetSyncGroupIndex(FName SyncGroupName) const override { return SyncGroupNames.IndexOfByKey(SyncGroupName); }
 	virtual const TArray<FExposedValueHandler>& GetExposedValueHandlers() const { return EvaluateGraphExposedInputs; }
 	virtual const TArray<FAnimBlueprintFunction>& GetAnimBlueprintFunctions() const override { return AnimBlueprintFunctions; }
-
+	
 public:
 #if WITH_EDITORONLY_DATA
 	FAnimBlueprintDebugData AnimBlueprintDebugData;
@@ -332,9 +332,12 @@ public:
 	const int32* GetNodePropertyIndexFromGuid(FGuid Guid, EPropertySearchMode::Type SearchMode = EPropertySearchMode::OnlyThis);
 #endif
 
-	// Called after Link to patch up refereces to the nodes in the CDO
+	// Called after Link to patch up references to the nodes in the CDO
 	void LinkFunctionsToDefaultObjectNodes(UObject* DefaultObject);
 
+	// Populates AnimBlueprintFunctions according to the UFunction(s) on this class
+	void GenerateAnimationBlueprintFunctions();
+
 	// UObject interface
 	virtual void Serialize(FArchive& Ar) override;
 	// End of UObject interface
@@ -347,6 +350,7 @@ public:
 	virtual void PurgeClass(bool bRecompilingOnLoad) override;
 	virtual uint8* GetPersistentUberGraphFrame(UObject* Obj, UFunction* FuncToCheck) const override;
 	virtual void PostLoadDefaultObject(UObject* Object) override;
+	virtual void PostLoad() override;
 	// End of UClass interface
 };
 
diff --git a/Engine/Source/Runtime/Engine/Private/Animation/AnimBlueprintGeneratedClass.cpp b/Engine/Source/Runtime/Engine/Private/Animation/AnimBlueprintGeneratedClass.cpp
--- a/Engine/Source/Runtime/Engine/Private/Animation/AnimBlueprintGeneratedClass.cpp
+++ b/Engine/Source/Runtime/Engine/Private/Animation/AnimBlueprintGeneratedClass.cpp
@@ -251,92 +251,15 @@
 {
 	Super::Link(Ar, bRelinkExistingProperties);
 
-	static const FName DefaultAnimGraphName("AnimGraph");
-
 	// @TODO: Shouldn't be necessary to clear these, but currently the class gets linked twice during compilation
 	AnimNodeProperties.Empty();
 	SubInstanceNodeProperties.Empty();
 	LayerNodeProperties.Empty();
-	AnimBlueprintFunctions.Empty();
-
-	// Patch up blueprint function info
-	for (TFieldIterator<UFunction> It(this); It; ++It)
-	{
-		bool bFoundOutput = false;
-#if WITH_EDITOR
-		// In editor we can grab the group from metadata, otherwise we need to wait until CDO post load (LinkFunctionsToDefaultObjectNodes)
-		FText CategoryText = It->GetMetaDataText(TEXT("Category"), TEXT("UObjectCategory"), It->GetFullGroupName(false));
-		FName Group = CategoryText.IsEmpty() ? NAME_None : FName(*CategoryText.ToString());
-#endif
-		UStructProperty* OutputPoseNodeProperty = nullptr;
-		TArray<FName> InputPoseNames;
-		TArray<int32> InputPoseNodeIndices;
-		TArray<UStructProperty*> InputPoseNodeProperties;
-		TArray<UProperty*> InputProperties;
-
-		// grab the input/output poses, their indices will be patched up later once the CDO is loaded in PostLoadDefaultObject
-		for (TFieldIterator<UProperty> ItParam(*It); ItParam; ++ItParam)
-		{
-			if(UStructProperty* StructProperty = Cast<UStructProperty>(*ItParam))
-			{
-				if(StructProperty->Struct->IsChildOf(FPoseLink::StaticStruct()))
-				{
-					if(StructProperty->GetPropertyFlags() & CPF_OutParm)
-					{
-						if(!bFoundOutput)
-						{
-							OutputPoseNodeProperty = StructProperty;
-							bFoundOutput = true;
-						}
-						else
-						{
-							// our required signature needs us to have a single post link output, so null it out if we find more than one
-							OutputPoseNodeProperty = nullptr;
-						}
-					}
-					else
-					{
-						InputPoseNames.Add(StructProperty->GetFName());
-						InputPoseNodeIndices.Add(INDEX_NONE);
-						InputPoseNodeProperties.Add(nullptr);
-					}
-				}
-				else
-				{
-					InputProperties.Add(*ItParam);
-				}
-			}
-			else
-			{
-				InputProperties.Add(*ItParam);
-			}
-		}
-
-		if(OutputPoseNodeProperty)
-		{
-			// We use the undecorated name here, so trim the postfix
-			FAnimBlueprintFunction* AnimBlueprintFunction = nullptr;
-
-			// Make sure that the default graph is at index 0
-			FName FunctionName = It->GetFName();
-			if(FunctionName == DefaultAnimGraphName)
-			{
-				AnimBlueprintFunction = &AnimBlueprintFunctions.Insert_GetRef(FAnimBlueprintFunction(FunctionName), 0);
-			}
-			else
-			{
-				AnimBlueprintFunction = &AnimBlueprintFunctions.Emplace_GetRef(FunctionName);
-			}
 
 #if WITH_EDITOR
-			AnimBlueprintFunction->Group = Group;
-#endif
-			AnimBlueprintFunction->InputPoseNames.Append(MoveTemp(InputPoseNames));
-			AnimBlueprintFunction->InputPoseNodeIndices.Append(MoveTemp(InputPoseNodeIndices));
-			AnimBlueprintFunction->InputPoseNodeProperties.Append(MoveTemp(InputPoseNodeProperties));
-			AnimBlueprintFunction->InputProperties.Append(MoveTemp(InputProperties));
-		}
-	}
+	// This relies on the entire class being fully loaded, this is not the case with EDL async-loading, in which case the functions are generated in PostLoad
+	GenerateAnimationBlueprintFunctions();
+#endif // WITH_EDITOR
 
 	// Initialize the various tracked node arrays & fix up function internals
 	for (TFieldIterator<UProperty> It(this); It; ++It)
@@ -423,6 +346,98 @@
 	LinkFunctionsToDefaultObjectNodes(Object);
 }
 
+void UAnimBlueprintGeneratedClass::PostLoad()
+{
+	Super::PostLoad();
+	GenerateAnimationBlueprintFunctions();
+}
+
+void UAnimBlueprintGeneratedClass::GenerateAnimationBlueprintFunctions()
+{
+	AnimBlueprintFunctions.Empty();
+
+	static const FName DefaultAnimGraphName("AnimGraph");
+	// Patch up blueprint function info
+	for (TFieldIterator<UFunction> It(this); It; ++It)
+	{
+		bool bFoundOutput = false;
+#if WITH_EDITOR
+		// In editor we can grab the group from metadata, otherwise we need to wait until CDO post load (LinkFunctionsToDefaultObjectNodes)
+		FText CategoryText = It->GetMetaDataText(TEXT("Category"), TEXT("UObjectCategory"), It->GetFullGroupName(false));
+		FName Group = CategoryText.IsEmpty() ? NAME_None : FName(*CategoryText.ToString());
+#endif
+		UStructProperty* OutputPoseNodeProperty = nullptr;
+		TArray<FName> InputPoseNames;
+		TArray<int32> InputPoseNodeIndices;
+		TArray<UStructProperty*> InputPoseNodeProperties;
+		TArray<UProperty*> InputProperties;
+
+		// grab the input/output poses, their indices will be patched up later once the CDO is loaded in PostLoadDefaultObject
+		for (TFieldIterator<UProperty> ItParam(*It); ItParam; ++ItParam)
+		{
+			if (UStructProperty* StructProperty = Cast<UStructProperty>(*ItParam))
+			{
+				if (StructProperty->Struct->IsChildOf(FPoseLink::StaticStruct()))
+				{
+					if (StructProperty->GetPropertyFlags() & CPF_OutParm)
+					{
+						if (!bFoundOutput)
+						{
+							OutputPoseNodeProperty = StructProperty;
+							bFoundOutput = true;
+						}
+						else
+						{
+							// our required signature needs us to have a single post link output, so null it out if we find more than one
+							OutputPoseNodeProperty = nullptr;
+						}
+					}
+					else
+					{
+						InputPoseNames.Add(StructProperty->GetFName());
+						InputPoseNodeIndices.Add(INDEX_NONE);
+						InputPoseNodeProperties.Add(nullptr);
+					}
+				}
+				else
+				{
+					InputProperties.Add(*ItParam);
+				}
+			}
+			else
+			{
+				InputProperties.Add(*ItParam);
+			}
+		}
+
+		if (OutputPoseNodeProperty)
+		{
+			// We use the undecorated name here, so trim the postfix
+			FAnimBlueprintFunction* AnimBlueprintFunction = nullptr;
+
+			FName FunctionName = It->GetFName();
+
+			// Make sure that the default graph is at index 0
+			if (FunctionName == DefaultAnimGraphName)
+			{
+				AnimBlueprintFunction = &AnimBlueprintFunctions.Insert_GetRef(FAnimBlueprintFunction(FunctionName), 0);
+			}
+			else
+			{
+				AnimBlueprintFunction = &AnimBlueprintFunctions.Emplace_GetRef(FunctionName);
+			}
+
+#if WITH_EDITOR
+			AnimBlueprintFunction->Group = Group;
+#endif
+			AnimBlueprintFunction->InputPoseNames.Append(MoveTemp(InputPoseNames));
+			AnimBlueprintFunction->InputPoseNodeIndices.Append(MoveTemp(InputPoseNodeIndices));
+			AnimBlueprintFunction->InputPoseNodeProperties.Append(MoveTemp(InputPoseNodeProperties));
+			AnimBlueprintFunction->InputProperties.Append(MoveTemp(InputProperties));
+		}
+	}
+}
+
 void UAnimBlueprintGeneratedClass::LinkFunctionsToDefaultObjectNodes(UObject* DefaultObject)
 {
 	// Link functions to their nodes
