Added throttling capabilities when copying mesh distance field to atlas distance field (r.DistanceFields.ThrottleCopyToAtlasInBytes)

#rnx

diff --git a/Engine/Source/Runtime/Engine/Private/DistanceFieldAtlas.cpp b/Engine/Source/Runtime/Engine/Private/DistanceFieldAtlas.cpp
--- a/Engine/Source/Runtime/Engine/Private/DistanceFieldAtlas.cpp
+++ b/Engine/Source/Runtime/Engine/Private/DistanceFieldAtlas.cpp
@@ -96,6 +96,13 @@ FAutoConsoleVariableRef CVarDistFieldForceAtlasRealloc(
 	ECVF_RenderThreadSafe
 );
 
+static TAutoConsoleVariable<int32> CVarDistFieldThrottleCopyToAtlasInBytes(
+	TEXT("r.DistanceFields.ThrottleCopyToAtlasInBytes"),
+	0,
+	TEXT("When enabled (higher than 0), throttle mesh distance field copy to global mesh distance field atlas volume (in bytes uncompressed)."),
+	ECVF_Default);
+
+
 static TAutoConsoleVariable<int32> CVarLandscapeGI(
 	TEXT("r.GenerateLandscapeGIData"),
 	0,
@@ -243,6 +250,12 @@ void FDistanceFieldVolumeTextureAtlas::AddAllocation(FDistanceFieldVolumeTexture
 {
 	InitializeIfNeeded();
 	PendingAllocations.AddUnique(Texture);
+	const int32 ThrottleSize = CVarDistFieldThrottleCopyToAtlasInBytes.GetValueOnAnyThread();
+	if (ThrottleSize >= 1024)
+	{
+		Texture->bThrottled = true;
+	}
+	
 }
 
 void FDistanceFieldVolumeTextureAtlas::RemoveAllocation(FDistanceFieldVolumeTexture* Texture)
@@ -273,27 +286,60 @@ void FDistanceFieldVolumeTextureAtlas::UpdateAllocations()
 	{
 		const double StartTime = FPlatformTime::Seconds();
 
+		const int32 FormatSize = GPixelFormats[Format].BlockBytes;
+
 		// Sort largest to smallest for best packing
 		PendingAllocations.Sort(FCompareVolumeAllocation());
+		
+		TArray<FDistanceFieldVolumeTexture*> ThrottledAllocations;
+		TArray<FDistanceFieldVolumeTexture*>* LocalPendingAllocations = &PendingAllocations;
+		int ThrottledCopyCount = 0;
+		int PendingCopyCount = PendingAllocations.Num();
 
-		for (int32 AllocationIndex = 0; AllocationIndex < PendingAllocations.Num(); AllocationIndex++)
+		const int32 ThrottleSize = CVarDistFieldThrottleCopyToAtlasInBytes.GetValueOnAnyThread();
+		const bool bThrottleUpdateAllocation = ThrottleSize >= 1024;
+
+		auto AllocateBlocks = [&]()
 		{
-			FDistanceFieldVolumeTexture* Texture = PendingAllocations[AllocationIndex];
-			const FIntVector Size = Texture->VolumeData.Size;
+			for (int32 AllocationIndex = 0; AllocationIndex < LocalPendingAllocations->Num(); AllocationIndex++)
+			{
+				FDistanceFieldVolumeTexture* Texture = (*LocalPendingAllocations)[AllocationIndex];
+				const FIntVector Size = Texture->VolumeData.Size;
+				Texture->bThrottled = false;
+
+				if (!BlockAllocator.AddElement((uint32&)Texture->AtlasAllocationMin.X, (uint32&)Texture->AtlasAllocationMin.Y, (uint32&)Texture->AtlasAllocationMin.Z, Size.X, Size.Y, Size.Z))
+				{
+					UE_LOG(LogStaticMesh, Error, TEXT("Failed to allocate %ux%ux%u in distance field atlas"), Size.X, Size.Y, Size.Z);
+					LocalPendingAllocations->RemoveAt(AllocationIndex);
+					AllocationIndex--;
+				}
+			}
+		};
+
+		if (bThrottleUpdateAllocation)
+		{
+			int32 CurrentSize = 0;
 
-			if (!BlockAllocator.AddElement((uint32&)Texture->AtlasAllocationMin.X, (uint32&)Texture->AtlasAllocationMin.Y, (uint32&)Texture->AtlasAllocationMin.Z, Size.X, Size.Y, Size.Z))
+			for (int32 AllocationIndex = 0; AllocationIndex < PendingAllocations.Num() && CurrentSize < ThrottleSize; ++AllocationIndex)
 			{
-				UE_LOG(LogStaticMesh,Error,TEXT("Failed to allocate %ux%ux%u in distance field atlas"), Size.X, Size.Y, Size.Z);
+				FDistanceFieldVolumeTexture* Texture = PendingAllocations[AllocationIndex];
+				const FIntVector Size = Texture->VolumeData.Size;
+				CurrentSize += Size.X * Size.Y * Size.Z * FormatSize;
+				ThrottledAllocations.Add(Texture);
 				PendingAllocations.RemoveAt(AllocationIndex);
 				AllocationIndex--;
 			}
+
+			LocalPendingAllocations = &ThrottledAllocations;
+			ThrottledCopyCount = ThrottledAllocations.Num();
+			PendingCopyCount = PendingAllocations.Num();
 		}
 
+		AllocateBlocks();
+			
 		static const auto CVarCompress = IConsoleManager::Get().FindTConsoleVariableDataInt(TEXT("r.DistanceFieldBuild.Compress"));
 		const bool bDataIsCompressed = CVarCompress->GetValueOnAnyThread() != 0;
 
-		const int32 FormatSize = GPixelFormats[Format].BlockBytes;
-
 		if (!VolumeTextureRHI
 			|| BlockAllocator.GetSizeX() > VolumeTextureRHI->GetSizeX()
 			|| BlockAllocator.GetSizeY() > VolumeTextureRHI->GetSizeY()
@@ -315,27 +361,30 @@ void FDistanceFieldVolumeTextureAtlas::UpdateAllocations()
 
 				// Re-upload all textures since we had to reallocate
 				PendingAllocations.Append(CurrentAllocations);
+				if (bThrottleUpdateAllocation)
+				{
+					PendingAllocations.Append(ThrottledAllocations);
+					ThrottledAllocations.Empty();
+				}
 				CurrentAllocations.Empty();
 
 				// Sort largest to smallest for best packing
 				PendingAllocations.Sort(FCompareVolumeAllocation());
 
-				// Add all allocations back to the layout
-				for (int32 AllocationIndex = 0; AllocationIndex < PendingAllocations.Num(); AllocationIndex++)
+				if (bThrottleUpdateAllocation)
 				{
-					FDistanceFieldVolumeTexture* Texture = PendingAllocations[AllocationIndex];
-					const FIntVector Size = Texture->VolumeData.Size;
-
-					if (!BlockAllocator.AddElement((uint32&)Texture->AtlasAllocationMin.X, (uint32&)Texture->AtlasAllocationMin.Y, (uint32&)Texture->AtlasAllocationMin.Z, Size.X, Size.Y, Size.Z))
-					{
-						UE_LOG(LogStaticMesh,Error,TEXT("Failed to allocate %ux%ux%u in distance field atlas"), Size.X, Size.Y, Size.Z);
-						PendingAllocations.RemoveAt(AllocationIndex);
-						AllocationIndex--;
-					}
+					// Throttling during a full realloc when not using the max sime of volume texture will make the same blocks being reused over and over
+					// allocate everything pending to avoid this
+					LocalPendingAllocations = &PendingAllocations;
+					ThrottledCopyCount = ThrottledAllocations.Num();
+					PendingCopyCount = PendingAllocations.Num();
 				}
+
+				// Add all allocations back to the layout
+				AllocateBlocks();
 			}
 
-			// Fully free the previous atlas memory before allocating a new one
+			// Fully free the previous atlas memory before allocating a new one6
 			{
 				// Remove last ref, add to deferred delete list
 				VolumeTextureRHI = NULL;
@@ -377,9 +426,9 @@ void FDistanceFieldVolumeTextureAtlas::UpdateAllocations()
 
 				TArray<uint8> UncompressedData;
 
-				for (int32 AllocationIndex = 0; AllocationIndex < PendingAllocations.Num(); AllocationIndex++)
+				for (int32 AllocationIndex = 0; AllocationIndex < LocalPendingAllocations->Num(); AllocationIndex++)
 				{
-					FDistanceFieldVolumeTexture* Texture = PendingAllocations[AllocationIndex];
+					FDistanceFieldVolumeTexture* Texture = (*LocalPendingAllocations)[AllocationIndex];
 					const FIntVector Size = Texture->VolumeData.Size;
 					const int32 UncompressedSize = Size.X * Size.Y * Size.Z * FormatSize;
 
@@ -426,7 +475,7 @@ void FDistanceFieldVolumeTextureAtlas::UpdateAllocations()
 		}
 		else
 		{
-			const int32 NumUpdates = PendingAllocations.Num();
+			const int32 NumUpdates = LocalPendingAllocations->Num();
 			TArray<FUpdateTexture3DData> UpdateDataArray;
 			UpdateDataArray.Empty(NumUpdates);
 			UpdateDataArray.AddUninitialized(NumUpdates);
@@ -434,7 +483,7 @@ void FDistanceFieldVolumeTextureAtlas::UpdateAllocations()
 			// Allocate upload buffers
 			for (int32 Idx = 0; Idx < NumUpdates; ++Idx)
 			{
-				FDistanceFieldVolumeTexture* Texture = PendingAllocations[Idx];
+				FDistanceFieldVolumeTexture* Texture = (*LocalPendingAllocations)[Idx];
 				const FIntVector& Size = Texture->VolumeData.Size;
 				const FUpdateTextureRegion3D UpdateRegion(Texture->AtlasAllocationMin, FIntVector::ZeroValue, Size);
 
@@ -448,10 +497,10 @@ void FDistanceFieldVolumeTextureAtlas::UpdateAllocations()
 			// Copy data to upload buffers and decompress source data if necessary
 			ParallelFor(
 				NumUpdates,
-				[this, FormatSize, bDataIsCompressed, &UpdateDataArray](int32 Idx)
+				[this, FormatSize, bDataIsCompressed, &UpdateDataArray, LocalPendingAllocations](int32 Idx)
 			{
 				FUpdateTexture3DData& UpdateData = UpdateDataArray[Idx];
-				FDistanceFieldVolumeTexture* Texture = PendingAllocations[Idx];
+				FDistanceFieldVolumeTexture* Texture = (*LocalPendingAllocations)[Idx];
 				const FIntVector& Size = Texture->VolumeData.Size;
 
 				TArray<uint8> UncompressedData;
@@ -507,8 +556,8 @@ void FDistanceFieldVolumeTextureAtlas::UpdateAllocations()
 			RHIEndMultiUpdateTexture3D(UpdateDataArray);
 		}
 
-		CurrentAllocations.Append(PendingAllocations);
-		PendingAllocations.Empty();
+		CurrentAllocations.Append(*LocalPendingAllocations);
+		LocalPendingAllocations->Empty();
 
 		const double EndTime = FPlatformTime::Seconds();
 		const float UpdateDurationMs = (float)(EndTime - StartTime) * 1000.0f;
diff --git a/Engine/Source/Runtime/Engine/Private/InstancedStaticMesh.cpp b/Engine/Source/Runtime/Engine/Private/InstancedStaticMesh.cpp
--- a/Engine/Source/Runtime/Engine/Private/InstancedStaticMesh.cpp
+++ b/Engine/Source/Runtime/Engine/Private/InstancedStaticMesh.cpp
@@ -920,9 +920,9 @@ bool FInstancedStaticMeshSceneProxy::GetWireframeMeshElement(int32 LODIndex, int
 	return false;
 }
 
-void FInstancedStaticMeshSceneProxy::GetDistancefieldAtlasData(FBox& LocalVolumeBounds, FVector2D& OutDistanceMinMax, FIntVector& OutBlockMin, FIntVector& OutBlockSize, bool& bOutBuiltAsIfTwoSided, bool& bMeshWasPlane, float& SelfShadowBias, TArray<FMatrix>& ObjectLocalToWorldTransforms) const
+void FInstancedStaticMeshSceneProxy::GetDistancefieldAtlasData(FBox& LocalVolumeBounds, FVector2D& OutDistanceMinMax, FIntVector& OutBlockMin, FIntVector& OutBlockSize, bool& bOutBuiltAsIfTwoSided, bool& bMeshWasPlane, float& SelfShadowBias, TArray<FMatrix>& ObjectLocalToWorldTransforms, bool& bOutThrottled) const
 {
-	FStaticMeshSceneProxy::GetDistancefieldAtlasData(LocalVolumeBounds, OutDistanceMinMax, OutBlockMin, OutBlockSize, bOutBuiltAsIfTwoSided, bMeshWasPlane, SelfShadowBias, ObjectLocalToWorldTransforms);
+	FStaticMeshSceneProxy::GetDistancefieldAtlasData(LocalVolumeBounds, OutDistanceMinMax, OutBlockMin, OutBlockSize, bOutBuiltAsIfTwoSided, bMeshWasPlane, SelfShadowBias, ObjectLocalToWorldTransforms, bOutThrottled);
 
 	ObjectLocalToWorldTransforms.Reset();
 
diff --git a/Engine/Source/Runtime/Engine/Private/InstancedStaticMesh.h b/Engine/Source/Runtime/Engine/Private/InstancedStaticMesh.h
--- a/Engine/Source/Runtime/Engine/Private/InstancedStaticMesh.h
+++ b/Engine/Source/Runtime/Engine/Private/InstancedStaticMesh.h
@@ -584,7 +584,7 @@ public:
 	/** Sets up a wireframe FMeshBatch for a specific LOD. */
 	virtual bool GetWireframeMeshElement(int32 LODIndex, int32 BatchIndex, const FMaterialRenderProxy* WireframeRenderProxy, uint8 InDepthPriorityGroup, bool bAllowPreCulledIndices, FMeshBatch& OutMeshBatch) const override;
 
-	virtual void GetDistancefieldAtlasData(FBox& LocalVolumeBounds, FVector2D& OutDistanceMinMax, FIntVector& OutBlockMin, FIntVector& OutBlockSize, bool& bOutBuiltAsIfTwoSided, bool& bMeshWasPlane, float& SelfShadowBias, TArray<FMatrix>& ObjectLocalToWorldTransforms) const override;
+	virtual void GetDistancefieldAtlasData(FBox& LocalVolumeBounds, FVector2D& OutDistanceMinMax, FIntVector& OutBlockMin, FIntVector& OutBlockSize, bool& bOutBuiltAsIfTwoSided, bool& bMeshWasPlane, float& SelfShadowBias, TArray<FMatrix>& ObjectLocalToWorldTransforms, bool& bOutThrottled) const override;
 
 	virtual void GetDistanceFieldInstanceInfo(int32& NumInstances, float& BoundsSurfaceArea) const override;
 
diff --git a/Engine/Source/Runtime/Engine/Private/StaticMeshRender.cpp b/Engine/Source/Runtime/Engine/Private/StaticMeshRender.cpp
--- a/Engine/Source/Runtime/Engine/Private/StaticMeshRender.cpp
+++ b/Engine/Source/Runtime/Engine/Private/StaticMeshRender.cpp
@@ -1470,7 +1470,7 @@ void FStaticMeshSceneProxy::GetLightRelevance(const FLightSceneProxy* LightScene
 	}
 }
 
-void FStaticMeshSceneProxy::GetDistancefieldAtlasData(FBox& LocalVolumeBounds, FVector2D& OutDistanceMinMax, FIntVector& OutBlockMin, FIntVector& OutBlockSize, bool& bOutBuiltAsIfTwoSided, bool& bMeshWasPlane, float& SelfShadowBias, TArray<FMatrix>& ObjectLocalToWorldTransforms) const
+void FStaticMeshSceneProxy::GetDistancefieldAtlasData(FBox& LocalVolumeBounds, FVector2D& OutDistanceMinMax, FIntVector& OutBlockMin, FIntVector& OutBlockSize, bool& bOutBuiltAsIfTwoSided, bool& bMeshWasPlane, float& SelfShadowBias, TArray<FMatrix>& ObjectLocalToWorldTransforms, bool& bOutThrottled) const
 {
 	if (DistanceFieldData)
 	{
@@ -1482,6 +1482,7 @@ void FStaticMeshSceneProxy::GetDistancefieldAtlasData(FBox& LocalVolumeBounds, F
 		bMeshWasPlane = DistanceFieldData->bMeshWasPlane;
 		ObjectLocalToWorldTransforms.Add(GetLocalToWorld());
 		SelfShadowBias = DistanceFieldSelfShadowBias;
+		bOutThrottled = DistanceFieldData->VolumeTexture.Throttled();
 	}
 	else
 	{
@@ -1492,6 +1493,7 @@ void FStaticMeshSceneProxy::GetDistancefieldAtlasData(FBox& LocalVolumeBounds, F
 		bOutBuiltAsIfTwoSided = false;
 		bMeshWasPlane = false;
 		SelfShadowBias = 0;
+		bOutThrottled = false;
 	}
 }
 
diff --git a/Engine/Source/Runtime/Engine/Public/DistanceFieldAtlas.h b/Engine/Source/Runtime/Engine/Public/DistanceFieldAtlas.h
--- a/Engine/Source/Runtime/Engine/Public/DistanceFieldAtlas.h
+++ b/Engine/Source/Runtime/Engine/Public/DistanceFieldAtlas.h
@@ -30,6 +30,7 @@ public:
 		VolumeData(InVolumeData),
 		AtlasAllocationMin(FIntVector(-1, -1, -1)),
 		bReferencedByAtlas(false),
+		bThrottled(false),
 		StaticMesh(NULL)
 	{}
 
@@ -55,6 +56,11 @@ public:
 
 	bool IsValidDistanceFieldVolume() const;
 
+	bool Throttled() const
+	{
+		return bThrottled;
+	}
+
 	UStaticMesh* GetStaticMesh() const
 	{
 		return StaticMesh;
@@ -63,7 +69,9 @@ public:
 private:
 	const FDistanceFieldVolumeData& VolumeData;
 	FIntVector AtlasAllocationMin;
-	bool bReferencedByAtlas;
+	bool bReferencedByAtlas : 1;
+	/** bThrottled prevents any objects using the texture from being uploaded to the scene buffer until upload of the texture to distance field atlas is complete */
+	bool bThrottled         : 1;
 	UStaticMesh* StaticMesh;
 
 	friend class FDistanceFieldVolumeTextureAtlas;
diff --git a/Engine/Source/Runtime/Engine/Public/PrimitiveSceneProxy.h b/Engine/Source/Runtime/Engine/Public/PrimitiveSceneProxy.h
--- a/Engine/Source/Runtime/Engine/Public/PrimitiveSceneProxy.h
+++ b/Engine/Source/Runtime/Engine/Public/PrimitiveSceneProxy.h
@@ -262,7 +262,7 @@ public:
 		bShadowMapped = false;
 	}
 
-	virtual void GetDistancefieldAtlasData(FBox& LocalVolumeBounds, FVector2D& OutDistanceMinMax, FIntVector& OutBlockMin, FIntVector& OutBlockSize, bool& bOutBuiltAsIfTwoSided, bool& bMeshWasPlane, float& SelfShadowBias, TArray<FMatrix>& ObjectLocalToWorldTransforms) const 
+	virtual void GetDistancefieldAtlasData(FBox& LocalVolumeBounds, FVector2D& OutDistanceMinMax, FIntVector& OutBlockMin, FIntVector& OutBlockSize, bool& bOutBuiltAsIfTwoSided, bool& bMeshWasPlane, float& SelfShadowBias, TArray<FMatrix>& ObjectLocalToWorldTransforms, bool& bOutThrottled) const
 	{
 		LocalVolumeBounds = FBox(ForceInit);
 		OutDistanceMinMax = FVector2D(0, 0);
@@ -271,6 +271,7 @@ public:
 		bOutBuiltAsIfTwoSided = false;
 		bMeshWasPlane = false;
 		SelfShadowBias = 0;
+		bOutThrottled = false;
 	}
 
 	virtual void GetDistanceFieldInstanceInfo(int32& NumInstances, float& BoundsSurfaceArea) const
diff --git a/Engine/Source/Runtime/Engine/Public/StaticMeshResources.h b/Engine/Source/Runtime/Engine/Public/StaticMeshResources.h
--- a/Engine/Source/Runtime/Engine/Public/StaticMeshResources.h
+++ b/Engine/Source/Runtime/Engine/Public/StaticMeshResources.h
@@ -798,7 +798,7 @@ public:
 	virtual bool CanBeOccluded() const override;
 	virtual bool IsUsingDistanceCullFade() const override;
 	virtual void GetLightRelevance(const FLightSceneProxy* LightSceneProxy, bool& bDynamic, bool& bRelevant, bool& bLightMapped, bool& bShadowMapped) const override;
-	virtual void GetDistancefieldAtlasData(FBox& LocalVolumeBounds, FVector2D& OutDistanceMinMax, FIntVector& OutBlockMin, FIntVector& OutBlockSize, bool& bOutBuiltAsIfTwoSided, bool& bMeshWasPlane, float& SelfShadowBias, TArray<FMatrix>& ObjectLocalToWorldTransforms) const override;
+	virtual void GetDistancefieldAtlasData(FBox& LocalVolumeBounds, FVector2D& OutDistanceMinMax, FIntVector& OutBlockMin, FIntVector& OutBlockSize, bool& bOutBuiltAsIfTwoSided, bool& bMeshWasPlane, float& SelfShadowBias, TArray<FMatrix>& ObjectLocalToWorldTransforms, bool& bOutThrottled) const override;
 	virtual void GetDistanceFieldInstanceInfo(int32& NumInstances, float& BoundsSurfaceArea) const override;
 	virtual bool HasDistanceFieldRepresentation() const override;
 	virtual bool HasDynamicIndirectShadowCasterRepresentation() const override;
diff --git a/Engine/Source/Runtime/Renderer/Private/DistanceFieldObjectManagement.cpp b/Engine/Source/Runtime/Renderer/Private/DistanceFieldObjectManagement.cpp
--- a/Engine/Source/Runtime/Renderer/Private/DistanceFieldObjectManagement.cpp
+++ b/Engine/Source/Runtime/Renderer/Private/DistanceFieldObjectManagement.cpp
@@ -775,7 +775,7 @@ void UpdateGlobalDistanceFieldObjectRemoves(FRHICommandListImmediate& RHICmdList
 }
 
 /** Gathers the information needed to represent a single object's distance field and appends it to the upload buffers. */
-void ProcessPrimitiveUpdate(
+bool ProcessPrimitiveUpdate(
 	bool bIsAddOperation,
 	FRHICommandListImmediate& RHICmdList, 
 	FSceneRenderer& SceneRenderer, 
@@ -800,7 +800,13 @@ void ProcessPrimitiveUpdate(
 	bool bBuiltAsIfTwoSided;
 	bool bMeshWasPlane;
 	float SelfShadowBias;
-	PrimitiveSceneInfo->Proxy->GetDistancefieldAtlasData(LocalVolumeBounds, DistanceMinMax, BlockMin, BlockSize, bBuiltAsIfTwoSided, bMeshWasPlane, SelfShadowBias, ObjectLocalToWorldTransforms);
+	bool bThrottled;
+	PrimitiveSceneInfo->Proxy->GetDistancefieldAtlasData(LocalVolumeBounds, DistanceMinMax, BlockMin, BlockSize, bBuiltAsIfTwoSided, bMeshWasPlane, SelfShadowBias, ObjectLocalToWorldTransforms, bThrottled);
+
+	if (bThrottled)
+	{
+		return false;
+	}
 
 	if (BlockMin.X >= 0 
 		&& BlockMin.Y >= 0 
@@ -1003,6 +1009,7 @@ void ProcessPrimitiveUpdate(
 			UE_LOG(LogDistanceField,Log,TEXT("Primitive %s %s excluded due to bounding radius %f"), *PrimitiveSceneInfo->Proxy->GetOwnerName().ToString(), *PrimitiveSceneInfo->Proxy->GetResourceName().ToString(), BoundingRadius);
 		}
 	}
+	return true;
 }
 
 bool bVerifySceneIntegrity = false;
@@ -1012,7 +1019,9 @@ void FDeferredShadingSceneRenderer::UpdateGlobalDistanceFieldObjectBuffers(FRHIC
 	FDistanceFieldSceneData& DistanceFieldSceneData = Scene->DistanceFieldSceneData;
 
 	if (GDistanceFieldVolumeTextureAtlas.VolumeTextureRHI
-		&& (DistanceFieldSceneData.HasPendingOperations() || DistanceFieldSceneData.AtlasGeneration != GDistanceFieldVolumeTextureAtlas.GetGeneration()))
+		&& (DistanceFieldSceneData.HasPendingOperations() || 
+			DistanceFieldSceneData.PendingThrottledOperations.Num()>0 || 
+			DistanceFieldSceneData.AtlasGeneration != GDistanceFieldVolumeTextureAtlas.GetGeneration()))
 	{
 		QUICK_SCOPE_CYCLE_COUNTER(STAT_UpdateObjectData);
 		SCOPED_DRAW_EVENT(RHICmdList, UpdateSceneObjectData);
@@ -1032,6 +1041,14 @@ void FDeferredShadingSceneRenderer::UpdateGlobalDistanceFieldObjectBuffers(FRHIC
 			DistanceFieldSceneData.InstancedSurfelBuffers = new FInstancedSurfelBuffers();
 		}
 
+		if (DistanceFieldSceneData.PendingAddOperations.Num()>0)
+		{
+			DistanceFieldSceneData.PendingThrottledOperations.Reserve(DistanceFieldSceneData.PendingThrottledOperations.Num() + DistanceFieldSceneData.PendingAddOperations.Num());
+		}
+
+		DistanceFieldSceneData.PendingAddOperations.Append(DistanceFieldSceneData.PendingThrottledOperations);
+		DistanceFieldSceneData.PendingThrottledOperations.Reset();
+
 		if (DistanceFieldSceneData.AtlasGeneration != GDistanceFieldVolumeTextureAtlas.GetGeneration())
 		{
 			DistanceFieldSceneData.AtlasGeneration = GDistanceFieldVolumeTextureAtlas.GetGeneration();
@@ -1182,18 +1199,21 @@ void FDeferredShadingSceneRenderer::UpdateGlobalDistanceFieldObjectBuffers(FRHIC
 			{
 				FPrimitiveSceneInfo* PrimitiveSceneInfo = DistanceFieldSceneData.PendingAddOperations[UploadPrimitiveIndex];
 
-				ProcessPrimitiveUpdate(
+				if (!ProcessPrimitiveUpdate(
 					true,
-					RHICmdList, 
-					*this, 
-					PrimitiveSceneInfo, 
-					OriginalNumObjects, 
-					InvTextureDim, 
-					bPrepareForDistanceFieldGI, 
+					RHICmdList,
+					*this,
+					PrimitiveSceneInfo,
+					OriginalNumObjects,
+					InvTextureDim,
+					bPrepareForDistanceFieldGI,
 					bAnyViewEnabledDistanceCulling,
-					ObjectLocalToWorldTransforms, 
-					UploadObjectIndices, 
-					UploadObjectData);
+					ObjectLocalToWorldTransforms,
+					UploadObjectIndices,
+					UploadObjectData))
+				{
+					DistanceFieldSceneData.PendingThrottledOperations.Add(PrimitiveSceneInfo);
+				}
 			}
 
 			for (TSet<FPrimitiveSceneInfo*>::TIterator It(DistanceFieldSceneData.PendingUpdateOperations); It; ++It)
@@ -1202,20 +1222,24 @@ void FDeferredShadingSceneRenderer::UpdateGlobalDistanceFieldObjectBuffers(FRHIC
 
 				ProcessPrimitiveUpdate(
 					false,
-					RHICmdList, 
-					*this, 
-					PrimitiveSceneInfo, 
-					OriginalNumObjects, 
-					InvTextureDim, 
-					bPrepareForDistanceFieldGI, 
+					RHICmdList,
+					*this,
+					PrimitiveSceneInfo,
+					OriginalNumObjects,
+					InvTextureDim,
+					bPrepareForDistanceFieldGI,
 					bAnyViewEnabledDistanceCulling,
-					ObjectLocalToWorldTransforms, 
-					UploadObjectIndices, 
+					ObjectLocalToWorldTransforms,
+					UploadObjectIndices,
 					UploadObjectData);
 			}
 
 			DistanceFieldSceneData.PendingAddOperations.Reset();
 			DistanceFieldSceneData.PendingUpdateOperations.Empty();
+			if (DistanceFieldSceneData.PendingThrottledOperations.Num() == 0)
+			{
+				DistanceFieldSceneData.PendingThrottledOperations.Empty();
+			}
 
 			if (DistanceFieldSceneData.ObjectBuffers->MaxObjects < DistanceFieldSceneData.NumObjectsInBuffer)
 			{
diff --git a/Engine/Source/Runtime/Renderer/Private/RendererScene.cpp b/Engine/Source/Runtime/Renderer/Private/RendererScene.cpp
--- a/Engine/Source/Runtime/Renderer/Private/RendererScene.cpp
+++ b/Engine/Source/Runtime/Renderer/Private/RendererScene.cpp
@@ -431,6 +431,7 @@ void FDistanceFieldSceneData::RemovePrimitive(FPrimitiveSceneInfo* InPrimitive)
 		{
 			PendingAddOperations.Remove(InPrimitive);
 			PendingUpdateOperations.Remove(InPrimitive);
+			PendingThrottledOperations.Remove(InPrimitive);
 
 			if (InPrimitive->DistanceFieldInstanceIndices.Num() > 0)
 			{
diff --git a/Engine/Source/Runtime/Renderer/Private/ScenePrivate.h b/Engine/Source/Runtime/Renderer/Private/ScenePrivate.h
--- a/Engine/Source/Runtime/Renderer/Private/ScenePrivate.h
+++ b/Engine/Source/Runtime/Renderer/Private/ScenePrivate.h
@@ -1882,6 +1882,7 @@ public:
 
 	/** Pending operations on the object buffers to be processed next frame. */
 	TArray<FPrimitiveSceneInfo*> PendingAddOperations;
+	TArray<FPrimitiveSceneInfo*> PendingThrottledOperations;
 	TSet<FPrimitiveSceneInfo*> PendingUpdateOperations;
 	TArray<FPrimitiveRemoveInfo> PendingRemoveOperations;
 	TArray<FVector4> PrimitiveModifiedBounds[GDF_Num];
