Turning a check into an ensure as it has been encountered in the field and the code is robust to it. The referenced issues in the comments associated are long since lost and predate Jira.

diff --git a/Engine/Source/Runtime/Engine/Private/Particles/ParticleSystemRender.cpp b/Engine/Source/Runtime/Engine/Private/Particles/ParticleSystemRender.cpp
--- a/Engine/Source/Runtime/Engine/Private/Particles/ParticleSystemRender.cpp
+++ b/Engine/Source/Runtime/Engine/Private/Particles/ParticleSystemRender.cpp
@@ -5496,8 +5496,8 @@ void FDynamicTrailsEmitterData::Init(bool bInSelected)
 {
 	bSelected = bInSelected;
 
-	check(SourcePointer->ActiveParticleCount < (16 * 1024));	// TTP #33330
-	check(SourcePointer->ParticleStride < (2 * 1024));			// TTP #33330
+	ensure(SourcePointer->ActiveParticleCount < (16 * 1024));	// TTP #33330
+	ensure(SourcePointer->ParticleStride < (2 * 1024));			// TTP #33330
 
 	MaterialResource = SourcePointer->MaterialInterface->GetRenderProxy();
 
