Fix compile error of TMaybeMappedAllocator

diff --git a/Engine/Source/Runtime/Engine/Public/Animation/AnimCompressionTypes.h b/Engine/Source/Runtime/Engine/Public/Animation/AnimCompressionTypes.h
--- a/Engine/Source/Runtime/Engine/Public/Animation/AnimCompressionTypes.h
+++ b/Engine/Source/Runtime/Engine/Public/Animation/AnimCompressionTypes.h
@@ -486,10 +486,10 @@ public:
 			Other.Data = nullptr;
 
 			MappedRegion = Other.MappedRegion;
-			Other->MappedRegion = nullptr;
+			Other.MappedRegion = nullptr;
 
 			MappedHandle = Other.MappedHandle;
-			Other->MappedHandle = nullptr;
+			Other.MappedHandle = nullptr;
 		}
 
 		/** Destructor. */
@@ -629,6 +629,12 @@ public:
 	}
 };
 
+template <uint32 Alignment>
+struct TAllocatorTraits<TMaybeMappedAllocator<Alignment>> : TAllocatorTraitsBase<TMaybeMappedAllocator<Alignment>>
+{
+	enum { SupportsMove = true };
+};
+
 struct ENGINE_API FCompressedAnimSequence
 {
 public:
