Disabled  buffer visualization and view mode cycling in DebugCameraController in cooked and/or shipped builds.
Fixed handling of invalid console var values.
Excluded unsupported view modes from view mode cycling settings.

diff --git a/Engine/Source/Editor/DetailCustomizations/Private/DebugCameraControllerSettingsCustomization.cpp b/Engine/Source/Editor/DetailCustomizations/Private/DebugCameraControllerSettingsCustomization.cpp
--- a/Engine/Source/Editor/DetailCustomizations/Private/DebugCameraControllerSettingsCustomization.cpp
+++ b/Engine/Source/Editor/DetailCustomizations/Private/DebugCameraControllerSettingsCustomization.cpp
@@ -10,6 +10,7 @@
 #include "DetailWidgetRow.h"
 #include "IDetailPropertyRow.h"
 #include "ShowFlags.h"
+#include "RHI.h"
 
 
 TSharedRef<IPropertyTypeCustomization> FDebugCameraControllerSettingsViewModeIndexCustomization::MakeInstance()
@@ -47,6 +48,14 @@ void FDebugCameraControllerSettingsViewModeIndexCustomization::CustomizeHeader(T
 	TSharedPtr<FPropertyRestriction> EnumRestriction = MakeShareable(new FPropertyRestriction(NSLOCTEXT("DebugCycleViewModes", "DebugCycleViewModes", "Cycle view modes for debug camera controller")));
 	const UEnum* const ViewModeIndexEnum = StaticEnum<EViewModeIndex>();
 	EnumRestriction->AddHiddenValue(ViewModeIndexEnum->GetNameStringByValue((uint8)EViewModeIndex::VMI_VisualizeBuffer));
+	EnumRestriction->AddHiddenValue(ViewModeIndexEnum->GetNameStringByValue((uint8)EViewModeIndex::VMI_StationaryLightOverlap));
+#if RHI_RAYTRACING
+	if (!GRHISupportsRayTracing)
+	{
+		EnumRestriction->AddHiddenValue(ViewModeIndexEnum->GetNameStringByValue((uint8)EViewModeIndex::VMI_PathTracing));
+		EnumRestriction->AddHiddenValue(ViewModeIndexEnum->GetNameStringByValue((uint8)EViewModeIndex::VMI_RayTracingDebug));
+	}
+#endif
 	ViewModeIndexHandle->AddRestriction(EnumRestriction.ToSharedRef());
 
 	HeaderRow
diff --git a/Engine/Source/Runtime/Engine/Classes/Engine/DebugCameraController.h b/Engine/Source/Runtime/Engine/Classes/Engine/DebugCameraController.h
--- a/Engine/Source/Runtime/Engine/Classes/Engine/DebugCameraController.h
+++ b/Engine/Source/Runtime/Engine/Classes/Engine/DebugCameraController.h
@@ -126,6 +126,14 @@ class ENGINE_API ADebugCameraController
 	/** Toggles camera orbit hitpoint */
 	void ToggleOrbitHitPoint();
 
+#if !(UE_BUILD_SHIPPING || UE_BUILD_TEST)
+
+	/** Whether buffer visualization option should be enabled */
+	static bool EnableDebugBuffers();
+
+	/** Whether cycle viewmode option should be enabled */
+	static bool EnableDebugViewmodes();
+
 	/** Cycle view mode */
 	void CycleViewMode();
 
@@ -153,15 +161,17 @@ class ENGINE_API ADebugCameraController
 	/** Set buffer visualization full mode */
 	void SetBufferVisualizationFullMode(bool bFullMode);
 
-#if !(UE_BUILD_SHIPPING || UE_BUILD_TEST)
-
 	/** Update visualize buffer post processing settings */
 	void UpdateVisualizeBufferPostProcessing(FFinalPostProcessSettings& InOutPostProcessingSettings);
 
-#endif
+	/** Get visualization buffer's material name used by post processing settings */
+	FString GetBufferMaterialName(const FString& InBuffer);
+
+	/** Get current visualization buffer's material name */
+	FString GetSelectedBufferMaterialName();
 
-	/** Get selected visualization buffer */
-	FString GetSelectedBufferName();
+
+#endif
 
 public:
 
@@ -280,6 +290,8 @@ protected:
 
 	virtual void SetSpectatorPawn(class ASpectatorPawn* NewSpectatorPawn) override;
 
+#if !(UE_BUILD_SHIPPING || UE_BUILD_TEST)
+
 	/** Get buffer visualization overview targets based on console var */
 	TArray<FString> GetBufferVisualizationOverviewTargets();
 
@@ -289,6 +301,8 @@ protected:
 	/** Get next buffer */
 	void GetNextBuffer(int32 Step = 1);
 
+#endif
+
 private:
 
 	/** The normalized screen location when a drag starts */
@@ -303,9 +317,6 @@ private:
 	/** Current orbit radius, if orbit is enabled */
 	float OrbitRadius;
 
-	/** Last view mode index before buffer visualization overview was enabled, see enum EViewModeIndex for valid values */
-	int32 LastViewModeIndex;
-
 	/** Last index in settings array for cycle view modes */
 	int32 LastViewModeSettingsIndex;
 
diff --git a/Engine/Source/Runtime/Engine/Private/DebugCameraController.cpp b/Engine/Source/Runtime/Engine/Private/DebugCameraController.cpp
--- a/Engine/Source/Runtime/Engine/Private/DebugCameraController.cpp
+++ b/Engine/Source/Runtime/Engine/Private/DebugCameraController.cpp
@@ -61,7 +61,6 @@ ADebugCameraController::ADebugCameraController(const FObjectInitializer& ObjectI
 	bEnableBufferVisualizationFullMode = false;
 	bIsBufferVisualizationInputSetup = false;
 	bLastDisplayEnabled = true;
-	LastViewModeIndex = VMI_Lit;
 	LastViewModeSettingsIndex = 0;
 }
 
@@ -84,9 +83,6 @@ void InitializeDebugCameraInputBindings()
 		UPlayerInput::AddEngineDefinedActionMapping(FInputActionKeyMapping("DebugCamera_FreezeRendering", EKeys::F));
 		UPlayerInput::AddEngineDefinedActionMapping(FInputActionKeyMapping("DebugCamera_OrbitHitPoint", EKeys::O));
 		UPlayerInput::AddEngineDefinedActionMapping(FInputActionKeyMapping("DebugCamera_OrbitCenter", EKeys::O, true));
-		UPlayerInput::AddEngineDefinedActionMapping(FInputActionKeyMapping("DebugCamera_CycleViewMode", EKeys::V));
-		UPlayerInput::AddEngineDefinedActionMapping(FInputActionKeyMapping("DebugCamera_ToggleBufferVisualizationOverview", EKeys::B));
-		UPlayerInput::AddEngineDefinedActionMapping(FInputActionKeyMapping("DebugCamera_ToggleBufferVisualizationFull", EKeys::Enter));
 
 		UPlayerInput::AddEngineDefinedActionMapping(FInputActionKeyMapping("DebugCamera_Select", EKeys::Gamepad_RightTrigger));
 		UPlayerInput::AddEngineDefinedActionMapping(FInputActionKeyMapping("DebugCamera_IncreaseSpeed", EKeys::Gamepad_RightShoulder));
@@ -96,16 +92,30 @@ void InitializeDebugCameraInputBindings()
 		UPlayerInput::AddEngineDefinedActionMapping(FInputActionKeyMapping("DebugCamera_ToggleDisplay", EKeys::Gamepad_FaceButton_Left));
 		UPlayerInput::AddEngineDefinedActionMapping(FInputActionKeyMapping("DebugCamera_FreezeRendering", EKeys::Gamepad_FaceButton_Top));
 
-		UPlayerInput::AddEngineDefinedActionMapping(FInputActionKeyMapping("DebugCamera_BufferVisualizationUp", EKeys::Up));
-		UPlayerInput::AddEngineDefinedActionMapping(FInputActionKeyMapping("DebugCamera_BufferVisualizationDown", EKeys::Down));
-		UPlayerInput::AddEngineDefinedActionMapping(FInputActionKeyMapping("DebugCamera_BufferVisualizationLeft", EKeys::Left));
-		UPlayerInput::AddEngineDefinedActionMapping(FInputActionKeyMapping("DebugCamera_BufferVisualizationRight", EKeys::Right));
+#if !(UE_BUILD_SHIPPING || UE_BUILD_TEST)
 
-		// The following axis mappings must be defined to override ADefaultPawn axis mappings when buffer visualization is enabled
-		UPlayerInput::AddEngineDefinedAxisMapping(FInputAxisKeyMapping("DebugCamera_DisableAxisMotion", EKeys::Up, 1.f));
-		UPlayerInput::AddEngineDefinedAxisMapping(FInputAxisKeyMapping("DebugCamera_DisableAxisMotion", EKeys::Down, -1.f));
-		UPlayerInput::AddEngineDefinedAxisMapping(FInputAxisKeyMapping("DebugCamera_DisableAxisMotion", EKeys::Left, -1.f));
-		UPlayerInput::AddEngineDefinedAxisMapping(FInputAxisKeyMapping("DebugCamera_DisableAxisMotion", EKeys::Right, 1.f));
+		if (ADebugCameraController::EnableDebugViewmodes())
+		{
+			UPlayerInput::AddEngineDefinedActionMapping(FInputActionKeyMapping("DebugCamera_CycleViewMode", EKeys::V));
+		}
+
+		if (ADebugCameraController::EnableDebugBuffers())
+		{
+			UPlayerInput::AddEngineDefinedActionMapping(FInputActionKeyMapping("DebugCamera_ToggleBufferVisualizationOverview", EKeys::B));
+			UPlayerInput::AddEngineDefinedActionMapping(FInputActionKeyMapping("DebugCamera_ToggleBufferVisualizationFull", EKeys::Enter));
+			UPlayerInput::AddEngineDefinedActionMapping(FInputActionKeyMapping("DebugCamera_BufferVisualizationUp", EKeys::Up));
+			UPlayerInput::AddEngineDefinedActionMapping(FInputActionKeyMapping("DebugCamera_BufferVisualizationDown", EKeys::Down));
+			UPlayerInput::AddEngineDefinedActionMapping(FInputActionKeyMapping("DebugCamera_BufferVisualizationLeft", EKeys::Left));
+			UPlayerInput::AddEngineDefinedActionMapping(FInputActionKeyMapping("DebugCamera_BufferVisualizationRight", EKeys::Right));
+
+			// The following axis mappings must be defined to override ADefaultPawn axis mappings when buffer visualization is enabled
+			UPlayerInput::AddEngineDefinedAxisMapping(FInputAxisKeyMapping("DebugCamera_DisableAxisMotion", EKeys::Up, 1.f));
+			UPlayerInput::AddEngineDefinedAxisMapping(FInputAxisKeyMapping("DebugCamera_DisableAxisMotion", EKeys::Down, -1.f));
+			UPlayerInput::AddEngineDefinedAxisMapping(FInputAxisKeyMapping("DebugCamera_DisableAxisMotion", EKeys::Left, -1.f));
+			UPlayerInput::AddEngineDefinedAxisMapping(FInputAxisKeyMapping("DebugCamera_DisableAxisMotion", EKeys::Right, 1.f));
+		}
+
+#endif
 	}
 }
 
@@ -127,15 +137,29 @@ void ADebugCameraController::SetupInputComponent()
 	InputComponent->BindAction("DebugCamera_FreezeRendering", IE_Pressed, this, &ADebugCameraController::ToggleFreezeRendering);
 	InputComponent->BindAction("DebugCamera_OrbitHitPoint", IE_Pressed, this, &ADebugCameraController::ToggleOrbitHitPoint);
 	InputComponent->BindAction("DebugCamera_OrbitCenter", IE_Pressed, this, &ADebugCameraController::ToggleOrbitCenter);
-	InputComponent->BindAction("DebugCamera_CycleViewMode", IE_Pressed, this, &ADebugCameraController::CycleViewMode);
-	InputComponent->BindAction("DebugCamera_ToggleBufferVisualizationOverview", IE_Pressed, this, &ADebugCameraController::ToggleBufferVisualizationOverviewMode);
-	InputComponent->BindAction("DebugCamera_ToggleBufferVisualizationFull", IE_Pressed, this, &ADebugCameraController::ToggleBufferVisualizationFullMode);
+
+#if !(UE_BUILD_SHIPPING || UE_BUILD_TEST)
+
+	if (ADebugCameraController::EnableDebugViewmodes())
+	{
+		InputComponent->BindAction("DebugCamera_CycleViewMode", IE_Pressed, this, &ADebugCameraController::CycleViewMode);
+	}
+
+	if (ADebugCameraController::EnableDebugBuffers())
+	{
+		InputComponent->BindAction("DebugCamera_ToggleBufferVisualizationOverview", IE_Pressed, this, &ADebugCameraController::ToggleBufferVisualizationOverviewMode);
+		InputComponent->BindAction("DebugCamera_ToggleBufferVisualizationFull", IE_Pressed, this, &ADebugCameraController::ToggleBufferVisualizationFullMode);
+	}
+
+#endif
 
 	InputComponent->BindTouch(IE_Pressed, this, &ADebugCameraController::OnTouchBegin);
 	InputComponent->BindTouch(IE_Released, this, &ADebugCameraController::OnTouchEnd);
 	InputComponent->BindTouch(IE_Repeat, this, &ADebugCameraController::OnFingerMove);
 }
 
+#if !(UE_BUILD_SHIPPING || UE_BUILD_TEST)
+
 void ADebugCameraController::SetupBufferVisualizationOverviewInput()
 {
 	if (InputComponent)
@@ -187,6 +211,8 @@ void ADebugCameraController::SetupBufferVisualizationOverviewInput()
 	}
 }
 
+#endif 
+
 void ADebugCameraController::OnTouchBegin(ETouchIndex::Type FingerIndex, FVector Location)
 {
 	if (FingerIndex == ETouchIndex::Touch1)
@@ -448,11 +474,15 @@ void ADebugCameraController::OnDeactivate( APlayerController* RestoredPC )
 		bIsFrozenRendering = false;
 	}
 
+#if !(UE_BUILD_SHIPPING || UE_BUILD_TEST)
+
 	if (bEnableBufferVisualization)
 	{
 		ToggleBufferVisualizationOverviewMode();
 	}
 
+#endif
+
 	bIsOrbitingSelectedActor = false;
 
 	DrawFrustum->SetVisibility(false);
@@ -675,6 +705,18 @@ void ADebugCameraController::ToggleOrbitHitPoint()
 	ToggleOrbit(false);
 }
 
+#if !(UE_BUILD_SHIPPING || UE_BUILD_TEST)
+
+bool ADebugCameraController::EnableDebugViewmodes()
+{
+	return AllowDebugViewmodes();
+}
+
+bool ADebugCameraController::EnableDebugBuffers()
+{
+	return AllowDebugViewmodes();
+}
+
 void ADebugCameraController::CycleViewMode()
 {
 	if (bEnableBufferVisualization)
@@ -688,14 +730,14 @@ void ADebugCameraController::CycleViewMode()
 
 		if (DebugViewModes.Num() == 0)
 		{
-			UE_LOG(LogController, Warning, TEXT("Debug camera controller settings must specify at least one view mode for view mode cycling."));
+			UE_LOG(LogPlayerController, Warning, TEXT("Debug camera controller settings must specify at least one view mode for view mode cycling."));
 			return;
 		}
 
 		int32 CurrViewModeIndex = GameViewportClient->ViewModeIndex;
 		int32 CurrIndex = LastViewModeSettingsIndex;
 
-		int32 NextIndex = (CurrIndex < DebugViewModes.Num() && CurrViewModeIndex == DebugViewModes[CurrIndex]) ? (CurrIndex + 1) % DebugViewModes.Num() : 0;
+		int32 NextIndex = CurrIndex < DebugViewModes.Num() ? (CurrIndex + 1) % DebugViewModes.Num() : 0;
 		int32 NextViewModeIndex = DebugViewModes[NextIndex];
 
 		if (NextViewModeIndex != CurrViewModeIndex)
@@ -708,10 +750,13 @@ void ADebugCameraController::CycleViewMode()
 				Cmd += NextViewModeName;
 				GameViewportClient->ConsoleCommand(Cmd);
 			}
+			else
+			{
+				UE_LOG(LogPlayerController, Warning, TEXT("Invalid view mode index %d."), NextViewModeIndex);
+			}
 		}
 
 		LastViewModeSettingsIndex = NextIndex;
-		LastViewModeIndex = NextViewModeIndex;
 	}
 }
 
@@ -721,25 +766,38 @@ TArray<FString> ADebugCameraController::GetBufferVisualizationOverviewTargets()
 
 	// Get the list of requested buffers from the console
 	static IConsoleVariable* CVar = IConsoleManager::Get().FindConsoleVariable(TEXT("r.BufferVisualizationOverviewTargets"));
-	FString SelectedBufferNames = CVar->GetString();
-
-	// Extract each material name from the comma separated string
-	while (SelectedBufferNames.Len())
+	if (CVar)
 	{
-		FString Left, Right;
+		FString SelectedBufferNames = CVar->GetString();
 
-		// Detect last entry in the list
-		if (!SelectedBufferNames.Split(TEXT(","), &Left, &Right))
+		// Extract each material name from the comma separated string
+		while (SelectedBufferNames.Len() && SelectedBuffers.Num() < 16)
 		{
-			Left = SelectedBufferNames;
-			Right = FString();
-		}
+			FString Left, Right;
 
-		Left.TrimStartInline();
-		SelectedBuffers.Add(*Left);
-		SelectedBufferNames = Right;
-	}
+			// Detect last entry in the list
+			if (!SelectedBufferNames.Split(TEXT(","), &Left, &Right))
+			{
+				Left = SelectedBufferNames;
+				Right = FString();
+			}
 
+			Left.TrimStartInline();
+			if (GetBufferMaterialName(*Left).IsEmpty())
+			{
+				SelectedBuffers.Add(TEXT(""));
+			}
+			else
+			{
+				SelectedBuffers.Add(*Left);
+			}
+			SelectedBufferNames = Right;
+		}
+	}
+	else
+	{
+		UE_LOG(LogPlayerController, Warning, TEXT("Console variable r.BufferVisualizationOverviewTargets is not found."));
+	}
 	return SelectedBuffers;
 }
 
@@ -756,7 +814,6 @@ void ADebugCameraController::ToggleBufferVisualizationOverviewMode()
 		if (bEnableBufferVisualization)
 		{
 			Cmd += GetViewModeName(VMI_VisualizeBuffer);
-			LastViewModeIndex = GameViewportClient->ViewModeIndex;
 
 			TArray<FString> SelectedBuffers = GetBufferVisualizationOverviewTargets();
 
@@ -770,7 +827,7 @@ void ADebugCameraController::ToggleBufferVisualizationOverviewMode()
 		}
 		else
 		{
-			Cmd += GetViewModeName((EViewModeIndex)LastViewModeIndex);
+			Cmd += GetViewModeName(EViewModeIndex::VMI_Lit);
 			SetDisplay(bLastDisplayEnabled);
 		}
 
@@ -780,13 +837,11 @@ void ADebugCameraController::ToggleBufferVisualizationOverviewMode()
 	}
 }
 
-#if !(UE_BUILD_SHIPPING || UE_BUILD_TEST)
-
 void ADebugCameraController::UpdateVisualizeBufferPostProcessing(FFinalPostProcessSettings& InOutPostProcessingSettings)
 {
 	if (bEnableBufferVisualization)
 	{
-		FString BufferMaterialName = GetSelectedBufferName();
+		FString BufferMaterialName = GetSelectedBufferMaterialName();
 		if (!BufferMaterialName.IsEmpty())
 		{
 			InOutPostProcessingSettings.bBufferVisualizationOverviewTargetIsSelected = true;
@@ -799,8 +854,6 @@ void ADebugCameraController::UpdateVisualizeBufferPostProcessing(FFinalPostProce
 	InOutPostProcessingSettings.BufferVisualizationOverviewSelectedTargetMaterialName.Empty();
 }
 
-#endif
-
 void ADebugCameraController::GetNextBuffer(int32 Step)
 {
 	if (bEnableBufferVisualization && !bEnableBufferVisualizationFullMode)
@@ -903,16 +956,24 @@ void ADebugCameraController::BufferVisualizationMoveLeft()
 	GetNextBuffer(-1);
 }
 
-FString ADebugCameraController::GetSelectedBufferName()
+FString ADebugCameraController::GetBufferMaterialName(const FString& InBufferName)
 {
-	if (UMaterial* Material = GetBufferVisualizationData().GetMaterial(*CurrSelectedBuffer))
+	if (!InBufferName.IsEmpty())
 	{
-		return Material->GetName();
+		if (UMaterial* Material = GetBufferVisualizationData().GetMaterial(*InBufferName))
+		{
+			return Material->GetName();
+		}
 	}
 
 	return TEXT("");
 }
 
+FString ADebugCameraController::GetSelectedBufferMaterialName()
+{
+	return GetBufferMaterialName(CurrSelectedBuffer);
+}
+
 void ADebugCameraController::ConsumeAxisMotion(float Val)
 {
 	// Just ignore the axis motion.
@@ -927,20 +988,26 @@ void ADebugCameraController::SetBufferVisualizationFullMode(bool bFullMode)
 {
 	if (bEnableBufferVisualizationFullMode != bFullMode)
 	{
-		bEnableBufferVisualizationFullMode = bFullMode;
-
 		static IConsoleVariable* ICVar = IConsoleManager::Get().FindConsoleVariable(FBufferVisualizationData::GetVisualizationTargetConsoleCommandName());
 		if (ICVar)
 		{
+			bEnableBufferVisualizationFullMode = bFullMode;
+
 			static const FName EmptyName = NAME_None;
 			ICVar->Set(bFullMode ? *CurrSelectedBuffer : *EmptyName.ToString(), ECVF_SetByCode);
-		}
 
-		SetupBufferVisualizationOverviewInput();
-		SetDisplay(bEnableBufferVisualizationFullMode);
+			SetupBufferVisualizationOverviewInput();
+			SetDisplay(bEnableBufferVisualizationFullMode);
+		}
+		else
+		{
+			UE_LOG(LogPlayerController, Verbose, TEXT("Console variable %s does not exist."), FBufferVisualizationData::GetVisualizationTargetConsoleCommandName());
+		}
 	}
 }
 
+#endif // !(UE_BUILD_SHIPPING || UE_BUILD_TEST)
+
 void ADebugCameraController::SelectTargetedObject()
 {
 	FVector CamLoc;
diff --git a/Engine/Source/Runtime/Engine/Private/DebugCameraHUD.cpp b/Engine/Source/Runtime/Engine/Private/DebugCameraHUD.cpp
--- a/Engine/Source/Runtime/Engine/Private/DebugCameraHUD.cpp
+++ b/Engine/Source/Runtime/Engine/Private/DebugCameraHUD.cpp
@@ -66,13 +66,14 @@ void ADebugCameraHUD::PostRender()
 		UFont* RenderFont = GEngine->GetSmallFont();
 		if( DCC != NULL )
 		{
-			if (DCC->bEnableBufferVisualizationFullMode)
+#if !(UE_BUILD_SHIPPING || UE_BUILD_TEST)
+			if (ADebugCameraController::EnableDebugBuffers() && DCC->bEnableBufferVisualizationFullMode)
 			{
 				RenderFont = GEngine->GetMediumFont();
 				FFontRenderInfo FontRenderInfo = Canvas->CreateFontRenderInfo(false, true);
 
 				Canvas->SetDrawColor(255, 255, 64, 255);
-				FString BufferText = DCC->GetSelectedBufferName();
+				FString BufferText = DCC->GetSelectedBufferMaterialName();
 				float xl, yl;
 				Canvas->StrLen(RenderFont, BufferText, xl, yl);
 				float X = Canvas->SizeX * 0.05f;
@@ -81,7 +82,8 @@ void ADebugCameraHUD::PostRender()
 
 				return;
 			}
-	
+#endif	
+
 			FFontRenderInfo FontRenderInfo = Canvas->CreateFontRenderInfo(false, true);
 
 			Canvas->SetDrawColor(64, 64, 255, 255);
@@ -117,10 +119,14 @@ void ADebugCameraHUD::PostRender()
 			Canvas->DrawText(RenderFont, SpeedString, X, yl, 1.f, 1.f, FontRenderInfo);
 			yl += Y;
 
-			FString const ViewModeString = FString::Printf(TEXT("ViewMode=%s"), GetWorld()->GetGameViewport() ? GetViewModeName((EViewModeIndex)GetWorld()->GetGameViewport()->ViewModeIndex) : TEXT("<Not found>"));
-			Canvas->DrawText(RenderFont, ViewModeString, X, yl, 1.f, 1.f, FontRenderInfo);
-			yl += Y;
-
+#if !(UE_BUILD_SHIPPING || UE_BUILD_TEST)
+			if (ADebugCameraController::EnableDebugViewmodes())
+			{
+				FString const ViewModeString = FString::Printf(TEXT("ViewMode=%s"), GetWorld()->GetGameViewport() ? GetViewModeName((EViewModeIndex)GetWorld()->GetGameViewport()->ViewModeIndex) : TEXT("<Not found>"));
+				Canvas->DrawText(RenderFont, ViewModeString, X, yl, 1.f, 1.f, FontRenderInfo);
+				yl += Y;
+			}
+#endif
 			const TCHAR* CVarComplexName = TEXT("g.DebugCameraTraceComplex");
 			bool bTraceComplex = true;
 
@@ -215,7 +221,7 @@ void ADebugCameraHUD::PostRender()
 
 
 			// controls display
-			yl += Y*8;
+			yl += Y*2;
 			
 			Canvas->SetDrawColor(64, 64, 255, 255);
 			Canvas->DrawText(RenderFont, TEXT("Controls"), X, yl, 1.f, 1.f, FontRenderInfo);
@@ -240,11 +246,20 @@ void ADebugCameraHUD::PostRender()
 			Canvas->DrawText(RenderFont, TEXT("Orbit Selected Actor: Shift-O"), X, yl, 1.f, 1.f, FontRenderInfo);
 			yl += Y;
 
-			Canvas->DrawText(RenderFont, TEXT("Toggle Buffer Visualization: B ('Enter' to toggle full buffer mode)"), X, yl, 1.f, 1.f, FontRenderInfo);
-			yl += Y;
+#if !(UE_BUILD_SHIPPING || UE_BUILD_TEST)
 
-			Canvas->DrawText(RenderFont, TEXT("Cycle View Mode: V"), X, yl, 1.f, 1.f, FontRenderInfo);
-			yl += Y;
+			if (ADebugCameraController::EnableDebugBuffers())
+			{
+				Canvas->DrawText(RenderFont, TEXT("Toggle Buffer Visualization: B ('Enter' to toggle full buffer mode)"), X, yl, 1.f, 1.f, FontRenderInfo);
+				yl += Y;
+			}
+
+			if (ADebugCameraController::EnableDebugViewmodes())
+			{
+				Canvas->DrawText(RenderFont, TEXT("Cycle View Mode: V"), X, yl, 1.f, 1.f, FontRenderInfo);
+				yl += Y;
+			}
+#endif
 		}
 	}
 #endif
