Fix "Implicit conversion from 'int32' to bool." warning
#rnx

diff --git a/Engine/Source/Runtime/Engine/Public/Rendering/MotionVectorSimulation.h b/Engine/Source/Runtime/Engine/Public/Rendering/MotionVectorSimulation.h
--- a/Engine/Source/Runtime/Engine/Public/Rendering/MotionVectorSimulation.h
+++ b/Engine/Source/Runtime/Engine/Public/Rendering/MotionVectorSimulation.h
@@ -79,7 +79,7 @@ public:
 private:
 
 	//~ FTickableGameObject interface - used for removing transforms that are no longer valid
-	virtual bool IsTickable() const { return SimulatedTransforms.Num(); }
+	virtual bool IsTickable() const { return SimulatedTransforms.Num() > 0; }
 	virtual bool IsTickableInEditor() const override { return true; }
 	virtual void Tick( float DeltaTime ) override;
 	virtual TStatId GetStatId() const override;
