Implement concurrency stop fade.

diff --git a/Engine/Source/Runtime/Engine/Classes/Components/AudioComponent.h b/Engine/Source/Runtime/Engine/Classes/Components/AudioComponent.h
--- a/Engine/Source/Runtime/Engine/Classes/Components/AudioComponent.h
+++ b/Engine/Source/Runtime/Engine/Classes/Components/AudioComponent.h
@@ -525,6 +525,8 @@ public:
 	virtual bool IsReadyForOwnerToAutoDestroy() const override;
 	//~ End ActorComponent Interface.
 
+	void AdjustVolumeInternal(float AdjustVolumeDuration, float AdjustVolumeLevel, bool bIsFadeOut);
+
 	/** Returns a pointer to the attenuation settings to be used (if any) for this audio component dependent on the SoundAttenuation asset or overrides set. */
 	const FSoundAttenuationSettings* GetAttenuationSettingsToApply() const;
 
diff --git a/Engine/Source/Runtime/Engine/Classes/Sound/SoundConcurrency.h b/Engine/Source/Runtime/Engine/Classes/Sound/SoundConcurrency.h
--- a/Engine/Source/Runtime/Engine/Classes/Sound/SoundConcurrency.h
+++ b/Engine/Source/Runtime/Engine/Classes/Sound/SoundConcurrency.h
@@ -100,6 +100,13 @@ public:
 	UPROPERTY(EditAnywhere, BlueprintReadWrite, Category = "Volume Scaling", meta = (DisplayName = "Release Time", EditCondition = "bVolumeScaleCanRelease", UIMin = "0.0", ClampMin = "0.0", UIMax = "10.0", ClampMax="1000000.0"))
 	float VolumeScaleReleaseTime;
 
+	/**
+	 * Time taken to fade out if voice is evicted or culled due to another voice in the group starting.
+	 */
+	UPROPERTY(EditAnywhere, BlueprintReadWrite, Category = "Voice Stealing", meta = (UIMin = "0.0", ClampMin = "0.0", UIMax = "1.0", ClampMax="1000000.0"))
+	float VoiceStealReleaseTime;
+
+
 	FSoundConcurrencySettings()
 		: MaxCount(16)
 		, bLimitToOwner(0)
@@ -107,7 +114,8 @@ public:
 		, VolumeScale(1.0f)
 		, VolumeScaleAttackTime(0.01f)
 		, bVolumeScaleCanRelease(0)
-		, VolumeScaleReleaseTime(0.0f)
+		, VolumeScaleReleaseTime(0.5f)
+		, VoiceStealReleaseTime(0.0f)
 	{}
 
 	/**
@@ -267,7 +275,7 @@ typedef TMap<FSoundObjectID, FConcurrencyGroupID> FPerSoundToActiveSoundsMap;
 class FSoundConcurrencyManager
 {
 public:
-	FSoundConcurrencyManager(class FAudioDevice* InAudioDevice);
+	FSoundConcurrencyManager(FAudioDevice* InAudioDevice);
 	ENGINE_API ~FSoundConcurrencyManager();
 
 	/** Returns a newly allocated active sound given the input active sound struct. Will return nullptr if the active sound concurrency evaluation doesn't allow for it. */
@@ -276,6 +284,9 @@ public:
 	/** Removes the active sound from concurrency tracking when active sound is stopped. */
 	void RemoveActiveSound(FActiveSound& ActiveSound);
 
+	// Fades sound if concurrency settings request to, or calls stop if not set
+	void FadeOutActiveSound(FActiveSound& ActiveSound);
+
 	/** Stops any active sounds due to max concurrency quietest sound resolution rule */
 	void UpdateQuietSoundsToStop();
 
diff --git a/Engine/Source/Runtime/Engine/Private/ActiveSound.cpp b/Engine/Source/Runtime/Engine/Private/ActiveSound.cpp
--- a/Engine/Source/Runtime/Engine/Private/ActiveSound.cpp
+++ b/Engine/Source/Runtime/Engine/Private/ActiveSound.cpp
@@ -42,7 +42,6 @@ FActiveSound::FActiveSound()
 	, bAllowSpatialization(true)
 	, bHasAttenuationSettings(false)
 	, bShouldRemainActiveIfDropped(false)
-	, bFadingOut(false)
 	, bFinished(false)
 	, bIsPaused(false)
 	, bShouldStopDueToMaxConcurrency(false)
@@ -73,6 +72,7 @@ FActiveSound::FActiveSound()
 	, bIsPlayingAudio(false)
 	, bIsStopping(false)
 	, UserIndex(0)
+	, FadeOut(EFadeOut::None)
 	, bIsOccluded(false)
 	, bAsyncOcclusionPending(false)
 	, PlaybackTime(0.f)
@@ -408,6 +408,7 @@ int32 FActiveSound::FindClosestListener( const TArray<FListener>& InListeners )
 void FActiveSound::GetConcurrencyHandles(TArray<FConcurrencyHandle>& OutConcurrencyHandles) const
 {
 	OutConcurrencyHandles.Reset();
+
 	if (!ConcurrencySet.Num() && Sound)
 	{
 		Sound->GetConcurrencyHandles(OutConcurrencyHandles);
@@ -438,6 +439,27 @@ void FActiveSound::GetConcurrencyHandles(TArray<FConcurrencyHandle>& OutConcurre
 	}
 }
 
+bool FActiveSound::GetConcurrencyFadeDuration(float& OutFadeDuration) const
+{
+	OutFadeDuration = -1.0f;
+	TArray<FConcurrencyHandle> Handles;
+	GetConcurrencyHandles(Handles);
+	for (FConcurrencyHandle& Handle : Handles)
+	{
+		OutFadeDuration = OutFadeDuration < 0.0f
+			? Handle.Settings.VoiceStealReleaseTime
+			: FMath::Min(Handle.Settings.VoiceStealReleaseTime, OutFadeDuration);
+	}
+
+	if (OutFadeDuration < 0.0f)
+	{
+		OutFadeDuration = 0.0f;
+		return false;
+	}
+
+	return true;
+}
+
 void FActiveSound::UpdateWaveInstances(TArray<FWaveInstance*> &InWaveInstances, const float DeltaTime)
 {
 	check(AudioDevice);
@@ -529,7 +551,7 @@ void FActiveSound::UpdateWaveInstances(TArray<FWaveInstance*> &InWaveInstances,
 
 	// Recurse nodes, have SoundWave's create new wave instances and update bFinished unless we finished fading out.
 	bFinished = true;
-	if (!bFadingOut || (PlaybackTime <= TargetAdjustVolumeStopTime))
+	if (FadeOut == EFadeOut::None || (PlaybackTime <= TargetAdjustVolumeStopTime))
 	{
 		if (bHasAttenuationSettings)
 		{
diff --git a/Engine/Source/Runtime/Engine/Private/AudioDevice.cpp b/Engine/Source/Runtime/Engine/Private/AudioDevice.cpp
--- a/Engine/Source/Runtime/Engine/Private/AudioDevice.cpp
+++ b/Engine/Source/Runtime/Engine/Private/AudioDevice.cpp
@@ -2094,39 +2094,53 @@ void FAudioDevice::StopQuietSoundsDueToMaxConcurrency(TArray<FWaveInstance*>& Wa
 		ConcurrencyManager.UpdateQuietSoundsToStop();
 	}
 
-	// Remove all wave instances from the wave instance list that are stopping due to max concurrency
-	for (int32 i = WaveInstances.Num() - 1; i >= 0; --i)
-	{
-		if (WaveInstances[i]->ShouldStopDueToMaxConcurrency())
-		{
-			WaveInstances.RemoveAtSwap(i, 1, false);
-		}
-	}
-
 	for (int32 i = ActiveSoundsCopy.Num() - 1; i >= 0; --i)
 	{
 		if (FActiveSound* ActiveSound = ActiveSoundsCopy[i])
 		{
-			if (ActiveSound->bShouldStopDueToMaxConcurrency)
+			if (!ActiveSound->bShouldStopDueToMaxConcurrency)
 			{
-				AddSoundToStop(ActiveSound);
-				const bool bDoRangeCheck = false;
-				FAudioVirtualLoop VirtualLoop;
-				if (FAudioVirtualLoop::Virtualize(*ActiveSound, bDoRangeCheck, VirtualLoop))
-				{
-					// Clear must be called after AddSoundToStop to ensure AudioComponent is properly removed from AudioComponentIDToActiveSoundMap
-					ActiveSound->ClearAudioComponent();
+				continue;
+			}
 
-					const bool bIsAtMaxConcurrency = true;
-					VirtualLoop.CalculateUpdateInterval(bIsAtMaxConcurrency);
-					if (USoundBase* Sound = ActiveSound->GetSound())
-					{
-						UE_LOG(LogAudio, Verbose, TEXT("Playing ActiveSound %s Virtualizing: Invalidated by 'StopQuiestest' Concurrency Rule."), *Sound->GetName());
-					}
-					AddVirtualLoop(VirtualLoop);
+			if (ActiveSound->FadeOut == FActiveSound::EFadeOut::Concurrency)
+			{
+				continue;
+			}
+
+			if (ActiveSound->bIsStopping)
+			{
+				continue;
+			}
+
+			ConcurrencyManager.RemoveActiveSound(*ActiveSound);
+			ConcurrencyManager.FadeOutActiveSound(*ActiveSound);
+
+			const bool bDoRangeCheck = false;
+			FAudioVirtualLoop VirtualLoop;
+			if (FAudioVirtualLoop::Virtualize(*ActiveSound, bDoRangeCheck, VirtualLoop))
+			{
+				// Clear must be called after AddSoundToStop to ensure AudioComponent is properly removed from AudioComponentIDToActiveSoundMap
+				ActiveSound->ClearAudioComponent();
+
+				if (USoundBase* Sound = ActiveSound->GetSound())
+				{
+					UE_LOG(LogAudio, Verbose, TEXT("Playing ActiveSound %s Virtualizing: Invalidated by 'StopQuiestest' Concurrency Rule."), *Sound->GetName());
 				}
-				ActiveSoundsCopy.RemoveAtSwap(i, 1, false);
+				AddVirtualLoop(VirtualLoop);
 			}
+			ActiveSoundsCopy.RemoveAtSwap(i, 1, false);
+		}
+	}
+
+	// Remove all wave instances from the wave instance list that are stopping due to max concurrency.
+	// Must be after checking if sound must fade out due to concurrency to avoid pre-maturally removing
+	// wave instances prior to concurrency system marking as fading out.
+	for (int32 i = WaveInstances.Num() - 1; i >= 0; --i)
+	{
+		if (WaveInstances[i]->ShouldStopDueToMaxConcurrency())
+		{
+			WaveInstances.RemoveAtSwap(i, 1, false);
 		}
 	}
 }
@@ -5726,7 +5740,7 @@ void FAudioDevice::UpdateVirtualLoops()
 			}
 
 			// If signaled to fade out and virtualized, add to pending stop list.
-			if (ActiveSound.bFadingOut)
+			if (ActiveSound.FadeOut != FActiveSound::EFadeOut::None)
 			{
 				AddSoundToStop(&ActiveSound);
 				continue;
@@ -5734,7 +5748,9 @@ void FAudioDevice::UpdateVirtualLoops()
 
 			const float DeltaTime = GetDeviceDeltaTime();
 
-			// Keep playback time up-to-date as it may be used to evaluate concurrency
+			// Keep playback time up-to-date as it may be used to evaluate whether or
+			// not virtual sound is eligible for playback when compared against
+			// actively playing sounds in concurrency checks.
 			ActiveSound.PlaybackTime += DeltaTime * ActiveSound.MinCurrentPitch;
 
 			// If the loop is ready to realize, add to array to be re-triggered outside of the loop
diff --git a/Engine/Source/Runtime/Engine/Private/AudioVirtualLoop.cpp b/Engine/Source/Runtime/Engine/Private/AudioVirtualLoop.cpp
--- a/Engine/Source/Runtime/Engine/Private/AudioVirtualLoop.cpp
+++ b/Engine/Source/Runtime/Engine/Private/AudioVirtualLoop.cpp
@@ -66,7 +66,7 @@ bool FAudioVirtualLoop::Virtualize(const FActiveSound& InActiveSound, FAudioDevi
 		return false;
 	}
 
-	if (InActiveSound.bFadingOut || InActiveSound.bIsStopping)
+	if (InActiveSound.FadeOut != FActiveSound::EFadeOut::None || InActiveSound.bIsStopping)
 	{
 		return false;
 	}
@@ -81,27 +81,16 @@ bool FAudioVirtualLoop::Virtualize(const FActiveSound& InActiveSound, FAudioDevi
 	return true;
 }
 
-void FAudioVirtualLoop::CalculateUpdateInterval(bool bIsAtMaxConcurrency)
+void FAudioVirtualLoop::CalculateUpdateInterval()
 {
-	// If calculating due to being at max concurrency, set to max rate as
-	// sound will most likely be killed again on next check until concurrency
-	// is no longer full.  This limits starting and stopping of excess sounds
-	// virtualizing.
-	if (bIsAtMaxConcurrency)
-	{
-		UpdateInterval = VirtualLoopsUpdateRateMaxCVar;
-	}
-	else
-	{
-		check(ActiveSound);
-		FAudioDevice* AudioDevice = ActiveSound->AudioDevice;
-		check(AudioDevice);
-
-		const float DistanceToListener = AudioDevice->GetDistanceToNearestListener(ActiveSound->Transform.GetLocation());
-		const float DistanceRatio = (DistanceToListener - ActiveSound->MaxDistance) / FMath::Max(VirtualLoopsPerfDistanceCVar, 1.0f);
-		const float DistanceRatioClamped = FMath::Clamp(DistanceRatio, 0.0f, 1.0f);
-		UpdateInterval = FMath::Lerp(VirtualLoopsUpdateRateMinCVar, VirtualLoopsUpdateRateMaxCVar, DistanceRatioClamped);
-	}
+	check(ActiveSound);
+	FAudioDevice* AudioDevice = ActiveSound->AudioDevice;
+	check(AudioDevice);
+
+	const float DistanceToListener = AudioDevice->GetDistanceToNearestListener(ActiveSound->Transform.GetLocation());
+	const float DistanceRatio = (DistanceToListener - ActiveSound->MaxDistance) / FMath::Max(VirtualLoopsPerfDistanceCVar, 1.0f);
+	const float DistanceRatioClamped = FMath::Clamp(DistanceRatio, 0.0f, 1.0f);
+	UpdateInterval = FMath::Lerp(VirtualLoopsUpdateRateMinCVar, VirtualLoopsUpdateRateMaxCVar, DistanceRatioClamped);
 }
 
 FActiveSound& FAudioVirtualLoop::GetActiveSound()
diff --git a/Engine/Source/Runtime/Engine/Private/Components/AudioComponent.cpp b/Engine/Source/Runtime/Engine/Private/Components/AudioComponent.cpp
--- a/Engine/Source/Runtime/Engine/Private/Components/AudioComponent.cpp
+++ b/Engine/Source/Runtime/Engine/Private/Components/AudioComponent.cpp
@@ -532,66 +532,84 @@ void UAudioComponent::FadeIn( float FadeInDuration, float FadeVolumeLevel, float
 	PlayInternal(StartTime, FadeInDuration, FadeVolumeLevel);
 }
 
-void UAudioComponent::FadeOut( float FadeOutDuration, float FadeVolumeLevel )
+void UAudioComponent::FadeOut(float FadeOutDuration, float FadeVolumeLevel)
 {
-	if (bIsActive)
+	const bool bIsFadeOut = true;
+	AdjustVolumeInternal(FadeOutDuration, FadeVolumeLevel, bIsFadeOut);
+}
+
+void UAudioComponent::AdjustVolume(float AdjustVolumeDuration, float AdjustVolumeLevel)
+{
+	const bool bIsFadeOut = false;
+	AdjustVolumeInternal(AdjustVolumeDuration, AdjustVolumeLevel, bIsFadeOut);
+}
+
+void UAudioComponent::AdjustVolumeInternal(float AdjustVolumeDuration, float AdjustVolumeLevel, bool bIsFadeOut)
+{
+	if (!bIsActive)
+	{
+		return;
+	}
+
+	FAudioDevice* AudioDevice = GetAudioDevice();
+	if (!AudioDevice)
+	{
+		return;
+	}
+
+	if (AdjustVolumeDuration == 0.0f && AdjustVolumeLevel == 0.0f)
+	{
+		Stop();
+		return;
+	}
+
+	const uint64 InAudioComponentID = AudioComponentID;
+	DECLARE_CYCLE_STAT(TEXT("FAudioThreadTask.AdjustVolume"), STAT_AudioAdjustVolume, STATGROUP_AudioThreadCommands);
+	FAudioThread::RunCommandOnAudioThread([AudioDevice, InAudioComponentID, AdjustVolumeDuration, AdjustVolumeLevel, bIsFadeOut]()
 	{
-		if (FadeOutDuration > 0.0f)
+		FActiveSound* ActiveSound = AudioDevice->FindActiveSound(InAudioComponentID);
+		if (!ActiveSound)
 		{
-			if (FAudioDevice* AudioDevice = GetAudioDevice())
+			return;
+		}
+
+		// Ignore fade out request if requested volume is higher than current target.
+		if (bIsFadeOut && AdjustVolumeLevel >= ActiveSound->TargetAdjustVolumeMultiplier)
+		{
+			return;
+		}
+
+		if (ActiveSound->FadeOut == FActiveSound::EFadeOut::Concurrency)
+		{
+			// Ignore adjust volume request if non-zero and currently voice stealing.
+			if (AdjustVolumeLevel != 0.0f)
 			{
-				DECLARE_CYCLE_STAT(TEXT("FAudioThreadTask.FadeOut"), STAT_AudioFadeOut, STATGROUP_AudioThreadCommands);
+				return;
+			}
 
-				const uint64 MyAudioComponentID = AudioComponentID;
-				FAudioThread::RunCommandOnAudioThread([AudioDevice, MyAudioComponentID, FadeOutDuration, FadeVolumeLevel]()
-				{
-					FActiveSound* ActiveSound = AudioDevice->FindActiveSound(MyAudioComponentID);
-					if (ActiveSound)
-					{
-						ActiveSound->TargetAdjustVolumeMultiplier = FadeVolumeLevel;
-						ActiveSound->TargetAdjustVolumeStopTime = ActiveSound->PlaybackTime + FadeOutDuration;
-						ActiveSound->bFadingOut = true;
-					}
-				}, GET_STATID(STAT_AudioFadeOut));
+			// Ignore request of longer fade out than active target if active is concurrency (voice stealing) fade.
+			if (AdjustVolumeDuration > ActiveSound->TargetAdjustVolumeStopTime)
+			{
+				return;
 			}
 		}
 		else
 		{
-			Stop();
+			ActiveSound->FadeOut = AdjustVolumeLevel == 0.0f ? FActiveSound::EFadeOut::User : FActiveSound::EFadeOut::None;
 		}
-	}
-}
-
-void UAudioComponent::AdjustVolume( float AdjustVolumeDuration, float AdjustVolumeLevel )
-{
-	if (bIsActive)
-	{
-		if (FAudioDevice* AudioDevice = GetAudioDevice())
-		{
-			DECLARE_CYCLE_STAT(TEXT("FAudioThreadTask.AdjustVolume"), STAT_AudioAdjustVolume, STATGROUP_AudioThreadCommands);
 
-			const uint64 MyAudioComponentID = AudioComponentID;
-			FAudioThread::RunCommandOnAudioThread([AudioDevice, MyAudioComponentID, AdjustVolumeDuration, AdjustVolumeLevel]()
-			{
-				FActiveSound* ActiveSound = AudioDevice->FindActiveSound(MyAudioComponentID);
-				if (ActiveSound)
-				{
-					ActiveSound->bFadingOut = false;
-					ActiveSound->TargetAdjustVolumeMultiplier = AdjustVolumeLevel;
+		ActiveSound->TargetAdjustVolumeMultiplier = AdjustVolumeLevel;
 
-					if (AdjustVolumeDuration > 0.0f)
-					{
-						ActiveSound->TargetAdjustVolumeStopTime = ActiveSound->PlaybackTime + AdjustVolumeDuration;
-					}
-					else
-					{
-						ActiveSound->CurrentAdjustVolumeMultiplier = AdjustVolumeLevel;
-						ActiveSound->TargetAdjustVolumeStopTime = -1.0f;
-					}
-				}
-			}, GET_STATID(STAT_AudioAdjustVolume));
+		if (AdjustVolumeDuration > 0.0f)
+		{
+			ActiveSound->TargetAdjustVolumeStopTime = ActiveSound->PlaybackTime + AdjustVolumeDuration;
 		}
-	}
+		else
+		{
+			ActiveSound->CurrentAdjustVolumeMultiplier = AdjustVolumeLevel;
+			ActiveSound->TargetAdjustVolumeStopTime = -1.0f;
+		}
+	}, GET_STATID(STAT_AudioAdjustVolume));
 }
 
 void UAudioComponent::Stop()
diff --git a/Engine/Source/Runtime/Engine/Private/SoundConcurrency.cpp b/Engine/Source/Runtime/Engine/Private/SoundConcurrency.cpp
--- a/Engine/Source/Runtime/Engine/Private/SoundConcurrency.cpp
+++ b/Engine/Source/Runtime/Engine/Private/SoundConcurrency.cpp
@@ -214,9 +214,15 @@ void FConcurrencyGroup::StopQuietSoundsDueToMaxConcurrency()
 	int32 i = 0;
 	for (; i < NumSoundsToStop; ++i)
 	{
+		FActiveSound* ActiveSound = ActiveSounds[i];
+		check(ActiveSound);
+
 		// Flag this active sound as needing to be stopped due to volume-based max concurrency.
 		// This will actually be stopped in the audio device update function.
-		ActiveSounds[i]->bShouldStopDueToMaxConcurrency = true;
+		if (ActiveSound->FadeOut != FActiveSound::EFadeOut::Concurrency)
+		{
+			ActiveSound->bShouldStopDueToMaxConcurrency = true;
+		}
 	}
 
 	const int32 NumActiveSounds = ActiveSounds.Num();
@@ -693,24 +699,30 @@ FActiveSound* FSoundConcurrencyManager::CreateAndEvictActiveSounds(const FActive
 		check(SoundToEvict);
 		check(AudioDevice == SoundToEvict->AudioDevice);
 
-		RemoveActiveSound(*SoundToEvict);
-
 		// Remove the active sound from the concurrency manager immediately so it doesn't count towards
 		// subsequent concurrency resolution checks (i.e. if sounds are triggered multiple times in this frame)
+		RemoveActiveSound(*SoundToEvict);
+
+		if (SoundToEvict->FadeOut == FActiveSound::EFadeOut::Concurrency)
+		{
+			continue;
+		}
+
 		if (AudioDevice->IsPendingStop(SoundToEvict))
 		{
 			continue;
 		}
 
-		TArray<FConcurrencyHandle> Handles;
-		SoundToEvict->GetConcurrencyHandles(Handles);
+		FadeOutActiveSound(*SoundToEvict);
 
-		AudioDevice->AddSoundToStop(SoundToEvict);
 		const bool bDoRangeCheck = false;
 		FAudioVirtualLoop VirtualLoop;
 		if (FAudioVirtualLoop::Virtualize(*SoundToEvict, bDoRangeCheck, VirtualLoop))
 		{
+			// Don't report that the sound is 'finished' by clearing the audio component of the evicted
+			// copy that is being virtualized and effectively stopped at the end of this frame.
 			SoundToEvict->ClearAudioComponent();
+
 			if (USoundBase* Sound = SoundToEvict->GetSound())
 			{
 				UE_LOG(LogAudio, Verbose, TEXT("Playing ActiveSound %s Virtualizing: Evicted due to concurrency rules."), *Sound->GetName());
@@ -795,6 +807,23 @@ void FSoundConcurrencyManager::RemoveActiveSound(FActiveSound& ActiveSound)
 	ActiveSound.ConcurrencyGroupData.Reset();
 }
 
+void FSoundConcurrencyManager::FadeOutActiveSound(FActiveSound& ActiveSound)
+{
+	check(AudioDevice);
+
+	float FadeOutDuration = 0.0f;
+	if (ActiveSound.GetConcurrencyFadeDuration(FadeOutDuration))
+	{
+		ActiveSound.FadeOut = FActiveSound::EFadeOut::Concurrency;
+		ActiveSound.TargetAdjustVolumeMultiplier = 0.0f;
+		ActiveSound.TargetAdjustVolumeStopTime = ActiveSound.PlaybackTime + FadeOutDuration;
+	}
+	else
+	{
+		AudioDevice->AddSoundToStop(&ActiveSound);
+	}
+}
+
 void FSoundConcurrencyManager::UpdateQuietSoundsToStop()
 {
 	check(IsInAudioThread());
diff --git a/Engine/Source/Runtime/Engine/Public/ActiveSound.h b/Engine/Source/Runtime/Engine/Public/ActiveSound.h
--- a/Engine/Source/Runtime/Engine/Public/ActiveSound.h
+++ b/Engine/Source/Runtime/Engine/Public/ActiveSound.h
@@ -305,6 +305,18 @@ private:
 	TMap<UPTRINT, FWaveInstance*> WaveInstances;
 
 public:
+	enum class EFadeOut : uint8
+	{
+		// Sound is not currently fading out
+		None,
+
+		// Client code (eg. AudioComponent) is requesting a fade out
+		User,
+
+		// The concurrency system is requesting a fade due to voice stealing
+		Concurrency
+	};
+
 	/** Whether or not the sound has checked if it was occluded already. Used to initialize a sound as occluded and bypassing occlusion interpolation. */
 	uint8 bHasCheckedOcclusion:1;
 
@@ -317,9 +329,6 @@ public:
 	/** Whether the wave instances should remain active if they're dropped by the prioritization code. Useful for e.g. vehicle sounds that shouldn't cut out. */
 	uint8 bShouldRemainActiveIfDropped:1;
 
-	/** Is the audio component currently fading out */
-	uint8 bFadingOut:1;
-
 	/** Whether the current component has finished playing */
 	uint8 bFinished:1;
 
@@ -403,9 +412,11 @@ public:
 	/** Whether or not the active sound is stopping. */
 	uint8 bIsStopping:1;
 
-public:
 	uint8 UserIndex;
 
+	/** Type of fade out currently being applied */
+	EFadeOut FadeOut;
+
 	/** whether we were occluded the last time we checked */
 	FThreadSafeBool bIsOccluded;
 
@@ -595,6 +606,8 @@ public:
 	/** Gets the sound concurrency handles applicable to this sound instance*/
 	void GetConcurrencyHandles(TArray<FConcurrencyHandle>& OutConcurrencyHandles) const;
 
+	bool GetConcurrencyFadeDuration(float& OutFadeDuration) const;
+
 	/** Delegate callback function when an async occlusion trace completes */
 	static void OcclusionTraceDone(const FTraceHandle& TraceHandle, FTraceDatum& TraceDatum);
 
diff --git a/Engine/Source/Runtime/Engine/Public/AudioVirtualLoop.h b/Engine/Source/Runtime/Engine/Public/AudioVirtualLoop.h
--- a/Engine/Source/Runtime/Engine/Public/AudioVirtualLoop.h
+++ b/Engine/Source/Runtime/Engine/Public/AudioVirtualLoop.h
@@ -55,7 +55,7 @@ public:
 	/**
 	 * Overrides the update interval to the provided length.
 	 */
-	void CalculateUpdateInterval(bool bIsAtMaxConcurrency = false);
+	void CalculateUpdateInterval();
 
 	/**
 	  * Updates the loop and checks if ready to play (or 'realize').
