Fix for unbound uniform buffer with SpeedTree node

diff --git a/Engine/Source/Runtime/Engine/Private/LocalVertexFactory.cpp b/Engine/Source/Runtime/Engine/Private/LocalVertexFactory.cpp
--- a/Engine/Source/Runtime/Engine/Private/LocalVertexFactory.cpp
+++ b/Engine/Source/Runtime/Engine/Private/LocalVertexFactory.cpp
@@ -128,15 +128,15 @@ void FLocalVertexFactoryShaderParametersBase::GetElementShaderBindingsBase(
 		}	
 	}
 
-	if (bAnySpeedTreeParamIsBound && Scene)
+	if (bAnySpeedTreeParamIsBound)
 	{
 		QUICK_SCOPE_CYCLE_COUNTER(STAT_FLocalVertexFactoryShaderParameters_SetMesh_SpeedTree);
-		FUniformBufferRHIParamRef SpeedTreeUniformBuffer = Scene->GetSpeedTreeUniformBuffer(VertexFactory);
-		if (SpeedTreeUniformBuffer == NULL)
+		FUniformBufferRHIParamRef SpeedTreeUniformBuffer = Scene? Scene->GetSpeedTreeUniformBuffer(VertexFactory) : nullptr;
+		if (SpeedTreeUniformBuffer == nullptr)
 		{
 			SpeedTreeUniformBuffer = GSpeedTreeWindNullUniformBuffer.GetUniformBufferRHI();
 		}
-		check(SpeedTreeUniformBuffer != NULL);
+		check(SpeedTreeUniformBuffer != nullptr);
 
 		ShaderBindings.Add(Shader->GetUniformBufferParameter<FSpeedTreeUniformParameters>(), SpeedTreeUniformBuffer);
 
