[NO NON-PII COMMIT MESSAGE]

diff --git a/Engine/Source/Runtime/Engine/Private/GameModeBase.cpp b/Engine/Source/Runtime/Engine/Private/GameModeBase.cpp
--- a/Engine/Source/Runtime/Engine/Private/GameModeBase.cpp
+++ b/Engine/Source/Runtime/Engine/Private/GameModeBase.cpp
@@ -507,6 +507,7 @@ void AGameModeBase::SwapPlayerControllers(APlayerController* OldPC, APlayerContr
 		UPlayer* Player = OldPC->Player;
 		NewPC->NetPlayerIndex = OldPC->NetPlayerIndex; //@warning: critical that this is first as SetPlayer() may trigger RPCs
 		NewPC->NetConnection = OldPC->NetConnection;
+		NewPC->SetReplicates(OldPC->GetIsReplicated());
 		NewPC->SetPlayer(Player);
 		NewPC->CopyRemoteRoleFrom(OldPC);
 
