#ue4 - Fix character movement debug visualization lasting forever. This used to work but apparently the lifetime no longer overrides bPersistentLines.

diff --git a/Engine/Source/Runtime/Engine/Private/Components/CharacterMovementComponent.cpp b/Engine/Source/Runtime/Engine/Private/Components/CharacterMovementComponent.cpp
--- a/Engine/Source/Runtime/Engine/Private/Components/CharacterMovementComponent.cpp
+++ b/Engine/Source/Runtime/Engine/Private/Components/CharacterMovementComponent.cpp
@@ -1493,7 +1493,7 @@ void UCharacterMovementComponent::SimulatedTick(float DeltaSeconds)
 				const FRotator NewRotation = UpdatedComponent->GetComponentRotation();
 				const FVector NewLocation = UpdatedComponent->GetComponentLocation();
 				DrawDebugCoordinateSystem(GetWorld(), CharacterOwner->GetMesh()->GetComponentLocation() + FVector(0,0,1), NewRotation, 50.f, false); //-V595
-				DrawDebugLine(GetWorld(), OldLocation, NewLocation, FColor::Red, true, 10.f);
+				DrawDebugLine(GetWorld(), OldLocation, NewLocation, FColor::Red, false, 10.f);
 
 				UE_LOG(LogRootMotion, Log,  TEXT("UCharacterMovementComponent::SimulatedTick DeltaMovement Translation: %s, Rotation: %s, MovementBase: %s"), //-V595
 					*(NewLocation - OldLocation).ToCompactString(), *(NewRotation - OldRotation).GetNormalized().ToCompactString(), *GetNameSafe(CharacterOwner->GetMovementBase()) );
@@ -2403,7 +2403,7 @@ void UCharacterMovementComponent::PerformMovement(float DeltaSeconds)
 				DrawDebugCoordinateSystem(MyWorld, CharacterOwner->GetMesh()->GetComponentLocation() + FVector(0,0,1), ResultingRotation, 50.f, false);
 
 				// Show resulting delta move.
-				DrawDebugLine(MyWorld, OldLocation, ResultingLocation, FColor::Red, true, 10.f);
+				DrawDebugLine(MyWorld, OldLocation, ResultingLocation, FColor::Red, false, 10.f);
 
 				// Log details.
 				UE_LOG(LogRootMotion, Warning,  TEXT("PerformMovement Resulting DeltaMove Translation: %s, Rotation: %s, MovementBase: %s"), //-V595
@@ -3404,7 +3404,7 @@ void UCharacterMovementComponent::CalcAvoidanceVelocity(float DeltaTime)
 #if !(UE_BUILD_SHIPPING || UE_BUILD_TEST)
 			if (bShowDebug)
 			{
-				DrawDebugLine(GetWorld(), GetActorFeetLocation(), GetActorFeetLocation() + Velocity, FColor::Blue, true, 0.5f, SDPG_MAX);
+				DrawDebugLine(GetWorld(), GetActorFeetLocation(), GetActorFeetLocation() + Velocity, FColor::Blue, false, 0.5f, SDPG_MAX);
 			}
 #endif
 		}
@@ -3424,7 +3424,7 @@ void UCharacterMovementComponent::CalcAvoidanceVelocity(float DeltaTime)
 #if !(UE_BUILD_SHIPPING || UE_BUILD_TEST)
 				if (bShowDebug)
 				{
-					DrawDebugLine(GetWorld(), GetActorFeetLocation(), GetActorFeetLocation() + Velocity, FColor::Red, true, 0.05f, SDPG_MAX, 10.0f);
+					DrawDebugLine(GetWorld(), GetActorFeetLocation(), GetActorFeetLocation() + Velocity, FColor::Red, false, 0.05f, SDPG_MAX, 10.0f);
 				}
 #endif
 			}
@@ -3435,7 +3435,7 @@ void UCharacterMovementComponent::CalcAvoidanceVelocity(float DeltaTime)
 #if !(UE_BUILD_SHIPPING || UE_BUILD_TEST)
 				if (bShowDebug)
 				{
-					//DrawDebugLine(GetWorld(), GetActorLocation(), GetActorLocation() + Velocity, FColor::Green, true, 0.05f, SDPG_MAX, 10.0f);
+					//DrawDebugLine(GetWorld(), GetActorLocation(), GetActorLocation() + Velocity, FColor::Green, false, 0.05f, SDPG_MAX, 10.0f);
 				}
 #endif
 			}
@@ -3448,13 +3448,13 @@ void UCharacterMovementComponent::CalcAvoidanceVelocity(float DeltaTime)
 #if !(UE_BUILD_SHIPPING || UE_BUILD_TEST)
 	else if (bShowDebug)
 	{
-		DrawDebugLine(GetWorld(), GetActorFeetLocation(), GetActorFeetLocation() + Velocity, FColor::Yellow, true, 0.05f, SDPG_MAX);
+		DrawDebugLine(GetWorld(), GetActorFeetLocation(), GetActorFeetLocation() + Velocity, FColor::Yellow, false, 0.05f, SDPG_MAX);
 	}
 
 	if (bShowDebug)
 	{
 		FVector UpLine(0,0,500);
-		DrawDebugLine(GetWorld(), GetActorFeetLocation(), GetActorFeetLocation() + UpLine, (AvoidanceLockTimer > 0.01f) ? FColor::Red : FColor::Blue, true, 0.05f, SDPG_MAX, 5.0f);
+		DrawDebugLine(GetWorld(), GetActorFeetLocation(), GetActorFeetLocation() + UpLine, (AvoidanceLockTimer > 0.01f) ? FColor::Red : FColor::Blue, false, 0.05f, SDPG_MAX, 5.0f);
 	}
 #endif
 }
@@ -3563,7 +3563,7 @@ void UCharacterMovementComponent::NotifyBumpedPawn(APawn* BumpedPawn)
 	const bool bShowDebug = Avoidance && Avoidance->IsDebugEnabled(AvoidanceUID);
 	if (bShowDebug)
 	{
-		DrawDebugLine(GetWorld(), GetActorFeetLocation(), GetActorFeetLocation() + FVector(0,0,500), (AvoidanceLockTimer > 0) ? FColor(255,64,64) : FColor(64,64,255), true, 2.0f, SDPG_MAX, 20.0f);
+		DrawDebugLine(GetWorld(), GetActorFeetLocation(), GetActorFeetLocation() + FVector(0,0,500), (AvoidanceLockTimer > 0) ? FColor(255,64,64) : FColor(64,64,255), false, 2.0f, SDPG_MAX, 20.0f);
 	}
 #endif
 
@@ -8739,8 +8739,8 @@ void UCharacterMovementComponent::ServerMoveHandleClientError(float ClientTimeSt
 			UE_LOG(LogNetPlayerMovement, Warning, TEXT("*** Server: Error for %s at Time=%.3f is %3.3f LocDiff(%s) ClientLoc(%s) ServerLoc(%s) Base: %s Bone: %s Accel(%s) Velocity(%s)"),
 				*GetNameSafe(CharacterOwner), ClientTimeStamp, LocDiff.Size(), *LocDiff.ToString(), *ClientLoc.ToString(), *UpdatedComponent->GetComponentLocation().ToString(), *BaseString, *ServerData->PendingAdjustment.NewBaseBoneName.ToString(), *Accel.ToString(), *Velocity.ToString());
 			const float DebugLifetime = CharacterMovementCVars::NetCorrectionLifetime;
-			DrawDebugCapsule(GetWorld(), UpdatedComponent->GetComponentLocation(), CharacterOwner->GetSimpleCollisionHalfHeight(), CharacterOwner->GetSimpleCollisionRadius(), FQuat::Identity, FColor(100, 255, 100), true, DebugLifetime);
-			DrawDebugCapsule(GetWorld(), ClientLoc                    , CharacterOwner->GetSimpleCollisionHalfHeight(), CharacterOwner->GetSimpleCollisionRadius(), FQuat::Identity, FColor(255, 100, 100), true, DebugLifetime);
+			DrawDebugCapsule(GetWorld(), UpdatedComponent->GetComponentLocation(), CharacterOwner->GetSimpleCollisionHalfHeight(), CharacterOwner->GetSimpleCollisionRadius(), FQuat::Identity, FColor(100, 255, 100), false, DebugLifetime);
+			DrawDebugCapsule(GetWorld(), ClientLoc                    , CharacterOwner->GetSimpleCollisionHalfHeight(), CharacterOwner->GetSimpleCollisionRadius(), FQuat::Identity, FColor(255, 100, 100), false, DebugLifetime);
 		}
 #endif
 
@@ -9257,8 +9257,8 @@ void UCharacterMovementComponent::OnClientCorrectionReceived(FNetworkPredictionD
 		UE_LOG(LogNetPlayerMovement, Warning, TEXT("*** Client: Error for %s at Time=%.3f is %3.3f LocDiff(%s) ClientLoc(%s) ServerLoc(%s) NewBase: %s NewBone: %s ClientVel(%s) ServerVel(%s) SavedMoves %d"),
 			   *GetNameSafe(CharacterOwner), TimeStamp, LocDiff.Size(), *LocDiff.ToString(), *UpdatedComponent->GetComponentLocation().ToString(), *NewLocation.ToString(), *NewBaseString, *NewBaseBoneName.ToString(), *Velocity.ToString(), *NewVelocity.ToString(), ClientData.SavedMoves.Num());
 		const float DebugLifetime = CharacterMovementCVars::NetCorrectionLifetime;
-		DrawDebugCapsule(GetWorld(), UpdatedComponent->GetComponentLocation(), CharacterOwner->GetSimpleCollisionHalfHeight(), CharacterOwner->GetSimpleCollisionRadius(), FQuat::Identity, FColor(255, 100, 100), true, DebugLifetime);
-		DrawDebugCapsule(GetWorld(), NewLocation, CharacterOwner->GetSimpleCollisionHalfHeight(), CharacterOwner->GetSimpleCollisionRadius(), FQuat::Identity, FColor(100, 255, 100), true, DebugLifetime);
+		DrawDebugCapsule(GetWorld(), UpdatedComponent->GetComponentLocation(), CharacterOwner->GetSimpleCollisionHalfHeight(), CharacterOwner->GetSimpleCollisionRadius(), FQuat::Identity, FColor(255, 100, 100), false, DebugLifetime);
+		DrawDebugCapsule(GetWorld(), NewLocation, CharacterOwner->GetSimpleCollisionHalfHeight(), CharacterOwner->GetSimpleCollisionRadius(), FQuat::Identity, FColor(100, 255, 100), false, DebugLifetime);
 	}
 #endif //!UE_BUILD_SHIPPING
 
