Streamable animation asset and streaming animation support

diff --git a/Engine/Source/Developer/AssetTools/Private/AssetTools.cpp b/Engine/Source/Developer/AssetTools/Private/AssetTools.cpp
--- a/Engine/Source/Developer/AssetTools/Private/AssetTools.cpp
+++ b/Engine/Source/Developer/AssetTools/Private/AssetTools.cpp
@@ -38,6 +38,7 @@
 #include "AssetTypeActions/AssetTypeActions_AnimationAsset.h"
 #include "AssetTypeActions/AssetTypeActions_AnimBlueprint.h"
 #include "AssetTypeActions/AssetTypeActions_AnimComposite.h"
+#include "AssetTypeActions/AssetTypeActions_AnimStreamable.h"
 #include "AssetTypeActions/AssetTypeActions_AnimCurveCompressionSettings.h"
 #include "AssetTypeActions/AssetTypeActions_AnimMontage.h"
 #include "AssetTypeActions/AssetTypeActions_AnimSequence.h"
@@ -166,6 +167,7 @@ UAssetToolsImpl::UAssetToolsImpl(const FObjectInitializer& ObjectInitializer)
 	RegisterAssetTypeActions(MakeShareable(new FAssetTypeActions_AnimationAsset));
 	RegisterAssetTypeActions(MakeShareable(new FAssetTypeActions_AnimBlueprint));
 	RegisterAssetTypeActions(MakeShareable(new FAssetTypeActions_AnimComposite));
+	RegisterAssetTypeActions(MakeShareable(new FAssetTypeActions_AnimStreamable));
 	RegisterAssetTypeActions(MakeShareable(new FAssetTypeActions_AnimCurveCompressionSettings));
 	RegisterAssetTypeActions(MakeShareable(new FAssetTypeActions_AnimMontage));
 	RegisterAssetTypeActions(MakeShareable(new FAssetTypeActions_AnimSequence));
diff --git a/Engine/Source/Developer/AssetTools/Private/AssetTypeActions/AssetTypeActions_AnimSequence.cpp b/Engine/Source/Developer/AssetTools/Private/AssetTypeActions/AssetTypeActions_AnimSequence.cpp
--- a/Engine/Source/Developer/AssetTools/Private/AssetTypeActions/AssetTypeActions_AnimSequence.cpp
+++ b/Engine/Source/Developer/AssetTools/Private/AssetTypeActions/AssetTypeActions_AnimSequence.cpp
@@ -7,10 +7,12 @@
 #include "EditorReimportHandler.h"
 #include "Animation/AnimMontage.h"
 #include "Factories/AnimCompositeFactory.h"
+#include "Factories/AnimStreamableFactory.h"
 #include "Factories/AnimMontageFactory.h"
 #include "Factories/PoseAssetFactory.h"
 #include "EditorFramework/AssetImportData.h"
 #include "Animation/AnimComposite.h"
+#include "Animation/AnimStreamable.h"
 #include "Animation/PoseAsset.h"
 #include "AssetTools.h"
 #include "IContentBrowserSingleton.h"
@@ -77,6 +79,16 @@ void FAssetTypeActions_AnimSequence::FillCreateMenu(FMenuBuilder& MenuBuilder, c
 			)
 		);
 
+	MenuBuilder.AddMenuEntry(
+		LOCTEXT("AnimSequence_NewAnimStreamable", "Create AnimStreamable"),
+		LOCTEXT("AnimSequence_NewAnimStreamableTooltip", "Creates an AnimStreamable using the selected anim sequence."),
+		FSlateIcon(FEditorStyle::GetStyleSetName(), "ClassIcon.AnimMontage"),
+		FUIAction(
+			FExecuteAction::CreateSP(this, &FAssetTypeActions_AnimSequence::ExecuteNewAnimStreamable, Sequences),
+			FCanExecuteAction()
+		)
+	);
+
 	MenuBuilder.AddMenuEntry(
 		LOCTEXT("AnimSequence_NewPoseAsset", "Create PoseAsset"),
 		LOCTEXT("AnimSequence_NewPoseAssetTooltip", "Creates an PoseAsset using the selected anim sequence."),
@@ -138,6 +150,21 @@ void FAssetTypeActions_AnimSequence::ExecuteNewAnimMontage(TArray<TWeakObjectPtr
 	CreateAnimationAssets(Objects, UAnimMontage::StaticClass(), Factory, DefaultSuffix, FOnConfigureFactory::CreateSP(this, &FAssetTypeActions_AnimSequence::ConfigureFactoryForAnimMontage));
 }
 
+void FAssetTypeActions_AnimSequence::ExecuteNewAnimStreamable(TArray<TWeakObjectPtr<UAnimSequence>> Objects) const
+{
+	const FString DefaultSuffix = TEXT("_Streamable");
+	UAnimStreamableFactory* Factory = NewObject<UAnimStreamableFactory>();
+
+	auto StreamableConfigure = [](UFactory* AssetFactory, UAnimSequence* SourceAnimation) -> bool
+	{
+		UAnimStreamableFactory* StreamableAnimFactory = CastChecked<UAnimStreamableFactory>(AssetFactory);
+		StreamableAnimFactory->SourceAnimation = SourceAnimation;
+		return true;
+	};
+
+	CreateAnimationAssets(Objects, UAnimStreamable::StaticClass(), Factory, DefaultSuffix, FOnConfigureFactory::CreateLambda(StreamableConfigure));
+}
+
 void FAssetTypeActions_AnimSequence::ExecuteNewPoseAsset(TArray<TWeakObjectPtr<UAnimSequence>> Objects) const
 {
 	const FString DefaultSuffix = TEXT("_PoseAsset");
diff --git a/Engine/Source/Developer/AssetTools/Private/AssetTypeActions/AssetTypeActions_AnimSequence.h b/Engine/Source/Developer/AssetTools/Private/AssetTypeActions/AssetTypeActions_AnimSequence.h
--- a/Engine/Source/Developer/AssetTools/Private/AssetTypeActions/AssetTypeActions_AnimSequence.h
+++ b/Engine/Source/Developer/AssetTools/Private/AssetTypeActions/AssetTypeActions_AnimSequence.h
@@ -35,6 +35,9 @@ private:
 	/** Handler for when Create AnimMontage is selected */
 	void ExecuteNewAnimMontage(TArray<TWeakObjectPtr<UAnimSequence>> Objects) const;
 
+	/** Handler for when Create AnimStreamable is selected */
+	void ExecuteNewAnimStreamable(TArray<TWeakObjectPtr<UAnimSequence>> Objects) const;
+
 	/** Handler for when Create Pose Asset is selected */
 	void ExecuteNewPoseAsset(TArray<TWeakObjectPtr<UAnimSequence>> Objects) const;
 
diff --git a/Engine/Source/Developer/AssetTools/Private/AssetTypeActions/AssetTypeActions_AnimStreamable.h b/Engine/Source/Developer/AssetTools/Private/AssetTypeActions/AssetTypeActions_AnimStreamable.h
new file mode 100644
--- /dev/null
+++ b/Engine/Source/Developer/AssetTools/Private/AssetTypeActions/AssetTypeActions_AnimStreamable.h
@@ -0,0 +1,80 @@
+// Copyright 1998-2019 Epic Games, Inc. All Rights Reserved.
+
+#pragma once
+
+#include "CoreMinimal.h"
+#include "AssetTypeActions/AssetTypeActions_AnimationAsset.h"
+#include "Animation/AnimStreamable.h"
+#include "Interfaces/ITargetPlatform.h"
+#include "Interfaces/ITargetPlatformManagerModule.h"
+
+class FAssetTypeActions_AnimStreamable : public FAssetTypeActions_AnimationAsset
+{
+public:
+	// IAssetTypeActions Implementation
+	virtual FText GetName() const override { return NSLOCTEXT("AssetTypeActions", "AssetTypeActions_AnimStreamable", "Streamable Animation"); }
+	virtual FColor GetTypeColor() const override { return FColor(181,230,29); }
+	virtual UClass* GetSupportedClass() const override { return UAnimStreamable::StaticClass(); }
+	virtual bool CanFilter() override { return true; }
+	virtual void GetActions(const TArray<UObject*>& InObjects, FMenuBuilder& MenuBuilder) override;
+};
+
+void FAssetTypeActions_AnimStreamable::GetActions(const TArray<UObject*>& InObjects, FMenuBuilder& MenuBuilder)
+{
+	/*auto Sequences = GetTypedWeakObjectPtrs<UAnimStreamable>(InObjects);
+
+	auto Logger = [Sequences]()
+	{
+		uint32 TotalCompressedData = 0;
+		uint32 TotalZeroChunks = 0;
+		uint32 TotalMaxStreaming = 0;
+
+		ITargetPlatformManagerModule* TPM = GetTargetPlatformManager();
+		check(TPM)
+		ITargetPlatform* Platform = TPM->GetRunningTargetPlatform();
+		check(Platform);
+
+		FPlatformMisc::LowLevelOutputDebugStringf(TEXT("EVENT Streaming Memory Breakdown:"));
+		for (TWeakObjectPtr<UAnimStreamable> Streamable : Sequences)
+		{
+			const FStreamableAnimPlatformData& StreamingData = Streamable->GetStreamingAnimPlatformData(Platform);
+
+			FPlatformMisc::LowLevelOutputDebugStringf(TEXT("\tAnim : %s Chunks:%i ChunkTime:%.3f\n"), *Streamable->GetName(), StreamingData.Chunks.Num(), Streamable->GetChunkSizeSeconds(Platform));
+
+			int32 MaxConsecutiveChunks = 0;
+			for (int32 ChunkIndex = 0; ChunkIndex < StreamingData.Chunks.Num(); ++ChunkIndex)
+			{
+				const int32 NextChunkIndex = (++MaxConsecutiveChunks) % StreamingData.Chunks.Num();
+				const int32 ChunkSize = StreamingData.Chunks[ChunkIndex].GetMemorySize();
+				
+				FPlatformMisc::LowLevelOutputDebugStringf(TEXT("\t\tChunk %i ) %i Bytes\n"), ChunkIndex, ChunkSize);
+
+				if (ChunkIndex > 0)
+				{
+					const int32 NextChunkSize = StreamingData.Chunks[NextChunkIndex].GetMemorySize();
+					const int32 ConsecutiveSize = ChunkSize + NextChunkSize;
+					MaxConsecutiveChunks = FMath::Max(MaxConsecutiveChunks, ConsecutiveSize);
+				}
+				else
+				{
+					TotalZeroChunks += ChunkSize;
+				}
+				TotalCompressedData += ChunkSize;
+			}
+
+			TotalMaxStreaming += MaxConsecutiveChunks;
+		}
+		TotalMaxStreaming += TotalZeroChunks;
+
+		FPlatformMisc::LowLevelOutputDebugStringf(TEXT("Total Compressed Data: %.2f Kb\n"), (double)TotalCompressedData / 1000.0);
+		FPlatformMisc::LowLevelOutputDebugStringf(TEXT("Total Zero Chunks Size: %.2f Kb\n"), (double)TotalZeroChunks / 1000.0);
+		FPlatformMisc::LowLevelOutputDebugStringf(TEXT("Max Streaming Size: %.2f Kb\n"), (double)TotalMaxStreaming / 1000.0);
+	};
+
+	MenuBuilder.AddMenuEntry(
+		NSLOCTEXT("AnimStreamable", "AnimSequence_ReimportWithNewSource", "LOGGER"),
+		NSLOCTEXT("AnimStreamable", "AnimSequence_ReimportWithNewSourceTooltip", "Reimport the selected sequence(s) from a new source file."),
+		FSlateIcon(FEditorStyle::GetStyleSetName(), "Persona.AssetActions.ReimportAnim"),
+		FUIAction(FExecuteAction::CreateLambda(Logger))
+	);*/
+}
\ No newline at end of file
diff --git a/Engine/Source/Editor/Persona/Private/PersonaModule.cpp b/Engine/Source/Editor/Persona/Private/PersonaModule.cpp
--- a/Engine/Source/Editor/Persona/Private/PersonaModule.cpp
+++ b/Engine/Source/Editor/Persona/Private/PersonaModule.cpp
@@ -32,6 +32,8 @@
 #include "SSequenceEditor.h"
 #include "Animation/AnimComposite.h"
 #include "SAnimCompositeEditor.h"
+#include "Animation/AnimStreamable.h"
+#include "SAnimStreamableEditor.h"
 #include "Animation/PoseAsset.h"
 #include "SPoseEditor.h"
 #include "Animation/BlendSpace.h"
@@ -376,6 +378,13 @@ TSharedRef<SWidget> FPersonaModule::CreateEditorWidgetForAnimDocument(const TSha
 
 			OutDocumentLink = TEXT("Engine/Animation/AnimMontage");
 		}
+		else if (UAnimStreamable* StreamableAnim = Cast<UAnimStreamable>(InAnimAsset))
+		{
+			Result = SNew(SAnimStreamableEditor, InArgs.PreviewScene.Pin().ToSharedRef(), InArgs.EditableSkeleton.Pin().ToSharedRef())
+				.StreamableAnim(StreamableAnim);
+
+			OutDocumentLink = TEXT("Engine/Animation/Sequences");
+		}
 		else if (UPoseAsset* PoseAsset = Cast<UPoseAsset>(InAnimAsset))
 		{
 			Result = SNew(SPoseEditor, InArgs.PersonaToolkit.Pin().ToSharedRef(), InArgs.EditableSkeleton.Pin().ToSharedRef(), InArgs.PreviewScene.Pin().ToSharedRef())
diff --git a/Engine/Source/Editor/Persona/Private/SAnimStreamableEditor.cpp b/Engine/Source/Editor/Persona/Private/SAnimStreamableEditor.cpp
new file mode 100644
--- /dev/null
+++ b/Engine/Source/Editor/Persona/Private/SAnimStreamableEditor.cpp
@@ -0,0 +1,157 @@
+// Copyright 1998-2019 Epic Games, Inc. All Rights Reserved.
+
+
+#include "SAnimStreamableEditor.h"
+#include "Animation/EditorAnimBaseObj.h"
+#include "IDocumentation.h"
+
+#include "SAnimNotifyPanel.h"
+#include "Editor.h"
+
+//////////////////////////////////////////////////////////////////////////
+// SAnimCompositeEditor
+
+TSharedRef<SWidget> SAnimStreamableEditor::CreateDocumentAnchor()
+{
+	return IDocumentation::Get()->CreateAnchor(TEXT("Engine/Animation/AnimationComposite")); // UPDATE TO STREAMABLE
+}
+
+void SAnimStreamableEditor::Construct(const FArguments& InArgs, const TSharedRef<class IPersonaPreviewScene>& InPreviewScene, const TSharedRef<class IEditableSkeleton>& InEditableSkeleton)
+{
+	bIsActiveTimerRegistered = false;
+	StreamableAnim = InArgs._StreamableAnim;
+	check(StreamableAnim);
+
+	SAnimEditorBase::Construct( SAnimEditorBase::FArguments()
+		.OnObjectsSelected(InArgs._OnObjectsSelected), 
+		InPreviewScene );
+
+	if(GEditor)
+	{
+		GEditor->RegisterForUndo(this);
+	}
+
+	/*EditorPanels->AddSlot()
+		.AutoHeight()
+		.Padding(0, 10)
+		[
+			SAssignNew( AnimCompositePanel, SAnimCompositePanel )
+			.Composite(CompositeObj)
+			.CompositeEditor(SharedThis(this))
+			.WidgetWidth(S2ColumnWidget::DEFAULT_RIGHT_COLUMN_WIDTH)
+			.ViewInputMin(this, &SAnimEditorBase::GetViewMinInput)
+			.ViewInputMax(this, &SAnimEditorBase::GetViewMaxInput)
+			.OnSetInputViewRange(this, &SAnimEditorBase::SetInputViewRange)
+		];*/
+
+	EditorPanels->AddSlot()
+		.AutoHeight()
+		.Padding(0, 10)
+		[
+			SAssignNew( AnimNotifyPanel, SAnimNotifyPanel, InEditableSkeleton )
+			.Sequence(StreamableAnim)
+			.WidgetWidth(S2ColumnWidget::DEFAULT_RIGHT_COLUMN_WIDTH)
+			.InputMin(this, &SAnimEditorBase::GetMinInput)
+			.InputMax(this, &SAnimEditorBase::GetMaxInput)
+			.ViewInputMin(this, &SAnimEditorBase::GetViewMinInput)
+			.ViewInputMax(this, &SAnimEditorBase::GetViewMaxInput)
+			.OnSetInputViewRange(this, &SAnimEditorBase::SetInputViewRange)
+			.OnGetScrubValue(this, &SAnimEditorBase::GetScrubValue)
+			.OnSelectionChanged(this, &SAnimEditorBase::OnSelectionChanged)
+			.OnInvokeTab(InArgs._OnInvokeTab)
+		];
+
+	EditorPanels->AddSlot()
+		.AutoHeight()
+		.Padding(0, 10)
+		[
+			SAssignNew( AnimCurvePanel, SAnimCurvePanel, InEditableSkeleton)
+			.Sequence(StreamableAnim)
+			.WidgetWidth(S2ColumnWidget::DEFAULT_RIGHT_COLUMN_WIDTH)
+			.ViewInputMin(this, &SAnimEditorBase::GetViewMinInput)
+			.ViewInputMax(this, &SAnimEditorBase::GetViewMaxInput)
+			.InputMin(this, &SAnimEditorBase::GetMinInput)
+			.InputMax(this, &SAnimEditorBase::GetMaxInput)
+			.OnSetInputViewRange(this, &SAnimEditorBase::SetInputViewRange)
+			.OnGetScrubValue(this, &SAnimEditorBase::GetScrubValue)
+		];
+}
+
+SAnimStreamableEditor::~SAnimStreamableEditor()
+{
+	if(GEditor)
+	{
+		GEditor->UnregisterForUndo(this);
+	}
+}
+
+void SAnimStreamableEditor::PreAnimUpdate()
+{
+	StreamableAnim->Modify();
+}
+
+void SAnimStreamableEditor::PostAnimUpdate()
+{
+	StreamableAnim->MarkPackageDirty();
+	//SortAndUpdateComposite();
+}
+
+void SAnimStreamableEditor::RebuildPanel()
+{
+	//SortAndUpdateComposite();
+	//AnimCompositePanel->Update();
+}
+
+void SAnimStreamableEditor::OnStreamableChange(class UObject *EditorAnimBaseObj, bool bRebuild)
+{
+	if (StreamableAnim != nullptr )
+	{
+		if(bRebuild && !bIsActiveTimerRegistered)
+		{
+			// sometimes crashes because the timer delay but animation still renders, so invalidating here before calling timer
+			bIsActiveTimerRegistered = true;
+			RegisterActiveTimer(0.f, FWidgetActiveTimerDelegate::CreateSP(this, &SAnimStreamableEditor::TriggerRebuildPanel));
+		} 
+
+		StreamableAnim->MarkPackageDirty();
+	}
+}
+
+void SAnimStreamableEditor::PostUndo( bool bSuccess )
+{
+	PostUndoRedo();
+}
+
+void SAnimStreamableEditor::PostRedo( bool bSuccess )
+{
+	PostUndoRedo();
+}
+
+void SAnimStreamableEditor::PostUndoRedo()
+{
+	if (!bIsActiveTimerRegistered)
+	{
+		bIsActiveTimerRegistered = true;
+		RegisterActiveTimer(0.f, FWidgetActiveTimerDelegate::CreateSP(this, &SAnimStreamableEditor::TriggerRebuildPanel));
+	}
+}
+
+void SAnimStreamableEditor::InitDetailsViewEditorObject(UEditorAnimBaseObj* EdObj)
+{
+	EdObj->InitFromAnim(StreamableAnim, FOnAnimObjectChange::CreateSP( SharedThis(this), &SAnimStreamableEditor::OnStreamableChange ));
+}
+
+EActiveTimerReturnType SAnimStreamableEditor::TriggerRebuildPanel(double InCurrentTime, float InDeltaTime)
+{
+	// we should not update any property related within PostEditChange, 
+	// so this is deferred to Tick, when it needs to rebuild, just mark it and this will update in the next tick
+	RebuildPanel();
+
+	bIsActiveTimerRegistered = false;
+	return EActiveTimerReturnType::Stop;
+}
+
+float SAnimStreamableEditor::CalculateSequenceLengthOfEditorObject() const
+{
+	return StreamableAnim->SequenceLength;
+}
diff --git a/Engine/Source/Editor/Persona/Private/SAnimStreamableEditor.h b/Engine/Source/Editor/Persona/Private/SAnimStreamableEditor.h
new file mode 100644
--- /dev/null
+++ b/Engine/Source/Editor/Persona/Private/SAnimStreamableEditor.h
@@ -0,0 +1,87 @@
+// Copyright 1998-2019 Epic Games, Inc. All Rights Reserved.
+
+
+#pragma once
+
+#include "CoreMinimal.h"
+#include "Widgets/SWidget.h"
+#include "Widgets/DeclarativeSyntaxSupport.h"
+#include "SAnimCurvePanel.h"
+#include "SAnimEditorBase.h"
+#include "Animation/AnimStreamable.h"
+
+class SAnimNotifyPanel;
+
+//////////////////////////////////////////////////////////////////////////
+// SAnimCompositeEditor
+
+/** Overall animation composite editing widget. This mostly contains functions for editing the UAnimComposite.
+
+	SAnimCompositeEditor will create the SAnimCompositePanel which is mostly responsible for setting up the UI 
+	portion of the composite tool and registering callbacks to the SAnimCompositeEditor to do the actual editing.
+	
+*/
+class SAnimStreamableEditor : public SAnimEditorBase, public FEditorUndoClient
+{
+public:
+	SLATE_BEGIN_ARGS(SAnimStreamableEditor)
+		: _StreamableAnim(NULL)
+		{}
+
+		SLATE_ARGUMENT( UAnimStreamable*, StreamableAnim)
+		SLATE_EVENT(FOnObjectsSelected, OnObjectsSelected)
+		SLATE_EVENT(FOnInvokeTab, OnInvokeTab)
+
+	SLATE_END_ARGS()
+
+private:
+	/** Slate editor panels */
+	TSharedPtr<class SAnimNotifyPanel> AnimNotifyPanel;
+	TSharedPtr<class SAnimCurvePanel>	AnimCurvePanel;
+
+	/** do I have to update the panel **/
+	bool bRebuildPanel;
+	void RebuildPanel();
+
+	/** Handler for when composite is edited in details view */
+	void OnStreamableChange(class UObject *EditorAnimBaseObj, bool bRebuild);
+
+protected:
+	virtual void InitDetailsViewEditorObject(class UEditorAnimBaseObj* EdObj) override;
+
+public:
+	void Construct(const FArguments& InArgs, const TSharedRef<class IPersonaPreviewScene>& InPreviewScene, const TSharedRef<class IEditableSkeleton>& InEditableSkeleton);
+
+	~SAnimStreamableEditor();
+
+	/** Return the animation composite being edited */
+	UAnimStreamable* GetStreamableAnimObj() const { return StreamableAnim; }
+	virtual UAnimationAsset* GetEditorObject() const override { return GetStreamableAnimObj(); }
+
+private:
+	/** Pointer to the animation composite being edited */
+	UAnimStreamable* StreamableAnim;
+	
+	/** FEditorUndoClient interface */
+	virtual void PostUndo( bool bSuccess ) override;
+	virtual void PostRedo( bool bSuccess ) override;
+	void PostUndoRedo();
+
+	virtual float CalculateSequenceLengthOfEditorObject() const override;
+
+	/** One-off active timer to trigger a panel rebuild */
+	EActiveTimerReturnType TriggerRebuildPanel(double InCurrentTime, float InDeltaTime);
+
+	/** Whether the active timer to trigger a panel rebuild is currently registered */
+	bool bIsActiveTimerRegistered;
+
+public:
+
+	/** delegate handlers for when the composite is being editor */
+	void			PreAnimUpdate();
+	void			PostAnimUpdate();
+
+	//~ Begin SAnimEditorBase Interface
+	virtual TSharedRef<SWidget> CreateDocumentAnchor() override;
+	//~ End SAnimEditorBase Interface
+};
diff --git a/Engine/Source/Editor/UnrealEd/Classes/Factories/AnimStreamableFactory.h b/Engine/Source/Editor/UnrealEd/Classes/Factories/AnimStreamableFactory.h
new file mode 100644
--- /dev/null
+++ b/Engine/Source/Editor/UnrealEd/Classes/Factories/AnimStreamableFactory.h
@@ -0,0 +1,34 @@
+// Copyright 1998-2019 Epic Games, Inc. All Rights Reserved.
+
+/**
+ *
+ */
+
+#pragma once
+
+#include "CoreMinimal.h"
+#include "UObject/ObjectMacros.h"
+#include "Factories/Factory.h"
+#include "AnimStreamableFactory.generated.h"
+
+struct FAssetData;
+class SWindow;
+
+UCLASS(HideCategories=Object,MinimalAPI)
+class UAnimStreamableFactory : public UFactory
+{
+	GENERATED_UCLASS_BODY()
+
+	UPROPERTY()
+	class USkeleton* TargetSkeleton;
+
+	/* Used when creating a composite from an AnimSequence, becomes the only AnimSequence contained */
+	UPROPERTY()
+	class UAnimSequence* SourceAnimation;
+
+	//~ Begin UFactory Interface
+	virtual bool ConfigureProperties() override;
+	virtual UObject* FactoryCreateNew(UClass* Class,UObject* InParent,FName Name,EObjectFlags Flags,UObject* Context,FFeedbackContext* Warn) override;
+	//~ Begin UFactory Interface	
+};
+
diff --git a/Engine/Source/Editor/UnrealEd/Private/Factories/AnimStreamableFactory.cpp b/Engine/Source/Editor/UnrealEd/Private/Factories/AnimStreamableFactory.cpp
new file mode 100644
--- /dev/null
+++ b/Engine/Source/Editor/UnrealEd/Private/Factories/AnimStreamableFactory.cpp
@@ -0,0 +1,83 @@
+// Copyright 1998-2019 Epic Games, Inc. All Rights Reserved.
+
+/*=============================================================================
+	AnimCompositeFactory.cpp: Factory for AnimComposite
+=============================================================================*/
+
+#include "Factories/AnimStreamableFactory.h"
+#include "Modules/ModuleManager.h"
+#include "Widgets/DeclarativeSyntaxSupport.h"
+#include "Widgets/SWindow.h"
+#include "Widgets/Layout/SBorder.h"
+#include "EditorStyleSet.h"
+#include "Animation/Skeleton.h"
+#include "Animation/AnimSequence.h"
+#include "Animation/AnimStreamable.h"
+#include "Editor.h"
+
+#include "IContentBrowserSingleton.h"
+#include "ContentBrowserModule.h"
+
+#define LOCTEXT_NAMESPACE "AnimStreamableFactory"
+
+UAnimStreamableFactory::UAnimStreamableFactory(const FObjectInitializer& ObjectInitializer)
+	: Super(ObjectInitializer)
+{
+	bCreateNew = false;
+	SupportedClass = UAnimStreamable::StaticClass();
+}
+
+bool UAnimStreamableFactory::ConfigureProperties()
+{
+	// Null the skeleton so we can check for selection later
+	/*TargetSkeleton = NULL;
+	SourceAnimation = NULL;
+
+	// Load the content browser module to display an asset picker
+	FContentBrowserModule& ContentBrowserModule = FModuleManager::LoadModuleChecked<FContentBrowserModule>("ContentBrowser");
+
+	FAssetPickerConfig AssetPickerConfig;
+	
+	AssetPickerConfig.Filter.ClassNames.Add(USkeleton::StaticClass()->GetFName());
+	AssetPickerConfig.Filter.bRecursiveClasses = true;
+
+	// The delegate that fires when an asset was selected
+	AssetPickerConfig.OnAssetSelected = FOnAssetSelected::CreateUObject(this, &UAnimCompositeFactory::OnTargetSkeletonSelected);
+
+	//The default view mode should be a list view
+	AssetPickerConfig.InitialAssetViewType = EAssetViewType::List;
+
+	PickerWindow = SNew(SWindow)
+	.Title(LOCTEXT("CreateAnimCompositeOptions", "Pick Skeleton"))
+	.ClientSize(FVector2D(500, 600))
+	.SupportsMinimize(false) .SupportsMaximize(false)
+	[
+		SNew(SBorder)
+		.BorderImage( FEditorStyle::GetBrush("Menu.Background") )
+		[
+			ContentBrowserModule.Get().CreateAssetPicker(AssetPickerConfig)
+		]
+	];
+
+	GEditor->EditorAddModalWindow(PickerWindow.ToSharedRef());
+	PickerWindow.Reset();
+
+	return TargetSkeleton != NULL;*/
+	return true;
+}
+
+UObject* UAnimStreamableFactory::FactoryCreateNew(UClass* Class, UObject* InParent, FName Name, EObjectFlags Flags, UObject* Context, FFeedbackContext* Warn)
+{
+	if (SourceAnimation)
+	{
+		UAnimStreamable* StreamableAnim = NewObject<UAnimStreamable>(InParent, Class, Name, Flags);
+
+		StreamableAnim->InitFrom(SourceAnimation);
+		
+		return StreamableAnim;
+	}
+
+	return NULL;
+}
+
+#undef LOCTEXT_NAMESPACE
diff --git a/Engine/Source/Runtime/Core/Public/Serialization/MemoryReader.h b/Engine/Source/Runtime/Core/Public/Serialization/MemoryReader.h
--- a/Engine/Source/Runtime/Core/Public/Serialization/MemoryReader.h
+++ b/Engine/Source/Runtime/Core/Public/Serialization/MemoryReader.h
@@ -7,6 +7,7 @@
 #include "Math/UnrealMathUtility.h"
 #include "Containers/UnrealString.h"
 #include "Serialization/MemoryArchive.h"
+#include "Containers/ArrayView.h"
 
 /**
  * Archive for reading arbitrary data from the specified memory location
@@ -66,3 +67,61 @@ private:
 	int64                LimitSize;
 };
 
+/**
+ * Archive for reading arbitrary data from the specified array view
+ */
+class FMemoryReaderView : public FMemoryArchive
+{
+public:
+	/**
+  	 * Returns the name of the Archive.  Useful for getting the name of the package a struct or object
+	 * is in when a loading error occurs.
+	 *
+	 * This is overridden for the specific Archive Types
+	 **/
+	virtual FString GetArchiveName() const override
+	{
+		return TEXT("FMemoryReaderView");
+	}
+
+	virtual int64 TotalSize() override
+	{
+		return FMath::Min((int64)Bytes.Num(), LimitSize);
+	}
+
+	void Serialize( void* Data, int64 Num )
+	{
+		if (Num && !ArIsError)
+		{
+			// Only serialize if we have the requested amount of data
+			if (Offset + Num <= TotalSize())
+			{
+				FMemory::Memcpy( Data, &Bytes[Offset], Num );
+				Offset += Num;
+			}
+			else
+			{
+				ArIsError = true;
+			}
+		}
+	}
+
+	explicit FMemoryReaderView( TArrayView<const uint8> InBytes, bool bIsPersistent = false )
+		: Bytes    (InBytes)
+		, LimitSize(INT64_MAX)
+	{
+		this->SetIsLoading(true);
+		this->SetIsPersistent(bIsPersistent);
+	}
+
+	/** With this method it's possible to attach data behind some serialized data. */
+	void SetLimitSize(int64 NewLimitSize)
+	{
+		LimitSize = NewLimitSize;
+	}
+
+private:
+	TArrayView<const uint8> Bytes;
+	int64                LimitSize;
+};
+
diff --git a/Engine/Source/Runtime/Engine/Classes/Animation/AnimSequence.h b/Engine/Source/Runtime/Engine/Classes/Animation/AnimSequence.h
--- a/Engine/Source/Runtime/Engine/Classes/Animation/AnimSequence.h
+++ b/Engine/Source/Runtime/Engine/Classes/Animation/AnimSequence.h
@@ -448,6 +448,7 @@ public:
 		bUseRawDataOnly = true;
 		RawDataGuid = bForceNewRawDatGuid ? FGuid::NewGuid() : GenerateGuidFromRawData();
 		FlagDependentAnimationsAsRawDataOnly();
+		UpdateDependentStreamingAnimations();
 	}
 #endif
 	//~ End UAnimSequenceBase Interface
@@ -506,16 +507,6 @@ public:
 	const FRawAnimSequenceTrack& GetRawAnimationTrack(int32 TrackIndex) const { return RawAnimationData[TrackIndex]; }
 
 private:
-	/** 
-	 * Utility function.
-	 * When extracting root motion, this function resets the root bone to the first frame of the animation.
-	 * 
-	 * @param	BoneTransform		Root Bone Transform
-	 * @param	RequiredBones		BoneContainer
-	 * @param	ExtractionContext	Extraction Context to access root motion extraction information.
-	 */
-	void ResetRootBoneForRootMotion(FTransform& BoneTransform, const FBoneContainer & RequiredBones, ERootMotionRootLock::Type InRootMotionRootLock) const;
-
 	/**
 	* Retarget a single bone transform, to apply right after extraction.
 	*
@@ -820,6 +811,9 @@ public:
 	// Helper function to allow us to notify animations that depend on us that they need to update
 	void FlagDependentAnimationsAsRawDataOnly() const;
 
+	// Helper function to allow us to update streaming animations that depend on us with our data when we are updated
+	void UpdateDependentStreamingAnimations() const;
+
 	// Generate a GUID from a hash of our own raw data
 	FGuid GenerateGuidFromRawData() const;
 
@@ -888,17 +882,9 @@ private:
 	/** Take a set of marker positions and validates them against a requested start position, updating them as desired */
 	void ValidateCurrentPosition(const FMarkerSyncAnimPosition& Position, bool bPlayingForwards, bool bLooping, float&CurrentTime, FMarkerPair& PreviousMarker, FMarkerPair& NextMarker) const;
 	bool UseRawDataForPoseExtraction(const FBoneContainer& RequiredBones) const;
-	void UpdateSHAWithCurves(FSHA1& Sha, const FRawCurveTracks& RawCurveData) const;
 	// Should we be always using our raw data (i.e is our compressed data stale)
 	bool bUseRawDataOnly;
 
-	// Populate InOutPose based on raw animation data. 
-	void BuildPoseFromRawData(const TArray<FRawAnimSequenceTrack>& InAnimationData, FCompactPose& InOutPose, float InTime) const;
-	
-	// Internal functionality used by BuildPoseFromRawData. Template param specifies whether KeyIndex2 and Alpha are expected to be used.
-	template<bool INTERPOLATE>
-	void BuildPoseFromRawDataInternal(const TArray<FRawAnimSequenceTrack>& InAnimationData, FCompactPose& InOutPose, int32 KeyIndex1, int32 KeyIndex2, float Alpha) const;
-
 public:
 	// Are we currently compressing this animation
 	bool bCompressionInProgress;
diff --git a/Engine/Source/Runtime/Engine/Classes/Animation/AnimStreamable.h b/Engine/Source/Runtime/Engine/Classes/Animation/AnimStreamable.h
new file mode 100644
--- /dev/null
+++ b/Engine/Source/Runtime/Engine/Classes/Animation/AnimStreamable.h
@@ -0,0 +1,233 @@
+// Copyright 1998-2019 Epic Games, Inc. All Rights Reserved.
+
+/**
+ * Animation that can be streamed instead of being loaded completely
+ *
+ */
+
+#pragma once
+
+#include "CoreMinimal.h"
+#include "UObject/ObjectMacros.h"
+#include "UObject/Object.h"
+#include "Misc/Guid.h"
+#include "Serialization/BulkData.h"
+#include "Animation/AnimTypes.h"
+#include "Animation/AnimCompress.h"
+#include "Animation/AnimCompressionTypes.h"
+#include "Animation/AnimSequenceBase.h"
+#include "AnimStreamable.generated.h"
+
+class UAnimSequence;
+class UAnimCompress;
+struct FCompactPose;
+
+class FAnimStreamableChunk
+{
+public:
+	FAnimStreamableChunk() : StartTime(0.f), SequenceLength(0.f), CompressedAnimSequence(nullptr) {}
+	~FAnimStreamableChunk()
+	{
+		if (CompressedAnimSequence)
+		{
+			delete CompressedAnimSequence;
+			CompressedAnimSequence = nullptr;
+		}
+	}
+
+	float StartTime;
+
+	float SequenceLength;
+
+	// Compressed Data for this chunk (if nullptr then data needs to be loaded via BulkData)
+	FCompressedAnimSequence* CompressedAnimSequence;
+
+	// Bulk data if stored in the package.
+	FByteBulkData BulkData;
+
+	int32 GetMemorySize() const
+	{
+		return	sizeof(FAnimStreamableChunk)
+			+ sizeof(FCompressedAnimSequence)
+			+ CompressedAnimSequence->CompressedTrackToSkeletonMapTable.GetAllocatedSize()
+			+ CompressedAnimSequence->CompressedCurveNames.GetAllocatedSize()
+			+ CompressedAnimSequence->CompressedCurveByteStream.GetAllocatedSize()
+			+ CompressedAnimSequence->CompressedDataStructure.GetApproxBoneCompressedSize();
+	}
+
+	/** Serialization. */
+	void Serialize(FArchive& Ar, UAnimStreamable* Owner, int32 ChunkIndex);
+};
+
+class FStreamableAnimPlatformData
+{
+public:
+	TArray<FAnimStreamableChunk> Chunks;
+
+	void Serialize(FArchive& Ar, class UAnimStreamable* Owner);
+
+	void Reset()
+	{
+		Chunks.Reset();
+	}
+};
+
+
+UCLASS(config=Engine, hidecategories=UObject, MinimalAPI, BlueprintType)
+class UAnimStreamable : public UAnimSequenceBase
+{
+	GENERATED_UCLASS_BODY()
+
+public:
+	/** Number of raw frames in this sequence (not used by engine - just for informational purposes). */
+	UPROPERTY(AssetRegistrySearchable, meta = (DisplayName = "Number of Keys"))
+	int32 NumFrames;
+
+	/** This defines how values between keys are calculated **/
+	UPROPERTY(EditAnywhere, AssetRegistrySearchable, Category = Animation)
+	EAnimInterpolationType Interpolation;
+
+	/** Base pose to use when retargeting */
+	UPROPERTY(EditAnywhere, AssetRegistrySearchable, Category = Animation)
+	FName RetargetSource;
+
+#if WITH_EDITORONLY_DATA
+
+	// Sequence the streamable was created from (used for reflecting changes to the source in editor)
+	UPROPERTY()
+	const UAnimSequence* SourceSequence;
+
+	/**
+	 * The compression scheme that was most recently used to compress this animation.
+	 */
+	UPROPERTY(Category = Compression, VisibleAnywhere)
+	class UAnimCompress* CompressionScheme;
+
+	UPROPERTY()
+	FGuid RawDataGuid;
+
+	/**
+	 * Raw uncompressed keyframe data.
+	 */
+	UPROPERTY()
+	TArray<struct FRawAnimSequenceTrack> RawAnimationData;
+
+	/**
+	 * In the future, maybe keeping RawAnimSequenceTrack + TrackMap as one would be good idea to avoid inconsistent array size
+	 * TrackToSkeletonMapTable(i) should contains  track mapping data for RawAnimationData(i).
+	 */
+	UPROPERTY()
+	TArray<struct FTrackToSkeletonMap> TrackToSkeletonMapTable;
+
+	/**
+	 * This is name of RawAnimationData tracks for editoronly - if we lose skeleton, we'll need relink them
+	 */
+	UPROPERTY()
+	TArray<FName> AnimationTrackNames;
+
+	// Editor can have multiple platforms loaded at once
+	TMap<const ITargetPlatform*, FStreamableAnimPlatformData*> StreamableAnimPlatformData;
+
+	FStreamableAnimPlatformData* RunningAnimPlatformData;
+#else
+
+	// Non editor only has one set of platform data
+	FStreamableAnimPlatformData RunningAnimPlatformData;
+#endif
+
+	bool HasRunningPlatformData() const
+	{
+#if WITH_EDITOR
+		return RunningAnimPlatformData != nullptr;
+#else
+		return true;
+#endif
+	}
+
+	FStreamableAnimPlatformData& GetRunningPlatformData()
+	{
+#if WITH_EDITOR
+		check(RunningAnimPlatformData);
+		return *RunningAnimPlatformData;
+#else
+		return RunningAnimPlatformData;
+#endif
+	}
+
+	const FStreamableAnimPlatformData& GetRunningPlatformData() const
+	{
+#if WITH_EDITOR
+		check(RunningAnimPlatformData);
+		return *RunningAnimPlatformData;
+#else
+		return RunningAnimPlatformData;
+#endif
+	}
+
+	/** The curve compression settings used to compress curves in this sequence. */
+	UPROPERTY(Category = Compression, EditAnywhere)
+	class UAnimCurveCompressionSettings* CurveCompressionSettings;
+
+	/** If this is on, it will allow extracting of root motion **/
+	UPROPERTY(EditAnywhere, AssetRegistrySearchable, Category = RootMotion, meta = (DisplayName = "EnableRootMotion"))
+	bool bEnableRootMotion;
+
+	/** Root Bone will be locked to that position when extracting root motion.**/
+	UPROPERTY(EditAnywhere, Category = RootMotion)
+	TEnumAsByte<ERootMotionRootLock::Type> RootMotionRootLock;
+
+	/** Force Root Bone Lock even if Root Motion is not enabled */
+	UPROPERTY(EditAnywhere, Category = RootMotion)
+	bool bForceRootLock;
+
+	/** If this is on, it will use a normalized scale value for the root motion extracted: FVector(1.0, 1.0, 1.0) **/
+	UPROPERTY(EditAnywhere, AssetRegistrySearchable, Category = RootMotion, meta = (DisplayName = "Use Normalized Root Motion Scale"))
+	bool bUseNormalizedRootMotionScale;
+
+	//~ Begin UObject Interface
+	virtual void PreSave(const class ITargetPlatform* TargetPlatform) override;
+	virtual void Serialize(FArchive& Ar) override;
+	virtual void PostLoad() override;
+	virtual void FinishDestroy() override;
+	//~ End UObject Interface
+
+	//~ Begin UAnimSequenceBase Interface
+	ENGINE_API virtual void HandleAssetPlayerTickedInternal(FAnimAssetTickContext &Context, const float PreviousTime, const float MoveDelta, const FAnimTickRecord &Instance, struct FAnimNotifyQueue& NotifyQueue) const override;
+	virtual void GetAnimationPose(FCompactPose& OutPose, FBlendedCurve& OutCurve, const FAnimExtractContext& ExtractionContext) const override;
+	virtual int32 GetNumberOfFrames() const override { return NumFrames; }
+	//virtual EAdditiveAnimationType GetAdditiveAnimType() const override;
+	//virtual bool IsValidAdditive() const override { return GetAdditiveAnimType() != AAT_None; }
+	//virtual bool HasRootMotion() const override;
+
+	//~ End UAnimSequenceBase Interface
+	//~ Begin UAnimSequence Interface
+#if WITH_EDITOR
+	//virtual class UAnimSequence* GetAdditiveBasePose() const override;
+	//virtual bool GetAllAnimationSequencesReferred(TArray<UAnimationAsset*>& AnimationAssets, bool bRecursive = true) override;
+	//virtual void ReplaceReferredAnimations(const TMap<UAnimationAsset*, UAnimationAsset*>& ReplacementMap) override;
+#endif
+	//~ End UAnimSequence Interface
+#if WITH_EDITOR
+	ENGINE_API void InitFrom(const UAnimSequence* InSourceSequence);
+#endif
+
+	ENGINE_API FStreamableAnimPlatformData& GetStreamingAnimPlatformData(const ITargetPlatform* Platform);
+
+	ENGINE_API float GetChunkSizeSeconds(const ITargetPlatform* Platform) const;
+
+	private:
+
+#if WITH_EDITOR
+	ENGINE_API void RequestCompressedData(const  ITargetPlatform* Platform=nullptr);
+
+	void UpdateRawData();
+
+	FString GetBaseDDCKey() const;
+
+	void RequestCompressedDataForChunk(const FString& ChunkDDCKey, FAnimStreamableChunk& Chunk, const uint32 FrameStart, const uint32 FrameEnd, TSharedRef<FAnimCompressContext> CompressContext);
+#endif
+
+	bool bUseRawDataOnly;
+	int32 GetChunkIndexForTime(const TArray<FAnimStreamableChunk>& Chunks, const float CurrentTime) const;
+};
+
diff --git a/Engine/Source/Runtime/Engine/Private/Animation/AnimCompress.cpp b/Engine/Source/Runtime/Engine/Private/Animation/AnimCompress.cpp
--- a/Engine/Source/Runtime/Engine/Private/Animation/AnimCompress.cpp
+++ b/Engine/Source/Runtime/Engine/Private/Animation/AnimCompress.cpp
@@ -1583,24 +1583,14 @@ void UAnimCompress::CoalesceCompressedSegments(UAnimSequence& AnimSeq, const TAr
 
 FString UAnimCompress::MakeDDCKey()
 {
-	FString Key;
-	TArray<uint8> TempBytes;
-	TempBytes.Reserve(64);
+	FArcToHexString ArcToHexString;
 
 	// Serialize the compression settings into a temporary array. The archive
 	// is flagged as persistent so that machines of different endianness produce
 	// identical binary results.
-	FMemoryWriter Ar(TempBytes, /*bIsPersistent=*/ true);
+	PopulateDDCKey(ArcToHexString.Ar);
 
-	PopulateDDCKey(Ar);
-
-	const uint8* SettingsAsBytes = TempBytes.GetData();
-	Key.Reserve(TempBytes.Num() + 1);
-	for (int32 ByteIndex = 0; ByteIndex < TempBytes.Num(); ++ByteIndex)
-	{
-		ByteToHex(SettingsAsBytes[ByteIndex], Key);
-	}
-	return Key;
+	return ArcToHexString.MakeString();
 }
 
 void UAnimCompress::PopulateDDCKey(FArchive& Ar)
diff --git a/Engine/Source/Runtime/Engine/Private/Animation/AnimCompressionDerivedData.cpp b/Engine/Source/Runtime/Engine/Private/Animation/AnimCompressionDerivedData.cpp
--- a/Engine/Source/Runtime/Engine/Private/Animation/AnimCompressionDerivedData.cpp
+++ b/Engine/Source/Runtime/Engine/Private/Animation/AnimCompressionDerivedData.cpp
@@ -14,124 +14,23 @@
 
 DECLARE_CYCLE_STAT(TEXT("Anim Compression (Derived Data)"), STAT_AnimCompressionDerivedData, STATGROUP_Anim);
 
-template<typename ArrayValue>
-void StripFramesEven(TArray<ArrayValue>& Keys, const int32 NumFrames)
-{
-	if (Keys.Num() > 1)
-	{
-		check(Keys.Num() == NumFrames);
-
-		for (int32 DstKey = 1, SrcKey = 2; SrcKey < NumFrames; ++DstKey, SrcKey+=2)
-		{
-			Keys[DstKey] = Keys[SrcKey];
-		}
-
-		const int32 HalfSize = (NumFrames - 1) / 2;
-		const int32 StartRemoval = HalfSize + 1;
-
-		Keys.RemoveAt(StartRemoval, NumFrames - StartRemoval);
-	}
-}
-
-template<typename ArrayValue>
-void StripFramesOdd(TArray<ArrayValue>& Keys, const int32 NumFrames)
-{
-	if (Keys.Num() > 1)
-	{
-		const int32 NewNumFrames = NumFrames / 2;
-
-		TArray<ArrayValue> NewKeys;
-		NewKeys.Reserve(NewNumFrames);
-
-		check(Keys.Num() == NumFrames);
-
-		NewKeys.Add(Keys[0]); //Always keep first 
-
-		//Always keep first and last
-		const int32 NumFramesToCalculate = NewNumFrames - 2;
-
-		// Frame increment is ratio of old frame spaces vs new frame spaces 
-		const double FrameIncrement = (double)(NumFrames - 1) / (double)(NewNumFrames - 1);
-			
-		for(int32 Frame=0; Frame < NumFramesToCalculate; ++Frame)
-		{
-			const double NextFramePosition = FrameIncrement * (Frame + 1);
-			const int32 Frame1 = (int32)NextFramePosition;
-			const float Alpha = (NextFramePosition - (double)Frame1);
-
-			NewKeys.Add(AnimationCompressionUtils::Interpolate(Keys[Frame1], Keys[Frame1 + 1], Alpha));
-
-		}
-
-		NewKeys.Add(Keys.Last()); // Always Keep Last
-
-		const int32 HalfSize = (NumFrames - 1) / 2;
-		const int32 StartRemoval = HalfSize + 1;
-
-		Keys = MoveTemp(NewKeys);
-	}
-}
-
-FDerivedDataAnimationCompression::FDerivedDataAnimationCompression(const FCompressibleAnimData& InDataToCompress, TSharedPtr<FAnimCompressContext> InCompressContext, int32 InPreviousCompressedSize, bool bInTryFrameStripping, bool bTryStrippingOnOddFramedAnims)
-	: DataToCompress(InDataToCompress)
+FDerivedDataAnimationCompression::FDerivedDataAnimationCompression(const TCHAR* InTypeName, const FString& InAssetDDCKey, TSharedPtr<FAnimCompressContext> InCompressContext, int32 InPreviousCompressedSize)
+	: TypeName(InTypeName)
+	, AssetDDCKey(InAssetDDCKey)
 	, CompressContext(InCompressContext)
 	, PreviousCompressedSize(InPreviousCompressedSize)
 {
-	check(DataToCompress.Skeleton != nullptr);
-
-	// Can only do stripping on animations that have an even number of frames once the end frame is removed)
-	bIsEvenFramed = ((DataToCompress.NumFrames - 1) % 2) == 0;
-	const bool bIsValidForStripping = bIsEvenFramed || bTryStrippingOnOddFramedAnims;
 
-	const bool bStripCandidate = (DataToCompress.NumFrames > 10) && bIsValidForStripping;
-	
-	bPerformStripping = bStripCandidate && bInTryFrameStripping;
 }
 
 FDerivedDataAnimationCompression::~FDerivedDataAnimationCompression()
 {
 }
 
-FString FDerivedDataAnimationCompression::GetPluginSpecificCacheKeySuffix() const
-{
-	enum { UE_ANIMCOMPRESSION_DERIVEDDATA_VER = 1 };
-
-	//Make up our content key consisting of:
-	//	* Our plugin version
-	//	* Global animation compression version
-	//	* Our raw data GUID
-	//	* Our skeleton GUID: If our skeleton changes our compressed data may now be stale
-	//	* Baked Additive Flag
-	//	* Additive ref pose GUID or hardcoded string if not available
-	//	* Compression Settings
-	//	* Curve compression settings
-
-	char AdditiveType = DataToCompress.bIsValidAdditive ? NibbleToTChar(DataToCompress.AdditiveAnimType) : '0';
-	char RefType = DataToCompress.bIsValidAdditive ? NibbleToTChar(DataToCompress.RefPoseType) : '0';
-
-	const int32 StripFrame = bPerformStripping ? 1 : 0;
-
-	FString Ret = FString::Printf(TEXT("%i_%i_%i_%i_%s%s%s_%c%c%i_%s_%s_%s"),
-		(int32)UE_ANIMCOMPRESSION_DERIVEDDATA_VER,
-		(int32)CURRENT_ANIMATION_ENCODING_PACKAGE_VERSION,
-		DataToCompress.CompressCommandletVersion,
-		StripFrame,
-		*DataToCompress.RawDataGuid.ToString(),
-		*DataToCompress.Skeleton->GetGuid().ToString(),
-		*DataToCompress.Skeleton->GetVirtualBoneGuid().ToString(),
-		AdditiveType,
-		RefType,
-		DataToCompress.RefFrameIndex,
-		(DataToCompress.bIsValidAdditive) ? *DataToCompress.AdditiveDataGuid.ToString() : TEXT("NotAdditive"),
-		*DataToCompress.RequestedCompressionScheme->MakeDDCKey(),
-		*DataToCompress.CurveCompressionSettings->MakeDDCKey()
-		);
-
-	return Ret;
-}
-
 bool FDerivedDataAnimationCompression::Build( TArray<uint8>& OutDataArray )
 {
+	check(DataToCompressPtr.IsValid());
+	FCompressibleAnimData& DataToCompress = *DataToCompressPtr.Get();
 	FCompressedAnimSequence OutData;
 
 	SCOPE_CYCLE_COUNTER(STAT_AnimCompressionDerivedData);
@@ -141,38 +40,6 @@ bool FDerivedDataAnimationCompression::Build( TArray<uint8>& OutDataArray )
 
 	bool bCompressionSuccessful = false;
 	{
-		if (bPerformStripping)
-		{
-			const int32 NumFrames = DataToCompress.NumFrames;
-			const int32 NumTracks = DataToCompress.RawAnimationData.Num();
-
-			//Strip every other frame from tracks
-			if(bIsEvenFramed)
-			{
-				for (FRawAnimSequenceTrack& Track : DataToCompress.RawAnimationData)
-				{
-					StripFramesEven(Track.PosKeys, NumFrames);
-					StripFramesEven(Track.RotKeys, NumFrames);
-					StripFramesEven(Track.ScaleKeys, NumFrames);
-				}
-
-				const int32 ActualFrames = DataToCompress.NumFrames - 1; // strip bookmark end frame
-				DataToCompress.NumFrames = (ActualFrames / 2) + 1;
-			}
-			else
-			{
-				for (FRawAnimSequenceTrack& Track : DataToCompress.RawAnimationData)
-				{
-					StripFramesOdd(Track.PosKeys, NumFrames);
-					StripFramesOdd(Track.RotKeys, NumFrames);
-					StripFramesOdd(Track.ScaleKeys, NumFrames);
-				}
-
-				const int32 ActualFrames = DataToCompress.NumFrames;
-				DataToCompress.NumFrames = (ActualFrames / 2);
-			}
-		}
-
 		DataToCompress.Update(OutData);
 
 		bool bCurveCompressionSuccess = FAnimationUtils::CompressAnimCurves(DataToCompress, OutData);
@@ -211,7 +78,7 @@ bool FDerivedDataAnimationCompression::Build( TArray<uint8>& OutDataArray )
 	if (bCompressionSuccessful)
 	{
 		FMemoryWriter Ar(OutDataArray, true);
-		OutData.SerializeCompressedData(Ar, true, nullptr, DataToCompress.CurveCompressionSettings); //Save out compressed
+		OutData.SerializeCompressedData(Ar, true, nullptr, DataToCompress.Skeleton, DataToCompress.CurveCompressionSettings); //Save out compressed
 	}
 
 	return bCompressionSuccessful;
diff --git a/Engine/Source/Runtime/Engine/Private/Animation/AnimCompressionDerivedData.h b/Engine/Source/Runtime/Engine/Private/Animation/AnimCompressionDerivedData.h
--- a/Engine/Source/Runtime/Engine/Private/Animation/AnimCompressionDerivedData.h
+++ b/Engine/Source/Runtime/Engine/Private/Animation/AnimCompressionDerivedData.h
@@ -10,7 +10,6 @@
 
 #include "Animation/AnimCompressionTypes.h"
 
-class UAnimSequence;
 struct FAnimCompressContext;
 
 #if WITH_EDITOR
@@ -20,7 +19,14 @@ struct FAnimCompressContext;
 class FDerivedDataAnimationCompression : public FDerivedDataPluginInterface
 {
 private:
-	FCompressibleAnimData DataToCompress;
+	// The anim data to compress
+	TSharedPtr<FCompressibleAnimData> DataToCompressPtr;
+
+	// The Type of anim data to compress (makes up part of DDC key)
+	const TCHAR* TypeName;
+
+	// Bulk of asset DDC key
+	const FString AssetDDCKey;
 
 	// FAnimCompressContext to use during compression if we don't pull from the DDC
 	TSharedPtr<FAnimCompressContext> CompressContext;
@@ -35,12 +41,18 @@ private:
 	bool bIsEvenFramed;
 
 public:
-	FDerivedDataAnimationCompression(const FCompressibleAnimData& InDataToCompress, TSharedPtr<FAnimCompressContext> InCompressContext, int32 InPreviousCompressedSize, bool bInTryFrameStripping, bool bTryStrippingOnOddFramedAnims);
+	FDerivedDataAnimationCompression(const TCHAR* InTypeName, const FString& InAssetDDCKey, TSharedPtr<FAnimCompressContext> InCompressContext, int32 InPreviousCompressedSize);
 	virtual ~FDerivedDataAnimationCompression();
 
+	void SetCompressibleData(TSharedRef<FCompressibleAnimData> InCompressibleAnimData)
+	{
+		DataToCompressPtr = InCompressibleAnimData;
+		check(DataToCompressPtr->Skeleton != nullptr);
+	}
+
 	virtual const TCHAR* GetPluginName() const override
 	{
-		return *DataToCompress.TypeName;
+		return TypeName;
 	}
 
 	virtual const TCHAR* GetVersionString() const override
@@ -48,10 +60,13 @@ public:
 		// This is a version string that mimics the old versioning scheme. If you
 		// want to bump this version, generate a new guid using VS->Tools->Create GUID and
 		// return it here. Ex.
-		return TEXT("1F1656B9E10142729AB16650D9821B1F");
+		return TEXT("1F1656B9E10142729AB16650D9821B1Ff");
 	}
 
-	virtual FString GetPluginSpecificCacheKeySuffix() const override;
+	virtual FString GetPluginSpecificCacheKeySuffix() const override
+	{
+		return AssetDDCKey;
+	}
 
 
 	virtual bool IsBuildThreadsafe() const override
@@ -64,7 +79,7 @@ public:
 	/** Return true if we can build **/
 	bool CanBuild()
 	{
-		return true;
+		return DataToCompressPtr.IsValid();
 	}
 };
 
diff --git a/Engine/Source/Runtime/Engine/Private/Animation/AnimCompressionTypes.cpp b/Engine/Source/Runtime/Engine/Private/Animation/AnimCompressionTypes.cpp
--- a/Engine/Source/Runtime/Engine/Private/Animation/AnimCompressionTypes.cpp
+++ b/Engine/Source/Runtime/Engine/Private/Animation/AnimCompressionTypes.cpp
@@ -3,12 +3,127 @@
 #include "Animation/AnimCompressionTypes.h"
 #include "AnimationUtils.h"
 #include "AnimEncoding.h"
+#include "Misc/SecureHash.h"
 #include "Interfaces/ITargetPlatform.h"
 #include "Animation/AnimCurveCompressionCodec.h"
 #include "Animation/AnimCurveCompressionSettings.h"
+#include "AnimationRuntime.h"
 #include "UObject/FortniteMainBranchObjectVersion.h"
 
-FCompressibleAnimData::FCompressibleAnimData(class UAnimSequence* InSeq) :
+CSV_DEFINE_CATEGORY(Animation, false);
+
+DECLARE_CYCLE_STAT(TEXT("Build Anim Track Pairs"), STAT_BuildAnimTrackPairs, STATGROUP_Anim);
+DECLARE_CYCLE_STAT(TEXT("Extract Pose From Anim Data"), STAT_ExtractPoseFromAnimData, STATGROUP_Anim);
+
+template <typename ArrayType>
+void UpdateSHAWithArray(FSHA1& Sha, const TArray<ArrayType>& Array)
+{
+	Sha.Update((uint8*)Array.GetData(), Array.Num() * Array.GetTypeSize());
+}
+
+void UpdateSHAWithRawTrack(FSHA1& Sha, const FRawAnimSequenceTrack& RawTrack)
+{
+	UpdateSHAWithArray(Sha, RawTrack.PosKeys);
+	UpdateSHAWithArray(Sha, RawTrack.RotKeys);
+	UpdateSHAWithArray(Sha, RawTrack.ScaleKeys);
+}
+
+template<class DataType>
+void UpdateWithData(FSHA1& Sha, const DataType& Data)
+{
+	Sha.Update((uint8*)(&Data), sizeof(DataType));
+}
+
+void UpdateSHAWithCurves(FSHA1& Sha, const FRawCurveTracks& InRawCurveData) 
+{
+	for (const FFloatCurve& Curve : InRawCurveData.FloatCurves)
+	{
+		UpdateWithData(Sha, Curve.Name.UID);
+		UpdateWithData(Sha, Curve.FloatCurve.DefaultValue);
+		UpdateSHAWithArray(Sha, Curve.FloatCurve.GetConstRefOfKeys());
+		UpdateWithData(Sha, Curve.FloatCurve.PreInfinityExtrap);
+		UpdateWithData(Sha, Curve.FloatCurve.PostInfinityExtrap);
+	}
+}
+
+FGuid GenerateGuidFromRawAnimData(const TArray<FRawAnimSequenceTrack>& RawAnimationData, const FRawCurveTracks& RawCurveData)
+{
+	FSHA1 Sha;
+
+	for (const FRawAnimSequenceTrack& Track : RawAnimationData)
+	{
+		UpdateSHAWithRawTrack(Sha, Track);
+	}
+
+	UpdateSHAWithCurves(Sha, RawCurveData);
+
+	Sha.Final();
+
+	uint32 Hash[5];
+	Sha.GetHash((uint8*)Hash);
+	FGuid Guid(Hash[0] ^ Hash[4], Hash[1], Hash[2], Hash[3]);
+	return Guid;
+}
+
+template<typename ArrayValue>
+void StripFramesEven(TArray<ArrayValue>& Keys, const int32 NumFrames)
+{
+	if (Keys.Num() > 1)
+	{
+		check(Keys.Num() == NumFrames);
+
+		for (int32 DstKey = 1, SrcKey = 2; SrcKey < NumFrames; ++DstKey, SrcKey += 2)
+		{
+			Keys[DstKey] = Keys[SrcKey];
+		}
+
+		const int32 HalfSize = (NumFrames - 1) / 2;
+		const int32 StartRemoval = HalfSize + 1;
+
+		Keys.RemoveAt(StartRemoval, NumFrames - StartRemoval);
+	}
+}
+
+template<typename ArrayValue>
+void StripFramesOdd(TArray<ArrayValue>& Keys, const int32 NumFrames)
+{
+	if (Keys.Num() > 1)
+	{
+		const int32 NewNumFrames = NumFrames / 2;
+
+		TArray<ArrayValue> NewKeys;
+		NewKeys.Reserve(NewNumFrames);
+
+		check(Keys.Num() == NumFrames);
+
+		NewKeys.Add(Keys[0]); //Always keep first 
+
+		//Always keep first and last
+		const int32 NumFramesToCalculate = NewNumFrames - 2;
+
+		// Frame increment is ratio of old frame spaces vs new frame spaces 
+		const double FrameIncrement = (double)(NumFrames - 1) / (double)(NewNumFrames - 1);
+
+		for (int32 Frame = 0; Frame < NumFramesToCalculate; ++Frame)
+		{
+			const double NextFramePosition = FrameIncrement * (Frame + 1);
+			const int32 Frame1 = (int32)NextFramePosition;
+			const float Alpha = (NextFramePosition - (double)Frame1);
+
+			NewKeys.Add(AnimationCompressionUtils::Interpolate(Keys[Frame1], Keys[Frame1 + 1], Alpha));
+
+		}
+
+		NewKeys.Add(Keys.Last()); // Always Keep Last
+
+		const int32 HalfSize = (NumFrames - 1) / 2;
+		const int32 StartRemoval = HalfSize + 1;
+
+		Keys = MoveTemp(NewKeys);
+	}
+}
+
+FCompressibleAnimData::FCompressibleAnimData(class UAnimSequence* InSeq, const bool bPerformStripping) :
 #if WITH_EDITOR
 	RequestedCompressionScheme(InSeq->CompressionScheme) ,
 #endif
@@ -19,13 +134,6 @@ FCompressibleAnimData::FCompressibleAnimData(class UAnimSequence* InSeq) :
 	, SequenceLength(InSeq->SequenceLength)
 	, NumFrames(InSeq->GetRawNumberOfFrames())
 	, bIsValidAdditive(InSeq->IsValidAdditive())
-#if WITH_EDITOR
-	, CompressCommandletVersion(InSeq->CompressCommandletVersion)
-	, RawDataGuid(InSeq->GetRawDataGuid())
-#endif
-	, RefFrameIndex(InSeq->RefFrameIndex)
-	, RefPoseType(InSeq->RefPoseType)
-	, AdditiveAnimType(InSeq->AdditiveAnimType)
 	, Name(InSeq->GetName())
 	, FullName(InSeq->GetFullName())
 	, AnimFName(InSeq->GetFName())
@@ -39,11 +147,6 @@ FCompressibleAnimData::FCompressibleAnimData(class UAnimSequence* InSeq) :
 	{
 		TArray<FName> TempTrackNames;
 		InSeq->BakeOutAdditiveIntoRawData(RawAnimationData, TempTrackNames, TrackToSkeletonMapTable, RawCurveData, AdditiveBaseAnimationData);
-
-		if (InSeq->RefPoseSeq)
-		{
-			AdditiveDataGuid = InSeq->RefPoseSeq->GetRawDataGuid();
-		}
 	}
 	else if (bHasVirtualBones)// If we aren't additive we must bake virtual bones
 	{
@@ -58,7 +161,55 @@ FCompressibleAnimData::FCompressibleAnimData(class UAnimSequence* InSeq) :
 		RawCurveData = InSeq->RawCurveData;
 	}
 
-	TypeName = TEXT("AnimSeq");
+	if (bPerformStripping)
+	{
+		const int32 NumTracks = RawAnimationData.Num();
+		
+		// End frame does not count towards "Even framed" calculation
+		const bool bIsEvenFramed = ((NumFrames - 1) % 2) == 0;
+
+		//Strip every other frame from tracks
+		if (bIsEvenFramed)
+		{
+			for (FRawAnimSequenceTrack& Track : RawAnimationData)
+			{
+				StripFramesEven(Track.PosKeys, NumFrames);
+				StripFramesEven(Track.RotKeys, NumFrames);
+				StripFramesEven(Track.ScaleKeys, NumFrames);
+			}
+
+			const int32 ActualFrames = NumFrames - 1; // strip bookmark end frame
+			NumFrames = (ActualFrames / 2) + 1;
+		}
+		else
+		{
+			for (FRawAnimSequenceTrack& Track : RawAnimationData)
+			{
+				StripFramesOdd(Track.PosKeys, NumFrames);
+				StripFramesOdd(Track.RotKeys, NumFrames);
+				StripFramesOdd(Track.ScaleKeys, NumFrames);
+			}
+
+			const int32 ActualFrames = NumFrames;
+			NumFrames = (ActualFrames / 2);
+		}
+	}
+#endif
+}
+
+FCompressibleAnimData::FCompressibleAnimData(UAnimCompress* InRequestedCompressionScheme, UAnimCurveCompressionSettings* InCurveCompressionSettings, USkeleton* InSkeleton, EAnimInterpolationType InInterpolation, float InSequenceLength, int32 InNumFrames) :
+#if WITH_EDITOR
+	RequestedCompressionScheme(InRequestedCompressionScheme) ,
+#endif
+	CurveCompressionSettings(InCurveCompressionSettings)
+	, Skeleton(InSkeleton)
+	, Interpolation(InInterpolation)
+	, SequenceLength(InSequenceLength)
+	, NumFrames(InNumFrames)
+	, bIsValidAdditive(false)
+{
+#if WITH_EDITOR
+	FAnimationUtils::BuildSkeletonMetaData(Skeleton, BoneData);
 #endif
 }
 
@@ -253,8 +404,9 @@ void FUECompressedAnimData::ByteSwapData(TArrayView<uint8> CompressedData, TArch
 template void FUECompressedAnimData::ByteSwapData(TArrayView<uint8> CompressedData, FMemoryReader& MemoryStream);
 template void FUECompressedAnimData::ByteSwapData(TArrayView<uint8> CompressedData, FMemoryWriter& MemoryStream);
 
-void FCompressedAnimSequence::SerializeCompressedData(FArchive& Ar, bool bDDCData, UObject* DataOwner, UAnimCurveCompressionSettings* CurveCompressionSettings)
+void FCompressedAnimSequence::SerializeCompressedData(FArchive& Ar, bool bDDCData, UObject* DataOwner, USkeleton* Skeleton, UAnimCurveCompressionSettings* CurveCompressionSettings, bool bCanUseBulkData)
 {
+	Ar << CompressedRawDataSize;
 	Ar << CompressedTrackToSkeletonMapTable;
 	Ar << CompressedCurveNames;
 
@@ -359,7 +511,7 @@ void FCompressedAnimSequence::SerializeCompressedData(FArchive& Ar, bool bDDCDat
 		// Make sure the entire byte stream was serialized.
 		check(NumBytes == SerializedData.Num());
 
-		bool bUseBulkDataForSave = NumBytes && bIsCooking && Ar.CookingTarget()->SupportsFeature(ETargetPlatformFeatures::MemoryMappedFiles) && Ar.CookingTarget()->SupportsFeature(ETargetPlatformFeatures::MemoryMappedAnimation);
+		bool bUseBulkDataForSave = bCanUseBulkData && NumBytes && bIsCooking && Ar.CookingTarget()->SupportsFeature(ETargetPlatformFeatures::MemoryMappedFiles) && Ar.CookingTarget()->SupportsFeature(ETargetPlatformFeatures::MemoryMappedAnimation);
 
 		bool bSavebUseBulkDataForSave = false;
 		if (!bDDCData)
@@ -437,4 +589,228 @@ void FCompressedAnimSequence::SerializeCompressedData(FArchive& Ar, bool bDDCDat
 		Ar << NumCurveBytes;
 		Ar.Serialize(CompressedCurveByteStream.GetData(), NumCurveBytes);
 	}
+
+#if WITH_EDITOR
+	if (bDDCData)
+	{
+		if (Ar.IsLoading() && Skeleton)
+		{
+			// Refresh the compressed curve names since the IDs might have changed since
+			for (FSmartName& CurveName : CompressedCurveNames)
+			{
+				Skeleton->VerifySmartName(USkeleton::AnimCurveMappingName, CurveName);
+			}
+		}
+	}
+#endif
 }
+
+struct FGetBonePoseScratchArea : public TThreadSingleton<FGetBonePoseScratchArea>
+{
+	BoneTrackArray RotationScalePairs;
+	BoneTrackArray TranslationPairs;
+	BoneTrackArray AnimScaleRetargetingPairs;
+	BoneTrackArray AnimRelativeRetargetingPairs;
+	BoneTrackArray OrientAndScaleRetargetingPairs;
+};
+
+void DecompressPose(FCompactPose& OutPose, const FCompressedAnimSequence& CompressedData, const FAnimExtractContext& ExtractionContext, USkeleton* Skeleton, float SequenceLength, EAnimInterpolationType Interpolation, bool bIsBakedAdditive, FName RetargetSource, FName SourceName, const FRootMotionReset& RootMotionReset)
+{
+	const FBoneContainer& RequiredBones = OutPose.GetBoneContainer();
+	const int32 NumTracks = CompressedData.CompressedTrackToSkeletonMapTable.Num();
+
+	TArray<int32> const& SkeletonToPoseBoneIndexArray = RequiredBones.GetSkeletonToPoseBoneIndexArray();
+
+	BoneTrackArray& RotationScalePairs = FGetBonePoseScratchArea::Get().RotationScalePairs;
+	BoneTrackArray& TranslationPairs = FGetBonePoseScratchArea::Get().TranslationPairs;
+	BoneTrackArray& AnimScaleRetargetingPairs = FGetBonePoseScratchArea::Get().AnimScaleRetargetingPairs;
+	BoneTrackArray& AnimRelativeRetargetingPairs = FGetBonePoseScratchArea::Get().AnimRelativeRetargetingPairs;
+	BoneTrackArray& OrientAndScaleRetargetingPairs = FGetBonePoseScratchArea::Get().OrientAndScaleRetargetingPairs;
+
+	// build a list of desired bones
+	RotationScalePairs.Reset();
+	TranslationPairs.Reset();
+	AnimScaleRetargetingPairs.Reset();
+	AnimRelativeRetargetingPairs.Reset();
+	OrientAndScaleRetargetingPairs.Reset();
+
+	// Optimization: assuming first index is root bone. That should always be the case in Skeletons.
+	checkSlow((SkeletonToPoseBoneIndexArray[0] == 0));
+	// this is not guaranteed for AnimSequences though... If Root is not animated, Track will not exist.
+	const bool bFirstTrackIsRootBone = (CompressedData.GetSkeletonIndexFromTrackIndex(0) == 0);
+
+	{
+		SCOPE_CYCLE_COUNTER(STAT_BuildAnimTrackPairs);
+
+		// Handle root bone separately if it is track 0. so we start w/ Index 1.
+		for (int32 TrackIndex = (bFirstTrackIsRootBone ? 1 : 0); TrackIndex < NumTracks; TrackIndex++)
+		{
+			const int32 SkeletonBoneIndex = CompressedData.GetSkeletonIndexFromTrackIndex(TrackIndex);
+			// not sure it's safe to assume that SkeletonBoneIndex can never be INDEX_NONE
+			if (SkeletonBoneIndex != INDEX_NONE)
+			{
+				const FCompactPoseBoneIndex BoneIndex = RequiredBones.GetCompactPoseIndexFromSkeletonIndex(SkeletonBoneIndex);
+				//Nasty, we break our type safety, code in the lower levels should be adjusted for this
+				const int32 CompactPoseBoneIndex = BoneIndex.GetInt();
+				if (CompactPoseBoneIndex != INDEX_NONE)
+				{
+					RotationScalePairs.Add(BoneTrackPair(CompactPoseBoneIndex, TrackIndex));
+
+					// Skip extracting translation component for EBoneTranslationRetargetingMode::Skeleton.
+					switch (Skeleton->GetBoneTranslationRetargetingMode(SkeletonBoneIndex))
+					{
+					case EBoneTranslationRetargetingMode::Animation:
+						TranslationPairs.Add(BoneTrackPair(CompactPoseBoneIndex, TrackIndex));
+						break;
+					case EBoneTranslationRetargetingMode::AnimationScaled:
+						TranslationPairs.Add(BoneTrackPair(CompactPoseBoneIndex, TrackIndex));
+						AnimScaleRetargetingPairs.Add(BoneTrackPair(CompactPoseBoneIndex, SkeletonBoneIndex));
+						break;
+					case EBoneTranslationRetargetingMode::AnimationRelative:
+						TranslationPairs.Add(BoneTrackPair(CompactPoseBoneIndex, TrackIndex));
+
+						// With baked additives, we can skip 'AnimationRelative' tracks, as the relative transform gets canceled out.
+						// (A1 + Rel) - (A2 + Rel) = A1 - A2.
+						if (!bIsBakedAdditive)
+						{
+							AnimRelativeRetargetingPairs.Add(BoneTrackPair(CompactPoseBoneIndex, SkeletonBoneIndex));
+						}
+						break;
+					case EBoneTranslationRetargetingMode::OrientAndScale:
+						TranslationPairs.Add(BoneTrackPair(CompactPoseBoneIndex, TrackIndex));
+
+						// Additives remain additives, they're not retargeted.
+						if (!bIsBakedAdditive)
+						{
+							OrientAndScaleRetargetingPairs.Add(BoneTrackPair(CompactPoseBoneIndex, SkeletonBoneIndex));
+						}
+						break;
+					}
+				}
+			}
+		}
+	}
+
+	{
+		SCOPE_CYCLE_COUNTER(STAT_ExtractPoseFromAnimData);
+		CSV_SCOPED_TIMING_STAT(Animation, ExtractPoseFromAnimData);
+		CSV_CUSTOM_STAT(Animation, NumberOfExtractedAnimations, 1, ECsvCustomStatOp::Accumulate);
+
+		FAnimSequenceDecompressionContext EvalDecompContext(SequenceLength, Interpolation, SourceName, CompressedData.CompressedDataStructure);
+		EvalDecompContext.Seek(ExtractionContext.CurrentTime);
+
+		// Handle Root Bone separately
+		if (bFirstTrackIsRootBone)
+		{
+			const int32 TrackIndex = 0;
+			FCompactPoseBoneIndex RootBone(0);
+			FTransform& RootAtom = OutPose[RootBone];
+
+			AnimationFormat_GetBoneAtom(
+				RootAtom,
+				EvalDecompContext,
+				TrackIndex);
+
+			// @laurent - we should look into splitting rotation and translation tracks, so we don't have to process translation twice.
+			FAnimationRuntime::RetargetBoneTransform(Skeleton, RetargetSource, RootAtom, 0, RootBone, RequiredBones, bIsBakedAdditive);
+		}
+
+		if (RotationScalePairs.Num() > 0)
+		{
+			// get the remaining bone atoms
+			FTransformArray LocalBones;
+			OutPose.MoveBonesTo(LocalBones);
+
+			AnimationFormat_GetAnimationPose(
+				LocalBones,
+				RotationScalePairs,
+				TranslationPairs,
+				RotationScalePairs,
+				EvalDecompContext);
+
+			OutPose.MoveBonesFrom(MoveTemp(LocalBones));
+		}
+	}
+
+	// Once pose has been extracted, snap root bone back to first frame if we are extracting root motion.
+	if ((ExtractionContext.bExtractRootMotion && RootMotionReset.bEnableRootMotion) || RootMotionReset.bForceRootLock)
+	{
+		RootMotionReset.ResetRootBoneForRootMotion(OutPose[FCompactPoseBoneIndex(0)], RequiredBones);
+	}
+
+	// Anim Scale Retargeting
+	int32 const NumBonesToScaleRetarget = AnimScaleRetargetingPairs.Num();
+	if (NumBonesToScaleRetarget > 0)
+	{
+		TArray<FTransform> const& AuthoredOnRefSkeleton = Skeleton->GetRefLocalPoses(RetargetSource);
+
+		for (const BoneTrackPair& BonePair : AnimScaleRetargetingPairs)
+		{
+			const FCompactPoseBoneIndex BoneIndex(BonePair.AtomIndex); //Nasty, we break our type safety, code in the lower levels should be adjusted for this
+			int32 const& SkeletonBoneIndex = BonePair.TrackIndex;
+
+			// @todo - precache that in FBoneContainer when we have SkeletonIndex->TrackIndex mapping. So we can just apply scale right away.
+			float const SourceTranslationLength = AuthoredOnRefSkeleton[SkeletonBoneIndex].GetTranslation().Size();
+			if (SourceTranslationLength > KINDA_SMALL_NUMBER)
+			{
+				float const TargetTranslationLength = RequiredBones.GetRefPoseTransform(BoneIndex).GetTranslation().Size();
+				OutPose[BoneIndex].ScaleTranslation(TargetTranslationLength / SourceTranslationLength);
+			}
+		}
+	}
+
+	// Anim Relative Retargeting
+	int32 const NumBonesToRelativeRetarget = AnimRelativeRetargetingPairs.Num();
+	if (NumBonesToRelativeRetarget > 0)
+	{
+		TArray<FTransform> const& AuthoredOnRefSkeleton = Skeleton->GetRefLocalPoses(RetargetSource);
+
+		for (const BoneTrackPair& BonePair : AnimRelativeRetargetingPairs)
+		{
+			const FCompactPoseBoneIndex BoneIndex(BonePair.AtomIndex); //Nasty, we break our type safety, code in the lower levels should be adjusted for this
+			int32 const& SkeletonBoneIndex = BonePair.TrackIndex;
+
+			const FTransform& RefPose = RequiredBones.GetRefPoseTransform(BoneIndex);
+
+			// Apply the retargeting as if it were an additive difference between the current skeleton and the retarget skeleton. 
+			OutPose[BoneIndex].SetRotation(OutPose[BoneIndex].GetRotation() * AuthoredOnRefSkeleton[SkeletonBoneIndex].GetRotation().Inverse() * RefPose.GetRotation());
+			OutPose[BoneIndex].SetTranslation(OutPose[BoneIndex].GetTranslation() + (RefPose.GetTranslation() - AuthoredOnRefSkeleton[SkeletonBoneIndex].GetTranslation()));
+			OutPose[BoneIndex].SetScale3D(OutPose[BoneIndex].GetScale3D() * (RefPose.GetScale3D() * AuthoredOnRefSkeleton[SkeletonBoneIndex].GetSafeScaleReciprocal(AuthoredOnRefSkeleton[SkeletonBoneIndex].GetScale3D())));
+			OutPose[BoneIndex].NormalizeRotation();
+		}
+	}
+
+	// Translation 'Orient and Scale' Translation Retargeting
+	const int32 NumBonesToOrientAndScaleRetarget = OrientAndScaleRetargetingPairs.Num();
+	if (NumBonesToOrientAndScaleRetarget > 0)
+	{
+		const FRetargetSourceCachedData& RetargetSourceCachedData = RequiredBones.GetRetargetSourceCachedData(RetargetSource);
+		const TArray<FOrientAndScaleRetargetingCachedData>& OrientAndScaleDataArray = RetargetSourceCachedData.OrientAndScaleData;
+		const TArray<int32>& CompactPoseIndexToOrientAndScaleIndex = RetargetSourceCachedData.CompactPoseIndexToOrientAndScaleIndex;
+
+		// If we have any cached retargeting data.
+		if ((OrientAndScaleDataArray.Num() > 0) && (CompactPoseIndexToOrientAndScaleIndex.Num() == RequiredBones.GetCompactPoseNumBones()))
+		{
+			for (int32 Index = 0; Index < NumBonesToOrientAndScaleRetarget; Index++)
+			{
+				const BoneTrackPair& BonePair = OrientAndScaleRetargetingPairs[Index];
+				const FCompactPoseBoneIndex CompactPoseBoneIndex(BonePair.AtomIndex);
+				const int32 OrientAndScaleIndex = CompactPoseIndexToOrientAndScaleIndex[CompactPoseBoneIndex.GetInt()];
+				if (OrientAndScaleIndex != INDEX_NONE)
+				{
+					const FOrientAndScaleRetargetingCachedData& OrientAndScaleData = OrientAndScaleDataArray[OrientAndScaleIndex];
+					FTransform& BoneTransform = OutPose[CompactPoseBoneIndex];
+					const FVector AnimatedTranslation = BoneTransform.GetTranslation();
+
+					// If Translation is not animated, we can just copy the TargetTranslation. No retargeting needs to be done.
+					const FVector NewTranslation = (AnimatedTranslation - OrientAndScaleData.SourceTranslation).IsNearlyZero(BONE_TRANS_RT_ORIENT_AND_SCALE_PRECISION) ?
+						OrientAndScaleData.TargetTranslation :
+						OrientAndScaleData.TranslationDeltaOrient.RotateVector(AnimatedTranslation) * OrientAndScaleData.TranslationScale;
+
+					BoneTransform.SetTranslation(NewTranslation);
+				}
+			}
+		}
+	}
+}
+
diff --git a/Engine/Source/Runtime/Engine/Private/Animation/AnimSequence.cpp b/Engine/Source/Runtime/Engine/Private/Animation/AnimSequence.cpp
--- a/Engine/Source/Runtime/Engine/Private/Animation/AnimSequence.cpp
+++ b/Engine/Source/Runtime/Engine/Private/Animation/AnimSequence.cpp
@@ -36,19 +36,15 @@
 #include "Animation/AnimNotifies/AnimNotifyState.h"
 #include "Widgets/Notifications/SNotificationList.h"
 #include "Framework/Notifications/NotificationManager.h"
-#include "ProfilingDebugging/CsvProfiler.h"
 #include "HAL/PlatformApplicationMisc.h"
 #include "DeviceProfiles/DeviceProfileManager.h"
 #include "DeviceProfiles/DeviceProfile.h"
+#include "Animation/AnimStreamable.h"
 #define USE_SLERP 0
 #define LOCTEXT_NAMESPACE "AnimSequence"
 
-CSV_DEFINE_CATEGORY(Animation, false);
-
 DECLARE_CYCLE_STAT(TEXT("AnimSeq GetBonePose"), STAT_AnimSeq_GetBonePose, STATGROUP_Anim);
 DECLARE_CYCLE_STAT(TEXT("AnimSeq EvalCurveData"), STAT_AnimSeq_EvalCurveData, STATGROUP_Anim);
-DECLARE_CYCLE_STAT(TEXT("Build Anim Track Pairs"), STAT_BuildAnimTrackPairs, STATGROUP_Anim);
-DECLARE_CYCLE_STAT(TEXT("Extract Pose From Anim Data"), STAT_ExtractPoseFromAnimData, STATGROUP_Anim);
 
 int32 GPerformFrameStripping = 0;
 int32 GPerformFrameStrippingOddFramedAnimations = 0;
@@ -203,6 +199,41 @@ void OnCVarsChanged()
 }
 
 FAutoConsoleVariableSink AnimationCVarSink(FConsoleCommandDelegate::CreateStatic(&OnCVarsChanged));
+
+FString GetAnimSequenceSpecificCacheKeySuffix(const UAnimSequence& Seq, const bool bPerformStripping)
+{
+	//Make up our content key consisting of:
+	//	* Global animation compression version
+	//  * Whether to strip frames
+	//	* Our raw data GUID
+	//	* Our skeleton GUID: If our skeleton changes our compressed data may now be stale
+	//	* Baked Additive Flag
+	//	* Additive ref pose GUID or hardcoded string if not available
+	//	* Compression Settings
+	//	* Curve compression settings
+
+	bool bIsValidAdditive = Seq.IsValidAdditive();
+	char AdditiveType = bIsValidAdditive ? NibbleToTChar(Seq.AdditiveAnimType) : '0';
+	char RefType = bIsValidAdditive ? NibbleToTChar(Seq.RefPoseType) : '0';
+
+	const int32 StripFrame = bPerformStripping ? 1 : 0;
+
+	FString Ret = FString::Printf(TEXT("%i_%i_%s%s%s_%c%c%i_%s_%s_%s"),
+		Seq.CompressCommandletVersion,
+		StripFrame,
+		*Seq.GetRawDataGuid().ToString(),
+		*Seq.GetSkeleton()->GetGuid().ToString(),
+		*Seq.GetSkeleton()->GetVirtualBoneGuid().ToString(),
+		AdditiveType,
+		RefType,
+		Seq.RefFrameIndex,
+		(bIsValidAdditive && Seq.RefPoseSeq) ? *Seq.RefPoseSeq->GetRawDataGuid().ToString() : TEXT("NoAdditiveGuid"),
+		*Seq.CompressionScheme->MakeDDCKey(),
+		*Seq.CurveCompressionSettings->MakeDDCKey()
+	);
+
+	return Ret;
+}
 #endif
 
 bool CompressRawAnimSequenceTrack(FRawAnimSequenceTrack& RawTrack, int32 NumFrames, FName ErrorName, float MaxPosDiff, float MaxAngleDiff)
@@ -1470,23 +1501,6 @@ void UAnimSequence::GetAnimationPose(FCompactPose& OutPose, FBlendedCurve& OutCu
 #endif
 }
 
-void UAnimSequence::ResetRootBoneForRootMotion(FTransform& BoneTransform, const FBoneContainer& RequiredBones, ERootMotionRootLock::Type InRootMotionRootLock) const
-{
-	switch (InRootMotionRootLock)
-	{
-		case ERootMotionRootLock::AnimFirstFrame: BoneTransform = ExtractRootTrackTransform(0.f, &RequiredBones); break;
-		case ERootMotionRootLock::Zero: BoneTransform = FTransform::Identity; break;
-		default:
-		case ERootMotionRootLock::RefPose: BoneTransform = RequiredBones.GetRefPoseArray()[0]; break;
-	}
-
-	if (IsValidAdditive() && InRootMotionRootLock != ERootMotionRootLock::AnimFirstFrame)
-	{
-		//Need to remove default scale here for additives
-		BoneTransform.SetScale3D(BoneTransform.GetScale3D() - FVector(1.f));
-	}
-}
-
 struct FRetargetTracking
 {
 	const FCompactPoseBoneIndex PoseBoneIndex;
@@ -1498,190 +1512,6 @@ struct FRetargetTracking
 	}
 };
 
-struct FGetBonePoseScratchArea : public TThreadSingleton<FGetBonePoseScratchArea>
-{
-	BoneTrackArray RotationScalePairs;
-	BoneTrackArray TranslationPairs;
-	BoneTrackArray AnimScaleRetargetingPairs;
-	BoneTrackArray AnimRelativeRetargetingPairs;
-	BoneTrackArray OrientAndScaleRetargetingPairs;
-	TArray<FRetargetTracking> RetargetTracking;
-	TArray<FVirtualBoneCompactPoseData> VirtualBoneCompactPoseData;
-};
-
-
-FTransform ExtractTransformForKey(int32 Key, const FRawAnimSequenceTrack &TrackToExtract)
-{
-	static const FVector DefaultScale3D = FVector(1.f);
-	const bool bHasScaleKey = TrackToExtract.ScaleKeys.Num() > 0;
-
-	const int32 PosKeyIndex = FMath::Min(Key, TrackToExtract.PosKeys.Num() - 1);
-	const int32 RotKeyIndex = FMath::Min(Key, TrackToExtract.RotKeys.Num() - 1);
-	if (bHasScaleKey)
-	{
-		const int32 ScaleKeyIndex = FMath::Min(Key, TrackToExtract.ScaleKeys.Num() - 1);
-		return FTransform(TrackToExtract.RotKeys[RotKeyIndex], TrackToExtract.PosKeys[PosKeyIndex], TrackToExtract.ScaleKeys[ScaleKeyIndex]);
-	}
-	else
-	{
-		return FTransform(TrackToExtract.RotKeys[RotKeyIndex], TrackToExtract.PosKeys[PosKeyIndex], DefaultScale3D);
-	}
-}
-
-template<bool bInterpolateT>
-void UAnimSequence::BuildPoseFromRawDataInternal(const TArray<FRawAnimSequenceTrack>& InAnimationData, FCompactPose& InOutPose, int32 KeyIndex1, int32 KeyIndex2, float Alpha) const
-{
-	const int32 NumTracks = InAnimationData.Num();
-	const FBoneContainer& RequiredBones = InOutPose.GetBoneContainer();
-
-	TArray<FRetargetTracking>& RetargetTracking = FGetBonePoseScratchArea::Get().RetargetTracking;
-	RetargetTracking.Reset(NumTracks);
-
-	TArray<FVirtualBoneCompactPoseData>& VBCompactPoseData = FGetBonePoseScratchArea::Get().VirtualBoneCompactPoseData;
-	VBCompactPoseData = RequiredBones.GetVirtualBoneCompactPoseData();
-
-	FCompactPose Key2Pose;
-	if (bInterpolateT)
-	{
-		Key2Pose.CopyBonesFrom(InOutPose);
-	}
-
-	for (int32 TrackIndex = 0; TrackIndex < NumTracks; TrackIndex++)
-	{
-		const int32 SkeletonBoneIndex = GetSkeletonIndexFromRawDataTrackIndex(TrackIndex);
-		// not sure it's safe to assume that SkeletonBoneIndex can never be INDEX_NONE
-		if ((SkeletonBoneIndex != INDEX_NONE) && (SkeletonBoneIndex < MAX_BONES))
-		{
-			const FCompactPoseBoneIndex PoseBoneIndex = RequiredBones.GetCompactPoseIndexFromSkeletonIndex(SkeletonBoneIndex);
-			if (PoseBoneIndex != INDEX_NONE)
-			{
-				for (int32 Idx = 0; Idx < VBCompactPoseData.Num(); ++Idx)
-				{
-					FVirtualBoneCompactPoseData& VB = VBCompactPoseData[Idx];
-					if (PoseBoneIndex == VB.VBIndex)
-					{
-						// Remove this bone as we have written data for it (false so we dont resize allocation)
-						VBCompactPoseData.RemoveAtSwap(Idx, 1, false);
-						break; //Modified TArray so must break here
-					}
-				}
-				// extract animation
-
-				const FRawAnimSequenceTrack& TrackToExtract = InAnimationData[TrackIndex];
-
-				// Bail out (with rather wacky data) if data is empty for some reason.
-				if (TrackToExtract.PosKeys.Num() == 0 || TrackToExtract.RotKeys.Num() == 0)
-				{
-#if WITH_EDITORONLY_DATA
-					UE_LOG(LogAnimation, Warning, TEXT("UAnimSequence::GetBoneTransform : No anim data in AnimSequence '%s' Track '%s'"), *GetPathName(), *AnimationTrackNames[TrackIndex].ToString() );
-#endif
-
-					InOutPose[PoseBoneIndex].SetIdentity();
-
-					if (bInterpolateT)
-					{
-						Key2Pose[PoseBoneIndex].SetIdentity();
-					}
-				}
-				else
-				{
-					InOutPose[PoseBoneIndex] = ExtractTransformForKey(KeyIndex1, TrackToExtract);
-
-					if (bInterpolateT)
-					{
-						Key2Pose[PoseBoneIndex] = ExtractTransformForKey(KeyIndex2, TrackToExtract);
-					}
-				}
-
-				RetargetTracking.Add(FRetargetTracking(PoseBoneIndex, SkeletonBoneIndex));
-			}
-		}
-	}
-
-	//Build Virtual Bones
-	if (VBCompactPoseData.Num() > 0)
-	{
-		FCSPose<FCompactPose> CSPose1;
-		CSPose1.InitPose(InOutPose);
-
-		FCSPose<FCompactPose> CSPose2;
-		if (bInterpolateT)
-		{
-			CSPose2.InitPose(Key2Pose);
-		}
-
-		for (FVirtualBoneCompactPoseData& VB : VBCompactPoseData)
-		{
-			FTransform Source = CSPose1.GetComponentSpaceTransform(VB.SourceIndex);
-			FTransform Target = CSPose1.GetComponentSpaceTransform(VB.TargetIndex);
-			InOutPose[VB.VBIndex] = Target.GetRelativeTransform(Source);
-
-			if (bInterpolateT)
-			{
-				FTransform Source2 = CSPose2.GetComponentSpaceTransform(VB.SourceIndex);
-				FTransform Target2 = CSPose2.GetComponentSpaceTransform(VB.TargetIndex);
-				Key2Pose[VB.VBIndex] = Target2.GetRelativeTransform(Source2);
-			}
-		}
-	}
-
-	if (bInterpolateT)
-	{
-		for (FCompactPoseBoneIndex BoneIndex : InOutPose.ForEachBoneIndex())
-		{
-			InOutPose[BoneIndex].Blend(InOutPose[BoneIndex], Key2Pose[BoneIndex], Alpha);
-		}
-	}
-}
-
-void UAnimSequence::BuildPoseFromRawData(const TArray<FRawAnimSequenceTrack>& InAnimationData, FCompactPose& InOutPose, float InTime) const
-{
-	int32 KeyIndex1, KeyIndex2;
-	float Alpha;
-	FAnimationRuntime::GetKeyIndicesFromTime(KeyIndex1, KeyIndex2, Alpha, InTime, NumFrames, SequenceLength);
-
-	if (Interpolation == EAnimInterpolationType::Step)
-	{
-		Alpha = 0.f;
-	}
-
-	bool bInterpolate = true;
-
-	if (Alpha < KINDA_SMALL_NUMBER)
-	{
-		Alpha = 0.f;
-		bInterpolate = false;
-	}
-	else if (Alpha > 1.f - KINDA_SMALL_NUMBER)
-	{
-		bInterpolate = false;
-		KeyIndex1 = KeyIndex2;
-	}
-
-	if (bInterpolate)
-	{
-		BuildPoseFromRawDataInternal<true>(InAnimationData, InOutPose, KeyIndex1, KeyIndex2, Alpha);
-	}
-	else
-	{
-		BuildPoseFromRawDataInternal<false>(InAnimationData, InOutPose, KeyIndex1, KeyIndex2, Alpha);
-	}
-
-	const FBoneContainer& RequiredBones = InOutPose.GetBoneContainer();
-	const bool bDisableRetargeting = RequiredBones.GetDisableRetargeting();
-
-	if (!bDisableRetargeting)
-	{
-		const TArray<FRetargetTracking>& RetargetTracking = FGetBonePoseScratchArea::Get().RetargetTracking;
-
-		for (const FRetargetTracking& RT : RetargetTracking)
-		{
-			RetargetBoneTransform(InOutPose[RT.PoseBoneIndex], RT.SkeletonBoneIndex, RT.PoseBoneIndex, RequiredBones, false);
-		}
-	}
-
-}
-
 void UAnimSequence::GetBonePose(FCompactPose& OutPose, FBlendedCurve& OutCurve, const FAnimExtractContext& ExtractionContext, bool bForceUseRawData) const
 {
 	SCOPE_CYCLE_COUNTER(STAT_AnimSeq_GetBonePose);
@@ -1747,6 +1577,8 @@ void UAnimSequence::GetBonePose(FCompactPose& OutPose, FBlendedCurve& OutCurve,
 		return;
 	}
 
+	FRootMotionReset RootMotionReset(bEnableRootMotion, RootMotionRootLock, bForceRootLock, ExtractRootTrackTransform(0.f, &RequiredBones), IsValidAdditive());
+
 #if WITH_EDITOR
 	// this happens only with editor data
 	// Slower path for disable retargeting, that's only used in editor and for debugging.
@@ -1755,209 +1587,30 @@ void UAnimSequence::GetBonePose(FCompactPose& OutPose, FBlendedCurve& OutCurve,
 		const bool bShouldUseSourceData = (RequiredBones.ShouldUseSourceData() && SourceRawAnimationData.Num() > 0);
 		const TArray<FRawAnimSequenceTrack>& AnimationData = bShouldUseSourceData ? SourceRawAnimationData : RawAnimationData;
 
-		BuildPoseFromRawData(AnimationData, OutPose, ExtractionContext.CurrentTime);
+		// Warning if we have invalid data
 
-		if ((ExtractionContext.bExtractRootMotion && bEnableRootMotion) || bForceRootLock)
+		for (int32 TrackIndex = 0; TrackIndex < NumTracks; TrackIndex++)
 		{
-			ResetRootBoneForRootMotion(OutPose[FCompactPoseBoneIndex(0)], RequiredBones, RootMotionRootLock);
-		}
-		return;
-	}
-#endif // WITH_EDITOR
-
-	TArray<int32> const& SkeletonToPoseBoneIndexArray = RequiredBones.GetSkeletonToPoseBoneIndexArray();
-
-	BoneTrackArray& RotationScalePairs = FGetBonePoseScratchArea::Get().RotationScalePairs;
-	BoneTrackArray& TranslationPairs = FGetBonePoseScratchArea::Get().TranslationPairs;
-	BoneTrackArray& AnimScaleRetargetingPairs = FGetBonePoseScratchArea::Get().AnimScaleRetargetingPairs;
-	BoneTrackArray& AnimRelativeRetargetingPairs = FGetBonePoseScratchArea::Get().AnimRelativeRetargetingPairs;
-	BoneTrackArray& OrientAndScaleRetargetingPairs = FGetBonePoseScratchArea::Get().OrientAndScaleRetargetingPairs;
-
-	// build a list of desired bones
-	RotationScalePairs.Reset();
-	TranslationPairs.Reset();
-	AnimScaleRetargetingPairs.Reset();
-	AnimRelativeRetargetingPairs.Reset();
-	OrientAndScaleRetargetingPairs.Reset();
+			const FRawAnimSequenceTrack& TrackToExtract = AnimationData[TrackIndex];
 
-	// Optimization: assuming first index is root bone. That should always be the case in Skeletons.
-	checkSlow((SkeletonToPoseBoneIndexArray[0] == 0));
-	// this is not guaranteed for AnimSequences though... If Root is not animated, Track will not exist.
-	const bool bFirstTrackIsRootBone = (GetSkeletonIndexFromCompressedDataTrackIndex(0) == 0);
-
-	{
-		SCOPE_CYCLE_COUNTER(STAT_BuildAnimTrackPairs);
-
-		// Handle root bone separately if it is track 0. so we start w/ Index 1.
-		for (int32 TrackIndex = (bFirstTrackIsRootBone ? 1 : 0); TrackIndex < NumTracks; TrackIndex++)
-		{
-			const int32 SkeletonBoneIndex = GetSkeletonIndexFromCompressedDataTrackIndex(TrackIndex);
-			// not sure it's safe to assume that SkeletonBoneIndex can never be INDEX_NONE
-			if (SkeletonBoneIndex != INDEX_NONE)
+			// Bail out (with rather wacky data) if data is empty for some reason.
+			if (TrackToExtract.PosKeys.Num() == 0 || TrackToExtract.RotKeys.Num() == 0)
 			{
-				const FCompactPoseBoneIndex BoneIndex = RequiredBones.GetCompactPoseIndexFromSkeletonIndex(SkeletonBoneIndex);
-				//Nasty, we break our type safety, code in the lower levels should be adjusted for this
-				const int32 CompactPoseBoneIndex = BoneIndex.GetInt();
-				if (CompactPoseBoneIndex != INDEX_NONE)
-				{
-					RotationScalePairs.Add(BoneTrackPair(CompactPoseBoneIndex, TrackIndex));
-
-					// Skip extracting translation component for EBoneTranslationRetargetingMode::Skeleton.
-					switch (MySkeleton->GetBoneTranslationRetargetingMode(SkeletonBoneIndex))
-					{
-					case EBoneTranslationRetargetingMode::Animation:
-						TranslationPairs.Add(BoneTrackPair(CompactPoseBoneIndex, TrackIndex));
-						break;
-					case EBoneTranslationRetargetingMode::AnimationScaled:
-						TranslationPairs.Add(BoneTrackPair(CompactPoseBoneIndex, TrackIndex));
-						AnimScaleRetargetingPairs.Add(BoneTrackPair(CompactPoseBoneIndex, SkeletonBoneIndex));
-						break;
-					case EBoneTranslationRetargetingMode::AnimationRelative:
-						TranslationPairs.Add(BoneTrackPair(CompactPoseBoneIndex, TrackIndex));
-
-						// With baked additives, we can skip 'AnimationRelative' tracks, as the relative transform gets canceled out.
-						// (A1 + Rel) - (A2 + Rel) = A1 - A2.
-						if (!bIsBakedAdditive)
-						{
-							AnimRelativeRetargetingPairs.Add(BoneTrackPair(CompactPoseBoneIndex, SkeletonBoneIndex));
-						}
-						break;
-					case EBoneTranslationRetargetingMode::OrientAndScale:
-						TranslationPairs.Add(BoneTrackPair(CompactPoseBoneIndex, TrackIndex));
-
-						// Additives remain additives, they're not retargeted.
-						if (!bIsBakedAdditive)
-						{
-							OrientAndScaleRetargetingPairs.Add(BoneTrackPair(CompactPoseBoneIndex, SkeletonBoneIndex));
-						}
-						break;
-					}
-				}
+				UE_LOG(LogAnimation, Warning, TEXT("UAnimSequence::GetBoneTransform : No anim data in AnimSequence '%s' Track '%s'"), *GetPathName(), *AnimationTrackNames[TrackIndex].ToString());
 			}
 		}
-	}
 
-	{
-		SCOPE_CYCLE_COUNTER(STAT_ExtractPoseFromAnimData);
-		CSV_SCOPED_TIMING_STAT(Animation, ExtractPoseFromAnimData);
-		CSV_CUSTOM_STAT(Animation, NumberOfExtractedAnimations, 1, ECsvCustomStatOp::Accumulate);
-
-		FAnimSequenceDecompressionContext EvalDecompContext(SequenceLength, Interpolation, GetFName(), CompressedData.CompressedDataStructure);
-		EvalDecompContext.Seek(ExtractionContext.CurrentTime);
+		BuildPoseFromRawData(AnimationData, TrackToSkeletonMapTable, OutPose, ExtractionContext.CurrentTime, Interpolation, NumFrames, SequenceLength, RetargetSource);
 
-		// Handle Root Bone separately
-		if (bFirstTrackIsRootBone)
+		if ((ExtractionContext.bExtractRootMotion && RootMotionReset.bEnableRootMotion) || RootMotionReset.bForceRootLock)
 		{
-			const int32 TrackIndex = 0;
-			FCompactPoseBoneIndex RootBone(0);
-			FTransform& RootAtom = OutPose[RootBone];
-
-			AnimationFormat_GetBoneAtom(
-				RootAtom,
-				EvalDecompContext,
-				TrackIndex);
-
-			// @laurent - we should look into splitting rotation and translation tracks, so we don't have to process translation twice.
-			RetargetBoneTransform(RootAtom, 0, RootBone, RequiredBones, bIsBakedAdditive);
-		}
-
-		if (RotationScalePairs.Num() > 0)
-		{
-			// get the remaining bone atoms
-			FTransformArray LocalBones;
-			OutPose.MoveBonesTo(LocalBones);
-
-			AnimationFormat_GetAnimationPose(
-				LocalBones,
-				RotationScalePairs,
-				TranslationPairs,
-				RotationScalePairs,
-				EvalDecompContext);
-
-			OutPose.MoveBonesFrom(MoveTemp(LocalBones));
-		}
-	}
-
-	// Once pose has been extracted, snap root bone back to first frame if we are extracting root motion.
-	if ((ExtractionContext.bExtractRootMotion && bEnableRootMotion) || bForceRootLock)
-	{
-		ResetRootBoneForRootMotion(OutPose[FCompactPoseBoneIndex(0)], RequiredBones, RootMotionRootLock);
-	}
-
-	// Anim Scale Retargeting
-	int32 const NumBonesToScaleRetarget = AnimScaleRetargetingPairs.Num();
-	if (NumBonesToScaleRetarget > 0)
-	{
-		TArray<FTransform> const& AuthoredOnRefSkeleton = MySkeleton->GetRefLocalPoses(RetargetSource);
-
-		for (const BoneTrackPair& BonePair : AnimScaleRetargetingPairs)
-		{
-			const FCompactPoseBoneIndex BoneIndex(BonePair.AtomIndex); //Nasty, we break our type safety, code in the lower levels should be adjusted for this
-			int32 const& SkeletonBoneIndex = BonePair.TrackIndex;
-
-			// @todo - precache that in FBoneContainer when we have SkeletonIndex->TrackIndex mapping. So we can just apply scale right away.
-			float const SourceTranslationLength = AuthoredOnRefSkeleton[SkeletonBoneIndex].GetTranslation().Size();
-			if (SourceTranslationLength > KINDA_SMALL_NUMBER)
-			{
-				float const TargetTranslationLength = RequiredBones.GetRefPoseTransform(BoneIndex).GetTranslation().Size();
-				OutPose[BoneIndex].ScaleTranslation(TargetTranslationLength / SourceTranslationLength);
-			}
-		}
-	}
-
-	// Anim Relative Retargeting
-	int32 const NumBonesToRelativeRetarget = AnimRelativeRetargetingPairs.Num();
-	if (NumBonesToRelativeRetarget > 0)
-	{
-		TArray<FTransform> const& AuthoredOnRefSkeleton = MySkeleton->GetRefLocalPoses(RetargetSource);
-	
-		for (const BoneTrackPair& BonePair : AnimRelativeRetargetingPairs)
-		{
-			const FCompactPoseBoneIndex BoneIndex(BonePair.AtomIndex); //Nasty, we break our type safety, code in the lower levels should be adjusted for this
-			int32 const& SkeletonBoneIndex = BonePair.TrackIndex;
-
-			const FTransform& RefPose = RequiredBones.GetRefPoseTransform(BoneIndex);
-
-			// Apply the retargeting as if it were an additive difference between the current skeleton and the retarget skeleton. 
-			OutPose[BoneIndex].SetRotation(OutPose[BoneIndex].GetRotation() * AuthoredOnRefSkeleton[SkeletonBoneIndex].GetRotation().Inverse() * RefPose.GetRotation());
-			OutPose[BoneIndex].SetTranslation(OutPose[BoneIndex].GetTranslation() + (RefPose.GetTranslation() - AuthoredOnRefSkeleton[SkeletonBoneIndex].GetTranslation()));
-			OutPose[BoneIndex].SetScale3D(OutPose[BoneIndex].GetScale3D() * (RefPose.GetScale3D() * AuthoredOnRefSkeleton[SkeletonBoneIndex].GetSafeScaleReciprocal(AuthoredOnRefSkeleton[SkeletonBoneIndex].GetScale3D())));
-			OutPose[BoneIndex].NormalizeRotation();
+			RootMotionReset.ResetRootBoneForRootMotion(OutPose[FCompactPoseBoneIndex(0)], RequiredBones);
 		}
+		return;
 	}
+#endif // WITH_EDITOR
 
-	// Translation 'Orient and Scale' Translation Retargeting
-	const int32 NumBonesToOrientAndScaleRetarget = OrientAndScaleRetargetingPairs.Num();
-	if (NumBonesToOrientAndScaleRetarget > 0)
-	{
-		const FRetargetSourceCachedData& RetargetSourceCachedData = RequiredBones.GetRetargetSourceCachedData(RetargetSource);
-		const TArray<FOrientAndScaleRetargetingCachedData>& OrientAndScaleDataArray = RetargetSourceCachedData.OrientAndScaleData;
-		const TArray<int32>& CompactPoseIndexToOrientAndScaleIndex = RetargetSourceCachedData.CompactPoseIndexToOrientAndScaleIndex;
-
-		// If we have any cached retargeting data.
-		if ((OrientAndScaleDataArray.Num() > 0) && (CompactPoseIndexToOrientAndScaleIndex.Num() == RequiredBones.GetCompactPoseNumBones()))
-		{
-			for (int32 Index = 0; Index < NumBonesToOrientAndScaleRetarget; Index++)
-			{
-				const BoneTrackPair& BonePair = OrientAndScaleRetargetingPairs[Index];
-				const FCompactPoseBoneIndex CompactPoseBoneIndex(BonePair.AtomIndex);
-				const int32 OrientAndScaleIndex = CompactPoseIndexToOrientAndScaleIndex[CompactPoseBoneIndex.GetInt()];
-				if (OrientAndScaleIndex != INDEX_NONE)
-				{
-					const FOrientAndScaleRetargetingCachedData& OrientAndScaleData = OrientAndScaleDataArray[OrientAndScaleIndex];
-					FTransform& BoneTransform = OutPose[CompactPoseBoneIndex];
-					const FVector AnimatedTranslation = BoneTransform.GetTranslation();
-
-					// If Translation is not animated, we can just copy the TargetTranslation. No retargeting needs to be done.
-					const FVector NewTranslation = (AnimatedTranslation - OrientAndScaleData.SourceTranslation).IsNearlyZero(BONE_TRANS_RT_ORIENT_AND_SCALE_PRECISION) ?
-						OrientAndScaleData.TargetTranslation :
-						OrientAndScaleData.TranslationDeltaOrient.RotateVector(AnimatedTranslation) * OrientAndScaleData.TranslationScale;
-
-					BoneTransform.SetTranslation(NewTranslation);
-				}
-			}
-		}
-	}
+	DecompressPose(OutPose, CompressedData, ExtractionContext, GetSkeleton(), SequenceLength, Interpolation, bIsBakedAdditive, RetargetSource, GetFName(), RootMotionReset);
 }
 
 #if WITH_EDITORONLY_DATA
@@ -2542,35 +2195,52 @@ void UAnimSequence::RequestAnimCompression(FRequestAnimCompressionParams Params)
 		const bool bPerformFrameStripping = Params.bPerformFrameStripping && bAllowFrameStripping;
 		const bool bPerformStrippingOnOddFramedAnims = Params.bPerformFrameStrippingOnOddNumberedFrames;
 
-		const int32 PreviousCompressionSize = GetApproxCompressedSize();
+		// Can only do stripping on animations that have an even number of frames once the end frame is removed)
+		const bool bIsEvenFramed = ((NumFrames - 1) % 2) == 0;
+		const bool bIsValidForStripping = bIsEvenFramed || bPerformStrippingOnOddFramedAnims;
+
+		const bool bStripCandidate = (NumFrames > 10) && bIsValidForStripping;
 
-		// Filter RAW data to get rid of mismatched tracks (translation/rotation data with a different number of keys than there are frames)
-		// No trivial key removal is done at this point (impossible error metrics of -1), since all of the techniques will perform it themselves
-		CompressRawAnimData(-1.0f, -1.0f);
+		const bool bPerformStripping = bStripCandidate && bPerformFrameStripping;
+
+		const int32 PreviousCompressionSize = GetApproxCompressedSize();
 
-		FCompressibleAnimData CompressibleData(this); //Build the data that the compression code will compress
+		FString AssetDDCKey = GetAnimSequenceSpecificCacheKeySuffix(*this, bPerformStripping);
 
 		TArray<uint8> OutData;
-		FDerivedDataAnimationCompression* AnimCompressor = new FDerivedDataAnimationCompression(CompressibleData, Params.CompressContext, PreviousCompressionSize, bPerformFrameStripping, bPerformStrippingOnOddFramedAnims);
-		// For debugging DDC/Compression issues		
-		const bool bSkipDDC = false;
-		if (bSkipDDC || (CompressCommandletVersion == INDEX_NONE))
 		{
-			AnimCompressor->Build(OutData);
+			bool bNeedToCleanUpAnimCompressor = true;
+			FDerivedDataAnimationCompression* AnimCompressor = new FDerivedDataAnimationCompression(TEXT("AnimSeq"), AssetDDCKey, Params.CompressContext, PreviousCompressionSize);
 
-			delete AnimCompressor;
-			AnimCompressor = nullptr;
-		}
-		else
-		{
-			if (AnimCompressor->CanBuild())
+			const FString FinalDDCKey = FDerivedDataCacheInterface::BuildCacheKey(AnimCompressor->GetPluginName(), AnimCompressor->GetVersionString(), *AnimCompressor->GetPluginSpecificCacheKeySuffix());
+
+			// For debugging DDC/Compression issues		
+			const bool bSkipDDC = false;
+
+			if (bSkipDDC || !GetDerivedDataCacheRef().GetSynchronous(*FinalDDCKey, OutData))
 			{
-				GetDerivedDataCacheRef().GetSynchronous(AnimCompressor, OutData);
+				// Data does not exist, need to build it.
+				// Filter RAW data to get rid of mismatched tracks (translation/rotation data with a different number of keys than there are frames)
+				// No trivial key removal is done at this point (impossible error metrics of -1), since all of the techniques will perform it themselves
+				CompressRawAnimData(-1.0f, -1.0f);
+
+				TSharedRef<FCompressibleAnimData> CompressibleData = MakeShared<FCompressibleAnimData>(this, bPerformStripping);
+				AnimCompressor->SetCompressibleData(CompressibleData);
+
+				if (bSkipDDC || (CompressCommandletVersion == INDEX_NONE))
+				{
+					AnimCompressor->Build(OutData);
+				}
+				else if (AnimCompressor->CanBuild())
+				{
+					bNeedToCleanUpAnimCompressor = false; // GetSynchronous will handle this
+					GetDerivedDataCacheRef().GetSynchronous(AnimCompressor, OutData);
+				}
 			}
-			else
+			
+			if(bNeedToCleanUpAnimCompressor)
 			{
-				// If we dont perform compression we need to clean this up
-				delete AnimCompressor;
+				delete AnimCompressor; // Would really like to do auto mem management but GetDerivedDataCacheRef().GetSynchronous expects a point it can delete, every
 				AnimCompressor = nullptr;
 			}
 		}
@@ -2591,7 +2261,7 @@ void UAnimSequence::SerializeCompressedData(FArchive& Ar, bool bDDCData)
 {
 	if (!HasAnyFlags(RF_ClassDefaultObject))
 	{
-		CompressedData.SerializeCompressedData(Ar, bDDCData, this, CurveCompressionSettings);
+		CompressedData.SerializeCompressedData(Ar, bDDCData, this, this->GetSkeleton(), CurveCompressionSettings);
 	}
 }
 
@@ -2604,47 +2274,6 @@ bool UAnimSequence::CanBakeAdditive() const
 			GetSkeleton();
 }
 
-FFloatCurve* GetFloatCurve(FRawCurveTracks& RawCurveTracks, USkeleton::AnimCurveUID& CurveUID)
-{
-	return static_cast<FFloatCurve *>(RawCurveTracks.GetCurveData(CurveUID, ERawCurveTrackTypes::RCT_Float));
-}
-
-bool IsNewKeyDifferent(const FRichCurveKey& LastKey, float NewValue)
-{
-	return LastKey.Value != NewValue;
-}
-
-template <typename ArrayType>
-void UpdateSHAWithArray(FSHA1& Sha, const TArray<ArrayType>& Array)
-{
-	Sha.Update((uint8*)Array.GetData(), Array.Num() * Array.GetTypeSize());
-}
-
-void UpdateSHAWithRawTrack(FSHA1& Sha, const FRawAnimSequenceTrack& RawTrack)
-{
-	UpdateSHAWithArray(Sha, RawTrack.PosKeys);
-	UpdateSHAWithArray(Sha, RawTrack.RotKeys);
-	UpdateSHAWithArray(Sha, RawTrack.ScaleKeys);
-}
-
-template<class DataType>
-void UpdateWithData(FSHA1& Sha, const DataType& Data)
-{
-	Sha.Update((uint8*)(&Data), sizeof(DataType));
-}
-
-void UAnimSequence::UpdateSHAWithCurves(FSHA1& Sha, const FRawCurveTracks& InRawCurveData) const
-{
-	for (const FFloatCurve& Curve : InRawCurveData.FloatCurves)
-	{
-		UpdateWithData(Sha, Curve.Name.UID);
-		UpdateWithData(Sha, Curve.FloatCurve.DefaultValue);
-		UpdateSHAWithArray(Sha, Curve.FloatCurve.GetConstRefOfKeys());
-		UpdateWithData(Sha, Curve.FloatCurve.PreInfinityExtrap);
-		UpdateWithData(Sha, Curve.FloatCurve.PostInfinityExtrap);
-	}
-}
-
 bool UAnimSequence::DoesSequenceContainZeroScale()
 {
 	for (const FRawAnimSequenceTrack& RawTrack : RawAnimationData)
@@ -2663,21 +2292,8 @@ bool UAnimSequence::DoesSequenceContainZeroScale()
 
 FGuid UAnimSequence::GenerateGuidFromRawData() const
 {
-	FSHA1 Sha;
-
-	for (const FRawAnimSequenceTrack& Track : RawAnimationData)
-	{
-		UpdateSHAWithRawTrack(Sha, Track);
-	}
+	return GenerateGuidFromRawAnimData(RawAnimationData, RawCurveData);
 
-	UpdateSHAWithCurves(Sha, RawCurveData);
-
-	Sha.Final();
-
-	uint32 Hash[5];
-	Sha.GetHash((uint8*)Hash);
-	FGuid Guid(Hash[0] ^ Hash[4], Hash[1], Hash[2], Hash[3]);
-	return Guid;
 }
 
 void CopyTransformToRawAnimationData(const FTransform& BoneTransform, FRawAnimSequenceTrack& Track, int32 Frame)
@@ -2824,6 +2440,16 @@ void UAnimSequence::TestEvalauteAnimation() const
 	}
 }
 
+FFloatCurve* GetFloatCurve(FRawCurveTracks& RawCurveTracks, USkeleton::AnimCurveUID& CurveUID)
+{
+	return static_cast<FFloatCurve *>(RawCurveTracks.GetCurveData(CurveUID, ERawCurveTrackTypes::RCT_Float));
+}
+
+bool IsNewKeyDifferent(const FRichCurveKey& LastKey, float NewValue)
+{
+	return LastKey.Value != NewValue;
+}
+
 void UAnimSequence::BakeOutAdditiveIntoRawData(TArray<FRawAnimSequenceTrack>& NewRawTracks, TArray<FName>& NewAnimationTrackNames, TArray<FTrackToSkeletonMap>& NewTrackToSkeletonMapTable, FRawCurveTracks& NewCurveTracks, TArray<FRawAnimSequenceTrack>& AdditiveBaseAnimationData)
 {
 	AddAnimLoadingDebugEntry(TEXT("BakeOutAdditiveIntoRawData"));
@@ -2983,6 +2609,18 @@ void UAnimSequence::FlagDependentAnimationsAsRawDataOnly() const
 	}
 }
 
+void UAnimSequence::UpdateDependentStreamingAnimations() const
+{
+	for (TObjectIterator<UAnimStreamable> Iter; Iter; ++Iter)
+	{
+		UAnimStreamable* Seq = *Iter;
+		if (Seq->SourceSequence == this)
+		{
+			Seq->InitFrom(this);
+		}
+	}
+}
+
 #endif
 
 void UAnimSequence::RecycleAnimSequence()
diff --git a/Engine/Source/Runtime/Engine/Private/Animation/AnimSingleNodeInstanceProxy.cpp b/Engine/Source/Runtime/Engine/Private/Animation/AnimSingleNodeInstanceProxy.cpp
--- a/Engine/Source/Runtime/Engine/Private/Animation/AnimSingleNodeInstanceProxy.cpp
+++ b/Engine/Source/Runtime/Engine/Private/Animation/AnimSingleNodeInstanceProxy.cpp
@@ -5,6 +5,7 @@
 #include "Animation/AnimComposite.h"
 #include "Animation/BlendSpaceBase.h"
 #include "Animation/PoseAsset.h"
+#include "Animation/AnimStreamable.h"
 #include "Animation/AnimSingleNodeInstance.h"
 #include "AnimEncoding.h"
 
@@ -299,6 +300,37 @@ void FAnimNode_SingleNode::Evaluate_AnyThread(FPoseContext& Output)
 				Sequence->GetAnimationPose(Output.Pose, Output.Curve, FAnimExtractContext(Proxy->CurrentTime, Sequence->bEnableRootMotion));
 			}
 		}
+		else if (UAnimStreamable* Streamable = Cast<UAnimStreamable>(Proxy->CurrentAsset))
+		{
+			// No Additive support yet
+			/*if (Streamable->IsValidAdditive())
+			{
+				FAnimExtractContext ExtractionContext(Proxy->CurrentTime, Streamable->bEnableRootMotion);
+
+				if (bCanProcessAdditiveAnimationsLocal)
+				{
+					Sequence->GetAdditiveBasePose(Output.Pose, Output.Curve, ExtractionContext);
+				}
+				else
+				{
+					Output.ResetToRefPose();
+				}
+
+				FCompactPose AdditivePose;
+				FBlendedCurve AdditiveCurve;
+				AdditivePose.SetBoneContainer(&Output.Pose.GetBoneContainer());
+				AdditiveCurve.InitFrom(Output.Curve);
+				Streamable->GetAnimationPose(AdditivePose, AdditiveCurve, ExtractionContext);
+
+				FAnimationRuntime::AccumulateAdditivePose(Output.Pose, AdditivePose, Output.Curve, AdditiveCurve, 1.f, Sequence->AdditiveAnimType);
+				Output.Pose.NormalizeRotations();
+			}
+			else*/
+			{
+				// if SkeletalMesh isn't there, we'll need to use skeleton
+				Streamable->GetAnimationPose(Output.Pose, Output.Curve, FAnimExtractContext(Proxy->CurrentTime, Streamable->bEnableRootMotion));
+			}
+		}
 		else if (UAnimComposite* Composite = Cast<UAnimComposite>(Proxy->CurrentAsset))
 		{
 			FAnimExtractContext ExtractionContext(Proxy->CurrentTime, Proxy->ShouldExtractRootMotion());
@@ -490,6 +522,20 @@ void FAnimNode_SingleNode::Update_AnyThread(const FAnimationUpdateContext& Conte
 				}
 			}
 		}
+		else if (UAnimStreamable* Streamable = Cast<UAnimStreamable>(Proxy->CurrentAsset))
+		{
+			FAnimTickRecord& TickRecord = Proxy->CreateUninitializedTickRecord(INDEX_NONE, /*out*/ SyncGroup);
+			Proxy->MakeSequenceTickRecord(TickRecord, Streamable, Proxy->bLooping, NewPlayRate, 1.f, /*inout*/ Proxy->CurrentTime, Proxy->MarkerTickRecord);
+			// if it's not looping, just set play to be false when reached to end
+			if (!Proxy->bLooping)
+			{
+				const float CombinedPlayRate = NewPlayRate * Streamable->RateScale;
+				if ((CombinedPlayRate < 0.f && Proxy->CurrentTime <= 0.f) || (CombinedPlayRate > 0.f && Proxy->CurrentTime >= Streamable->SequenceLength))
+				{
+					Proxy->SetPlaying(false);
+				}
+			}
+		}
 		else if(UAnimComposite* Composite = Cast<UAnimComposite>(Proxy->CurrentAsset))
 		{
 			FAnimTickRecord& TickRecord = Proxy->CreateUninitializedTickRecord(INDEX_NONE, /*out*/ SyncGroup);
diff --git a/Engine/Source/Runtime/Engine/Private/Animation/AnimStreamable.cpp b/Engine/Source/Runtime/Engine/Private/Animation/AnimStreamable.cpp
new file mode 100644
--- /dev/null
+++ b/Engine/Source/Runtime/Engine/Private/Animation/AnimStreamable.cpp
@@ -0,0 +1,683 @@
+// Copyright 1998-2019 Epic Games, Inc. All Rights Reserved.
+
+/*=============================================================================
+	AnimStreamable.cpp: Animation that can be streamed instead of being loaded completely
+=============================================================================*/ 
+
+#include "Animation/AnimStreamable.h"
+#include "Interfaces/ITargetPlatform.h"
+#include "Interfaces/ITargetPlatformManagerModule.h"
+#include "DeviceProfiles/DeviceProfileManager.h"
+#include "DeviceProfiles/DeviceProfile.h"
+#include "Animation/AnimCompressionDerivedData.h"
+#include "DerivedDataCacheInterface.h"
+#include "Animation/AnimCurveCompressionSettings.h"
+#include "Animation/AnimCompress.h"
+#include "Animation/AnimCurveCompressionCodec.h"
+#include "BonePose.h"
+#include "ContentStreaming.h"
+
+DECLARE_CYCLE_STAT(TEXT("AnimStreamable GetAnimationPose"), STAT_AnimStreamable_GetAnimationPose, STATGROUP_Anim);
+
+// This is a version string for the streaming anim chunk logic
+// If you want to bump this version, generate a new guid using VS->Tools->Create GUID and
+// set it here
+const TCHAR* StreamingAnimChunkVersion = TEXT("1F1656B9E10142729AB16650D9821B1Fuf");
+
+const float MINIMUM_CHUNK_SIZE = 1.0f;
+
+float GChunkSizeSeconds = MINIMUM_CHUNK_SIZE;
+
+static const TCHAR* ChunkSizeSecondsCVarName = TEXT("a.Streaming.ChunkSizeSeconds");
+
+static FAutoConsoleVariableRef CVarChunkSizeSeconds(
+	ChunkSizeSecondsCVarName,
+	GChunkSizeSeconds,
+	TEXT("Size of streaming animation chunk in seconds, 0 or negative signifies only have 1 chunk"));
+
+void FAnimStreamableChunk::Serialize(FArchive& Ar, UAnimStreamable* Owner, int32 ChunkIndex)
+{
+	DECLARE_SCOPE_CYCLE_COUNTER(TEXT("FAnimStreamableChunk::Serialize"), STAT_StreamedAnimChunk_Serialize, STATGROUP_LoadTime);
+
+	bool bCooked = Ar.IsCooking();
+	Ar << bCooked;
+
+	Ar << StartTime;
+	Ar << SequenceLength;
+
+	if(bCooked)
+	{
+		if (ChunkIndex == 0)
+		{
+			// Chunk 0 just serializes the compressed data directly
+			if (Ar.IsLoading())
+			{
+				check(!CompressedAnimSequence);
+				CompressedAnimSequence = new FCompressedAnimSequence();
+			}
+			CompressedAnimSequence->SerializeCompressedData(Ar, false, Owner, Owner->GetSkeleton(), Owner->CurveCompressionSettings, false);
+		}
+		else
+		{
+			BulkData.SetBulkDataFlags(BULKDATA_Force_NOT_InlinePayload);
+
+			if (Ar.IsSaving())
+			{
+				//Need to pack compressed data into BulkData
+				TArray<uint8> TempBytes;
+				const int32 InitialSize = CompressedAnimSequence->CompressedDataStructure.GetApproxBoneCompressedSize();
+				TempBytes.Reset(InitialSize);
+
+				FMemoryWriter TempAr(TempBytes, true);
+				CompressedAnimSequence->SerializeCompressedData(TempAr, false, Owner, Owner->GetSkeleton(), Owner->CurveCompressionSettings, false);
+
+				BulkData.Lock(LOCK_READ_WRITE);
+				void* ChunkData = BulkData.Realloc(TempBytes.Num());
+				FMemory::Memcpy(ChunkData, TempBytes.GetData(), TempBytes.Num());
+				BulkData.Unlock();
+			}
+
+			// streaming doesn't use memory mapped IO
+			BulkData.Serialize(Ar, Owner, ChunkIndex, false);
+		}
+	}
+}
+
+void FStreamableAnimPlatformData::Serialize(FArchive& Ar, UAnimStreamable* Owner)
+{
+	int32 NumChunks = Chunks.Num();
+	Ar << NumChunks;
+
+	if (Ar.IsLoading())
+	{
+		Chunks.Reset(NumChunks);
+		Chunks.AddDefaulted(NumChunks);
+	}
+	for (int32 ChunkIndex = 0; ChunkIndex < NumChunks; ++ChunkIndex)
+	{
+		Chunks[ChunkIndex].Serialize(Ar, Owner, ChunkIndex);
+	}
+}
+
+UAnimStreamable::UAnimStreamable(const FObjectInitializer& ObjectInitializer)
+	: Super(ObjectInitializer)
+#if WITH_EDITOR
+	, RunningAnimPlatformData(nullptr)
+#endif
+{
+	bUseRawDataOnly = true;
+}
+
+void UAnimStreamable::PreSave(const class ITargetPlatform* TargetPlatform)
+{
+#if WITH_EDITOR
+	if (TargetPlatform)
+	{
+		RequestCompressedData(TargetPlatform); //Make sure target platform data is built
+	}
+
+#endif
+
+	Super::PreSave(TargetPlatform);
+}
+
+void UAnimStreamable::Serialize(FArchive& Ar)
+{
+	Super::Serialize(Ar);
+	
+
+	bool bCooked = Ar.IsCooking();
+	Ar << bCooked;
+
+	if (FPlatformProperties::RequiresCookedData() && !bCooked && Ar.IsLoading())
+	{
+		UE_LOG(LogAnimation, Fatal, TEXT("This platform requires cooked packages, and animation data was not cooked into %s."), *GetFullName());
+	}
+
+	if (bCooked)
+	{
+		if (Ar.IsLoading())
+		{
+			GetRunningPlatformData().Serialize(Ar, this);
+		}
+		else
+		{
+			FStreamableAnimPlatformData& PlatformData = GetStreamingAnimPlatformData(Ar.CookingTarget());
+			PlatformData.Serialize(Ar, this);
+		}
+	}
+}
+
+#if WITH_EDITOR
+/*bool UAnimComposite::GetAllAnimationSequencesReferred(TArray<UAnimationAsset*>& AnimationAssets, bool bRecursive) 
+{
+	return AnimationTrack.GetAllAnimationSequencesReferred(AnimationAssets, bRecursive);
+}
+
+void UAnimComposite::ReplaceReferredAnimations(const TMap<UAnimationAsset*, UAnimationAsset*>& ReplacementMap)
+{
+	AnimationTrack.ReplaceReferredAnimations(ReplacementMap);
+}*/
+#endif
+
+
+void UAnimStreamable::HandleAssetPlayerTickedInternal(FAnimAssetTickContext &Context, const float PreviousTime, const float MoveDelta, const FAnimTickRecord &Instance, struct FAnimNotifyQueue& NotifyQueue) const
+{
+	Super::HandleAssetPlayerTickedInternal(Context, PreviousTime, MoveDelta, Instance, NotifyQueue);
+	
+	const int32 ChunkIndex = GetChunkIndexForTime(GetRunningPlatformData().Chunks, PreviousTime);
+	IAnimationStreamingManager& StreamingManager = IStreamingManager::Get().GetAnimationStreamingManager();
+	const FCompressedAnimSequence* CurveCompressedDataChunk = StreamingManager.GetLoadedChunk(this, ChunkIndex);
+	
+	//ExtractRootMotionFromTrack(AnimationTrack, PreviousTime, PreviousTime + MoveDelta, Context.RootMotionMovementParams);
+}
+
+void UAnimStreamable::GetAnimationPose(FCompactPose& OutPose, FBlendedCurve& OutCurve, const FAnimExtractContext& ExtractionContext) const
+{
+	OutPose.ResetToRefPose();
+
+	SCOPE_CYCLE_COUNTER(STAT_AnimStreamable_GetAnimationPose);
+	CSV_SCOPED_TIMING_STAT(Animation, AnimStreamable_GetAnimationPose);
+
+	const FBoneContainer& RequiredBones = OutPose.GetBoneContainer();
+	//const bool bUseRawDataForPoseExtraction = bForceUseRawData || UseRawDataForPoseExtraction(RequiredBones);
+
+	const bool bIsBakedAdditive = false;//!bUseRawDataForPoseExtraction && IsValidAdditive();
+
+	const USkeleton* MySkeleton = GetSkeleton();
+	if (!MySkeleton)
+	{
+		if (bIsBakedAdditive)
+		{
+			OutPose.ResetToAdditiveIdentity();
+		}
+		else
+		{
+			OutPose.ResetToRefPose();
+		}
+		return;
+	}
+
+	const bool bDisableRetargeting = RequiredBones.GetDisableRetargeting();
+
+	// initialize with ref-pose
+	if (bIsBakedAdditive)
+	{
+		//When using baked additive ref pose is identity
+		OutPose.ResetToAdditiveIdentity();
+	}
+	else
+	{
+		// if retargeting is disabled, we initialize pose with 'Retargeting Source' ref pose.
+		if (bDisableRetargeting)
+		{
+			TArray<FTransform> const& AuthoredOnRefSkeleton = MySkeleton->GetRefLocalPoses(RetargetSource);
+			TArray<FBoneIndexType> const& RequireBonesIndexArray = RequiredBones.GetBoneIndicesArray();
+
+			int32 const NumRequiredBones = RequireBonesIndexArray.Num();
+			for (FCompactPoseBoneIndex PoseBoneIndex : OutPose.ForEachBoneIndex())
+			{
+				int32 const& SkeletonBoneIndex = RequiredBones.GetSkeletonIndex(PoseBoneIndex);
+
+				// Pose bone index should always exist in Skeleton
+				checkSlow(SkeletonBoneIndex != INDEX_NONE);
+				OutPose[PoseBoneIndex] = AuthoredOnRefSkeleton[SkeletonBoneIndex];
+			}
+		}
+		else
+		{
+			OutPose.ResetToRefPose();
+		}
+	}
+
+	//FRootMotionReset RootMotionReset(bEnableRootMotion, RootMotionRootLock, bForceRootLock, ExtractRootTrackTransform(0.f, &RequiredBones), IsValidAdditive());
+	FRootMotionReset RootMotionReset(bEnableRootMotion, RootMotionRootLock, bForceRootLock, FTransform(), false); // MDW Does not support root motion yet
+
+#if WITH_EDITOR
+	if (!HasRunningPlatformData())
+	{
+		//Need to evaluate raw data
+		RawCurveData.EvaluateCurveData(OutCurve, ExtractionContext.CurrentTime);
+
+		const int32 NumTracks = TrackToSkeletonMapTable.Num();
+
+		// Warning if we have invalid data
+
+		for (int32 TrackIndex = 0; TrackIndex < NumTracks; TrackIndex++)
+		{
+			const FRawAnimSequenceTrack& TrackToExtract = RawAnimationData[TrackIndex];
+
+			// Bail out (with rather wacky data) if data is empty for some reason.
+			if (TrackToExtract.PosKeys.Num() == 0 || TrackToExtract.RotKeys.Num() == 0)
+			{
+				UE_LOG(LogAnimation, Warning, TEXT("UAnimSequence::GetBoneTransform : No anim data in AnimSequence '%s' Track '%s'"), *GetPathName(), *AnimationTrackNames[TrackIndex].ToString());
+			}
+		}
+
+		BuildPoseFromRawData(RawAnimationData, TrackToSkeletonMapTable, OutPose, ExtractionContext.CurrentTime, Interpolation, NumFrames, SequenceLength, RetargetSource);
+
+		if ((ExtractionContext.bExtractRootMotion && RootMotionReset.bEnableRootMotion) || RootMotionReset.bForceRootLock)
+		{
+			RootMotionReset.ResetRootBoneForRootMotion(OutPose[FCompactPoseBoneIndex(0)], RequiredBones);
+		}
+		return;
+	}
+#endif
+
+	const int32 ChunkIndex = GetChunkIndexForTime(GetRunningPlatformData().Chunks, ExtractionContext.CurrentTime);
+	if(ensureMsgf(ChunkIndex != INDEX_NONE, TEXT("Could not get valid chunk with Time %.2f for Streaming Anim %s"), ExtractionContext.CurrentTime, *GetFullName()))
+	{
+		IAnimationStreamingManager& StreamingManager = IStreamingManager::Get().GetAnimationStreamingManager();
+
+		const FCompressedAnimSequence* CurveCompressedDataChunk = StreamingManager.GetLoadedChunk(this, 0); //Curve Data stored in chunk 0 till it is properly cropped
+
+		if (!CurveCompressedDataChunk)
+		{
+#if WITH_EDITOR
+			CurveCompressedDataChunk = GetRunningPlatformData().Chunks[0].CompressedAnimSequence;
+#else
+			UE_LOG(LogAnimation, Warning, TEXT("Failed to get streamed compressed data Time: %.2f, ChunkIndex:%i, Anim: %s"), ExtractionContext.CurrentTime, 0, *GetFullName());
+			return;
+#endif
+		}
+
+		CurveCompressedDataChunk->CurveCompressionCodec->DecompressCurves(*CurveCompressedDataChunk, OutCurve, ExtractionContext.CurrentTime);
+
+		const FCompressedAnimSequence* CompressedData = (ChunkIndex == 0) ? CurveCompressedDataChunk : StreamingManager.GetLoadedChunk(this, ChunkIndex);
+
+		if (!CompressedData)
+		{
+#if WITH_EDITOR
+			CompressedData = GetRunningPlatformData().Chunks[ChunkIndex].CompressedAnimSequence;
+#else
+			UE_LOG(LogAnimation, Warning, TEXT("Failed to get streamed compressed data Time: %.2f, ChunkIndex:%i, Anim: %s"), ExtractionContext.CurrentTime, ChunkIndex, *GetFullName());
+			return;
+#endif
+		}
+
+		const int32 NumTracks = CompressedData->CompressedTrackToSkeletonMapTable.Num();
+		if (NumTracks == 0)
+		{
+			return;
+		}
+
+		float ChunkCurrentTime = ExtractionContext.CurrentTime - GetRunningPlatformData().Chunks[ChunkIndex].StartTime;
+		FAnimExtractContext ChunkExtractionContext(ChunkCurrentTime, ExtractionContext.bExtractRootMotion);
+		ChunkExtractionContext.BonesRequired = ExtractionContext.BonesRequired;
+		ChunkExtractionContext.PoseCurves = ExtractionContext.PoseCurves;
+
+		//FPlatformMisc::LowLevelOutputDebugStringf(TEXT("Playing Streaming Anim %s Time: %.2f Chunk:%i\n"), *GetName(), ExtractionContext.CurrentTime, ChunkIndex);
+		DecompressPose(OutPose, *CompressedData, ChunkExtractionContext, GetSkeleton(), GetRunningPlatformData().Chunks[ChunkIndex].SequenceLength, Interpolation, bIsBakedAdditive, RetargetSource, GetFName(), RootMotionReset);
+	}
+}
+
+void UAnimStreamable::PostLoad()
+{
+	//Parent PostLoad will ensure that skeleton is fully loaded
+	//before we do anything further in PostLoad
+	Super::PostLoad();
+
+#if WITH_EDITOR
+	RequestCompressedData(); // Grab compressed data for current platform
+#else
+	IStreamingManager::Get().GetAnimationStreamingManager().AddStreamingAnim(this); // This will be handled by RequestCompressedData in editor builds
+
+	if (USkeleton* CurrentSkeleton = GetSkeleton())
+	{
+		for (FSmartName& CurveName : GetRunningPlatformData().Chunks[0].CompressedAnimSequence->CompressedCurveNames)
+		{
+			CurrentSkeleton->VerifySmartName(USkeleton::AnimCurveMappingName, CurveName);
+		}
+	}
+#endif
+}
+
+void UAnimStreamable::FinishDestroy()
+{
+	Super::FinishDestroy();
+
+	IStreamingManager::Get().GetAnimationStreamingManager().RemoveStreamingAnim(this);
+}
+
+int32 UAnimStreamable::GetChunkIndexForTime(const TArray<FAnimStreamableChunk>& Chunks, float CurrentTime) const
+{
+	for (int32 ChunkIndex = 0; ChunkIndex < Chunks.Num(); ++ChunkIndex)
+	{
+		const float ChunkLength = Chunks[ChunkIndex].SequenceLength;
+
+		if (CurrentTime < ChunkLength)
+		{
+			return ChunkIndex;
+		}
+		CurrentTime -= ChunkLength;
+	}
+	return Chunks.Num() - 1;
+}
+
+#if WITH_EDITOR
+void UAnimStreamable::InitFrom(const UAnimSequence* InSourceSequence)
+{
+	SetSkeleton(InSourceSequence->GetSkeleton());
+	SourceSequence = InSourceSequence;
+
+	RawAnimationData = InSourceSequence->GetRawAnimationData();
+	RawCurveData = InSourceSequence->RawCurveData;
+
+	Notifies = InSourceSequence->Notifies;
+
+	TrackToSkeletonMapTable = InSourceSequence->GetRawTrackToSkeletonMapTable();
+	AnimationTrackNames = InSourceSequence->GetAnimationTrackNames();
+	
+	NumFrames = InSourceSequence->GetNumberOfFrames();
+	SequenceLength = InSourceSequence->SequenceLength;
+	
+	RateScale = InSourceSequence->RateScale;
+
+	Interpolation = InSourceSequence->Interpolation;
+
+	RetargetSource = InSourceSequence->RetargetSource;
+
+	bEnableRootMotion = InSourceSequence->bEnableRootMotion;
+	RootMotionRootLock = InSourceSequence->RootMotionRootLock;
+	bForceRootLock = InSourceSequence->bForceRootLock;
+	bUseNormalizedRootMotionScale = InSourceSequence->bUseNormalizedRootMotionScale;
+
+	UpdateRawData();
+}
+
+FString GetChunkDDCKey(const FString& BaseKey, uint32 ChunkIndex, uint32 NumChunks)
+{
+	FArcToHexString HexStringBuilder;
+
+	HexStringBuilder.Ar << ChunkIndex;
+	HexStringBuilder.Ar << NumChunks;
+
+	return FString::Printf(TEXT("%s%s"),
+		*BaseKey,
+		*HexStringBuilder.MakeString()
+	);
+}
+
+
+void UAnimStreamable::RequestCompressedData(const ITargetPlatform* Platform)
+{
+	check(IsInGameThread());
+
+	bUseRawDataOnly = true;
+
+	USkeleton* CurrentSkeleton = GetSkeleton();
+	if (CurrentSkeleton == nullptr)
+	{
+		return;
+	}
+
+	if (GetOutermost() == GetTransientPackage())
+	{
+		return; // Skip transient animations, they are most likely the leftovers of previous compression attempts.
+	}
+
+	if (FPlatformProperties::RequiresCookedData())
+	{
+		return;
+	}
+	
+	ITargetPlatformManagerModule* TPM = GetTargetPlatformManager();
+	if (!TPM)
+	{
+		return; // No platform manager
+	}
+
+	if (!Platform)
+	{
+		Platform = TPM->GetRunningTargetPlatform();
+	}
+
+	const bool bIsRunningPlatform = (Platform == TPM->GetRunningTargetPlatform());
+
+	if (bIsRunningPlatform)
+	{
+		IStreamingManager::Get().GetAnimationStreamingManager().RemoveStreamingAnim(this);
+	}
+
+	if (!CompressionScheme)
+	{
+		CompressionScheme = FAnimationUtils::GetDefaultAnimationCompressionAlgorithm();
+	}
+
+	if (CurveCompressionSettings == nullptr || !CurveCompressionSettings->AreSettingsValid())
+	{
+		CurveCompressionSettings = FAnimationUtils::GetDefaultAnimationCurveCompressionSettings();
+	}
+
+	checkf(Platform, TEXT("Failed to specify platform for streamable animation compression"));
+
+	FStreamableAnimPlatformData& PlatformData = GetStreamingAnimPlatformData(Platform);
+
+	if (bIsRunningPlatform)
+	{
+		RunningAnimPlatformData = &PlatformData;
+	}
+
+	PlatformData.Reset();
+
+	float ChunkSizeSeconds = GetChunkSizeSeconds(Platform);
+
+	uint32 NumChunks = 1;
+	if (ChunkSizeSeconds > 0.f) // <= 0.f signifies to not chunk
+	{
+		ChunkSizeSeconds = FMath::Max(ChunkSizeSeconds, MINIMUM_CHUNK_SIZE);
+		const int32 InitialNumChunks = FMath::FloorToInt(SequenceLength / ChunkSizeSeconds);
+		NumChunks = FMath::Max(InitialNumChunks, 1);
+	}
+
+	int32 NumFramesToChunk = NumFrames - 1;
+	int32 FramesPerChunk = NumFrames / NumChunks;
+
+	PlatformData.Chunks.AddDefaulted(NumChunks);
+
+	const FString BaseDDCKey = GetBaseDDCKey();
+
+	const bool bInAllowAlternateCompressor = false;
+	const bool bInOutput				   = false;
+
+	TSharedRef<FAnimCompressContext> CompressContext = MakeShared<FAnimCompressContext>(bInAllowAlternateCompressor, bInOutput);
+
+	for (uint32 ChunkIndex = 0; ChunkIndex < NumChunks; ++ChunkIndex)
+	{
+		FString ChunkDDCKey = GetChunkDDCKey(BaseDDCKey, ChunkIndex, NumChunks);
+
+		const bool bLastChunk = (ChunkIndex == (NumChunks - 1));
+		const uint32 FrameStart = ChunkIndex * FramesPerChunk;
+		const uint32 FrameEnd = bLastChunk ? NumFramesToChunk : (ChunkIndex + 1) * FramesPerChunk;
+
+		RequestCompressedDataForChunk(ChunkDDCKey, PlatformData.Chunks[ChunkIndex], FrameStart, FrameEnd, CompressContext);
+	}
+
+	if (bIsRunningPlatform)
+	{
+		IStreamingManager::Get().GetAnimationStreamingManager().AddStreamingAnim(this);
+	}
+	//PlatformData.SetSkeletonVirtualBoneGuid(GetSkeleton()->GetVirtualBoneGuid()); //MDW DO THIS
+	//PlatformData.bUseRawDataOnly = false; //MDW Need to do something with this? 
+}
+
+float UAnimStreamable::GetChunkSizeSeconds(const ITargetPlatform* Platform) const
+{
+	float CVarPlatformChunkSizeSeconds = 1.f;
+	if (UDeviceProfile* DeviceProfile = UDeviceProfileManager::Get().FindProfile(Platform->IniPlatformName()))
+	{
+		if (!DeviceProfile->GetConsolidatedCVarValue(ChunkSizeSecondsCVarName, CVarPlatformChunkSizeSeconds))
+		{
+			CVarPlatformChunkSizeSeconds = MINIMUM_CHUNK_SIZE; //Could not get value so set to Minimum
+		}
+	}
+	
+	//UE_LOG(LogAnimation, Log, TEXT("Anim Chunk Size for platform %s : %.2f\n"), *Platform->DisplayName().ToString(), CVarPlatformChunkSizeSeconds);
+	return CVarPlatformChunkSizeSeconds;
+}
+
+template<typename KeyType>
+void MakeKeyChunk(const TArray<KeyType>& SrcKeys, TArray<KeyType>& DestKeys, int32 NumFrames, const uint32 FrameStart, const uint32 FrameEnd)
+{
+	if (SrcKeys.Num() == 1)
+	{
+		DestKeys.Add(SrcKeys[0]);
+	}
+	else
+	{
+		check(SrcKeys.Num() == NumFrames); // Invalid data otherwise
+
+		DestKeys.Reset((FrameEnd - FrameStart) + 1);
+		for (uint32 FrameIndex = FrameStart; FrameIndex <= FrameEnd; ++FrameIndex)
+		{
+			DestKeys.Add(SrcKeys[FrameIndex]);
+		}
+	}
+}
+
+void UAnimStreamable::RequestCompressedDataForChunk(const FString& ChunkDDCKey, FAnimStreamableChunk& Chunk, const uint32 FrameStart, const uint32 FrameEnd, TSharedRef<FAnimCompressContext> CompressContext)
+{
+	// Need to unify with Anim Sequence!
+
+	TArray<uint8> OutData;
+	{
+		bool bNeedToCleanUpAnimCompressor = true;
+		FDerivedDataAnimationCompression* AnimCompressor = new FDerivedDataAnimationCompression(TEXT("StreamAnim"), ChunkDDCKey, CompressContext, 0);
+
+		const FString FinalDDCKey = FDerivedDataCacheInterface::BuildCacheKey(AnimCompressor->GetPluginName(), AnimCompressor->GetVersionString(), *AnimCompressor->GetPluginSpecificCacheKeySuffix());
+
+		// For debugging DDC/Compression issues		
+		const bool bSkipDDC = false;
+
+		const int32 ChunkNumFrames = FrameEnd - FrameStart;
+		const float FrameLength = SequenceLength / (float)(NumFrames - 1);
+		Chunk.StartTime = FrameStart * FrameLength;
+		Chunk.SequenceLength = ChunkNumFrames * FrameLength;
+
+		if (bSkipDDC || !GetDerivedDataCacheRef().GetSynchronous(*FinalDDCKey, OutData))
+		{
+			TSharedRef<FCompressibleAnimData> CompressibleData = MakeShared<FCompressibleAnimData>(CompressionScheme, CurveCompressionSettings, GetSkeleton(), Interpolation, Chunk.SequenceLength, ChunkNumFrames+1);
+
+			CompressibleData->RawAnimationData.AddDefaulted(RawAnimationData.Num());
+
+			for (int32 TrackIndex = 0; TrackIndex < RawAnimationData.Num(); ++TrackIndex)
+			{
+				FRawAnimSequenceTrack& SrcTrack = RawAnimationData[TrackIndex];
+				FRawAnimSequenceTrack& DestTrack = CompressibleData->RawAnimationData[TrackIndex];
+
+				MakeKeyChunk(SrcTrack.PosKeys, DestTrack.PosKeys, NumFrames, FrameStart, FrameEnd);
+				MakeKeyChunk(SrcTrack.RotKeys, DestTrack.RotKeys, NumFrames, FrameStart, FrameEnd);
+				if (SrcTrack.ScaleKeys.Num() > 0)
+				{
+					MakeKeyChunk(SrcTrack.ScaleKeys, DestTrack.ScaleKeys, NumFrames, FrameStart, FrameEnd);
+				}
+			}
+
+			if (FrameStart == 0)
+			{
+				//Crop curve logic broken, for the moment store curve data in always loaded chunk 0
+				CompressibleData->RawCurveData = RawCurveData;
+			}
+
+			CompressibleData->TrackToSkeletonMapTable = TrackToSkeletonMapTable;
+			AnimCompressor->SetCompressibleData(CompressibleData);
+
+			if (bSkipDDC)
+			{
+				AnimCompressor->Build(OutData);
+			}
+			else if (AnimCompressor->CanBuild())
+			{
+				bNeedToCleanUpAnimCompressor = false; // GetSynchronous will handle this
+				GetDerivedDataCacheRef().GetSynchronous(AnimCompressor, OutData);
+			}
+		}
+
+		if (bNeedToCleanUpAnimCompressor)
+		{
+			delete AnimCompressor; // Would really like to do auto mem management but GetDerivedDataCacheRef().GetSynchronous expects a point it can delete, every
+			AnimCompressor = nullptr;
+		}
+	}
+
+	check(OutData.Num() > 0); // Should always have "something"
+	{
+		FMemoryReader MemAr(OutData);
+
+		if (!Chunk.CompressedAnimSequence)
+		{
+			Chunk.CompressedAnimSequence = new FCompressedAnimSequence();
+		}
+		Chunk.CompressedAnimSequence->SerializeCompressedData(MemAr, true, this, this->GetSkeleton(), CurveCompressionSettings);
+	}
+}
+
+void UAnimStreamable::UpdateRawData()
+{
+	RawDataGuid = GenerateGuidFromRawAnimData(RawAnimationData, RawCurveData);
+	RequestCompressedData();
+}
+
+FString UAnimStreamable::GetBaseDDCKey() const
+{
+	//Make up our content key consisting of:
+	//  * Streaming Anim Chunk logic version
+	//	* Our raw data GUID
+	//	* Our skeleton GUID: If our skeleton changes our compressed data may now be stale
+	//  * Our skeletons virtual bone guid
+	//	* Compression Settings
+	//	* Curve compression settings
+
+	FString Ret = FString::Printf(TEXT("%s%s%s%s_%s_%s"),
+		StreamingAnimChunkVersion,
+		*RawDataGuid.ToString(),
+		*GetSkeleton()->GetGuid().ToString(),
+		*GetSkeleton()->GetVirtualBoneGuid().ToString(),
+		*CompressionScheme->MakeDDCKey(),
+		*CurveCompressionSettings->MakeDDCKey()
+	);
+
+	return Ret;
+}
+
+#endif
+
+FStreamableAnimPlatformData& UAnimStreamable::GetStreamingAnimPlatformData(const ITargetPlatform* Platform)
+{
+#if WITH_EDITOR
+	FStreamableAnimPlatformData** Data = StreamableAnimPlatformData.Find(Platform);
+	if (!Data)
+	{
+		return *StreamableAnimPlatformData.Add(Platform, new FStreamableAnimPlatformData());
+	}
+	checkf(*Data, TEXT("StreamableAnimPlatformData contains nullptr!"));
+	return **Data;
+#else
+	return RunningAnimPlatformData;
+#endif
+}
+
+/*EAdditiveAnimationType UAnimComposite::GetAdditiveAnimType() const
+{
+	int32 AdditiveType = AnimationTrack.GetTrackAdditiveType();
+
+	if (AdditiveType != -1)
+	{
+		return (EAdditiveAnimationType)AdditiveType;
+	}
+
+	return AAT_None;
+}
+
+bool UAnimComposite::HasRootMotion() const
+{
+	return AnimationTrack.HasRootMotion();
+}*/
+
+#if WITH_EDITOR
+/*class UAnimSequence* UAnimComposite::GetAdditiveBasePose() const
+{
+	// @todo : for now it just picks up the first sequence
+	return AnimationTrack.GetAdditiveBasePose();
+}*/
+#endif 
diff --git a/Engine/Source/Runtime/Engine/Private/Animation/AnimationStreaming.cpp b/Engine/Source/Runtime/Engine/Private/Animation/AnimationStreaming.cpp
new file mode 100644
--- /dev/null
+++ b/Engine/Source/Runtime/Engine/Private/Animation/AnimationStreaming.cpp
@@ -0,0 +1,607 @@
+// Copyright 1998-2019 Epic Games, Inc. All Rights Reserved.
+
+/*=============================================================================
+AnimationStreaming.cpp: Manager to handle streaming animation data
+=============================================================================*/
+
+#include "Animation/AnimationStreaming.h"
+#include "Misc/CoreStats.h"
+#include "Animation/AnimStreamable.h"
+#include "Algo/Find.h"
+/*#include "Sound/SoundWave.h"
+#include "Sound/AudioSettings.h"
+#include "DerivedDataCacheInterface.h"
+#include "Serialization/MemoryReader.h"
+#include "HAL/PlatformFile.h"
+#include "HAL/PlatformFilemanager.h"
+#include "Async/AsyncFileHandle.h"
+#include "Misc/ScopeLock.h"
+#include "HAL/IConsoleManager.h"
+#include "HAL/LowLevelMemTracker.h"
+#include "AudioDecompress.h"*/
+
+static int32 SpoofFailedAnimationChunkLoad = 0;
+FAutoConsoleVariableRef CVarSpoofFailedAnimationChunkLoad(
+	TEXT("a.Streaming.SpoofFailedChunkLoad"),
+	SpoofFailedAnimationChunkLoad,
+	TEXT("Forces failing to load streamed animation chunks.\n")
+	TEXT("0: Not Enabled, 1: Enabled"),
+	ECVF_Default);
+
+
+/*------------------------------------------------------------------------------
+	Streaming chunks from the derived data cache.
+------------------------------------------------------------------------------*/
+
+#if WITH_EDITORONLY_DATA
+/*
+// Initialization constructor.
+FAsyncStreamDerivedChunkWorker::FAsyncStreamDerivedChunkWorker(
+	const FString& InDerivedDataKey,
+	void* InDestChunkData,
+	int32 InChunkSize,
+	FThreadSafeCounter* InThreadSafeCounter
+	)
+	: DerivedDataKey(InDerivedDataKey)
+	, DestChunkData(InDestChunkData)
+	, ExpectedChunkSize(InChunkSize)
+	, bRequestFailed(false)
+	, ThreadSafeCounter(InThreadSafeCounter)
+{
+}
+
+// Retrieves the derived chunk from the derived data cache.
+void FAsyncStreamDerivedChunkWorker::DoWork()
+{
+	DECLARE_SCOPE_CYCLE_COUNTER(TEXT("FAsyncStreamDerivedChunkWorker::DoWork"), STAT_AsyncStreamDerivedChunkWorker_DoWork, STATGROUP_StreamingDetails);
+
+	UE_LOG(LogAudio, Verbose, TEXT("Start of ASync DDC Chunk read for key: %s"), *DerivedDataKey);
+
+	TArray<uint8> DerivedChunkData;
+
+	if (GetDerivedDataCacheRef().GetSynchronous(*DerivedDataKey, DerivedChunkData))
+	{
+		FMemoryReader Ar(DerivedChunkData, true);
+		int32 ChunkSize = 0;
+		Ar << ChunkSize;
+		checkf(ChunkSize == ExpectedChunkSize, TEXT("ChunkSize(%d) != ExpectedSize(%d)"), ChunkSize, ExpectedChunkSize);
+		Ar.Serialize(DestChunkData, ChunkSize);
+	}
+	else
+	{
+		bRequestFailed = true;
+	}
+	FPlatformMisc::MemoryBarrier();
+	ThreadSafeCounter->Decrement();
+
+	UE_LOG(LogAudio, Verbose, TEXT("End of ASync DDC Chunk read for key: %s"), *DerivedDataKey);
+}
+*/
+#endif // #if WITH_EDITORONLY_DATA
+
+void FLoadedAnimationChunk::CleanUpIORequest()
+{
+	if (IORequest)
+	{
+		IORequest->WaitCompletion();
+		delete IORequest;
+		IORequest = nullptr;
+	}
+}
+
+////////////////////////
+// FStreamingAnimationData //
+////////////////////////
+
+FStreamingAnimationData::FStreamingAnimationData()
+	: StreamableAnim(NULL)
+	, IORequestHandle(nullptr)
+	, AnimationStreamingManager(nullptr)
+{
+	ResetRequestedChunks();
+}
+
+FStreamingAnimationData::~FStreamingAnimationData()
+{
+	check(IORequestHandle == nullptr);
+}
+
+void FStreamingAnimationData::FreeResources()
+{
+	// Make sure there are no pending requests in flight.
+	for (int32 Pass = 0; Pass < 3; Pass++)
+	{
+		BlockTillAllRequestsFinished();
+		if (!UpdateStreamingStatus())
+		{
+			break;
+		}
+		check(Pass < 2); // we should be done after two passes. Pass 0 will start anything we need and pass 1 will complete those requests
+	}
+
+	for (FLoadedAnimationChunk& LoadedChunk : LoadedChunks)
+	{
+		FreeLoadedChunk(LoadedChunk);
+	}
+
+	if (IORequestHandle)
+	{
+		delete IORequestHandle;
+		IORequestHandle = nullptr;
+	}
+}
+
+bool FStreamingAnimationData::Initialize(UAnimStreamable* InStreamableAnim, FAnimationStreamingManager* InAnimationStreamingManager)
+{
+	check(!IORequestHandle);
+
+	check(InStreamableAnim && InStreamableAnim->HasRunningPlatformData());
+
+	FStreamableAnimPlatformData& RunningAnimPlatformData = InStreamableAnim->GetRunningPlatformData();
+
+	if (!RunningAnimPlatformData.Chunks.Num())
+	{
+		UE_LOG(LogAnimation, Error, TEXT("Failed to initialize streaming animation due to lack of anim or serialized stream chunks. '%s'"), *InStreamableAnim->GetFullName());
+		return false;
+	}
+
+	StreamableAnim = InStreamableAnim;
+	AnimationStreamingManager = InAnimationStreamingManager;
+
+	// Always get the first chunk of data so we can play immediately
+	check(LoadedChunks.Num() == 0);
+	check(LoadedChunkIndices.Num() == 0);
+
+	AddNewLoadedChunk(0, RunningAnimPlatformData.Chunks[0].CompressedAnimSequence);
+	LoadedChunkIndices.Add(0);
+
+	return true;
+}
+
+bool FStreamingAnimationData::UpdateStreamingStatus()
+{
+	if (!StreamableAnim)
+	{
+		return false;
+	}
+
+	bool bHasPendingRequestInFlight = false;
+
+	TArray<uint32> IndicesToLoad;
+	TArray<uint32> IndicesToFree;
+
+	if (HasPendingRequests(IndicesToLoad, IndicesToFree))
+	{
+		// could maybe iterate over the things we know are done, but I couldn't tell if that was IndicesToLoad or not.
+		for (FLoadedAnimationChunk& LoadedChunk : LoadedChunks)
+		{
+			if (LoadedChunk.IORequest)
+			{
+				const bool bRequestFinished = LoadedChunk.IORequest->PollCompletion();
+				bHasPendingRequestInFlight |= !bRequestFinished;
+				if (bRequestFinished)
+				{
+					LoadedChunk.CleanUpIORequest();
+				}
+			}
+		}
+
+		LoadedChunkIndices = RequestedChunks;
+
+		BeginPendingRequests(IndicesToLoad, IndicesToFree);
+	}
+
+	ResetRequestedChunks();
+
+	return bHasPendingRequestInFlight;
+}
+
+/*void FStreamingAnimationData::UpdateChunkRequests(FWaveRequest& InWaveRequest)
+{
+	// Might change this but ensures chunk 0 stays loaded for now
+	check(InWaveRequest.RequiredIndices.Contains(0));
+
+	CurrentRequest = InWaveRequest;
+
+}*/
+
+bool FStreamingAnimationData::HasPendingRequests(TArray<uint32>& IndicesToLoad, TArray<uint32>& IndicesToFree) const
+{
+	IndicesToLoad.Reset();
+	IndicesToFree.Reset();
+
+	// Find indices that aren't loaded
+	for (auto NeededIndex : RequestedChunks)
+	{
+		if (!LoadedChunkIndices.Contains(NeededIndex))
+		{
+			IndicesToLoad.AddUnique(NeededIndex);
+		}
+	}
+
+	// Find indices that aren't needed anymore
+	for (auto CurrentIndex : LoadedChunkIndices)
+	{
+		if (!RequestedChunks.Contains(CurrentIndex))
+		{
+			IndicesToFree.AddUnique(CurrentIndex);
+		}
+	}
+
+	return IndicesToLoad.Num() > 0 || IndicesToFree.Num() > 0;
+}
+
+void FStreamingAnimationData::BeginPendingRequests(const TArray<uint32>& IndicesToLoad, const TArray<uint32>& IndicesToFree)
+{
+	/*if (UE_LOG_ACTIVE(LogAudio, Log) && IndicesToLoad.Num() > 0)
+	{
+		FString LogString = FString::Printf(TEXT("Requesting ASync load of chunk(s) %d"), IndicesToLoad[0]);
+		for (int32 Index = 1; Index < IndicesToLoad.Num(); ++Index)
+		{
+			LogString += FString::Printf(TEXT(", %d"), IndicesToLoad[Index]);
+		}
+		LogString += FString::Printf(TEXT(" from SoundWave'%s'"), *SoundWave->GetName());
+		UE_LOG(LogAudio, Log, TEXT("%s"), *LogString);
+	}*/
+
+	TArray<uint32> FreeChunkIndices;
+
+	// Mark Chunks for removal in case they can be reused
+	{
+		for (auto Index : IndicesToFree)
+		{
+			for (int32 ChunkIndex = 0; ChunkIndex < LoadedChunks.Num(); ++ChunkIndex)
+			{
+				check(Index != 0);
+				if (LoadedChunks[ChunkIndex].Index == Index)
+				{
+					FreeLoadedChunk(LoadedChunks[ChunkIndex]);
+					LoadedChunks.RemoveAtSwap(ChunkIndex,1,false);
+					break;
+				}
+			}
+		}
+	}
+
+	// Set off all IO Requests
+
+	const EAsyncIOPriorityAndFlags AsyncIOPriority = AIOP_High;
+
+	for (auto ChunkIndex : IndicesToLoad)
+	{
+		const FAnimStreamableChunk& Chunk = StreamableAnim->GetRunningPlatformData().Chunks[ChunkIndex];
+
+		FCompressedAnimSequence* ExistingCompressedData = Chunk.CompressedAnimSequence;
+
+		FLoadedAnimationChunk& ChunkStorage = AddNewLoadedChunk(ChunkIndex, ExistingCompressedData);
+
+		if(!ExistingCompressedData)
+		{
+			check(!ChunkStorage.CompressedAnimData);
+			check(Chunk.BulkData.GetFilename().Len());
+			UE_CLOG(Chunk.BulkData.IsStoredCompressedOnDisk(), LogAnimation, Fatal, TEXT("Package level compression is not supported for streaming animation."));
+			check(!ChunkStorage.IORequest);
+			
+			if (!IORequestHandle)
+			{
+				IORequestHandle = FPlatformFileManager::Get().GetPlatformFile().OpenAsyncRead(*Chunk.BulkData.GetFilename());
+				check(IORequestHandle); // this generally cannot fail because it is async
+			}
+
+			int64 ChunkSize = Chunk.BulkData.GetBulkDataSize();
+			ChunkStorage.RequestStart = FPlatformTime::Seconds();
+			FPlatformMisc::LowLevelOutputDebugStringf(TEXT("Request Started %.2f\n"), ChunkStorage.RequestStart);
+			FAsyncFileCallBack AsyncFileCallBack =
+				[this, ChunkIndex, ChunkSize](bool bWasCancelled, IAsyncReadRequest* Req)
+			{
+				AnimationStreamingManager->OnAsyncFileCallback(this, ChunkIndex, ChunkSize, Req);
+			};
+
+			FPlatformMisc::LowLevelOutputDebugStringf(TEXT("Loading Stream Chunk %s Chunk:%i Length: %.3f Offset:%i Size:%i File:%i\n"), *StreamableAnim->GetName(), ChunkIndex, Chunk.SequenceLength, Chunk.BulkData.GetBulkDataOffsetInFile(), Chunk.BulkData.GetBulkDataSize(), *Chunk.BulkData.GetFilename());
+			ChunkStorage.IORequest = IORequestHandle->ReadRequest(Chunk.BulkData.GetBulkDataOffsetInFile(), Chunk.BulkData.GetBulkDataSize(), AsyncIOPriority, &AsyncFileCallBack);
+			if (!ChunkStorage.IORequest)
+			{
+				UE_LOG(LogAnimation, Error, TEXT("Animation streaming read request failed."));
+			}
+		}
+	}
+}
+
+bool FStreamingAnimationData::BlockTillAllRequestsFinished(float TimeLimit)
+{
+	QUICK_SCOPE_CYCLE_COUNTER(FStreamingAnimData_BlockTillAllRequestsFinished);
+	if (TimeLimit == 0.0f)
+	{
+		for (FLoadedAnimationChunk& LoadedChunk : LoadedChunks)
+		{
+			LoadedChunk.CleanUpIORequest();
+		}
+	}
+	else
+	{
+		double EndTime = FPlatformTime::Seconds() + TimeLimit;
+		for (FLoadedAnimationChunk& LoadedChunk : LoadedChunks)
+		{
+			if (LoadedChunk.IORequest)
+			{
+				float ThisTimeLimit = EndTime - FPlatformTime::Seconds();
+				if (ThisTimeLimit < .001f || // one ms is the granularity of the platform event system
+					!LoadedChunk.IORequest->WaitCompletion(ThisTimeLimit))
+				{
+					return false;
+				}
+
+				LoadedChunk.CleanUpIORequest();
+			}
+		}
+	}
+	return true;
+}
+
+FLoadedAnimationChunk& FStreamingAnimationData::AddNewLoadedChunk(uint32 ChunkIndex, FCompressedAnimSequence* ExistingData)
+{
+	int32 NewIndex = LoadedChunks.Num();
+	LoadedChunks.AddDefaulted();
+
+	LoadedChunks[NewIndex].CompressedAnimData = ExistingData;
+	LoadedChunks[NewIndex].bOwnsCompressedData = false;
+	LoadedChunks[NewIndex].Index = ChunkIndex;
+	return LoadedChunks[NewIndex];
+}
+
+void FStreamingAnimationData::FreeLoadedChunk(FLoadedAnimationChunk& LoadedChunk)
+{
+	if (LoadedChunk.IORequest)
+	{
+		LoadedChunk.IORequest->Cancel();
+		LoadedChunk.IORequest->WaitCompletion();
+		delete LoadedChunk.IORequest;
+		LoadedChunk.IORequest = nullptr;
+	}
+
+	if (LoadedChunk.bOwnsCompressedData)
+	{
+		delete LoadedChunk.CompressedAnimData;
+
+		//DEC_DWORD_STAT_BY(STAT_AudioMemorySize, LoadedChunk.DataSize);
+		//DEC_DWORD_STAT_BY(STAT_AudioMemory, LoadedChunk.DataSize);
+	}
+
+	LoadedChunk.CompressedAnimData = NULL;
+	LoadedChunk.bOwnsCompressedData = false;
+	LoadedChunk.Index = 0;
+}
+
+void FStreamingAnimationData::ResetRequestedChunks()
+{
+	RequestedChunks.Reset();
+	RequestedChunks.Add(0); //Always want 1
+}
+
+////////////////////////////
+// FAnimationStreamingManager //
+////////////////////////////
+
+FAnimationStreamingManager::FAnimationStreamingManager()
+{
+}
+
+FAnimationStreamingManager::~FAnimationStreamingManager()
+{
+}
+
+void FAnimationStreamingManager::OnAsyncFileCallback(FStreamingAnimationData* StreamingAnimData, int32 ChunkIndex, int64 ReadSize, IAsyncReadRequest* ReadRequest)
+{
+	// Check to see if we successfully managed to load anything
+	uint8* Mem = ReadRequest->GetReadResults();
+
+	if (Mem)
+	{
+		int32 LoadedChunkIndex = StreamingAnimData->LoadedChunks.IndexOfByPredicate([ChunkIndex](const FLoadedAnimationChunk& Chunk) { return Chunk.Index == ChunkIndex; });
+		FLoadedAnimationChunk& ChunkStorage = StreamingAnimData->LoadedChunks[LoadedChunkIndex];
+
+		checkf(!ChunkStorage.CompressedAnimData, TEXT("Chunk storage already has data. (0x%p)"), ChunkStorage.CompressedAnimData);
+
+		ChunkStorage.CompressedAnimData = new FCompressedAnimSequence();
+
+		TArrayView<const uint8> MemView(Mem, ReadSize);
+		FMemoryReaderView Reader(MemView);
+
+		UAnimStreamable* Anim = StreamingAnimData->StreamableAnim;
+		ChunkStorage.CompressedAnimData->SerializeCompressedData(Reader, false, Anim, Anim->GetSkeleton(), Anim->CurveCompressionSettings);
+
+		FPlatformMisc::LowLevelOutputDebugStringf(TEXT("Request Finished %.2f\nAnim Chunk Streamed %.4f\n"), FPlatformTime::Seconds(), FPlatformTime::Seconds() - ChunkStorage.RequestStart);
+		/*DEC_MEMORY_STAT_BY(STAT_AsyncFileMemory, ChunkStorage->DataSize);
+		INC_DWORD_STAT_BY(STAT_AudioMemorySize, ChunkStorage->DataSize);
+		INC_DWORD_STAT_BY(STAT_AudioMemory, ChunkStorage->DataSize);*/
+	}
+}
+
+void FAnimationStreamingManager::UpdateResourceStreaming(float DeltaTime, bool bProcessEverything /*= false*/)
+{
+	LLM_SCOPE(ELLMTag::Audio);
+
+	FScopeLock Lock(&CriticalSection);
+
+	for (auto& AnimData : StreamingAnimations)
+	{
+		AnimData.Value->UpdateStreamingStatus();
+	}
+
+	/*for (ICompressedAudioInfo* Decoder : CompressedAudioInfos)
+	{
+		USoundWave* SoundWave = Decoder->GetStreamingSoundWave();
+		if (SoundWave)
+		{
+			FStreamingWaveData** WaveDataPtr = StreamingSoundWaves.Find(SoundWave);
+			if (WaveDataPtr && (*WaveDataPtr)->PendingChunkChangeRequestStatus.GetValue() == AudioState_ReadyFor_Requests)
+			{
+				FStreamingWaveData* WaveData = *WaveDataPtr;
+				// Request the chunk the source is using and the one after that
+				FWaveRequest& WaveRequest = GetWaveRequest(SoundWave);
+				int32 SourceChunk = Decoder->GetCurrentChunkIndex();
+				if (SourceChunk >= 0 && SourceChunk < SoundWave->RunningPlatformData->NumChunks)
+				{
+					WaveRequest.RequiredIndices.AddUnique(SourceChunk);
+					WaveRequest.RequiredIndices.AddUnique((SourceChunk + 1) % SoundWave->RunningPlatformData->NumChunks);
+					WaveRequest.bPrioritiseRequest = true;
+				}
+				else
+				{
+					UE_LOG(LogAudio, Log, TEXT("Invalid chunk request curIndex=%d numChunks=%d\n"), SourceChunk, SoundWave->RunningPlatformData->NumChunks);
+				}
+			}
+		}
+	}
+
+	for (auto Iter = WaveRequests.CreateIterator(); Iter; ++Iter)
+	{
+		USoundWave* Wave = Iter.Key();
+		FStreamingWaveData* WaveData = StreamingSoundWaves.FindRef(Wave);
+
+		if (WaveData && WaveData->PendingChunkChangeRequestStatus.GetValue() == AudioState_ReadyFor_Requests)
+		{
+			WaveData->UpdateChunkRequests(Iter.Value());
+			WaveData->UpdateStreamingStatus();
+			Iter.RemoveCurrent();
+		}
+	}*/
+}
+
+int32 FAnimationStreamingManager::BlockTillAllRequestsFinished(float TimeLimit, bool)
+{
+	/*{
+		FScopeLock Lock(&CriticalSection);
+
+		QUICK_SCOPE_CYCLE_COUNTER(FAudioStreamingManager_BlockTillAllRequestsFinished);
+		int32 Result = 0;
+
+		if (TimeLimit == 0.0f)
+		{
+			for (auto& WavePair : StreamingSoundWaves)
+			{
+				WavePair.Value->BlockTillAllRequestsFinished();
+			}
+		}
+		else
+		{
+			double EndTime = FPlatformTime::Seconds() + TimeLimit;
+			for (auto& WavePair : StreamingSoundWaves)
+			{
+				float ThisTimeLimit = EndTime - FPlatformTime::Seconds();
+				if (ThisTimeLimit < .001f || // one ms is the granularity of the platform event system
+					!WavePair.Value->BlockTillAllRequestsFinished(ThisTimeLimit))
+				{
+					Result = 1; // we don't report the actual number, just 1 for any number of outstanding requests
+					break;
+				}
+			}
+		}
+
+		return Result;
+	}*/
+
+	// Not sure yet whether this will work the same as textures - aside from just before destroying
+	return 0;
+}
+
+void FAnimationStreamingManager::CancelForcedResources()
+{
+}
+
+void FAnimationStreamingManager::NotifyLevelChange()
+{
+}
+
+void FAnimationStreamingManager::SetDisregardWorldResourcesForFrames(int32 NumFrames)
+{
+}
+
+void FAnimationStreamingManager::AddLevel(class ULevel* Level)
+{
+}
+
+void FAnimationStreamingManager::RemoveLevel(class ULevel* Level)
+{
+}
+
+void FAnimationStreamingManager::NotifyLevelOffset(class ULevel* Level, const FVector& Offset)
+{
+}
+
+void FAnimationStreamingManager::AddStreamingAnim(UAnimStreamable* Anim)
+{
+	FScopeLock Lock(&CriticalSection);
+	if (StreamingAnimations.FindRef(Anim) == nullptr)
+	{
+		FStreamingAnimationData* NewStreamingAnim = new FStreamingAnimationData;
+		if (NewStreamingAnim->Initialize(Anim, this))
+		{
+			StreamingAnimations.Add(Anim, NewStreamingAnim);
+		}
+		else
+		{
+			// Failed to initialize, don't add to list of streaming sound waves
+			delete NewStreamingAnim;
+		}
+	}
+}
+
+bool FAnimationStreamingManager::RemoveStreamingAnim(UAnimStreamable* Anim)
+{
+	FScopeLock Lock(&CriticalSection);
+	FStreamingAnimationData* AnimData = StreamingAnimations.FindRef(Anim);
+	if (AnimData)
+	{
+		StreamingAnimations.Remove(Anim);
+
+		// Free the resources of the streaming wave data. This blocks pending IO requests
+		AnimData->FreeResources();
+
+		{
+			// Then we need to remove any results from those pending requests before we delete so that we don't process them
+			/*FScopeLock StreamChunkResults(&ChunkResultCriticalSection);
+			for (int32 i = AsyncAudioStreamChunkResults.Num() - 1 ; i >= 0; --i)
+			{
+				FASyncAudioChunkLoadResult* LoadResult = AsyncAudioStreamChunkResults[i];
+				FStreamingWaveData* StreamingWaveData = LoadResult->StreamingWaveData;
+				if (StreamingWaveData == WaveData)
+				{
+					delete LoadResult;
+					
+					AsyncAudioStreamChunkResults.RemoveAtSwap(i, 1, false);
+				}
+			}*/
+		}
+		delete AnimData;
+		return true;
+	}
+	return false;
+}
+
+const FCompressedAnimSequence* FAnimationStreamingManager::GetLoadedChunk(const UAnimStreamable* Anim, uint32 ChunkIndex) const
+{
+	// Check for the spoof of failing to load a stream chunk
+	if (SpoofFailedAnimationChunkLoad > 0)
+	{
+		return nullptr;
+	}
+
+	// If we fail at getting the critical section here, early out.
+	FScopeLock MapLock(&CriticalSection);
+
+	FStreamingAnimationData* AnimData = StreamingAnimations.FindRef(Anim);
+	if (AnimData)
+	{
+		AnimData->RequestedChunks.AddUnique(ChunkIndex);
+		AnimData->RequestedChunks.AddUnique((ChunkIndex + 1) % Anim->GetRunningPlatformData().Chunks.Num());
+
+		if (AnimData->LoadedChunkIndices.Contains(ChunkIndex))
+		{
+			const FLoadedAnimationChunk* Chunk = Algo::FindBy(AnimData->LoadedChunks, ChunkIndex, &FLoadedAnimationChunk::Index);
+			return Chunk ? Chunk->CompressedAnimData : nullptr;
+		}
+	}
+
+	return nullptr;
+}
\ No newline at end of file
diff --git a/Engine/Source/Runtime/Engine/Private/Animation/BonePose.cpp b/Engine/Source/Runtime/Engine/Private/Animation/BonePose.cpp
--- a/Engine/Source/Runtime/Engine/Private/Animation/BonePose.cpp
+++ b/Engine/Source/Runtime/Engine/Private/Animation/BonePose.cpp
@@ -2,6 +2,8 @@
 
 #include "BonePose.h"
 #include "AnimationRuntime.h"
+#include "AnimEncoding.h"
+#include "HAL/ThreadSingleton.h"
 
 void FMeshPose::ResetToRefPose()
 {
@@ -45,3 +47,192 @@ bool FMeshPose::IsNormalized() const
 	return true;
 }
 
+struct FRetargetTracking
+{
+	const FCompactPoseBoneIndex PoseBoneIndex;
+	const int32 SkeletonBoneIndex;
+
+	FRetargetTracking(const FCompactPoseBoneIndex InPoseBoneIndex, const int32 InSkeletonBoneIndex)
+		: PoseBoneIndex(InPoseBoneIndex), SkeletonBoneIndex(InSkeletonBoneIndex)
+	{
+	}
+};
+
+FTransform ExtractTransformForKey(int32 Key, const FRawAnimSequenceTrack &TrackToExtract)
+{
+	static const FVector DefaultScale3D = FVector(1.f);
+	const bool bHasScaleKey = TrackToExtract.ScaleKeys.Num() > 0;
+
+	const int32 PosKeyIndex = FMath::Min(Key, TrackToExtract.PosKeys.Num() - 1);
+	const int32 RotKeyIndex = FMath::Min(Key, TrackToExtract.RotKeys.Num() - 1);
+	if (bHasScaleKey)
+	{
+		const int32 ScaleKeyIndex = FMath::Min(Key, TrackToExtract.ScaleKeys.Num() - 1);
+		return FTransform(TrackToExtract.RotKeys[RotKeyIndex], TrackToExtract.PosKeys[PosKeyIndex], TrackToExtract.ScaleKeys[ScaleKeyIndex]);
+	}
+	else
+	{
+		return FTransform(TrackToExtract.RotKeys[RotKeyIndex], TrackToExtract.PosKeys[PosKeyIndex], DefaultScale3D);
+	}
+}
+
+struct FBuildRawPoseScratchArea : public TThreadSingleton<FBuildRawPoseScratchArea>
+{
+	TArray<FRetargetTracking> RetargetTracking;
+	TArray<FVirtualBoneCompactPoseData> VirtualBoneCompactPoseData;
+};
+
+
+template<bool bInterpolateT>
+void BuildPoseFromRawDataInternal(const TArray<FRawAnimSequenceTrack>& InAnimationData, const TArray<struct FTrackToSkeletonMap>& TrackToSkeletonMapTable, FCompactPose& InOutPose, int32 KeyIndex1, int32 KeyIndex2, float Alpha)
+{
+	const int32 NumTracks = InAnimationData.Num();
+	const FBoneContainer& RequiredBones = InOutPose.GetBoneContainer();
+
+	TArray<FRetargetTracking>& RetargetTracking = FBuildRawPoseScratchArea::Get().RetargetTracking;
+	RetargetTracking.Reset(NumTracks);
+
+	TArray<FVirtualBoneCompactPoseData>& VBCompactPoseData = FBuildRawPoseScratchArea::Get().VirtualBoneCompactPoseData;
+	VBCompactPoseData = RequiredBones.GetVirtualBoneCompactPoseData();
+
+	FCompactPose Key2Pose;
+	if (bInterpolateT)
+	{
+		Key2Pose.CopyBonesFrom(InOutPose);
+	}
+
+	for (int32 TrackIndex = 0; TrackIndex < NumTracks; TrackIndex++)
+	{
+		const int32 SkeletonBoneIndex = TrackToSkeletonMapTable[TrackIndex].BoneTreeIndex;
+		// not sure it's safe to assume that SkeletonBoneIndex can never be INDEX_NONE
+		if ((SkeletonBoneIndex != INDEX_NONE) && (SkeletonBoneIndex < MAX_BONES))
+		{
+			const FCompactPoseBoneIndex PoseBoneIndex = RequiredBones.GetCompactPoseIndexFromSkeletonIndex(SkeletonBoneIndex);
+			if (PoseBoneIndex != INDEX_NONE)
+			{
+				for (int32 Idx = 0; Idx < VBCompactPoseData.Num(); ++Idx)
+				{
+					FVirtualBoneCompactPoseData& VB = VBCompactPoseData[Idx];
+					if (PoseBoneIndex == VB.VBIndex)
+					{
+						// Remove this bone as we have written data for it (false so we dont resize allocation)
+						VBCompactPoseData.RemoveAtSwap(Idx, 1, false);
+						break; //Modified TArray so must break here
+					}
+				}
+				// extract animation
+
+				const FRawAnimSequenceTrack& TrackToExtract = InAnimationData[TrackIndex];
+
+				// Bail out (with rather wacky data) if data is empty for some reason.
+				if (TrackToExtract.PosKeys.Num() == 0 || TrackToExtract.RotKeys.Num() == 0)
+				{
+					InOutPose[PoseBoneIndex].SetIdentity();
+
+					if (bInterpolateT)
+					{
+						Key2Pose[PoseBoneIndex].SetIdentity();
+					}
+				}
+				else
+				{
+					InOutPose[PoseBoneIndex] = ExtractTransformForKey(KeyIndex1, TrackToExtract);
+
+					if (bInterpolateT)
+					{
+						Key2Pose[PoseBoneIndex] = ExtractTransformForKey(KeyIndex2, TrackToExtract);
+					}
+				}
+
+				RetargetTracking.Add(FRetargetTracking(PoseBoneIndex, SkeletonBoneIndex));
+			}
+		}
+	}
+
+	//Build Virtual Bones
+	if (VBCompactPoseData.Num() > 0)
+	{
+		FCSPose<FCompactPose> CSPose1;
+		CSPose1.InitPose(InOutPose);
+
+		FCSPose<FCompactPose> CSPose2;
+		if (bInterpolateT)
+		{
+			CSPose2.InitPose(Key2Pose);
+		}
+
+		for (FVirtualBoneCompactPoseData& VB : VBCompactPoseData)
+		{
+			FTransform Source = CSPose1.GetComponentSpaceTransform(VB.SourceIndex);
+			FTransform Target = CSPose1.GetComponentSpaceTransform(VB.TargetIndex);
+			InOutPose[VB.VBIndex] = Target.GetRelativeTransform(Source);
+
+			if (bInterpolateT)
+			{
+				FTransform Source2 = CSPose2.GetComponentSpaceTransform(VB.SourceIndex);
+				FTransform Target2 = CSPose2.GetComponentSpaceTransform(VB.TargetIndex);
+				Key2Pose[VB.VBIndex] = Target2.GetRelativeTransform(Source2);
+			}
+		}
+	}
+
+	if (bInterpolateT)
+	{
+		for (FCompactPoseBoneIndex BoneIndex : InOutPose.ForEachBoneIndex())
+		{
+			InOutPose[BoneIndex].Blend(InOutPose[BoneIndex], Key2Pose[BoneIndex], Alpha);
+		}
+	}
+}
+
+void BuildPoseFromRawData(const TArray<FRawAnimSequenceTrack>& InAnimationData, const TArray<struct FTrackToSkeletonMap>& TrackToSkeletonMapTable, FCompactPose& InOutPose, float InTime, EAnimInterpolationType Interpolation, int32 NumFrames, float SequenceLength, FName RetargetSource)
+{
+	int32 KeyIndex1, KeyIndex2;
+	float Alpha;
+	FAnimationRuntime::GetKeyIndicesFromTime(KeyIndex1, KeyIndex2, Alpha, InTime, NumFrames, SequenceLength);
+
+	if (Interpolation == EAnimInterpolationType::Step)
+	{
+		Alpha = 0.f;
+	}
+
+	bool bInterpolate = true;
+
+	if (Alpha < KINDA_SMALL_NUMBER)
+	{
+		Alpha = 0.f;
+		bInterpolate = false;
+	}
+	else if (Alpha > 1.f - KINDA_SMALL_NUMBER)
+	{
+		bInterpolate = false;
+		KeyIndex1 = KeyIndex2;
+	}
+
+	if (bInterpolate)
+	{
+		BuildPoseFromRawDataInternal<true>(InAnimationData, TrackToSkeletonMapTable, InOutPose, KeyIndex1, KeyIndex2, Alpha);
+	}
+	else
+	{
+		BuildPoseFromRawDataInternal<false>(InAnimationData, TrackToSkeletonMapTable, InOutPose, KeyIndex1, KeyIndex2, Alpha);
+	}
+
+	const FBoneContainer& RequiredBones = InOutPose.GetBoneContainer();
+	const bool bDisableRetargeting = RequiredBones.GetDisableRetargeting();
+
+	if (!bDisableRetargeting)
+	{
+		const TArray<FRetargetTracking>& RetargetTracking = FBuildRawPoseScratchArea::Get().RetargetTracking;
+
+		USkeleton* Skeleton = RequiredBones.GetSkeletonAsset();
+
+		for (const FRetargetTracking& RT : RetargetTracking)
+		{
+			FAnimationRuntime::RetargetBoneTransform(Skeleton, RetargetSource, InOutPose[RT.PoseBoneIndex], RT.SkeletonBoneIndex, RT.PoseBoneIndex, RequiredBones, false);
+		}
+	}
+
+}
+
+
diff --git a/Engine/Source/Runtime/Engine/Private/ContentStreaming.cpp b/Engine/Source/Runtime/Engine/Private/ContentStreaming.cpp
--- a/Engine/Source/Runtime/Engine/Private/ContentStreaming.cpp
+++ b/Engine/Source/Runtime/Engine/Private/ContentStreaming.cpp
@@ -17,6 +17,7 @@
 #include "Streaming/TextureStreamingHelpers.h"
 #include "Streaming/StreamingManagerTexture.h"
 #include "AudioStreaming.h"
+#include "Animation/AnimationStreaming.h"
 
 /*-----------------------------------------------------------------------------
 	Globals.
@@ -754,6 +755,9 @@ FStreamingManagerCollection::FStreamingManagerCollection()
 
 	AudioStreamingManager = new FAudioStreamingManager();
 	AddStreamingManager( AudioStreamingManager );
+
+	AnimationStreamingManager = new FAnimationStreamingManager();
+	AddStreamingManager(AnimationStreamingManager);
 }
 
 /**
@@ -952,6 +956,12 @@ IAudioStreamingManager& FStreamingManagerCollection::GetAudioStreamingManager()
 	return *AudioStreamingManager;
 }
 
+IAnimationStreamingManager& FStreamingManagerCollection::GetAnimationStreamingManager() const
+{
+	check(AnimationStreamingManager);
+	return *AnimationStreamingManager;
+}
+
 /** Don't stream world resources for the next NumFrames. */
 void FStreamingManagerCollection::SetDisregardWorldResourcesForFrames(int32 NumFrames )
 {
diff --git a/Engine/Source/Runtime/Engine/Public/Animation/AnimCompressionTypes.h b/Engine/Source/Runtime/Engine/Public/Animation/AnimCompressionTypes.h
--- a/Engine/Source/Runtime/Engine/Public/Animation/AnimCompressionTypes.h
+++ b/Engine/Source/Runtime/Engine/Public/Animation/AnimCompressionTypes.h
@@ -4,9 +4,12 @@
 
 #include "CoreMinimal.h"
 #include "UObject/ObjectMacros.h"
+#include "Misc/Guid.h"
 #include "Misc/MemStack.h"
+#include "BonePose.h"
 #include "Animation/AnimTypes.h"
 #include "Animation/AnimCurveTypes.h"
+#include "Animation/AnimationAsset.h"
 
 #include "Async/MappedFileHandle.h"
 #include "HAL/PlatformFilemanager.h"
@@ -17,8 +20,12 @@
 #include "Serialization/MemoryReader.h"
 #include "Serialization/MemoryWriter.h"
 
+#include "ProfilingDebugging/CsvProfiler.h"
+
 #include "AnimCompressionTypes.generated.h"
 
+CSV_DECLARE_CATEGORY_EXTERN(Animation);
+
 /**
  * Indicates animation data key format.
  */
@@ -34,6 +41,12 @@ enum AnimationKeyFormat
 class FMemoryReader;
 class FMemoryWriter;
 
+class UAnimCompress;
+class UAnimCurveCompressionSettings;
+class USkeleton;
+
+extern FGuid GenerateGuidFromRawAnimData(const TArray<FRawAnimSequenceTrack>& RawAnimationData, const FRawCurveTracks& RawCurveData);
+
 template<typename ArrayClass>
 struct ENGINE_API FCompressedOffsetDataBase
 {
@@ -90,6 +103,34 @@ struct ENGINE_API FCompressedOffsetDataBase
 	}
 };
 
+// Helper for buiilding DDC keys of settings
+struct FArcToHexString
+{
+private:
+	TArray<uint8> TempBytes;
+
+public:
+	FMemoryWriter Ar;
+
+	FArcToHexString()
+		: Ar(TempBytes)
+	{
+		TempBytes.Reserve(64);
+	}
+
+	FString MakeString() const
+	{
+		FString Key;
+		const uint8* SettingsAsBytes = TempBytes.GetData();
+		Key.Reserve(TempBytes.Num() + 1);
+		for (int32 ByteIndex = 0; ByteIndex < TempBytes.Num(); ++ByteIndex)
+		{
+			ByteToHex(SettingsAsBytes[ByteIndex], Key);
+		}
+		return Key;
+	}
+};
+
 struct FCompressedOffsetData : public FCompressedOffsetDataBase<TArray<int32>>
 {
 
@@ -150,20 +191,18 @@ public:
 		, SequenceLength(0.f)
 		, NumFrames(0)
 		, bIsValidAdditive(false)
-		, CompressCommandletVersion(0)
-		, RefFrameIndex(0)
-		, RefPoseType((EAdditiveBasePoseType)0)
-		, AdditiveAnimType((EAdditiveAnimationType)0)
 	{
 	}
 
-	FCompressibleAnimData(class UAnimSequence* InSeq);
+	FCompressibleAnimData(UAnimCompress* InRequestedCompressionScheme, UAnimCurveCompressionSettings* InCurveCompressionSettings, USkeleton* InSkeleton, EAnimInterpolationType InInterpolation, float InSequenceLength, int32 InNumFrames);
 
-	class UAnimCompress* RequestedCompressionScheme;
+	FCompressibleAnimData(class UAnimSequence* InSeq, const bool bPerformStripping);
 
-	class UAnimCurveCompressionSettings* CurveCompressionSettings;
+	UAnimCompress* RequestedCompressionScheme;
 
-	class USkeleton* Skeleton;
+	UAnimCurveCompressionSettings* CurveCompressionSettings;
+
+	USkeleton* Skeleton;
 
 	TArray<FTrackToSkeletonMap> TrackToSkeletonMapTable;
 
@@ -183,19 +222,6 @@ public:
 
 	bool bIsValidAdditive;
 
-	//For DDC
-	FString TypeName;
-	FString DDCKey;
-
-	int32 CompressCommandletVersion;
-	FGuid RawDataGuid;
-
-	//Additive
-	FGuid AdditiveDataGuid;
-	int32 RefFrameIndex;
-	EAdditiveBasePoseType RefPoseType;
-	EAdditiveAnimationType AdditiveAnimType;
-
 	//For Logging
 	FString Name;
 	FString FullName;
@@ -674,5 +700,61 @@ public:
 	// The size of the raw data used to create the compressed data
 	int32 CompressedRawDataSize;
 
-	void SerializeCompressedData(FArchive& Ar, bool bDDCData, UObject* DataOwner, class UAnimCurveCompressionSettings* CurveCompressionSettings);
-};
\ No newline at end of file
+	void SerializeCompressedData(FArchive& Ar, bool bDDCData, UObject* DataOwner, USkeleton* Skeleton, class UAnimCurveCompressionSettings* CurveCompressionSettings, bool bCanUseBulkData=true);
+
+	int32 GetSkeletonIndexFromTrackIndex(const int32 TrackIndex) const
+	{
+		return CompressedTrackToSkeletonMapTable[TrackIndex].BoneTreeIndex;
+	}
+};
+
+struct FRootMotionReset
+{
+
+	FRootMotionReset(bool bInEnableRootMotion, ERootMotionRootLock::Type InRootMotionRootLock, bool bInForceRootLock, FTransform InAnimFirstFrame, bool bInIsValidAdditive)
+		: bEnableRootMotion(bInEnableRootMotion)
+		, RootMotionRootLock(InRootMotionRootLock)
+		, bForceRootLock(bInForceRootLock)
+		, AnimFirstFrame(InAnimFirstFrame)
+		, bIsValidAdditive(bInIsValidAdditive)
+	{
+	}
+
+	bool bEnableRootMotion;
+
+	ERootMotionRootLock::Type RootMotionRootLock;
+
+	bool bForceRootLock;
+
+	FTransform AnimFirstFrame;
+
+	bool bIsValidAdditive;
+
+	void ResetRootBoneForRootMotion(FTransform& BoneTransform, const FBoneContainer& RequiredBones) const
+	{
+		switch (RootMotionRootLock)
+		{
+		case ERootMotionRootLock::AnimFirstFrame: BoneTransform = AnimFirstFrame; break;
+		case ERootMotionRootLock::Zero: BoneTransform = FTransform::Identity; break;
+		default:
+		case ERootMotionRootLock::RefPose: BoneTransform = RequiredBones.GetRefPoseArray()[0]; break;
+		}
+
+		if (bIsValidAdditive && RootMotionRootLock != ERootMotionRootLock::AnimFirstFrame)
+		{
+			//Need to remove default scale here for additives
+			BoneTransform.SetScale3D(BoneTransform.GetScale3D() - FVector(1.f));
+		}
+	}
+};
+
+extern void DecompressPose(	FCompactPose& OutPose,
+							const FCompressedAnimSequence& CompressedData,
+							const FAnimExtractContext& ExtractionContext,
+							USkeleton* Skeleton,
+							float SequenceLength,
+							EAnimInterpolationType Interpolation,
+							bool bIsBakedAdditive,
+							FName RetargetSource,
+							FName SourceName,
+							const FRootMotionReset& RootMotionReset);
\ No newline at end of file
diff --git a/Engine/Source/Runtime/Engine/Public/Animation/AnimationStreaming.h b/Engine/Source/Runtime/Engine/Public/Animation/AnimationStreaming.h
new file mode 100644
--- /dev/null
+++ b/Engine/Source/Runtime/Engine/Public/Animation/AnimationStreaming.h
@@ -0,0 +1,214 @@
+// Copyright 1998-2019 Epic Games, Inc. All Rights Reserved.
+
+/*=============================================================================
+AudioStreaming.h: Definitions of classes used for audio streaming.
+=============================================================================*/
+
+#pragma once
+
+#include "CoreMinimal.h"
+#include "HAL/ThreadSafeCounter.h"
+#include "Containers/IndirectArray.h"
+#include "Containers/Queue.h"
+#include "Stats/Stats.h"
+#include "ContentStreaming.h"
+#include "Async/AsyncWork.h"
+#include "Async/AsyncFileHandle.h"
+
+class UAnimStreamable;
+struct FAnimationStreamingManager;
+
+// 
+struct FLoadedAnimationChunk
+{
+	FCompressedAnimSequence* CompressedAnimData;
+
+	class IAsyncReadRequest* IORequest;
+	double RequestStart;
+
+	uint32	Index;
+	bool	bOwnsCompressedData;
+
+	FLoadedAnimationChunk()
+		: CompressedAnimData(nullptr)
+		, IORequest(nullptr)
+		, Index(0)
+		, bOwnsCompressedData(false)
+	{
+	}
+
+	~FLoadedAnimationChunk()
+	{
+		checkf(CompressedAnimData == nullptr, TEXT("Animation chunk compressed data ptr not null (%p), DataSize: %d"));
+	}
+
+	void CleanUpIORequest();
+};
+
+/**
+ * Contains everything that will be needed by a Streamable Anim that's streaming in data
+ */
+struct FStreamingAnimationData final
+{
+	FStreamingAnimationData();
+	~FStreamingAnimationData();
+
+	// Frees streaming animation data resources, blocks pending async IO requests
+	void FreeResources();
+
+	/**
+	 * Sets up the streaming wave data and loads the first chunk of audio for instant play
+	 *
+	 * @param Anim	The streamable animation we are managing
+	 */
+	bool Initialize(UAnimStreamable* InStreamableAnim, FAnimationStreamingManager* InAnimationStreamingManager);
+
+	/**
+	 * Updates the streaming status of the animation and performs finalization when appropriate. The function returns
+	 * true while there are pending requests in flight and updating needs to continue.
+	 *
+	 * @return					true if there are requests in flight, false otherwise
+	 */
+	bool UpdateStreamingStatus();
+
+	/**
+	 * Tells the SoundWave which chunks are currently required so that it can start loading any needed
+	 *
+	 * @param InChunkIndices The Chunk Indices that are currently needed by all sources using this sound
+	 * @param bShouldPrioritizeAsyncIORequest Whether request should have higher priority than usual
+	 */
+	//void UpdateChunkRequests(FWaveRequest& InWaveRequest);
+
+		/**
+	 * Checks whether the requested chunk indices differ from those loaded
+	 *
+	 * @param IndicesToLoad		List of chunk indices that should be loaded
+	 * @param IndicesToFree		List of chunk indices that should be freed
+	 * @return Whether any changes to loaded chunks are required
+	 */
+	bool HasPendingRequests(TArray<uint32>& IndicesToLoad, TArray<uint32>& IndicesToFree) const;
+
+	/**
+	 * Kicks off any pending requests
+	 */
+	void BeginPendingRequests(const TArray<uint32>& IndicesToLoad, const TArray<uint32>& IndicesToFree);
+
+	/**
+	* Blocks till all pending requests are fulfilled.
+	*
+	* @param TimeLimit		Optional time limit for processing, in seconds. Specifying 0 means infinite time limit.
+	* @return				Return true if there are no requests left in flight, false if the time limit was reached before they were finished.
+	*/
+	bool BlockTillAllRequestsFinished(float TimeLimit = 0.0f);
+
+#if WITH_EDITORONLY_DATA
+	/**
+	 * Finishes any Derived Data Cache requests that may be in progress
+	 *
+	 * @return Whether any of the requests failed.
+	 */
+	bool FinishDDCRequests();
+#endif //WITH_EDITORONLY_DATA
+
+private:
+	// Don't allow copy construction as it could free shared memory
+	FStreamingAnimationData(const FStreamingAnimationData& that);
+	FStreamingAnimationData& operator=(FStreamingAnimationData const&);
+
+	// Creates a new chunk, returns the chunk index
+	FLoadedAnimationChunk& AddNewLoadedChunk(uint32 ChunkIndex, FCompressedAnimSequence* ExistingData);
+	void FreeLoadedChunk(FLoadedAnimationChunk& LoadedChunk);
+	void ResetRequestedChunks();
+
+public:
+	/** AnimStreamable this streaming data is for */
+	UAnimStreamable* StreamableAnim;
+
+	/* Contains pointers to Chunks of audio data that have been streamed in */
+	TArray<FLoadedAnimationChunk> LoadedChunks;
+
+	class IAsyncReadFileHandle* IORequestHandle;
+
+	/** Indices of chunks that are currently loaded */
+	TArray<uint32>	LoadedChunkIndices;
+
+	TArray<uint32> RequestedChunks;
+
+	/** Indices of chunks we want to have loaded */
+	//FWaveRequest	CurrentRequest;
+
+#if WITH_EDITORONLY_DATA
+	/** Pending async derived data streaming tasks */
+	//TIndirectArray<FAsyncStreamDerivedChunkTask> PendingAsyncStreamDerivedChunkTasks;
+#endif // #if WITH_EDITORONLY_DATA
+
+	/** Ptr to owning audio streaming manager. */
+	FAnimationStreamingManager* AnimationStreamingManager;
+};
+
+/** Struct used to store results of an async file load. */
+/*struct FASyncAnimationChunkLoadResult
+{
+	// Place to safely copy the ptr of a loaded audio chunk when load result is finished
+	uint8* DataResults;
+
+	// Actual storage of the loaded audio chunk, will be filled on audio thread.
+	FStreamingAnimationData* StreamingAnimData;
+
+	// Loaded audio chunk index
+	int32 LoadedChunkIndex;
+
+	FASyncAnimationChunkLoadResult()
+		: DataResults(nullptr)
+		, StreamingAnimData(nullptr)
+		, LoadedChunkIndex(INDEX_NONE)
+	{}
+};*/
+
+
+/**
+* Streaming manager dealing with audio.
+*/
+struct FAnimationStreamingManager : public IAnimationStreamingManager
+{
+	/** Constructor, initializing all members */
+	FAnimationStreamingManager();
+
+	virtual ~FAnimationStreamingManager();
+
+	// IStreamingManager interface
+	virtual void UpdateResourceStreaming( float DeltaTime, bool bProcessEverything=false ) override;
+	virtual int32 BlockTillAllRequestsFinished( float TimeLimit = 0.0f, bool bLogResults = false ) override;
+	virtual void CancelForcedResources() override;
+	virtual void NotifyLevelChange() override;
+	virtual void SetDisregardWorldResourcesForFrames( int32 NumFrames ) override;
+	virtual void AddLevel( class ULevel* Level ) override;
+	virtual void RemoveLevel( class ULevel* Level ) override;
+	virtual void NotifyLevelOffset( class ULevel* Level, const FVector& Offset ) override;
+	// End IStreamingManager interface
+
+	// IAudioStreamingManager interface
+	virtual void AddStreamingAnim(UAnimStreamable* Anim) override;
+	virtual bool RemoveStreamingAnim(UAnimStreamable* Anim) override;
+	/*void AddDecoder(ICompressedAudioInfo* CompressedAudioInfo);
+	void RemoveDecoder(ICompressedAudioInfo* CompressedAudioInfo);
+	bool IsManagedStreamingSoundWave(const USoundWave* SoundWave) const;
+	bool IsStreamingInProgress(const USoundWave* SoundWave);
+	bool CanCreateSoundSource(const FWaveInstance* WaveInstance) const;
+	void AddStreamingSoundSource(FSoundSource* SoundSource);
+	void RemoveStreamingSoundSource(FSoundSource* SoundSource);
+	bool IsManagedStreamingSoundSource(const FSoundSource* SoundSource) const;*/
+	virtual const FCompressedAnimSequence* GetLoadedChunk(const UAnimStreamable* Anim, uint32 ChunkIndex) const override;
+	// End IAudioStreamingManager interface
+
+	/** Called when an async callback is made on an async loading audio chunk request. */
+	void OnAsyncFileCallback(FStreamingAnimationData* StreamingAnimData, int32 ChunkIndex, int64 ReadSize, IAsyncReadRequest* ReadRequest);
+
+protected:
+
+	/** Sound Waves being managed. */
+	TMap<UAnimStreamable*, FStreamingAnimationData*> StreamingAnimations;
+
+	/** Critical section to protect usage of shared gamethread/workerthread members */
+	mutable FCriticalSection CriticalSection;
+};
diff --git a/Engine/Source/Runtime/Engine/Public/BonePose.h b/Engine/Source/Runtime/Engine/Public/BonePose.h
--- a/Engine/Source/Runtime/Engine/Public/BonePose.h
+++ b/Engine/Source/Runtime/Engine/Public/BonePose.h
@@ -882,4 +882,7 @@ void FCSPose<PoseType>::ConvertComponentPosesToLocalPoses(FCSPose<PoseType>&& In
 			OutPose[BoneIndex].NormalizeRotation();
 		}
 	}
-}
\ No newline at end of file
+}
+
+// Populate InOutPose based on raw animation data. 
+extern void BuildPoseFromRawData(const TArray<FRawAnimSequenceTrack>& InAnimationData, const TArray<struct FTrackToSkeletonMap>& TrackToSkeletonMapTable, FCompactPose& InOutPose, float InTime, EAnimInterpolationType Interpolation, int32 NumFrames, float SequenceLength, FName RetargetSource);
\ No newline at end of file
diff --git a/Engine/Source/Runtime/Engine/Public/ContentStreaming.h b/Engine/Source/Runtime/Engine/Public/ContentStreaming.h
--- a/Engine/Source/Runtime/Engine/Public/ContentStreaming.h
+++ b/Engine/Source/Runtime/Engine/Public/ContentStreaming.h
@@ -17,6 +17,8 @@ class ICompressedAudioInfo;
 class UTexture2D;
 struct FRenderAssetStreamingManager;
 struct FWaveInstance;
+class UAnimStreamable;
+struct FCompressedAnimSequence;
 
 /*-----------------------------------------------------------------------------
 	Stats.
@@ -443,6 +445,27 @@ struct IAudioStreamingManager : public IStreamingManager
 	virtual const uint8* GetLoadedChunk(const USoundWave* SoundWave, uint32 ChunkIndex, uint32* OutChunkSize = NULL) const = 0;
 };
 
+/**
+ * Interface to add functions specifically related to animation streaming
+ */
+struct IAnimationStreamingManager : public IStreamingManager
+{
+	/** Adds a new Streamable Anim to the streaming manager. */
+	virtual void AddStreamingAnim(UAnimStreamable* Anim) = 0;
+
+	/** Removes a Streamable Anim from the streaming manager. */
+	virtual bool RemoveStreamingAnim(UAnimStreamable* Anim) = 0;
+
+	/**
+	 * Gets a pointer to a chunk of animation data
+	 *
+	 * @param Anim			AnimStreamable we want a chunk from
+	 * @param ChunkIndex	Index of the chunk we want
+	 * @return Either the desired chunk or NULL if it's not loaded
+	 */
+	virtual const FCompressedAnimSequence* GetLoadedChunk(const UAnimStreamable* Anim, uint32 ChunkIndex) const = 0;
+};
+
 /**
  * Streaming manager collection, routing function calls to streaming managers that have been added
  * via AddStreamingManager.
@@ -526,6 +549,11 @@ struct ENGINE_VTABLE FStreamingManagerCollection : public IStreamingManager
 	 */
 	ENGINE_API IAudioStreamingManager& GetAudioStreamingManager() const;
 
+	/**
+	 * Gets a reference to the Animation Streaming Manager interface
+	 */
+	ENGINE_API IAnimationStreamingManager& GetAnimationStreamingManager() const;
+
 	/**
 	 * Adds a streaming manager to the array of managers to route function calls to.
 	 *
@@ -616,5 +644,8 @@ protected:
 
 	/** The audio streaming manager, should always exist */
 	IAudioStreamingManager* AudioStreamingManager;
+
+	/** The animation streaming manager, should always exist */
+	IAnimationStreamingManager* AnimationStreamingManager;
 };
 
diff --git a/Engine/Source/Runtime/MovieSceneTracks/Public/Sections/MovieSceneSkeletalAnimationSection.h b/Engine/Source/Runtime/MovieSceneTracks/Public/Sections/MovieSceneSkeletalAnimationSection.h
--- a/Engine/Source/Runtime/MovieSceneTracks/Public/Sections/MovieSceneSkeletalAnimationSection.h
+++ b/Engine/Source/Runtime/MovieSceneTracks/Public/Sections/MovieSceneSkeletalAnimationSection.h
@@ -23,7 +23,7 @@ struct FMovieSceneSkeletalAnimationParams
 	float GetSequenceLength() const { return Animation != nullptr ? Animation->SequenceLength : 0.f; }
 
 	/** The animation this section plays */
-	UPROPERTY(BlueprintReadWrite, EditAnywhere, Category="Animation", meta=(AllowedClasses = "AnimSequence, AnimComposite"))
+	UPROPERTY(BlueprintReadWrite, EditAnywhere, Category="Animation", meta=(AllowedClasses = "AnimSequence,AnimComposite,AnimStreamable"))
 	UAnimSequenceBase* Animation;
 
 	/** The offset into the beginning of the animation clip */
