Clean up commented out streaming code

diff --git a/Engine/Source/Runtime/Engine/Classes/Animation/AnimStreamable.h b/Engine/Source/Runtime/Engine/Classes/Animation/AnimStreamable.h
--- a/Engine/Source/Runtime/Engine/Classes/Animation/AnimStreamable.h
+++ b/Engine/Source/Runtime/Engine/Classes/Animation/AnimStreamable.h
@@ -195,18 +195,8 @@ public:
 	ENGINE_API virtual void HandleAssetPlayerTickedInternal(FAnimAssetTickContext &Context, const float PreviousTime, const float MoveDelta, const FAnimTickRecord &Instance, struct FAnimNotifyQueue& NotifyQueue) const override;
 	virtual void GetAnimationPose(FCompactPose& OutPose, FBlendedCurve& OutCurve, const FAnimExtractContext& ExtractionContext) const override;
 	virtual int32 GetNumberOfFrames() const override { return NumFrames; }
-	//virtual EAdditiveAnimationType GetAdditiveAnimType() const override;
-	//virtual bool IsValidAdditive() const override { return GetAdditiveAnimType() != AAT_None; }
-	//virtual bool HasRootMotion() const override;
-
 	//~ End UAnimSequenceBase Interface
-	//~ Begin UAnimSequence Interface
-#if WITH_EDITOR
-	//virtual class UAnimSequence* GetAdditiveBasePose() const override;
-	//virtual bool GetAllAnimationSequencesReferred(TArray<UAnimationAsset*>& AnimationAssets, bool bRecursive = true) override;
-	//virtual void ReplaceReferredAnimations(const TMap<UAnimationAsset*, UAnimationAsset*>& ReplacementMap) override;
-#endif
-	//~ End UAnimSequence Interface
+
 #if WITH_EDITOR
 	ENGINE_API void InitFrom(const UAnimSequence* InSourceSequence);
 #endif
diff --git a/Engine/Source/Runtime/Engine/Private/Animation/AnimationStreaming.cpp b/Engine/Source/Runtime/Engine/Private/Animation/AnimationStreaming.cpp
--- a/Engine/Source/Runtime/Engine/Private/Animation/AnimationStreaming.cpp
+++ b/Engine/Source/Runtime/Engine/Private/Animation/AnimationStreaming.cpp
@@ -8,17 +8,6 @@ AnimationStreaming.cpp: Manager to handle streaming animation data
 #include "Misc/CoreStats.h"
 #include "Animation/AnimStreamable.h"
 #include "Algo/Find.h"
-/*#include "Sound/SoundWave.h"
-#include "Sound/AudioSettings.h"
-#include "DerivedDataCacheInterface.h"
-#include "Serialization/MemoryReader.h"
-#include "HAL/PlatformFile.h"
-#include "HAL/PlatformFilemanager.h"
-#include "Async/AsyncFileHandle.h"
-#include "Misc/ScopeLock.h"
-#include "HAL/IConsoleManager.h"
-#include "HAL/LowLevelMemTracker.h"
-#include "AudioDecompress.h"*/
 
 static int32 SpoofFailedAnimationChunkLoad = 0;
 FAutoConsoleVariableRef CVarSpoofFailedAnimationChunkLoad(
@@ -29,55 +18,6 @@ FAutoConsoleVariableRef CVarSpoofFailedAnimationChunkLoad(
 	ECVF_Default);
 
 
-/*------------------------------------------------------------------------------
-	Streaming chunks from the derived data cache.
-------------------------------------------------------------------------------*/
-
-#if WITH_EDITORONLY_DATA
-/*
-// Initialization constructor.
-FAsyncStreamDerivedChunkWorker::FAsyncStreamDerivedChunkWorker(
-	const FString& InDerivedDataKey,
-	void* InDestChunkData,
-	int32 InChunkSize,
-	FThreadSafeCounter* InThreadSafeCounter
-	)
-	: DerivedDataKey(InDerivedDataKey)
-	, DestChunkData(InDestChunkData)
-	, ExpectedChunkSize(InChunkSize)
-	, bRequestFailed(false)
-	, ThreadSafeCounter(InThreadSafeCounter)
-{
-}
-
-// Retrieves the derived chunk from the derived data cache.
-void FAsyncStreamDerivedChunkWorker::DoWork()
-{
-	DECLARE_SCOPE_CYCLE_COUNTER(TEXT("FAsyncStreamDerivedChunkWorker::DoWork"), STAT_AsyncStreamDerivedChunkWorker_DoWork, STATGROUP_StreamingDetails);
-
-	UE_LOG(LogAudio, Verbose, TEXT("Start of ASync DDC Chunk read for key: %s"), *DerivedDataKey);
-
-	TArray<uint8> DerivedChunkData;
-
-	if (GetDerivedDataCacheRef().GetSynchronous(*DerivedDataKey, DerivedChunkData))
-	{
-		FMemoryReader Ar(DerivedChunkData, true);
-		int32 ChunkSize = 0;
-		Ar << ChunkSize;
-		checkf(ChunkSize == ExpectedChunkSize, TEXT("ChunkSize(%d) != ExpectedSize(%d)"), ChunkSize, ExpectedChunkSize);
-		Ar.Serialize(DestChunkData, ChunkSize);
-	}
-	else
-	{
-		bRequestFailed = true;
-	}
-	FPlatformMisc::MemoryBarrier();
-	ThreadSafeCounter->Decrement();
-
-	UE_LOG(LogAudio, Verbose, TEXT("End of ASync DDC Chunk read for key: %s"), *DerivedDataKey);
-}
-*/
-#endif // #if WITH_EDITORONLY_DATA
 
 void FLoadedAnimationChunk::CleanUpIORequest()
 {
@@ -196,15 +136,6 @@ bool FStreamingAnimationData::UpdateStreamingStatus()
 	return bHasPendingRequestInFlight;
 }
 
-/*void FStreamingAnimationData::UpdateChunkRequests(FWaveRequest& InWaveRequest)
-{
-	// Might change this but ensures chunk 0 stays loaded for now
-	check(InWaveRequest.RequiredIndices.Contains(0));
-
-	CurrentRequest = InWaveRequest;
-
-}*/
-
 bool FStreamingAnimationData::HasPendingRequests(TArray<uint32>& IndicesToLoad, TArray<uint32>& IndicesToFree) const
 {
 	IndicesToLoad.Reset();
@@ -233,17 +164,6 @@ bool FStreamingAnimationData::HasPendingRequests(TArray<uint32>& IndicesToLoad,
 
 void FStreamingAnimationData::BeginPendingRequests(const TArray<uint32>& IndicesToLoad, const TArray<uint32>& IndicesToFree)
 {
-	/*if (UE_LOG_ACTIVE(LogAudio, Log) && IndicesToLoad.Num() > 0)
-	{
-		FString LogString = FString::Printf(TEXT("Requesting ASync load of chunk(s) %d"), IndicesToLoad[0]);
-		for (int32 Index = 1; Index < IndicesToLoad.Num(); ++Index)
-		{
-			LogString += FString::Printf(TEXT(", %d"), IndicesToLoad[Index]);
-		}
-		LogString += FString::Printf(TEXT(" from SoundWave'%s'"), *SoundWave->GetName());
-		UE_LOG(LogAudio, Log, TEXT("%s"), *LogString);
-	}*/
-
 	TArray<uint32> FreeChunkIndices;
 
 	// Mark Chunks for removal in case they can be reused
@@ -362,9 +282,6 @@ void FStreamingAnimationData::FreeLoadedChunk(FLoadedAnimationChunk& LoadedChunk
 	if (LoadedChunk.bOwnsCompressedData)
 	{
 		delete LoadedChunk.CompressedAnimData;
-
-		//DEC_DWORD_STAT_BY(STAT_AudioMemorySize, LoadedChunk.DataSize);
-		//DEC_DWORD_STAT_BY(STAT_AudioMemory, LoadedChunk.DataSize);
 	}
 
 	LoadedChunk.CompressedAnimData = NULL;
@@ -411,9 +328,6 @@ void FAnimationStreamingManager::OnAsyncFileCallback(FStreamingAnimationData* St
 		ChunkStorage.CompressedAnimData->SerializeCompressedData(Reader, false, Anim, Anim->GetSkeleton(), Anim->CurveCompressionSettings);
 
 		FPlatformMisc::LowLevelOutputDebugStringf(TEXT("Request Finished %.2f\nAnim Chunk Streamed %.4f\n"), FPlatformTime::Seconds(), FPlatformTime::Seconds() - ChunkStorage.RequestStart);
-		/*DEC_MEMORY_STAT_BY(STAT_AsyncFileMemory, ChunkStorage->DataSize);
-		INC_DWORD_STAT_BY(STAT_AudioMemorySize, ChunkStorage->DataSize);
-		INC_DWORD_STAT_BY(STAT_AudioMemory, ChunkStorage->DataSize);*/
 	}
 }
 
@@ -427,70 +341,31 @@ void FAnimationStreamingManager::UpdateResourceStreaming(float DeltaTime, bool b
 	{
 		AnimData.Value->UpdateStreamingStatus();
 	}
-
-	/*for (ICompressedAudioInfo* Decoder : CompressedAudioInfos)
-	{
-		USoundWave* SoundWave = Decoder->GetStreamingSoundWave();
-		if (SoundWave)
-		{
-			FStreamingWaveData** WaveDataPtr = StreamingSoundWaves.Find(SoundWave);
-			if (WaveDataPtr && (*WaveDataPtr)->PendingChunkChangeRequestStatus.GetValue() == AudioState_ReadyFor_Requests)
-			{
-				FStreamingWaveData* WaveData = *WaveDataPtr;
-				// Request the chunk the source is using and the one after that
-				FWaveRequest& WaveRequest = GetWaveRequest(SoundWave);
-				int32 SourceChunk = Decoder->GetCurrentChunkIndex();
-				if (SourceChunk >= 0 && SourceChunk < SoundWave->RunningPlatformData->NumChunks)
-				{
-					WaveRequest.RequiredIndices.AddUnique(SourceChunk);
-					WaveRequest.RequiredIndices.AddUnique((SourceChunk + 1) % SoundWave->RunningPlatformData->NumChunks);
-					WaveRequest.bPrioritiseRequest = true;
-				}
-				else
-				{
-					UE_LOG(LogAudio, Log, TEXT("Invalid chunk request curIndex=%d numChunks=%d\n"), SourceChunk, SoundWave->RunningPlatformData->NumChunks);
-				}
-			}
-		}
-	}
-
-	for (auto Iter = WaveRequests.CreateIterator(); Iter; ++Iter)
-	{
-		USoundWave* Wave = Iter.Key();
-		FStreamingWaveData* WaveData = StreamingSoundWaves.FindRef(Wave);
-
-		if (WaveData && WaveData->PendingChunkChangeRequestStatus.GetValue() == AudioState_ReadyFor_Requests)
-		{
-			WaveData->UpdateChunkRequests(Iter.Value());
-			WaveData->UpdateStreamingStatus();
-			Iter.RemoveCurrent();
-		}
-	}*/
 }
 
 int32 FAnimationStreamingManager::BlockTillAllRequestsFinished(float TimeLimit, bool)
 {
-	/*{
+	{
 		FScopeLock Lock(&CriticalSection);
 
-		QUICK_SCOPE_CYCLE_COUNTER(FAudioStreamingManager_BlockTillAllRequestsFinished);
+		QUICK_SCOPE_CYCLE_COUNTER(FAnimStreamingManager_BlockTillAllRequestsFinished);
 		int32 Result = 0;
 
 		if (TimeLimit == 0.0f)
 		{
-			for (auto& WavePair : StreamingSoundWaves)
+			for (auto& AnimPair : StreamingAnimations)
 			{
-				WavePair.Value->BlockTillAllRequestsFinished();
+				AnimPair.Value->BlockTillAllRequestsFinished();
 			}
 		}
 		else
 		{
 			double EndTime = FPlatformTime::Seconds() + TimeLimit;
-			for (auto& WavePair : StreamingSoundWaves)
+			for (auto& AnimPair : StreamingAnimations)
 			{
 				float ThisTimeLimit = EndTime - FPlatformTime::Seconds();
 				if (ThisTimeLimit < .001f || // one ms is the granularity of the platform event system
-					!WavePair.Value->BlockTillAllRequestsFinished(ThisTimeLimit))
+					!AnimPair.Value->BlockTillAllRequestsFinished(ThisTimeLimit))
 				{
 					Result = 1; // we don't report the actual number, just 1 for any number of outstanding requests
 					break;
@@ -499,7 +374,7 @@ int32 FAnimationStreamingManager::BlockTillAllRequestsFinished(float TimeLimit,
 		}
 
 		return Result;
-	}*/
+	}
 
 	// Not sure yet whether this will work the same as textures - aside from just before destroying
 	return 0;
@@ -557,22 +432,6 @@ bool FAnimationStreamingManager::RemoveStreamingAnim(UAnimStreamable* Anim)
 
 		// Free the resources of the streaming wave data. This blocks pending IO requests
 		AnimData->FreeResources();
-
-		{
-			// Then we need to remove any results from those pending requests before we delete so that we don't process them
-			/*FScopeLock StreamChunkResults(&ChunkResultCriticalSection);
-			for (int32 i = AsyncAudioStreamChunkResults.Num() - 1 ; i >= 0; --i)
-			{
-				FASyncAudioChunkLoadResult* LoadResult = AsyncAudioStreamChunkResults[i];
-				FStreamingWaveData* StreamingWaveData = LoadResult->StreamingWaveData;
-				if (StreamingWaveData == WaveData)
-				{
-					delete LoadResult;
-					
-					AsyncAudioStreamChunkResults.RemoveAtSwap(i, 1, false);
-				}
-			}*/
-		}
 		delete AnimData;
 		return true;
 	}
