Mark streamable anim as dirty when raw data changed

diff --git a/Engine/Source/Runtime/Engine/Private/Animation/AnimStreamable.cpp b/Engine/Source/Runtime/Engine/Private/Animation/AnimStreamable.cpp
--- a/Engine/Source/Runtime/Engine/Private/Animation/AnimStreamable.cpp
+++ b/Engine/Source/Runtime/Engine/Private/Animation/AnimStreamable.cpp
@@ -357,6 +357,7 @@ int32 UAnimStreamable::GetChunkIndexForTime(const TArray<FAnimStreamableChunk>&
 #if WITH_EDITOR
 void UAnimStreamable::InitFrom(const UAnimSequence* InSourceSequence)
 {
+	Modify();
 	SetSkeleton(InSourceSequence->GetSkeleton());
 	SourceSequence = InSourceSequence;
 
