[NO NON-PII COMMIT MESSAGE]

diff --git a/Engine/Source/Runtime/Engine/Private/LevelStreaming.cpp b/Engine/Source/Runtime/Engine/Private/LevelStreaming.cpp
--- a/Engine/Source/Runtime/Engine/Private/LevelStreaming.cpp
+++ b/Engine/Source/Runtime/Engine/Private/LevelStreaming.cpp
@@ -1716,9 +1716,18 @@ ULevelStreamingDynamic* ULevelStreamingDynamic::LoadLevelInstanceBySoftObjectPtr
 
 ULevelStreamingDynamic* ULevelStreamingDynamic::LoadLevelInstance_Internal(UWorld* World, const FString& LongPackageName, const FVector Location, const FRotator Rotation, bool& bOutSuccess)
 {
-    // Create Unique Name for sub-level package
-	const FString ShortPackageName = FPackageName::GetShortName(LongPackageName);
 	const FString PackagePath = FPackageName::GetLongPackagePath(LongPackageName);
+	FString ShortPackageName = FPackageName::GetShortName(LongPackageName);
+
+	if (ShortPackageName.StartsWith(World->StreamingLevelsPrefix))
+	{
+		ShortPackageName = ShortPackageName.RightChop(World->StreamingLevelsPrefix.Len());
+	}
+
+	// Remove PIE prefix if it's there before we actually load the level
+	FString OnDiskPackageName = PackagePath + TEXT("/") + ShortPackageName;
+
+	// Create Unique Name for sub-level package
 	FString UniqueLevelPackageName = PackagePath + TEXT("/") + World->StreamingLevelsPrefix + ShortPackageName;
 	UniqueLevelPackageName += TEXT("_LevelInstance_") + FString::FromInt(++UniqueLevelInstanceId);
     
@@ -1734,7 +1743,7 @@ ULevelStreamingDynamic* ULevelStreamingDynamic::LoadLevelInstance_Internal(UWorl
 	// Transform
     StreamingLevel->LevelTransform = FTransform(Rotation, Location);
 	// Map to Load
-    StreamingLevel->PackageNameToLoad = FName(*LongPackageName);
+    StreamingLevel->PackageNameToLoad = FName(*OnDiskPackageName);
           
     // Add the new level to world.
     World->AddStreamingLevel(StreamingLevel);
