Change CustomDeltaState tracking over to TArray now that we iterate CustomDeltaProperties by index.

diff --git a/Engine/Source/Runtime/Engine/Private/DataReplication.cpp b/Engine/Source/Runtime/Engine/Private/DataReplication.cpp
--- a/Engine/Source/Runtime/Engine/Private/DataReplication.cpp
+++ b/Engine/Source/Runtime/Engine/Private/DataReplication.cpp
@@ -224,7 +224,7 @@ public:
 		UObject* Object,
 		UNetConnection* Connection,
 		FReplicationChangelistMgr& ChangelistMgr,
-		TMap<int32, TSharedPtr<INetDeltaBaseState>>& CustomDeltaStates)
+		TArray<TSharedPtr<INetDeltaBaseState>>& CustomDeltaStates)
 	{
 		RepLayout.PreSendCustomDeltaProperties(Object, Connection, ChangelistMgr, CustomDeltaStates);
 	}
@@ -234,7 +234,7 @@ public:
 		UObject* Object,
 		UNetConnection* Connection,
 		FReplicationChangelistMgr& ChangelistMgr,
-		TMap<int32, TSharedPtr<INetDeltaBaseState>>& CustomDeltaStates)
+		TArray<TSharedPtr<INetDeltaBaseState>>& CustomDeltaStates)
 	{
 		RepLayout.PostSendCustomDeltaProperties(Object, Connection, ChangelistMgr, CustomDeltaStates);
 	}
@@ -426,19 +426,32 @@ void FObjectReplicator::InitRecentProperties(uint8* Source)
 
 		// Init custom delta property state
 		const uint16 NumLifetimeCustomDeltaProperties = FNetSerializeCB::GetNumLifetimeCustomDeltaProperties(LocalRepLayout);
+		SendingRepState->RecentCustomDeltaState.SetNum(NumLifetimeCustomDeltaProperties);
+
+		const bool bIsRecordingReplay = Connection->InternalAck;
+
+		if (bIsRecordingReplay)
+		{
+			SendingRepState->CDOCustomDeltaState.SetNum(NumLifetimeCustomDeltaProperties);
+			SendingRepState->CheckpointCustomDeltaState.SetNum(NumLifetimeCustomDeltaProperties);
+		}
+
 		for (uint16 CustomDeltaProperty = 0; CustomDeltaProperty < NumLifetimeCustomDeltaProperties; ++CustomDeltaProperty)
 		{
 			FOutBunch DeltaState(Connection->PackageMap);
-			TSharedPtr<INetDeltaBaseState>& NewState = SendingRepState->RecentCustomDeltaState.FindOrAdd(CustomDeltaProperty);
+			TSharedPtr<INetDeltaBaseState>& NewState = SendingRepState->RecentCustomDeltaState[CustomDeltaProperty];
 			NewState.Reset();
 
 			TSharedPtr<INetDeltaBaseState> OldState;
 
 			SendCustomDeltaProperty(UseObject, CustomDeltaProperty, DeltaState, NewState, OldState);
 
-			// Store the initial delta state in case we need it for when we're asked to resend all data since channel was first opened (bResendAllDataSinceOpen)
-			SendingRepState->CDOCustomDeltaState.Add(CustomDeltaProperty, NewState);
-			SendingRepState->CheckpointCustomDeltaState.Add(CustomDeltaProperty, NewState);
+			if (bIsRecordingReplay)
+			{
+				// Store the initial delta state in case we need it for when we're asked to resend all data since channel was first opened (bResendAllDataSinceOpen)
+				SendingRepState->CDOCustomDeltaState[CustomDeltaProperty] = NewState;
+				SendingRepState->CheckpointCustomDeltaState[CustomDeltaProperty] = NewState;
+			}
 		}
 	}
 }
@@ -855,13 +868,7 @@ void FObjectReplicator::ReceivedNak( int32 NakPacketId )
 
 						// The Nack'd packet did update this property, so we need to replace the buffer in RecentDynamic
 						// with the buffer we used to create this update (which was dropped), so that the update will be recreated on the next replicate actor
-						if (Rec->DynamicState.IsValid())
-						{
-							TSharedPtr<INetDeltaBaseState> & RecentState = SendingRepState->RecentCustomDeltaState.FindChecked(i);
-
-							RecentState.Reset();
-							RecentState = Rec->DynamicState;
-						}
+						SendingRepState->RecentCustomDeltaState[i] = Rec->DynamicState;
 
 						// We can get rid of the rest of the saved off base states since we will be regenerating these updates on the next replicate actor
 						while (Rec != nullptr)
@@ -1433,7 +1440,7 @@ void FObjectReplicator::ReplicateCustomDeltaProperties( FNetBitWriter & Bunch, F
 	check(OwningChannel);
 	check(Connection == OwningChannel->Connection);
 
-	TMap<int32, TSharedPtr<INetDeltaBaseState>>& UsingCustomDeltaStates =
+	TArray<TSharedPtr<INetDeltaBaseState>>& UsingCustomDeltaStates =
 		EResendAllDataState::None == Connection->ResendAllDataState ? SendingRepState->RecentCustomDeltaState :
 		EResendAllDataState::SinceOpen == Connection->ResendAllDataState ? SendingRepState->CDOCustomDeltaState :
 		SendingRepState->CheckpointCustomDeltaState;
@@ -1473,12 +1480,12 @@ void FObjectReplicator::ReplicateCustomDeltaProperties( FNetBitWriter & Bunch, F
 
 		TempBitWriter.Reset();
 
+		TSharedPtr<INetDeltaBaseState>& OldState = UsingCustomDeltaStates[CustomDeltaProperty];
+
 		if (Connection->ResendAllDataState != EResendAllDataState::None)
 		{
 			if (Connection->ResendAllDataState == EResendAllDataState::SinceCheckpoint)
 			{
-				TSharedPtr<INetDeltaBaseState>& OldState = UsingCustomDeltaStates.FindChecked(CustomDeltaProperty);
-
 				if (!SendCustomDeltaProperty(Object, CustomDeltaProperty, TempBitWriter, NewState, OldState))
 				{
 					continue;
@@ -1491,8 +1498,6 @@ void FObjectReplicator::ReplicateCustomDeltaProperties( FNetBitWriter & Bunch, F
 			{
 				// If we are resending data since open, we don't want to affect the current state of channel/replication, so just do the minimum and send the data, and return
 				// In this case, we'll send all of the properties since the CDO, so use the initial CDO delta state
-				TSharedPtr<INetDeltaBaseState>& OldState = UsingCustomDeltaStates.FindChecked(CustomDeltaProperty);
-
 				if (!SendCustomDeltaProperty(Object, CustomDeltaProperty, TempBitWriter, NewState, OldState))
 				{
 					continue;
@@ -1517,8 +1522,6 @@ void FObjectReplicator::ReplicateCustomDeltaProperties( FNetBitWriter & Bunch, F
 
 		ValidateRetirementHistory(Retire, Object);
 
-		TSharedPtr<INetDeltaBaseState>& OldState = UsingCustomDeltaStates.FindOrAdd(CustomDeltaProperty);
-
 		//-----------------------------------------
 		//	Do delta serialization on dynamic properties
 		//-----------------------------------------
diff --git a/Engine/Source/Runtime/Engine/Private/RepLayout.cpp b/Engine/Source/Runtime/Engine/Private/RepLayout.cpp
--- a/Engine/Source/Runtime/Engine/Private/RepLayout.cpp
+++ b/Engine/Source/Runtime/Engine/Private/RepLayout.cpp
@@ -6063,7 +6063,7 @@ void FRepLayout::PreSendCustomDeltaProperties(
 	UObject* Object,
 	UNetConnection* Connection,
 	FReplicationChangelistMgr& ChangelistMgr,
-	TMap<int32, TSharedPtr<INetDeltaBaseState>>& CustomDeltaStates) const
+	TArray<TSharedPtr<INetDeltaBaseState>>& CustomDeltaStates) const
 {
 	using namespace UE4_RepLayout_Private;
 
@@ -6154,7 +6154,7 @@ void FRepLayout::PostSendCustomDeltaProperties(
 	UObject* Object,
 	UNetConnection* Connection,
 	FReplicationChangelistMgr& ChangelistMgr,
-	TMap<int32, TSharedPtr<INetDeltaBaseState>>& CustomDeltaStates) const
+	TArray<TSharedPtr<INetDeltaBaseState>>& CustomDeltaStates) const
 {
 }
 
@@ -6948,20 +6948,31 @@ void FSendingRepState::CountBytes(FArchive& Ar) const
 
 	GRANULAR_NETWORK_MEMORY_TRACKING_TRACK("RecentCustomDeltaState",
 		RecentCustomDeltaState.CountBytes(Ar);
-	for (const auto& RecentCustomDeltaStatePair : RecentCustomDeltaState)
-	{
-		if (INetDeltaBaseState const * const BaseState = RecentCustomDeltaStatePair.Value.Get())
+		for (const TSharedPtr<INetDeltaBaseState>& LocalRecentCustomDeltaState : RecentCustomDeltaState)
 		{
-			BaseState->CountBytes(Ar);
+			if (INetDeltaBaseState const * const BaseState = LocalRecentCustomDeltaState.Get())
+			{
+				BaseState->CountBytes(Ar);
+			}
 		}
-	}
 	);
 
 	GRANULAR_NETWORK_MEMORY_TRACKING_TRACK("CDOCustomDeltaState",
 		CDOCustomDeltaState.CountBytes(Ar);
-		for (const auto& CDOCustomDeltaStatePair : CDOCustomDeltaState)
+		for (const TSharedPtr<INetDeltaBaseState>& LocalRecentCustomDeltaState : CDOCustomDeltaState)
+		{
+			if (INetDeltaBaseState const* const BaseState = LocalRecentCustomDeltaState.Get())
+			{
+				BaseState->CountBytes(Ar);
+			}
+		}
+	);
+
+	GRANULAR_NETWORK_MEMORY_TRACKING_TRACK("CheckpointCustomDeltaState",
+		CheckpointCustomDeltaState.CountBytes(Ar);
+		for (const TSharedPtr<INetDeltaBaseState>& LocalRecentCustomDeltaState : CheckpointCustomDeltaState)
 		{
-			if (INetDeltaBaseState const * const BaseState = CDOCustomDeltaStatePair.Value.Get())
+			if (INetDeltaBaseState const* const BaseState = LocalRecentCustomDeltaState.Get())
 			{
 				BaseState->CountBytes(Ar);
 			}
diff --git a/Engine/Source/Runtime/Engine/Public/Net/RepLayout.h b/Engine/Source/Runtime/Engine/Public/Net/RepLayout.h
--- a/Engine/Source/Runtime/Engine/Public/Net/RepLayout.h
+++ b/Engine/Source/Runtime/Engine/Public/Net/RepLayout.h
@@ -613,13 +613,13 @@ public:
 	TBitArray<> InactiveParents;
 
 	/** This is the delta state we need to compare with when determining what to send to a client for custom delta properties. */
-	TMap<int32, TSharedPtr<INetDeltaBaseState>> RecentCustomDeltaState;
+	TArray<TSharedPtr<INetDeltaBaseState>> RecentCustomDeltaState;
 
 	/** Same as RecentCustomDeltaState, but this will always remain as the initial CDO version. We use this to send all properties since channel was first opened (for bResendAllDataSinceOpen). */
-	TMap<int32, TSharedPtr<INetDeltaBaseState>>	CDOCustomDeltaState;
+	TArray<TSharedPtr<INetDeltaBaseState>>	CDOCustomDeltaState;
 
 	/** Same as RecentCustomDeltaState, but will represent the state at the last checkpoint. */
-	TMap<int32, TSharedPtr<INetDeltaBaseState>>	CheckpointCustomDeltaState;
+	TArray<TSharedPtr<INetDeltaBaseState>>	CheckpointCustomDeltaState;
 };
 
 /** Replication State that is unique Per Object Per Net Connection. */
@@ -1723,13 +1723,13 @@ private:
 		UObject* Object,
 		UNetConnection* Connection,
 		FReplicationChangelistMgr& ChangelistMgr,
-		TMap<int32, TSharedPtr<INetDeltaBaseState>>& CustomDeltaStates) const;
+		TArray<TSharedPtr<INetDeltaBaseState>>& CustomDeltaStates) const;
 
 	void PostSendCustomDeltaProperties(
 		UObject* Object,
 		UNetConnection* Connection,
 		FReplicationChangelistMgr& ChangelistMgr,
-		TMap<int32, TSharedPtr<INetDeltaBaseState>>& CustomDeltaStates) const;
+		TArray<TSharedPtr<INetDeltaBaseState>>& CustomDeltaStates) const;
 
 	const uint16 GetNumLifetimeCustomDeltaProperties() const;
 
