added missing post process call to screen projection

diff --git a/Engine/Source/Runtime/Engine/Private/PlayerController.cpp b/Engine/Source/Runtime/Engine/Private/PlayerController.cpp
--- a/Engine/Source/Runtime/Engine/Private/PlayerController.cpp
+++ b/Engine/Source/Runtime/Engine/Private/PlayerController.cpp
@@ -2044,6 +2044,7 @@ bool APlayerController::ProjectWorldLocationToScreenWithDistance(FVector WorldLo
 				}
 
 				ScreenLocation = FVector(ScreenPosition2D.X, ScreenPosition2D.Y, FVector::Dist(ProjectionData.ViewOrigin, WorldLocation));
+				PostProcessWorldToScreen(WorldLocation, ScreenPosition2D, bPlayerViewportRelative);
 
 				return true;
 			}
