Make a number of scope cycle counters in networking functions conditional behind net.EnableDetailedScopeCounters. This should allow framepro captures to be faster by default while preserving the ability to get detailed info as needed.

diff --git a/Engine/Source/Runtime/Engine/Classes/Engine/NetSerialization.h b/Engine/Source/Runtime/Engine/Classes/Engine/NetSerialization.h
--- a/Engine/Source/Runtime/Engine/Classes/Engine/NetSerialization.h
+++ b/Engine/Source/Runtime/Engine/Classes/Engine/NetSerialization.h
@@ -27,6 +27,8 @@ DECLARE_CYCLE_STAT_EXTERN(TEXT("NetSerializeFast Array"), STAT_NetSerializeFastA
 DECLARE_CYCLE_STAT_EXTERN(TEXT("NetSerializeFast Array BuildMap"), STAT_NetSerializeFastArray_BuildMap, STATGROUP_ServerCPU, ENGINE_API);
 DECLARE_CYCLE_STAT_EXTERN(TEXT("NetSerializeFast Array Delta Struct"), STAT_NetSerializeFastArray_DeltaStruct, STATGROUP_ServerCPU, ENGINE_API);
 
+extern ENGINE_API TAutoConsoleVariable<int32> CVarEnableDetailedScopeCounters;
+
 /**
  *	===================== NetSerialize and NetDeltaSerialize customization. =====================
  *
@@ -1070,7 +1072,7 @@ bool FFastArraySerializer::FastArrayDeltaSerialize(TArray<Type> &Items, FNetDelt
 		return FastArrayDeltaSerialize_DeltaSerializeStructs(Items, Parms, ArraySerializer);
 	}
 
-	SCOPE_CYCLE_COUNTER(STAT_NetSerializeFastArray);
+	CONDITIONAL_SCOPE_CYCLE_COUNTER(STAT_NetSerializeFastArray, CVarEnableDetailedScopeCounters.GetValueOnGameThread() > 0);
 	class UScriptStruct* InnerStruct = Type::StaticStruct();
 
 	UE_LOG(LogNetFastTArray, Log, TEXT("FastArrayDeltaSerialize for %s. %s. %s"), *InnerStruct->GetName(), *InnerStruct->GetOwnerStruct()->GetName(), Parms.Reader ? TEXT("Reading") : TEXT("Writing"));
@@ -1498,7 +1500,8 @@ bool FFastArraySerializer::FastArrayDeltaSerialize_DeltaSerializeStructs(TArray<
 		}
 	};
 
-	SCOPE_CYCLE_COUNTER(STAT_NetSerializeFastArray_DeltaStruct);
+	CONDITIONAL_SCOPE_CYCLE_COUNTER(STAT_NetSerializeFastArray_DeltaStruct, CVarEnableDetailedScopeCounters.GetValueOnGameThread() > 0);
+
 	class UScriptStruct* InnerStruct = Type::StaticStruct();
 
 	TFastArraySerializeHelper<Type, SerializerType> Helper{
diff --git a/Engine/Source/Runtime/Engine/Private/DataChannel.cpp b/Engine/Source/Runtime/Engine/Private/DataChannel.cpp
--- a/Engine/Source/Runtime/Engine/Private/DataChannel.cpp
+++ b/Engine/Source/Runtime/Engine/Private/DataChannel.cpp
@@ -73,6 +73,7 @@ FAutoConsoleVariableRef CVarNetSkipReplicatorForDestructionInfos(
 	TEXT("If enabled, skip creation of object replicator in SetChannelActor when we know there is no content payload and we're going to immediately destroy the actor."));
 
 extern TAutoConsoleVariable<int32> CVarFilterGuidRemapping;
+extern TAutoConsoleVariable<int32> CVarEnableDetailedScopeCounters;
 
 /*-----------------------------------------------------------------------------
 	UChannel implementation.
@@ -2739,7 +2740,7 @@ int64 UActorChannel::ReplicateActor()
 	const UWorld* const ActorWorld = Actor->GetWorld();
 	check(ActorWorld);
 
-#if STATS
+#if STATS || ENABLE_STATNAMEDEVENTS
 	UClass* ParentNativeClass = GetParentNativeClass(Actor->GetClass());
 	SCOPE_CYCLE_UOBJECT(ParentNativeClass, ParentNativeClass);
 #endif
@@ -3761,7 +3762,7 @@ FObjectReplicator & UActorChannel::GetActorReplicationData()
 
 TSharedRef< FObjectReplicator > & UActorChannel::FindOrCreateReplicator( UObject* Obj, bool* bOutCreated )
 {
-	SCOPE_CYCLE_COUNTER(Stat_ActorChanFindOrCreateRep);
+	CONDITIONAL_SCOPE_CYCLE_COUNTER(Stat_ActorChanFindOrCreateRep, CVarEnableDetailedScopeCounters.GetValueOnGameThread() > 0);
 	SCOPE_CYCLE_UOBJECT(ActorChannelFindOrCreateRep, Obj);
 
 	// First, try to find it on the channel replication map
diff --git a/Engine/Source/Runtime/Engine/Private/DataReplication.cpp b/Engine/Source/Runtime/Engine/Private/DataReplication.cpp
--- a/Engine/Source/Runtime/Engine/Private/DataReplication.cpp
+++ b/Engine/Source/Runtime/Engine/Private/DataReplication.cpp
@@ -49,6 +49,8 @@ static FAutoConsoleVariableRef CVarSupportsFastArrayDelta(
 	TEXT("Whether or not Fast Array Struct Delta Serialization is enabled.")
 );
 
+extern TAutoConsoleVariable<int32> CVarEnableDetailedScopeCounters;
+
 PRAGMA_DISABLE_DEPRECATION_WARNINGS
 class FNetSerializeCB : public INetSerializeCB
 {
@@ -307,7 +309,7 @@ bool FObjectReplicator::SerializeCustomDeltaProperty(
 {
 	check( NewFullState.IsValid() == false ); // NewState is passed in as nullptr and instantiated within this function if necessary
 
-	SCOPE_CYCLE_COUNTER( STAT_NetSerializeItemDeltaTime );
+	CONDITIONAL_SCOPE_CYCLE_COUNTER( STAT_NetSerializeItemDeltaTime, CVarEnableDetailedScopeCounters.GetValueOnGameThread() > 0 );
 
 	UStructProperty * StructProperty = CastChecked< UStructProperty >( Property );
 
@@ -358,7 +360,7 @@ bool FObjectReplicator::SendCustomDeltaProperty(UObject* InObject, uint16 Custom
 	check(!NewFullState.IsValid()); // NewState is passed in as nullptr and instantiated within this function if necessary
 	check(RepLayout);
 
-	SCOPE_CYCLE_COUNTER(STAT_NetSerializeItemDeltaTime);
+	CONDITIONAL_SCOPE_CYCLE_COUNTER(STAT_NetSerializeItemDeltaTime, CVarEnableDetailedScopeCounters.GetValueOnGameThread() > 0);
 
 	UNetDriver* const ConnectionDriver = Connection->GetDriver();
 	FNetSerializeCB NetSerializeCB(ConnectionDriver);
@@ -1418,7 +1420,7 @@ static FORCEINLINE FPropertyRetirement** UpdateAckedRetirements(
 
 void FObjectReplicator::ReplicateCustomDeltaProperties( FNetBitWriter & Bunch, FReplicationFlags RepFlags )
 {
-	SCOPE_CYCLE_COUNTER(STAT_NetReplicateCustomDeltaPropTime);
+	CONDITIONAL_SCOPE_CYCLE_COUNTER(STAT_NetReplicateCustomDeltaPropTime, CVarEnableDetailedScopeCounters.GetValueOnGameThread() > 0);
 
 	check(RepLayout);
 	const FRepLayout& LocalRepLayout = *RepLayout;
diff --git a/Engine/Source/Runtime/Engine/Private/NetConnection.cpp b/Engine/Source/Runtime/Engine/Private/NetConnection.cpp
--- a/Engine/Source/Runtime/Engine/Private/NetConnection.cpp
+++ b/Engine/Source/Runtime/Engine/Private/NetConnection.cpp
@@ -60,6 +60,8 @@ static TAutoConsoleVariable<int32> CVarNetPacketOrderMaxMissingPackets(TEXT("net
 
 static TAutoConsoleVariable<int32> CVarNetPacketOrderMaxCachedPackets(TEXT("net.PacketOrderMaxCachedPackets"), 32, TEXT("(NOTE: Must be power of 2!) The maximum number of packets to cache while waiting for missing packet sequences, before treating missing packets as lost."));
 
+TAutoConsoleVariable<int32> CVarEnableDetailedScopeCounters(TEXT("net.EnableDetailedScopeCounters"), 1, TEXT("Enables detailed networking scope cycle counters. There are often lots of these which can negatively impact performance."));
+
 extern int32 GNetDormancyValidate;
 
 DECLARE_CYCLE_STAT(TEXT("NetConnection SendAcks"), Stat_NetConnectionSendAck, STATGROUP_Net);
diff --git a/Engine/Source/Runtime/Engine/Private/RepLayout.cpp b/Engine/Source/Runtime/Engine/Private/RepLayout.cpp
--- a/Engine/Source/Runtime/Engine/Private/RepLayout.cpp
+++ b/Engine/Source/Runtime/Engine/Private/RepLayout.cpp
@@ -66,6 +66,8 @@ int32 MaxRepArrayMemory = UNetworkSettings::DefaultMaxRepArrayMemory;
 extern int32 GNumSharedSerializationHit;
 extern int32 GNumSharedSerializationMiss;
 
+extern TAutoConsoleVariable<int32> CVarEnableDetailedScopeCounters;
+
 FConsoleVariableSinkHandle CreateMaxArraySizeCVarAndRegisterSink()
 {
 	static FAutoConsoleVariable CVarMaxArraySize(TEXT("net.MaxRepArraySize"), MaxRepArraySize, TEXT("Maximum allowable size for replicated dynamic arrays (in number of elements). Value must be between 1 and 65535."));
@@ -1151,7 +1153,7 @@ bool FRepLayout::CompareProperties(
 	const FConstRepObjectDataBuffer Data,
 	const FReplicationFlags& RepFlags) const
 {
-	SCOPE_CYCLE_COUNTER(STAT_NetReplicateDynamicPropCompareTime);
+	CONDITIONAL_SCOPE_CYCLE_COUNTER(STAT_NetReplicateDynamicPropCompareTime, CVarEnableDetailedScopeCounters.GetValueOnGameThread() > 0);
 
 	if (LayoutState == ERepLayoutState::Empty)
 	{
@@ -1253,7 +1255,7 @@ bool FRepLayout::ReplicateProperties(
 	FNetBitWriter& Writer,
 	const FReplicationFlags& RepFlags) const
 {
-	SCOPE_CYCLE_COUNTER(STAT_NetReplicateDynamicPropTime);
+	CONDITIONAL_SCOPE_CYCLE_COUNTER(STAT_NetReplicateDynamicPropTime, CVarEnableDetailedScopeCounters.GetValueOnGameThread() > 0);
 
 	check(ObjectClass == Owner);
 
@@ -6162,7 +6164,7 @@ bool FRepLayout::DeltaSerializeFastArrayProperty(FFastArrayDeltaSerializeParams&
 {
 	using namespace UE4_RepLayout_Private;
 
-	SCOPE_CYCLE_COUNTER(STAT_RepLayout_DeltaSerializeFastArray);
+	CONDITIONAL_SCOPE_CYCLE_COUNTER(STAT_RepLayout_DeltaSerializeFastArray, CVarEnableDetailedScopeCounters.GetValueOnGameThread() > 0);
 
 	// A portion of this work could be shared across all Fast Array Properties for a given object,
 	// but that would be easier to do if the Custom Delta Serialization was completely encapsulated in FRepLayout.
