Fix bug with max distance not being correctly updated if referenced attenuation is updated but SoundCue is not in editor.

diff --git a/Engine/Source/Runtime/Engine/Classes/Sound/SoundCue.h b/Engine/Source/Runtime/Engine/Classes/Sound/SoundCue.h
--- a/Engine/Source/Runtime/Engine/Classes/Sound/SoundCue.h
+++ b/Engine/Source/Runtime/Engine/Classes/Sound/SoundCue.h
@@ -286,7 +286,9 @@ private:
 	/** Recursively sets the branch culling exclusion on random nodes in this sound cue. */
 	void RecursivelySetExcludeBranchCulling(USoundNode* CurrentNode);
 
+	float FindMaxDistanceInternal() const;
+
 	/** Ptr to interface to sound cue editor operations. */
 	static ENGINE_API TSharedPtr<ISoundCueAudioEditor> SoundCueAudioEditor;
-#endif
+#endif // WITH_EDITOR
 };
diff --git a/Engine/Source/Runtime/Engine/Private/SoundCue.cpp b/Engine/Source/Runtime/Engine/Private/SoundCue.cpp
--- a/Engine/Source/Runtime/Engine/Private/SoundCue.cpp
+++ b/Engine/Source/Runtime/Engine/Private/SoundCue.cpp
@@ -78,21 +78,7 @@ void USoundCue::CacheAggregateValues()
 
 		Duration = FirstNode->GetDuration();
 
-		MaxDistance = FirstNode->GetMaxDistance();
-
-		if (const FSoundAttenuationSettings* Settings = GetAttenuationSettingsToApply())
-		{
-			if (Settings->bAttenuate)
-			{
-				MaxDistance = FMath::Max(MaxDistance, Settings->GetMaxDimension());
-			}
-		}
-
-		// If no sound cue nodes overrode the max distance, we need to check the base attenuation
-		if (MaxDistance == 0.0f)
-		{
-			MaxDistance = USoundBase::GetMaxDistance();
-		}
+		MaxDistance = FindMaxDistanceInternal();
 
 		bHasDelayNode = FirstNode->HasDelayNode();
 		bHasConcatenatorNode = FirstNode->HasConcatenatorNode();
@@ -245,6 +231,33 @@ void USoundCue::RecursivelySetExcludeBranchCulling(USoundNode* CurrentNode)
 	}
 }
 
+float USoundCue::FindMaxDistanceInternal() const
+{
+	float OutMaxDistance = 0.0f;
+	if (const FSoundAttenuationSettings* Settings = GetAttenuationSettingsToApply())
+	{
+		if (!Settings->bAttenuate)
+		{
+			return WORLD_MAX;
+		}
+
+		OutMaxDistance = FMath::Max(OutMaxDistance, Settings->GetMaxDimension());
+	}
+
+	if (FirstNode)
+	{
+		OutMaxDistance = FMath::Max(OutMaxDistance, FirstNode->GetMaxDistance());
+	}
+
+	if (OutMaxDistance >= 0.0f)
+	{
+		return OutMaxDistance;
+	}
+
+	// If no sound cue nodes has overridden the max distance, check the base attenuation
+	return USoundBase::GetMaxDistance();
+}
+
 void USoundCue::PostEditChangeProperty(struct FPropertyChangedEvent& PropertyChangedEvent)
 {
 	Super::PostEditChangeProperty(PropertyChangedEvent);
@@ -410,7 +423,10 @@ int32 USoundCue::GetResourceSizeForFormat(FName Format)
 
 float USoundCue::GetMaxDistance() const
 {
-	return MaxDistance;
+	// Always recalc the max distance when in the editor as it could change
+	// from a referenced attenuation asset being updated without this cue
+	// asset re-caching the aggregate 'MaxDistance' value
+	return GIsEditor ? FindMaxDistanceInternal() : MaxDistance;
 }
 
 float USoundCue::GetDuration()
