- Extract out focus calcuation code and make publically available so virtualization system can periodically update/check if focus updates push sounds out of virtualization

diff --git a/Engine/Source/Runtime/Engine/Private/ActiveSound.cpp b/Engine/Source/Runtime/Engine/Private/ActiveSound.cpp
--- a/Engine/Source/Runtime/Engine/Private/ActiveSound.cpp
+++ b/Engine/Source/Runtime/Engine/Private/ActiveSound.cpp
@@ -88,14 +88,10 @@ FActiveSound::FActiveSound()
 	, CurrentOcclusionVolumeAttenuation(1.0f)
 	, SubtitlePriority(DEFAULT_SUBTITLE_PRIORITY)
 	, Priority(1.0f)
-	, FocusPriorityScale(1.0f)
-	, FocusDistanceScale(1.0f)
 	, VolumeConcurrency(0.0f)
 	, OcclusionCheckInterval(0.f)
 	, LastOcclusionCheckTime(TNumericLimits<float>::Lowest())
 	, MaxDistance(WORLD_MAX)
-	, Azimuth(0.0f)
-	, AbsoluteAzimuth(0.0f)
 	, LastLocation(FVector::ZeroVector)
 	, AudioVolumeID(0)
 	, LastUpdateTime(0.f)
@@ -109,7 +105,6 @@ FActiveSound::FActiveSound()
 	, DebugColor(FColor::Black)
 #endif // !(UE_BUILD_SHIPPING || UE_BUILD_TEST)
 	, ClosestListenerPtr(nullptr)
-	, InternalFocusFactor(1.0f)
 {
 	if (!ActiveSoundTraceDelegate.IsBound())
 	{
@@ -555,7 +550,7 @@ void FActiveSound::UpdateWaveInstances(TArray<FWaveInstance*> &InWaveInstances,
 	{
 		if (bHasAttenuationSettings)
 		{
-			ApplyAttenuation(ParseParams, *ClosestListenerPtr);
+			UpdateAttenuation(DeltaTime, ParseParams, *ClosestListenerPtr);
 		}
 		else
 		{
@@ -1434,20 +1429,70 @@ float FActiveSound::GetAttenuationFrequency(const FSoundAttenuationSettings* Set
 	return FMath::Clamp<float>(OutputFrequency, MIN_FILTER_FREQUENCY, MAX_FILTER_FREQUENCY);
 }
 
-void FActiveSound::ApplyAttenuation(FSoundParseParameters& ParseParams, const FListener& Listener, const FSoundAttenuationSettings* SettingsAttenuationNode)
+void FActiveSound::UpdateFocusData(float DeltaTime, const FAttenuationListenerData& ListenerData, FAttenuationFocusData* OutFocusData)
+{
+	FAttenuationFocusData* FocusDataToUpdate = OutFocusData ? OutFocusData : &FocusData;
+
+	AudioDevice->GetAzimuth(ListenerData, FocusDataToUpdate->Azimuth, FocusDataToUpdate->AbsoluteAzimuth);
+
+	FocusDataToUpdate->DistanceScale = 1.0f;
+	FocusDataToUpdate->PriorityScale = 1.0f;
+
+	if (!ListenerData.AttenuationSettings->bEnableListenerFocus)
+	{
+		return;
+	}
+
+	if (!ListenerData.AttenuationSettings->bSpatialize)
+	{
+		return;
+	}
+
+	const FGlobalFocusSettings& FocusSettings = AudioDevice->GetGlobalFocusSettings();
+	const float TargetFocusFactor = AudioDevice->GetFocusFactor(FocusDataToUpdate->Azimuth, *ListenerData.AttenuationSettings);
+
+	// User opt-in for focus interpolation
+	if (ListenerData.AttenuationSettings->bEnableFocusInterpolation)
+	{
+		// Determine which interpolation speed to use (attack/release)
+		float InterpSpeed;
+		if (TargetFocusFactor <= FocusDataToUpdate->FocusFactor)
+		{
+			InterpSpeed = ListenerData.AttenuationSettings->FocusAttackInterpSpeed;
+		}
+		else
+		{
+			InterpSpeed = ListenerData.AttenuationSettings->FocusReleaseInterpSpeed;
+		}
+
+		FocusDataToUpdate->FocusFactor = FMath::FInterpTo(FocusDataToUpdate->FocusFactor, TargetFocusFactor, DeltaTime, InterpSpeed);
+	}
+	else
+	{
+		// Set focus directly to target value
+		FocusDataToUpdate->FocusFactor = TargetFocusFactor;
+	}
+
+	// Scale the volume-weighted priority scale value we use for sorting this sound for voice-stealing
+	FocusDataToUpdate->PriorityScale = ListenerData.AttenuationSettings->GetFocusPriorityScale(FocusSettings, FocusDataToUpdate->FocusFactor);
+	FocusDataToUpdate->DistanceScale = ListenerData.AttenuationSettings->GetFocusDistanceScale(FocusSettings, FocusDataToUpdate->FocusFactor);
+	FocusDataToUpdate->VolumeScale = ListenerData.AttenuationSettings->GetFocusAttenuation(FocusSettings, FocusDataToUpdate->FocusFactor);
+}
+
+void FActiveSound::ParseAttenuation(FSoundParseParameters& ParseParams, const FListener& Listener, const FSoundAttenuationSettings& AttenuationSettings)
+{
+	UpdateAttenuation(0.0f, ParseParams, Listener, &AttenuationSettings);
+}
+
+void FActiveSound::UpdateAttenuation(float DeltaTime, FSoundParseParameters& ParseParams, const FListener& Listener, const FSoundAttenuationSettings* SettingsAttenuationNode)
 {
-	float& Volume = ParseParams.Volume;
-	const FTransform& SoundTransform = ParseParams.Transform;
 	const FVector& ListenerLocation = Listener.Transform.GetTranslation();
 
 	// Get the attenuation settings to use for this application to the active sound
 	const FSoundAttenuationSettings* Settings = SettingsAttenuationNode ? SettingsAttenuationNode : &AttenuationSettings;
 
-	FAttenuationListenerData ListenerData;
-
-	// Reset distance and priority scale to 1.0 in case changed in editor
-	FocusDistanceScale = 1.0f;
-	FocusPriorityScale = 1.0f;
+	// Reset Focus data and recompute if necessary
+	FAttenuationFocusData FocusDataToApply;
 
 	if (Settings->bEnableReverbSend)
 	{
@@ -1458,60 +1503,26 @@ void FActiveSound::ApplyAttenuation(FSoundParseParameters& ParseParams, const FL
 		ParseParams.ReverbSendLevelDistanceRange = { Settings->ReverbDistanceMin, Settings->ReverbDistanceMax };
 	}
 
-	if (Settings->bSpatialize || Settings->bEnableListenerFocus)
+	check(AudioDevice);
+	FAttenuationListenerData ListenerData = FAttenuationListenerData::Create(*AudioDevice, Listener.Transform, ParseParams.Transform, *Settings);
+
+	// Ignore focus calculation if attenuation is being updated due to override attenuation parameters being passed in.
+	// Focus is always only calculated by the parent attenuation's data, never by the attenuation node's settings.
+	const bool bIgnoreFocusCalculation = false;// SettingsAttenuationNode != nullptr;
+	if (!bIgnoreFocusCalculation && (Settings->bSpatialize || Settings->bEnableListenerFocus))
 	{
-		AudioDevice->GetAzimuth(ListenerData, SoundTransform, *Settings, Listener.Transform, Azimuth, AbsoluteAzimuth);
+		// Update azimuth angles prior to updating focus as it uses this in calculating
+		// in and out of focus values.
+		UpdateFocusData(DeltaTime, ListenerData, &FocusDataToApply);
 
+		ParseParams.Volume *= FocusDataToApply.VolumeScale;
+		ParseParams.Priority *= FocusDataToApply.PriorityScale;
 		if (Settings->bSpatialize)
 		{
 			ParseParams.AttenuationDistance = ListenerData.AttenuationDistance;
 			ParseParams.ListenerToSoundDistance = ListenerData.ListenerToSoundDistance;
 			ParseParams.ListenerToSoundDistanceForPanning = ListenerData.ListenerToSoundDistanceForPanning;
-			ParseParams.AbsoluteAzimuth = AbsoluteAzimuth;
-		}
-
-		if (Settings->bEnableListenerFocus)
-		{
-			// Compute the azimuth of the active sound
-			const FGlobalFocusSettings& FocusSettings = AudioDevice->GetGlobalFocusSettings();
-
-			// Get the current target focus factor
-			const float TargetFocusFactor = AudioDevice->GetFocusFactor(Azimuth, *Settings);
-
-			// User opt-in for focus interpolation
-			if (Settings->bEnableFocusInterpolation)
-			{
-				// Determine which interpolation speed to use (attack/release)
-				float InterpSpeed;
-				if (TargetFocusFactor <= InternalFocusFactor)
-				{
-					InterpSpeed = Settings->FocusAttackInterpSpeed;
-				}
-				else
-				{
-					InterpSpeed = Settings->FocusReleaseInterpSpeed;
-				}
-
-				// Interpolate the internal focus factor to the target value
-				const float DeviceDeltaTime = AudioDevice->GetDeviceDeltaTime();
-				InternalFocusFactor = FMath::FInterpTo(InternalFocusFactor, TargetFocusFactor, DeviceDeltaTime, InterpSpeed);
-			}
-			else
-			{
-				// Set focus directly to target value
-				InternalFocusFactor = TargetFocusFactor;
-			}
-
-			// Get the volume scale to apply the volume calculation based on the focus factor
-			const float FocusVolumeAttenuation = Settings->GetFocusAttenuation(FocusSettings, InternalFocusFactor);
-			Volume *= FocusVolumeAttenuation;
-
-			// Scale the volume-weighted priority scale value we use for sorting this sound for voice-stealing
-			FocusPriorityScale = Settings->GetFocusPriorityScale(FocusSettings, InternalFocusFactor);
-			ParseParams.Priority *= FocusPriorityScale;
-
-			// Get the distance scale to use when computing distance-calculations for 3d attenuation
-			FocusDistanceScale = Settings->GetFocusDistanceScale(FocusSettings, InternalFocusFactor);
+			ParseParams.AbsoluteAzimuth = FocusDataToApply.AbsoluteAzimuth;
 		}
 	}
 
@@ -1523,12 +1534,11 @@ void FActiveSound::ApplyAttenuation(FSoundParseParameters& ParseParams, const FL
 		if (Settings->AttenuationShape == EAttenuationShape::Sphere)
 		{
 			// Update attenuation data in-case it hasn't been updated
-			AudioDevice->GetAttenuationListenerData(ListenerData, SoundTransform, *Settings, &Listener.Transform);
-			ParseParams.DistanceAttenuation *= Settings->AttenuationEval(ListenerData.AttenuationDistance, Settings->FalloffDistance, FocusDistanceScale);
+			ParseParams.DistanceAttenuation *= Settings->AttenuationEval(ListenerData.AttenuationDistance, Settings->FalloffDistance, FocusDataToApply.DistanceScale);
 		}
 		else
 		{
-			ParseParams.DistanceAttenuation *= Settings->Evaluate(SoundTransform, ListenerLocation, FocusDistanceScale);
+			ParseParams.DistanceAttenuation *= Settings->Evaluate(ParseParams.Transform, ListenerLocation, FocusDataToApply.DistanceScale);
 		}
 	}
 
@@ -1551,7 +1561,7 @@ void FActiveSound::ApplyAttenuation(FSoundParseParameters& ParseParams, const FL
 				}
 			}
 		}
-		else if (Volume > 0.0f && !AudioDevice->IsAudioDeviceMuted())
+		else if (ParseParams.Volume > 0.0f && !AudioDevice->IsAudioDeviceMuted())
 		{
 			check(ClosestListenerPtr);
 			CheckOcclusion(ClosestListenerPtr->Transform.GetTranslation(), ParseParams.Transform.GetTranslation(), Settings);
@@ -1600,8 +1610,6 @@ void FActiveSound::ApplyAttenuation(FSoundParseParameters& ParseParams, const FL
 	// Attenuate with the absorption filter if necessary
 	if (Settings->bAttenuateWithLPF)
 	{
-		AudioDevice->GetAttenuationListenerData(ListenerData, SoundTransform, *Settings, &Listener.Transform);
-
 		FVector2D AbsorptionLowPassFrequencyRange = { Settings->LPFFrequencyAtMin, Settings->LPFFrequencyAtMax };
 		FVector2D AbsorptionHighPassFrequencyRange = { Settings->HPFFrequencyAtMin, Settings->HPFFrequencyAtMax };
 		const float AttenuationLowpassFilterFrequency = GetAttenuationFrequency(Settings, ListenerData, AbsorptionLowPassFrequencyRange, Settings->CustomLowpassAirAbsorptionCurve);
@@ -1638,4 +1646,10 @@ void FActiveSound::ApplyAttenuation(FSoundParseParameters& ParseParams, const FL
 	{
 		ParseParams.SpatializationMethod = Settings->SpatializationAlgorithm;
 	}
+
+	// If not overriding from a node, set focus data
+	if (!SettingsAttenuationNode)
+	{
+		FocusData = FocusDataToApply;
+	}
 }
diff --git a/Engine/Source/Runtime/Engine/Private/AudioDevice.cpp b/Engine/Source/Runtime/Engine/Private/AudioDevice.cpp
--- a/Engine/Source/Runtime/Engine/Private/AudioDevice.cpp
+++ b/Engine/Source/Runtime/Engine/Private/AudioDevice.cpp
@@ -155,8 +155,45 @@ FAutoConsoleVariableRef CVarDisableBinauralSpatialization(
 namespace
 {
 	using FVirtualLoopPair = TPair<FActiveSound*, FAudioVirtualLoop>;
+
+#if !UE_BUILD_SHIPPING
+	int32 PrecachedRealtime = 0;
+	int32 PrecachedNative = 0;
+	int32 TotalNativeSize = 0;
+	float AverageNativeLength = 0.f;
+	TMap<int32, int32> NativeChannelCount;
+	TMap<int32, int32> NativeSampleRateCount;
+#endif // !UE_BUILD_SHIPPING
 } // namespace <>
 
+
+FAttenuationListenerData FAttenuationListenerData::Create(const FAudioDevice& AudioDevice, const FTransform& InListenerTransform, const FTransform& InSoundTransform, const FSoundAttenuationSettings& InAttenuationSettings)
+{
+	FAttenuationListenerData ListenerData(InListenerTransform, InSoundTransform, InAttenuationSettings);
+
+	const FVector SoundTranslation = InSoundTransform.GetTranslation();
+	FVector ListenerToSound = SoundTranslation - InListenerTransform.GetTranslation();
+	ListenerToSound.ToDirectionAndLength(ListenerData.ListenerToSoundDir, ListenerData.ListenerToSoundDistance);
+
+	// Store the actual distance for surround-panning sources with spread (AudioMixer)
+	ListenerData.ListenerToSoundDistanceForPanning = ListenerData.ListenerToSoundDistance;
+
+	if (AudioDevice.IsUsingListenerAttenuationOverride())
+	{
+		const FVector& AttenuationOverride = AudioDevice.GetListenerAttenuationOverride();
+		ListenerData.ListenerToSoundDistance = (SoundTranslation - AttenuationOverride).Size();
+	}
+
+	const FSoundAttenuationSettings& AttenuationSettings = *ListenerData.AttenuationSettings;
+	if ((AttenuationSettings.bAttenuate && AttenuationSettings.AttenuationShape == EAttenuationShape::Sphere) || AttenuationSettings.bAttenuateWithLPF)
+	{
+		ListenerData.AttenuationDistance = FMath::Max(ListenerData.ListenerToSoundDistance - AttenuationSettings.AttenuationShapeExtents.X, 0.0f);
+	}
+
+	return MoveTemp(ListenerData);
+}
+
+
 /*-----------------------------------------------------------------------------
 	FAudioDevice implementation.
 -----------------------------------------------------------------------------*/
@@ -941,13 +978,6 @@ void FAudioDevice::ShowSoundClassHierarchy(FOutputDevice& Ar, USoundClass* InSou
 	}
 }
 
-int32 PrecachedRealtime = 0;
-int32 PrecachedNative = 0;
-int32 TotalNativeSize = 0;
-float AverageNativeLength = 0.f;
-TMap<int32, int32> NativeChannelCount;
-TMap<int32, int32> NativeSampleRateCount;
-
 bool FAudioDevice::HandleDumpSoundInfoCommand(const TCHAR* Cmd, FOutputDevice& Ar)
 {
 	FAudioThreadSuspendContext AudioThreadSuspend;
@@ -4658,10 +4688,6 @@ void FAudioDevice::GetMaxDistanceAndFocusFactor(USoundBase* Sound, const UWorld*
 
 		if (AttenuationSettingsToApply->bSpatialize && AttenuationSettingsToApply->bEnableListenerFocus)
 		{
-			// Now scale the max distance based on the focus settings in the attenuation settings
-			FAttenuationListenerData ListenerData;
-			float Azimuth = 0.0f;
-			float AbsoluteAzimuth = 0.0f;
 			const int32 ClosestListenerIndex = FindClosestListenerIndex(SoundTransform);
 			if (ClosestListenerIndex >= ListenerTransforms.Num())
 			{
@@ -4669,8 +4695,12 @@ void FAudioDevice::GetMaxDistanceAndFocusFactor(USoundBase* Sound, const UWorld*
 				return;
 			}
 
-			const FTransform& ListenerTransform = ListenerTransforms[ClosestListenerIndex];
-			GetAzimuth(ListenerData, SoundTransform, *AttenuationSettingsToApply, ListenerTransform, Azimuth, AbsoluteAzimuth);
+			// Now scale the max distance based on the focus settings in the attenuation settings
+			FAttenuationListenerData ListenerData = FAttenuationListenerData::Create(*this, ListenerTransforms[ClosestListenerIndex], SoundTransform, *AttenuationSettingsToApply);
+
+			float Azimuth = 0.0f;
+			float AbsoluteAzimuth = 0.0f;
+			GetAzimuth(ListenerData, Azimuth, AbsoluteAzimuth);
 			OutFocusFactor = GetFocusFactor(Azimuth, *AttenuationSettingsToApply);
 		}
 	}
@@ -4741,7 +4771,7 @@ bool FAudioDevice::SoundIsAudible(const FActiveSound& NewActiveSound)
 		return true;
 	}
 
-	const float ApparentMaxDistance = NewActiveSound.MaxDistance * NewActiveSound.FocusDistanceScale;
+	const float ApparentMaxDistance = NewActiveSound.MaxDistance * NewActiveSound.FocusData.DistanceScale;
 	if (LocationIsAudible(NewActiveSound.Transform.GetLocation(), ApparentMaxDistance))
 	{
 		return true;
@@ -4801,71 +4831,18 @@ int32 FAudioDevice::FindClosestListenerIndex(const FTransform& SoundTransform) c
 	return INDEX_NONE;
 }
 
-void FAudioDevice::GetAttenuationListenerData(FAttenuationListenerData& OutListenerData, const FTransform& SoundTransform, const FSoundAttenuationSettings& AttenuationSettings, const FTransform* InListenerTransform) const
+void FAudioDevice::GetAzimuth(const FAttenuationListenerData& ListenerData, float& OutAzimuth, float& OutAbsoluteAzimuth) const
 {
-	// Only compute various components of the listener of it hasn't been computed yet
-	if (!OutListenerData.bDataComputed)
-	{
+	const FVector& ListenerForwardDir = ListenerData.ListenerTransform.GetUnitAxis(EAxis::X);
 
-		// Use the optional input listener param
-		if (InListenerTransform)
-		{
-			OutListenerData.ListenerTransform = *InListenerTransform;
-		}
-		// If not set, then we need to find the closest listener
-		else
-		{
-			const int32 ClosestListenerIndex = FindClosestListenerIndex(SoundTransform);
-			if (IsInAudioThread())
-			{
-				OutListenerData.ListenerTransform = Listeners[ClosestListenerIndex].Transform;
-			}
-			else if (IsInGameThread())
-			{
-				OutListenerData.ListenerTransform = ListenerTransforms[ClosestListenerIndex];
-			}
-		}
-
-		FVector ListenerToSound;
-		const FVector ListenerLocation = OutListenerData.ListenerTransform.GetTranslation();
-		ListenerToSound = SoundTransform.GetTranslation() - ListenerLocation;
-		ListenerToSound.ToDirectionAndLength(OutListenerData.ListenerToSoundDir, OutListenerData.ListenerToSoundDistance);
-
-		// store the actual distance for surround-panning sources with spread (AudioMixer)
-		OutListenerData.ListenerToSoundDistanceForPanning = OutListenerData.ListenerToSoundDistance;
-
-		if (bUseListenerAttenuationOverride)
-		{
-			ListenerToSound = SoundTransform.GetTranslation() - ListenerAttenuationOverride;
-			FVector Temp;
-			ListenerToSound.ToDirectionAndLength(Temp, OutListenerData.ListenerToSoundDistance);
-		}
-
-		OutListenerData.AttenuationDistance = 0.0f;
-
-		if ((AttenuationSettings.bAttenuate && AttenuationSettings.AttenuationShape == EAttenuationShape::Sphere) || AttenuationSettings.bAttenuateWithLPF)
-		{
-			OutListenerData.AttenuationDistance = FMath::Max(OutListenerData.ListenerToSoundDistance - AttenuationSettings.AttenuationShapeExtents.X, 0.f);
-		}
-
-		OutListenerData.bDataComputed = true;
-	}
-}
-
-void FAudioDevice::GetAzimuth(FAttenuationListenerData& OutListenerData, const FTransform& SoundTransform, const FSoundAttenuationSettings& AttenuationSettings, const FTransform& ListenerTransform, float& OutAzimuth, float& OutAbsoluteAzimuth) const
-{
-	GetAttenuationListenerData(OutListenerData, SoundTransform, AttenuationSettings, &ListenerTransform);
-
-	const FVector& ListenerForwardDir = OutListenerData.ListenerTransform.GetUnitAxis(EAxis::X);
-
-	const float SoundToListenerForwardDotProduct = FVector::DotProduct(ListenerForwardDir, OutListenerData.ListenerToSoundDir);
+	const float SoundToListenerForwardDotProduct = FVector::DotProduct(ListenerForwardDir, ListenerData.ListenerToSoundDir);
 	const float SoundListenerAngleRadians = FMath::Acos(SoundToListenerForwardDotProduct);
 
 	// Normal azimuth only goes to 180 (0 is in front, 180 is behind).
 	OutAzimuth = FMath::RadiansToDegrees(SoundListenerAngleRadians);
 
-	const FVector& ListenerRightDir = OutListenerData.ListenerTransform.GetUnitAxis(EAxis::Y);
-	const float SoundToListenerRightDotProduct = FVector::DotProduct(ListenerRightDir, OutListenerData.ListenerToSoundDir);
+	const FVector& ListenerRightDir = ListenerData.ListenerTransform.GetUnitAxis(EAxis::Y);
+	const float SoundToListenerRightDotProduct = FVector::DotProduct(ListenerRightDir, ListenerData.ListenerToSoundDir);
 
 	FVector AbsAzimuthVector2D = FVector(SoundToListenerForwardDotProduct, SoundToListenerRightDotProduct, 0.0f);
 	AbsAzimuthVector2D.Normalize();
@@ -5104,7 +5081,7 @@ void FAudioDevice::PlaySoundAtLocation(USoundBase* Sound, UWorld* World, float V
 
 	const FSoundAttenuationSettings* AttenuationSettingsToApply = (AttenuationSettings ? &AttenuationSettings->Attenuation : Sound->GetAttenuationSettingsToApply());
 	float MaxDistance = 0.0f;
-	float FocusFactor = 0.0f;
+	float FocusFactor = 1.0f;
 
 	GetMaxDistanceAndFocusFactor(Sound, World, Location, AttenuationSettingsToApply, MaxDistance, FocusFactor);
 
@@ -5130,8 +5107,8 @@ void FAudioDevice::PlaySoundAtLocation(USoundBase* Sound, UWorld* World, float V
 			const FGlobalFocusSettings& FocusSettings = GetGlobalFocusSettings();
 
 			NewActiveSound.AttenuationSettings = *AttenuationSettingsToApply;
-			NewActiveSound.FocusPriorityScale = AttenuationSettingsToApply->GetFocusPriorityScale(FocusSettings, FocusFactor);
-			NewActiveSound.FocusDistanceScale = AttenuationSettingsToApply->GetFocusDistanceScale(FocusSettings, FocusFactor);
+			NewActiveSound.FocusData.PriorityScale = AttenuationSettingsToApply->GetFocusPriorityScale(FocusSettings, FocusFactor);
+			NewActiveSound.FocusData.DistanceScale = AttenuationSettingsToApply->GetFocusDistanceScale(FocusSettings, FocusFactor);
 		}
 
 		NewActiveSound.MaxDistance = MaxDistance;
@@ -5778,8 +5755,8 @@ void FAudioDevice::UpdateVirtualLoops(bool bForceUpdate)
 			// actively playing sounds in concurrency checks.
 			ActiveSound.PlaybackTime += DeltaTime * ActiveSound.MinCurrentPitch;
 
-			// If the loop is ready to realize, add to array to be re-triggered outside of the loop
-			// to avoid map manipulation while iterating.
+			// If the loop is ready to realize, add to array to be re-triggered
+			// outside of the loop to avoid map manipulation while iterating.
 			if (VirtualLoop.CanRealize(DeltaTime, bForceUpdate))
 			{
 				VirtualLoopsToRetrigger.Add(VirtualLoop);
diff --git a/Engine/Source/Runtime/Engine/Private/AudioVirtualLoop.cpp b/Engine/Source/Runtime/Engine/Private/AudioVirtualLoop.cpp
--- a/Engine/Source/Runtime/Engine/Private/AudioVirtualLoop.cpp
+++ b/Engine/Source/Runtime/Engine/Private/AudioVirtualLoop.cpp
@@ -140,29 +140,61 @@ bool FAudioVirtualLoop::IsInAudibleRange(const FActiveSound& InActiveSound, cons
 	if (InActiveSound.bHasAttenuationSettings)
 	{
 		// If we are not using distance-based attenuation, this sound will be audible regardless of distance.
-		const FSoundAttenuationSettings* AttenuationSettingsToApply = InActiveSound.bHasAttenuationSettings ? &InActiveSound.AttenuationSettings : nullptr;
-		if (!AttenuationSettingsToApply->bAttenuate)
+		if (!InActiveSound.AttenuationSettings.bAttenuate)
 		{
 			return true;
 		}
 
-		DistanceScale = InActiveSound.FocusDistanceScale;
+		DistanceScale = InActiveSound.FocusData.DistanceScale;
 	}
 
-	DistanceScale = FMath::Max(DistanceScale, 0.0001f);
+	DistanceScale = FMath::Max(DistanceScale, KINDA_SMALL_NUMBER);
 	const FVector Location = InActiveSound.Transform.GetLocation();
 	return AudioDevice->LocationIsAudible(Location, InActiveSound.MaxDistance / DistanceScale);
 }
 
+void FAudioVirtualLoop::UpdateFocusData(float DeltaTime)
+{
+	check(ActiveSound);
+
+	if (!ActiveSound->bHasAttenuationSettings)
+	{
+		return;
+	}
+
+	// If we are not using distance-based attenuation, this sound will be audible regardless of distance.
+	if (!ActiveSound->AttenuationSettings.bAttenuate)
+	{
+		return;
+	}
+
+	check(ActiveSound->AudioDevice);
+	const FAudioDevice& AudioDevice = *ActiveSound->AudioDevice;
+
+	FAttenuationFocusData FocusData;
+	FTransform ListenerTransform;
+	const TArray<FListener>& Listeners = AudioDevice.GetListeners();
+	if (Listeners.Num() > 0)
+	{
+		int32 ClosestListenerIndex = FAudioDevice::FindClosestListenerIndex(ActiveSound->Transform, Listeners);
+		ListenerTransform = Listeners[ClosestListenerIndex].Transform;
+	}
+
+	FAttenuationListenerData ListenerData = FAttenuationListenerData::Create(AudioDevice, ListenerTransform, ActiveSound->Transform, ActiveSound->AttenuationSettings);
+	ActiveSound->UpdateFocusData(DeltaTime, ListenerData);
+}
+
 bool FAudioVirtualLoop::CanRealize(float DeltaTime, bool bForceUpdate)
 {
+	const float UpdateDelta = TimeSinceLastUpdate + DeltaTime;
+
 	if (bForceUpdate)
 	{
 		TimeSinceLastUpdate = 0.0f;
 	}
 	else if (UpdateInterval > 0.0f)
 	{
-		TimeSinceLastUpdate += DeltaTime;
+		TimeSinceLastUpdate = UpdateDelta;
 		if (UpdateInterval > TimeSinceLastUpdate)
 		{
 			return false;
@@ -174,6 +206,8 @@ bool FAudioVirtualLoop::CanRealize(float DeltaTime, bool bForceUpdate)
 	FAudioDebugger::DrawDebugInfo(*this);
 #endif // ENABLE_AUDIO_DEBUG
 
+	UpdateFocusData(UpdateDelta);
+
 	// If not audible, update when will be checked again and return false
 	if (!IsInAudibleRange(*ActiveSound))
 	{
diff --git a/Engine/Source/Runtime/Engine/Private/Components/AudioComponent.cpp b/Engine/Source/Runtime/Engine/Private/Components/AudioComponent.cpp
--- a/Engine/Source/Runtime/Engine/Private/Components/AudioComponent.cpp
+++ b/Engine/Source/Runtime/Engine/Private/Components/AudioComponent.cpp
@@ -390,7 +390,7 @@ void UAudioComponent::PlayInternal(const float StartTime, const float FadeInDura
 			const FSoundAttenuationSettings* AttenuationSettingsToApply = bAllowSpatialization ? GetAttenuationSettingsToApply() : nullptr;
 
 			float MaxDistance = 0.0f;
-			float FocusFactor = 0.0f;
+			float FocusFactor = 1.0f;
 			FVector Location = GetComponentTransform().GetLocation();
 
 			AudioDevice->GetMaxDistanceAndFocusFactor(Sound, World, Location, AttenuationSettingsToApply, MaxDistance, FocusFactor);
@@ -454,7 +454,7 @@ void UAudioComponent::PlayInternal(const float StartTime, const float FadeInDura
 			if (NewActiveSound.bHasAttenuationSettings)
 			{
 				NewActiveSound.AttenuationSettings = *AttenuationSettingsToApply;
-				NewActiveSound.FocusPriorityScale = AttenuationSettingsToApply->GetFocusPriorityScale(AudioDevice->GetGlobalFocusSettings(), FocusFactor);
+				NewActiveSound.FocusData.PriorityScale = AttenuationSettingsToApply->GetFocusPriorityScale(AudioDevice->GetGlobalFocusSettings(), FocusFactor);
 			}
 
 			NewActiveSound.EnvelopeFollowerAttackTime = FMath::Max(EnvelopeFollowerAttackTime, 0);
diff --git a/Engine/Source/Runtime/Engine/Private/SoundNodeAttenuation.cpp b/Engine/Source/Runtime/Engine/Private/SoundNodeAttenuation.cpp
--- a/Engine/Source/Runtime/Engine/Private/SoundNodeAttenuation.cpp
+++ b/Engine/Source/Runtime/Engine/Private/SoundNodeAttenuation.cpp
@@ -51,22 +51,22 @@ const FSoundAttenuationSettings* USoundNodeAttenuation::GetAttenuationSettingsTo
 	return Settings;
 }
 
-void USoundNodeAttenuation::ParseNodes( FAudioDevice* AudioDevice, const UPTRINT NodeWaveInstanceHash, FActiveSound& ActiveSound, const FSoundParseParameters& ParseParams, TArray<FWaveInstance*>& WaveInstances )
+void USoundNodeAttenuation::ParseNodes(FAudioDevice* AudioDevice, const UPTRINT NodeWaveInstanceHash, FActiveSound& ActiveSound, const FSoundParseParameters& ParseParams, TArray<FWaveInstance*>& WaveInstances)
 {
-	const FSoundAttenuationSettings* Settings = (ActiveSound.bAllowSpatialization ? GetAttenuationSettingsToApply() : NULL);
-
 	FSoundParseParameters UpdatedParseParams = ParseParams;
+
+	const FSoundAttenuationSettings* Settings = (ActiveSound.bAllowSpatialization ? GetAttenuationSettingsToApply() : nullptr);
 	if (Settings)
 	{
 		const FListener& Listener = AudioDevice->GetListeners()[0];
 
 		// Update this node's attenuation settings overrides
-		ActiveSound.ApplyAttenuation(UpdatedParseParams, Listener, Settings);
+		ActiveSound.ParseAttenuation(UpdatedParseParams, Listener, *Settings);
 	}
 	else
 	{
 		UpdatedParseParams.bUseSpatialization = false;
 	}
 
-	Super::ParseNodes( AudioDevice, NodeWaveInstanceHash, ActiveSound, UpdatedParseParams, WaveInstances );
+	Super::ParseNodes(AudioDevice, NodeWaveInstanceHash, ActiveSound, UpdatedParseParams, WaveInstances);
 }
diff --git a/Engine/Source/Runtime/Engine/Public/ActiveSound.h b/Engine/Source/Runtime/Engine/Public/ActiveSound.h
--- a/Engine/Source/Runtime/Engine/Public/ActiveSound.h
+++ b/Engine/Source/Runtime/Engine/Public/ActiveSound.h
@@ -26,6 +26,40 @@ class USoundWave;
 struct FListener;
 struct FAttenuationListenerData;
 
+/**
+ * Attenuation focus system data computed per update per active sound
+ */
+struct FAttenuationFocusData
+{
+	/** Azimuth of the active sound relative to the listener. Used by sound focus. */
+	float Azimuth;
+
+	/** Absolute azimuth of the active sound relative to the listener. Used for 3d audio calculations. */
+	float AbsoluteAzimuth;
+
+	/** Value used to allow smooth interpolation in/out of focus */
+	float FocusFactor;
+
+	/** Cached calculation of the amount distance is scaled due to focus */
+	float DistanceScale;
+
+	/** The amount priority is scaled due to focus */
+	float PriorityScale;
+
+	/** The amount volume is scaled due to focus */
+	float VolumeScale;
+
+	FAttenuationFocusData()
+		: Azimuth(0.0f)
+		, AbsoluteAzimuth(0.0f)
+		, FocusFactor(1.0f)
+		, DistanceScale(1.0f)
+		, PriorityScale(1.0f)
+		, VolumeScale(1.0f)
+	{
+	}
+};
+
 /**
  *	Struct used for gathering the final parameters to apply to a wave instance
  */
@@ -52,6 +86,7 @@ struct FSoundParseParameters
 	// A volume scale on the sound specified by user
 	float VolumeMultiplier;
 
+
 	// Attack time of the source envelope follower
 	int32 EnvelopeFollowerAttackTime;
 
@@ -448,12 +483,6 @@ public:
 	/** The product of the component priority and the USoundBase priority */
 	float Priority;
 
-	/** The amount priority is scaled due to focus */
-	float FocusPriorityScale;
-
-	/** The amount distance is scaled due to focus */
-	float FocusDistanceScale;
-
 	/** The volume used to determine concurrency resolution for "quietest" active sound.
 	// If negative, tracking is disabled for lifetime of ActiveSound */
 	float VolumeConcurrency;
@@ -469,11 +498,10 @@ public:
 
 	FTransform Transform;
 
-	/** Azimuth of the active sound relative to the listener. Used by sound focus. */
-	float Azimuth;
-
-	/** Absolute azimuth of the active sound relative to the listener. Used for 3d audio calculations. */
-	float AbsoluteAzimuth;
+	/**
+	 * Cached data pertaining to focus system updated each frame
+	 */
+	FAttenuationFocusData FocusData;
 
 	/** Location last time playback was updated */
 	FVector LastLocation;
@@ -612,10 +640,10 @@ public:
 	static void OcclusionTraceDone(const FTraceHandle& TraceHandle, FTraceDatum& TraceDatum);
 
 	/** Applies the active sound's attenuation settings to the input parse params using the given listener */
-	void ApplyAttenuation(FSoundParseParameters& ParseParams, const FListener& Listener, const FSoundAttenuationSettings* SettingsAttenuationNode = nullptr);
+	void ParseAttenuation(FSoundParseParameters& ParseParams, const FListener& Listener, const FSoundAttenuationSettings& AttenuationSettings);
 
 	/** Returns the effective priority of the active sound */
-	float GetPriority() const { return Priority * FocusPriorityScale; }
+	float GetPriority() const { return Priority * FocusData.PriorityScale; }
 
 	/** Sets the amount of audio from this active sound to send to the submix. */
 	void SetSubmixSend(const FSoundSubmixSendInfo& SubmixSendInfo);
@@ -623,6 +651,12 @@ public:
 	/** Sets the amount of audio from this active sound to send to the source bus. */
 	void SetSourceBusSend(EBusSendType BusSendTyoe, const FSoundSourceBusSendInfo& SourceBusSendInfo);
 
+	/** Updates the active sound's attenuation settings to the input parse params using the given listener */
+	void UpdateAttenuation(float DeltaTime, FSoundParseParameters& ParseParams, const FListener& Listener, const FSoundAttenuationSettings* SettingsAttenuationNode = nullptr);
+
+	/** Updates the provided focus data using the local */
+	void UpdateFocusData(float DeltaTime, const FAttenuationListenerData& ListenerData, FAttenuationFocusData* OutFocusData = nullptr);
+
 private:
 
 	struct FAsyncTraceDetails
@@ -672,7 +706,4 @@ private:
 
 	/** Helper function which retrieves attenuation frequency value for HPF and LPF distance-based filtering. */
 	float GetAttenuationFrequency(const FSoundAttenuationSettings* InSettings, const FAttenuationListenerData& ListenerData, const FVector2D& FrequencyRange, const FRuntimeFloatCurve& CustomCurve);
-
-	/** Internal Focus Factor value used to allow smooth interpolation in/out of Focus */
-	float InternalFocusFactor;
 };
diff --git a/Engine/Source/Runtime/Engine/Public/AudioDevice.h b/Engine/Source/Runtime/Engine/Public/AudioDevice.h
--- a/Engine/Source/Runtime/Engine/Public/AudioDevice.h
+++ b/Engine/Source/Runtime/Engine/Public/AudioDevice.h
@@ -26,6 +26,7 @@
  */
 
 class FArchive;
+class FAudioDevice;
 class FAudioEffectsManager;
 class FCanvas;
 class FOutputDevice;
@@ -49,6 +50,7 @@ class USoundWave;
 class UWorld;
 
 struct FActiveSound;
+struct FAttenuationFocusData;
 struct FAudioComponentParam;
 struct FAudioQualitySettings;
 struct FRotator;
@@ -268,7 +270,6 @@ struct FActivatedReverb
 struct FAttenuationListenerData
 {
 	FVector ListenerToSoundDir;
-	FTransform ListenerTransform;
 	float AttenuationDistance;
 	float ListenerToSoundDistance;
 
@@ -276,30 +277,24 @@ struct FAttenuationListenerData
 	// Non-attenuation distance for calculating surround sound speaker maps for sources w/ spread
 	float ListenerToSoundDistanceForPanning;
 
-	bool bDataComputed;
+	FTransform ListenerTransform;
+	const FTransform SoundTransform;
+	const FSoundAttenuationSettings* AttenuationSettings;
+
+	/** Computes and returns some geometry related to the listener and the given sound transform. */
+	static FAttenuationListenerData Create(const FAudioDevice& AudioDevice, const FTransform& InListenerTransform, const FTransform& InSoundTransform, const FSoundAttenuationSettings& InAttenuationSettings);
 
-	FAttenuationListenerData()
+private:
+	FAttenuationListenerData(const FTransform& InListenerTransform, const FTransform& InSoundTransform, const FSoundAttenuationSettings& InAttenuationSettings)
 		: ListenerToSoundDir(FVector::ZeroVector)
 		, AttenuationDistance(0.0f)
 		, ListenerToSoundDistance(0.0f)
 		, ListenerToSoundDistanceForPanning(0.0f)
-		, bDataComputed(false)
-	{}
-};
-
-struct FAttenuationFocusData
-{
-	float FocusFactor;
-	float DistanceScale;
-	float PriorityScale;
-	float VolumeScale;
-
-	FAttenuationFocusData()
-		: FocusFactor(1.0f)
-		, DistanceScale(1.0f)
-		, PriorityScale(1.0f)
-		, VolumeScale(1.0f)
-	{}
+		, ListenerTransform(InListenerTransform)
+		, SoundTransform(InSoundTransform)
+		, AttenuationSettings(&InAttenuationSettings)
+	{
+	}
 };
 
 /*
@@ -325,7 +320,8 @@ struct FGlobalFocusSettings
 		, NonFocusVolumeScale(1.0f)
 		, FocusPriorityScale(1.0f)
 		, NonFocusPriorityScale(1.0f)
-	{}
+	{
+	}
 };
 
 /** Interface to register a device changed listener to respond to audio device changes. */
@@ -505,8 +501,8 @@ public:
 	 */
 	void Flush(UWorld* WorldToFlush, bool bClearActivatedReverb = true);
 
-	/** 
-	 * Allows audio rendering command queue to flush during audio device flush. 
+	/**
+	 * Allows audio rendering command queue to flush during audio device flush.
 	 * @param bPumpSynchronously must be called in situations where the audio render thread is not being called.
 	 */
 	virtual void FlushAudioRenderingCommands(bool bPumpSynchronously = false) {}
@@ -988,11 +984,8 @@ public:
 
 	void SetDeviceMuted(bool bMuted);
 
-	/** Computes and returns some geometry related to the listener and the given sound transform. */
-	void GetAttenuationListenerData(FAttenuationListenerData& OutListenerData, const FTransform& SoundTransform, const FSoundAttenuationSettings& AttenuationSettings, const FTransform* InListenerTransform = nullptr) const;
-
 	/** Returns the azimuth angle of the sound relative to the sound's nearest listener. Used for 3d audio calculations. */
-	void GetAzimuth(FAttenuationListenerData& OutListenerData, const FTransform& SoundTransform, const FSoundAttenuationSettings& AttenuationSettings, const FTransform& ListenerTransform, float& OutAzimuth, float& AbsoluteAzimuth) const;
+	void GetAzimuth(const FAttenuationListenerData& OutListenerData, float& OutAzimuth, float& AbsoluteAzimuth) const;
 
 	/** Returns the focus factor of a sound based on its position and listener data. */
 	float GetFocusFactor(const float Azimuth, const FSoundAttenuationSettings& AttenuationSettings) const;
@@ -1474,6 +1467,18 @@ public:
 	/** Returns the game's delta time */
 	float GetGameDeltaTime() const;
 
+	/** Whether device is using listener attenuation override or not. */
+	bool IsUsingListenerAttenuationOverride() const
+	{
+		return bUseListenerAttenuationOverride;
+	}
+
+	/** Returns the listener attenuation override */
+	const FVector& GetListenerAttenuationOverride() const
+	{
+		return ListenerAttenuationOverride;
+	}
+
 	void UpdateVirtualLoops(bool bForceUpdate);
 
 	/** Sets the update delta time for the audio frame */
diff --git a/Engine/Source/Runtime/Engine/Public/AudioVirtualLoop.h b/Engine/Source/Runtime/Engine/Public/AudioVirtualLoop.h
--- a/Engine/Source/Runtime/Engine/Public/AudioVirtualLoop.h
+++ b/Engine/Source/Runtime/Engine/Public/AudioVirtualLoop.h
@@ -43,7 +43,7 @@ public:
 	static bool IsEnabled();
 
 	/**
-	 * Returns the internally-managed active sound.
+	 * Returns the internally-managed active sound
 	 */
 	FActiveSound& GetActiveSound();
 
@@ -53,15 +53,20 @@ public:
 	float GetUpdateInterval() const { return UpdateInterval; }
 
 	/**
-	 * Returns the internally-managed active sound.
 	 */
 	const FActiveSound& GetActiveSound() const;
 
 	/**
-	 * Overrides the update interval to the provided length.
+	 * Overrides the update interval to the provided length
 	 */
 	void CalculateUpdateInterval();
 
+	/**
+	 * Takes aggregate update delta and updates focus so that realization
+	 * check can test if ready to play.
+	 */
+	void UpdateFocusData(float DeltaTime);
+
 	/**
 	  * Updates the loop and checks if ready to play (or 'realize').
 	  * Returns whether or not the sound is ready to be realized.
