Fix issues with -nomcp causing crashes in PIE with dedicated server.

diff --git a/Engine/Source/Editor/UnrealEd/Private/UnrealEdEngine.cpp b/Engine/Source/Editor/UnrealEd/Private/UnrealEdEngine.cpp
--- a/Engine/Source/Editor/UnrealEd/Private/UnrealEdEngine.cpp
+++ b/Engine/Source/Editor/UnrealEd/Private/UnrealEdEngine.cpp
@@ -188,6 +188,12 @@ void UUnrealEdEngine::Init(IEngineLoop* InEngineLoop)
 		}
 	}
 
+	if (FParse::Param(FCommandLine::Get(), TEXT("nomcp")))
+	{
+		// If our editor has nomcp, pass it through to any subprocesses.
+		FCommandLine::AddToSubprocessCommandline(TEXT(" -nomcp"));
+	}
+
 	bPivotMovedIndependently = false;
 }
 
diff --git a/Engine/Source/Runtime/Engine/Private/UnrealEngine.cpp b/Engine/Source/Runtime/Engine/Private/UnrealEngine.cpp
--- a/Engine/Source/Runtime/Engine/Private/UnrealEngine.cpp
+++ b/Engine/Source/Runtime/Engine/Private/UnrealEngine.cpp
@@ -12447,7 +12447,7 @@ bool UEngine::LoadMap( FWorldContext& WorldContext, FURL URL, class UPendingNetG
 				GIsPlayInEditorWorld = true;
 			}
 			// Otherwise we are probably loading new map while in PIE, so we need to rename world package and all streaming levels
-			else if ((Pending == nullptr) || (Pending->DemoNetDriver != nullptr))
+			else if (WorldContext.PIEInstance != -1 && ((Pending == nullptr) || (Pending->DemoNetDriver != nullptr)))
 			{
 				NewWorld->RenameToPIEWorld(WorldContext.PIEInstance);
 			}
