Fix compiler errors in non-editor builds/warning in editor build

diff --git a/Engine/Source/Runtime/Engine/Private/ActiveSound.cpp b/Engine/Source/Runtime/Engine/Private/ActiveSound.cpp
--- a/Engine/Source/Runtime/Engine/Private/ActiveSound.cpp
+++ b/Engine/Source/Runtime/Engine/Private/ActiveSound.cpp
@@ -1479,9 +1479,9 @@ void FActiveSound::UpdateFocusData(float DeltaTime, const FAttenuationListenerDa
 	FocusDataToUpdate->VolumeScale = ListenerData.AttenuationSettings->GetFocusAttenuation(FocusSettings, FocusDataToUpdate->FocusFactor);
 }
 
-void FActiveSound::ParseAttenuation(FSoundParseParameters& ParseParams, const FListener& Listener, const FSoundAttenuationSettings& AttenuationSettings)
+void FActiveSound::ParseAttenuation(FSoundParseParameters& OutParseParams, const FListener& InListener, const FSoundAttenuationSettings& InAttenuationSettings)
 {
-	UpdateAttenuation(0.0f, ParseParams, Listener, &AttenuationSettings);
+	UpdateAttenuation(0.0f, OutParseParams, InListener, &InAttenuationSettings);
 }
 
 void FActiveSound::UpdateAttenuation(float DeltaTime, FSoundParseParameters& ParseParams, const FListener& Listener, const FSoundAttenuationSettings* SettingsAttenuationNode)
diff --git a/Engine/Source/Runtime/Engine/Public/ActiveSound.h b/Engine/Source/Runtime/Engine/Public/ActiveSound.h
--- a/Engine/Source/Runtime/Engine/Public/ActiveSound.h
+++ b/Engine/Source/Runtime/Engine/Public/ActiveSound.h
@@ -640,7 +640,7 @@ public:
 	static void OcclusionTraceDone(const FTraceHandle& TraceHandle, FTraceDatum& TraceDatum);
 
 	/** Applies the active sound's attenuation settings to the input parse params using the given listener */
-	void ParseAttenuation(FSoundParseParameters& ParseParams, const FListener& Listener, const FSoundAttenuationSettings& AttenuationSettings);
+	void ParseAttenuation(FSoundParseParameters& OutParseParams, const FListener& InListener, const FSoundAttenuationSettings& InAttenuationSettings);
 
 	/** Returns the effective priority of the active sound */
 	float GetPriority() const { return Priority * FocusData.PriorityScale; }
