Fix asserts in demo net driver due to adding conditional scope checks reliant on a cvar

diff --git a/Engine/Source/Runtime/Engine/Classes/Engine/NetSerialization.h b/Engine/Source/Runtime/Engine/Classes/Engine/NetSerialization.h
--- a/Engine/Source/Runtime/Engine/Classes/Engine/NetSerialization.h
+++ b/Engine/Source/Runtime/Engine/Classes/Engine/NetSerialization.h
@@ -27,7 +27,7 @@ DECLARE_CYCLE_STAT_EXTERN(TEXT("NetSerializeFast Array"), STAT_NetSerializeFastA
 DECLARE_CYCLE_STAT_EXTERN(TEXT("NetSerializeFast Array BuildMap"), STAT_NetSerializeFastArray_BuildMap, STATGROUP_ServerCPU, ENGINE_API);
 DECLARE_CYCLE_STAT_EXTERN(TEXT("NetSerializeFast Array Delta Struct"), STAT_NetSerializeFastArray_DeltaStruct, STATGROUP_ServerCPU, ENGINE_API);
 
-extern ENGINE_API TAutoConsoleVariable<int32> CVarEnableDetailedScopeCounters;
+extern ENGINE_API TAutoConsoleVariable<int32> CVarNetEnableDetailedScopeCounters;
 
 /**
  *	===================== NetSerialize and NetDeltaSerialize customization. =====================
@@ -1072,7 +1072,7 @@ bool FFastArraySerializer::FastArrayDeltaSerialize(TArray<Type> &Items, FNetDelt
 		return FastArrayDeltaSerialize_DeltaSerializeStructs(Items, Parms, ArraySerializer);
 	}
 
-	CONDITIONAL_SCOPE_CYCLE_COUNTER(STAT_NetSerializeFastArray, CVarEnableDetailedScopeCounters.GetValueOnGameThread() > 0);
+	CONDITIONAL_SCOPE_CYCLE_COUNTER(STAT_NetSerializeFastArray, CVarNetEnableDetailedScopeCounters.GetValueOnAnyThread() > 0);
 	class UScriptStruct* InnerStruct = Type::StaticStruct();
 
 	UE_LOG(LogNetFastTArray, Log, TEXT("FastArrayDeltaSerialize for %s. %s. %s"), *InnerStruct->GetName(), *InnerStruct->GetOwnerStruct()->GetName(), Parms.Reader ? TEXT("Reading") : TEXT("Writing"));
@@ -1500,7 +1500,7 @@ bool FFastArraySerializer::FastArrayDeltaSerialize_DeltaSerializeStructs(TArray<
 		}
 	};
 
-	CONDITIONAL_SCOPE_CYCLE_COUNTER(STAT_NetSerializeFastArray_DeltaStruct, CVarEnableDetailedScopeCounters.GetValueOnGameThread() > 0);
+	CONDITIONAL_SCOPE_CYCLE_COUNTER(STAT_NetSerializeFastArray_DeltaStruct, CVarNetEnableDetailedScopeCounters.GetValueOnAnyThread() > 0);
 
 	class UScriptStruct* InnerStruct = Type::StaticStruct();
 
diff --git a/Engine/Source/Runtime/Engine/Private/DataChannel.cpp b/Engine/Source/Runtime/Engine/Private/DataChannel.cpp
--- a/Engine/Source/Runtime/Engine/Private/DataChannel.cpp
+++ b/Engine/Source/Runtime/Engine/Private/DataChannel.cpp
@@ -73,7 +73,7 @@ FAutoConsoleVariableRef CVarNetSkipReplicatorForDestructionInfos(
 	TEXT("If enabled, skip creation of object replicator in SetChannelActor when we know there is no content payload and we're going to immediately destroy the actor."));
 
 extern TAutoConsoleVariable<int32> CVarFilterGuidRemapping;
-extern TAutoConsoleVariable<int32> CVarEnableDetailedScopeCounters;
+extern TAutoConsoleVariable<int32> CVarNetEnableDetailedScopeCounters;
 
 /*-----------------------------------------------------------------------------
 	UChannel implementation.
@@ -3762,7 +3762,7 @@ FObjectReplicator & UActorChannel::GetActorReplicationData()
 
 TSharedRef< FObjectReplicator > & UActorChannel::FindOrCreateReplicator( UObject* Obj, bool* bOutCreated )
 {
-	CONDITIONAL_SCOPE_CYCLE_COUNTER(Stat_ActorChanFindOrCreateRep, CVarEnableDetailedScopeCounters.GetValueOnGameThread() > 0);
+	CONDITIONAL_SCOPE_CYCLE_COUNTER(Stat_ActorChanFindOrCreateRep, CVarNetEnableDetailedScopeCounters.GetValueOnAnyThread() > 0);
 	SCOPE_CYCLE_UOBJECT(ActorChannelFindOrCreateRep, Obj);
 
 	// First, try to find it on the channel replication map
diff --git a/Engine/Source/Runtime/Engine/Private/DataReplication.cpp b/Engine/Source/Runtime/Engine/Private/DataReplication.cpp
--- a/Engine/Source/Runtime/Engine/Private/DataReplication.cpp
+++ b/Engine/Source/Runtime/Engine/Private/DataReplication.cpp
@@ -49,7 +49,7 @@ static FAutoConsoleVariableRef CVarSupportsFastArrayDelta(
 	TEXT("Whether or not Fast Array Struct Delta Serialization is enabled.")
 );
 
-extern TAutoConsoleVariable<int32> CVarEnableDetailedScopeCounters;
+extern TAutoConsoleVariable<int32> CVarNetEnableDetailedScopeCounters;
 
 PRAGMA_DISABLE_DEPRECATION_WARNINGS
 class FNetSerializeCB : public INetSerializeCB
@@ -309,7 +309,7 @@ bool FObjectReplicator::SerializeCustomDeltaProperty(
 {
 	check( NewFullState.IsValid() == false ); // NewState is passed in as nullptr and instantiated within this function if necessary
 
-	CONDITIONAL_SCOPE_CYCLE_COUNTER( STAT_NetSerializeItemDeltaTime, CVarEnableDetailedScopeCounters.GetValueOnGameThread() > 0 );
+	CONDITIONAL_SCOPE_CYCLE_COUNTER( STAT_NetSerializeItemDeltaTime, CVarNetEnableDetailedScopeCounters.GetValueOnAnyThread() > 0 );
 
 	UStructProperty * StructProperty = CastChecked< UStructProperty >( Property );
 
@@ -360,7 +360,7 @@ bool FObjectReplicator::SendCustomDeltaProperty(UObject* InObject, uint16 Custom
 	check(!NewFullState.IsValid()); // NewState is passed in as nullptr and instantiated within this function if necessary
 	check(RepLayout);
 
-	CONDITIONAL_SCOPE_CYCLE_COUNTER(STAT_NetSerializeItemDeltaTime, CVarEnableDetailedScopeCounters.GetValueOnGameThread() > 0);
+	CONDITIONAL_SCOPE_CYCLE_COUNTER(STAT_NetSerializeItemDeltaTime, CVarNetEnableDetailedScopeCounters.GetValueOnAnyThread() > 0);
 
 	UNetDriver* const ConnectionDriver = Connection->GetDriver();
 	FNetSerializeCB NetSerializeCB(ConnectionDriver);
@@ -1420,7 +1420,7 @@ static FORCEINLINE FPropertyRetirement** UpdateAckedRetirements(
 
 void FObjectReplicator::ReplicateCustomDeltaProperties( FNetBitWriter & Bunch, FReplicationFlags RepFlags )
 {
-	CONDITIONAL_SCOPE_CYCLE_COUNTER(STAT_NetReplicateCustomDeltaPropTime, CVarEnableDetailedScopeCounters.GetValueOnGameThread() > 0);
+	CONDITIONAL_SCOPE_CYCLE_COUNTER(STAT_NetReplicateCustomDeltaPropTime, CVarNetEnableDetailedScopeCounters.GetValueOnAnyThread() > 0);
 
 	check(RepLayout);
 	const FRepLayout& LocalRepLayout = *RepLayout;
diff --git a/Engine/Source/Runtime/Engine/Private/NetConnection.cpp b/Engine/Source/Runtime/Engine/Private/NetConnection.cpp
--- a/Engine/Source/Runtime/Engine/Private/NetConnection.cpp
+++ b/Engine/Source/Runtime/Engine/Private/NetConnection.cpp
@@ -60,7 +60,7 @@ static TAutoConsoleVariable<int32> CVarNetPacketOrderMaxMissingPackets(TEXT("net
 
 static TAutoConsoleVariable<int32> CVarNetPacketOrderMaxCachedPackets(TEXT("net.PacketOrderMaxCachedPackets"), 32, TEXT("(NOTE: Must be power of 2!) The maximum number of packets to cache while waiting for missing packet sequences, before treating missing packets as lost."));
 
-TAutoConsoleVariable<int32> CVarEnableDetailedScopeCounters(TEXT("net.EnableDetailedScopeCounters"), 1, TEXT("Enables detailed networking scope cycle counters. There are often lots of these which can negatively impact performance."));
+TAutoConsoleVariable<int32> CVarNetEnableDetailedScopeCounters(TEXT("net.EnableDetailedScopeCounters"), 1, TEXT("Enables detailed networking scope cycle counters. There are often lots of these which can negatively impact performance."));
 
 extern int32 GNetDormancyValidate;
 
diff --git a/Engine/Source/Runtime/Engine/Private/RepLayout.cpp b/Engine/Source/Runtime/Engine/Private/RepLayout.cpp
--- a/Engine/Source/Runtime/Engine/Private/RepLayout.cpp
+++ b/Engine/Source/Runtime/Engine/Private/RepLayout.cpp
@@ -66,7 +66,7 @@ int32 MaxRepArrayMemory = UNetworkSettings::DefaultMaxRepArrayMemory;
 extern int32 GNumSharedSerializationHit;
 extern int32 GNumSharedSerializationMiss;
 
-extern TAutoConsoleVariable<int32> CVarEnableDetailedScopeCounters;
+extern TAutoConsoleVariable<int32> CVarNetEnableDetailedScopeCounters;
 
 FConsoleVariableSinkHandle CreateMaxArraySizeCVarAndRegisterSink()
 {
@@ -1153,7 +1153,7 @@ bool FRepLayout::CompareProperties(
 	const FConstRepObjectDataBuffer Data,
 	const FReplicationFlags& RepFlags) const
 {
-	CONDITIONAL_SCOPE_CYCLE_COUNTER(STAT_NetReplicateDynamicPropCompareTime, CVarEnableDetailedScopeCounters.GetValueOnGameThread() > 0);
+	CONDITIONAL_SCOPE_CYCLE_COUNTER(STAT_NetReplicateDynamicPropCompareTime, CVarNetEnableDetailedScopeCounters.GetValueOnAnyThread() > 0);
 
 	if (LayoutState == ERepLayoutState::Empty)
 	{
@@ -1255,7 +1255,7 @@ bool FRepLayout::ReplicateProperties(
 	FNetBitWriter& Writer,
 	const FReplicationFlags& RepFlags) const
 {
-	CONDITIONAL_SCOPE_CYCLE_COUNTER(STAT_NetReplicateDynamicPropTime, CVarEnableDetailedScopeCounters.GetValueOnGameThread() > 0);
+	CONDITIONAL_SCOPE_CYCLE_COUNTER(STAT_NetReplicateDynamicPropTime, CVarNetEnableDetailedScopeCounters.GetValueOnAnyThread() > 0);
 
 	check(ObjectClass == Owner);
 
@@ -6164,7 +6164,7 @@ bool FRepLayout::DeltaSerializeFastArrayProperty(FFastArrayDeltaSerializeParams&
 {
 	using namespace UE4_RepLayout_Private;
 
-	CONDITIONAL_SCOPE_CYCLE_COUNTER(STAT_RepLayout_DeltaSerializeFastArray, CVarEnableDetailedScopeCounters.GetValueOnGameThread() > 0);
+	CONDITIONAL_SCOPE_CYCLE_COUNTER(STAT_RepLayout_DeltaSerializeFastArray, CVarNetEnableDetailedScopeCounters.GetValueOnAnyThread() > 0);
 
 	// A portion of this work could be shared across all Fast Array Properties for a given object,
 	// but that would be easier to do if the Custom Delta Serialization was completely encapsulated in FRepLayout.
