Added URuntimeOptionsBase which serves as a useful base class for groups of config variables that you want to be developer friendly
- In non-Shipping builds, each property will be exposed both as a console variable and as a command-line argument for easy testing during development

diff --git a/Engine/Source/Runtime/Engine/Classes/Engine/RuntimeOptionsBase.h b/Engine/Source/Runtime/Engine/Classes/Engine/RuntimeOptionsBase.h
new file mode 100644
--- /dev/null
+++ b/Engine/Source/Runtime/Engine/Classes/Engine/RuntimeOptionsBase.h
@@ -0,0 +1,67 @@
+// Copyright 1998-2019 Epic Games, Inc. All Rights Reserved.
+
+#pragma once
+
+#include "CoreMinimal.h"
+#include "RuntimeOptionsBase.generated.h"
+
+/**
+ *
+ * URuntimeOptionsBase: Base class designed to be subclassed in your game.
+ *
+ * Supports checking at runtime whether features are enabled/disabled, and changing
+ * configuration parameters via console cheats or startup commands.
+ *
+ * Add new config properties to your subclasses which default to the desired normal state
+ * This can be adjusted via the development-only tools (command line or cvar) or via an
+ * override in the config hierarchy to adjust it as needed (e.g., via a hotfix).
+ *
+ * In non-Shipping builds, each property will be exposed both as a console variable and as a
+ * command-line argument for easy testing during development.
+ *
+ * Debug console syntax (disabled in Shipping configurations):
+ *   prefix.PropertyName Value
+ * Command line syntax (disabled in Shipping configurations):
+ *   -prefix.PropertyName=Value
+ * DefaultRuntimeOptions.ini syntax (note that there is no prefix for these):
+ *   [/Script/YourModule.YourRuntimeOptionsSubclass]
+ *   PropertyName=Value
+ *
+ * Where the prefix is set by the value of OptionCommandPrefix (defaults to "ro" but can be overridden)
+ *
+ * You can also change the name of the ini file that settings are gathered from in your derived
+ * UCLASS() declaration
+ */
+
+UCLASS(config=RuntimeOptions, BlueprintType, Abstract)
+class ENGINE_API URuntimeOptionsBase : public UObject
+{
+	GENERATED_BODY()
+
+public:
+	URuntimeOptionsBase();
+
+	//~UObject interface
+	virtual void PostInitProperties() override;
+	virtual void PostReloadConfig(UProperty* PropertyThatWasLoaded) override;
+	//~End of UObject interface
+
+	// Returns the runtime options for the specified type, typically you will make a non-templated overload of Get in your subclass
+	template<typename TRuntimeOptionsSubclass>
+	static const TRuntimeOptionsSubclass& Get()
+	{
+		return *GetDefault<TRuntimeOptionsSubclass>();
+	}
+
+protected:
+	// The command prefix for all options declared in a subclass (defaults to "ro", so a property named TestProp would be exposed as ro.TestProp)
+	FString OptionCommandPrefix;
+
+protected:
+	// Applies overrides from the command line
+	void ApplyCommandlineOverrides();
+	void RegisterSupportedConsoleVariables(bool bDuringReload);
+	void InitializeRuntimeOptions();
+
+	void ApplyPluginOverridesIfPresent(const FString& PluginName);
+};
diff --git a/Engine/Source/Runtime/Engine/Private/RuntimeOptionsBase.cpp b/Engine/Source/Runtime/Engine/Private/RuntimeOptionsBase.cpp
new file mode 100644
--- /dev/null
+++ b/Engine/Source/Runtime/Engine/Private/RuntimeOptionsBase.cpp
@@ -0,0 +1,180 @@
+// Copyright 1998-2019 Epic Games, Inc. All Rights Reserved.
+
+#include "Engine/RuntimeOptionsBase.h"
+#include "UObject/PropertyPortFlags.h"
+#include "Misc/Paths.h"
+#include "Misc/FileHelper.h"
+#include "Misc/ConfigCacheIni.h"
+#include "Interfaces/IPluginManager.h"
+
+DEFINE_LOG_CATEGORY_STATIC(LogRuntimeOptionsBase, Log, All);
+
+#define UE_RUNTIMEOPTIONSBASE_SUPPORT_CVARS !UE_BUILD_SHIPPING
+#define UE_RUNTIMEOPTIONSBASE_SUPPORT_COMMANDLINE !UE_BUILD_SHIPPING
+
+URuntimeOptionsBase::URuntimeOptionsBase()
+{
+	OptionCommandPrefix = TEXT("ro");
+}
+
+void URuntimeOptionsBase::InitializeRuntimeOptions()
+{
+	check(HasAnyFlags(RF_ClassDefaultObject));
+	check(!OptionCommandPrefix.IsEmpty());
+
+	if (!GetClass()->HasAnyClassFlags(CLASS_Abstract))
+	{
+		// Verify that the user hasn't created a hierarchy of non-abstract versions (this won't work with, e.g., cvar setting as only one of them could be registered)
+		int32 NumNonAbstractVersions = 0;
+		for (UClass* CurrentClass = GetClass(); CurrentClass != URuntimeOptionsBase::StaticClass(); CurrentClass = CurrentClass->GetSuperClass())
+		{
+			if (!CurrentClass->HasAnyClassFlags(CLASS_Abstract))
+			{
+				++NumNonAbstractVersions;
+			}
+		}
+		ensureAlwaysMsgf(NumNonAbstractVersions == 1, TEXT("Error in %s hierarchy; properties should only be introduced in leaf non-abstract subclasses of URuntimeOptionsBase"), *GetClass()->GetName());
+
+		// Load the .ini file if it hasn't already been loaded
+		FString RuntimeOptionsIniName;
+		FConfigCacheIni::LoadGlobalIniFile(/*out*/ RuntimeOptionsIniName, *GetClass()->ClassConfigName.ToString(), nullptr);
+	
+		// Apply command line overrides and expose properties as console variables
+		ApplyCommandlineOverrides();
+		RegisterSupportedConsoleVariables(false);
+	}
+}
+
+void URuntimeOptionsBase::ApplyPluginOverridesIfPresent(const FString& PluginName)
+{
+	check(HasAnyFlags(RF_ClassDefaultObject));
+
+	TSharedPtr<IPlugin> Plugin = IPluginManager::Get().FindPlugin(*PluginName);
+	if (Plugin.IsValid() && Plugin->IsEnabled())
+	{
+		const FString PluginConfigDir = Plugin->GetBaseDir() / TEXT("Config/");
+		TArray<FString> PluginConfigs;
+		IFileManager::Get().FindFiles(PluginConfigs, *PluginConfigDir, TEXT("ini"));
+		const FString PlaformName = FPlatformProperties::PlatformName();
+
+		for (const FString& ConfigFile : PluginConfigs)
+		{
+			const FString PluginConfigFilename = FString::Printf(TEXT("%s%s/%s.ini"), *FPaths::GeneratedConfigDir(), *PlaformName, *FPaths::GetBaseFilename(ConfigFile));
+			if (FConfigFile* FoundConfig = GConfig->Find(PluginConfigFilename, false))
+			{
+				FString PluginConfigContent;
+				if (FFileHelper::LoadFileToString(PluginConfigContent, *FPaths::Combine(PluginConfigDir, ConfigFile)))
+				{
+					FoundConfig->CombineFromBuffer(PluginConfigContent);
+					// if plugin config overrides are applied then don't save
+					FoundConfig->NoSave = true;
+
+					FConfigFile::OverrideFromCommandline(FoundConfig, PluginConfigFilename);
+				}
+			}
+		}
+	}
+}
+
+void URuntimeOptionsBase::ApplyCommandlineOverrides()
+{
+#if UE_RUNTIMEOPTIONSBASE_SUPPORT_COMMANDLINE
+	// In non-shipping builds check the commandline for overrides and register variables with the console
+	for (const UProperty* Property : TFieldRange<const UProperty>(GetClass()))
+	{
+		const FString FullyQualifiedName = FString::Printf(TEXT("%s.%s"), *OptionCommandPrefix, *Property->GetName());
+
+		uint8* DataPtr = Property->ContainerPtrToValuePtr<uint8>(this, 0);
+
+		FString CommandLineOverride;
+		if (FParse::Value(FCommandLine::Get(), *FString::Printf(TEXT("%s="), *FullyQualifiedName), /*out*/ CommandLineOverride))
+		{
+			Property->ImportText(*CommandLineOverride, DataPtr, PPF_None, this);
+		}
+	}
+#endif
+}
+
+void URuntimeOptionsBase::RegisterSupportedConsoleVariables(bool bDuringReload)
+{
+#if UE_RUNTIMEOPTIONSBASE_SUPPORT_CVARS
+	for (const UProperty* Property : TFieldRange<const UProperty>(GetClass()))
+	{
+		const FString FullyQualifiedName = FString::Printf(TEXT("%s.%s"), *OptionCommandPrefix, *Property->GetName());
+
+		uint8* DataPtr = Property->ContainerPtrToValuePtr<uint8>(this, 0);
+
+		FString DisplayName;
+#if WITH_EDITOR
+		DisplayName = Property->GetMetaData("DisplayName");
+#endif
+
+		if (bDuringReload)
+		{
+			IConsoleManager::Get().UnregisterConsoleObject(*FullyQualifiedName, /*bKeepState=*/ false);
+		}
+		
+		if (IConsoleObject* ExistingObj = IConsoleManager::Get().FindConsoleObject(*FullyQualifiedName))
+		{
+			UE_LOG(LogRuntimeOptionsBase, Error, TEXT("A conflicting CVar '%s' already exists when trying to register a runtime option for the property %s%s::%s"),
+				*FullyQualifiedName,
+				GetClass()->GetPrefixCPP(),
+				*GetClass()->GetName(),
+				*Property->GetName());
+		}
+		else
+		{
+			if (const UNumericProperty* NumericProperty = Cast<const UNumericProperty>(Property))
+			{
+				if (NumericProperty->IsFloatingPoint())
+				{
+					IConsoleManager::Get().RegisterConsoleVariableRef(*FullyQualifiedName, *(float*)DataPtr, *DisplayName, ECVF_Default);
+				}
+				else if (NumericProperty->IsInteger())
+				{
+					IConsoleManager::Get().RegisterConsoleVariableRef(*FullyQualifiedName, *(int32*)DataPtr, *DisplayName, ECVF_Default);
+				}
+			}
+			else if (const UBoolProperty* BoolProperty = Cast<UBoolProperty>(Property))
+			{
+				IConsoleManager::Get().RegisterConsoleVariableRef(*FullyQualifiedName, *(bool*)DataPtr, *DisplayName, ECVF_Default);
+			}
+			else if (const UStrProperty* StringProperty = Cast<UStrProperty>(Property))
+			{
+				IConsoleManager::Get().RegisterConsoleVariableRef(*FullyQualifiedName, *(FString*)DataPtr, *DisplayName, ECVF_Default);
+			}
+		}
+	}
+#endif
+}
+
+void URuntimeOptionsBase::PostInitProperties()
+{
+	Super::PostInitProperties();
+	if (HasAnyFlags(RF_ClassDefaultObject))
+	{
+		InitializeRuntimeOptions();
+	}
+}
+
+void URuntimeOptionsBase::PostReloadConfig(UProperty* PropertyThatWasLoaded)
+{
+	// Re-apply command line options
+	ApplyCommandlineOverrides();
+
+	// CVars hold an internal cache that can be misleading, so update it
+	RegisterSupportedConsoleVariables(true);
+
+	// Print out the post-reload active values for the properties
+	if (UE_LOG_ACTIVE(LogRuntimeOptionsBase, Verbose))
+	{
+		UE_LOG(LogRuntimeOptionsBase, Verbose, TEXT("After %s%s::PostReloadConfig:"), GetClass()->GetPrefixCPP(), *GetClass()->GetName());
+		for (const UProperty* Property : TFieldRange<const UProperty>(GetClass()))
+		{
+			FString	Value;
+			Property->ExportText_InContainer(0, Value, this, this, this, 0);
+
+			UE_LOG(LogRuntimeOptionsBase, Verbose, TEXT("\t%s%s::%s = %s"), GetClass()->GetPrefixCPP(), *GetClass()->GetName(), *Property->GetName(), *Value);
+		}
+	}
+}
