IWYU fixes for URuntimeOptionsBase
#rnx

diff --git a/Engine/Source/Runtime/Engine/Private/RuntimeOptionsBase.cpp b/Engine/Source/Runtime/Engine/Private/RuntimeOptionsBase.cpp
--- a/Engine/Source/Runtime/Engine/Private/RuntimeOptionsBase.cpp
+++ b/Engine/Source/Runtime/Engine/Private/RuntimeOptionsBase.cpp
@@ -6,6 +6,8 @@
 #include "Misc/FileHelper.h"
 #include "Misc/ConfigCacheIni.h"
 #include "Interfaces/IPluginManager.h"
+#include "Misc/CommandLine.h"
+#include "HAL/IConsoleManager.h"
 
 DEFINE_LOG_CATEGORY_STATIC(LogRuntimeOptionsBase, Log, All);
 
