CIS fix

diff --git a/Engine/Source/Runtime/Engine/Private/Animation/AnimSequence.cpp b/Engine/Source/Runtime/Engine/Private/Animation/AnimSequence.cpp
--- a/Engine/Source/Runtime/Engine/Private/Animation/AnimSequence.cpp
+++ b/Engine/Source/Runtime/Engine/Private/Animation/AnimSequence.cpp
@@ -1501,17 +1501,6 @@ void UAnimSequence::GetAnimationPose(FCompactPose& OutPose, FBlendedCurve& OutCu
 #endif
 }
 
-struct FRetargetTracking
-{
-	const FCompactPoseBoneIndex PoseBoneIndex;
-	const int32 SkeletonBoneIndex;
-
-	FRetargetTracking(const FCompactPoseBoneIndex InPoseBoneIndex, const int32 InSkeletonBoneIndex)
-		: PoseBoneIndex(InPoseBoneIndex), SkeletonBoneIndex(InSkeletonBoneIndex) 
-	{
-	}
-};
-
 void UAnimSequence::GetBonePose(FCompactPose& OutPose, FBlendedCurve& OutCurve, const FAnimExtractContext& ExtractionContext, bool bForceUseRawData) const
 {
 	SCOPE_CYCLE_COUNTER(STAT_AnimSeq_GetBonePose);
