Exposed DistanceFieldAtlas classes to other modules. This is required to have custom mesh proxies generate their own distance field data outside of the engine module.

diff --git a/Engine/Source/Runtime/Engine/Public/DistanceFieldAtlas.h b/Engine/Source/Runtime/Engine/Public/DistanceFieldAtlas.h
--- a/Engine/Source/Runtime/Engine/Public/DistanceFieldAtlas.h
+++ b/Engine/Source/Runtime/Engine/Public/DistanceFieldAtlas.h
@@ -23,7 +23,7 @@ class UStaticMesh;
 template <class T> class TLockFreePointerListLIFO;
 
 /** Represents a distance field volume texture for a single UStaticMesh. */
-class FDistanceFieldVolumeTexture
+class ENGINE_API FDistanceFieldVolumeTexture
 {
 public:
 	FDistanceFieldVolumeTexture(const class FDistanceFieldVolumeData& InVolumeData) :
@@ -78,7 +78,7 @@ private:
 };
 
 /** Global volume texture atlas that collects all static mesh resource distance fields. */
-class FDistanceFieldVolumeTextureAtlas : public FRenderResource
+class ENGINE_API FDistanceFieldVolumeTextureAtlas : public FRenderResource
 {
 public:
 	FDistanceFieldVolumeTextureAtlas();
@@ -94,9 +94,9 @@ public:
 	int32 GetSizeY() const { return VolumeTextureRHI->GetSizeY(); }
 	int32 GetSizeZ() const { return VolumeTextureRHI->GetSizeZ(); }
 
-	ENGINE_API FString GetSizeString() const;
+	FString GetSizeString() const;
 
-	ENGINE_API void ListMeshDistanceFields() const;
+	void ListMeshDistanceFields() const;
 
 	/** Add an allocation to the atlas. */
 	void AddAllocation(FDistanceFieldVolumeTexture* Texture);
@@ -105,7 +105,7 @@ public:
 	void RemoveAllocation(FDistanceFieldVolumeTexture* Texture);
 
 	/** Reallocates the volume texture if necessary and uploads new allocations. */
-	ENGINE_API void UpdateAllocations();
+	void UpdateAllocations();
 
 	int32 GetGeneration() const { return Generation; }
 
@@ -131,7 +131,7 @@ private:
 extern ENGINE_API TGlobalResource<FDistanceFieldVolumeTextureAtlas> GDistanceFieldVolumeTextureAtlas;
 
 /** Distance field data payload and output of the mesh build process. */
-class FDistanceFieldVolumeData : public FDeferredCleanupInterface
+class ENGINE_API FDistanceFieldVolumeData : public FDeferredCleanupInterface
 {
 public:
 
