Allow FShadowMap to be destroyed on the Garbage Collector async destryction thread

diff --git a/Engine/Source/Runtime/Engine/Public/ShadowMap.h b/Engine/Source/Runtime/Engine/Public/ShadowMap.h
--- a/Engine/Source/Runtime/Engine/Public/ShadowMap.h
+++ b/Engine/Source/Runtime/Engine/Public/ShadowMap.h
@@ -258,7 +258,7 @@ public:
 	}
 	void Release()
 	{
-		check(IsInGameThread() || IsInAsyncLoadingThread());
+		check(IsInGameThread() || IsInAsyncLoadingThread() || IsInGarbageCollectorThread());
 		checkSlow(NumRefs > 0);
 		if (--NumRefs == 0)
 		{
