Fixed RendererSettings edit condition.

diff --git a/Engine/Source/Runtime/Engine/Classes/Engine/RendererSettings.h b/Engine/Source/Runtime/Engine/Classes/Engine/RendererSettings.h
--- a/Engine/Source/Runtime/Engine/Classes/Engine/RendererSettings.h
+++ b/Engine/Source/Runtime/Engine/Classes/Engine/RendererSettings.h
@@ -456,7 +456,7 @@ class ENGINE_API URendererSettings : public UDeveloperSettings
 	uint32 bDefaultFeatureLensFlare : 1;
 
 	UPROPERTY(config, EditAnywhere, Category = DefaultSettings, meta = (
-		EditCondition = "DefaultFeatureAntiAliasing == AAM_TemporalAA",
+		EditCondition = "DefaultFeatureAntiAliasing == EAntiAliasingMethod::AAM_TemporalAA",
 		ConsoleVariable = "r.TemporalAA.Upsampling", DisplayName = "Temporal Upsampling",
 		ToolTip = "Whether to do primary screen percentage with temporal AA or not."))
 	uint32 bTemporalUpsampling : 1;
