Fix phsyics replication error.

diff --git a/Engine/Source/Runtime/Engine/Private/PhysicsEngine/PhysicsReplication.cpp b/Engine/Source/Runtime/Engine/Private/PhysicsEngine/PhysicsReplication.cpp
--- a/Engine/Source/Runtime/Engine/Private/PhysicsEngine/PhysicsReplication.cpp
+++ b/Engine/Source/Runtime/Engine/Private/PhysicsEngine/PhysicsReplication.cpp
@@ -191,7 +191,7 @@ bool FPhysicsReplication::ApplyRigidBodyState(float DeltaSeconds, FBodyInstance*
 	float AngDiff;
 	const FQuat DeltaQuat = InvCurrentQuat * TargetQuat;
 	DeltaQuat.ToAxisAndAngle(AngDiffAxis, AngDiff);
-	AngDiff = FMath::RadiansToDegrees(FMath::UnwindRadians(AngDiff));
+	AngDiff = FMath::Abs(FMath::RadiansToDegrees(FMath::UnwindRadians(AngDiff)));
 
 	/////// ACCUMULATE ERROR IF NOT APPROACHING SOLUTION ///////
 
