Fix update streamable anim if source anim has changed.

diff --git a/Engine/Source/Runtime/Engine/Private/Animation/AnimStreamable.cpp b/Engine/Source/Runtime/Engine/Private/Animation/AnimStreamable.cpp
--- a/Engine/Source/Runtime/Engine/Private/Animation/AnimStreamable.cpp
+++ b/Engine/Source/Runtime/Engine/Private/Animation/AnimStreamable.cpp
@@ -318,6 +318,24 @@ void UAnimStreamable::PostLoad()
 	Super::PostLoad();
 
 #if WITH_EDITOR
+	if (UAnimSequence* NonConstSeq = const_cast<UAnimSequence*>(SourceSequence))
+	{
+		if (FLinkerLoad* Linker = NonConstSeq->GetLinker())
+		{
+			Linker->Preload(NonConstSeq);
+		}
+		NonConstSeq->ConditionalPostLoad();
+	}
+
+	if (SourceSequence && (GenerateGuidFromRawAnimData(SourceSequence->GetRawAnimationData(), SourceSequence->RawCurveData) != RawDataGuid))
+	{
+		InitFrom(SourceSequence);
+	}
+	else
+	{
+		RequestCompressedData(); // Grab compressed data for current platform
+	}
+
 	RequestCompressedData(); // Grab compressed data for current platform
 #else
 	IStreamingManager::Get().GetAnimationStreamingManager().AddStreamingAnim(this); // This will be handled by RequestCompressedData in editor builds
