Remove test code that was causing static analysis warning.

diff --git a/Engine/Source/Runtime/Engine/Private/ActiveSound.cpp b/Engine/Source/Runtime/Engine/Private/ActiveSound.cpp
--- a/Engine/Source/Runtime/Engine/Private/ActiveSound.cpp
+++ b/Engine/Source/Runtime/Engine/Private/ActiveSound.cpp
@@ -1506,10 +1506,7 @@ void FActiveSound::UpdateAttenuation(float DeltaTime, FSoundParseParameters& Par
 	check(AudioDevice);
 	FAttenuationListenerData ListenerData = FAttenuationListenerData::Create(*AudioDevice, Listener.Transform, ParseParams.Transform, *Settings);
 
-	// Ignore focus calculation if attenuation is being updated due to override attenuation parameters being passed in.
-	// Focus is always only calculated by the parent attenuation's data, never by the attenuation node's settings.
-	const bool bIgnoreFocusCalculation = false;// SettingsAttenuationNode != nullptr;
-	if (!bIgnoreFocusCalculation && (Settings->bSpatialize || Settings->bEnableListenerFocus))
+	if (Settings->bSpatialize || Settings->bEnableListenerFocus)
 	{
 		// Update azimuth angles prior to updating focus as it uses this in calculating
 		// in and out of focus values.
