Fix CIS issue

diff --git a/Engine/Source/Runtime/Engine/Private/Animation/AnimStreamable.cpp b/Engine/Source/Runtime/Engine/Private/Animation/AnimStreamable.cpp
--- a/Engine/Source/Runtime/Engine/Private/Animation/AnimStreamable.cpp
+++ b/Engine/Source/Runtime/Engine/Private/Animation/AnimStreamable.cpp
@@ -9,6 +9,7 @@
 #include "Interfaces/ITargetPlatformManagerModule.h"
 #include "DeviceProfiles/DeviceProfileManager.h"
 #include "DeviceProfiles/DeviceProfile.h"
+#include "UObject/LinkerLoad.h"
 #include "Animation/AnimCompressionDerivedData.h"
 #include "DerivedDataCacheInterface.h"
 #include "Animation/AnimCurveCompressionSettings.h"
