Fix assert when calling ParticleSystemComponent::FinalizeTickComponent in UParticleSystemComponent::DestroyRenderState_Concurrent

diff --git a/Engine/Source/Runtime/Engine/Private/Particles/ParticleComponents.cpp b/Engine/Source/Runtime/Engine/Private/Particles/ParticleComponents.cpp
--- a/Engine/Source/Runtime/Engine/Private/Particles/ParticleComponents.cpp
+++ b/Engine/Source/Runtime/Engine/Private/Particles/ParticleComponents.cpp
@@ -3980,7 +3980,13 @@ void UParticleSystemComponent::DestroyRenderState_Concurrent()
 	SCOPE_CYCLE_COUNTER(STAT_ParticleSystemComponent_DestroyRenderState_Concurrent);
 	SCOPE_CYCLE_COUNTER(STAT_ParticlesOverview_GT_CNC);
 
-	ForceAsyncWorkCompletion(ENSURE_AND_STALL);
+	{
+		const bool bSaveSkipUpdate = bSkipUpdateDynamicDataDuringTick;
+		bSkipUpdateDynamicDataDuringTick = true;
+		ForceAsyncWorkCompletion(ENSURE_AND_STALL);
+		bSkipUpdateDynamicDataDuringTick = bSaveSkipUpdate;
+	}
+
 	check( GetWorld() );
 	UE_LOG(LogParticles,Verbose,
 		TEXT("DestroyRenderState_Concurrent @ %fs %s"), GetWorld()->TimeSeconds,
