#ue4 - Avoid using bAutoManageAttachment within UParticleSystemComponent unless the world is a game world. Avoid including auto-detached PSCs in CalcBounds().

diff --git a/Engine/Source/Runtime/Engine/Classes/Particles/ParticleSystemComponent.h b/Engine/Source/Runtime/Engine/Classes/Particles/ParticleSystemComponent.h
--- a/Engine/Source/Runtime/Engine/Classes/Particles/ParticleSystemComponent.h
+++ b/Engine/Source/Runtime/Engine/Classes/Particles/ParticleSystemComponent.h
@@ -488,7 +488,7 @@ private:
 	volatile bool bAsyncWorkOutstanding;
 	
 	/** Restore relative transform from auto attachment and optionally detach from parent (regardless of whether it was an auto attachment). */
-	void CancelAutoAttachment(bool bDetachFromParent);
+	void CancelAutoAttachment(bool bDetachFromParent, const UWorld* MyWorld);
 
 	/** Handle into the FParticleSystemWorldManager. INDEX_NONE if this component does not have managed ticks. */
 	int32 ManagerHandle : 30;
diff --git a/Engine/Source/Runtime/Engine/Private/Particles/ParticleComponents.cpp b/Engine/Source/Runtime/Engine/Private/Particles/ParticleComponents.cpp
--- a/Engine/Source/Runtime/Engine/Private/Particles/ParticleComponents.cpp
+++ b/Engine/Source/Runtime/Engine/Private/Particles/ParticleComponents.cpp
@@ -3822,18 +3822,22 @@ void UParticleSystemComponent::OnRegister()
 				AutoAttachSocketName = GetAttachSocketName();
 			}
 
-			// Prevent attachment before Super::OnRegister() tries to attach us, since we only attach when activated.
-			if (GetAttachParent()->GetAttachChildren().Contains(this))
+			// If in a game world, detach now if necessary. Activation will cause auto-attachment.
+			if (World->IsGameWorld())
 			{
-				// Only detach if we are not about to auto attach to the same target, that would be wasteful.
-				if (!bAutoActivate || (AutoAttachLocationRule != EAttachmentRule::KeepRelative && AutoAttachRotationRule != EAttachmentRule::KeepRelative && AutoAttachScaleRule != EAttachmentRule::KeepRelative) || (AutoAttachSocketName != GetAttachSocketName()) || (AutoAttachParent != GetAttachParent()))
+				// Prevent attachment before Super::OnRegister() tries to attach us, since we only attach when activated.
+				if (GetAttachParent()->GetAttachChildren().Contains(this))
 				{
-					DetachFromComponent(FDetachmentTransformRules(EDetachmentRule::KeepRelative, /*bCallModify=*/ false));
+					// Only detach if we are not about to auto attach to the same target, that would be wasteful.
+					if (!bAutoActivate || (AutoAttachLocationRule != EAttachmentRule::KeepRelative && AutoAttachRotationRule != EAttachmentRule::KeepRelative && AutoAttachScaleRule != EAttachmentRule::KeepRelative) || (AutoAttachSocketName != GetAttachSocketName()) || (AutoAttachParent != GetAttachParent()))
+					{
+						DetachFromComponent(FDetachmentTransformRules(EDetachmentRule::KeepRelative, /*bCallModify=*/ false));
+					}
+				}
+				else
+				{
+					SetupAttachment(nullptr, NAME_None);
 				}
-			}
-			else
-			{
-				SetupAttachment(nullptr, NAME_None);
 			}
 		}
 
@@ -4621,7 +4625,13 @@ FBoxSphereBounds UParticleSystemComponent::CalcBounds(const FTransform& LocalToW
 	FBox BoundingBox;
 	BoundingBox.Init();
 
-	if (FXConsoleVariables::bAllowCulling == false)
+	const USceneComponent* UseAutoParent = (bAutoManageAttachment && GetAttachParent() == nullptr) ? AutoAttachParent.Get() : nullptr;
+	if (UseAutoParent)
+	{
+		// We use auto attachment but have detached, don't use our own bogus bounds (we're off near 0,0,0), use the usual parent's bounds.
+		return UseAutoParent->Bounds;
+	}
+	else if (FXConsoleVariables::bAllowCulling == false)
 	{
 		BoundingBox.Min = FVector(-HALF_WORLD_MAX);
 		BoundingBox.Max = FVector(+HALF_WORLD_MAX);
@@ -5904,7 +5914,7 @@ void UParticleSystemComponent::ActivateSystem(bool bFlagAsJustAttached)
 		// Auto attach if requested
 		const bool bWasAutoAttached = bDidAutoAttach;
 		bDidAutoAttach = false;
-		if (bAutoManageAttachment)
+		if (bAutoManageAttachment && bIsGameWorld)
 		{
 			USceneComponent* NewParent = AutoAttachParent.Get();
 			if (NewParent)
@@ -5913,7 +5923,7 @@ void UParticleSystemComponent::ActivateSystem(bool bFlagAsJustAttached)
 				if (!bAlreadyAttached)
 				{
 					bDidAutoAttach = bWasAutoAttached;
-					CancelAutoAttachment(true);
+					CancelAutoAttachment(true, World);
 					SavedAutoAttachRelativeLocation = RelativeLocation;
 					SavedAutoAttachRelativeRotation = RelativeRotation;
 					SavedAutoAttachRelativeScale3D = RelativeScale3D;
@@ -5925,7 +5935,7 @@ void UParticleSystemComponent::ActivateSystem(bool bFlagAsJustAttached)
 			}
 			else
 			{
-				CancelAutoAttachment(true);
+				CancelAutoAttachment(true, World);
 			}
 		}
 
@@ -6094,7 +6104,7 @@ void UParticleSystemComponent::Complete()
 	}
 	else if (bAutoManageAttachment)
 	{
-		CancelAutoAttachment(/*bDetachFromParent=*/ true);
+		CancelAutoAttachment(/*bDetachFromParent=*/ true, World);
 	}
 }
 
@@ -6168,9 +6178,9 @@ void UParticleSystemComponent::DeactivateSystem()
 	SetLastRenderTime(World->GetTimeSeconds());
 }
 
-void UParticleSystemComponent::CancelAutoAttachment(bool bDetachFromParent)
+void UParticleSystemComponent::CancelAutoAttachment(bool bDetachFromParent, const UWorld* MyWorld)
 {
-	if (bAutoManageAttachment)
+	if (bAutoManageAttachment && MyWorld && MyWorld->IsGameWorld())
 	{
 		if (bDidAutoAttach)
 		{
