#ue4 - Avoid using bAutoManageAttachment within UAudioComponent unless the world is a game world. Avoid including auto-detached components in CalcBounds().

diff --git a/Engine/Source/Runtime/Engine/Classes/Components/AudioComponent.h b/Engine/Source/Runtime/Engine/Classes/Components/AudioComponent.h
--- a/Engine/Source/Runtime/Engine/Classes/Components/AudioComponent.h
+++ b/Engine/Source/Runtime/Engine/Classes/Components/AudioComponent.h
@@ -516,6 +516,7 @@ public:
 	virtual void Activate(bool bReset=false) override;
 	virtual void Deactivate() override;
 	virtual void OnUpdateTransform(EUpdateTransformFlags UpdateTransformFlags, ETeleportType Teleport = ETeleportType::None) override;
+	virtual FBoxSphereBounds CalcBounds(const FTransform& LocalToWorld) const override;
 	//~ End USceneComponent Interface
 
 	//~ Begin ActorComponent Interface.
@@ -605,7 +606,7 @@ private:
 	TMap<uint32, FSoundWavePlaybackTimeData> SoundWavePlaybackTimes;
 
 	/** Restore relative transform from auto attachment and optionally detach from parent (regardless of whether it was an auto attachment). */
-	void CancelAutoAttachment(bool bDetachFromParent);
+	void CancelAutoAttachment(bool bDetachFromParent, const UWorld* MyWorld);
 
 protected:
 	/** Utility function called by Play and FadeIn to start a sound playing. */
diff --git a/Engine/Source/Runtime/Engine/Private/Components/AudioComponent.cpp b/Engine/Source/Runtime/Engine/Private/Components/AudioComponent.cpp
--- a/Engine/Source/Runtime/Engine/Private/Components/AudioComponent.cpp
+++ b/Engine/Source/Runtime/Engine/Private/Components/AudioComponent.cpp
@@ -189,18 +189,23 @@ void UAudioComponent::OnRegister()
 				AutoAttachSocketName = GetAttachSocketName();
 			}
 
-			// Prevent attachment before Super::OnRegister() tries to attach us, since we only attach when activated.
-			if (GetAttachParent()->GetAttachChildren().Contains(this))
+			// If in a game world, detach now if necessary. Activation will cause auto-attachment.
+			const UWorld* World = GetWorld();
+			if (World->IsGameWorld())
 			{
-				// Only detach if we are not about to auto attach to the same target, that would be wasteful.
-				if (!bAutoActivate || (AutoAttachLocationRule != EAttachmentRule::KeepRelative && AutoAttachRotationRule != EAttachmentRule::KeepRelative && AutoAttachScaleRule != EAttachmentRule::KeepRelative) || (AutoAttachSocketName != GetAttachSocketName()) || (AutoAttachParent != GetAttachParent()))
+				// Prevent attachment before Super::OnRegister() tries to attach us, since we only attach when activated.
+				if (GetAttachParent()->GetAttachChildren().Contains(this))
 				{
-					DetachFromComponent(FDetachmentTransformRules(EDetachmentRule::KeepRelative, /*bCallModify=*/ false));
+					// Only detach if we are not about to auto attach to the same target, that would be wasteful.
+					if (!bAutoActivate || (AutoAttachLocationRule != EAttachmentRule::KeepRelative && AutoAttachRotationRule != EAttachmentRule::KeepRelative && AutoAttachScaleRule != EAttachmentRule::KeepRelative) || (AutoAttachSocketName != GetAttachSocketName()) || (AutoAttachParent != GetAttachParent()))
+					{
+						DetachFromComponent(FDetachmentTransformRules(EDetachmentRule::KeepRelative, /*bCallModify=*/ false));
+					}
+				}
+				else
+				{
+					SetupAttachment(nullptr, NAME_None);
 				}
-			}
-			else
-			{
-				SetupAttachment(nullptr, NAME_None);
 			}
 		}
 
@@ -290,9 +295,21 @@ void UAudioComponent::OnUpdateTransform(EUpdateTransformFlags UpdateTransformFla
 	}
 };
 
-void UAudioComponent::CancelAutoAttachment(bool bDetachFromParent)
+FBoxSphereBounds UAudioComponent::CalcBounds(const FTransform& LocalToWorld) const
+{
+	const USceneComponent* UseAutoParent = (bAutoManageAttachment && GetAttachParent() == nullptr) ? AutoAttachParent.Get() : nullptr;
+	if (UseAutoParent)
+	{
+		// We use auto attachment but have detached, don't use our own bogus bounds (we're off near 0,0,0), use the usual parent's bounds.
+		return UseAutoParent->Bounds;
+	}
+
+	return Super::CalcBounds(LocalToWorld);
+}
+
+void UAudioComponent::CancelAutoAttachment(bool bDetachFromParent, const UWorld* MyWorld)
 {
-	if (bAutoManageAttachment)
+	if (bAutoManageAttachment && MyWorld && MyWorld->IsGameWorld())
 	{
 		if (bDidAutoAttach)
 		{
@@ -362,7 +379,7 @@ void UAudioComponent::PlayInternal(const float StartTime, const float FadeInDura
 			// Auto attach if requested
 			const bool bWasAutoAttached = bDidAutoAttach;
 			bDidAutoAttach = false;
-			if (bAutoManageAttachment)
+			if (bAutoManageAttachment && World->IsGameWorld())
 			{
 				USceneComponent* NewParent = AutoAttachParent.Get();
 				if (NewParent)
@@ -371,7 +388,7 @@ void UAudioComponent::PlayInternal(const float StartTime, const float FadeInDura
 					if (!bAlreadyAttached)
 					{
 						bDidAutoAttach = bWasAutoAttached;
-						CancelAutoAttachment(true);
+						CancelAutoAttachment(true, World);
 						SavedAutoAttachRelativeLocation = RelativeLocation;
 						SavedAutoAttachRelativeRotation = RelativeRotation;
 						SavedAutoAttachRelativeScale3D = RelativeScale3D;
@@ -382,7 +399,7 @@ void UAudioComponent::PlayInternal(const float StartTime, const float FadeInDura
 				}
 				else
 				{
-					CancelAutoAttachment(true);
+					CancelAutoAttachment(true, World);
 				}
 			}
 
@@ -685,7 +702,8 @@ void UAudioComponent::PlaybackCompleted(bool bFailedToStart)
 	// Mark inactive before calling destroy to avoid recursion
 	bIsActive = false;
 
-	if (!bFailedToStart && GetWorld() != nullptr && (OnAudioFinished.IsBound() || OnAudioFinishedNative.IsBound()))
+	const UWorld* MyWorld = GetWorld();
+	if (!bFailedToStart && MyWorld != nullptr && (OnAudioFinished.IsBound() || OnAudioFinishedNative.IsBound()))
 	{
 		INC_DWORD_STAT(STAT_AudioFinishedDelegatesCalled);
 		SCOPE_CYCLE_COUNTER(STAT_AudioFinishedDelegates);
@@ -702,7 +720,7 @@ void UAudioComponent::PlaybackCompleted(bool bFailedToStart)
 	// Otherwise see if we should detach ourself and wait until we're needed again
 	else if (bAutoManageAttachment)
 	{
-		CancelAutoAttachment(true);
+		CancelAutoAttachment(true, MyWorld);
 	}
 }
 
