Add Replay Playlist Support

diff --git a/Engine/Source/Runtime/Engine/Classes/Engine/DemoNetDriver.h b/Engine/Source/Runtime/Engine/Classes/Engine/DemoNetDriver.h
--- a/Engine/Source/Runtime/Engine/Classes/Engine/DemoNetDriver.h
+++ b/Engine/Source/Runtime/Engine/Classes/Engine/DemoNetDriver.h
@@ -475,7 +475,14 @@ private:
 
 	TArray<TUniquePtr<FDeltaCheckpointData>> PlaybackDeltaCheckpointData;
 	
-public:	
+	TSharedPtr<struct FReplayPlaylistTracker> PlaylistTracker;
+
+public:
+
+	void SetPlayingPlaylist(TSharedPtr<struct FReplayPlaylistTracker> InPlaylistTracker)
+	{
+		PlaylistTracker = InPlaylistTracker;
+	}
 
 	virtual void Serialize(FArchive& Ar) override;
 
diff --git a/Engine/Source/Runtime/Engine/Classes/Engine/GameInstance.h b/Engine/Source/Runtime/Engine/Classes/Engine/GameInstance.h
--- a/Engine/Source/Runtime/Engine/Classes/Engine/GameInstance.h
+++ b/Engine/Source/Runtime/Engine/Classes/Engine/GameInstance.h
@@ -361,11 +361,29 @@ public:
 	/** Stop recording a replay if one is currently in progress */
 	virtual void StopRecordingReplay();
 
-	/** Start playing back a previously recorded replay.
-	 *	Return false if it fails to play.
-	*/
+	/**
+	 * Start playing back a previously recorded replay.
+	 *
+	 * @param InName				Name of the replay file.
+	 * @param WorldOverride			World in which the replay will be played. Passing null will cause the current world to be used.
+	 * @param AdditionalOptions		Additional options that can be read by derived game instances, or the Demo Net Driver.
+	 *
+	 * @return True if the replay began successfully.
+	 */
 	virtual bool PlayReplay(const FString& InName, UWorld* WorldOverride = nullptr, const TArray<FString>& AdditionalOptions = TArray<FString>());
 
+	/**
+	 * Start playing back a playlist of previously recorded replays.
+	 *
+	 * Using "ExitAfterReplay" on the command line will cause the system to exit *after* the last
+	 * replay has been played.
+	 *
+	 * Using the "Demo.Loop" CVar will cause the entire replay playlist to loop.
+	 *
+	 * @return True if the first replay began successfully.
+	 */
+	bool PlayReplayPlaylist(const struct FReplayPlaylistParams& PlaylistParams);
+
 	/**
 	 * Adds a join-in-progress user to the set of users associated with the currently recording replay (if any)
 	 *
diff --git a/Engine/Source/Runtime/Engine/Private/DemoNetDriver.cpp b/Engine/Source/Runtime/Engine/Private/DemoNetDriver.cpp
--- a/Engine/Source/Runtime/Engine/Private/DemoNetDriver.cpp
+++ b/Engine/Source/Runtime/Engine/Private/DemoNetDriver.cpp
@@ -38,6 +38,7 @@
 #include "Misc/EngineVersion.h"
 #include "Stats/Stats2.h"
 #include "Engine/ChildConnection.h"
+#include "Net/ReplayPlaylistTracker.h"
 
 DEFINE_LOG_CATEGORY( LogDemo );
 
@@ -3747,6 +3748,18 @@ static FCsvDemoSettings GetCsvDemoSettings()
 }
 #endif // (CSV_PROFILER && (!UE_BUILD_SHIPPING))
 
+class FDemoNetDriverReplayPlaylistHelper
+{
+private:
+
+	friend class UDemoNetDriver;
+
+	static void RestartPlaylist(FReplayPlaylistTracker& ToRestart)
+	{
+		ToRestart.Restart();
+	}
+};
+
 void UDemoNetDriver::TickDemoPlayback(float DeltaSeconds)
 {
 	LLM_SCOPE(ELLMTag::Networking);
@@ -3838,24 +3851,44 @@ void UDemoNetDriver::TickDemoPlayback(float DeltaSeconds)
 		if (!ReplayStreamer->IsLive() && bIsAtEnd)
 		{
 			OnDemoFinishPlaybackDelegate.Broadcast();
+			FReplayPlaylistTracker* LocalPlaylistTracker = PlaylistTracker.Get();
 
 			// checking against 1 so the count will mean total number of playthroughs, not additional loops
 			if (GDemoLoopCount > 1)
 			{
-				--GDemoLoopCount;
-
-				GotoTimeInSeconds(0.0f);
+				if (LocalPlaylistTracker)
+				{
+					if (LocalPlaylistTracker->IsOnLastReplay())
+					{
+						--GDemoLoopCount;
+						FDemoNetDriverReplayPlaylistHelper::RestartPlaylist(*LocalPlaylistTracker);
+					}
+				}
+				else
+				{
+					--GDemoLoopCount;
+					GotoTimeInSeconds(0.0f);
+				}	
 			}
 			else
 			{
-				if (FParse::Param(FCommandLine::Get(), TEXT("ExitAfterReplay")))
+				if (FParse::Param(FCommandLine::Get(), TEXT("ExitAfterReplay")) && (!LocalPlaylistTracker || LocalPlaylistTracker->IsOnLastReplay()))
 				{
 					FPlatformMisc::RequestExit(false);
 				}
-
-				if (CVarLoopDemo.GetValueOnGameThread() > 0)
+				else
 				{
-					GotoTimeInSeconds(0.0f);
+					if (CVarLoopDemo.GetValueOnGameThread() > 0)
+					{
+						if (!LocalPlaylistTracker)
+						{
+							GotoTimeInSeconds(0.0f);
+						}
+						else if (LocalPlaylistTracker->IsOnLastReplay())
+						{
+							FDemoNetDriverReplayPlaylistHelper::RestartPlaylist(*LocalPlaylistTracker);
+						}
+					}
 				}
 			}
 		}
diff --git a/Engine/Source/Runtime/Engine/Private/GameInstance.cpp b/Engine/Source/Runtime/Engine/Private/GameInstance.cpp
--- a/Engine/Source/Runtime/Engine/Private/GameInstance.cpp
+++ b/Engine/Source/Runtime/Engine/Private/GameInstance.cpp
@@ -28,6 +28,7 @@
 #include "Framework/Application/SlateApplication.h"
 #include "GenericPlatform/GenericApplication.h"
 #include "Misc/PackageName.h"
+#include "Net/ReplayPlaylistTracker.h"
 
 #if WITH_EDITOR
 #include "Settings/LevelEditorPlaySettings.h"
@@ -475,13 +476,24 @@ void UGameInstance::StartGameInstance()
 #if !UE_SERVER
 	// Parse replay name if specified on cmdline
 	FString ReplayCommand;
-	if ( FParse::Value( Tmp, TEXT( "-REPLAY=" ), ReplayCommand ) )
+	if (FParse::Value(Tmp, TEXT("-REPLAY="), ReplayCommand))
 	{
-		if(PlayReplay( ReplayCommand ))
+		if (PlayReplay(ReplayCommand))
 		{
 			return;
 		}
 	}
+	else if (FParse::Value(Tmp, TEXT("-REPLAYPLAYLIST="), ReplayCommand, false))
+	{
+		FReplayPlaylistParams Params;
+		if (ReplayCommand.ParseIntoArray(Params.Playlist, TEXT(",")))
+		{
+			if (PlayReplayPlaylist(Params))
+			{
+				return;
+			}
+		}
+	}
 #endif // !UE_SERVER
 
 	const UGameMapsSettings* GameMapsSettings = GetDefault<UGameMapsSettings>();
@@ -1102,6 +1114,27 @@ bool UGameInstance::PlayReplay(const FString& Name, UWorld* WorldOverride, const
 	return true;
 }
 
+class FGameInstanceReplayPlaylistHelper
+{
+private:
+
+	friend class UGameInstance;
+
+	static const bool StartReplay(const FReplayPlaylistParams& PlaylistParams, UGameInstance* GameInstance)
+	{
+		// Can't use MakeShared directly since the PlaylistTracker constructor is private.
+		// Also, the Playlist Tracker will manage holding onto references to itself.
+		return TSharedRef<FReplayPlaylistTracker>(new FReplayPlaylistTracker(PlaylistParams, GameInstance))->Start();
+	}
+};
+
+bool UGameInstance::PlayReplayPlaylist(const FReplayPlaylistParams& PlaylistParams)
+{
+	LLM_SCOPE(ELLMTag::Networking);
+
+	return FGameInstanceReplayPlaylistHelper::StartReplay(PlaylistParams, this);
+}
+
 void UGameInstance::AddUserToReplay(const FString& UserString)
 {
 	UWorld* CurrentWorld = GetWorld();
diff --git a/Engine/Source/Runtime/Engine/Private/Net/ReplayPlaylistTracker.cpp b/Engine/Source/Runtime/Engine/Private/Net/ReplayPlaylistTracker.cpp
new file mode 100644
--- /dev/null
+++ b/Engine/Source/Runtime/Engine/Private/Net/ReplayPlaylistTracker.cpp
@@ -0,0 +1,189 @@
+// Copyright 1998-2019 Epic Games, Inc. All Rights Reserved.
+
+#include "Net/ReplayPlaylistTracker.h"
+#include "Engine/DemoNetDriver.h"
+#include "Engine/GameInstance.h"
+#include "Engine/World.h"
+
+FReplayPlaylistTracker::FReplayPlaylistTracker(const FReplayPlaylistParams& Params, UGameInstance* InGameInstance) :
+	bIsStartingReplay(false),
+	CurrentReplay(0),
+	Playlist(Params.Playlist),
+	AdditionalOptions(Params.InitialOptions.NewAdditionalOptions.Get(TArray<FString>())),
+	WorldOverride(const_cast<UWorld*>(Params.InitialOptions.NewWorldOverride.Get(nullptr))),
+	GameInstance(InGameInstance),
+	PreStartNextReplay(Params.PreStartNextPlaylistReplayDelegate)
+{
+	const FString QuoteString(TEXT("\""));
+	for (int32 i = Playlist.Num() - 1; i >= 0; --i)
+	{
+		FString& Replay = Playlist[i];
+		Replay.TrimStartAndEndInline();
+
+		// Remove leading and trailing quotes.
+		const bool bMismatchedQuotes = Replay.StartsWith(QuoteString) != Replay.EndsWith(QuoteString);
+		if (Replay.RemoveFromStart(QuoteString) || Replay.RemoveFromEnd(QuoteString))
+		{
+			// Log the untouched replay.
+			UE_LOG(LogDemo, Log, TEXT("FReplayPlaylistTracker: Removed mismatched quotes from %s"), *Params.Playlist[i]);
+		}
+
+		// Now that we've maybe removed quotes, try trimming one last time.
+		Replay.TrimStartAndEndInline();
+
+		if (Replay.Len() == 0)
+		{
+			UE_LOG(LogDemo, Log, TEXT("FReplayPlaylistTracker: Removed invalid replay %s"), *Params.Playlist[i]);
+			Playlist.RemoveAt(i, 1, false);
+		}
+	}
+
+	if (Playlist.Num() == 0)
+	{
+		UE_LOG(LogDemo, Warning, TEXT("FReplayPlaylistTracker: No valid replays in playlist."))
+	}
+}
+
+void FReplayPlaylistTracker::Reset()
+{
+	if (UDemoNetDriver * LocalDemoNetDriver = DemoNetDriver.Get())
+	{
+		DemoNetDriver.Reset();
+		LocalDemoNetDriver->OnDemoFailedToStart.Remove(OnDemoFailedToStartHandle);
+		LocalDemoNetDriver->OnDemoFinishPlaybackDelegate.Remove(OnDemoPlaybackFinishedHandle);
+		LocalDemoNetDriver->OnDemoFinishRecordingDelegate.Remove(OnDemoStoppedHandle);
+		LocalDemoNetDriver->SetPlayingPlaylist(nullptr);
+	}
+
+	OnDemoFailedToStartHandle.Reset();
+	OnDemoPlaybackFinishedHandle.Reset();
+	OnDemoStoppedHandle.Reset();
+
+	// When we are starting a new replay, the old DemoNetDriver will be torn down.
+	// It's fine to handle the above code, as we won't reassociate until we know
+	// the replay has started successfully, we definitely don't want to reset
+	// our current replay.
+	if (!bIsStartingReplay)
+	{
+		CurrentReplay = 0;
+	}
+}
+
+bool FReplayPlaylistTracker::Start()
+{
+	// Don't log this case, as it will be called out earlier.
+	if (Playlist.Num() > 0 && CurrentReplay < Playlist.Num())
+	{
+		UWorld* World = WorldOverride.Get();
+		World = World ? World : GameInstance->GetWorld();
+
+		bool bStartedPlayingSuccessfully = false;
+		{
+			TGuardValue<bool> GuardStarting(bIsStartingReplay, true);
+			bStartedPlayingSuccessfully = GameInstance->PlayReplay(Playlist[CurrentReplay], World, AdditionalOptions);
+		}
+
+		if (bStartedPlayingSuccessfully)
+		{
+			TSharedRef<ThisClass> This = AsShared();
+
+			if (UDemoNetDriver * LocalDemoNetDriver = World->DemoNetDriver)
+			{
+				DemoNetDriver = LocalDemoNetDriver;
+				OnDemoFailedToStartHandle = LocalDemoNetDriver->OnDemoFailedToStart.AddStatic(&ThisClass::OnDemoFailedToStart, This);
+				OnDemoPlaybackFinishedHandle = LocalDemoNetDriver->OnDemoFinishPlaybackDelegate.AddStatic(&ThisClass::OnDemoPlaybackFinished, This);
+				OnDemoPlaybackFinishedHandle = LocalDemoNetDriver->OnDemoFinishRecordingDelegate.AddStatic(&ThisClass::OnDemoStopped, This);
+				LocalDemoNetDriver->SetPlayingPlaylist(This);
+
+				UE_LOG(LogDemo, Log, TEXT("FReplayPlaylistTracker::Start: Successfully started Replay=%s, CurrentReplayIdx=%d"), *Playlist[CurrentReplay], CurrentReplay);
+				return true;
+			}
+			else
+			{
+				UE_LOG(LogDemo, Error, TEXT("FReplayPlaylistTracker::Start: Unable to retrieve DemoNetDriver, replay may be started, but playlist won't continue"));
+			}
+		}
+		else
+		{
+			UE_LOG(LogDemo, Warning, TEXT("FReplayPlaylistTracker::Start: Failed to play replay."));
+		}
+	}
+
+	return false;
+}
+
+bool FReplayPlaylistTracker::Restart()
+{
+	CurrentReplay = 0;
+	return Start();
+}
+
+void FReplayPlaylistTracker::OnDemoFailedToStart(UDemoNetDriver* DemoNetDriver, EDemoPlayFailure::Type FailureType, TSharedRef<ThisClass> This)
+{
+	This->Quit();
+}
+
+void FReplayPlaylistTracker::OnDemoStopped(TSharedRef<FReplayPlaylistTracker> Tracker)
+{
+	Tracker->Quit();
+}
+
+void FReplayPlaylistTracker::Quit()
+{
+	Reset();
+}
+
+void FReplayPlaylistTracker::OnDemoPlaybackFinished(TSharedRef<FReplayPlaylistTracker> Tracker)
+{
+	Tracker->PlayNextReplay();
+}
+
+void FReplayPlaylistTracker::PlayNextReplay()
+{
+	bool bWasSuccesful = false;
+	if (UGameInstance * LocalGameInstance = GameInstance.Get())
+	{
+		CurrentReplay = FMath::Min(Playlist.Num(), CurrentReplay + 1);
+
+		if (CurrentReplay < Playlist.Num())
+		{
+			if (PreStartNextReplay.IsBound())
+			{
+				FReplayPlaylistUpdatedOptions NewOptions;
+				PreStartNextReplay.Execute(*this, NewOptions);
+
+				if (NewOptions.NewWorldOverride.IsSet())
+				{
+					UE_LOG(LogDemo, Log, TEXT("FReplayPlaylistTracker::PlayNextReplay: Updated World Override, NewWorldOverride=%s, OldWorldOverride=%s"),
+						*GetPathNameSafe(NewOptions.NewWorldOverride.GetValue()),
+						*GetPathNameSafe(WorldOverride.Get()));
+
+					WorldOverride = NewOptions.NewWorldOverride.GetValue();
+				}
+				if (NewOptions.NewAdditionalOptions.IsSet())
+				{
+					UE_LOG(LogDemo, Log, TEXT("FReplayPlaylistTracker::PlayNextReplay: Updated AdditionalOptions, NewOptions=[%s], OldOptions=[%s]"),
+						*FString::Join(NewOptions.NewAdditionalOptions.GetValue(), TEXT(",")),
+						*FString::Join(AdditionalOptions, TEXT(",")));
+
+					AdditionalOptions = MoveTemp(NewOptions.NewAdditionalOptions.GetValue());
+				}
+			}
+
+			bWasSuccesful = Start();
+		}
+		else
+		{
+			// If we've reached the end of our playlist, we don't need to do anything.
+			// The DemoNetDriver may fire every the finish notification every frame,
+			// and we don't want to inadvertently reset our state (restarting our playlist).
+			// Once the demo has been closed out, we should get the Stopped (or error) delegate callbacks,
+			// and that's what will clean us up.
+			bWasSuccesful = true;
+		}
+	}
+	if (!bWasSuccesful)
+	{
+		Quit();
+	}
+}
\ No newline at end of file
diff --git a/Engine/Source/Runtime/Engine/Public/Net/ReplayPlaylistTracker.h b/Engine/Source/Runtime/Engine/Public/Net/ReplayPlaylistTracker.h
new file mode 100644
--- /dev/null
+++ b/Engine/Source/Runtime/Engine/Public/Net/ReplayPlaylistTracker.h
@@ -0,0 +1,137 @@
+// Copyright 1998-2019 Epic Games, Inc. All Rights Reserved.
+
+#pragma once
+
+#include "CoreMinimal.h"
+#include "Engine/EngineBaseTypes.h"
+#include "Templates/UnrealTemplate.h"
+#include "Templates/SharedPointer.h"
+
+/**
+ * Delegate that is called prior to starting each replay in a playlist,
+ * giving game code a chance to update any necessary options or
+ * perform other work.
+ */
+DECLARE_DELEGATE_TwoParams(FPreStartNextPlaylistReplay, const struct FReplayPlaylistTracker& /*ThisPlaylistTracker*/, struct FReplayPlaylistUpdatedOptions& /*OverrideWorldUpdate*/);
+
+/** Set of options that can be updated before starting each replay in a playlist. */
+struct FReplayPlaylistUpdatedOptions
+{
+	TOptional<class UWorld*> NewWorldOverride;
+	TOptional<TArray<FString>> NewAdditionalOptions;
+};
+
+/** Set of parameters used with UGameInstance::PlayReplayPlatlist to control how the playlist is started. */
+struct FReplayPlaylistParams
+{
+	/** Set of replays to use as the playlist. */
+	TArray<FString> Playlist;
+
+	/** Callback that can be used to update playback options after each individual replay has been played. */
+	FPreStartNextPlaylistReplay PreStartNextPlaylistReplayDelegate;
+
+	/**
+	 * Set of options that will be used to play the very first replay.
+	 * These can be updated via PreStartNextPlaylistReplayDelegate before each new replay begins
+	 * (including before looping if "Demo.Loop" is enabled.
+	 */
+	FReplayPlaylistUpdatedOptions InitialOptions;
+};
+
+/** Used with UGameInstance::PlayReplayPlaylist to manage playing a set of Replays in succession. */
+struct FReplayPlaylistTracker : public TSharedFromThis<FReplayPlaylistTracker>, public FNoncopyable
+{
+private:
+
+	using ThisClass = FReplayPlaylistTracker;
+	friend class FGameInstanceReplayPlaylistHelper;
+	friend class FDemoNetDriverReplayPlaylistHelper;
+
+public:
+
+	~FReplayPlaylistTracker()
+	{
+		// Forcibly reset this in case we end up cancelling requests while trying to start a new replay.
+		bIsStartingReplay = false;
+		Reset();
+	}
+
+	/** Whether or not we're currently on the last replay in a playlist. */
+	const int32 IsOnLastReplay() const
+	{
+		return GetCurrentReplay() == GetNumReplays();
+	}
+
+	/** Gets the number of replays in the playlist. */
+	const int32 GetNumReplays() const
+	{
+		return Playlist.Num();
+	}
+
+	/** Gets the index of the current replay in the playlist. */
+	const int32 GetCurrentReplay() const
+	{
+		return CurrentReplay;
+	}
+
+	/** Gets the actual playlist. */
+	const TArrayView<const FString> GetPlaylist() const
+	{
+		return Playlist;
+	}
+
+	/** Gets the last updated value of AdditionalOptions. */
+	const TArrayView<const FString> GetAdditionalOptions() const
+	{
+		return AdditionalOptions;
+	}
+
+	/** Gets the last updated value of WorldOverride. */
+	class UWorld* GetWorldOverride() const
+	{
+		return WorldOverride.Get();
+	}
+
+	/** Gets the GameInstance used to start replays. */
+	class UGameInstance* GetGameInstance() const
+	{
+		return GameInstance.Get();
+	}
+
+private:
+
+	FReplayPlaylistTracker(const FReplayPlaylistParams& Params, class UGameInstance* InGameInstance);
+
+	void Reset();
+
+	bool Start();
+
+	bool Restart();
+
+	void Quit();
+
+	void PlayNextReplay();
+
+	/** Called if an error occurs *after* the replay was requested to start (successfully), but before the replay actually starts. */
+	static void OnDemoFailedToStart(class UDemoNetDriver* DemoNetDriver, EDemoPlayFailure::Type FailureType, TSharedRef<ThisClass> This);
+
+	/** Called when the DemoNetDriver hits a Hard Stop, like an error or a user requesting a stop. */
+	static void OnDemoStopped(TSharedRef<FReplayPlaylistTracker> Tracker);
+
+	/** Called when a replay has reached its end. */
+	static void OnDemoPlaybackFinished(TSharedRef<FReplayPlaylistTracker> Tracker);
+
+	bool bIsStartingReplay;
+	int32 CurrentReplay;
+	TArray<FString> Playlist;
+	TArray<FString> AdditionalOptions;
+	TWeakObjectPtr<class UWorld> WorldOverride;
+	TWeakObjectPtr<class UGameInstance> GameInstance;
+	TWeakObjectPtr<class UDemoNetDriver> DemoNetDriver;
+
+	FDelegateHandle OnDemoFailedToStartHandle;
+	FDelegateHandle OnDemoPlaybackFinishedHandle;
+	FDelegateHandle OnDemoStoppedHandle;
+
+	FPreStartNextPlaylistReplay PreStartNextReplay;
+};
\ No newline at end of file
