First pass at adding packet loss burst analytics.

diff --git a/Engine/Source/Runtime/Engine/Classes/Engine/NetConnection.h b/Engine/Source/Runtime/Engine/Classes/Engine/NetConnection.h
--- a/Engine/Source/Runtime/Engine/Classes/Engine/NetConnection.h
+++ b/Engine/Source/Runtime/Engine/Classes/Engine/NetConnection.h
@@ -1099,6 +1099,8 @@ public:
 	/** Removes stale entries from DormantReplicatorMap. */
 	void CleanupStaleDormantReplicators();
 
+	void PostTickDispatch();
+
 	/**
 	 * Flush the cache of sequenced packets waiting for a missing packet. Will flush only up to the next missing packet, unless bFlushWholeCache is set.
 	 *
@@ -1135,6 +1137,19 @@ public:
 	/** Resets the current saturation analytics. */
 	ENGINE_API void ResetSaturationAnalytics();
 
+	/**
+	 * Returns the current packet stability analytics and resets them.
+	 * This would be similar to calls to Get and Reset separately, except that the caller
+	 * will assume ownership of the data in this case.
+	 */
+	ENGINE_API void ConsumePacketAnalytics(FNetConnectionPacketAnalytics& Out);
+
+	/** Returns the current packet stability analytics. */
+	ENGINE_API const FNetConnectionPacketAnalytics& GetPacketAnalytics() const;
+
+	/** Resets the current packet stability analytics. */
+	ENGINE_API void ResetPacketAnalytics();
+
 	/**
 	 * Called to notify the connection that we attempted to replicate its actors this frame.
 	 * This is primarily for analytics book keeping.
@@ -1259,6 +1274,7 @@ private:
 	int32 PacketOrderCacheCount;
 
 	FNetConnectionSaturationAnalytics SaturationAnalytics;
+	FNetConnectionPacketAnalytics PacketAnalytics;
 
 	/** Whether or not PacketOrderCache is presently being flushed */
 	bool bFlushingPacketOrderCache;
diff --git a/Engine/Source/Runtime/Engine/Private/Net/NetAnalyticsTypes.cpp b/Engine/Source/Runtime/Engine/Private/Net/NetAnalyticsTypes.cpp
--- a/Engine/Source/Runtime/Engine/Private/Net/NetAnalyticsTypes.cpp
+++ b/Engine/Source/Runtime/Engine/Private/Net/NetAnalyticsTypes.cpp
@@ -3,6 +3,34 @@
 #include "Net/NetAnalyticsTypes.h"
 #include "Algo/BinarySearch.h"
 #include "Serialization/Archive.h"
+#include "HAL/IConsoleManager.h"
+
+namespace NetAnalyticsCVars
+{
+	static int32 NumberOfConsecutiveDroppedPacketsToConsiderBurst = 3;
+	static FAutoConsoleVariableRef CVarNumberOfDroppedPacketsToConsiderBurst(
+		TEXT("NetAnalytics.NumberOfConsecutiveDroppedPacketsToConsiderBurst"),
+		NumberOfConsecutiveDroppedPacketsToConsiderBurst,
+		TEXT("The number of packets lost in a row (in or out) for us to consider the frame as having bursts of packet loss.")
+		TEXT("Not affected by NetAnalytics.MinimumNumberOfPacketsForBurstTracking.")
+	);
+
+	static int32 MinimumNumberOfPacketsForBurstTracking = 5;
+	static FAutoConsoleVariableRef CVarMinimumNumberOfPacketsForBurstTracking(
+		TEXT("NetAnalytics.MinimumNumberOfPacketsForBurstTracking"),
+		MinimumNumberOfPacketsForBurstTracking,
+		TEXT("The minimum number of packets that must have been notified (in our out) in order to consider a frame for packet loss by percentage.")
+		TEXT("See NetAnalytics.PercentOfDroppedPacketsToConsiderBurst")
+	);
+
+	static float PercentOfDroppedPacketsToConsiderBurst = 0.2f;
+	static FAutoConsoleVariableRef CVarPercentOfDroppedPacketsToConsiderBurst(
+		TEXT("NetAnalytics.PercentOfDroppedPacketsToConsiderBurst"),
+		PercentOfDroppedPacketsToConsiderBurst,
+		TEXT("The percentage of packets lost in a frame (in or out) for us to consider the frame as having bursts of packet loss.\n")
+		TEXT("See NetAnalytics.MinimumNumberOfPacketsForBurstTracking.")
+	);
+}
 
 FDelinquencyAnalytics::FDelinquencyAnalytics(const uint32 InNumberOfTopOffendersToTrack) :
 	TotalTime(0.f),
@@ -11,11 +39,11 @@ FDelinquencyAnalytics::FDelinquencyAnalytics(const uint32 InNumberOfTopOffenders
 	TopOffenders.Reserve(NumberOfTopOffendersToTrack);
 }
 
-FDelinquencyAnalytics::FDelinquencyAnalytics(FDelinquencyAnalytics&& Other) :
-	TopOffenders(MoveTemp(Other.TopOffenders)),
-	AllDelinquents(MoveTemp(Other.AllDelinquents)),
-	TotalTime(Other.TotalTime),
-	NumberOfTopOffendersToTrack(Other.NumberOfTopOffendersToTrack)
+FDelinquencyAnalytics::FDelinquencyAnalytics(FDelinquencyAnalytics&& Other)
+	: TopOffenders(MoveTemp(Other.TopOffenders))
+	, AllDelinquents(MoveTemp(Other.AllDelinquents))
+	, TotalTime(Other.TotalTime)
+	, NumberOfTopOffendersToTrack(Other.NumberOfTopOffendersToTrack)
 {
 	TopOffenders.Reserve(NumberOfTopOffendersToTrack);
 }
@@ -176,4 +204,95 @@ void FNetConnectionSaturationAnalytics::Reset()
 
 	CurrentRunOfSaturatedFrames = 0;
 	CurrentRunOfSaturatedReplications = 0;
+}
+
+void FNetConnectionPacketAnalytics::Reset()
+{
+	bSawPacketLossBurstThisFrame = false;
+	NumberOfAcksThisFrame = 0;
+	NumberOfNaksThisFrame = 0;
+	NumberOfMissingPacketsThisFrame = 0;
+	NumberOfPacketsThisFrame = 0;
+	CurrentRunOfDroppedOutPackets = 0;
+	LongestRunOfDroppedOutPackets = 0;
+	LongestRunOfDroppedInPackets = 0;
+	NumberOfFramesWithBurstsOfPacketLoss = 0;
+	NumberOfFramesWithNoPackets = 0;
+	NumberOfTrackedFrames = 0;
+}
+
+void FNetConnectionPacketAnalytics::TrackAck(int32 PacketId)
+{
+	++NumberOfAcksThisFrame;
+
+	// We will only know about NAKs when we receive ACKs, so we only need to update this here.
+	if (UNLIKELY(CurrentRunOfDroppedOutPackets > 0))
+	{
+		if (NetAnalyticsCVars::NumberOfConsecutiveDroppedPacketsToConsiderBurst > 0)
+		{
+			bSawPacketLossBurstThisFrame |= CurrentRunOfDroppedOutPackets >= static_cast<uint32>(NetAnalyticsCVars::NumberOfConsecutiveDroppedPacketsToConsiderBurst);
+		}
+
+		LongestRunOfDroppedOutPackets = FMath::Max<>(CurrentRunOfDroppedOutPackets, LongestRunOfDroppedOutPackets);
+		CurrentRunOfDroppedOutPackets = 0;
+	}
+}
+
+void FNetConnectionPacketAnalytics::TrackNak(int32 PacketId)
+{
+	++NumberOfNaksThisFrame;
+	++CurrentRunOfDroppedOutPackets;
+}
+
+void FNetConnectionPacketAnalytics::TrackInPacket(uint32 InPacketId, uint32 NumberOfMissingPackets)
+{
+	++NumberOfPacketsThisFrame;
+
+	// We will only know about missing packets when we receive a packet, so we only need to update this here.
+	if (UNLIKELY(NumberOfMissingPackets))
+	{
+		if (NetAnalyticsCVars::NumberOfConsecutiveDroppedPacketsToConsiderBurst > 0)
+		{
+			bSawPacketLossBurstThisFrame |= NumberOfMissingPackets >= static_cast<uint32>(NetAnalyticsCVars::NumberOfConsecutiveDroppedPacketsToConsiderBurst);
+		}
+
+		NumberOfMissingPackets += NumberOfMissingPacketsThisFrame;
+		LongestRunOfDroppedInPackets = FMath::Max<>(NumberOfMissingPackets, LongestRunOfDroppedInPackets);
+	}
+}
+
+void FNetConnectionPacketAnalytics::Tick()
+{
+	++NumberOfTrackedFrames;
+
+	const int32 NumberOfInPackets = NumberOfPacketsThisFrame + NumberOfMissingPacketsThisFrame;
+	const int32 NumberOfOutPackets = NumberOfAcksThisFrame + NumberOfNaksThisFrame;
+
+	if (NumberOfInPackets == 0 && NumberOfOutPackets == 0)
+	{
+		++NumberOfFramesWithNoPackets;
+		return;
+	}
+
+	if (!bSawPacketLossBurstThisFrame &&
+		NetAnalyticsCVars::MinimumNumberOfPacketsForBurstTracking > 0 &&
+		NetAnalyticsCVars::PercentOfDroppedPacketsToConsiderBurst > 0.f)
+	{
+		const float PercentInPacketsLost = (NumberOfInPackets > 0) ? static_cast<float>(NumberOfMissingPacketsThisFrame) / static_cast<float>(NumberOfInPackets) : 0.f;
+		const float PercentOutPacketsLost = (NumberOfOutPackets > 0) ? static_cast<float>(NumberOfNaksThisFrame) / static_cast<float>(NumberOfOutPackets) : 0.f;
+
+		bSawPacketLossBurstThisFrame = (PercentInPacketsLost >= NetAnalyticsCVars::PercentOfDroppedPacketsToConsiderBurst) ||
+			(PercentOutPacketsLost >= NetAnalyticsCVars::PercentOfDroppedPacketsToConsiderBurst);
+	}
+
+	if (bSawPacketLossBurstThisFrame)
+	{
+		++NumberOfFramesWithBurstsOfPacketLoss;
+	}
+
+	bSawPacketLossBurstThisFrame = false;
+	NumberOfAcksThisFrame = 0;
+	NumberOfNaksThisFrame = 0;
+	NumberOfMissingPacketsThisFrame = 0;
+	NumberOfPacketsThisFrame = 0;
 }
\ No newline at end of file
diff --git a/Engine/Source/Runtime/Engine/Private/NetConnection.cpp b/Engine/Source/Runtime/Engine/Private/NetConnection.cpp
--- a/Engine/Source/Runtime/Engine/Private/NetConnection.cpp
+++ b/Engine/Source/Runtime/Engine/Private/NetConnection.cpp
@@ -1169,6 +1169,15 @@ void UNetConnection::ReceivedRawPacket( void* InData, int32 Count )
 	}
 }
 
+void UNetConnection::PostTickDispatch()
+{
+	if (!InternalAck)
+	{
+		FlushPacketOrderCache(/*bFlushWholeCache=*/true);
+		PacketAnalytics.Tick();
+	}
+}
+
 void UNetConnection::FlushPacketOrderCache(bool bFlushWholeCache/*=false*/)
 {
 	if (PacketOrderCache.IsSet() && PacketOrderCacheCount > 0)
@@ -1423,6 +1432,8 @@ void UNetConnection::ReceivedAck(int32 AckPacketId)
 	// Process the bunch.
 	LastRecvAckTime = Driver->Time;
 
+	PacketAnalytics.TrackAck(AckPacketId);
+
 	if (PackageMap != NULL)
 	{
 		PackageMap->ReceivedAck( AckPacketId );
@@ -1467,6 +1478,8 @@ void UNetConnection::ReceivedNak( int32 NakPacketId )
 
 	SCOPE_CYCLE_COUNTER(Stat_NetConnectionReceivedNak);
 
+	PacketAnalytics.TrackNak(NakPacketId);
+
 	// Update pending NetGUIDs
 	PackageMap->ReceivedNak(NakPacketId);
 
@@ -1727,7 +1740,6 @@ void UNetConnection::ReceivedPacket( FBitReader& Reader )
 				return;
 			}
 
-
 			if (MissingPacketCount > 10)
 			{
 				UE_LOG(LogNetTraffic, Verbose, TEXT("High single frame packet loss. PacketsLost: %i %s" ), MissingPacketCount, *Describe());
@@ -1738,6 +1750,8 @@ void UNetConnection::ReceivedPacket( FBitReader& Reader )
 			Driver->InPacketsLost += MissingPacketCount;
 			Driver->InTotalPacketsLost += MissingPacketCount;
 			InPacketId += PacketSequenceDelta;
+
+			PacketAnalytics.TrackInPacket(InPacketId, MissingPacketCount);
 		}
 		else
 		{
@@ -3531,6 +3545,22 @@ void UNetConnection::ResetSaturationAnalytics()
 	SaturationAnalytics.Reset();
 }
 
+void UNetConnection::ConsumePacketAnalytics(FNetConnectionPacketAnalytics& Out)
+{
+	Out = MoveTemp(PacketAnalytics);
+	PacketAnalytics.Reset();
+}
+
+const FNetConnectionPacketAnalytics& UNetConnection::GetPacketAnalytics() const
+{
+	return PacketAnalytics;
+}
+
+void UNetConnection::ResetPacketAnalytics()
+{
+	PacketAnalytics.Reset();
+}
+
 void UNetConnection::TrackReplicationForAnalytics(const bool bWasSaturated)
 {
 	++TickCount;
diff --git a/Engine/Source/Runtime/Engine/Private/NetworkDriver.cpp b/Engine/Source/Runtime/Engine/Private/NetworkDriver.cpp
--- a/Engine/Source/Runtime/Engine/Private/NetworkDriver.cpp
+++ b/Engine/Source/Runtime/Engine/Private/NetworkDriver.cpp
@@ -1727,16 +1727,18 @@ void UNetDriver::PostTickDispatch()
 	// Flush out of order packet caches for connections that did not receive the missing packets during TickDispatch
 	if (ServerConnection != nullptr)
 	{
-		ServerConnection->FlushPacketOrderCache(true);
+		if (!ServerConnection->IsPendingKill())
+		{
+			ServerConnection->PostTickDispatch();
+		}
 	}
 
 	TArray<UNetConnection*> ClientConnCopy = ClientConnections;
-
 	for (UNetConnection* CurConn : ClientConnCopy)
 	{
 		if (!CurConn->IsPendingKill())
 		{
-			CurConn->FlushPacketOrderCache(true);
+			CurConn->PostTickDispatch();
 		}
 	}
 
diff --git a/Engine/Source/Runtime/Engine/Public/Net/NetAnalyticsTypes.h b/Engine/Source/Runtime/Engine/Public/Net/NetAnalyticsTypes.h
--- a/Engine/Source/Runtime/Engine/Public/Net/NetAnalyticsTypes.h
+++ b/Engine/Source/Runtime/Engine/Public/Net/NetAnalyticsTypes.h
@@ -15,9 +15,9 @@ struct ENGINE_API FDelinquencyNameTimePair
 {
 public:
 
-	FDelinquencyNameTimePair(FName InName, float InTimeSeconds) :
-		Name(InName),
-		TimeSeconds(InTimeSeconds)
+	FDelinquencyNameTimePair(FName InName, float InTimeSeconds)
+		: Name(InName)
+		, TimeSeconds(InTimeSeconds)
 	{
 	}
 
@@ -124,9 +124,9 @@ private:
  */
 struct ENGINE_API FNetAsyncLoadDelinquencyAnalytics
 {
-	FNetAsyncLoadDelinquencyAnalytics() :
-		DelinquentAsyncLoads(0),
-		MaxConcurrentAsyncLoads(0)
+	FNetAsyncLoadDelinquencyAnalytics()
+		: DelinquentAsyncLoads(0)
+		, MaxConcurrentAsyncLoads(0)
 	{
 	}
 
@@ -169,9 +169,9 @@ struct ENGINE_API FNetAsyncLoadDelinquencyAnalytics
  */
 struct ENGINE_API FNetQueuedActorDelinquencyAnalytics
 {
-	FNetQueuedActorDelinquencyAnalytics() :
-		DelinquentQueuedActors(0),
-		MaxConcurrentQueuedActors(0)
+	FNetQueuedActorDelinquencyAnalytics()
+		: DelinquentQueuedActors(0)
+		, MaxConcurrentQueuedActors(0)
 	{
 	}
 
@@ -210,15 +210,17 @@ struct ENGINE_API FNetQueuedActorDelinquencyAnalytics
 /** Struct wrapping Per Net Connection saturation analytics. */
 struct ENGINE_API FNetConnectionSaturationAnalytics
 {
-	FNetConnectionSaturationAnalytics() :
-		NumberOfTrackedFrames(0),
-		NumberOfSaturatedFrames(0),
-		LongestRunOfSaturatedFrames(0),
-		NumberOfReplications(0),
-		NumberOfSaturatedReplications(0),
-		LongestRunOfSaturatedReplications(0),
-		CurrentRunOfSaturatedFrames(0),
-		CurrentRunOfSaturatedReplications(0)
+public:
+
+	FNetConnectionSaturationAnalytics()
+		: NumberOfTrackedFrames(0)
+		, NumberOfSaturatedFrames(0)
+		, LongestRunOfSaturatedFrames(0)
+		, NumberOfReplications(0)
+		, NumberOfSaturatedReplications(0)
+		, LongestRunOfSaturatedReplications(0)
+		, CurrentRunOfSaturatedFrames(0)
+		, CurrentRunOfSaturatedReplications(0)
 	{
 	}
 
@@ -282,4 +284,84 @@ private:
 
 	uint32 CurrentRunOfSaturatedFrames;
 	uint32 CurrentRunOfSaturatedReplications;
+};
+
+/** Struct wrapper Per Net Connection analytics for things like packet loss and jitter. */
+struct ENGINE_API FNetConnectionPacketAnalytics
+{
+public:
+
+	FNetConnectionPacketAnalytics()
+		: bSawPacketLossBurstThisFrame(false)
+		, NumberOfAcksThisFrame(0)
+		, NumberOfNaksThisFrame(0)
+		, NumberOfMissingPacketsThisFrame(0)
+		, NumberOfPacketsThisFrame(0)
+		, CurrentRunOfDroppedOutPackets(0)
+		, LongestRunOfDroppedOutPackets(0)
+		, LongestRunOfDroppedInPackets(0)
+		, NumberOfFramesWithBurstsOfPacketLoss(0)
+		, NumberOfFramesWithNoPackets(0)
+		, NumberOfTrackedFrames(0)
+	{
+	}
+
+	/** Longest number of consecutive dropped incoming packets that was tracked. */
+	const uint32 GetLongestRunOfDroppedInPackets() const
+	{
+		return LongestRunOfDroppedInPackets;
+	}
+
+	/** Longest number of consecutive outgoing packets that was tracked. */
+	const uint32 GetLongestRunOfDroppedOutPackets() const
+	{
+		return LongestRunOfDroppedOutPackets;
+	}
+
+	/** Number of frames where we saw bursts of packet loss. */
+	const uint32 GetNumberOfFramesWithBurstsOfPacketLoss() const
+	{
+		return NumberOfFramesWithBurstsOfPacketLoss;
+	}
+
+	/** The total number of frames where we were not notified of any packets (ACK, NAK, in or out). */
+	const uint32 GetNumberOfFramesWithNoPackets() const
+	{
+		return NumberOfFramesWithNoPackets;
+	}
+
+	/** Total number of frames that have been tracked. */
+	const uint64 GetNumberOfTrackedFrames() const
+	{
+		return NumberOfTrackedFrames;
+	}
+
+	/** Resets the state of tracking. */
+	void Reset();
+
+private:
+
+	friend class UNetConnection;
+
+	void TrackAck(int32 PacketId);
+	void TrackNak(int32 PacketId);
+	void TrackInPacket(uint32 InPacketId, uint32 NumberOfMissingPackets);
+
+	void Tick();
+
+	bool bSawPacketLossBurstThisFrame : 1;
+
+	uint32 NumberOfAcksThisFrame;
+	uint32 NumberOfNaksThisFrame;
+
+	uint32 NumberOfMissingPacketsThisFrame;
+	uint32 NumberOfPacketsThisFrame;
+
+	uint32 CurrentRunOfDroppedOutPackets;
+	uint32 LongestRunOfDroppedOutPackets;
+	uint32 LongestRunOfDroppedInPackets;
+
+	uint32 NumberOfFramesWithBurstsOfPacketLoss;
+	uint32 NumberOfFramesWithNoPackets;
+	uint64 NumberOfTrackedFrames;
 };
\ No newline at end of file
