#ue4 - APlayerController::ServerAcknowledgePossession_Implementation() resets server network prediction data since that is called only from ClientRestart() which resets the client prediction data. This keeps the movement timestamps in sync.

diff --git a/Engine/Source/Runtime/Engine/Private/PlayerController.cpp b/Engine/Source/Runtime/Engine/Private/PlayerController.cpp
--- a/Engine/Source/Runtime/Engine/Private/PlayerController.cpp
+++ b/Engine/Source/Runtime/Engine/Private/PlayerController.cpp
@@ -75,6 +75,18 @@ DECLARE_CYCLE_STAT(TEXT("  PC Tick Input"), STAT_PC_TickInput, STATGROUP_PlayerC
 DECLARE_CYCLE_STAT(TEXT("    PC Build Input Stack"), STAT_PC_BuildInputStack, STATGROUP_PlayerController);
 DECLARE_CYCLE_STAT(TEXT("    PC Process Input Stack"), STAT_PC_ProcessInputStack, STATGROUP_PlayerController);
 
+// CVars
+namespace PlayerControllerCVars
+{
+	// Resync timestamps on pawn ack
+	static int32 NetResetServerPredictionDataOnPawnAck = 1;
+	FAutoConsoleVariableRef CVarNetResetServerPredictionDataOnPawnAck(
+		TEXT("PlayerController.NetResetServerPredictionDataOnPawnAck"),
+		NetResetServerPredictionDataOnPawnAck,
+		TEXT("Whether to reset server prediction data for the possessed Pawn when the pawn ack handshake completes.\n")
+		TEXT("0: Disable, 1: Enable"),
+		ECVF_Default);
+}
 
 const float RetryClientRestartThrottleTime = 0.5f;
 const float RetryServerAcknowledgeThrottleTime = 0.25f;
@@ -1204,6 +1216,18 @@ void APlayerController::ServerAcknowledgePossession_Implementation(APawn* P)
 {
 	UE_LOG(LogPlayerController, Verbose, TEXT("ServerAcknowledgePossession_Implementation %s"), *GetNameSafe(P));
 	AcknowledgedPawn = P;
+
+	if (PlayerControllerCVars::NetResetServerPredictionDataOnPawnAck != 0)
+	{
+		if (AcknowledgedPawn && AcknowledgedPawn == GetPawn())
+		{
+			INetworkPredictionInterface* NetworkPredictionInterface = GetPawn() ? Cast<INetworkPredictionInterface>(GetPawn()->GetMovementComponent()) : NULL;
+			if (NetworkPredictionInterface)
+			{
+				NetworkPredictionInterface->ResetPredictionData_Server();
+			}
+		}
+	}
 }
 
 bool APlayerController::ServerAcknowledgePossession_Validate(APawn* P)
