Set correct streaming chunk size

diff --git a/Engine/Source/Runtime/Engine/Private/Animation/AnimStreamable.cpp b/Engine/Source/Runtime/Engine/Private/Animation/AnimStreamable.cpp
--- a/Engine/Source/Runtime/Engine/Private/Animation/AnimStreamable.cpp
+++ b/Engine/Source/Runtime/Engine/Private/Animation/AnimStreamable.cpp
@@ -25,7 +25,7 @@ DECLARE_CYCLE_STAT(TEXT("AnimStreamable GetAnimationPose"), STAT_AnimStreamable_
 // set it here
 const TCHAR* StreamingAnimChunkVersion = TEXT("1F1656B9E10142729AB16650D9821B1Fuf");
 
-const float MINIMUM_CHUNK_SIZE = 1.0f;
+const float MINIMUM_CHUNK_SIZE = 2.0f;
 
 float GChunkSizeSeconds = MINIMUM_CHUNK_SIZE;
 
@@ -235,7 +235,7 @@ void UAnimStreamable::GetAnimationPose(FCompactPose& OutPose, FBlendedCurve& Out
 	FRootMotionReset RootMotionReset(bEnableRootMotion, RootMotionRootLock, bForceRootLock, FTransform(), false); // MDW Does not support root motion yet
 
 #if WITH_EDITOR
-	if (!HasRunningPlatformData())
+	if (!HasRunningPlatformData() || RequiredBones.ShouldUseRawData())
 	{
 		//Need to evaluate raw data
 		RawCurveData.EvaluateCurveData(OutCurve, ExtractionContext.CurrentTime);
