More generic  fix for end-of-frame update issue when finalizing the particle component ticks

diff --git a/Engine/Source/Runtime/Engine/Classes/Particles/ParticleSystemComponent.h b/Engine/Source/Runtime/Engine/Classes/Particles/ParticleSystemComponent.h
--- a/Engine/Source/Runtime/Engine/Classes/Particles/ParticleSystemComponent.h
+++ b/Engine/Source/Runtime/Engine/Classes/Particles/ParticleSystemComponent.h
@@ -1232,6 +1232,18 @@ public:
 		ENSURE_AND_STALL,
 		SILENT, // this would only be appropriate for editor only or other unusual things that we never see in game
 	};
+	/** If there is async work outstanding, force it to be completed now **/
+	FORCEINLINE void ForceAsyncWorkCompletion(EForceAsyncWorkCompletion Behavior, bool bDefinitelyGameThread, bool InSkipUpdateDynamicDataDuringTick)
+	{
+		if (AsyncWork.GetReference())
+		{
+			const bool bSavedSkipUpdate = bSkipUpdateDynamicDataDuringTick;
+			bSkipUpdateDynamicDataDuringTick |= InSkipUpdateDynamicDataDuringTick;
+			WaitForAsyncAndFinalize(Behavior, bDefinitelyGameThread);
+			bSkipUpdateDynamicDataDuringTick = bSavedSkipUpdate;
+		}
+	}
+
 	/** If there is async work outstanding, force it to be completed now **/
 	FORCEINLINE void ForceAsyncWorkCompletion(EForceAsyncWorkCompletion Behavior, bool bDefinitelyGameThread = true) const
 	{
diff --git a/Engine/Source/Runtime/Engine/Private/Particles/ParticleComponents.cpp b/Engine/Source/Runtime/Engine/Private/Particles/ParticleComponents.cpp
--- a/Engine/Source/Runtime/Engine/Private/Particles/ParticleComponents.cpp
+++ b/Engine/Source/Runtime/Engine/Private/Particles/ParticleComponents.cpp
@@ -3719,7 +3719,7 @@ void UParticleSystemComponent::Serialize( FArchive& Ar )
 
 void UParticleSystemComponent::BeginDestroy()
 {
-	ForceAsyncWorkCompletion(ENSURE_AND_STALL);
+	ForceAsyncWorkCompletion(ENSURE_AND_STALL, true, true);
 	Super::BeginDestroy();
 
 	if (PoolingMethod == EPSCPoolMethod::AutoRelease || PoolingMethod == EPSCPoolMethod::ManualRelease)
@@ -3738,7 +3738,7 @@ void UParticleSystemComponent::BeginDestroy()
 
 void UParticleSystemComponent::FinishDestroy()
 {
-	ForceAsyncWorkCompletion(ENSURE_AND_STALL);
+	ForceAsyncWorkCompletion(ENSURE_AND_STALL, true, true);
 	for (int32 EmitterIndex = 0; EmitterIndex < EmitterInstances.Num(); EmitterIndex++)
 	{
 		FParticleEmitterInstance* EmitInst = EmitterInstances[EmitterIndex];
@@ -3896,7 +3896,7 @@ void UParticleSystemComponent::CreateRenderState_Concurrent()
 	SCOPE_CYCLE_COUNTER(STAT_ParticleSystemComponent_CreateRenderState_Concurrent);
 	SCOPE_CYCLE_COUNTER(STAT_ParticlesOverview_GT_CNC);
 
-	ForceAsyncWorkCompletion(ENSURE_AND_STALL);
+	ForceAsyncWorkCompletion(ENSURE_AND_STALL, false, true);
 	check( GetWorld() );
 	UE_LOG(LogParticles,Verbose,
 		TEXT("CreateRenderState_Concurrent @ %fs %s"), GetWorld()->TimeSeconds,
@@ -3931,7 +3931,7 @@ void UParticleSystemComponent::SendRenderTransform_Concurrent()
 	SCOPE_CYCLE_COUNTER(STAT_ParticleSystemComponent_SendRenderTransform_Concurrent);
 	SCOPE_CYCLE_COUNTER(STAT_ParticlesOverview_GT_CNC);
 
-	ForceAsyncWorkCompletion(ENSURE_AND_STALL);
+	ForceAsyncWorkCompletion(ENSURE_AND_STALL, false, true);
 	if (bIsActive)
 	{
 		if (bSkipUpdateDynamicDataDuringTick == false)
@@ -3949,7 +3949,7 @@ void UParticleSystemComponent::SendRenderDynamicData_Concurrent()
 	SCOPE_CYCLE_COUNTER(STAT_ParticleSystemComponent_SendRenderDynamicData_Concurrent);
 	SCOPE_CYCLE_COUNTER(STAT_ParticlesOverview_GT_CNC);
 
-	ForceAsyncWorkCompletion(ENSURE_AND_STALL);
+	ForceAsyncWorkCompletion(ENSURE_AND_STALL, false, true);
 	Super::SendRenderDynamicData_Concurrent();
 
 	check(!bAsyncDataCopyIsValid);
@@ -3984,12 +3984,7 @@ void UParticleSystemComponent::DestroyRenderState_Concurrent()
 	SCOPE_CYCLE_COUNTER(STAT_ParticleSystemComponent_DestroyRenderState_Concurrent);
 	SCOPE_CYCLE_COUNTER(STAT_ParticlesOverview_GT_CNC);
 
-	{
-		const bool bSaveSkipUpdate = bSkipUpdateDynamicDataDuringTick;
-		bSkipUpdateDynamicDataDuringTick = true;
-		ForceAsyncWorkCompletion(ENSURE_AND_STALL);
-		bSkipUpdateDynamicDataDuringTick = bSaveSkipUpdate;
-	}
+	ForceAsyncWorkCompletion(ENSURE_AND_STALL, false, true);
 
 	check( GetWorld() );
 	UE_LOG(LogParticles,Verbose,
