[NO NON-PII COMMIT MESSAGE]

diff --git a/Engine/Source/Runtime/Engine/Private/SoundCue.cpp b/Engine/Source/Runtime/Engine/Private/SoundCue.cpp
--- a/Engine/Source/Runtime/Engine/Private/SoundCue.cpp
+++ b/Engine/Source/Runtime/Engine/Private/SoundCue.cpp
@@ -79,7 +79,6 @@ void USoundCue::CacheAggregateValues()
 		Duration = FirstNode->GetDuration();
 
 		MaxDistance = FindMaxDistanceInternal();
-
 		bHasDelayNode = FirstNode->HasDelayNode();
 		bHasConcatenatorNode = FirstNode->HasConcatenatorNode();
 		bHasPlayWhenSilent = FirstNode->IsPlayWhenSilent();
@@ -229,7 +228,7 @@ float USoundCue::FindMaxDistanceInternal() const
 		OutMaxDistance = FMath::Max(OutMaxDistance, FirstNode->GetMaxDistance());
 	}
 
-	if (OutMaxDistance >= 0.0f)
+	if (OutMaxDistance > KINDA_SMALL_NUMBER)
 	{
 		return OutMaxDistance;
 	}
