#ue4 - Adjusted logging for ForcePositionUpdate() and expired timestamps in ServerMove().

diff --git a/Engine/Source/Runtime/Engine/Private/Components/CharacterMovementComponent.cpp b/Engine/Source/Runtime/Engine/Private/Components/CharacterMovementComponent.cpp
--- a/Engine/Source/Runtime/Engine/Private/Components/CharacterMovementComponent.cpp
+++ b/Engine/Source/Runtime/Engine/Private/Components/CharacterMovementComponent.cpp
@@ -204,6 +204,13 @@ namespace CharacterMovementCVars
 		TEXT("If 1, force a jump substep to always reach the peak position of a jump, which can often be cut off as framerate lowers."),
 		ECVF_Default);
 
+	static float NetServerMoveTimestampExpiredWarningThreshold = 1.0f;
+	FAutoConsoleVariableRef CVarNetServerMoveTimestampExpiredWarningThreshold(
+		TEXT("net.NetServerMoveTimestampExpiredWarningThreshold"),
+		NetServerMoveTimestampExpiredWarningThreshold,
+		TEXT("Tolerance for ServerMove() to warn when client moves are expired more than this time threshold behind the server."),
+		ECVF_Default);
+
 #if !UE_BUILD_SHIPPING
 
 	int32 NetShowCorrections = 0;
@@ -222,7 +229,7 @@ namespace CharacterMovementCVars
 		TEXT("Time in seconds each visualized network correction persists."),
 		ECVF_Cheat);
 
-#endif // !UI_BUILD_SHIPPING
+#endif // !UE_BUILD_SHIPPING
 
 
 #if !(UE_BUILD_SHIPPING || UE_BUILD_TEST)
@@ -7769,13 +7776,11 @@ bool UCharacterMovementComponent::ForcePositionUpdate(float DeltaTime)
 	const double SavedServerTimestamp = ServerData->ServerAccumulatedClientTimeStamp;
 	ServerData->ServerAccumulatedClientTimeStamp += DeltaTime;
 
-#if !(UE_BUILD_SHIPPING)
 	const bool bServerMoveHasOccurred = (ServerData->ServerTimeStampLastServerMove != 0.f);
 	if (bServerMoveHasOccurred)
 	{
-		UE_LOG(LogNetPlayerMovement, Log, TEXT("ForcePositionUpdate %s (DeltaTime %.2f)"), *CharacterOwner->GetName(), DeltaTime);
+		UE_LOG(LogNetPlayerMovement, Log, TEXT("ForcePositionUpdate: %s (DeltaTime %.2f -> ServerTimeStamp %.2f)"), *GetNameSafe(CharacterOwner), DeltaTime, ServerData->CurrentClientTimeStamp);
 	}
-#endif
 
 	// Force movement update.
 	PerformMovement(DeltaTime);
@@ -8212,7 +8217,7 @@ void UCharacterMovementComponent::ServerMoveOld_Implementation
 
 	if( !VerifyClientTimeStamp(OldTimeStamp, *ServerData) )
 	{
-		UE_LOG(LogNetPlayerMovement, VeryVerbose, TEXT("ServerMoveOld: TimeStamp expired. %f, CurrentTimeStamp: %f"), OldTimeStamp, ServerData->CurrentClientTimeStamp);
+		UE_LOG(LogNetPlayerMovement, VeryVerbose, TEXT("ServerMoveOld: TimeStamp expired. %f, CurrentTimeStamp: %f, Character: %s"), OldTimeStamp, ServerData->CurrentClientTimeStamp, *GetNameSafe(CharacterOwner));
 		return;
 	}
 
@@ -8606,7 +8611,16 @@ void UCharacterMovementComponent::ServerMove_Implementation(
 
 	if( !VerifyClientTimeStamp(TimeStamp, *ServerData) )
 	{
-		UE_LOG(LogNetPlayerMovement, Log, TEXT("ServerMove: TimeStamp expired. %f, CurrentTimeStamp: %f"), TimeStamp, ServerData->CurrentClientTimeStamp);
+		const float ServerTimeStamp = ServerData->CurrentClientTimeStamp;
+		// This is more severe if the timestamp has a large discrepancy and hasn't been recently reset.
+		if (ServerTimeStamp > 1.0f && FMath::Abs(ServerTimeStamp - TimeStamp) > CharacterMovementCVars::NetServerMoveTimestampExpiredWarningThreshold)
+		{
+			UE_LOG(LogNetPlayerMovement, Warning, TEXT("ServerMove: TimeStamp expired: %f, CurrentTimeStamp: %f, Character: %s"), TimeStamp, ServerTimeStamp, *GetNameSafe(CharacterOwner));
+		}
+		else
+		{
+			UE_LOG(LogNetPlayerMovement, Log, TEXT("ServerMove: TimeStamp expired: %f, CurrentTimeStamp: %f, Character: %s"), TimeStamp, ServerTimeStamp, *GetNameSafe(CharacterOwner));
+		}		
 		return;
 	}
 
