HasAnyShadingModel now takes array by ref instead of by value

#rnx

diff --git a/Engine/Source/Runtime/Engine/Classes/Engine/EngineTypes.h b/Engine/Source/Runtime/Engine/Classes/Engine/EngineTypes.h
--- a/Engine/Source/Runtime/Engine/Classes/Engine/EngineTypes.h
+++ b/Engine/Source/Runtime/Engine/Classes/Engine/EngineTypes.h
@@ -497,7 +497,7 @@ public:
 	void ClearShadingModels()												{ ShadingModelField = 0; }
 
 	// Check if any of the given shading models are present
-	bool HasAnyShadingModel(TArray<EMaterialShadingModel> InShadingModels) const	
+	bool HasAnyShadingModel(const TArray<EMaterialShadingModel>& InShadingModels) const	
 	{ 
 		for (EMaterialShadingModel ShadingModel : InShadingModels)
 		{
