Allow FShadowMap and FLightMap to be destroyed on the Garbage Collector async destruction thread

diff --git a/Engine/Source/Runtime/Engine/Public/LightMap.h b/Engine/Source/Runtime/Engine/Public/LightMap.h
--- a/Engine/Source/Runtime/Engine/Public/LightMap.h
+++ b/Engine/Source/Runtime/Engine/Public/LightMap.h
@@ -92,7 +92,7 @@ public:
 	}
 	void Release()
 	{
-		check(IsInGameThread() || IsInAsyncLoadingThread());
+		check(IsInGameThread() || IsInAsyncLoadingThread() || IsInGarbageCollectorThread());
 		checkSlow(NumRefs > 0);
 		if(--NumRefs == 0)
 		{
