Prevent replay crash when spectator controller is created on a dedicated server

diff --git a/Engine/Source/Runtime/Engine/Private/DemoNetDriver.cpp b/Engine/Source/Runtime/Engine/Private/DemoNetDriver.cpp
--- a/Engine/Source/Runtime/Engine/Private/DemoNetDriver.cpp
+++ b/Engine/Source/Runtime/Engine/Private/DemoNetDriver.cpp
@@ -5824,7 +5824,14 @@ void UDemoNetConnection::HandleClientPlayer(APlayerController* PC, UNetConnectio
 
 	int32 SavedNetSpeed = LocalPlayer ? LocalPlayer->CurrentNetSpeed : 0;
 
-	Super::HandleClientPlayer(PC, NetConnection);
+	if (LocalPlayer != nullptr)
+	{
+		Super::HandleClientPlayer(PC, NetConnection);
+	}
+	else
+	{
+		DemoDriver->RestoreConnectionPostScrub(PC, NetConnection);
+	}
 	
 	// Restore the netspeed if we're a local replay
 	if (DemoDriver->bIsLocalReplay && LocalPlayer)
