Expose Sleep/Constraint functions.

diff --git a/Engine/Source/Runtime/Engine/Public/Physics/PhysScene_PhysX.h b/Engine/Source/Runtime/Engine/Public/Physics/PhysScene_PhysX.h
--- a/Engine/Source/Runtime/Engine/Public/Physics/PhysScene_PhysX.h
+++ b/Engine/Source/Runtime/Engine/Public/Physics/PhysScene_PhysX.h
@@ -153,9 +153,9 @@ public:
 	void ClearPreSimKinematicUpdate(USkeletalMeshComponent* InSkelComp);
 
 	/** Pending constraint break events */
-	void AddPendingOnConstraintBreak(FConstraintInstance* ConstraintInstance);
+	ENGINE_API void AddPendingOnConstraintBreak(FConstraintInstance* ConstraintInstance);
 	/** Pending wake/sleep events */
-	void AddPendingSleepingEvent(FBodyInstance* BI, ESleepEvent SleepEventType);
+	ENGINE_API void AddPendingSleepingEvent(FBodyInstance* BI, ESleepEvent SleepEventType);
 
 	/** Gets the array of collision notifications, pending execution at the end of the physics engine run. */
 	TArray<FCollisionNotifyInfo>& GetPendingCollisionNotifies() { return PendingCollisionData.PendingCollisionNotifies; }
