DepthOfField engine showflag is no longer shipping-only

#rnx

diff --git a/Engine/Source/Runtime/Engine/Public/ShowFlagsValues.inl b/Engine/Source/Runtime/Engine/Public/ShowFlagsValues.inl
--- a/Engine/Source/Runtime/Engine/Public/ShowFlagsValues.inl
+++ b/Engine/Source/Runtime/Engine/Public/ShowFlagsValues.inl
@@ -69,7 +69,7 @@ SHOWFLAG_FIXED_IN_SHIPPING(1, ColorGrading, SFG_PostProcess, NSLOCTEXT("UnrealEd
 /** Visualize vector fields. */
 SHOWFLAG_FIXED_IN_SHIPPING(0, VectorFields, SFG_Developer, NSLOCTEXT("UnrealEd", "VectorFieldsSF", "Vector Fields"))
 /** Depth of Field */
-SHOWFLAG_FIXED_IN_SHIPPING(1, DepthOfField, SFG_PostProcess, NSLOCTEXT("UnrealEd", "DepthOfFieldSF", "Depth Of Field"))
+SHOWFLAG_ALWAYS_ACCESSIBLE(DepthOfField, SFG_PostProcess, NSLOCTEXT("UnrealEd", "DepthOfFieldSF", "Depth Of Field"))
 /** Highlight materials that indicate performance issues or show unrealistic materials */
 SHOWFLAG_FIXED_IN_SHIPPING(0, GBufferHints, SFG_Developer, NSLOCTEXT("UnrealEd", "GBufferHintsSF", "GBuffer Hints (material attributes)"))
 /** MotionBlur, for now only camera motion blur, for now SHOWFLAG_ALWAYS_ACCESSIBLE because it's exposed in SceneCapture */
