Enable static mesh LOD streaming on 3 platforms

diff --git a/Engine/Source/Runtime/Engine/Private/Streaming/AsyncTextureStreaming.cpp b/Engine/Source/Runtime/Engine/Private/Streaming/AsyncTextureStreaming.cpp
--- a/Engine/Source/Runtime/Engine/Private/Streaming/AsyncTextureStreaming.cpp
+++ b/Engine/Source/Runtime/Engine/Private/Streaming/AsyncTextureStreaming.cpp
@@ -327,6 +327,7 @@ void FRenderAssetStreamingMipCalcTask::UpdateBudgetedMips_Async(int64& MemoryUse
 	TArray<int32> PrioritizedRenderAssets;
 
 	int64 MemoryBudgeted = 0;
+	int64 MemoryUsedByNonTextures = 0;
 	MemoryUsed = 0;
 	TempMemoryUsed = 0;
 
@@ -335,7 +336,14 @@ void FRenderAssetStreamingMipCalcTask::UpdateBudgetedMips_Async(int64& MemoryUse
 		if (IsAborted()) break;
 
 		MemoryBudgeted += StreamingRenderAsset.UpdateRetentionPriority_Async();
-		MemoryUsed += StreamingRenderAsset.GetSize(StreamingRenderAsset.ResidentMips);
+		const int32 AssetMemUsed = StreamingRenderAsset.GetSize(StreamingRenderAsset.ResidentMips);
+		MemoryUsed += AssetMemUsed;
+
+		// TODO: Use RHI metrics
+		if (StreamingRenderAsset.RenderAssetType != FStreamingRenderAsset::AT_Texture)
+		{
+			MemoryUsedByNonTextures += AssetMemUsed;
+		}
 
 		if (StreamingRenderAsset.ResidentMips != StreamingRenderAsset.RequestedMips)
 		{
@@ -362,7 +370,7 @@ void FRenderAssetStreamingMipCalcTask::UpdateBudgetedMips_Async(int64& MemoryUse
 	}
 
 
-	const int64 NonStreamingRenderAssetMemory =  AllocatedMemory - MemoryUsed;
+	const int64 NonStreamingRenderAssetMemory =  AllocatedMemory - MemoryUsed + MemoryUsedByNonTextures;
 	int64 AvailableMemoryForStreaming = PoolSize - NonStreamingRenderAssetMemory - MemoryMargin;
 
 	// If the platform defines a max VRAM usage, check if the pool size must be reduced,
