Use a single size instead of width and height for RVT
Width and height are then inferred from RVT volume aspect ratio
Integrated from //UE4/Dev-Rendering
#rnx

diff --git a/Engine/Source/Editor/VirtualTexturingEditor/Classes/RuntimeVirtualTextureDetailsCustomization.h b/Engine/Source/Editor/VirtualTexturingEditor/Classes/RuntimeVirtualTextureDetailsCustomization.h
--- a/Engine/Source/Editor/VirtualTexturingEditor/Classes/RuntimeVirtualTextureDetailsCustomization.h
+++ b/Engine/Source/Editor/VirtualTexturingEditor/Classes/RuntimeVirtualTextureDetailsCustomization.h
@@ -26,8 +26,7 @@ protected:
 private:
 	class URuntimeVirtualTexture* VirtualTexture;
 
-	TSharedPtr<class STextBlock> WidthText;
-	TSharedPtr<class STextBlock> HeightText;
+	TSharedPtr<class STextBlock> SizeText;
 	TSharedPtr<class STextBlock> TileSizeText;
 	TSharedPtr<class STextBlock> TileBorderSizeText;
 	TSharedPtr<class STextBlock> RemoveLowMipsText;
diff --git a/Engine/Source/Editor/VirtualTexturingEditor/Private/RuntimeVirtualTextureDetailsCustomization.cpp b/Engine/Source/Editor/VirtualTexturingEditor/Private/RuntimeVirtualTextureDetailsCustomization.cpp
--- a/Engine/Source/Editor/VirtualTexturingEditor/Private/RuntimeVirtualTextureDetailsCustomization.cpp
+++ b/Engine/Source/Editor/VirtualTexturingEditor/Private/RuntimeVirtualTextureDetailsCustomization.cpp
@@ -92,8 +92,7 @@ void FRuntimeVirtualTextureDetailsCustomization::CustomizeDetails(IDetailLayoutB
 
 	// Add size helpers
 	IDetailCategoryBuilder& SizeCategory = DetailBuilder.EditCategory("Size", FText::GetEmpty());
-	AddTextToProperty(DetailBuilder, SizeCategory, "Width", WidthText);
-	AddTextToProperty(DetailBuilder, SizeCategory, "Height", HeightText);
+	AddTextToProperty(DetailBuilder, SizeCategory, "Size", SizeText);
 	AddTextToProperty(DetailBuilder, SizeCategory, "TileSize", TileSizeText);
 	AddTextToProperty(DetailBuilder, SizeCategory, "TileBorderSize", TileBorderSizeText);
 	AddTextToProperty(DetailBuilder, SizeCategory, "RemoveLowMips", RemoveLowMipsText);
@@ -124,8 +123,7 @@ void FRuntimeVirtualTextureDetailsCustomization::CustomizeDetails(IDetailLayoutB
 	];
 
 	// Add refresh callback for all properties 
-	DetailBuilder.GetProperty(FName(TEXT("Width")))->SetOnPropertyValueChanged(FSimpleDelegate::CreateSP(this, &FRuntimeVirtualTextureDetailsCustomization::RefreshDetails));
-	DetailBuilder.GetProperty(FName(TEXT("Height")))->SetOnPropertyValueChanged(FSimpleDelegate::CreateSP(this, &FRuntimeVirtualTextureDetailsCustomization::RefreshDetails));
+	DetailBuilder.GetProperty(FName(TEXT("Size")))->SetOnPropertyValueChanged(FSimpleDelegate::CreateSP(this, &FRuntimeVirtualTextureDetailsCustomization::RefreshDetails));
 	DetailBuilder.GetProperty(FName(TEXT("TileSize")))->SetOnPropertyValueChanged(FSimpleDelegate::CreateSP(this, &FRuntimeVirtualTextureDetailsCustomization::RefreshDetails));
 	DetailBuilder.GetProperty(FName(TEXT("TileBorderSize")))->SetOnPropertyValueChanged(FSimpleDelegate::CreateSP(this, &FRuntimeVirtualTextureDetailsCustomization::RefreshDetails));
 	DetailBuilder.GetProperty(FName(TEXT("MaterialType")))->SetOnPropertyValueChanged(FSimpleDelegate::CreateSP(this, &FRuntimeVirtualTextureDetailsCustomization::RefreshDetails));
@@ -142,8 +140,7 @@ void FRuntimeVirtualTextureDetailsCustomization::RefreshDetails()
 	SizeOptions.UseGrouping = false;
 	SizeOptions.MaximumFractionalDigits = 0;
 
- 	WidthText->SetText(FText::Format(LOCTEXT("Details_Number", "{0}"), FText::AsNumber(VirtualTexture->GetWidth(), &SizeOptions)));
- 	HeightText->SetText(FText::Format(LOCTEXT("Details_Number", "{0}"), FText::AsNumber(VirtualTexture->GetHeight(), &SizeOptions)));
+ 	SizeText->SetText(FText::Format(LOCTEXT("Details_Number", "{0}"), FText::AsNumber(VirtualTexture->GetSize(), &SizeOptions)));
  	TileSizeText->SetText(FText::Format(LOCTEXT("Details_Number", "{0}"), FText::AsNumber(VirtualTexture->GetTileSize(), &SizeOptions)));
  	TileBorderSizeText->SetText(FText::Format(LOCTEXT("Details_Number", "{0}"), FText::AsNumber(VirtualTexture->GetTileBorderSize(), &SizeOptions)));
  	RemoveLowMipsText->SetText(FText::Format(LOCTEXT("Details_Number", "{0}"), FText::AsNumber(VirtualTexture->GetRemoveLowMips(), &SizeOptions)));
diff --git a/Engine/Source/Runtime/Engine/Classes/VT/RuntimeVirtualTexture.h b/Engine/Source/Runtime/Engine/Classes/VT/RuntimeVirtualTexture.h
--- a/Engine/Source/Runtime/Engine/Classes/VT/RuntimeVirtualTexture.h
+++ b/Engine/Source/Runtime/Engine/Classes/VT/RuntimeVirtualTexture.h
@@ -24,13 +24,9 @@ protected:
 	UPROPERTY(EditAnywhere, BlueprintReadOnly, Category = Layout, meta = (DisplayName = "Enable BC texture compression"))
 	bool bCompressTextures = true;
 
-	/** Width of virtual texture. (Actual values increase in powers of 2) */
-	UPROPERTY(EditAnywhere, BlueprintReadOnly, Category = Size, meta = (UIMin = "0", UIMax = "8", DisplayName = "Width of the virtual texture"))
-	int32 Width = 6; // 65536
-
-	/** Height of virtual texture. (Actual values increase in powers of 2) */
-	UPROPERTY(EditAnywhere, BlueprintReadOnly, Category = Size, meta = (UIMin = "0", UIMax = "8", DisplayName = "Height of the virtual texture"))
-	int32 Height = 6; // 65536
+	/** Size of virtual texture along the largest axis. (Actual values increase in powers of 2) */
+	UPROPERTY(EditAnywhere, BlueprintReadOnly, Category = Size, meta = (UIMin = "0", UIMax = "8", DisplayName = "Size of the virtual texture"))
+	int32 Size = 6; // 65536
 
 	/** Page tile size. (Actual values increase in powers of 2) */
 	UPROPERTY(EditAnywhere, BlueprintReadOnly, Category = Size, meta = (UIMin = "0", UIMax = "4", DisplayName = "Size of each virtual texture tile"))
@@ -55,10 +51,8 @@ public:
 	/** Get the material set that this virtual texture stores. */
 	ERuntimeVirtualTextureMaterialType GetMaterialType() const { return MaterialType; }
 
-	/** Public getter for virtual texture width */
-	int32 GetWidth() const { return 1 << FMath::Clamp(Width + 10, 10, 18); }
-	/** Public getter for virtual texture height */
-	int32 GetHeight() const { return 1 << FMath::Clamp(Height + 10, 10, 18); }
+	/** Public getter for virtual texture size */
+	int32 GetSize() const { return 1 << FMath::Clamp(Size + 10, 10, 18); }
 	/** Public getter for virtual texture tile size */
 	int32 GetTileSize() const { return 1 << FMath::Clamp(TileSize + 6, 6, 10); }
 	/** Public getter for virtual texture tile border size */
@@ -71,13 +65,13 @@ public:
 	/** Returns an approximate estimated value for the memory used by the physical texture. */
 	int32 GetEstimatedPhysicalTextureMemoryKb() const;
 
-	/** Get virtual texture description based on the properties of this object. */
-	void GetProducerDescription(FVTProducerDescription& OutDesc) const;
+	/** Get virtual texture description based on the properties of this object and the passed in volume transform. */
+	void GetProducerDescription(FVTProducerDescription& OutDesc, FTransform const& VolumeToWorld) const;
 	/** Return true if the virtual texture layer should be sampled as sRGB */
 	bool IsLayerSRGB(int32 LayerIndex) const;
 
 	/** (Re)Initialize this object. Call this whenever we modify the producer or transform. */
-	void Initialize(IVirtualTexture* InProducer, FTransform const& BoxToWorld);
+	void Initialize(IVirtualTexture* InProducer, FTransform const& VolumeToWorld);
 
 	/** Release the resources for this object This will need to be called if our producer becomes stale and we aren't doing a full reinit with a new producer. */
 	void Release();
@@ -92,7 +86,7 @@ public:
 
 protected:
 	/** Initialize the render resources. This kicks off render thread work. */
-	void InitResource(IVirtualTexture* InProducer);
+	void InitResource(IVirtualTexture* InProducer, FTransform const& VolumeToWorld);
 	/** Initialize the render resources with a null producer. This kicks off render thread work. */
 	void InitNullResource();
 
diff --git a/Engine/Source/Runtime/Engine/Private/VT/RuntimeVirtualTexture.cpp b/Engine/Source/Runtime/Engine/Private/VT/RuntimeVirtualTexture.cpp
--- a/Engine/Source/Runtime/Engine/Private/VT/RuntimeVirtualTexture.cpp
+++ b/Engine/Source/Runtime/Engine/Private/VT/RuntimeVirtualTexture.cpp
@@ -175,19 +175,36 @@ URuntimeVirtualTexture::~URuntimeVirtualTexture()
 	delete Resource;
 }
 
-void URuntimeVirtualTexture::GetProducerDescription(FVTProducerDescription& OutDesc) const
+void URuntimeVirtualTexture::GetProducerDescription(FVTProducerDescription& OutDesc, FTransform const& VolumeToWorld) const
 {
 	OutDesc.Name = GetFName();
 	OutDesc.Dimensions = 2;
 	OutDesc.TileSize = GetTileSize();
 	OutDesc.TileBorderSize = GetTileBorderSize();
-	OutDesc.BlockWidthInTiles = GetWidth() / GetTileSize();
-	OutDesc.BlockHeightInTiles = GetHeight() / GetTileSize();
-	OutDesc.MaxLevel = FMath::Max((int32)FMath::CeilLogTwo(FMath::Max(OutDesc.BlockWidthInTiles, OutDesc.BlockHeightInTiles)) - RemoveLowMips, 1);
 	OutDesc.DepthInTiles = 1;
 	OutDesc.WidthInBlocks = 1;
 	OutDesc.HeightInBlocks = 1;
 
+	// Set width and height to best match the aspect ratio
+	const FVector VolumeSize = VolumeToWorld.GetScale3D();
+	const float AspectRatio = VolumeSize.X / VolumeSize.Y;
+	int32 Width, Height;
+	if (AspectRatio >= 1.f)
+	{
+		Width = GetSize();
+		Height = FMath::Max((int32)FMath::RoundUpToPowerOfTwo(Width / AspectRatio), GetTileSize());
+	}
+	else
+	{
+		Height = GetSize();
+		Width = FMath::Max((int32)FMath::RoundUpToPowerOfTwo(Height * AspectRatio), GetTileSize());
+	}
+
+	OutDesc.BlockWidthInTiles = Width / GetTileSize();
+	OutDesc.BlockHeightInTiles = Height / GetTileSize();
+	OutDesc.MaxLevel = FMath::Max((int32)FMath::CeilLogTwo(FMath::Max(OutDesc.BlockWidthInTiles, OutDesc.BlockHeightInTiles)) - RemoveLowMips, 1);
+
+	// Set layer description based on material type
 	switch (MaterialType)
 	{
 	case ERuntimeVirtualTextureMaterialType::BaseColor:
@@ -258,14 +275,14 @@ FVector4 URuntimeVirtualTexture::GetUniformParameter(int32 Index)
 	return WorldToUVTransformParameters[Index];
 }
 
-void URuntimeVirtualTexture::Initialize(IVirtualTexture* InProducer, FTransform const& BoxToWorld)
+void URuntimeVirtualTexture::Initialize(IVirtualTexture* InProducer, FTransform const& VolumeToWorld)
 {
 	//todo[vt]: possible issues with precision in large worlds here it might be better to calculate/upload camera space relative transform per frame?
-	WorldToUVTransformParameters[0] = BoxToWorld.GetTranslation();
-	WorldToUVTransformParameters[1] = BoxToWorld.GetUnitAxis(EAxis::X) * 1.f / BoxToWorld.GetScale3D().X;
-	WorldToUVTransformParameters[2] = BoxToWorld.GetUnitAxis(EAxis::Y) * 1.f / BoxToWorld.GetScale3D().Y;
+	WorldToUVTransformParameters[0] = VolumeToWorld.GetTranslation();
+	WorldToUVTransformParameters[1] = VolumeToWorld.GetUnitAxis(EAxis::X) * 1.f / VolumeToWorld.GetScale3D().X;
+	WorldToUVTransformParameters[2] = VolumeToWorld.GetUnitAxis(EAxis::Y) * 1.f / VolumeToWorld.GetScale3D().Y;
 
-	InitResource(InProducer);
+	InitResource(InProducer, VolumeToWorld);
 }
 
 void URuntimeVirtualTexture::Release()
@@ -273,10 +290,10 @@ void URuntimeVirtualTexture::Release()
 	InitNullResource();
 }
 
-void URuntimeVirtualTexture::InitResource(IVirtualTexture* InProducer)
+void URuntimeVirtualTexture::InitResource(IVirtualTexture* InProducer, FTransform const& VolumeToWorld)
 {
 	FVTProducerDescription Desc;
-	GetProducerDescription(Desc);
+	GetProducerDescription(Desc, VolumeToWorld);
 	Resource->Init(Desc, InProducer);
 }
 
@@ -292,8 +309,7 @@ void URuntimeVirtualTexture::GetAssetRegistryTags(TArray<FAssetRegistryTag>& Out
 {
 	Super::GetAssetRegistryTags(OutTags);
 
-	OutTags.Add(FAssetRegistryTag("Width", FString::FromInt(GetWidth()), FAssetRegistryTag::TT_Numerical));
-	OutTags.Add(FAssetRegistryTag("Height", FString::FromInt(GetHeight()), FAssetRegistryTag::TT_Numerical));
+	OutTags.Add(FAssetRegistryTag("Size", FString::FromInt(GetSize()), FAssetRegistryTag::TT_Numerical));
 	OutTags.Add(FAssetRegistryTag("TileSize", FString::FromInt(GetTileSize()), FAssetRegistryTag::TT_Numerical));
 	OutTags.Add(FAssetRegistryTag("TileBorderSize", FString::FromInt(GetTileBorderSize()), FAssetRegistryTag::TT_Numerical));
 }
diff --git a/Engine/Source/Runtime/Renderer/Private/VT/RuntimeVirtualTextureSceneProxy.cpp b/Engine/Source/Runtime/Renderer/Private/VT/RuntimeVirtualTextureSceneProxy.cpp
--- a/Engine/Source/Runtime/Renderer/Private/VT/RuntimeVirtualTextureSceneProxy.cpp
+++ b/Engine/Source/Runtime/Renderer/Private/VT/RuntimeVirtualTextureSceneProxy.cpp
@@ -14,9 +14,6 @@ FRuntimeVirtualTextureSceneProxy::FRuntimeVirtualTextureSceneProxy(URuntimeVirtu
 {
 	if (VirtualTexture != nullptr)
 	{
-		FVTProducerDescription Desc;
-		VirtualTexture->GetProducerDescription(Desc);
-
 		// We store a ProducerId here so that we will be able to find our SceneIndex from the Producer during rendering.
 		// We will need the SceneIndex to determine which primitives should render to this Producer.
 		ProducerId = ProducerIdGenerator++;
@@ -26,8 +23,12 @@ FRuntimeVirtualTextureSceneProxy::FRuntimeVirtualTextureSceneProxy(URuntimeVirtu
 		// Transform is based on bottom left of the URuntimeVirtualTextureComponent unit box (which is centered on the origin)
 		FTransform Transform = FTransform(FVector(-0.5f, -0.5f, 0.f)) * InComponent->GetComponentTransform();
 
-		// The Producer object created here will be passed into the Virtual Texture system which will take ownership.
-		// The Initialize() call will allocate the VT by spawning work on the render thread.
+		// The producer description is calculated using the transform to determine the aspect ratio
+		FVTProducerDescription Desc;
+		VirtualTexture->GetProducerDescription(Desc, Transform);
+
+		// The Producer object created here will be passed into the virtual texture system which will take ownership.
+		// The Initialize() call will allocate the virtual texture by spawning work on the render thread.
 		FRuntimeVirtualTextureProducer* Producer = new FRuntimeVirtualTextureProducer(Desc, ProducerId, MaterialType, InComponent->GetScene(), Transform);
 		VirtualTexture->Initialize(Producer, Transform);
 	}
