Remove unused height pins from RVT material nodes
Change default name of RuntimeVirtualTextureVolume actor
Integrated from //UE4/Dev-Rendering
#rnx

diff --git a/Engine/Shaders/Private/VirtualTextureMaterial.usf b/Engine/Shaders/Private/VirtualTextureMaterial.usf
--- a/Engine/Shaders/Private/VirtualTextureMaterial.usf
+++ b/Engine/Shaders/Private/VirtualTextureMaterial.usf
@@ -75,8 +75,7 @@ void FPixelShaderInOut_MainPS(
 	half Specular = GetVirtualTextureOutput1(MaterialParameters);
 	half Roughness = GetVirtualTextureOutput2(MaterialParameters);
 	half3 Normal = GetVirtualTextureOutput3(MaterialParameters);
-	float Height = GetVirtualTextureOutput4(MaterialParameters);
-	float Opacity = GetVirtualTextureOutput5(MaterialParameters);
+	float Opacity = GetVirtualTextureOutput4(MaterialParameters);
 #else
 	// Output is from standard material output attribute node
 	half3 BaseColor = GetMaterialBaseColor(PixelMaterialInputs);
diff --git a/Engine/Source/Editor/UnrealEd/Private/Factories/ActorFactory.cpp b/Engine/Source/Editor/UnrealEd/Private/Factories/ActorFactory.cpp
--- a/Engine/Source/Editor/UnrealEd/Private/Factories/ActorFactory.cpp
+++ b/Engine/Source/Editor/UnrealEd/Private/Factories/ActorFactory.cpp
@@ -1902,7 +1902,7 @@ UActorFactoryRuntimeVirtualTextureVolume::UActorFactoryRuntimeVirtualTextureVolu
 
 void UActorFactoryRuntimeVirtualTextureVolume::PostSpawnActor(UObject* Asset, AActor* NewActor)
 {
-	FText ActorName = LOCTEXT("VirtualTextureVolume_DefaultActorName", "Virtual Texture");
+	FText ActorName = LOCTEXT("VirtualTextureVolume_DefaultActorName", "Virtual Texture Volume");
 	NewActor->SetActorLabel(ActorName.ToString());
 
 	// Good default size to see object in editor
diff --git a/Engine/Source/Runtime/Engine/Classes/Materials/MaterialExpressionRuntimeVirtualTextureOutput.h b/Engine/Source/Runtime/Engine/Classes/Materials/MaterialExpressionRuntimeVirtualTextureOutput.h
--- a/Engine/Source/Runtime/Engine/Classes/Materials/MaterialExpressionRuntimeVirtualTextureOutput.h
+++ b/Engine/Source/Runtime/Engine/Classes/Materials/MaterialExpressionRuntimeVirtualTextureOutput.h
@@ -29,10 +29,6 @@ class UMaterialExpressionRuntimeVirtualTextureOutput : public UMaterialExpressio
 	UPROPERTY()
 	FExpressionInput Normal;
 
-	/** Input for Height to output to virtual texture. */
-	UPROPERTY()
-	FExpressionInput Height;
-
 	/** Input for Opacity value used for blending to virtual texture. */
 	UPROPERTY()
 	FExpressionInput Opacity;
diff --git a/Engine/Source/Runtime/Engine/Private/Materials/MaterialExpressions.cpp b/Engine/Source/Runtime/Engine/Private/Materials/MaterialExpressions.cpp
--- a/Engine/Source/Runtime/Engine/Private/Materials/MaterialExpressions.cpp
+++ b/Engine/Source/Runtime/Engine/Private/Materials/MaterialExpressions.cpp
@@ -1918,10 +1918,6 @@ int32 UMaterialExpressionRuntimeVirtualTextureOutput::Compile(class FMaterialCom
 		CodeInput = Normal.IsConnected() ? Normal.Compile(Compiler) : Compiler->Constant3(0.f, 0.f, 1.f);
 	}
 	else if (OutputIndex == 4)
-	{
-		CodeInput = Height.IsConnected() ? Height.Compile(Compiler) : Compiler->Constant(0.f);
-	}
-	else if (OutputIndex == 5)
 	{
 		CodeInput = Opacity.IsConnected() ? Opacity.Compile(Compiler) : Compiler->Constant(1.f);
 	}
@@ -1939,7 +1935,7 @@ void UMaterialExpressionRuntimeVirtualTextureOutput::GetCaption(TArray<FString>&
 
 int32 UMaterialExpressionRuntimeVirtualTextureOutput::GetNumOutputs() const
 {
-	return 6; 
+	return 5; 
 }
 
 FString UMaterialExpressionRuntimeVirtualTextureOutput::GetFunctionName() const
@@ -1994,7 +1990,6 @@ void UMaterialExpressionRuntimeVirtualTextureSample::InitOutputs()
 	Outputs.Add(FExpressionOutput(TEXT("Specular"), 1, 0, 0, 1, 0));
 	Outputs.Add(FExpressionOutput(TEXT("Roughness"), 1, 1, 0, 0, 0));
 	Outputs.Add(FExpressionOutput(TEXT("Normal")));
-	Outputs.Add(FExpressionOutput(TEXT("Height")));
 #endif // WITH_EDITORONLY_DATA
 }
 
@@ -2089,9 +2084,6 @@ int32 UMaterialExpressionRuntimeVirtualTextureSample::Compile(class FMaterialCom
 			return Compiler->Constant3(0.f, 0.f, 1.f);
 		}
 		break;
-	case 4:
-		//todo[vt]: Implement virtual texture displacement map support
-		return Compiler->Constant(0.f);
 	default:
 		return INDEX_NONE;
 	}
