Expose GeomSweepMulti and GeomOverlapMulti

diff --git a/Engine/Source/Runtime/Engine/Public/Physics/PhysicsInterfacePhysX.h b/Engine/Source/Runtime/Engine/Public/Physics/PhysicsInterfacePhysX.h
--- a/Engine/Source/Runtime/Engine/Public/Physics/PhysicsInterfacePhysX.h
+++ b/Engine/Source/Runtime/Engine/Public/Physics/PhysicsInterfacePhysX.h
@@ -411,10 +411,10 @@ template <>
 ENGINE_API int32 FPhysicsInterface_PhysX::GetAllShapes_AssumedLocked(const FPhysicsActorHandle_PhysX& InActorHandle, PhysicsInterfaceTypes::FInlineShapeArray& OutShapes);
 
 template<>
-bool FGenericPhysicsInterface::GeomSweepMulti(const UWorld* World, const FPhysicsGeometryCollection& InGeom, const FQuat& InGeomRot, TArray<FHitResult>& OutHits, FVector Start, FVector End, ECollisionChannel TraceChannel, const FCollisionQueryParams& Params, const FCollisionResponseParams& ResponseParams, const FCollisionObjectQueryParams& ObjectParams /*= FCollisionObjectQueryParams::DefaultObjectQueryParam*/);
+ENGINE_API bool FGenericPhysicsInterface::GeomSweepMulti(const UWorld* World, const FPhysicsGeometryCollection& InGeom, const FQuat& InGeomRot, TArray<FHitResult>& OutHits, FVector Start, FVector End, ECollisionChannel TraceChannel, const FCollisionQueryParams& Params, const FCollisionResponseParams& ResponseParams, const FCollisionObjectQueryParams& ObjectParams /*= FCollisionObjectQueryParams::DefaultObjectQueryParam*/);
 
 template<>
-bool FGenericPhysicsInterface::GeomOverlapMulti(const UWorld* World, const FPhysicsGeometryCollection& InGeom, const FVector& InPosition, const FQuat& InRotation, TArray<FOverlapResult>& OutOverlaps, ECollisionChannel TraceChannel, const FCollisionQueryParams& Params, const FCollisionResponseParams& ResponseParams, const FCollisionObjectQueryParams& ObjectParams);
+ENGINE_API bool FGenericPhysicsInterface::GeomOverlapMulti(const UWorld* World, const FPhysicsGeometryCollection& InGeom, const FVector& InPosition, const FQuat& InRotation, TArray<FOverlapResult>& OutOverlaps, ECollisionChannel TraceChannel, const FCollisionQueryParams& Params, const FCollisionResponseParams& ResponseParams, const FCollisionObjectQueryParams& ObjectParams);
 
 
 #endif
\ No newline at end of file
