Fix clearning up preview physics scene.

diff --git a/Engine/Source/Runtime/Engine/Classes/Engine/World.h b/Engine/Source/Runtime/Engine/Classes/Engine/World.h
--- a/Engine/Source/Runtime/Engine/Classes/Engine/World.h
+++ b/Engine/Source/Runtime/Engine/Classes/Engine/World.h
@@ -2464,6 +2464,8 @@ public:
 	 */
 	void UpdateLevelStreaming();
 
+	/** Releases PhysicsScene manually */
+	void ReleasePhysicsScene();
 public:
 	/**
 	 * Flushes level streaming in blocking fashion and returns when all levels are loaded/ visible/ hidden
diff --git a/Engine/Source/Runtime/Engine/Private/PreviewScene.cpp b/Engine/Source/Runtime/Engine/Private/PreviewScene.cpp
--- a/Engine/Source/Runtime/Engine/Private/PreviewScene.cpp
+++ b/Engine/Source/Runtime/Engine/Private/PreviewScene.cpp
@@ -99,6 +99,10 @@ FPreviewScene::~FPreviewScene()
 	
 	PreviewWorld->CleanupWorld();
 	GEngine->DestroyWorldContext(GetWorld());
+
+	// Release PhysicsScene for fixing big fbx importing bug
+	PreviewWorld->ReleasePhysicsScene();
+
 }
 
 void FPreviewScene::AddComponent(UActorComponent* Component,const FTransform& LocalToWorld, bool bAttachToRoot /*= false*/)
diff --git a/Engine/Source/Runtime/Engine/Private/World.cpp b/Engine/Source/Runtime/Engine/Private/World.cpp
--- a/Engine/Source/Runtime/Engine/Private/World.cpp
+++ b/Engine/Source/Runtime/Engine/Private/World.cpp
@@ -821,6 +821,22 @@ void UWorld::BeginDestroy()
 	}
 }
 
+void UWorld::ReleasePhysicsScene()
+{
+	if (PhysicsScene)
+	{
+		delete PhysicsScene;
+		PhysicsScene = NULL;
+
+#if WITH_PHYSX
+		if (GPhysCommandHandler)
+		{
+			GPhysCommandHandler->Flush();
+		}
+#endif // WITH_PHYSX
+	}
+}
+
 void UWorld::FinishDestroy()
 {
 	// Avoid cleanup if the world hasn't been initialized, e.g., the default object or a world object that got loaded
@@ -842,18 +858,7 @@ void UWorld::FinishDestroy()
 			FXSystem = NULL;
 		}
 
-		if (PhysicsScene)
-		{
-			delete PhysicsScene;
-			PhysicsScene = NULL;
-
-#if WITH_PHYSX
-			if(GPhysCommandHandler)
-			{
-				GPhysCommandHandler->Flush();
-			}
-#endif // WITH_PHYSX
-		}
+		ReleasePhysicsScene();
 
 		if (Scene)
 		{
