Cherry-pick 6975160: Added optimization to cut the cost of gathering ISM/HISM

diff --git a/Engine/Source/Runtime/Engine/Private/InstancedStaticMesh.cpp b/Engine/Source/Runtime/Engine/Private/InstancedStaticMesh.cpp
--- a/Engine/Source/Runtime/Engine/Private/InstancedStaticMesh.cpp
+++ b/Engine/Source/Runtime/Engine/Private/InstancedStaticMesh.cpp
@@ -1027,6 +1027,28 @@ void FInstancedStaticMeshSceneProxy::GetDynamicRayTracingInstances(struct FRayTr
 	uint32 LOD = 0;
 	const int32 InstanceCount = InstancedRenderData.Component->PerInstanceSMData.Num();
 
+	//preallocate the worst-case to prevent an explosion of reallocs
+	//#dxr_todo: possibly track used instances and reserve based on previous behavior
+	OutRayTracingInstances.Reserve(InstanceCount);
+
+	//setup a 'template' for the instance first, so we aren't duplicating work
+	//#dxr_todo: when multiple LODs are used, template needs to be an array of templates, probably best initialized on-demand via a lamda
+	FRayTracingInstance RayTracingInstanceTemplate;
+	RayTracingInstanceTemplate.Geometry = &RenderData->LODResources[LOD].RayTracingGeometry;
+
+	int32 SectionCount = InstancedRenderData.LODModels[LOD].Sections.Num();
+
+	for (int32 SectionIdx = 0; SectionIdx < SectionCount; ++SectionIdx)
+	{
+		//#dxr_todo: so far we use the parent static mesh path to get material data
+		FMeshBatch MeshBatch;
+		FStaticMeshSceneProxy::GetMeshElement(LOD, 0, SectionIdx, 0, false, false, MeshBatch);
+
+		RayTracingInstanceTemplate.Materials.Add(MeshBatch);
+	}
+	RayTracingInstanceTemplate.BuildInstanceMaskAndFlags();
+	RayTracingInstanceTemplate.InstanceTransforms.AddZeroed(1);
+
 	if (CVarRayTracingRenderInstancesCulling.GetValueOnRenderThread() > 0 && RayTracingCullClusterInstances.Num() > 0)
 	{
 		const float BVHCullRadius = CVarRayTracingInstancesCullClusterRadius.GetValueOnRenderThread();
@@ -1086,23 +1108,8 @@ void FInstancedStaticMeshSceneProxy::GetDynamicRayTracingInstances(struct FRayTr
 							continue;
 					}
 
-					FRayTracingInstance RayTracingInstance;
-					RayTracingInstance.Geometry = &RenderData->LODResources[LOD].RayTracingGeometry;
-					RayTracingInstance.InstanceTransforms.Add(InstanceTransform);
-
-					int32 SectionCount = InstancedRenderData.LODModels[LOD].Sections.Num();
-
-					for (int32 SectionIdx = 0; SectionIdx < SectionCount; ++SectionIdx)
-					{
-						//#dxr_todo: so far we use the parent static mesh path to get material data
-						FMeshBatch MeshBatch;
-						FStaticMeshSceneProxy::GetMeshElement(LOD, 0, SectionIdx, 0, false, false, MeshBatch);
-
-						RayTracingInstance.Materials.Add(MeshBatch);
-					}
-
-					RayTracingInstance.BuildInstanceMaskAndFlags();
-					OutRayTracingInstances.Add(RayTracingInstance);
+					OutRayTracingInstances.Emplace(RayTracingInstanceTemplate);
+					OutRayTracingInstances.Last().InstanceTransforms[0] = InstanceTransform;
 				}
 			}
 		}
@@ -1118,23 +1125,8 @@ void FInstancedStaticMeshSceneProxy::GetDynamicRayTracingInstances(struct FRayTr
 				FMatrix ComponentLocalToWorld = InstancedRenderData.Component->GetComponentTransform().ToMatrixWithScale();
 				FMatrix InstanceTransform = InstanceData.Transform * ComponentLocalToWorld;
 
-				FRayTracingInstance RayTracingInstance;
-				RayTracingInstance.Geometry = &RenderData->LODResources[LOD].RayTracingGeometry;
-				RayTracingInstance.InstanceTransforms.Add(InstanceTransform);
-
-				int32 SectionCount = InstancedRenderData.LODModels[LOD].Sections.Num();
-
-				for (int32 SectionIdx = 0; SectionIdx < SectionCount; ++SectionIdx)
-				{
-					//#dxr_todo: so far we use the parent static mesh path to get material data
-					FMeshBatch MeshBatch;
-					FStaticMeshSceneProxy::GetMeshElement(LOD, 0, SectionIdx, 0, false, false, MeshBatch);
-
-					RayTracingInstance.Materials.Add(MeshBatch);
-				}
-
-				RayTracingInstance.BuildInstanceMaskAndFlags();
-				OutRayTracingInstances.Add(RayTracingInstance);
+				OutRayTracingInstances.Emplace(RayTracingInstanceTemplate);
+				OutRayTracingInstances.Last().InstanceTransforms[0] = InstanceTransform;
 			}
 		}
 	}
