Cherry-pick 6988376: Implemented support for multiple instance transforms within one FRayTracingInstance. This is now used in FInstancedStaticMeshSceneProxy::GetDynamicRayTracingInstances(), which can generate thousands of instances.

This gives a significant CPU performance improvement in scenes with many instanced static meshes.

diff --git a/Engine/Source/Runtime/Engine/Private/InstancedStaticMesh.cpp b/Engine/Source/Runtime/Engine/Private/InstancedStaticMesh.cpp
--- a/Engine/Source/Runtime/Engine/Private/InstancedStaticMesh.cpp
+++ b/Engine/Source/Runtime/Engine/Private/InstancedStaticMesh.cpp
@@ -1027,15 +1027,20 @@ void FInstancedStaticMeshSceneProxy::GetDynamicRayTracingInstances(struct FRayTr
 	uint32 LOD = 0;
 	const int32 InstanceCount = InstancedRenderData.Component->PerInstanceSMData.Num();
 
-	//preallocate the worst-case to prevent an explosion of reallocs
-	//#dxr_todo: possibly track used instances and reserve based on previous behavior
-	OutRayTracingInstances.Reserve(InstanceCount);
-
+	if (InstanceCount == 0)
+	{
+		return;
+	}
 	//setup a 'template' for the instance first, so we aren't duplicating work
 	//#dxr_todo: when multiple LODs are used, template needs to be an array of templates, probably best initialized on-demand via a lamda
 	FRayTracingInstance RayTracingInstanceTemplate;
 	RayTracingInstanceTemplate.Geometry = &RenderData->LODResources[LOD].RayTracingGeometry;
 
+	//preallocate the worst-case to prevent an explosion of reallocs
+	//#dxr_todo: possibly track used instances and reserve based on previous behavior
+	RayTracingInstanceTemplate.InstanceTransforms.Reserve(InstanceCount);
+
+
 	int32 SectionCount = InstancedRenderData.LODModels[LOD].Sections.Num();
 
 	for (int32 SectionIdx = 0; SectionIdx < SectionCount; ++SectionIdx)
@@ -1047,7 +1052,6 @@ void FInstancedStaticMeshSceneProxy::GetDynamicRayTracingInstances(struct FRayTr
 		RayTracingInstanceTemplate.Materials.Add(MeshBatch);
 	}
 	RayTracingInstanceTemplate.BuildInstanceMaskAndFlags();
-	RayTracingInstanceTemplate.InstanceTransforms.AddZeroed(1);
 
 	if (CVarRayTracingRenderInstancesCulling.GetValueOnRenderThread() > 0 && RayTracingCullClusterInstances.Num() > 0)
 	{
@@ -1108,8 +1112,7 @@ void FInstancedStaticMeshSceneProxy::GetDynamicRayTracingInstances(struct FRayTr
 							continue;
 					}
 
-					OutRayTracingInstances.Emplace(RayTracingInstanceTemplate);
-					OutRayTracingInstances.Last().InstanceTransforms[0] = InstanceTransform;
+					RayTracingInstanceTemplate.InstanceTransforms.Add(InstanceTransform);
 				}
 			}
 		}
@@ -1125,11 +1128,12 @@ void FInstancedStaticMeshSceneProxy::GetDynamicRayTracingInstances(struct FRayTr
 				FMatrix ComponentLocalToWorld = InstancedRenderData.Component->GetComponentTransform().ToMatrixWithScale();
 				FMatrix InstanceTransform = InstanceData.Transform * ComponentLocalToWorld;
 
-				OutRayTracingInstances.Emplace(RayTracingInstanceTemplate);
-				OutRayTracingInstances.Last().InstanceTransforms[0] = InstanceTransform;
+				RayTracingInstanceTemplate.InstanceTransforms.Add(InstanceTransform);
 			}
 		}
 	}
+
+	OutRayTracingInstances.Add(RayTracingInstanceTemplate);
 }
 
 void FInstancedStaticMeshSceneProxy::SetupRayTracingCullClusters()
diff --git a/Engine/Source/Runtime/Renderer/Private/DeferredShadingRenderer.cpp b/Engine/Source/Runtime/Renderer/Private/DeferredShadingRenderer.cpp
--- a/Engine/Source/Runtime/Renderer/Private/DeferredShadingRenderer.cpp
+++ b/Engine/Source/Runtime/Renderer/Private/DeferredShadingRenderer.cpp
@@ -724,28 +724,31 @@ bool FDeferredShadingSceneRenderer::GatherRayTracingWorldInstances(FRHICommandLi
 					for (FRayTracingInstance& Instance : RayTracingInstances)
 					{
 						FRayTracingGeometryInstance RayTracingInstance = { Instance.Geometry->RayTracingGeometryRHI };
-						RayTracingInstance.Transform = Instance.InstanceTransforms[0];
-						ensureMsgf(Instance.InstanceTransforms.Num() == 1, TEXT("Multi-instancing hasn't been supported"));
 						RayTracingInstance.UserData = (uint32)PrimitiveIndex;
 						RayTracingInstance.Mask = Instance.Mask;
 						RayTracingInstance.bForceOpaque = Instance.bForceOpaque;
 
 						check(Instance.Materials.Num() == Instance.Geometry->Initializer.Segments.Num() || (Instance.Geometry->Initializer.Segments.Num() == 0 && Instance.Materials.Num() == 1));
 
-						uint32 InstanceIndex = ReferenceView.RayTracingGeometryInstances.Add(RayTracingInstance);
-
-						for (int32 ViewIndex = 1; ViewIndex < Views.Num(); ViewIndex++)
+						for (const FMatrix& InstanceTransform : Instance.InstanceTransforms)
 						{
-							Views[ViewIndex].RayTracingGeometryInstances.Add(RayTracingInstance);
-						}
+							RayTracingInstance.Transform = InstanceTransform;
 
-						for (int32 SegmentIndex = 0; SegmentIndex < Instance.Materials.Num(); SegmentIndex++)
-						{
-							FMeshBatch& MeshBatch = Instance.Materials[SegmentIndex];
-							FDynamicRayTracingMeshCommandContext CommandContext(ReferenceView.DynamicRayTracingMeshCommandStorage, ReferenceView.VisibleRayTracingMeshCommands, SegmentIndex, InstanceIndex);
-							FRayTracingMeshProcessor RayTracingMeshProcessor(&CommandContext, Scene, &ReferenceView);
+							uint32 InstanceIndex = ReferenceView.RayTracingGeometryInstances.Add(RayTracingInstance);
 
-							RayTracingMeshProcessor.AddMeshBatch(MeshBatch, 1, SceneProxy);
+							for (int32 ViewIndex = 1; ViewIndex < Views.Num(); ViewIndex++)
+							{
+								Views[ViewIndex].RayTracingGeometryInstances.Add(RayTracingInstance);
+							}
+
+							for (int32 SegmentIndex = 0; SegmentIndex < Instance.Materials.Num(); SegmentIndex++)
+							{
+								FMeshBatch& MeshBatch = Instance.Materials[SegmentIndex];
+								FDynamicRayTracingMeshCommandContext CommandContext(ReferenceView.DynamicRayTracingMeshCommandStorage, ReferenceView.VisibleRayTracingMeshCommands, SegmentIndex, InstanceIndex);
+								FRayTracingMeshProcessor RayTracingMeshProcessor(&CommandContext, Scene, &ReferenceView);
+
+								RayTracingMeshProcessor.AddMeshBatch(MeshBatch, 1, SceneProxy);
+							}
 						}
 					}
 				}
