Don't add body instance to ucs properites if not set.

diff --git a/Engine/Source/Editor/DetailCustomizations/Private/BodyInstanceCustomization.cpp b/Engine/Source/Editor/DetailCustomizations/Private/BodyInstanceCustomization.cpp
--- a/Engine/Source/Editor/DetailCustomizations/Private/BodyInstanceCustomization.cpp
+++ b/Engine/Source/Editor/DetailCustomizations/Private/BodyInstanceCustomization.cpp
@@ -1017,7 +1017,7 @@ bool FBodyInstanceCustomization::IsCollisionEnabled() const
 	bool bEnabled = false;
 	if(BodyInstanceHandle.IsValid())
 	{
-		bEnabled = !BodyInstanceHandle->IsEditConst() && FSlateApplication::Get().GetNormalExecutionAttribute().Get();
+		bEnabled = BodyInstanceHandle->IsEditable() && FSlateApplication::Get().GetNormalExecutionAttribute().Get();
 	}
 
 	return bEnabled;
diff --git a/Engine/Source/Runtime/Engine/Classes/Components/PrimitiveComponent.h b/Engine/Source/Runtime/Engine/Classes/Components/PrimitiveComponent.h
--- a/Engine/Source/Runtime/Engine/Classes/Components/PrimitiveComponent.h
+++ b/Engine/Source/Runtime/Engine/Classes/Components/PrimitiveComponent.h
@@ -896,7 +896,7 @@ public:
 	// Internal physics engine data.
 	
 	/** Physics scene information for this component, holds a single rigid body with multiple shapes. */
-	UPROPERTY(EditAnywhere, BlueprintReadOnly, Category=Collision, meta=(ShowOnlyInnerProperties))
+	UPROPERTY(EditAnywhere, BlueprintReadOnly, Category=Collision, meta=(ShowOnlyInnerProperties, SkipUCSModifiedProperties))
 	FBodyInstance BodyInstance;
 
 	/** 
diff --git a/Engine/Source/Runtime/Engine/Private/Components/ActorComponent.cpp b/Engine/Source/Runtime/Engine/Private/Components/ActorComponent.cpp
--- a/Engine/Source/Runtime/Engine/Private/Components/ActorComponent.cpp
+++ b/Engine/Source/Runtime/Engine/Private/Components/ActorComponent.cpp
@@ -1780,8 +1780,10 @@ void UActorComponent::DetermineUCSModifiedProperties()
 
 			virtual bool ShouldSkipProperty(const UProperty* InProperty) const override
 			{
+				static const FName MD_SkipUCSModifiedProperties(TEXT("SkipUCSModifiedProperties"));
 				return (    InProperty->HasAnyPropertyFlags(CPF_Transient)
-						|| !InProperty->HasAnyPropertyFlags(CPF_Edit | CPF_Interp));
+						|| !InProperty->HasAnyPropertyFlags(CPF_Edit | CPF_Interp)
+						|| InProperty->HasMetaData(MD_SkipUCSModifiedProperties));
 			}
 		} PropertySkipper;
 
