Disable runtime skin weight changes

diff --git a/Engine/Source/Runtime/Engine/Private/Animation/SkinWeightProfile.cpp b/Engine/Source/Runtime/Engine/Private/Animation/SkinWeightProfile.cpp
--- a/Engine/Source/Runtime/Engine/Private/Animation/SkinWeightProfile.cpp
+++ b/Engine/Source/Runtime/Engine/Private/Animation/SkinWeightProfile.cpp
@@ -15,7 +15,7 @@
 
 static void OnDefaultProfileCVarsChanged(IConsoleVariable* Variable)
 {
-	if (GSkinWeightProfilesLoadByDefaultMode != 1)
+	if (GSkinWeightProfilesLoadByDefaultMode >= 0)
 	{
 		const bool bClearBuffer = GSkinWeightProfilesLoadByDefaultMode == 2 || GSkinWeightProfilesLoadByDefaultMode == 0;
 		const bool bSetBuffer = GSkinWeightProfilesLoadByDefaultMode == 3;
@@ -49,11 +49,12 @@ static void OnDefaultProfileCVarsChanged(IConsoleVariable* Variable)
 	}
 }
 
-int32 GSkinWeightProfilesLoadByDefaultMode = 0;
+int32 GSkinWeightProfilesLoadByDefaultMode = -1;
 static FAutoConsoleVariableRef CVarSkinWeightsLoadByDefaultMode(
 	TEXT("a.SkinWeightProfile.LoadByDefaultMode"),
 	GSkinWeightProfilesLoadByDefaultMode,
 	TEXT("Enables/disables run-time optimization to override the original skin weights with a profile designated as the default to replace it. Can be used to optimize memory for specific platforms or devices")
+	TEXT("-1 = disabled")
 	TEXT("0 = static disabled")
 	TEXT("1 = static enabled")
 	TEXT("2 = dynamic disabled")
