Only check for metadata with editor.

diff --git a/Engine/Source/Runtime/Engine/Private/Components/ActorComponent.cpp b/Engine/Source/Runtime/Engine/Private/Components/ActorComponent.cpp
--- a/Engine/Source/Runtime/Engine/Private/Components/ActorComponent.cpp
+++ b/Engine/Source/Runtime/Engine/Private/Components/ActorComponent.cpp
@@ -1781,9 +1781,12 @@ void UActorComponent::DetermineUCSModifiedProperties()
 			virtual bool ShouldSkipProperty(const UProperty* InProperty) const override
 			{
 				static const FName MD_SkipUCSModifiedProperties(TEXT("SkipUCSModifiedProperties"));
-				return (    InProperty->HasAnyPropertyFlags(CPF_Transient)
-						|| !InProperty->HasAnyPropertyFlags(CPF_Edit | CPF_Interp)
-						|| InProperty->HasMetaData(MD_SkipUCSModifiedProperties));
+				return (InProperty->HasAnyPropertyFlags(CPF_Transient)
+					|| !InProperty->HasAnyPropertyFlags(CPF_Edit | CPF_Interp)
+#if WITH_EDITOR
+					|| InProperty->HasMetaData(MD_SkipUCSModifiedProperties)
+#endif
+					);
 			}
 		} PropertySkipper;
 
