When HMD display resolution is not zero, override with HMD resolution in FSceneViewport::ResizeFrame().

diff --git a/Engine/Source/Runtime/Engine/Private/Slate/SceneViewport.cpp b/Engine/Source/Runtime/Engine/Private/Slate/SceneViewport.cpp
--- a/Engine/Source/Runtime/Engine/Private/Slate/SceneViewport.cpp
+++ b/Engine/Source/Runtime/Engine/Private/Slate/SceneViewport.cpp
@@ -1334,7 +1334,7 @@ void FSceneViewport::ResizeFrame(uint32 NewWindowSizeX, uint32 NewWindowSizeY, E
 				IHeadMountedDisplay::MonitorInfo MonitorInfo;
 				if (GEngine->XRSystem.IsValid() && GEngine->XRSystem->GetHMDDevice() && GEngine->XRSystem->GetHMDDevice()->GetHMDMonitorInfo(MonitorInfo))
 				{
-					if (MonitorInfo.DesktopX > 0 || MonitorInfo.DesktopY > 0)
+					if (MonitorInfo.DesktopX > 0 || MonitorInfo.DesktopY > 0 || MonitorInfo.ResolutionX > 0 || MonitorInfo.ResolutionY > 0)
 					{
 						NewWindowSize.X = MonitorInfo.ResolutionX;
 						NewWindowSize.Y = MonitorInfo.ResolutionY;
