Fix RVT sampler state
Integrated from //UE4/Dev-Rendering
#rnx

diff --git a/Engine/Source/Runtime/Engine/Private/Materials/MaterialUniformExpressions.cpp b/Engine/Source/Runtime/Engine/Private/Materials/MaterialUniformExpressions.cpp
--- a/Engine/Source/Runtime/Engine/Private/Materials/MaterialUniformExpressions.cpp
+++ b/Engine/Source/Runtime/Engine/Private/Materials/MaterialUniformExpressions.cpp
@@ -826,7 +826,7 @@ void FUniformExpressionSet::FillUniformBuffer(const FMaterialRenderContext& Mate
 						if (PhysicalViewRHI != nullptr)
 						{
 							*ResourceTablePhysicalTexturePtr = PhysicalViewRHI;
-							*ResourceTablePhysicalSamplerPtr = TStaticSamplerState<SF_Bilinear, AM_Clamp, AM_Clamp, AM_Clamp, 0, 8>::GetRHI();
+							*ResourceTablePhysicalSamplerPtr = TStaticSamplerState<SF_AnisotropicPoint, AM_Clamp, AM_Clamp, AM_Clamp, 0, 8>::GetRHI();
 							bValidResources = true;
 						}
 					}
