Fixed naming inconsistency between distance field shadows and ray tracing shadows parameters

diff --git a/Engine/Source/Runtime/Engine/Classes/Components/LightComponent.h b/Engine/Source/Runtime/Engine/Classes/Components/LightComponent.h
--- a/Engine/Source/Runtime/Engine/Classes/Components/LightComponent.h
+++ b/Engine/Source/Runtime/Engine/Classes/Components/LightComponent.h
@@ -213,7 +213,7 @@ class ENGINE_API ULightComponent : public ULightComponentBase
 	 * They have less aliasing artifacts than standard shadowmaps, but inherit all the limitations of distance field representations (only uniform scale, no deformation).
 	 * These shadows have a low per-object cost (and don't depend on triangle count) so they are effective for distant shadows from a dynamic sun.
 	 */
-	UPROPERTY(EditAnywhere, BlueprintReadOnly, Category=DistanceFieldShadows, meta=(DisplayName = "Ray Tracing DistanceField Shadows"))
+	UPROPERTY(EditAnywhere, BlueprintReadOnly, Category=DistanceFieldShadows, meta=(DisplayName = "Distance Field Shadows"))
 	bool bUseRayTracedDistanceFieldShadows;
 
 	/** 
diff --git a/Engine/Source/Runtime/Engine/Public/ShowFlagsValues.inl b/Engine/Source/Runtime/Engine/Public/ShowFlagsValues.inl
--- a/Engine/Source/Runtime/Engine/Public/ShowFlagsValues.inl
+++ b/Engine/Source/Runtime/Engine/Public/ShowFlagsValues.inl
@@ -111,7 +111,7 @@ SHOWFLAG_ALWAYS_ACCESSIBLE(ScreenSpaceReflections, SFG_LightingFeatures, NSLOCTE
 /** If Screen space contact shadows are enabled. */
 SHOWFLAG_ALWAYS_ACCESSIBLE(ContactShadows, SFG_LightingFeatures, NSLOCTEXT("UnrealEd", "ContactShadows", "Screen Space Contact Shadows"))
 /** If RTDF shadows are enabled. */
-SHOWFLAG_ALWAYS_ACCESSIBLE(RayTracedDistanceFieldShadows, SFG_LightingFeatures, NSLOCTEXT("UnrealEd", "RayTracedDistanceFieldShadows", "Ray Traced Distance Field Shadows"))
+SHOWFLAG_ALWAYS_ACCESSIBLE(RayTracedDistanceFieldShadows, SFG_LightingFeatures, NSLOCTEXT("UnrealEd", "RayTracedDistanceFieldShadows", "Distance Field Shadows"))
 /** If Capsule shadows are enabled. */
 SHOWFLAG_ALWAYS_ACCESSIBLE(CapsuleShadows, SFG_LightingFeatures, NSLOCTEXT("UnrealEd", "CapsuleShadows", "Capsule Shadows"))
 /** If Screen Space Subsurface Scattering enabled */
