Guard against indexing out of bounds into SkeletalBodySetups when logging warnings for invalid configurations. #jira UE-75461 #rb none

diff --git a/Engine/Source/Runtime/Engine/Private/PhysicsEngine/PhysAnim.cpp b/Engine/Source/Runtime/Engine/Private/PhysicsEngine/PhysAnim.cpp
--- a/Engine/Source/Runtime/Engine/Private/PhysicsEngine/PhysAnim.cpp
+++ b/Engine/Source/Runtime/Engine/Private/PhysicsEngine/PhysAnim.cpp
@@ -587,14 +587,22 @@ void USkeletalMeshComponent::UpdateKinematicBonesToAnim(const TArray<FTransform>
 						// If we could not find it - warn.
 						if(BoneIndex == INDEX_NONE || BoneIndex >= NumComponentSpaceTransforms)
 						{
-							const FName BodyName = PhysicsAsset->SkeletalBodySetups[i]->BoneName;
+							FName BodyName = TEXT("UNKNOWN");
+							if (PhysicsAsset->SkeletalBodySetups.IsValidIndex(i))
+							{
+								BodyName = PhysicsAsset->SkeletalBodySetups[i]->BoneName;
+							}
 							UE_LOG(LogPhysics, Log, TEXT("UpdateRBBones: WARNING: Failed to find bone '%s' need by PhysicsAsset '%s' in SkeletalMesh '%s'."), *BodyName.ToString(), *PhysicsAsset->GetName(), *SkeletalMesh->GetName());
 						}
 						else
 						{
 							if (BoneIndex >= InSpaceBases.Num())
 							{
-								const FName BodyName = PhysicsAsset->SkeletalBodySetups[i]->BoneName;
+								FName BodyName = TEXT("UNKNOWN");
+								if (PhysicsAsset->SkeletalBodySetups.IsValidIndex(i))
+								{
+									BodyName = PhysicsAsset->SkeletalBodySetups[i]->BoneName;
+								}
 								UE_LOG(LogPhysics, Warning, TEXT("BoneIndex %d out of range of SpaceBases (Size %d) on PhysicsAsset '%s' in SkeletalMesh '%s' for bone '%s'"), BoneIndex, InSpaceBases.Num(), *PhysicsAsset->GetName(), *SkeletalMesh->GetName(), *BodyName.ToString());
 								continue;
 							}
@@ -603,7 +611,12 @@ void USkeletalMeshComponent::UpdateKinematicBonesToAnim(const TArray<FTransform>
 							const FTransform BoneTransform = InSpaceBases[BoneIndex] * CurrentLocalToWorld;
 							if(!BoneTransform.IsValid())
 							{
-								const FName BodyName = PhysicsAsset->SkeletalBodySetups[i]->BoneName;
+								FName BodyName = TEXT("UNKNOWN");
+								if (PhysicsAsset->SkeletalBodySetups.IsValidIndex(i))
+								{
+									BodyName = PhysicsAsset->SkeletalBodySetups[i]->BoneName;
+								}
+
 								UE_LOG(LogPhysics, Warning, TEXT("UpdateKinematicBonesToAnim: Trying to set transform with bad data %s on PhysicsAsset '%s' in SkeletalMesh '%s' for bone '%s'"), *BoneTransform.ToHumanReadableString(), *PhysicsAsset->GetName(), *SkeletalMesh->GetName(), *BodyName.ToString());
 								BoneTransform.DiagnosticCheck_IsValid();	//In special nan mode we want to actually ensure
 
