#ue4 - Fixed FindTeleportSpot() when trying to find spots when spawning an actor (before BeginPlay). It now applies the proposed Z adjustment to all permutations instead of only the last one in the XY-only permutations if those failed. When an adjustment axis is nearly zero, more permutations are tried rather than duplicating the 4 axial permutations.

Fixes being unable to spawn 4+ players when there is only 1 player start available. Note that it is still preferable to include multiple player starts to avoid contention when player starts are occupied.

Behavior of FindTeleportSpot() for actors that have begun play is unchanged.

diff --git a/Engine/Source/Runtime/Engine/Private/LevelActor.cpp b/Engine/Source/Runtime/Engine/Private/LevelActor.cpp
--- a/Engine/Source/Runtime/Engine/Private/LevelActor.cpp
+++ b/Engine/Source/Runtime/Engine/Private/LevelActor.cpp
@@ -811,7 +811,8 @@ bool UWorld::FindTeleportSpot(const AActor* TestActor, FVector& TestLocation, FR
 	}
 
 	// first do only Z
-	if (!FMath::IsNearlyZero(Adjust.Z)) 
+	const bool bZeroZ = FMath::IsNearlyZero(Adjust.Z, KINDA_SMALL_NUMBER);
+	if (!bZeroZ)
 	{
 		TestLocation.Z += Adjust.Z;
 		if( !EncroachingBlockingGeometry(TestActor, TestLocation, TestRotation) )
@@ -823,40 +824,67 @@ bool UWorld::FindTeleportSpot(const AActor* TestActor, FVector& TestLocation, FR
 	}
 
 	// now try just XY
-	if (!FMath::IsNearlyZero(Adjust.X) || !FMath::IsNearlyZero(Adjust.Y))
+	const bool bZeroX = FMath::IsNearlyZero(Adjust.X, KINDA_SMALL_NUMBER);
+	const bool bZeroY = FMath::IsNearlyZero(Adjust.Y, KINDA_SMALL_NUMBER);
+	if (!bZeroX || !bZeroY)
 	{
+		const float X = bZeroX ? 0.f : Adjust.X;
+		const float Y = bZeroY ? 0.f : Adjust.Y;
+		FVector Adjustments[8];
+		Adjustments[0] = FVector(X, Y, 0);
+
 		// If initially spawning allow testing a few permutations (though this needs improvement).
 		// During play only test the first adjustment, permuting axes could put the location on other sides of geometry.
 		const int32 Iterations = (TestActor->HasActorBegunPlay() ? 1 : 8);
-		for (int i = 0; i < Iterations; ++i)
+
+		if (Iterations > 1)
 		{
-			TestLocation = OriginalTestLocation;
-			TestLocation.X += (i < 4 ? Adjust.X : Adjust.Y) * (i % 2 == 0 ? 1 : -1);
-			TestLocation.Y += (i < 4 ? Adjust.Y : Adjust.X) * (i % 4 < 2 ? 1 : -1);
-			if (!EncroachingBlockingGeometry(TestActor, TestLocation, TestRotation))
+			if (!bZeroX && !bZeroY)
 			{
-				return true;
+				Adjustments[1] = FVector(-X,  Y, 0);
+				Adjustments[2] = FVector( X, -Y, 0);
+				Adjustments[3] = FVector(-X, -Y, 0);
+				Adjustments[4] = FVector( Y,  X, 0);
+				Adjustments[5] = FVector(-Y,  X, 0);
+				Adjustments[6] = FVector( Y, -X, 0);
+				Adjustments[7] = FVector(-Y, -X, 0);
+			}
+			else
+			{
+				// If either X or Y was zero, the permutations above would result in only 4 unique attempts.
+				Adjustments[1] = FVector(-X, -Y, 0);
+				Adjustments[2] = FVector( Y,  X, 0);
+				Adjustments[3] = FVector(-Y, -X, 0);
+				// Mirror the dominant non-zero value
+				const float D = bZeroY ? X : Y;
+				Adjustments[4] = FVector( D,  D, 0);
+				Adjustments[5] = FVector( D, -D, 0);
+				Adjustments[6] = FVector(-D,  D, 0);
+				Adjustments[7] = FVector(-D, -D, 0);
 			}
 		}
 
-		// now z again (with the last XY offset we tried)
-		if (!FMath::IsNearlyZero(Adjust.Z))
+		for (int i = 0; i < Iterations; ++i)
 		{
-			TestLocation.Z += Adjust.Z;
-			if( !EncroachingBlockingGeometry(TestActor, TestLocation, TestRotation) )
+			TestLocation = OriginalTestLocation + Adjustments[i];
+			if (!EncroachingBlockingGeometry(TestActor, TestLocation, TestRotation))
 			{
 				return true;
 			}
 		}
-	}
 
-	// Now try full adjustment (only if we're not in a match, otherwise this test is actually redundant because we did it right above.)
-	if (!TestActor->HasActorBegunPlay())
-	{
-		TestLocation = OriginalTestLocation + Adjust;
-		if( !EncroachingBlockingGeometry(TestActor, TestLocation, TestRotation, &Adjust) )
+		// Try XY adjustment including Z. Note that even with only 1 iteration, this will still try the full proposed (X,Y,Z) adjustment.
+		if (!bZeroZ)
 		{
-			return true;
+			for (int i = 0; i < Iterations; ++i)
+			{
+				TestLocation = OriginalTestLocation + Adjustments[i];
+				TestLocation.Z += Adjust.Z;
+				if (!EncroachingBlockingGeometry(TestActor, TestLocation, TestRotation))
+				{
+					return true;
+				}
+			}
 		}
 	}
 
