Guard against index out of bounds in UpdateWorldBoneTM. #jira UE-75464 #rb none

diff --git a/Engine/Source/Runtime/Engine/Private/PhysicsEngine/PhysAnim.cpp b/Engine/Source/Runtime/Engine/Private/PhysicsEngine/PhysAnim.cpp
--- a/Engine/Source/Runtime/Engine/Private/PhysicsEngine/PhysAnim.cpp
+++ b/Engine/Source/Runtime/Engine/Private/PhysicsEngine/PhysAnim.cpp
@@ -140,11 +140,14 @@ void UpdateWorldBoneTM(TAssetWorldBoneTMArray& WorldBoneTMs, const TArray<FTrans
 		ParentTM = WorldBoneTMs[ParentIndex].TM;
 	}
 
-	RelTM = InBoneSpaceTransforms[BoneIndex];
-	RelTM.ScaleTranslation( Scale3D );
+	if (InBoneSpaceTransforms.IsValidIndex(BoneIndex))
+	{
+		RelTM = InBoneSpaceTransforms[BoneIndex];
+		RelTM.ScaleTranslation(Scale3D);
 
-	WorldBoneTMs[BoneIndex].TM = RelTM * ParentTM;
-	WorldBoneTMs[BoneIndex].bUpToDate = true;
+		WorldBoneTMs[BoneIndex].TM = RelTM * ParentTM;
+		WorldBoneTMs[BoneIndex].bUpToDate = true;
+	}
 }
 TAutoConsoleVariable<int32> CVarPhysicsAnimBlendUpdatesPhysX(TEXT("p.PhysicsAnimBlendUpdatesPhysX"), 1, TEXT("Whether to update the physx simulation with the results of physics animation blending"));
 
