Fixes for engine analytics.
 - bTerminated now initialized properly.
 - We now track the process id in the session data and don't report running instances as abnormal terminations.

diff --git a/Engine/Source/Runtime/Engine/Private/EngineSessionManager.cpp b/Engine/Source/Runtime/Engine/Private/EngineSessionManager.cpp
--- a/Engine/Source/Runtime/Engine/Private/EngineSessionManager.cpp
+++ b/Engine/Source/Runtime/Engine/Private/EngineSessionManager.cpp
@@ -56,6 +56,7 @@ namespace SessionManagerDefs
 	static const FString DeactivatedStoreKey(TEXT("IsDeactivated"));
 	static const FString BackgroundStoreKey(TEXT("IsInBackground"));
 	static const FString TerminatingKey(TEXT("Terminating"));
+	static const FString PlatformProcessIDKey(TEXT("PlatformProcessID"));
 	static const FString EngineVersionStoreKey(TEXT("EngineVersion"));
 	static const FString TimestampStoreKey(TEXT("Timestamp"));
 	static const FString StartupTimeStoreKey(TEXT("StartupTimestamp"));
@@ -252,19 +253,7 @@ void FEngineSessionManager::Shutdown()
 	{
 		if (!CurrentSession.bCrashed)
 		{
-			FPlatformMisc::DeleteStoredValue(SessionManagerDefs::StoreId, CurrentSessionSectionName, SessionManagerDefs::ModeStoreKey);
-			FPlatformMisc::DeleteStoredValue(SessionManagerDefs::StoreId, CurrentSessionSectionName, SessionManagerDefs::ProjectNameStoreKey);
-			FPlatformMisc::DeleteStoredValue(SessionManagerDefs::StoreId, CurrentSessionSectionName, SessionManagerDefs::CrashStoreKey);
-			FPlatformMisc::DeleteStoredValue(SessionManagerDefs::StoreId, CurrentSessionSectionName, SessionManagerDefs::GPUCrashStoreKey);
-			FPlatformMisc::DeleteStoredValue(SessionManagerDefs::StoreId, CurrentSessionSectionName, SessionManagerDefs::EngineVersionStoreKey);
-			FPlatformMisc::DeleteStoredValue(SessionManagerDefs::StoreId, CurrentSessionSectionName, SessionManagerDefs::TimestampStoreKey);
-			FPlatformMisc::DeleteStoredValue(SessionManagerDefs::StoreId, CurrentSessionSectionName, SessionManagerDefs::DebuggerStoreKey);
-			FPlatformMisc::DeleteStoredValue(SessionManagerDefs::StoreId, CurrentSessionSectionName, SessionManagerDefs::WasDebuggerStoreKey);
-			FPlatformMisc::DeleteStoredValue(SessionManagerDefs::StoreId, CurrentSessionSectionName, SessionManagerDefs::DeactivatedStoreKey);
-			FPlatformMisc::DeleteStoredValue(SessionManagerDefs::StoreId, CurrentSessionSectionName, SessionManagerDefs::BackgroundStoreKey);
-			FPlatformMisc::DeleteStoredValue(SessionManagerDefs::StoreId, CurrentSessionSectionName, SessionManagerDefs::UserActivityStoreKey);
-			FPlatformMisc::DeleteStoredValue(SessionManagerDefs::StoreId, CurrentSessionSectionName, SessionManagerDefs::VanillaStoreKey);
-			FPlatformMisc::DeleteStoredValue(SessionManagerDefs::StoreId, CurrentSessionSectionName, SessionManagerDefs::TerminatingKey);
+			DeleteStoredRecordValues(CurrentSessionSectionName);
 
 #if PLATFORM_SUPPORTS_WATCHDOG
 			if (!WatchdogSectionName.IsEmpty())
@@ -307,12 +296,29 @@ bool FEngineSessionManager::BeginReadWriteRecords()
 		FString TimestampString;
 		FString IsDebuggerString;
 
-		// Read manditory values
+		// Read mandatory values
 		if (FPlatformMisc::GetStoredValue(SessionManagerDefs::StoreId, SectionName, SessionManagerDefs::CrashStoreKey, IsCrashString) &&
 			FPlatformMisc::GetStoredValue(SessionManagerDefs::StoreId, SectionName, SessionManagerDefs::EngineVersionStoreKey, EngineVersionString) &&
 			FPlatformMisc::GetStoredValue(SessionManagerDefs::StoreId, SectionName, SessionManagerDefs::TimestampStoreKey, TimestampString) &&
 			FPlatformMisc::GetStoredValue(SessionManagerDefs::StoreId, SectionName, SessionManagerDefs::DebuggerStoreKey, IsDebuggerString))
 		{
+			// If the process is still running we don't need to report it.
+			FString PlatformProcessIDString;
+			if (FPlatformMisc::GetStoredValue(SessionManagerDefs::StoreId, SectionName, SessionManagerDefs::PlatformProcessIDKey, PlatformProcessIDString))
+			{
+				const uint32 ProcId = (uint32)FCString::Atoi(*PlatformProcessIDString);
+				FProcHandle Handle = FPlatformProcess::OpenProcess(ProcId);
+				if (Handle.IsValid())
+				{
+					const bool bIsRunning = FPlatformProcess::IsProcRunning(Handle);
+					FPlatformProcess::CloseProc(Handle);
+					if (bIsRunning)
+					{
+						continue;
+					}
+				}
+			}
+
 			// Read optional values
 			FString WasDebuggerString = IsDebuggerString;
 			FPlatformMisc::GetStoredValue(SessionManagerDefs::StoreId, SectionName, SessionManagerDefs::WasDebuggerStoreKey, WasDebuggerString);
@@ -355,19 +361,7 @@ bool FEngineSessionManager::BeginReadWriteRecords()
 		else
 		{
 			// Clean up orphaned values, if there are any
-			FPlatformMisc::DeleteStoredValue(SessionManagerDefs::StoreId, SectionName, SessionManagerDefs::ModeStoreKey);
-			FPlatformMisc::DeleteStoredValue(SessionManagerDefs::StoreId, SectionName, SessionManagerDefs::ProjectNameStoreKey);
-			FPlatformMisc::DeleteStoredValue(SessionManagerDefs::StoreId, SectionName, SessionManagerDefs::CrashStoreKey);
-			FPlatformMisc::DeleteStoredValue(SessionManagerDefs::StoreId, SectionName, SessionManagerDefs::GPUCrashStoreKey);
-			FPlatformMisc::DeleteStoredValue(SessionManagerDefs::StoreId, SectionName, SessionManagerDefs::EngineVersionStoreKey);
-			FPlatformMisc::DeleteStoredValue(SessionManagerDefs::StoreId, SectionName, SessionManagerDefs::TimestampStoreKey);
-			FPlatformMisc::DeleteStoredValue(SessionManagerDefs::StoreId, SectionName, SessionManagerDefs::DebuggerStoreKey);
-			FPlatformMisc::DeleteStoredValue(SessionManagerDefs::StoreId, SectionName, SessionManagerDefs::WasDebuggerStoreKey);
-			FPlatformMisc::DeleteStoredValue(SessionManagerDefs::StoreId, SectionName, SessionManagerDefs::DeactivatedStoreKey);
-			FPlatformMisc::DeleteStoredValue(SessionManagerDefs::StoreId, SectionName, SessionManagerDefs::BackgroundStoreKey);
-			FPlatformMisc::DeleteStoredValue(SessionManagerDefs::StoreId, SectionName, SessionManagerDefs::UserActivityStoreKey);
-			FPlatformMisc::DeleteStoredValue(SessionManagerDefs::StoreId, SectionName, SessionManagerDefs::VanillaStoreKey);
-			FPlatformMisc::DeleteStoredValue(SessionManagerDefs::StoreId, SectionName, SessionManagerDefs::TerminatingKey);
+			DeleteStoredRecordValues(SectionName);
 		}
 	}
 
@@ -401,6 +395,14 @@ void FEngineSessionManager::DeleteStoredRecord(const FSessionRecord& Record)
 	FString SessionId = Record.SessionId;
 	FString SectionName = GetStoreSectionString(SessionId);
 
+	DeleteStoredRecordValues(SectionName);
+
+	// Remove the session record from SessionRecords list
+	SessionRecords.RemoveAll([&SessionId](const FSessionRecord& X){ return X.SessionId == SessionId; });
+}
+
+void FEngineSessionManager::DeleteStoredRecordValues(const FString& SectionName) const
+{
 	FPlatformMisc::DeleteStoredValue(SessionManagerDefs::StoreId, SectionName, SessionManagerDefs::ModeStoreKey);
 	FPlatformMisc::DeleteStoredValue(SessionManagerDefs::StoreId, SectionName, SessionManagerDefs::ProjectNameStoreKey);
 	FPlatformMisc::DeleteStoredValue(SessionManagerDefs::StoreId, SectionName, SessionManagerDefs::CrashStoreKey);
@@ -414,9 +416,7 @@ void FEngineSessionManager::DeleteStoredRecord(const FSessionRecord& Record)
 	FPlatformMisc::DeleteStoredValue(SessionManagerDefs::StoreId, SectionName, SessionManagerDefs::UserActivityStoreKey);
 	FPlatformMisc::DeleteStoredValue(SessionManagerDefs::StoreId, SectionName, SessionManagerDefs::VanillaStoreKey);
 	FPlatformMisc::DeleteStoredValue(SessionManagerDefs::StoreId, SectionName, SessionManagerDefs::TerminatingKey);
-
-	// Remove the session record from SessionRecords list
-	SessionRecords.RemoveAll([&SessionId](const FSessionRecord& X){ return X.SessionId == SessionId; });
+	FPlatformMisc::DeleteStoredValue(SessionManagerDefs::StoreId, SectionName, SessionManagerDefs::PlatformProcessIDKey);
 }
 
 /**
@@ -531,6 +531,9 @@ void FEngineSessionManager::CreateAndWriteRecordForSession()
 		CurrentSession.SessionId = FEngineAnalytics::GetProvider().GetSessionID();
 	}
 
+	const uint32 ProcId = FPlatformProcess::GetCurrentProcessId();
+	FString CurrentProcessIDString = FString::FromInt(ProcId);
+
 	const UGeneralProjectSettings& ProjectSettings = *GetDefault<UGeneralProjectSettings>();
 
 	CurrentSession.Mode = Mode;
@@ -562,6 +565,7 @@ void FEngineSessionManager::CreateAndWriteRecordForSession()
 	FPlatformMisc::SetStoredValue(SessionManagerDefs::StoreId, CurrentSessionSectionName, SessionManagerDefs::UserActivityStoreKey, CurrentSession.CurrentUserActivity);
 	FPlatformMisc::SetStoredValue(SessionManagerDefs::StoreId, CurrentSessionSectionName, SessionManagerDefs::VanillaStoreKey, IsVanillaString);
 	FPlatformMisc::SetStoredValue(SessionManagerDefs::StoreId, CurrentSessionSectionName, SessionManagerDefs::TerminatingKey, IsTerminatingString);
+	FPlatformMisc::SetStoredValue(SessionManagerDefs::StoreId, CurrentSessionSectionName, SessionManagerDefs::PlatformProcessIDKey, CurrentProcessIDString);
 
 	SessionRecords.Add(CurrentSession);
 
diff --git a/Engine/Source/Runtime/Engine/Private/EngineSessionManager.h b/Engine/Source/Runtime/Engine/Private/EngineSessionManager.h
--- a/Engine/Source/Runtime/Engine/Private/EngineSessionManager.h
+++ b/Engine/Source/Runtime/Engine/Private/EngineSessionManager.h
@@ -35,6 +35,7 @@ public:
 		CurrentSession.bIsDeactivated = false;
 		CurrentSession.bIsInBackground = false;
 		CurrentSession.bIsVanilla = false;
+		CurrentSession.bIsTerminating = false;
 	}
 
 	void Initialize();
@@ -67,6 +68,7 @@ private:
 	bool BeginReadWriteRecords();
 	void EndReadWriteRecords();
 	void DeleteStoredRecord(const FSessionRecord& Record);
+	void DeleteStoredRecordValues(const FString& SectionName) const;
 	void SendAbnormalShutdownReport(const FSessionRecord& Record);
 	void CreateAndWriteRecordForSession();
 	void OnCrashing();
