Several ray tracing override flags were not marked as EditAnywhere, despite being used as InlineEditConditionToggles.

diff --git a/Engine/Source/Runtime/Engine/Classes/Engine/Scene.h b/Engine/Source/Runtime/Engine/Classes/Engine/Scene.h
--- a/Engine/Source/Runtime/Engine/Classes/Engine/Scene.h
+++ b/Engine/Source/Runtime/Engine/Classes/Engine/Scene.h
@@ -902,7 +902,7 @@ struct FPostProcessSettings
 	UPROPERTY(EditAnywhere, BlueprintReadWrite, Category=Overrides, meta=(PinHiddenByDefault, InlineEditConditionToggle))
 	uint8 bOverride_AmbientOcclusionMipThreshold:1;
 
-	UPROPERTY(BlueprintReadWrite, Category = Overrides, meta = (PinHiddenByDefault, InlineEditConditionToggle))
+	UPROPERTY(EditAnywhere, BlueprintReadWrite, Category = Overrides, meta = (PinHiddenByDefault, InlineEditConditionToggle))
 	uint32 bOverride_RayTracingAOSamplesPerPixel : 1;
 
 	UPROPERTY(EditAnywhere, BlueprintReadWrite, Category=Overrides, meta=(PinHiddenByDefault, InlineEditConditionToggle))
@@ -1044,52 +1044,52 @@ struct FPostProcessSettings
 
 	// Ray Tracing
 	
-	UPROPERTY(BlueprintReadWrite, Category = Overrides, meta = (PinHiddenByDefault, InlineEditConditionToggle))
+	UPROPERTY(EditAnywhere, BlueprintReadWrite, Category = Overrides, meta = (PinHiddenByDefault, InlineEditConditionToggle))
 	uint32 bOverride_ReflectionsType : 1;
 
-	UPROPERTY(BlueprintReadWrite, Category = Overrides, meta = (PinHiddenByDefault, InlineEditConditionToggle))
+	UPROPERTY(EditAnywhere, BlueprintReadWrite, Category = Overrides, meta = (PinHiddenByDefault, InlineEditConditionToggle))
 	uint32 bOverride_RayTracingReflectionsMaxRoughness : 1;
 
-	UPROPERTY(BlueprintReadWrite, Category = Overrides, meta = (PinHiddenByDefault, InlineEditConditionToggle))
+	UPROPERTY(EditAnywhere, BlueprintReadWrite, Category = Overrides, meta = (PinHiddenByDefault, InlineEditConditionToggle))
 	uint32 bOverride_RayTracingReflectionsMaxBounces : 1;
 
-	UPROPERTY(BlueprintReadWrite, Category = Overrides, meta = (PinHiddenByDefault, InlineEditConditionToggle))
+	UPROPERTY(EditAnywhere, BlueprintReadWrite, Category = Overrides, meta = (PinHiddenByDefault, InlineEditConditionToggle))
 	uint32 bOverride_RayTracingReflectionsSamplesPerPixel : 1;
 
-	UPROPERTY(BlueprintReadWrite, Category = Overrides, meta = (PinHiddenByDefault, InlineEditConditionToggle))
+	UPROPERTY(EditAnywhere, BlueprintReadWrite, Category = Overrides, meta = (PinHiddenByDefault, InlineEditConditionToggle))
 	uint32 bOverride_RayTracingReflectionsShadows : 1;
 
-	UPROPERTY(BlueprintReadWrite, Category = Overrides, meta = (PinHiddenByDefault, InlineEditConditionToggle))
+	UPROPERTY(EditAnywhere, BlueprintReadWrite, Category = Overrides, meta = (PinHiddenByDefault, InlineEditConditionToggle))
 	uint32 bOverride_TranslucencyType : 1;
 
-	UPROPERTY(BlueprintReadWrite, Category = Overrides, meta = (PinHiddenByDefault, InlineEditConditionToggle))
+	UPROPERTY(EditAnywhere, BlueprintReadWrite, Category = Overrides, meta = (PinHiddenByDefault, InlineEditConditionToggle))
 	uint32 bOverride_RayTracingTranslucencyMaxRoughness : 1;
 
-	UPROPERTY(BlueprintReadWrite, Category = Overrides, meta = (PinHiddenByDefault, InlineEditConditionToggle))
+	UPROPERTY(EditAnywhere, BlueprintReadWrite, Category = Overrides, meta = (PinHiddenByDefault, InlineEditConditionToggle))
 	uint32 bOverride_RayTracingTranslucencyRefractionRays : 1;
 
-	UPROPERTY(BlueprintReadWrite, Category = Overrides, meta = (PinHiddenByDefault, InlineEditConditionToggle))
+	UPROPERTY(EditAnywhere, BlueprintReadWrite, Category = Overrides, meta = (PinHiddenByDefault, InlineEditConditionToggle))
 	uint32 bOverride_RayTracingTranslucencySamplesPerPixel : 1;
 
-	UPROPERTY(BlueprintReadWrite, Category = Overrides, meta = (PinHiddenByDefault, InlineEditConditionToggle))
+	UPROPERTY(EditAnywhere, BlueprintReadWrite, Category = Overrides, meta = (PinHiddenByDefault, InlineEditConditionToggle))
 	uint32 bOverride_RayTracingTranslucencyShadows : 1;
 
-	UPROPERTY(BlueprintReadWrite, Category = Overrides, meta = (PinHiddenByDefault, InlineEditConditionToggle))
+	UPROPERTY(EditAnywhere, BlueprintReadWrite, Category = Overrides, meta = (PinHiddenByDefault, InlineEditConditionToggle))
 	uint32 bOverride_RayTracingTranslucencyRefraction : 1;
 
-	UPROPERTY(BlueprintReadWrite, Category = Overrides, meta = (PinHiddenByDefault, InlineEditConditionToggle))
+	UPROPERTY(EditAnywhere, BlueprintReadWrite, Category = Overrides, meta = (PinHiddenByDefault, InlineEditConditionToggle))
 	uint32 bOverride_RayTracingGI : 1;
 
-	UPROPERTY(BlueprintReadWrite, Category = Overrides, meta = (PinHiddenByDefault, InlineEditConditionToggle))
+	UPROPERTY(EditAnywhere, BlueprintReadWrite, Category = Overrides, meta = (PinHiddenByDefault, InlineEditConditionToggle))
 	uint32 bOverride_RayTracingGIMaxBounces : 1;
 
-	UPROPERTY(BlueprintReadWrite, Category = Overrides, meta = (PinHiddenByDefault, InlineEditConditionToggle))
+	UPROPERTY(EditAnywhere, BlueprintReadWrite, Category = Overrides, meta = (PinHiddenByDefault, InlineEditConditionToggle))
 	uint32 bOverride_RayTracingGISamplesPerPixel : 1;
 
-	UPROPERTY(BlueprintReadWrite, Category = Overrides, meta = (PinHiddenByDefault, InlineEditConditionToggle))
+	UPROPERTY(EditAnywhere, BlueprintReadWrite, Category = Overrides, meta = (PinHiddenByDefault, InlineEditConditionToggle))
 	uint32 bOverride_PathTracingMaxBounces : 1;
 
-	UPROPERTY(BlueprintReadWrite, Category = Overrides, meta = (PinHiddenByDefault, InlineEditConditionToggle))
+	UPROPERTY(EditAnywhere, BlueprintReadWrite, Category = Overrides, meta = (PinHiddenByDefault, InlineEditConditionToggle))
 	uint32 bOverride_PathTracingSamplesPerPixel : 1;
 
 	// -----------------------------------------------------------------------
