Edigrate 6853491
Fixed crash when adding a component in the blueprint editor after undoing an add component action - explicitly move objects out of the way in order to adhere to component naming convention

diff --git a/Engine/Source/Runtime/Engine/Private/SimpleConstructionScript.cpp b/Engine/Source/Runtime/Engine/Private/SimpleConstructionScript.cpp
--- a/Engine/Source/Runtime/Engine/Private/SimpleConstructionScript.cpp
+++ b/Engine/Source/Runtime/Engine/Private/SimpleConstructionScript.cpp
@@ -1292,16 +1292,27 @@ USCS_Node* USimpleConstructionScript::CreateNode(UClass* NewComponentClass, FNam
 	// At this point we should have a unique, explicit name to use for the template object.
 	check(NewComponentVariableName != NAME_None);
 
+	UPackage* TransientPackage = GetTransientPackage();
+
 	// A bit of a hack, but by doing this we ensure that the original object isn't outered to the BPGC. That way if we undo this action later, it'll rename the template away from the BPGC.
 	// This is necessary because of our template object naming scheme that's in place to ensure deterministic cooking. We have to keep the SCS node and template object names in sync as a result,
 	// and leaving the template outered to the BPGC can lead to template object name collisions when attempting to rename the remaining SCS nodes. See USCS_Node::NameWasModified() for more details.
-	UActorComponent* NewComponentTemplate = NewObject<UActorComponent>(GetTransientPackage(), NewComponentClass, NAME_None, RF_ArchetypeObject | RF_Transactional | RF_Public);
+	UActorComponent* NewComponentTemplate = NewObject<UActorComponent>(TransientPackage, NewComponentClass, NAME_None, RF_ArchetypeObject | RF_Transactional | RF_Public);
 
 	// Record initial object state in case we're in a transaction context.
 	NewComponentTemplate->Modify();
 
+	FString Name = NewComponentVariableName.ToString() + ComponentTemplateNameSuffix;
+	// First, make sure that e.g. undo/redo hasn't orphaned any objects in our 'position':
+	UObject* Collision = FindObject<UObject>(Blueprint->GeneratedClass, *Name);
+	while(Collision)
+	{
+		Collision->Rename(nullptr, GetTransientPackage(), REN_DoNotDirty|REN_DontCreateRedirectors|REN_ForceNoResetLoaders);
+		Collision = FindObject<UObject>(Blueprint->GeneratedClass, *Name);
+	}
+
 	// Now set the actual name and outer to the BPGC.
-	NewComponentTemplate->Rename(*(NewComponentVariableName.ToString() + ComponentTemplateNameSuffix), Blueprint->GeneratedClass, REN_DoNotDirty|REN_DontCreateRedirectors|REN_ForceNoResetLoaders);
+	NewComponentTemplate->Rename(*Name, Blueprint->GeneratedClass, REN_DoNotDirty|REN_DontCreateRedirectors|REN_ForceNoResetLoaders);
 
 	return CreateNodeImpl(NewComponentTemplate, NewComponentVariableName);
 }
