Add way to tell ShowFlags not to serialize, set HighResScreenshotMask to not save

diff --git a/Engine/Plugins/Tests/EditorTests/Source/EditorTests/Private/UnrealEd/EditorViewportFlagsTest.cpp b/Engine/Plugins/Tests/EditorTests/Source/EditorTests/Private/UnrealEd/EditorViewportFlagsTest.cpp
--- a/Engine/Plugins/Tests/EditorTests/Source/EditorTests/Private/UnrealEd/EditorViewportFlagsTest.cpp
+++ b/Engine/Plugins/Tests/EditorTests/Source/EditorTests/Private/UnrealEd/EditorViewportFlagsTest.cpp
@@ -55,6 +55,7 @@ namespace FViewportTestHelper
 				GroupName = TEXT("Light Features");
 				break;
 			case SFG_Hidden:
+			case SFG_Transient:
 				GroupName = TEXT("Hidden");
 				break;
 			case SFG_Max:
diff --git a/Engine/Source/Editor/DetailCustomizations/Private/SceneCaptureDetails.cpp b/Engine/Source/Editor/DetailCustomizations/Private/SceneCaptureDetails.cpp
--- a/Engine/Source/Editor/DetailCustomizations/Private/SceneCaptureDetails.cpp
+++ b/Engine/Source/Editor/DetailCustomizations/Private/SceneCaptureDetails.cpp
@@ -161,6 +161,7 @@ void FSceneCaptureDetails::CustomizeDetails( IDetailLayoutBuilder& DetailLayout
 				GroupName = LOCTEXT("CollisionModesShowFlagsMenu", "Collision Modes Show Flags");
 				break;
 			case SFG_Hidden:
+			case SFG_Transient:
 				GroupName = LOCTEXT("HiddenShowFlagsMenu", "Hidden Show Flags");
 				break;
 			default:
diff --git a/Engine/Source/Editor/UnrealEd/Private/EditorShowFlags.cpp b/Engine/Source/Editor/UnrealEd/Private/EditorShowFlags.cpp
--- a/Engine/Source/Editor/UnrealEd/Private/EditorShowFlags.cpp
+++ b/Engine/Source/Editor/UnrealEd/Private/EditorShowFlags.cpp
@@ -66,7 +66,7 @@ TArray<FShowFlagData>& GetShowFlagMenuItems()
 				bool OnEngineShowFlag(uint32 InIndex, const FString& InName)
 				{
 					EShowFlagGroup Group = FEngineShowFlags::FindShowFlagGroup(*InName);
-					if( Group != SFG_Hidden  )
+					if( Group != SFG_Hidden && Group != SFG_Transient )
 					{
 						FText FlagDisplayName;
 						FEngineShowFlags::FindShowFlagDisplayName(InName, FlagDisplayName);
diff --git a/Engine/Source/Runtime/Engine/Private/ShowFlags.cpp b/Engine/Source/Runtime/Engine/Private/ShowFlags.cpp
--- a/Engine/Source/Runtime/Engine/Private/ShowFlags.cpp
+++ b/Engine/Source/Runtime/Engine/Private/ShowFlags.cpp
@@ -40,15 +40,19 @@ FString FEngineShowFlags::ToString() const
 
 		bool OnEngineShowFlag(uint32 InIndex, const FString& InName)
 		{
-			if(!ret.IsEmpty())
+			EShowFlagGroup Group = FEngineShowFlags::FindShowFlagGroup(*InName);
+			if (Group != SFG_Transient)
 			{
-				ret += (TCHAR)',';
-			}
+				if (!ret.IsEmpty())
+				{
+					ret += (TCHAR)',';
+				}
 
-			AddNameByIndex(InIndex, ret);
+				AddNameByIndex(InIndex, ret);
 
-			ret += (TCHAR)'=';
-			ret += EngineShowFlags.GetSingleFlag(InIndex) ? (TCHAR)'1' : (TCHAR)'0';
+				ret += (TCHAR)'=';
+				ret += EngineShowFlags.GetSingleFlag(InIndex) ? (TCHAR)'1' : (TCHAR)'0';
+			}
 			return true;
 		}
 
diff --git a/Engine/Source/Runtime/Engine/Public/ShowFlags.h b/Engine/Source/Runtime/Engine/Public/ShowFlags.h
--- a/Engine/Source/Runtime/Engine/Public/ShowFlags.h
+++ b/Engine/Source/Runtime/Engine/Public/ShowFlags.h
@@ -22,6 +22,7 @@ enum EShowFlagGroup
 	SFG_LightingComponents,
 	SFG_LightingFeatures,
 	SFG_Hidden,
+	SFG_Transient, // Hidden, and don't serialize it
 	SFG_Max
 };
 
diff --git a/Engine/Source/Runtime/Engine/Public/ShowFlagsValues.inl b/Engine/Source/Runtime/Engine/Public/ShowFlagsValues.inl
--- a/Engine/Source/Runtime/Engine/Public/ShowFlagsValues.inl
+++ b/Engine/Source/Runtime/Engine/Public/ShowFlagsValues.inl
@@ -261,7 +261,7 @@ SHOWFLAG_ALWAYS_ACCESSIBLE(TextRender, SFG_Advanced, NSLOCTEXT("UnrealEd", "Text
 /** Any rendering/buffer clearing  (good for benchmarking and for pausing rendering while the app is not in focus to save cycles). */
 SHOWFLAG_ALWAYS_ACCESSIBLE(Rendering, SFG_Hidden, NSLOCTEXT("UnrealEd", "RenderingSF", "Any Rendering")) // do not make it FIXED_IN_SHIPPING, used by Oculus plugin.
 /** Show the current mask being used by the highres screenshot capture */
-SHOWFLAG_FIXED_IN_SHIPPING(0, HighResScreenshotMask, SFG_Hidden, NSLOCTEXT("UnrealEd", "HighResScreenshotMaskSF", "High Res Screenshot Mask"))
+SHOWFLAG_FIXED_IN_SHIPPING(0, HighResScreenshotMask, SFG_Transient, NSLOCTEXT("UnrealEd", "HighResScreenshotMaskSF", "High Res Screenshot Mask"))
 /** Distortion of output for HMD devices, SHOWFLAG_ALWAYS_ACCESSIBLE for now because USceneCaptureComponent needs that */
 SHOWFLAG_ALWAYS_ACCESSIBLE(HMDDistortion, SFG_PostProcess, NSLOCTEXT("UnrealEd", "HMDDistortionSF", "HMD Distortion"))
 /** Whether to render in stereoscopic 3d, for now SHOWFLAG_ALWAYS_ACCESSIBLE because it's used by StereoRendering */
