Better naming convention for ray tracing instance CVars

diff --git a/Engine/Source/Runtime/Engine/Private/InstancedStaticMesh.cpp b/Engine/Source/Runtime/Engine/Private/InstancedStaticMesh.cpp
--- a/Engine/Source/Runtime/Engine/Private/InstancedStaticMesh.cpp
+++ b/Engine/Source/Runtime/Engine/Private/InstancedStaticMesh.cpp
@@ -54,32 +54,32 @@ TAutoConsoleVariable<int32> CVarMinLOD(
 	TEXT("Used to discard the top LODs for performance evaluation. -1: Disable all effects of this cvar."));
 
 static TAutoConsoleVariable<int32> CVarRayTracingRenderInstances(
-	TEXT("r.RayTracing.Instances"),
+	TEXT("r.RayTracing.InstancedStaticMeshes"),
 	1,
 	TEXT("Include static mesh instances in ray tracing effects (default = 1 (Instances enabled in ray tracing))"));
 
 static TAutoConsoleVariable<int32> CVarRayTracingRenderInstancesCulling(
-	TEXT("r.RayTracing.Instances.Culling"),
+	TEXT("r.RayTracing.InstancedStaticMeshes.Culling"),
 	0,
 	TEXT("Enable culling for instances in ray tracing (default = 0 (Culling disabled))"));
 
 static TAutoConsoleVariable<float> CVarRayTracingInstancesCullClusterMaxRadiusMultiplier(
-	TEXT("r.RayTracing.Instances.CullClusterMaxRadiusMultiplier"),
+	TEXT("r.RayTracing.InstancedStaticMeshes.CullClusterMaxRadiusMultiplier"),
 	20.0f, 
 	TEXT("Multiplier for the maximum instance size (default = 20cm)"));
 
 static TAutoConsoleVariable<float> CVarRayTracingInstancesCullClusterRadius(
-	TEXT("r.RayTracing.Instances.CullClusterRadius"),
+	TEXT("r.RayTracing.InstancedStaticMeshes.CullClusterRadius"),
 	10000.0f, // 100 m
 	TEXT("Ignore instances outside of this radius in ray tracing effects (default = 10000 (100m))"));
 
 static TAutoConsoleVariable<float> CVarRayTracingInstancesLowScaleThreshold(
-	TEXT("r.RayTracing.Instances.LowScaleRadiusThreshold"),
+	TEXT("r.RayTracing.InstancedStaticMeshes.LowScaleRadiusThreshold"),
 	50.0f, // Instances with a radius smaller than this threshold get culled after CVarRayTracingInstancesLowScaleCullRadius
 	TEXT("Threshold that classifies instances as small (default = 50cm))"));
 
 static TAutoConsoleVariable<float> CVarRayTracingInstancesLowScaleCullRadius(
-	TEXT("r.RayTracing.Instances.LowScaleCullRadius"),
+	TEXT("r.RayTracing.InstancedStaticMeshes.LowScaleCullRadius"),
 	1000.0f, 
 	TEXT("Cull radius for small instances (default = 1000 (10m))"));
 
