Missed file from previous commit

diff --git a/Engine/Source/Runtime/Engine/Private/ShaderCompiler/ShaderCompiler.cpp b/Engine/Source/Runtime/Engine/Private/ShaderCompiler/ShaderCompiler.cpp
--- a/Engine/Source/Runtime/Engine/Private/ShaderCompiler/ShaderCompiler.cpp
+++ b/Engine/Source/Runtime/Engine/Private/ShaderCompiler/ShaderCompiler.cpp
@@ -3377,7 +3377,7 @@ void GlobalBeginCompileShader(
 	
 	{
 		static const auto CVar = IConsoleManager::Get().FindConsoleVariable(TEXT("r.vt.FeedbackFactor"));
-		Input.Environment.SetDefine(TEXT("VIRTUAL_TEXTURE_FEEDBACK_FACTOR"), CVar ? CVar->GetInt() : 0);
+		Input.Environment.SetDefine(TEXT("VIRTUAL_TEXTURE_FEEDBACK_FACTOR"), CVar ? FMath::Max(CVar->GetInt(), 1) : 1);
 	}
 
 
