Don't crash with VT tile border size of 0 or 1
Copied from Dev-Rendering CL 7097467

diff --git a/Engine/Source/Runtime/Engine/Private/TextureDerivedData.cpp b/Engine/Source/Runtime/Engine/Private/TextureDerivedData.cpp
--- a/Engine/Source/Runtime/Engine/Private/TextureDerivedData.cpp
+++ b/Engine/Source/Runtime/Engine/Private/TextureDerivedData.cpp
@@ -450,7 +450,13 @@ static void GetTextureBuildSettings(
 		OutBuildSettings.bVirtualTextureEnableCompressZlib = CVarVTCompressZlib.GetValueOnAnyThread() != 0;
 		OutBuildSettings.bVirtualTextureEnableCompressCrunch = CVarVTCompressCrunch.GetValueOnAnyThread() != 0;
 		OutBuildSettings.VirtualTextureTileSize = FMath::RoundUpToPowerOfTwo(CVarVTTileSize.GetValueOnAnyThread());
-		OutBuildSettings.VirtualTextureBorderSize = FMath::RoundUpToPowerOfTwo(CVarVTTileBorderSize.GetValueOnAnyThread());
+		
+		// 0 is a valid value for border size
+		// 1 would be OK in some cases, but breaks BC compressed formats, since it will result in physical tiles that aren't divisible by block size (4)
+		// Could allow border size of 1 for non BC compressed virtual textures, but somewhat complicated to get that correct, especially with multiple layers
+		// Doesn't seem worth the complexity for now, so clamp the size to be at least 2
+		const int32 TileBorderSize = CVarVTTileBorderSize.GetValueOnAnyThread();
+		OutBuildSettings.VirtualTextureBorderSize = (TileBorderSize > 0) ? FMath::RoundUpToPowerOfTwo(FMath::Max(TileBorderSize, 2)) : 0;
 	}
 	else
 	{
