FnMerge - Add nullptr check when no vertexdata is available

diff --git a/Engine/Source/Runtime/Engine/Private/PositionVertexBuffer.cpp b/Engine/Source/Runtime/Engine/Private/PositionVertexBuffer.cpp
--- a/Engine/Source/Runtime/Engine/Private/PositionVertexBuffer.cpp
+++ b/Engine/Source/Runtime/Engine/Private/PositionVertexBuffer.cpp
@@ -232,7 +232,7 @@ void FPositionVertexBuffer::InitRHI()
 		// When bAllowCPUAccess is true, the meshes is likely going to be used for Niagara to spawn particles on mesh surface.
 		// And it can be the case for CPU *and* GPU access: no differenciation today. That is why we create a SRV in this case.
 		// This also avoid setting lots of states on all the members of all the different buffers used by meshes. Follow up: https://jira.it.epicgames.net/browse/UE-69376.
-		bSRV |= VertexData->GetAllowCPUAccess();
+		bSRV |= (VertexData && VertexData->GetAllowCPUAccess());
 		if(bSRV)
 		{
 			PositionComponentSRV = RHICreateShaderResourceView(VertexBufferRHI, 4, PF_R32_FLOAT);
