Don't queue render command from material serialize method
- Otherwise render command may trigger after material is destroyed, causing random memory corruption, seems to occur in cooked builds

diff --git a/Engine/Source/Runtime/Engine/Private/Materials/MaterialShared.cpp b/Engine/Source/Runtime/Engine/Private/Materials/MaterialShared.cpp
--- a/Engine/Source/Runtime/Engine/Private/Materials/MaterialShared.cpp
+++ b/Engine/Source/Runtime/Engine/Private/Materials/MaterialShared.cpp
@@ -977,7 +977,7 @@ void FMaterial::SerializeInlineShaderMap(FArchive& Ar)
 			{
 				TRefCountPtr<FMaterialShaderMap> LoadedShaderMap = new FMaterialShaderMap();
 				LoadedShaderMap->Serialize(Ar, true, bCooked && Ar.IsLoading());
-				SetGameThreadShaderMap(LoadedShaderMap);
+				GameThreadShaderMap = LoadedShaderMap;
 			}
 		}
 	}
