Added ability to choose if a Texture Render Target 2D should auto generate mip maps when created with the Create Render Target 2D blueprint utility function.

diff --git a/Engine/Source/Runtime/Engine/Classes/Kismet/KismetRenderingLibrary.h b/Engine/Source/Runtime/Engine/Classes/Kismet/KismetRenderingLibrary.h
--- a/Engine/Source/Runtime/Engine/Classes/Kismet/KismetRenderingLibrary.h
+++ b/Engine/Source/Runtime/Engine/Classes/Kismet/KismetRenderingLibrary.h
@@ -47,7 +47,7 @@ class UKismetRenderingLibrary : public UBlueprintFunctionLibrary
 	 * Creates a new render target and initializes it to the specified dimensions
 	 */
 	UFUNCTION(BlueprintCallable, Category="Rendering", meta=(WorldContext="WorldContextObject"))
-	static ENGINE_API UTextureRenderTarget2D* CreateRenderTarget2D(UObject* WorldContextObject, int32 Width = 256, int32 Height = 256, ETextureRenderTargetFormat Format = RTF_RGBA16f, FLinearColor ClearColor = FLinearColor::Black);
+	static ENGINE_API UTextureRenderTarget2D* CreateRenderTarget2D(UObject* WorldContextObject, int32 Width = 256, int32 Height = 256, ETextureRenderTargetFormat Format = RTF_RGBA16f, FLinearColor ClearColor = FLinearColor::Black, bool bAutoGenerateMipMaps = false);
 	
 	/**
 	 * Manually releases GPU resources of a render target. This is useful for blueprint creating a lot of render target that would
diff --git a/Engine/Source/Runtime/Engine/Private/KismetRenderingLibrary.cpp b/Engine/Source/Runtime/Engine/Private/KismetRenderingLibrary.cpp
--- a/Engine/Source/Runtime/Engine/Private/KismetRenderingLibrary.cpp
+++ b/Engine/Source/Runtime/Engine/Private/KismetRenderingLibrary.cpp
@@ -60,7 +60,7 @@ void UKismetRenderingLibrary::ClearRenderTarget2D(UObject* WorldContextObject, U
 	}
 }
 
-UTextureRenderTarget2D* UKismetRenderingLibrary::CreateRenderTarget2D(UObject* WorldContextObject, int32 Width, int32 Height, ETextureRenderTargetFormat Format, FLinearColor ClearColor)
+UTextureRenderTarget2D* UKismetRenderingLibrary::CreateRenderTarget2D(UObject* WorldContextObject, int32 Width, int32 Height, ETextureRenderTargetFormat Format, FLinearColor ClearColor, bool bAutoGenerateMipMaps)
 {
 	UWorld* World = GEngine->GetWorldFromContextObject(WorldContextObject, EGetWorldErrorMode::LogAndReturnNull);
 
@@ -70,7 +70,8 @@ UTextureRenderTarget2D* UKismetRenderingLibrary::CreateRenderTarget2D(UObject* W
 		check(NewRenderTarget2D);
 		NewRenderTarget2D->RenderTargetFormat = Format;
 		NewRenderTarget2D->ClearColor = ClearColor;
-		NewRenderTarget2D->InitAutoFormat(Width, Height);		
+		NewRenderTarget2D->bAutoGenerateMips = bAutoGenerateMipMaps;
+		NewRenderTarget2D->InitAutoFormat(Width, Height);	
 		NewRenderTarget2D->UpdateResourceImmediate(true);
 
 		return NewRenderTarget2D; 
