SourceEffectChain on AudioComponent

Aaron.McLeran

diff --git a/Engine/Source/Runtime/Engine/Classes/Components/AudioComponent.h b/Engine/Source/Runtime/Engine/Classes/Components/AudioComponent.h
--- a/Engine/Source/Runtime/Engine/Classes/Components/AudioComponent.h
+++ b/Engine/Source/Runtime/Engine/Classes/Components/AudioComponent.h
@@ -249,6 +249,9 @@ public:
 	UPROPERTY(EditAnywhere, BlueprintReadWrite, Category = Sound, meta = (ClampMin = "0.0", UIMin = "0.0", EditCondition = "bOverrideSubtitlePriority"))
 	float SubtitlePriority;
 
+	UPROPERTY(EditAnywhere, BlueprintReadWrite, Category = Sound)
+	USoundEffectSourcePresetChain* SourceEffectChain;
+
 #if WITH_EDITORONLY_DATA
 	UPROPERTY()
 	float VolumeWeightedPriorityScale_DEPRECATED;
@@ -550,6 +553,7 @@ public:
 	// Will be set if the audio component is using baked FFT or envelope following data so as to be able to feed that data to BP based on playback time
 	void SetPlaybackTimes(const TMap<uint32, float>& InSoundWavePlaybackTimes);
 
+	void SetSourceEffectChain(USoundEffectSourcePresetChain* InSourceEffectChain);
 public:
 
 	/**
diff --git a/Engine/Source/Runtime/Engine/Private/ActiveSound.cpp b/Engine/Source/Runtime/Engine/Private/ActiveSound.cpp
--- a/Engine/Source/Runtime/Engine/Private/ActiveSound.cpp
+++ b/Engine/Source/Runtime/Engine/Private/ActiveSound.cpp
@@ -33,6 +33,7 @@ FActiveSound::FActiveSound()
 	: World(nullptr)
 	, WorldID(0)
 	, Sound(nullptr)
+	, SourceEffectChain(nullptr)
 	, AudioComponentID(0)
 	, OwnerID(0)
 	, AudioDevice(nullptr)
@@ -162,6 +163,7 @@ void FActiveSound::AddReferencedObjects( FReferenceCollector& Collector)
 	}
 
 	Collector.AddReferencedObject(Sound);
+	Collector.AddReferencedObject(SourceEffectChain);
 	Collector.AddReferencedObject(SoundClassOverride);
 
 	for (USoundConcurrency* Concurrency : ConcurrencySet)
@@ -204,6 +206,11 @@ void FActiveSound::SetSound(USoundBase* InSound)
 							|| (Sound && Sound->ShouldApplyInteriorVolumes());
 }
 
+void FActiveSound::SetSourceEffectChain(USoundEffectSourcePresetChain* InSourceEffectChain)
+{
+	SourceEffectChain = InSourceEffectChain;
+}
+
 void FActiveSound::SetSoundClass(USoundClass* SoundClass)
 {
 	SoundClassOverride = SoundClass;
@@ -522,7 +529,7 @@ void FActiveSound::UpdateWaveInstances(TArray<FWaveInstance*> &InWaveInstances,
 	}
 
 	// Set up the base source effect chain.
-	ParseParams.SourceEffectChain = Sound->SourceEffectChain;
+	ParseParams.SourceEffectChain = GetSourceEffectChain();
 
 	// Setup the envelope attack and release times
 	ParseParams.EnvelopeFollowerAttackTime = EnvelopeFollowerAttackTime;
diff --git a/Engine/Source/Runtime/Engine/Private/AudioDevice.cpp b/Engine/Source/Runtime/Engine/Private/AudioDevice.cpp
--- a/Engine/Source/Runtime/Engine/Private/AudioDevice.cpp
+++ b/Engine/Source/Runtime/Engine/Private/AudioDevice.cpp
@@ -3424,7 +3424,8 @@ int32 FAudioDevice::GetSortedActiveWaveInstances(TArray<FWaveInstance*>& WaveIns
 				if (ActiveSound->IsOneShot() && !ActiveSound->bIsPreviewSound)
 				{
 					// Don't stop a sound if it's playing effect chain tails or has effects playing, active sound will stop on its own in this case (in audio mixer).
-					if (!ActiveSound->Sound->SourceEffectChain || !ActiveSound->Sound->SourceEffectChain->bPlayEffectChainTails || ActiveSound->Sound->SourceEffectChain->Chain.Num() == 0)
+					USoundEffectSourcePresetChain* ActiveSourceEffectChain = ActiveSound->GetSourceEffectChain();
+					if (!ActiveSourceEffectChain || !ActiveSourceEffectChain->bPlayEffectChainTails || ActiveSourceEffectChain->Chain.Num() == 0)
 					{
 						const float Duration = ActiveSound->Sound->GetDuration();
 						if ((ActiveSound->Sound->HasDelayNode() || ActiveSound->Sound->HasConcatenatorNode()))
diff --git a/Engine/Source/Runtime/Engine/Private/Components/AudioComponent.cpp b/Engine/Source/Runtime/Engine/Private/Components/AudioComponent.cpp
--- a/Engine/Source/Runtime/Engine/Private/Components/AudioComponent.cpp
+++ b/Engine/Source/Runtime/Engine/Private/Components/AudioComponent.cpp
@@ -416,6 +416,7 @@ void UAudioComponent::PlayInternal(const float StartTime, const float FadeInDura
 			NewActiveSound.SetAudioComponent(*this);
 			NewActiveSound.SetWorld(GetWorld());
 			NewActiveSound.SetSound(Sound);
+			NewActiveSound.SetSourceEffectChain(SourceEffectChain);
 			NewActiveSound.SetSoundClass(SoundClassOverride);
 			NewActiveSound.ConcurrencySet = ConcurrencySet;
 
@@ -1477,3 +1478,8 @@ bool UAudioComponent::GetCookedEnvelopeDataForAllPlayingSounds(TArray<FSoundWave
 	}
 	return bHadData;
 }
+
+void UAudioComponent::SetSourceEffectChain(USoundEffectSourcePresetChain* InSourceEffectChain)
+{
+	SourceEffectChain = InSourceEffectChain;
+}
\ No newline at end of file
diff --git a/Engine/Source/Runtime/Engine/Public/ActiveSound.h b/Engine/Source/Runtime/Engine/Public/ActiveSound.h
--- a/Engine/Source/Runtime/Engine/Public/ActiveSound.h
+++ b/Engine/Source/Runtime/Engine/Public/ActiveSound.h
@@ -255,6 +255,7 @@ private:
 	uint32 WorldID;
 
 	USoundBase* Sound;
+	USoundEffectSourcePresetChain* SourceEffectChain;
 
 	uint64 AudioComponentID;
 	FName AudioComponentUserID;
@@ -294,6 +295,9 @@ public:
 	USoundBase* GetSound() const { return Sound; }
 	void SetSound(USoundBase* InSound);
 
+	USoundEffectSourcePresetChain* GetSourceEffectChain() const { return SourceEffectChain ? SourceEffectChain : Sound->SourceEffectChain; }
+	void SetSourceEffectChain(USoundEffectSourcePresetChain* InSourceEffectChain);
+
 	void SetSoundClass(USoundClass* SoundClass);
 
 	void SetAudioDevice(FAudioDevice* InAudioDevice)
