Updating and clarifying comments for spatialization method enumeration

diff --git a/Engine/Source/Runtime/Engine/Classes/Sound/SoundAttenuation.h b/Engine/Source/Runtime/Engine/Classes/Sound/SoundAttenuation.h
--- a/Engine/Source/Runtime/Engine/Classes/Sound/SoundAttenuation.h
+++ b/Engine/Source/Runtime/Engine/Classes/Sound/SoundAttenuation.h
@@ -20,10 +20,10 @@ enum ESoundDistanceCalc
 UENUM()
 enum ESoundSpatializationAlgorithm
 {
-	// Standard panning method for spatialization, built-in to the audio engine.
+	// Standard panning method for spatialization (linear or equal power method defined in project settings)
 	SPATIALIZATION_Default UMETA(DisplayName = "Panning"),
 
-	// 3rd party binaural spatialization (HRTF, Atmos). Requires a spatializaton plugin.
+	// Binaural spatialization method if available (requires headphones, enabled by plugins)
 	SPATIALIZATION_HRTF UMETA(DisplayName = "Binaural"),
 };
 
