Error on no input to RVT Output material node
Output black instead of magenta on missing Base Color input
Error on no RVT bound to RVT Sample material node
Integrated from //UE4/Dev-Rendering
#rnx

diff --git a/Engine/Source/Runtime/Engine/Private/Materials/MaterialExpressions.cpp b/Engine/Source/Runtime/Engine/Private/Materials/MaterialExpressions.cpp
--- a/Engine/Source/Runtime/Engine/Private/Materials/MaterialExpressions.cpp
+++ b/Engine/Source/Runtime/Engine/Private/Materials/MaterialExpressions.cpp
@@ -1899,11 +1899,16 @@ int32 UMaterialExpressionRuntimeVirtualTextureOutput::Compile(class FMaterialCom
 {
 	int32 CodeInput = INDEX_NONE;
 
+	if (!BaseColor.IsConnected() && !Specular.IsConnected() && !Roughness.IsConnected() && !Normal.IsConnected() && !Opacity.IsConnected())
+	{
+		Compiler->Error(TEXT("No inputs to Runtime Virtual Texture Output."));
+	}
+
 	// Order of outputs generates function names GetVirtualTextureOutput{index}
 	// These must match the function names called in VirtualTextureMaterial.usf
 	if (OutputIndex == 0)
 	{
-		CodeInput = BaseColor.IsConnected() ? BaseColor.Compile(Compiler) : Compiler->Constant3(1.f, 0.f, 1.f);
+		CodeInput = BaseColor.IsConnected() ? BaseColor.Compile(Compiler) : Compiler->Constant3(0.f, 0.f, 0.f);
 	}
 	else if (OutputIndex == 1)
 	{
@@ -2025,7 +2030,11 @@ int32 UMaterialExpressionRuntimeVirtualTextureSample::Compile(class FMaterialCom
 {
 	// Check validity of current virtual texture
 	bool bIsVirtualTextureValid = VirtualTexture != nullptr;
-	if (bIsVirtualTextureValid && VirtualTexture->GetMaterialType() != MaterialType)
+	if (!bIsVirtualTextureValid)
+	{
+		Compiler->Error(TEXT("Missing input Virtual Texture"));
+	}
+	else if (VirtualTexture->GetMaterialType() != MaterialType)
 	{
 		UEnum const* Enum = StaticEnum<ERuntimeVirtualTextureMaterialType>();
 		FString MaterialTypeDisplayName = Enum->GetDisplayNameTextByValue((int64)MaterialType).ToString();
@@ -2059,7 +2068,7 @@ int32 UMaterialExpressionRuntimeVirtualTextureSample::Compile(class FMaterialCom
 		}
 		else
 		{
-			return Compiler->Constant3(1.f, 0.f, 1.f);
+			return Compiler->Constant3(0.f, 0.f, 0.f);
 		}
 		break;
 	case 1:
diff --git a/Engine/Source/Runtime/Engine/Private/Materials/MaterialUniformExpressions.cpp b/Engine/Source/Runtime/Engine/Private/Materials/MaterialUniformExpressions.cpp
--- a/Engine/Source/Runtime/Engine/Private/Materials/MaterialUniformExpressions.cpp
+++ b/Engine/Source/Runtime/Engine/Private/Materials/MaterialUniformExpressions.cpp
@@ -835,7 +835,7 @@ void FUniformExpressionSet::FillUniformBuffer(const FMaterialRenderContext& Mate
 			// Don't have valid resources to bind for this VT, so make sure something is bound
 			if (!bValidResources)
 			{
-				*ResourceTablePhysicalTexturePtr = GWhiteTextureWithSRV->ShaderResourceViewRHI;
+				*ResourceTablePhysicalTexturePtr = GBlackTextureWithSRV->ShaderResourceViewRHI;
 				*ResourceTablePhysicalSamplerPtr = TStaticSamplerState<SF_Bilinear, AM_Clamp, AM_Clamp, AM_Clamp, 0, 8>::GetRHI();
 			}
 		}
