Fixing Atlas Curve expression to use parameters so it handles atlas resizing correctly
#rnx

diff --git a/Engine/Shaders/Public/ShaderVersion.ush b/Engine/Shaders/Public/ShaderVersion.ush
--- a/Engine/Shaders/Public/ShaderVersion.ush
+++ b/Engine/Shaders/Public/ShaderVersion.ush
@@ -2,4 +2,5 @@
 // Each time the console command is executed it generates a new GUID. As this file is included
 // in Platform.ush (which should be included in any shader) it allows to invalidate the shader DDC.
 // 
-// GUID = 3A7029DE8C044F0A97C52C4CD3F6182E
\ No newline at end of file
+// If you are merging streams and there is a conflict with this GUID you should make a new GUID rather than taking one or the other.
+// GUID = 676CDC95922E43F19A67C4089A159E89
diff --git a/Engine/Source/Editor/MaterialEditor/Private/MaterialPropertyHelpers.cpp b/Engine/Source/Editor/MaterialEditor/Private/MaterialPropertyHelpers.cpp
--- a/Engine/Source/Editor/MaterialEditor/Private/MaterialPropertyHelpers.cpp
+++ b/Engine/Source/Editor/MaterialEditor/Private/MaterialPropertyHelpers.cpp
@@ -1222,7 +1222,7 @@ void FMaterialPropertyHelpers::SetPositionFromCurveAsset(const FAssetData& Asset
 		return;
 	}
 
-	float NewValue = ((float)Index * Atlas->GradientPixelSize) / Atlas->TextureSize + (0.5f * Atlas->GradientPixelSize) / Atlas->TextureSize;
+	float NewValue = (float)Index;
 
 	// If changed, propagate the update
 	if (InParameter->ParameterValue != NewValue)
diff --git a/Engine/Source/Runtime/Engine/Classes/Curves/CurveLinearColorAtlas.h b/Engine/Source/Runtime/Engine/Classes/Curves/CurveLinearColorAtlas.h
--- a/Engine/Source/Runtime/Engine/Classes/Curves/CurveLinearColorAtlas.h
+++ b/Engine/Source/Runtime/Engine/Classes/Curves/CurveLinearColorAtlas.h
@@ -34,7 +34,7 @@ class ENGINE_API UCurveLinearColorAtlas : public UTexture2D
 	// How many slots are available per texture
 	FORCEINLINE uint32 MaxSlotsPerTexture()
 	{
-		return TextureSize / GradientPixelSize;
+		return TextureSize;
 	}
 
 	// Immediately render a new material to the specified slot index(SlotIndex must be within this section's range)
@@ -63,9 +63,6 @@ class ENGINE_API UCurveLinearColorAtlas : public UTexture2D
 	UPROPERTY(EditAnywhere, Category = "Curves")
 	uint32	TextureSize;						// Size of the lookup textures
 
-	UPROPERTY(VisibleAnywhere, Category = "Curves")
-	uint32	GradientPixelSize;					// How many pixels tall is any gradient slot
-
 	UPROPERTY(EditAnywhere, Category = "Curves")
 	TArray<UCurveLinearColor*> GradientCurves;
 
diff --git a/Engine/Source/Runtime/Engine/Private/Curves/CurveLinearColorAtlas.cpp b/Engine/Source/Runtime/Engine/Private/Curves/CurveLinearColorAtlas.cpp
--- a/Engine/Source/Runtime/Engine/Private/Curves/CurveLinearColorAtlas.cpp
+++ b/Engine/Source/Runtime/Engine/Private/Curves/CurveLinearColorAtlas.cpp
@@ -14,11 +14,10 @@ UCurveLinearColorAtlas::UCurveLinearColorAtlas(const FObjectInitializer& ObjectI
 	: Super(ObjectInitializer)
 {
 	TextureSize = 256;
-	GradientPixelSize = 1;
 #if WITH_EDITORONLY_DATA
 	bHasAnyDirtyTextures = false;
 	bShowDebugColorsForNullGradients = false;
-	SizeXY = { (float)TextureSize, (float)GradientPixelSize };
+	SizeXY = { (float)TextureSize,  1.0f };
 	MipGenSettings = TMGS_NoMipmaps;
 #endif
 	Filter = TextureFilter::TF_Bilinear;
@@ -48,7 +47,7 @@ void UCurveLinearColorAtlas::PostEditChangeProperty(struct FPropertyChangedEvent
 
 			Source.Init(TextureSize, TextureSize, 1, 1, TSF_RGBA16F);
 
-			SizeXY = { (float)TextureSize, (float)GradientPixelSize };
+			SizeXY = { (float)TextureSize, 1.0f };
 			UpdateTextures();
 			bRequiresNotifyMaterials = true;
 		}
@@ -93,7 +92,7 @@ void UCurveLinearColorAtlas::PostLoad()
 		}
 	}
 	Source.Init(TextureSize, TextureSize, 1, 1, TSF_RGBA16F);
-	SizeXY = { (float)TextureSize, (float)GradientPixelSize };
+	SizeXY = { (float)TextureSize, 1.0f };
 	UpdateTextures();
 #endif
 
@@ -133,7 +132,7 @@ void UCurveLinearColorAtlas::UpdateGradientSlot(UCurveLinearColor* Gradient)
 	if (SlotIndex != INDEX_NONE && (uint32)SlotIndex < MaxSlotsPerTexture())
 	{
 		// Determine the position of the gradient
-		int32 StartXY = SlotIndex * TextureSize * GradientPixelSize;
+		int32 StartXY = SlotIndex * TextureSize;
 
 		// Render the single gradient to the render target
 		RenderGradient(SrcData, Gradient, StartXY, SizeXY);
@@ -164,7 +163,7 @@ void UCurveLinearColorAtlas::UpdateTextures()
 	{
 		if (GradientCurves[i] != nullptr)
 		{
-			int32 StartXY = i * TextureSize * GradientPixelSize;
+			int32 StartXY = i * TextureSize;
 			RenderGradient(SrcData, GradientCurves[i], StartXY, SizeXY);
 		}
 
@@ -205,7 +204,7 @@ bool UCurveLinearColorAtlas::GetCurvePosition(UCurveLinearColor* InCurve, float&
 	Position = 0.0f;
 	if (Index != INDEX_NONE)
 	{
-		Position = ((float)Index * GradientPixelSize) / TextureSize + (0.5f * GradientPixelSize) / TextureSize;
+		Position = (float)Index;
 		return true;
 	}
 	return false;
diff --git a/Engine/Source/Runtime/Engine/Private/Materials/MaterialExpressions.cpp b/Engine/Source/Runtime/Engine/Private/Materials/MaterialExpressions.cpp
--- a/Engine/Source/Runtime/Engine/Private/Materials/MaterialExpressions.cpp
+++ b/Engine/Source/Runtime/Engine/Private/Materials/MaterialExpressions.cpp
@@ -15731,32 +15731,24 @@ void UMaterialExpressionCurveAtlasRowParameter::GetTexturesForceMaterialRecompil
 
 int32 UMaterialExpressionCurveAtlasRowParameter::Compile(class FMaterialCompiler* Compiler, int32 OutputIndex)
 {
-	// Need to recalculate the row index b/c a recompile could be triggered from the texture changing
-	int32 SlotIndex = INDEX_NONE;
 	if (Atlas && Curve)
 	{
-		SlotIndex = Atlas->GradientCurves.Find(Curve);
+		int32 Slot = Compiler->ScalarParameter(ParameterName, DefaultValue);
 
-		if (SlotIndex != INDEX_NONE)
-		{
-			float NewValue = ((float)SlotIndex * Atlas->GradientPixelSize) / Atlas->TextureSize + (0.5f * Atlas->GradientPixelSize) / Atlas->TextureSize;
-			SetParameterValue(ParameterName, NewValue);
-		}
+		// Get Atlas texture object and texture size
+		int32 AtlasRef = INDEX_NONE;
+		int32 AtlasCode = Compiler->Texture(Atlas, AtlasRef, SAMPLERTYPE_LinearColor, SSM_Clamp_WorldGroupSettings, TMVM_None);
+		int32 AtlasSize = Compiler->ForceCast(Compiler->TextureProperty(AtlasCode, TMTM_TextureSize), MCT_Float1);
+
+		// Calculate UVs from size and slot
 		// if the input is hooked up, use it, otherwise use the internal constant
 		int32 Arg1 = InputTime.GetTracedInput().Expression ? InputTime.Compile(Compiler) : Compiler->Constant(0);
-		int32 Arg2 = Compiler->ScalarParameter(ParameterName, DefaultValue);
+		int32 Arg2 = Compiler->Div(Compiler->Add(Slot, Compiler->Constant(0.5)), AtlasSize);
 
 		int32 UV = Compiler->AppendVector(Arg1, Arg2);
-		return CompileTextureSample(
-			Compiler,
-			Atlas,
-			UV,
-			SAMPLERTYPE_LinearColor,
-			TOptional<FName>(),
-			INDEX_NONE,
-			INDEX_NONE,
-			TMVM_None,
-			SSM_Clamp_WorldGroupSettings);
+
+		// Sample texture
+		return Compiler->TextureSample(AtlasCode, UV, SAMPLERTYPE_LinearColor, INDEX_NONE, INDEX_NONE, TMVM_None, SSM_Clamp_WorldGroupSettings, AtlasRef, false);
 	}
 	return INDEX_NONE;
 }
@@ -15773,7 +15765,7 @@ void UMaterialExpressionCurveAtlasRowParameter::PostEditChangeProperty(FProperty
 		}
 		if (SlotIndex != INDEX_NONE)
 		{
-			float NewValue = ((float)SlotIndex * Atlas->GradientPixelSize) / Atlas->TextureSize + (0.5f * Atlas->GradientPixelSize) / Atlas->TextureSize;
+			float NewValue = (float)SlotIndex;
 			SetParameterValue(ParameterName, NewValue);
 		}
 		else
diff --git a/Engine/Source/Runtime/Engine/Private/Materials/MaterialInstance.cpp b/Engine/Source/Runtime/Engine/Private/Materials/MaterialInstance.cpp
--- a/Engine/Source/Runtime/Engine/Private/Materials/MaterialInstance.cpp
+++ b/Engine/Source/Runtime/Engine/Private/Materials/MaterialInstance.cpp
@@ -3496,7 +3496,7 @@ void UMaterialInstance::SetScalarParameterAtlasInternal(const FMaterialParameter
 			return;
 		}
 
-		float NewValue = ((float)Index * Atlas->GradientPixelSize) / Atlas->TextureSize + (0.5f * Atlas->GradientPixelSize) / Atlas->TextureSize;
+		float NewValue = (float)Index;
 		
 		// Don't enqueue an update if it isn't needed
 		if (ParameterValue->ParameterValue != NewValue)
