Add compression error scale to anim sequence. Allows compression error tests to account for the playback scale of the animation

diff --git a/Engine/Source/Runtime/Engine/Classes/Animation/AnimCompress.h b/Engine/Source/Runtime/Engine/Classes/Animation/AnimCompress.h
--- a/Engine/Source/Runtime/Engine/Classes/Animation/AnimCompress.h
+++ b/Engine/Source/Runtime/Engine/Classes/Animation/AnimCompress.h
@@ -633,7 +633,7 @@ public:
 #endif
 
 #if WITH_EDITOR
-	FString MakeDDCKey();
+	void PopulateDDCKeyArchive(FArchive& Ar) { PopulateDDCKey(Ar); }
 
 protected:
 	virtual void PopulateDDCKey(FArchive& Ar);
diff --git a/Engine/Source/Runtime/Engine/Classes/Animation/AnimCurveCompressionSettings.h b/Engine/Source/Runtime/Engine/Classes/Animation/AnimCurveCompressionSettings.h
--- a/Engine/Source/Runtime/Engine/Classes/Animation/AnimCurveCompressionSettings.h
+++ b/Engine/Source/Runtime/Engine/Classes/Animation/AnimCurveCompressionSettings.h
@@ -38,6 +38,6 @@ class ENGINE_API UAnimCurveCompressionSettings : public UObject
 	bool Compress(const FCompressibleAnimData& AnimSeq, FCompressedAnimSequence& OutCompressedData) const;
 
 	/** Generates a DDC key that takes into account the current settings and selected codec. */
-	FString MakeDDCKey() const;
+	void PopulateDDCKey(FArchive& Ar) const;
 #endif
 };
diff --git a/Engine/Source/Runtime/Engine/Classes/Animation/AnimSequence.h b/Engine/Source/Runtime/Engine/Classes/Animation/AnimSequence.h
--- a/Engine/Source/Runtime/Engine/Classes/Animation/AnimSequence.h
+++ b/Engine/Source/Runtime/Engine/Classes/Animation/AnimSequence.h
@@ -310,6 +310,14 @@ public:
 	UPROPERTY(Category = Compression, EditAnywhere)
 	bool bAllowFrameStripping;
 
+	/**
+	 * Set a scale for error threshold on compression. This is useful if the animation will 
+	 * be played back at a different scale (e.g. if you know the animation will be played
+	 * on an actor/component that is scaled up by a factor of 10, set this value to 10)
+	 */
+	UPROPERTY(Category = Compression, EditAnywhere)
+	float CompressionErrorThresholdScale;
+
 #endif
 
 	/** The curve compression settings used to compress curves in this sequence. */
@@ -497,6 +505,9 @@ public:
 	
 	// Adds a new track (if no track of the supplied name is found) to the raw animation data, optionally setting it to TrackData.
 	int32 AddNewRawTrack(FName TrackName, FRawAnimSequenceTrack* TrackData = nullptr);
+
+	// Get the Alternate compression error threshold 
+	float GetAltCompressionErrorThreshold() const;
 #endif
 
 	const TArray<FTrackToSkeletonMap>& GetRawTrackToSkeletonMapTable() const { return TrackToSkeletonMapTable; }
diff --git a/Engine/Source/Runtime/Engine/Classes/Animation/AnimStreamable.h b/Engine/Source/Runtime/Engine/Classes/Animation/AnimStreamable.h
--- a/Engine/Source/Runtime/Engine/Classes/Animation/AnimStreamable.h
+++ b/Engine/Source/Runtime/Engine/Classes/Animation/AnimStreamable.h
@@ -208,13 +208,15 @@ public:
 	private:
 
 #if WITH_EDITOR
+	float GetAltCompressionErrorThreshold() const;
+
 	ENGINE_API void RequestCompressedData(const  ITargetPlatform* Platform=nullptr);
 
 	void UpdateRawData();
 
-	FString GetBaseDDCKey() const;
+	FString GetBaseDDCKey(uint32 NumChunks, float AltCompressionErrorThreshold) const;
 
-	void RequestCompressedDataForChunk(const FString& ChunkDDCKey, FAnimStreamableChunk& Chunk, const uint32 FrameStart, const uint32 FrameEnd, TSharedRef<FAnimCompressContext> CompressContext);
+	void RequestCompressedDataForChunk(const FString& ChunkDDCKey, FAnimStreamableChunk& Chunk, const int32 ChunkIndex, const uint32 FrameStart, const uint32 FrameEnd, TSharedRef<FAnimCompressContext> CompressContext);
 #endif
 
 	bool bUseRawDataOnly;
diff --git a/Engine/Source/Runtime/Engine/Private/Animation/AnimCompress.cpp b/Engine/Source/Runtime/Engine/Private/Animation/AnimCompress.cpp
--- a/Engine/Source/Runtime/Engine/Private/Animation/AnimCompress.cpp
+++ b/Engine/Source/Runtime/Engine/Private/Animation/AnimCompress.cpp
@@ -1581,18 +1581,6 @@ void UAnimCompress::CoalesceCompressedSegments(UAnimSequence& AnimSeq, const TAr
 
 #if WITH_EDITOR
 
-FString UAnimCompress::MakeDDCKey()
-{
-	FArcToHexString ArcToHexString;
-
-	// Serialize the compression settings into a temporary array. The archive
-	// is flagged as persistent so that machines of different endianness produce
-	// identical binary results.
-	PopulateDDCKey(ArcToHexString.Ar);
-
-	return ArcToHexString.MakeString();
-}
-
 void UAnimCompress::PopulateDDCKey(FArchive& Ar)
 {
 	uint8 TCF, RCF, SCF;
diff --git a/Engine/Source/Runtime/Engine/Private/Animation/AnimCompressionDerivedData.cpp b/Engine/Source/Runtime/Engine/Private/Animation/AnimCompressionDerivedData.cpp
--- a/Engine/Source/Runtime/Engine/Private/Animation/AnimCompressionDerivedData.cpp
+++ b/Engine/Source/Runtime/Engine/Private/Animation/AnimCompressionDerivedData.cpp
@@ -27,6 +27,14 @@ FDerivedDataAnimationCompression::~FDerivedDataAnimationCompression()
 {
 }
 
+const TCHAR* FDerivedDataAnimationCompression::GetVersionString() const
+{
+	// This is a version string that mimics the old versioning scheme. If you
+	// want to bump this version, generate a new guid using VS->Tools->Create GUID and
+	// return it here. Ex.
+	return TEXT("2E79BF10172A48FDACA76883B8951538");
+}
+
 bool FDerivedDataAnimationCompression::Build( TArray<uint8>& OutDataArray )
 {
 	check(DataToCompressPtr.IsValid());
diff --git a/Engine/Source/Runtime/Engine/Private/Animation/AnimCompressionDerivedData.h b/Engine/Source/Runtime/Engine/Private/Animation/AnimCompressionDerivedData.h
--- a/Engine/Source/Runtime/Engine/Private/Animation/AnimCompressionDerivedData.h
+++ b/Engine/Source/Runtime/Engine/Private/Animation/AnimCompressionDerivedData.h
@@ -55,13 +55,7 @@ public:
 		return TypeName;
 	}
 
-	virtual const TCHAR* GetVersionString() const override
-	{
-		// This is a version string that mimics the old versioning scheme. If you
-		// want to bump this version, generate a new guid using VS->Tools->Create GUID and
-		// return it here. Ex.
-		return TEXT("1F1656B9E10142729AB16650D9821B1Ff");
-	}
+	virtual const TCHAR* GetVersionString() const override;
 
 	virtual FString GetPluginSpecificCacheKeySuffix() const override
 	{
diff --git a/Engine/Source/Runtime/Engine/Private/Animation/AnimCompressionTypes.cpp b/Engine/Source/Runtime/Engine/Private/Animation/AnimCompressionTypes.cpp
--- a/Engine/Source/Runtime/Engine/Private/Animation/AnimCompressionTypes.cpp
+++ b/Engine/Source/Runtime/Engine/Private/Animation/AnimCompressionTypes.cpp
@@ -123,7 +123,7 @@ void StripFramesOdd(TArray<ArrayValue>& Keys, const int32 NumFrames)
 	}
 }
 
-FCompressibleAnimData::FCompressibleAnimData(class UAnimSequence* InSeq, const bool bPerformStripping) :
+FCompressibleAnimData::FCompressibleAnimData(class UAnimSequence* InSeq, const bool bPerformStripping, const float InAltCompressionErrorThreshold) :
 #if WITH_EDITOR
 	RequestedCompressionScheme(InSeq->CompressionScheme) ,
 #endif
@@ -134,6 +134,7 @@ FCompressibleAnimData::FCompressibleAnimData(class UAnimSequence* InSeq, const b
 	, SequenceLength(InSeq->SequenceLength)
 	, NumFrames(InSeq->GetRawNumberOfFrames())
 	, bIsValidAdditive(InSeq->IsValidAdditive())
+	, AltCompressionErrorThreshold(InAltCompressionErrorThreshold)
 	, Name(InSeq->GetName())
 	, FullName(InSeq->GetFullName())
 	, AnimFName(InSeq->GetFName())
@@ -197,7 +198,7 @@ FCompressibleAnimData::FCompressibleAnimData(class UAnimSequence* InSeq, const b
 #endif
 }
 
-FCompressibleAnimData::FCompressibleAnimData(UAnimCompress* InRequestedCompressionScheme, UAnimCurveCompressionSettings* InCurveCompressionSettings, USkeleton* InSkeleton, EAnimInterpolationType InInterpolation, float InSequenceLength, int32 InNumFrames) :
+FCompressibleAnimData::FCompressibleAnimData(UAnimCompress* InRequestedCompressionScheme, UAnimCurveCompressionSettings* InCurveCompressionSettings, USkeleton* InSkeleton, EAnimInterpolationType InInterpolation, float InSequenceLength, int32 InNumFrames, const float InAltCompressionErrorThreshold) :
 #if WITH_EDITOR
 	RequestedCompressionScheme(InRequestedCompressionScheme) ,
 #endif
@@ -207,12 +208,25 @@ FCompressibleAnimData::FCompressibleAnimData(UAnimCompress* InRequestedCompressi
 	, SequenceLength(InSequenceLength)
 	, NumFrames(InNumFrames)
 	, bIsValidAdditive(false)
+	, AltCompressionErrorThreshold(InAltCompressionErrorThreshold)
 {
 #if WITH_EDITOR
 	FAnimationUtils::BuildSkeletonMetaData(Skeleton, BoneData);
 #endif
 }
 
+FCompressibleAnimData::FCompressibleAnimData() : RequestedCompressionScheme(nullptr)
+, CurveCompressionSettings(nullptr)
+, Skeleton(nullptr)
+, Interpolation((EAnimInterpolationType)0)
+, SequenceLength(0.f)
+, NumFrames(0)
+, bIsValidAdditive(false)
+, AltCompressionErrorThreshold(FAnimationUtils::GetAlternativeCompressionThreshold())
+{
+
+}
+
 void FCompressibleAnimData::Update(FCompressedAnimSequence& InOutCompressedData) const
 {
 	InOutCompressedData.CompressedTrackToSkeletonMapTable = TrackToSkeletonMapTable;
diff --git a/Engine/Source/Runtime/Engine/Private/Animation/AnimCurveCompressionSettings.cpp b/Engine/Source/Runtime/Engine/Private/Animation/AnimCurveCompressionSettings.cpp
--- a/Engine/Source/Runtime/Engine/Private/Animation/AnimCurveCompressionSettings.cpp
+++ b/Engine/Source/Runtime/Engine/Private/Animation/AnimCurveCompressionSettings.cpp
@@ -41,30 +41,17 @@ bool UAnimCurveCompressionSettings::Compress(const FCompressibleAnimData& AnimSe
 	return Success;
 }
 
-FString UAnimCurveCompressionSettings::MakeDDCKey() const
+void UAnimCurveCompressionSettings::PopulateDDCKey(FArchive& Ar) const
 {
-	if (Codec == nullptr)
+	if (Codec)
 	{
-		return TEXT("<Missing Codec>");
+		Codec->PopulateDDCKey(Ar);
 	}
-
-	TArray<uint8> TempBytes;
-	TempBytes.Reserve(64);
-
-	// Serialize the compression settings into a temporary array. The archive
-	// is flagged as persistent so that machines of different endianness produce
-	// identical binary results.
-	FMemoryWriter Ar(TempBytes, /*bIsPersistent=*/ true);
-
-	Codec->PopulateDDCKey(Ar);
-
-	FString Key;
-	Key.Reserve(TempBytes.Num() + 1);
-	for (int32 ByteIndex = 0; ByteIndex < TempBytes.Num(); ++ByteIndex)
+	else
 	{
-		ByteToHex(TempBytes[ByteIndex], Key);
+		static FString NoCodecString(TEXT("<Missing Codec>"));
+		Ar << NoCodecString;
 	}
-
-	return Key;
 }
+
 #endif
diff --git a/Engine/Source/Runtime/Engine/Private/Animation/AnimSequence.cpp b/Engine/Source/Runtime/Engine/Private/Animation/AnimSequence.cpp
--- a/Engine/Source/Runtime/Engine/Private/Animation/AnimSequence.cpp
+++ b/Engine/Source/Runtime/Engine/Private/Animation/AnimSequence.cpp
@@ -200,7 +200,7 @@ void OnCVarsChanged()
 
 FAutoConsoleVariableSink AnimationCVarSink(FConsoleCommandDelegate::CreateStatic(&OnCVarsChanged));
 
-FString GetAnimSequenceSpecificCacheKeySuffix(const UAnimSequence& Seq, const bool bPerformStripping)
+FString GetAnimSequenceSpecificCacheKeySuffix(const UAnimSequence& Seq, bool bPerformStripping, float AltErrorThreshold)
 {
 	//Make up our content key consisting of:
 	//	* Global animation compression version
@@ -216,11 +216,15 @@ FString GetAnimSequenceSpecificCacheKeySuffix(const UAnimSequence& Seq, const bo
 	char AdditiveType = bIsValidAdditive ? NibbleToTChar(Seq.AdditiveAnimType) : '0';
 	char RefType = bIsValidAdditive ? NibbleToTChar(Seq.RefPoseType) : '0';
 
-	const int32 StripFrame = bPerformStripping ? 1 : 0;
+	FArcToHexString ArcToHexString;
 
-	FString Ret = FString::Printf(TEXT("%i_%i_%s%s%s_%c%c%i_%s_%s_%s"),
+	ArcToHexString.Ar << AltErrorThreshold;
+	ArcToHexString.Ar << bPerformStripping;
+	Seq.CompressionScheme->PopulateDDCKeyArchive(ArcToHexString.Ar);
+	Seq.CurveCompressionSettings->PopulateDDCKey(ArcToHexString.Ar);
+
+	FString Ret = FString::Printf(TEXT("%i_%s%s%s_%c%c%i_%s_%s"),
 		Seq.CompressCommandletVersion,
-		StripFrame,
 		*Seq.GetRawDataGuid().ToString(),
 		*Seq.GetSkeleton()->GetGuid().ToString(),
 		*Seq.GetSkeleton()->GetVirtualBoneGuid().ToString(),
@@ -228,8 +232,7 @@ FString GetAnimSequenceSpecificCacheKeySuffix(const UAnimSequence& Seq, const bo
 		RefType,
 		Seq.RefFrameIndex,
 		(bIsValidAdditive && Seq.RefPoseSeq) ? *Seq.RefPoseSeq->GetRawDataGuid().ToString() : TEXT("NoAdditiveGuid"),
-		*Seq.CompressionScheme->MakeDDCKey(),
-		*Seq.CurveCompressionSettings->MakeDDCKey()
+		*ArcToHexString.MakeString()
 	);
 
 	return Ret;
@@ -547,7 +550,7 @@ UAnimSequence::UAnimSequence(const FObjectInitializer& ObjectInitializer)
 	ImportFileFramerate = 0.0f;
 	ImportResampleFramerate = 0;
 	bAllowFrameStripping = true;
-
+	CompressionErrorThresholdScale = 1.f;
 #endif
 
 	InitCurveCompressionScheme();
@@ -1163,6 +1166,8 @@ void UAnimSequence::PostEditChangeProperty(FPropertyChangedEvent& PropertyChange
 	}
 
 	bool bAdditiveSettingsChanged = false;
+	bool bCompressionAffectingSettingsChanged = false;
+
 	if(PropertyChangedEvent.Property)
 	{
 		const bool bChangedRefFrameIndex = PropertyChangedEvent.Property->GetFName() == GET_MEMBER_NAME_CHECKED(UAnimSequence, RefFrameIndex);
@@ -1180,10 +1185,14 @@ void UAnimSequence::PostEditChangeProperty(FPropertyChangedEvent& PropertyChange
 			bAdditiveSettingsChanged = true;
 		}
 		
+		bCompressionAffectingSettingsChanged =   PropertyChangedEvent.Property->GetFName() == GET_MEMBER_NAME_CHECKED(UAnimSequence, bAllowFrameStripping)
+											  || PropertyChangedEvent.Property->GetFName() == GET_MEMBER_NAME_CHECKED(UAnimSequence, CompressionErrorThresholdScale);
 	}
 
+	const bool bNeedPostProcess = !IsCompressedDataValid() || bAdditiveSettingsChanged || bCompressionAffectingSettingsChanged;
+
 	// @Todo fix me: This is temporary fix to make sure they always have compressed data
-	if (RawAnimationData.Num() > 0 && (!IsCompressedDataValid() || bAdditiveSettingsChanged))
+	if (RawAnimationData.Num() > 0 && bNeedPostProcess)
 	{
 		PostProcessSequence();
 	}
@@ -1783,6 +1792,12 @@ void UAnimSequence::RetargetBoneTransform(FTransform& BoneTransform, const int32
 }
 
 #if WITH_EDITOR
+float UAnimSequence::GetAltCompressionErrorThreshold() const
+{
+	// CompressionErrorThresholdScale represents a world scale we will play back the animation at
+	return FAnimationUtils::GetAlternativeCompressionThreshold() / CompressionErrorThresholdScale;
+}
+
 /** Utility function to crop data from a RawAnimSequenceTrack */
 static int32 CropRawTrack(FRawAnimSequenceTrack& RawTrack, int32 StartKey, int32 NumKeys, int32 TotalNumOfFrames)
 {
@@ -2194,7 +2209,9 @@ void UAnimSequence::RequestAnimCompression(FRequestAnimCompressionParams Params)
 
 		const int32 PreviousCompressionSize = GetApproxCompressedSize();
 
-		FString AssetDDCKey = GetAnimSequenceSpecificCacheKeySuffix(*this, bPerformStripping);
+		const float AltCompressionErrorThreshold = GetAltCompressionErrorThreshold();
+
+		FString AssetDDCKey = GetAnimSequenceSpecificCacheKeySuffix(*this, bPerformStripping, AltCompressionErrorThreshold);
 
 		TArray<uint8> OutData;
 		{
@@ -2213,7 +2230,7 @@ void UAnimSequence::RequestAnimCompression(FRequestAnimCompressionParams Params)
 				// No trivial key removal is done at this point (impossible error metrics of -1), since all of the techniques will perform it themselves
 				CompressRawAnimData(-1.0f, -1.0f);
 
-				TSharedRef<FCompressibleAnimData> CompressibleData = MakeShared<FCompressibleAnimData>(this, bPerformStripping);
+				TSharedRef<FCompressibleAnimData> CompressibleData = MakeShared<FCompressibleAnimData>(this, bPerformStripping, AltCompressionErrorThreshold);
 				AnimCompressor->SetCompressibleData(CompressibleData);
 
 				if (bSkipDDC || (CompressCommandletVersion == INDEX_NONE))
diff --git a/Engine/Source/Runtime/Engine/Private/Animation/AnimStreamable.cpp b/Engine/Source/Runtime/Engine/Private/Animation/AnimStreamable.cpp
--- a/Engine/Source/Runtime/Engine/Private/Animation/AnimStreamable.cpp
+++ b/Engine/Source/Runtime/Engine/Private/Animation/AnimStreamable.cpp
@@ -23,7 +23,7 @@ DECLARE_CYCLE_STAT(TEXT("AnimStreamable GetAnimationPose"), STAT_AnimStreamable_
 // This is a version string for the streaming anim chunk logic
 // If you want to bump this version, generate a new guid using VS->Tools->Create GUID and
 // set it here
-const TCHAR* StreamingAnimChunkVersion = TEXT("1F1656B9E10142729AB16650D9821B1Fuf");
+const TCHAR* StreamingAnimChunkVersion = TEXT("1F1656B9E10142729AB16650D9821B1F");
 
 const float MINIMUM_CHUNK_SIZE = 2.0f;
 
@@ -328,6 +328,11 @@ void UAnimStreamable::PostLoad()
 		NonConstSeq->ConditionalPostLoad();
 	}
 
+	if (SourceSequence)
+	{
+		CompressionScheme = SourceSequence->CompressionScheme;
+	}
+
 	if (SourceSequence && (GenerateGuidFromRawAnimData(SourceSequence->GetRawAnimationData(), SourceSequence->RawCurveData) != RawDataGuid))
 	{
 		InitFrom(SourceSequence);
@@ -336,8 +341,6 @@ void UAnimStreamable::PostLoad()
 	{
 		RequestCompressedData(); // Grab compressed data for current platform
 	}
-
-	RequestCompressedData(); // Grab compressed data for current platform
 #else
 	IStreamingManager::Get().GetAnimationStreamingManager().AddStreamingAnim(this); // This will be handled by RequestCompressedData in editor builds
 
@@ -374,11 +377,17 @@ int32 UAnimStreamable::GetChunkIndexForTime(const TArray<FAnimStreamableChunk>&
 }
 
 #if WITH_EDITOR
+float UAnimStreamable::GetAltCompressionErrorThreshold() const
+{
+	return SourceSequence ? SourceSequence->GetAltCompressionErrorThreshold() : FAnimationUtils::GetAlternativeCompressionThreshold();
+}
+
 void UAnimStreamable::InitFrom(const UAnimSequence* InSourceSequence)
 {
 	Modify();
 	SetSkeleton(InSourceSequence->GetSkeleton());
 	SourceSequence = InSourceSequence;
+	CompressionScheme = InSourceSequence->CompressionScheme;
 
 	RawAnimationData = InSourceSequence->GetRawAnimationData();
 	RawCurveData = InSourceSequence->RawCurveData;
@@ -405,16 +414,11 @@ void UAnimStreamable::InitFrom(const UAnimSequence* InSourceSequence)
 	UpdateRawData();
 }
 
-FString GetChunkDDCKey(const FString& BaseKey, uint32 ChunkIndex, uint32 NumChunks)
+FString GetChunkDDCKey(const FString& BaseKey, uint32 ChunkIndex)
 {
-	FArcToHexString HexStringBuilder;
-
-	HexStringBuilder.Ar << ChunkIndex;
-	HexStringBuilder.Ar << NumChunks;
-
-	return FString::Printf(TEXT("%s%s"),
+	return FString::Printf(TEXT("%s%u"),
 		*BaseKey,
-		*HexStringBuilder.MakeString()
+		ChunkIndex
 	);
 }
 
@@ -495,7 +499,7 @@ void UAnimStreamable::RequestCompressedData(const ITargetPlatform* Platform)
 
 	PlatformData.Chunks.AddDefaulted(NumChunks);
 
-	const FString BaseDDCKey = GetBaseDDCKey();
+	const FString BaseDDCKey = GetBaseDDCKey(NumChunks, GetAltCompressionErrorThreshold());
 
 	const bool bInAllowAlternateCompressor = false;
 	const bool bInOutput				   = false;
@@ -504,13 +508,13 @@ void UAnimStreamable::RequestCompressedData(const ITargetPlatform* Platform)
 
 	for (uint32 ChunkIndex = 0; ChunkIndex < NumChunks; ++ChunkIndex)
 	{
-		FString ChunkDDCKey = GetChunkDDCKey(BaseDDCKey, ChunkIndex, NumChunks);
+		FString ChunkDDCKey = GetChunkDDCKey(BaseDDCKey, ChunkIndex);
 
 		const bool bLastChunk = (ChunkIndex == (NumChunks - 1));
 		const uint32 FrameStart = ChunkIndex * FramesPerChunk;
 		const uint32 FrameEnd = bLastChunk ? NumFramesToChunk : (ChunkIndex + 1) * FramesPerChunk;
 
-		RequestCompressedDataForChunk(ChunkDDCKey, PlatformData.Chunks[ChunkIndex], FrameStart, FrameEnd, CompressContext);
+		RequestCompressedDataForChunk(ChunkDDCKey, PlatformData.Chunks[ChunkIndex], ChunkIndex, FrameStart, FrameEnd, CompressContext);
 	}
 
 	if (bIsRunningPlatform)
@@ -555,7 +559,7 @@ void MakeKeyChunk(const TArray<KeyType>& SrcKeys, TArray<KeyType>& DestKeys, int
 	}
 }
 
-void UAnimStreamable::RequestCompressedDataForChunk(const FString& ChunkDDCKey, FAnimStreamableChunk& Chunk, const uint32 FrameStart, const uint32 FrameEnd, TSharedRef<FAnimCompressContext> CompressContext)
+void UAnimStreamable::RequestCompressedDataForChunk(const FString& ChunkDDCKey, FAnimStreamableChunk& Chunk, const int32 ChunkIndex, const uint32 FrameStart, const uint32 FrameEnd, TSharedRef<FAnimCompressContext> CompressContext)
 {
 	// Need to unify with Anim Sequence!
 
@@ -576,7 +580,7 @@ void UAnimStreamable::RequestCompressedDataForChunk(const FString& ChunkDDCKey,
 
 		if (bSkipDDC || !GetDerivedDataCacheRef().GetSynchronous(*FinalDDCKey, OutData))
 		{
-			TSharedRef<FCompressibleAnimData> CompressibleData = MakeShared<FCompressibleAnimData>(CompressionScheme, CurveCompressionSettings, GetSkeleton(), Interpolation, Chunk.SequenceLength, ChunkNumFrames+1);
+			TSharedRef<FCompressibleAnimData> CompressibleData = MakeShared<FCompressibleAnimData>(CompressionScheme, CurveCompressionSettings, GetSkeleton(), Interpolation, Chunk.SequenceLength, ChunkNumFrames+1, GetAltCompressionErrorThreshold());
 
 			CompressibleData->RawAnimationData.AddDefaulted(RawAnimationData.Num());
 
@@ -599,6 +603,19 @@ void UAnimStreamable::RequestCompressedDataForChunk(const FString& ChunkDDCKey,
 				CompressibleData->RawCurveData = RawCurveData;
 			}
 
+			/*if (SourceSequence && ChunkIndex == 7)
+			{
+				UAnimSequence* Seq = const_cast<UAnimSequence*>(SourceSequence); //Stop update loop
+				SourceSequence = nullptr;
+				Seq->RawAnimationData = CompressibleData->RawAnimationData;
+				Seq->SequenceLength = CompressibleData->SequenceLength;
+				Seq->NumFrames = CompressibleData->NumFrames;
+				Seq->MarkRawDataAsModified(false);
+				Seq->OnRawDataChanged();
+
+				SourceSequence = Seq;
+			}*/
+
 			CompressibleData->TrackToSkeletonMapTable = TrackToSkeletonMapTable;
 			AnimCompressor->SetCompressibleData(CompressibleData);
 
@@ -638,7 +655,7 @@ void UAnimStreamable::UpdateRawData()
 	RequestCompressedData();
 }
 
-FString UAnimStreamable::GetBaseDDCKey() const
+FString UAnimStreamable::GetBaseDDCKey(uint32 NumChunks, float AltCompressionErrorThreshold) const
 {
 	//Make up our content key consisting of:
 	//  * Streaming Anim Chunk logic version
@@ -648,13 +665,19 @@ FString UAnimStreamable::GetBaseDDCKey() const
 	//	* Compression Settings
 	//	* Curve compression settings
 
-	FString Ret = FString::Printf(TEXT("%s%s%s%s_%s_%s"),
+	FArcToHexString ArcToHexString;
+
+	ArcToHexString.Ar << NumChunks;
+	ArcToHexString.Ar << AltCompressionErrorThreshold;
+	CompressionScheme->PopulateDDCKeyArchive(ArcToHexString.Ar);
+	CurveCompressionSettings->PopulateDDCKey(ArcToHexString.Ar);
+
+	FString Ret = FString::Printf(TEXT("%s%s%s%s_%s"),
 		StreamingAnimChunkVersion,
 		*RawDataGuid.ToString(),
 		*GetSkeleton()->GetGuid().ToString(),
 		*GetSkeleton()->GetVirtualBoneGuid().ToString(),
-		*CompressionScheme->MakeDDCKey(),
-		*CurveCompressionSettings->MakeDDCKey()
+		*ArcToHexString.MakeString()
 	);
 
 	return Ret;
diff --git a/Engine/Source/Runtime/Engine/Private/Animation/AnimationUtils.cpp b/Engine/Source/Runtime/Engine/Private/Animation/AnimationUtils.cpp
--- a/Engine/Source/Runtime/Engine/Private/Animation/AnimationUtils.cpp
+++ b/Engine/Source/Runtime/Engine/Private/Animation/AnimationUtils.cpp
@@ -789,9 +789,6 @@ void FAnimationUtils::CompressAnimSequence(const FCompressibleAnimData& Compress
 {
 	if (FPlatformProperties::HasEditorOnlyData())
 	{
-		// get the master tolerance we will use to guide recompression
-		float MasterTolerance = GetAlternativeCompressionThreshold(); 
-
 		bool bOnlyCheckForMissingSkeletalMeshes = UAnimationSettings::Get()->bOnlyCheckForMissingSkeletalMeshes;
 		if (!bOnlyCheckForMissingSkeletalMeshes)
 		{
@@ -834,7 +831,7 @@ void FAnimationUtils::CompressAnimSequence(const FCompressibleAnimData& Compress
 					CompressibleAnimData,
 					OutCompressedData,
 					CompressContext,
-					CompressContext.bAllowAlternateCompressor ? MasterTolerance : 0.0f,
+					CompressibleAnimData.AltCompressionErrorThreshold,
 					bFirstRecompressUsingCurrentOrDefault,
 					bForceBelowThreshold,
 					bRaiseMaxErrorToExisting,
@@ -943,14 +940,15 @@ void FAnimationUtils::CompressAnimSequenceExplicit(
 	// we must have raw data to continue
 	if(NumRawDataTracks > 0 )
 	{
+		// See if we're trying alternate compressors
 		// If compression Scheme is automatic, then we definitely want to 'TryAlternateCompressor'.
-		if (CompressibleAnimData.RequestedCompressionScheme && CompressibleAnimData.RequestedCompressionScheme->IsA(UAnimCompress_Automatic::StaticClass()))
-		{
-			MasterTolerance = GetAlternativeCompressionThreshold();
-		}
+		const bool bIsAutomatic = CompressibleAnimData.RequestedCompressionScheme && CompressibleAnimData.RequestedCompressionScheme->IsA(UAnimCompress_Automatic::StaticClass());
+		const bool bTryAlternateCompressor = bIsAutomatic || CompressContext.bAllowAlternateCompressor;
 
-		// See if we're trying alternate compressors
-		const bool bTryAlternateCompressor = MasterTolerance > 0.0f;
+		// We shouldn't override as this value can come from automatic compression sequences
+		// but that was broken by CL 3489273
+		// Preserving override behaviour till issue can be properly addressed
+		MasterTolerance = CompressibleAnimData.AltCompressionErrorThreshold;
 
 		AnimationErrorStats TrueOriginalErrorStats;
 		FAnimationUtils::ComputeCompressionError(CompressibleAnimData, OutCompressedData, TrueOriginalErrorStats);
diff --git a/Engine/Source/Runtime/Engine/Public/Animation/AnimCompressionTypes.h b/Engine/Source/Runtime/Engine/Public/Animation/AnimCompressionTypes.h
--- a/Engine/Source/Runtime/Engine/Public/Animation/AnimCompressionTypes.h
+++ b/Engine/Source/Runtime/Engine/Public/Animation/AnimCompressionTypes.h
@@ -183,20 +183,11 @@ struct ENGINE_API FCompressedSegment
 struct ENGINE_API FCompressibleAnimData
 {
 public:
-	FCompressibleAnimData()
-		: RequestedCompressionScheme(nullptr)
-		, CurveCompressionSettings(nullptr)
-		, Skeleton(nullptr)
-		, Interpolation((EAnimInterpolationType)0)
-		, SequenceLength(0.f)
-		, NumFrames(0)
-		, bIsValidAdditive(false)
-	{
-	}
+	FCompressibleAnimData();
 
-	FCompressibleAnimData(UAnimCompress* InRequestedCompressionScheme, UAnimCurveCompressionSettings* InCurveCompressionSettings, USkeleton* InSkeleton, EAnimInterpolationType InInterpolation, float InSequenceLength, int32 InNumFrames);
+	FCompressibleAnimData(UAnimCompress* InRequestedCompressionScheme, UAnimCurveCompressionSettings* InCurveCompressionSettings, USkeleton* InSkeleton, EAnimInterpolationType InInterpolation, float InSequenceLength, int32 InNumFrames, const float InAltCompressionErrorThreshold);
 
-	FCompressibleAnimData(class UAnimSequence* InSeq, const bool bPerformStripping);
+	FCompressibleAnimData(class UAnimSequence* InSeq, const bool bPerformStripping, const float InAltCompressionErrorThreshold);
 
 	UAnimCompress* RequestedCompressionScheme;
 
@@ -222,6 +213,8 @@ public:
 
 	bool bIsValidAdditive;
 
+	float AltCompressionErrorThreshold;
+
 	//For Logging
 	FString Name;
 	FString FullName;
diff --git a/Engine/Source/Runtime/Engine/Public/AnimationUtils.h b/Engine/Source/Runtime/Engine/Public/AnimationUtils.h
--- a/Engine/Source/Runtime/Engine/Public/AnimationUtils.h
+++ b/Engine/Source/Runtime/Engine/Public/AnimationUtils.h
@@ -158,7 +158,7 @@ public:
 	 *
 	 * @return				World-space error tolerance for considering an alternative compression method
 	 */
-	static float GetAlternativeCompressionThreshold();
+	ENGINE_API static float GetAlternativeCompressionThreshold();
 	
 	/**
 	 * Determines the current setting for recompressing all animations upon load. The default value 
