Fix MaxDuration calculation when cooking surround

diff --git a/Engine/Source/Runtime/Engine/Private/AudioDerivedData.cpp b/Engine/Source/Runtime/Engine/Private/AudioDerivedData.cpp
--- a/Engine/Source/Runtime/Engine/Private/AudioDerivedData.cpp
+++ b/Engine/Source/Runtime/Engine/Private/AudioDerivedData.cpp
@@ -1119,7 +1119,7 @@ static void CookSurroundWave( USoundWave* SoundWave, FName FormatName, const IAu
 				WaveSampleRate = SampleRateOverride;
 			}
 
-			UE_LOG(LogAudioDerivedData, Log,  TEXT( "Cooking %d channels for: %s" ), ChannelCount, *SoundWave->GetFullName() );
+			UE_LOG(LogAudioDerivedData, Display, TEXT("Cooking %d channels for: %s"), ChannelCount, *SoundWave->GetFullName());
 
 			FSoundQualityInfo QualityInfo = { 0 };
 
@@ -1153,16 +1153,20 @@ static void CookSurroundWave( USoundWave* SoundWave, FName FormatName, const IAu
 					UE_LOG(LogAudioDerivedData, Warning, TEXT( "Updated SoundWave->NumChannels during cooking %s." ), *SoundWave->GetFullName() );
 					SoundWave->NumChannels = ChannelCount;
 				}
-				if (SoundWave->RawPCMDataSize != SampleDataSize * ChannelCount)
+
+				const int32 SampleCount = SampleDataSize * ChannelCount;
+				if (SoundWave->RawPCMDataSize != SampleCount)
 				{
-					UE_LOG(LogAudioDerivedData, Log, TEXT( "Updated SoundWave->RawPCMDataSize during cooking %s." ), *SoundWave->GetFullName() );
-					SoundWave->RawPCMDataSize = SampleDataSize * ChannelCount;
+					UE_LOG(LogAudioDerivedData, Display, TEXT("Updated SoundWave->RawPCMDataSize during cooking %s."), *SoundWave->GetFullName() );
+					SoundWave->RawPCMDataSize = SampleCount;
 				}
-				if (SoundWave->Duration != ( float )SampleDataSize / (QualityInfo.SampleRate * sizeof( int16 )))
+
+				const float NewDuration = (float)SampleCount / (QualityInfo.SampleRate * sizeof(int16));
+				if (SoundWave->Duration != NewDuration)
 				{
-					UE_LOG(LogAudioDerivedData, Warning, TEXT( "Updated SoundWave->Duration during cooking %s." ), *SoundWave->GetFullName() );
-					SoundWave->Duration = ( float )SampleDataSize / (QualityInfo.SampleRate * sizeof( int16 ));
-				}			
+					UE_LOG(LogAudioDerivedData, Display, TEXT("Updated SoundWave->Duration during cooking %s."), *SoundWave->GetFullName());
+					SoundWave->Duration = NewDuration;
+				}
 			}
 			else
 			{
