np: Fix variable name

diff --git a/Engine/Source/Runtime/Engine/Private/AudioDerivedData.cpp b/Engine/Source/Runtime/Engine/Private/AudioDerivedData.cpp
--- a/Engine/Source/Runtime/Engine/Private/AudioDerivedData.cpp
+++ b/Engine/Source/Runtime/Engine/Private/AudioDerivedData.cpp
@@ -1154,14 +1154,14 @@ static void CookSurroundWave( USoundWave* SoundWave, FName FormatName, const IAu
 					SoundWave->NumChannels = ChannelCount;
 				}
 
-				const int32 SampleCount = SampleDataSize * ChannelCount;
-				if (SoundWave->RawPCMDataSize != SampleCount)
+				const int32 PCMDataSize = SampleDataSize * ChannelCount;
+				if (SoundWave->RawPCMDataSize != PCMDataSize)
 				{
 					UE_LOG(LogAudioDerivedData, Display, TEXT("Updated SoundWave->RawPCMDataSize during cooking %s."), *SoundWave->GetFullName() );
-					SoundWave->RawPCMDataSize = SampleCount;
+					SoundWave->RawPCMDataSize = PCMDataSize;
 				}
 
-				const float NewDuration = (float)SampleCount / (QualityInfo.SampleRate * sizeof(int16));
+				const float NewDuration = (float)PCMDataSize / (QualityInfo.SampleRate * sizeof(int16));
 				if (SoundWave->Duration != NewDuration)
 				{
 					UE_LOG(LogAudioDerivedData, Display, TEXT("Updated SoundWave->Duration during cooking %s."), *SoundWave->GetFullName());
