Fix GLEO error on saving streamable animations

diff --git a/Engine/Source/Runtime/Engine/Private/Animation/AnimStreamable.cpp b/Engine/Source/Runtime/Engine/Private/Animation/AnimStreamable.cpp
--- a/Engine/Source/Runtime/Engine/Private/Animation/AnimStreamable.cpp
+++ b/Engine/Source/Runtime/Engine/Private/Animation/AnimStreamable.cpp
@@ -330,7 +330,7 @@ void UAnimStreamable::PostLoad()
 
 	if (SourceSequence)
 	{
-		CompressionScheme = SourceSequence->CompressionScheme;
+		CompressionScheme = DuplicateObject<UAnimCompress>(SourceSequence->CompressionScheme, this);
 	}
 
 	if (SourceSequence && (GenerateGuidFromRawAnimData(SourceSequence->GetRawAnimationData(), SourceSequence->RawCurveData) != RawDataGuid))
@@ -387,7 +387,7 @@ void UAnimStreamable::InitFrom(const UAnimSequence* InSourceSequence)
 	Modify();
 	SetSkeleton(InSourceSequence->GetSkeleton());
 	SourceSequence = InSourceSequence;
-	CompressionScheme = InSourceSequence->CompressionScheme;
+	CompressionScheme = DuplicateObject<UAnimCompress>(SourceSequence->CompressionScheme, this);
 
 	RawAnimationData = InSourceSequence->GetRawAnimationData();
 	RawCurveData = InSourceSequence->RawCurveData;
