Disable morph buffers initialization for mobile, as its used only on SM5

diff --git a/Engine/Source/Runtime/Engine/Private/SkeletalMeshLODRenderData.cpp b/Engine/Source/Runtime/Engine/Private/SkeletalMeshLODRenderData.cpp
--- a/Engine/Source/Runtime/Engine/Private/SkeletalMeshLODRenderData.cpp
+++ b/Engine/Source/Runtime/Engine/Private/SkeletalMeshLODRenderData.cpp
@@ -223,7 +223,8 @@ void FSkeletalMeshLODRenderData::InitResources(bool bNeedsVertexColors, int32 LO
         }
     }
 
-	if (RHISupportsComputeShaders(GMaxRHIShaderPlatform) && InMorphTargets.Num() > 0)
+	// UseGPUMorphTargets() can be toggled only on SM5 atm
+	if (IsFeatureLevelSupported(GMaxRHIShaderPlatform, ERHIFeatureLevel::SM5) && InMorphTargets.Num() > 0)
 	{
 		MorphTargetVertexInfoBuffers.VertexIndices.Empty();
 		MorphTargetVertexInfoBuffers.MorphDeltas.Empty();
