Revert earlier change.  Bug was not due to import channel division issue but is due to re-import logic being incorrect and failing on surround.

diff --git a/Engine/Source/Runtime/Engine/Private/AudioDerivedData.cpp b/Engine/Source/Runtime/Engine/Private/AudioDerivedData.cpp
--- a/Engine/Source/Runtime/Engine/Private/AudioDerivedData.cpp
+++ b/Engine/Source/Runtime/Engine/Private/AudioDerivedData.cpp
@@ -70,7 +70,7 @@ static void GetStreamedAudioDerivedDataKeySuffix(
 			Version = AudioFormat->GetVersion(AudioFormatName);
 		}
 	}
-	
+
 	FString AudioFormatNameString = AudioFormatName.ToString();
 
 	// If we have compression overrides for this target platform, append them to this string.
@@ -281,7 +281,7 @@ class FStreamedAudioCacheDerivedDataWorker : public FNonAbandonableTask
 		GetStreamedAudioDerivedDataKeySuffix(SoundWave, AudioFormatName, CompressionOverrides, KeySuffix);
 
 		DerivedData->Chunks.Empty();
-		
+
 		ITargetPlatformManagerModule* TPM = GetTargetPlatformManager();
 		const IAudioFormat* AudioFormat = NULL;
 		if (TPM)
@@ -362,7 +362,7 @@ class FStreamedAudioCacheDerivedDataWorker : public FNonAbandonableTask
 
 				// Store it in the cache.
 				// @todo: This will remove the streaming bulk data, which we immediately reload below!
-				// Should ideally avoid this redundant work, but it only happens when we actually have 
+				// Should ideally avoid this redundant work, but it only happens when we actually have
 				// to build the texture, which should only ever be once.
 				this->BytesCached = PutDerivedDataInCache(DerivedData, KeySuffix);
 			}
@@ -374,7 +374,7 @@ class FStreamedAudioCacheDerivedDataWorker : public FNonAbandonableTask
 					);
 			}
 		}
-		
+
 		if (DerivedData->Chunks.Num())
 		{
 			bool bInlineChunks = (CacheFlags & EStreamedAudioCacheFlags::InlineChunks) != 0;
@@ -455,7 +455,7 @@ public:
 	bool Finalize()
 	{
 		check(IsInGameThread());
-		// if we couldn't get from the DDC or didn't build synchronously, then we have to build now. 
+		// if we couldn't get from the DDC or didn't build synchronously, then we have to build now.
 		// This is a super edge case that should rarely happen.
 		if (!bSucceeded)
 		{
@@ -802,7 +802,7 @@ static void ResampleWaveData(TArray<uint8>& WaveData, size_t& NumBytes, int32 Nu
 	// Convert wav data from int16 to float:
 	ResamplerInputData.AddUninitialized(NumSamples);
 	int16* InputData = (int16*) WaveData.GetData();
-	
+
 	for (int32 Index = 0; Index < NumSamples; Index++)
 	{
 		ResamplerInputData[Index] = ((float)InputData[Index]) / 32767.0f;
@@ -863,7 +863,7 @@ static void ResampleWaveData(TArray<uint8>& WaveData, size_t& NumBytes, int32 Nu
 	{
 		UE_LOG(LogAudioDerivedData, Error, TEXT("Resampling operation failed."));
 	}
-	
+
 	double TimeDelta = FPlatformTime::Seconds() - StartTime;
 	UE_LOG(LogAudioDerivedData, Display, TEXT("Resampling file from %f to %f took %f seconds."), SourceSampleRate, DestinationSampleRate, TimeDelta);
 }
@@ -929,7 +929,7 @@ static void CookSimpleWave(USoundWave* SoundWave, FName FormatName, const IAudio
 			if (SampleRateOverride > 0 && SampleRateOverride != (float)WaveSampleRate)
 			{
 				size_t TotalDataSize = WaveInfo.SampleDataSize;
-				
+
 				ResampleWaveData(Input, TotalDataSize, *WaveInfo.pChannels, WaveSampleRate, SampleRateOverride);
 				WaveSampleRate = SampleRateOverride;
 				WaveInfo.SampleDataSize = TotalDataSize;
@@ -954,7 +954,7 @@ static void CookSimpleWave(USoundWave* SoundWave, FName FormatName, const IAudio
 		{
 			QualityInfo.Quality = SoundWave->CompressionQuality;
 		}
-		
+
 		QualityInfo.NumChannels = *WaveInfo.pChannels;
 		QualityInfo.SampleRate = WaveSampleRate;
 		QualityInfo.SampleDataSize = Input.Num();
@@ -962,7 +962,7 @@ static void CookSimpleWave(USoundWave* SoundWave, FName FormatName, const IAudio
 		QualityInfo.DebugName = SoundWave->GetFullName();
 
 		// Cook the data.
-		if(Format.Cook(FormatName, Input, QualityInfo, Output)) 
+		if(Format.Cook(FormatName, Input, QualityInfo, Output))
 		{
 			SoundWave->SetSampleRate(QualityInfo.SampleRate);
 
@@ -1095,7 +1095,7 @@ static void CookSurroundWave( USoundWave* SoundWave, FName FormatName, const IAu
 				Input.AddZeroed(SampleDataSize);
 			}
 		}
-	
+
 		// Only allow the formats that can be played back through
 		if( ChannelCount == 4 || ChannelCount == 6 || ChannelCount == 7 || ChannelCount == 8 )
 		{
@@ -1114,7 +1114,7 @@ static void CookSurroundWave( USoundWave* SoundWave, FName FormatName, const IAu
 				{
 					size_t DataSize = SourceBuffers[ChannelIndex].Num();
 					ResampleWaveData(SourceBuffers[ChannelIndex], DataSize, 1, WaveSampleRate, SampleRateOverride);
-				}	
+				}
 
 				WaveSampleRate = SampleRateOverride;
 			}
@@ -1144,7 +1144,7 @@ static void CookSurroundWave( USoundWave* SoundWave, FName FormatName, const IAu
 			QualityInfo.bStreaming = SoundWave->IsStreaming();
 			QualityInfo.DebugName = SoundWave->GetFullName();
 			//@todo tighten up the checking for empty results here
-			if(Format.CookSurround(FormatName, SourceBuffers, QualityInfo, Output)) 
+			if(Format.CookSurround(FormatName, SourceBuffers, QualityInfo, Output))
 			{
 				SoundWave->SetSampleRate(QualityInfo.SampleRate);
 
@@ -1161,7 +1161,7 @@ static void CookSurroundWave( USoundWave* SoundWave, FName FormatName, const IAu
 					SoundWave->RawPCMDataSize = PCMDataSize;
 				}
 
-				const float NewDuration = (float)PCMDataSize / (QualityInfo.SampleRate * sizeof(int16));
+				const float NewDuration = (float)SampleDataSize / (QualityInfo.SampleRate * sizeof(int16));
 				if (SoundWave->Duration != NewDuration)
 				{
 					UE_LOG(LogAudioDerivedData, Display, TEXT("Updated SoundWave->Duration during cooking %s."), *SoundWave->GetFullName());
@@ -1393,7 +1393,7 @@ void USoundWave::BeginCacheForCookedPlatformData(const ITargetPlatform *TargetPl
 	Super::BeginCacheForCookedPlatformData(TargetPlatform);
 }
 
-bool USoundWave::IsCachedCookedPlatformDataLoaded( const ITargetPlatform* TargetPlatform ) 
+bool USoundWave::IsCachedCookedPlatformDataLoaded( const ITargetPlatform* TargetPlatform )
 {
 	if (TargetPlatform->SupportsFeature(ETargetPlatformFeatures::AudioStreaming) && IsStreaming())
 	{
@@ -1418,7 +1418,7 @@ bool USoundWave::IsCachedCookedPlatformDataLoaded( const ITargetPlatform* Target
 
 		return PlatformData->IsFinishedCache();
 	}
-	return true; 
+	return true;
 }
 
 
@@ -1452,22 +1452,22 @@ void USoundWave::ClearCachedCookedPlatformData( const ITargetPlatform* TargetPla
 		FString DerivedDataKey;
 		GetStreamedAudioDerivedDataKeySuffix(*this, PlatformFormat, CompressionOverrides, DerivedDataKey);
 
-		
+
 		if ( CookedPlatformData.Contains(DerivedDataKey) )
 		{
 			FStreamedAudioPlatformData *PlatformData = CookedPlatformData.FindAndRemoveChecked( DerivedDataKey );
 			delete PlatformData;
-		}	
+		}
 	}
 }
 
 void USoundWave::WillNeverCacheCookedPlatformDataAgain()
 {
 	// this is called after we have finished caching the platform data but before we have saved the data
-	// so need to keep the cached platform data around 
+	// so need to keep the cached platform data around
 	Super::WillNeverCacheCookedPlatformDataAgain();
 
-	// TODO: We can clear these arrays if we never need to cook again. 
+	// TODO: We can clear these arrays if we never need to cook again.
 	RawData.RemoveBulkData();
 	CompressedFormatData.FlushData();
 }
