Added a multicast delegate to be broadcasted as part of UWorld::BeginTearingDown

diff --git a/Engine/Source/Runtime/Engine/Classes/Engine/World.h b/Engine/Source/Runtime/Engine/Classes/Engine/World.h
--- a/Engine/Source/Runtime/Engine/Classes/Engine/World.h
+++ b/Engine/Source/Runtime/Engine/Classes/Engine/World.h
@@ -1391,6 +1391,11 @@ private:
 	/** Broadcasts whenever the number of levels changes */
 	FOnLevelsChangedEvent LevelsChangedEvent;
 
+	DECLARE_EVENT(UWorld, FOnBeginTearingDownEvent);
+
+	/** Broadcasted on UWorld::BeginTearingDown */
+	FOnBeginTearingDownEvent BeginTearingDownEvent;
+
 #if WITH_EDITOR
 
 	/** Broadcasts that selected levels have changed. */
@@ -3244,6 +3249,9 @@ public:
 	/** Returns the LevelsChangedEvent member. */
 	FOnLevelsChangedEvent& OnLevelsChanged() { return LevelsChangedEvent; }
 
+	/** Returns the BeginTearingDownEvent member. */
+	FOnBeginTearingDownEvent& OnBeginTearingDown() { return BeginTearingDownEvent; }
+
 	/** Returns the actor count. */
 	int32 GetProgressDenominator();
 	
diff --git a/Engine/Source/Runtime/Engine/Private/World.cpp b/Engine/Source/Runtime/Engine/Private/World.cpp
--- a/Engine/Source/Runtime/Engine/Private/World.cpp
+++ b/Engine/Source/Runtime/Engine/Private/World.cpp
@@ -2512,6 +2512,7 @@ void UWorld::BeginTearingDown()
 {
 	bIsTearingDown = true;
 	UE_LOG(LogWorld, Log, TEXT("BeginTearingDown for %s"), *GetOutermost()->GetName());
+	BeginTearingDownEvent.Broadcast();
 }
 
 void UWorld::RemoveFromWorld( ULevel* Level, bool bAllowIncrementalRemoval )
