Removed the remaining the traces of the original UNavigationSystem from the engine code

diff --git a/Engine/Source/Runtime/Engine/Classes/AI/Navigation/NavigationTypes.h b/Engine/Source/Runtime/Engine/Classes/AI/Navigation/NavigationTypes.h
--- a/Engine/Source/Runtime/Engine/Classes/AI/Navigation/NavigationTypes.h
+++ b/Engine/Source/Runtime/Engine/Classes/AI/Navigation/NavigationTypes.h
@@ -35,10 +35,10 @@ typedef uint64 NavNodeRef;
 
 namespace FNavigationSystem
 {
-	/** used as a fallback value for navigation agent radius, when none specified via UNavigationSystem::SupportedAgents */
+	/** used as a fallback value for navigation agent radius, when none specified via UNavigationSystemV1::SupportedAgents */
 	extern ENGINE_API const float FallbackAgentRadius;
 
-	/** used as a fallback value for navigation agent height, when none specified via UNavigationSystem::SupportedAgents */
+	/** used as a fallback value for navigation agent height, when none specified via UNavigationSystemV1::SupportedAgents */
 	extern ENGINE_API const float FallbackAgentHeight;
 
 	static const FBox InvalidBoundingBox(ForceInit);
diff --git a/Engine/Source/Runtime/Engine/Classes/AI/NavigationSystemBase.h b/Engine/Source/Runtime/Engine/Classes/AI/NavigationSystemBase.h
--- a/Engine/Source/Runtime/Engine/Classes/AI/NavigationSystemBase.h
+++ b/Engine/Source/Runtime/Engine/Classes/AI/NavigationSystemBase.h
@@ -285,84 +285,3 @@ protected:
 	static FNavigationSystem::FControllerBasedSignature& StopMovementDelegate();
 	static FNavigationSystem::FBoolControllerBasedSignature& IsFollowingAPathDelegate();
 };
-
-
-UCLASS(Abstract, config = Engine, defaultconfig)
-class ENGINE_API UNavigationSystem : public UObject
-{
-	GENERATED_BODY()
-public:
-	UNavigationSystem(const FObjectInitializer& ObjectInitializer = FObjectInitializer::Get());
-
-	UE_DEPRECATED(4.20, "UpdateActorInNavOctree is deprecated. Use FNavigationSystem::RemoveActorData instead")
-	static void ClearNavOctreeAll(AActor* Actor);
-	UE_DEPRECATED(4.20, "UpdateActorInNavOctree is deprecated. Use FNavigationSystem::UpdateActorData instead")
-	static void UpdateActorInNavOctree(AActor& Actor);
-	UE_DEPRECATED(4.20, "UpdateComponentInNavOctree is deprecated. Use FNavigationSystem::UpdateComponentData instead")
-	static void UpdateComponentInNavOctree(UActorComponent& Comp);
-	UE_DEPRECATED(4.20, "UpdateActorAndComponentsInNavOctree is deprecated. Use FNavigationSystem::UpdateActorAndComponentData instead")
-	static void UpdateActorAndComponentsInNavOctree(AActor& Actor, bool bUpdateAttachedActors = true);
-	UE_DEPRECATED(4.20, "UpdateNavOctreeAfterMove is deprecated. Use FNavigationSystem::UpdateComponentDataAfterMove instead")
-	static void UpdateNavOctreeAfterMove(USceneComponent* Comp);
-	UE_DEPRECATED(4.20, "UpdateNavOctreeBounds is deprecated. Use FNavigationSystem::OnActorBoundsChanged instead")
-	static void UpdateNavOctreeBounds(AActor* Actor);
-	UE_DEPRECATED(4.20, "InitializeForWorld is deprecated. Use FNavigationSystem::CreateNavigationSystem instead")
-	static void InitializeForWorld(UWorld* World, FNavigationSystemRunMode Mode);
-	UE_DEPRECATED(4.20, "CreateNavigationSystem is deprecated. Use FNavigationSystem::CreateNavigationSystem instead")
-	static UNavigationSystem* CreateNavigationSystem(UWorld* WorldOwner);
-	UE_DEPRECATED(4.20, "OnComponentRegistered is deprecated. Use FNavigationSystem::OnComponentRegistered instead")
-	static void OnComponentRegistered(UActorComponent* Comp);
-	UE_DEPRECATED(4.20, "OnComponentUnregistered is deprecated. Use FNavigationSystem::OnComponentUnregistered instead")
-	static void OnComponentUnregistered(UActorComponent* Comp);
-	UE_DEPRECATED(4.20, "OnActorRegistered is deprecated. Use FNavigationSystem::OnActorRegistered instead")
-	static void OnActorRegistered(AActor* Actor);
-	UE_DEPRECATED(4.20, "OnActorUnregistered is deprecated. Use FNavigationSystem::OnActorUnregistered instead")
-	static void OnActorUnregistered(AActor* Actor);
-	UE_DEPRECATED(4.20, "GetCurrent is deprecated. Use FNavigationSystem::GetCurrent instead")
-	static UNavigationSystem* GetCurrent(UWorld* World);
-	UE_DEPRECATED(4.20, "GetCurrent is deprecated. Use FNavigationSystem::GetCurrent instead")
-	static UNavigationSystem* GetCurrent(UObject* WorldContextObject);
-	UE_DEPRECATED(4.20, "ShouldUpdateNavOctreeOnComponentChange is deprecated. Use FNavigationSystem::WantsComponentChangeNotifies instead")
-	static bool ShouldUpdateNavOctreeOnComponentChange();
-	UE_DEPRECATED(4.20, "GetDefaultWalkableArea is deprecated. Use FNavigationSystem::GetDefaultWalkableArea instead")
-	static TSubclassOf<UNavAreaBase> GetDefaultWalkableArea();
-	UE_DEPRECATED(4.20, "UNavigationSystem::GetNavDataForProps is deprecated. Use FNavigationSystem::GetDefaultObstacleArea instead")
-	static TSubclassOf<UNavAreaBase> GetDefaultObstacleArea();
-	UE_DEPRECATED(4.20, "UNavigationSystem::K2_GetRandomReachablePointInRadius is deprecated. Use UNavigationSystemV1::K2_GetRandomReachablePointInRadius instead")
-	static bool K2_GetRandomReachablePointInRadius(UObject* WorldContextObject, const FVector& Origin, FVector& RandomLocation, float Radius, UObject* NavData = NULL, TSubclassOf<UObject> FilterClass = NULL) { return false; }
-	UE_DEPRECATED(4.20, "UNavigationSystem::SimpleMoveToActor is deprecated. Use UAIBlueprintHelperLibrary::SimpleMoveToActor instead")
-	UFUNCTION(BlueprintCallable, Category = "AI|Navigation", meta = (DisplayName = "SimpleMoveToActor_DEPRECATED", ScriptNoExport, DeprecatedFunction, DeprecationMessage = "SimpleMoveToActor is deprecated. Use AIBlueprintHelperLibrary::SimpleMoveToActor instead"))
-	static void SimpleMoveToActor(AController* Controller, const AActor* Goal) {}
-	UE_DEPRECATED(4.20, "UNavigationSystem::SimpleMoveToLocation is deprecated. Use UAIBlueprintHelperLibrary::SimpleMoveToLocation instead")
-	UFUNCTION(BlueprintCallable, Category = "AI|Navigation", meta = (DisplayName = "SimpleMoveToLocation_DEPRECATED", ScriptNoExport, DeprecatedFunction, DeprecationMessage = "SimpleMoveToLocation is deprecated. Use AIBlueprintHelperLibrary::SimpleMoveToLocation instead"))
-	static void SimpleMoveToLocation(AController* Controller, const FVector& Goal) {}
-
-	UE_DEPRECATED(4.20, "UpdateNavOctreeElementBounds is deprecated. Use FNavigationSystem::OnComponentBoundsChanged instead")
-	bool UpdateNavOctreeElementBounds(UActorComponent* Comp, const FBox& NewBounds, const FBox& DirtyArea);
-
-	UE_DEPRECATED(4.20, "GetObjectsNavOctreeId is deprecated. You need to access NavigationSystem's modul NavigationSystem to access this functionality")
-	const class FOctreeElementId* GetObjectsNavOctreeId(const UObject* Object) const { return nullptr; }
-
-	UE_DEPRECATED(4.20, "HasPendingObjectNavOctreeId is deprecated. You need to access NavigationSystem's modul NavigationSystem class to access this functionality")
-	bool HasPendingObjectNavOctreeId(UObject* Object) const { return false; }
-
-	UE_DEPRECATED(4.20, "GetDefaultSupportedAgentConfig is deprecated. Use FNavigationSystem::GetDefaultSupportedAgent instead")
-	const FNavDataConfig& GetDefaultSupportedAgentConfig() const;
-	
-	UE_DEPRECATED(4.20, "GetNavDataForProps is deprecated. Use FNavigationSystem::GetNavDataForProps instead")
-	INavigationDataInterface* GetNavDataForProps(const FNavAgentProperties& AgentProperties);
-
-	UE_DEPRECATED(4.20, "GetSupportedAgents has been moved. Use UNavigationSystemV1::GetSupportedAgents instead")
-	const TArray<FNavDataConfig>& GetSupportedAgents() const { return FakeSupportedAgents; }
-
-	UE_DEPRECATED(4.20, "GetMainNavData has been moved. Use FNavigationSystem::GetCurrent<UNavigationSystemV1>()->GetDefaultNavDataInstance instead")
-	UObject* GetMainNavData(int) { return nullptr; }
-	UE_DEPRECATED(4.20, "GetMainNavData has been moved. Use FNavigationSystem::GetCurrent<UNavigationSystemV1>()->GetDefaultNavDataInstance instead")
-	UObject* GetMainNavData() const { return nullptr; }
-
-	UE_DEPRECATED(4.20, "ProjectPointToNavigation has been deprecated along with the whole UNavigationSystem class. Use UNavigationSystemV1 instead")
-	bool ProjectPointToNavigation(const FVector& Point, FNavLocation& OutLocation, const FVector& Extent = INVALID_NAVEXTENT, const UObject* NavData = NULL, void* QueryFilter = NULL) const { return false; }
-
-private:
-	const TArray<FNavDataConfig> FakeSupportedAgents;
-};
\ No newline at end of file
diff --git a/Engine/Source/Runtime/Engine/Classes/Engine/World.h b/Engine/Source/Runtime/Engine/Classes/Engine/World.h
--- a/Engine/Source/Runtime/Engine/Classes/Engine/World.h
+++ b/Engine/Source/Runtime/Engine/Classes/Engine/World.h
@@ -49,7 +49,7 @@ class ULocalPlayer;
 class UMaterialParameterCollection;
 class UMaterialParameterCollectionInstance;
 class UModel;
-class UNavigationSystem;
+class UNavigationSystemBase;
 class UNetConnection;
 class UNetDriver;
 class UPrimitiveComponent;
diff --git a/Engine/Source/Runtime/Engine/Private/AI/NavigationSystemBase.cpp b/Engine/Source/Runtime/Engine/Private/AI/NavigationSystemBase.cpp
--- a/Engine/Source/Runtime/Engine/Private/AI/NavigationSystemBase.cpp
+++ b/Engine/Source/Runtime/Engine/Private/AI/NavigationSystemBase.cpp
@@ -306,52 +306,8 @@ void FNavigationLockContext::UnlockUpdates()
 }
 
 //----------------------------------------------------------------------//
-// 
+// UNavigationSystemBase
 //----------------------------------------------------------------------//
-UNavigationSystem::UNavigationSystem(const FObjectInitializer& ObjectInitializer)
-{
-#if !UE_BUILD_SHIPPING
-	if (HasAnyFlags(RF_ClassDefaultObject) && GetClass() == UNavigationSystem::StaticClass())
-	{
-		struct FIniChecker
-		{
-			FIniChecker()
-			{
-				const TCHAR EngineTemplage[] = TEXT("/Script/Engine.%s");
-				const TCHAR MessageTemplate[] = TEXT("[/Script/Engine.%s] found in the DefaultEngine.ini file. This class has been moved. Please rename that section to [/Script/NavigationSystem.%s]");
-				const TArray<FString> MovedIniClasses = {
-					TEXT("RecastNavMesh")
-					, TEXT("NavArea")
-					, TEXT("NavAreaMeta")
-					, TEXT("NavArea_Default")
-					, TEXT("NavArea_LowHeight")
-					, TEXT("NavArea_Null")
-					, TEXT("NavArea_Obstacle")
-					, TEXT("NavAreaMeta_SwitchByAgent")
-					, TEXT("AbstractNavData")
-					, TEXT("NavCollision")
-					, TEXT("NavigationData")
-					, TEXT("NavigationGraph")
-					, TEXT("NavigationGraphNode")
-					, TEXT("NavigationGraphNodeComponent")
-				};
-
-				// NavigationSystem changed name, treat tit separately
-				UE_CLOG(GConfig->DoesSectionExist(*FString::Printf(EngineTemplage, TEXT("NavigationSystem")), GEngineIni)
-					, LogNavigation, Error, MessageTemplate, TEXT("NavigationSystem"), TEXT("NavigationSystemV1"));
-
-				for (auto ClassName : MovedIniClasses)
-				{
-					UE_CLOG(GConfig->DoesSectionExist(*FString::Printf(EngineTemplage, *ClassName), GEngineIni)
-						, LogNavigation, Error, MessageTemplate, *ClassName, *ClassName);
-				}
-			}
-		};
-		static FIniChecker IniChecker;
-	}
-#endif // !UE_BUILD_SHIPPING
-}
-
 void UNavigationSystemBase::SetCoordTransformTo(const ENavigationCoordSystem::Type CoordType, const FTransform& Transform)
 {
 	SetCoordTransform(ENavigationCoordSystem::Unreal, CoordType, Transform);
diff --git a/Engine/Source/Runtime/Engine/Public/EngineFwd.h b/Engine/Source/Runtime/Engine/Public/EngineFwd.h
--- a/Engine/Source/Runtime/Engine/Public/EngineFwd.h
+++ b/Engine/Source/Runtime/Engine/Public/EngineFwd.h
@@ -79,7 +79,6 @@ class UMaterialInterface;
 class UMeshComponent;
 class UModel;
 class UMorphTarget;
-class UNavigationSystem;
 class UNetConnection;
 class UNetDriver;
 class UParticleEmitter;
diff --git a/Engine/Source/Runtime/Landscape/Private/LandscapeEdit.cpp b/Engine/Source/Runtime/Landscape/Private/LandscapeEdit.cpp
--- a/Engine/Source/Runtime/Landscape/Private/LandscapeEdit.cpp
+++ b/Engine/Source/Runtime/Landscape/Private/LandscapeEdit.cpp
@@ -137,7 +137,6 @@ void ULandscapeComponent::UpdateNavigationRelevance()
 	if (CollisionComponent && Proxy)
 	{
 		CollisionComponent->SetCanEverAffectNavigation(Proxy->bUsedForNavigation);
-		// UNavigationSystem::UpdateNavOctree(CollisionComponent.Get());
 	}
 }
 
diff --git a/Engine/Source/Runtime/NavigationSystem/Private/NavMesh/RecastNavMesh.cpp b/Engine/Source/Runtime/NavigationSystem/Private/NavMesh/RecastNavMesh.cpp
--- a/Engine/Source/Runtime/NavigationSystem/Private/NavMesh/RecastNavMesh.cpp
+++ b/Engine/Source/Runtime/NavigationSystem/Private/NavMesh/RecastNavMesh.cpp
@@ -293,37 +293,6 @@ void ARecastNavMesh::DestroyRecastPImpl()
 	}
 }
 
-ARecastNavMesh* ARecastNavMesh::SpawnInstance(UNavigationSystem* NavSys, const FNavDataConfig* AgentProps)
-{
-	FActorSpawnParameters SpawnInfo;
-	SpawnInfo.OverrideLevel = NavSys->GetWorld()->PersistentLevel;
-	ARecastNavMesh* Instance = NavSys->GetWorld()->SpawnActor<ARecastNavMesh>( SpawnInfo );
-
-	if (Instance != NULL && AgentProps != NULL)
-	{
-		Instance->SetConfig(*AgentProps);
-		if (AgentProps->Name != NAME_None)
-		{
-			FString StrName = FString::Printf(TEXT("%s-%s"), *(Instance->GetFName().GetPlainNameString()), *(AgentProps->Name.ToString()));
-			// temporary solution to make sure we don't try to change name while there's already
-			// an object with this name
-			UObject* ExistingObject = StaticFindObject(/*Class=*/ NULL, Instance->GetOuter(), *StrName, true);
-			if (ExistingObject != NULL)
-			{
-				ExistingObject->Rename(NULL, NULL, REN_DontCreateRedirectors | REN_ForceGlobalUnique | REN_DoNotDirty | REN_NonTransactional);
-			}
-
-			// Set descriptive name
-			Instance->Rename(*StrName);
-#if WITH_EDITOR
-			Instance->SetActorLabel(StrName);
-#endif // WITH_EDITOR
-		}
-	}
-
-	return Instance;
-}
-
 UPrimitiveComponent* ARecastNavMesh::ConstructRenderingComponent() 
 {
 	return NewObject<UNavMeshRenderingComponent>(this, TEXT("NavRenderingComp"), RF_Transient);
diff --git a/Engine/Source/Runtime/NavigationSystem/Public/CrowdManagerBase.h b/Engine/Source/Runtime/NavigationSystem/Public/CrowdManagerBase.h
--- a/Engine/Source/Runtime/NavigationSystem/Public/CrowdManagerBase.h
+++ b/Engine/Source/Runtime/NavigationSystem/Public/CrowdManagerBase.h
@@ -5,7 +5,7 @@
 #include "CrowdManagerBase.generated.h"
 
 /** Base class for Crowd Managers. If you want to create a custom crowd manager
- *	implement a class extending this one and set UNavigationSystem::CrowdManagerClass
+ *	implement a class extending this one and set UNavigationSystemV1::CrowdManagerClass
  *	to point at your class */
 UCLASS(Abstract, Transient)
 class NAVIGATIONSYSTEM_API UCrowdManagerBase : public UObject
diff --git a/Engine/Source/Runtime/NavigationSystem/Public/NavFilters/NavigationQueryFilter.h b/Engine/Source/Runtime/NavigationSystem/Public/NavFilters/NavigationQueryFilter.h
--- a/Engine/Source/Runtime/NavigationSystem/Public/NavFilters/NavigationQueryFilter.h
+++ b/Engine/Source/Runtime/NavigationSystem/Public/NavFilters/NavigationQueryFilter.h
@@ -47,7 +47,7 @@ struct NAVIGATIONSYSTEM_API FNavigationFilterArea
 };
 
 // 
-// Use UNavigationSystem.DescribeFilterFlags() to setup user friendly names of flags
+// Use UNavigationSystemV1.DescribeFilterFlags() to setup user friendly names of flags
 // 
 USTRUCT()
 struct NAVIGATIONSYSTEM_API FNavigationFilterFlags
diff --git a/Engine/Source/Runtime/NavigationSystem/Public/NavMesh/PImplRecastNavMesh.h b/Engine/Source/Runtime/NavigationSystem/Public/NavMesh/PImplRecastNavMesh.h
--- a/Engine/Source/Runtime/NavigationSystem/Public/NavMesh/PImplRecastNavMesh.h
+++ b/Engine/Source/Runtime/NavigationSystem/Public/NavMesh/PImplRecastNavMesh.h
@@ -21,7 +21,6 @@
 #endif
 
 class FRecastNavMeshGenerator;
-class UNavigationSystem;
 
 #if WITH_RECAST
 
diff --git a/Engine/Source/Runtime/NavigationSystem/Public/NavMesh/RecastNavMesh.h b/Engine/Source/Runtime/NavigationSystem/Public/NavMesh/RecastNavMesh.h
--- a/Engine/Source/Runtime/NavigationSystem/Public/NavMesh/RecastNavMesh.h
+++ b/Engine/Source/Runtime/NavigationSystem/Public/NavMesh/RecastNavMesh.h
@@ -41,7 +41,6 @@ class FRecastQueryFilter;
 class INavLinkCustomInterface;
 class UCanvas;
 class UNavArea;
-class UNavigationSystem;
 class UPrimitiveComponent;
 class URecastNavMeshDataChunk;
 struct FRecastAreaNavModifierElement;
@@ -916,9 +915,6 @@ protected:
 	/** Invalidates active paths that go through changed tiles  */
 	void InvalidateAffectedPaths(const TArray<uint32>& ChangedTiles);
 
-	/** Spawns an ARecastNavMesh instance, and configures it if AgentProps != NULL */
-	static ARecastNavMesh* SpawnInstance(UNavigationSystem* NavSys, const FNavDataConfig* AgentProps = NULL);
-
 	/** created a new FRecastNavMeshGenerator instance. Overrider to supply your
 	 *	own extentions. Note: needs to derive from FRecastNavMeshGenerator */
 	virtual FRecastNavMeshGenerator* CreateGeneratorInstance();
diff --git a/Engine/Source/Runtime/NavigationSystem/Public/NavigationSystem.h b/Engine/Source/Runtime/NavigationSystem/Public/NavigationSystem.h
--- a/Engine/Source/Runtime/NavigationSystem/Public/NavigationSystem.h
+++ b/Engine/Source/Runtime/NavigationSystem/Public/NavigationSystem.h
@@ -1003,9 +1003,9 @@ public:
 	UE_DEPRECATED(4.20, "SimpleMoveToLocation is deprecated. Use UAIBlueprintHelperLibrary::SimpleMoveToLocation instead")
 	UFUNCTION(BlueprintCallable, Category = "AI|Navigation", meta = (DisplayName = "SimpleMoveToLocation_DEPRECATED", ScriptNoExport, DeprecatedFunction, DeprecationMessage = "SimpleMoveToLocation is deprecated. Use AIBlueprintHelperLibrary::SimpleMoveToLocation instead"))
 	static void SimpleMoveToLocation(AController* Controller, const FVector& Goal);
-	UE_DEPRECATED(4.20, "UNavigationSystem::GetDefaultWalkableArea is deprecated. Use FNavigationSystem::GetDefaultWalkableArea instead")
+	UE_DEPRECATED(4.20, "UNavigationSystemV1::GetDefaultWalkableArea is deprecated. Use FNavigationSystem::GetDefaultWalkableArea instead")
 	static TSubclassOf<UNavAreaBase> GetDefaultWalkableArea() { return FNavigationSystem::GetDefaultWalkableArea(); }
-	UE_DEPRECATED(4.20, "UNavigationSystem::GetDefaultObstacleArea is deprecated. Use FNavigationSystem::GetDefaultObstacleArea instead")
+	UE_DEPRECATED(4.20, "UNavigationSystemV1::GetDefaultObstacleArea is deprecated. Use FNavigationSystem::GetDefaultObstacleArea instead")
 	static TSubclassOf<UNavAreaBase> GetDefaultObstacleArea() { return FNavigationSystem::GetDefaultObstacleArea(); }
 	UE_DEPRECATED(4.22, "This version is deprecated.  Please use GetRandomLocationInNavigableRadius instead")
 	UFUNCTION(BlueprintPure, Category = "AI|Navigation", meta = (WorldContext = "WorldContextObject", DisplayName = "GetRandomPointInNavigableRadius", ScriptName = "GetRandomPointInNavigableRadius"))
