Fix duplication in streaming level name table when loading a delta checkpoint

diff --git a/Engine/Source/Runtime/Engine/Classes/Engine/DemoNetDriver.h b/Engine/Source/Runtime/Engine/Classes/Engine/DemoNetDriver.h
--- a/Engine/Source/Runtime/Engine/Classes/Engine/DemoNetDriver.h
+++ b/Engine/Source/Runtime/Engine/Classes/Engine/DemoNetDriver.h
@@ -1026,7 +1026,7 @@ private:
 	// Maintain a quick lookup for loaded levels directly to LevelStatus
 	TMap<const ULevel*, int32> LevelStatusIndexByLevel;
 
-	// List of seen level statuses indices (in ALlLevelStatuses).
+	// List of seen level statuses indices (in AllLevelStatuses).
 	TArray<int32> SeenLevelStatuses;
 
 	// Time of the last packet we've processed (in seconds).
diff --git a/Engine/Source/Runtime/Engine/Private/DemoNetDriver.cpp b/Engine/Source/Runtime/Engine/Private/DemoNetDriver.cpp
--- a/Engine/Source/Runtime/Engine/Private/DemoNetDriver.cpp
+++ b/Engine/Source/Runtime/Engine/Private/DemoNetDriver.cpp
@@ -1224,7 +1224,7 @@ void UDemoNetDriver::ResetLevelStatuses()
 	check(World);
 
 	// ResetLevelStatuses should only ever be called before receiving *any* data from the Replay stream,
-	// immediately before processing checkpoint data, or after a level transistion (in which case no data
+	// immediately before processing checkpoint data, or after a level transition (in which case no data
 	// will be relevant to the new sublevels).
 	// In any case, we can just flag these sublevels as ready immediately.
 	FindOrAddLevelStatus(*(World->PersistentLevel)).bIsReady = true;
@@ -5353,6 +5353,7 @@ bool UDemoNetDriver::LoadCheckpoint(const FGotoResult& GotoResult)
 
 			if (GotoCheckpointArchive->IsError())
 			{
+				UE_LOG(LogDemo, Error, TEXT("Guid cache serialization error while loading checkpoint."));
 				break;
 			}
 		}
@@ -5371,6 +5372,12 @@ bool UDemoNetDriver::LoadCheckpoint(const FGotoResult& GotoResult)
 			CastChecked<UPackageMapClient>(ServerConnection->PackageMap)->SerializeNetFieldExportGroupMap(*GotoCheckpointArchive);
 		}
 
+		if (bDeltaCheckpoint)
+		{
+			// each set of checkpoint packets we read will have a full name table, so only keep the last version
+			SeenLevelStatuses.Reset();
+		}
+
 		ReadDemoFrameIntoPlaybackPackets(*GotoCheckpointArchive);
 
 		if (bDeltaCheckpoint)
