[NO NON-PII COMMIT MESSAGE]

diff --git a/Engine/Source/Developer/Android/AndroidTargetPlatform/Private/AndroidTargetPlatform.cpp b/Engine/Source/Developer/Android/AndroidTargetPlatform/Private/AndroidTargetPlatform.cpp
--- a/Engine/Source/Developer/Android/AndroidTargetPlatform/Private/AndroidTargetPlatform.cpp
+++ b/Engine/Source/Developer/Android/AndroidTargetPlatform/Private/AndroidTargetPlatform.cpp
@@ -656,12 +656,23 @@ void FAndroidTargetPlatform::GetAllWaveFormats(TArray<FName>& OutFormats) const
 	OutFormats.Add(NAME_ADPCM);
 }
 
-namespace
+namespace Android
 {
 	void CachePlatformAudioCookOverrides(FPlatformAudioCookOverrides& OutOverrides)
 	{
 		const TCHAR* CategoryName = TEXT("/Script/AndroidRuntimeSettings.AndroidRuntimeSettings");
 
+		GConfig->GetBool(CategoryName, TEXT("bUseAudioStreamCaching"), OutOverrides.bUseStreamCaching, GEngineIni);
+
+		/** Memory Load On Demand Settings */
+		if (OutOverrides.bUseStreamCaching)
+		{
+			// Cache size:
+			int32 RetrievedCacheSize = 32 * 1024;
+			GConfig->GetInt(CategoryName, TEXT("CacheSizeKB"), RetrievedCacheSize, GEngineIni);
+			OutOverrides.StreamCachingSettings.CacheSizeKB = RetrievedCacheSize;
+		}
+
 		GConfig->GetBool(CategoryName, TEXT("bResampleForDevice"), OutOverrides.bResampleForDevice, GEngineIni);
 
 		GConfig->GetFloat(CategoryName, TEXT("CompressionQualityModifier"), OutOverrides.CompressionQualityModifier, GEngineIni);
@@ -768,11 +779,11 @@ FPlatformAudioCookOverrides* FAndroidTargetPlatform::GetAudioCompressionSettings
 
 	if (!bCachedPlatformSettings)
 	{
-		CachePlatformAudioCookOverrides(Settings);
+		Android::CachePlatformAudioCookOverrides(Settings);
 		bCachedPlatformSettings = true;
 	}
 #else
-	CachePlatformAudioCookOverrides(Settings);
+	Android::CachePlatformAudioCookOverrides(Settings);
 #endif
 
 	return &Settings;
diff --git a/Engine/Source/Developer/AudioFormatADPCM/Private/AudioFormatADPCM.cpp b/Engine/Source/Developer/AudioFormatADPCM/Private/AudioFormatADPCM.cpp
--- a/Engine/Source/Developer/AudioFormatADPCM/Private/AudioFormatADPCM.cpp
+++ b/Engine/Source/Developer/AudioFormatADPCM/Private/AudioFormatADPCM.cpp
@@ -509,7 +509,19 @@ public:
 		return 0;
 	}
 
-	virtual bool SplitDataForStreaming(const TArray<uint8>& SrcBuffer, TArray<TArray<uint8>>& OutBuffers, const int32 MaxChunkSize) const override
+	virtual int32 GetMinimumSizeForInitialChunk(FName Format, const TArray<uint8>& SrcBuffer) const override
+	{
+		uint8 const*	SrcData = SrcBuffer.GetData();
+		int32			SrcSize = SrcBuffer.Num();
+		uint32			BytesProcessed = 0;
+
+		FWaveModInfo	WaveInfo;
+		WaveInfo.ReadWaveInfo((uint8*)SrcData, SrcSize);
+
+		return WaveInfo.SampleDataStart - SrcData;
+	}
+
+	virtual bool SplitDataForStreaming(const TArray<uint8>& SrcBuffer, TArray<TArray<uint8>>& OutBuffers, const int32 InitialMaxChunkSize, const int32 MaxChunkSize) const override
 	{
 		uint8 const*	SrcData = SrcBuffer.GetData();
 		int32			SrcSize = SrcBuffer.Num();
@@ -537,7 +549,9 @@ public:
 
 			// Add a new chunk for the header
 			int32 HeaderSize = WaveInfo.SampleDataStart - SrcData;
-			AddNewChunk(OutBuffers, MaxChunkSize);
+
+			check(InitialMaxChunkSize >= HeaderSize);
+			AddNewChunk(OutBuffers, InitialMaxChunkSize);
 			AddChunkData(OutBuffers, SrcData, HeaderSize);
 
 			int32 CurChunkDataSize = HeaderSize;
diff --git a/Engine/Source/Developer/AudioFormatOgg/Private/AudioFormatOgg.cpp b/Engine/Source/Developer/AudioFormatOgg/Private/AudioFormatOgg.cpp
--- a/Engine/Source/Developer/AudioFormatOgg/Private/AudioFormatOgg.cpp
+++ b/Engine/Source/Developer/AudioFormatOgg/Private/AudioFormatOgg.cpp
@@ -32,7 +32,7 @@ class FAudioFormatOgg : public IAudioFormat
 	enum
 	{
 		/** Version for OGG format, this becomes part of the DDC key. */
-		UE_AUDIO_OGG_VER = 2,
+		UE_AUDIO_OGG_VER = 3,
 	};
 
 public:
@@ -435,7 +435,14 @@ public:
 		return CompressedDataStore.Num();
 	}
 
-	virtual bool SplitDataForStreaming(const TArray<uint8>& SrcBuffer, TArray<TArray<uint8>>& OutBuffers, const int32 MaxChunkSize) const override
+	virtual int32 GetMinimumSizeForInitialChunk(FName Format, const TArray<uint8>& SrcBuffer) const override
+	{
+		// The header packet is typically 753 bytes, but depending on various metadata fields this can vary.
+		// To be safe this is set to 6 kB.
+		return 6 * 1024;
+	}
+
+	virtual bool SplitDataForStreaming(const TArray<uint8>& SrcBuffer, TArray<TArray<uint8>>& OutBuffers, const int32 InitialChunkMaxSize, const int32 MaxChunkSize) const override
 	{
 		// Just chunk purely on MONO_PCM_BUFFER_SIZE
 		uint8 const*	SrcData = SrcBuffer.GetData();
@@ -449,20 +456,29 @@ public:
 			return false;
 		}
 
-		// Use a base chunk size of MONO_PCM_BUFFER_SIZE * 2 because Android uses MONO_PCM_BUFFER_SIZE to submit buffers to the os, the streaming system has more scheduling flexability when the chunk size is
-		//	larger the the buffer size submitted to the OS
-#if PLATFORM_ANDROID
-		uint32	ChunkSize = MONO_PCM_BUFFER_SIZE * 2 * QualityInfo.NumChannels;
-#else
-		uint32 ChunkSize = MaxChunkSize;
-#endif
+
+		// Set up initial chunk:
+		uint32 ChunkSize = InitialChunkMaxSize;
+		uint32	CurChunkDataSize = FMath::Min<uint32>(SrcSize, ChunkSize);
+
+		AddNewChunk(OutBuffers, CurChunkDataSize);
+		AddChunkData(OutBuffers, SrcData, CurChunkDataSize);
+
+		check(SrcSize >= CurChunkDataSize);
+		SrcSize -= CurChunkDataSize;
+		SrcData += CurChunkDataSize;
+
+		ChunkSize = MaxChunkSize;
+
+		// Set up the rest of the chunks:
 		while(SrcSize > 0)
 		{
-			int32	CurChunkDataSize = FMath::Min<uint32>(SrcSize, ChunkSize);
-				
+			CurChunkDataSize = FMath::Min<uint32>(SrcSize, ChunkSize);
+
 			AddNewChunk(OutBuffers, CurChunkDataSize);
 			AddChunkData(OutBuffers, SrcData, CurChunkDataSize);
-				
+
+			check(SrcSize >= CurChunkDataSize);
 			SrcSize -= CurChunkDataSize;
 			SrcData += CurChunkDataSize;
 		}
diff --git a/Engine/Source/Developer/AudioFormatOpus/Private/AudioFormatOpus.cpp b/Engine/Source/Developer/AudioFormatOpus/Private/AudioFormatOpus.cpp
--- a/Engine/Source/Developer/AudioFormatOpus/Private/AudioFormatOpus.cpp
+++ b/Engine/Source/Developer/AudioFormatOpus/Private/AudioFormatOpus.cpp
@@ -30,7 +30,7 @@ class FAudioFormatOpus : public IAudioFormat
 	enum
 	{
 		/** Version for OPUS format, this becomes part of the DDC key. */
-		UE_AUDIO_OPUS_VER = 5,
+		UE_AUDIO_OPUS_VER = 6,
 	};
 
 public:
@@ -179,6 +179,9 @@ public:
 		}
 		else
 		{
+			SrcBufferCopies.Reset();
+			SrcBufferCopies.AddDefaulted(SrcBuffers.Num());
+
 			// Take a copy of the source regardless
 			for (int32 Index = 0; Index < SrcBuffers.Num(); Index++)
 			{
@@ -345,7 +348,17 @@ public:
 		return CompressedDataStore.Num();
 	}
 
-	virtual bool SplitDataForStreaming(const TArray<uint8>& SrcBuffer, TArray<TArray<uint8>>& OutBuffers, const int32 MaxChunkSize) const override
+	virtual int32 GetMinimumSizeForInitialChunk(FName Format, const TArray<uint8>& SrcBuffer) const override
+	{
+		// Since UE4 uses it's own version of the header, we hardcode the size of that here. See SplitDataForStreaming below to see our initial info.
+		return FCStringAnsi::Strlen(OPUS_ID_STRING) + 1 // Format identifier
+			+ sizeof(uint16) // Sample Rate
+			+ sizeof(uint32) // True Sample Count
+			+ sizeof(uint8)  // Number of Channels
+			+ sizeof(uint16); // Serialized Frames
+	}
+
+	virtual bool SplitDataForStreaming(const TArray<uint8>& SrcBuffer, TArray<TArray<uint8>>& OutBuffers, const int32 MaxInitialChunkSize, const int32 MaxChunkSize) const override
 	{
 		if (SrcBuffer.Num() == 0)
 		{
@@ -373,19 +386,23 @@ public:
 		ReadOffset += sizeof(uint16);
 
 		// Should always be able to store basic info in a single chunk
-		check(ReadOffset - WriteOffset < (uint32)MaxChunkSize)
+		check(ReadOffset - WriteOffset <= (uint32)MaxInitialChunkSize);
+
+		int32 ChunkSize = MaxInitialChunkSize;
 
 		while (ProcessedFrames < SerializedFrames)
 		{
 			uint16 FrameSize = *((uint16*)(LockedSrc + ReadOffset));
 
-			if ( (ReadOffset + sizeof(uint16) + FrameSize) - WriteOffset >= MaxChunkSize)
+			if ( (ReadOffset + sizeof(uint16) + FrameSize) - WriteOffset >= ChunkSize)
 			{
 				WriteOffset += AddDataChunk(OutBuffers, LockedSrc + WriteOffset, ReadOffset - WriteOffset);
 			}
 
 			ReadOffset += sizeof(uint16) + FrameSize;
 			ProcessedFrames++;
+
+			ChunkSize = MaxChunkSize;
 		}
 		if (WriteOffset < ReadOffset)
 		{
diff --git a/Engine/Source/Developer/IOS/IOSTargetPlatform/Private/IOSTargetPlatform.cpp b/Engine/Source/Developer/IOS/IOSTargetPlatform/Private/IOSTargetPlatform.cpp
--- a/Engine/Source/Developer/IOS/IOSTargetPlatform/Private/IOSTargetPlatform.cpp
+++ b/Engine/Source/Developer/IOS/IOSTargetPlatform/Private/IOSTargetPlatform.cpp
@@ -727,12 +727,23 @@ void FIOSTargetPlatform::GetAllWaveFormats(TArray<FName>& OutFormat) const
 	OutFormat.Add(NAME_ADPCM);
 }
 
-namespace
+namespace IOS
 {
 	void CachePlatformAudioCookOverrides(FPlatformAudioCookOverrides& OutOverrides)
 	{
 		const TCHAR* CategoryName = TEXT("/Script/IOSRuntimeSettings.IOSRuntimeSettings");
 
+		GConfig->GetBool(CategoryName, TEXT("bUseAudioStreamCaching"), OutOverrides.bUseStreamCaching, GEngineIni);
+
+		/** Memory Load On Demand Settings */
+		if (OutOverrides.bUseStreamCaching)
+		{
+			// Cache size:
+			int32 RetrievedCacheSize = 32 * 1024;
+			GConfig->GetInt(CategoryName, TEXT("CacheSizeKB"), RetrievedCacheSize, GEngineIni);
+			OutOverrides.StreamCachingSettings.CacheSizeKB = RetrievedCacheSize;
+		}
+
 		GConfig->GetBool(CategoryName, TEXT("bResampleForDevice"), OutOverrides.bResampleForDevice, GEngineIni);
 
 		GConfig->GetFloat(CategoryName, TEXT("CompressionQualityModifier"), OutOverrides.CompressionQualityModifier, GEngineIni);
@@ -839,11 +850,11 @@ FPlatformAudioCookOverrides* FIOSTargetPlatform::GetAudioCompressionSettings() c
 
 	if (!bCachedPlatformSettings)
 	{
-		CachePlatformAudioCookOverrides(Settings);
+		IOS::CachePlatformAudioCookOverrides(Settings);
 		bCachedPlatformSettings = true;
 	}
 #else
-	CachePlatformAudioCookOverrides(Settings);
+	IOS::CachePlatformAudioCookOverrides(Settings);
 #endif
 
 	return &Settings;
diff --git a/Engine/Source/Developer/TargetPlatform/Public/Interfaces/IAudioFormat.h b/Engine/Source/Developer/TargetPlatform/Public/Interfaces/IAudioFormat.h
--- a/Engine/Source/Developer/TargetPlatform/Public/Interfaces/IAudioFormat.h
+++ b/Engine/Source/Developer/TargetPlatform/Public/Interfaces/IAudioFormat.h
@@ -100,15 +100,24 @@ public:
 	 */
 	virtual int32 Recompress( FName Format, const TArray<uint8>& SrcBuffer, FSoundQualityInfo& QualityInfo, TArray<uint8>& OutBuffer ) const = 0;
 
+	/**
+	 * Given the encoded buffer, returns the minimum number of bytes required to perform ICompressedAudioInfo::ReadCompressedInfo() or IStreamedCompressedInfo::ParseHeader() for this file.
+	 * 
+	 * @param Format the codec that SrcBuffer was compressed as.
+	 * @param SrcBuffer the compressed data that will later be split into individual chunks.
+	 */
+	virtual int32 GetMinimumSizeForInitialChunk(FName Format, const TArray<uint8>& SrcBuffer) const { return 0; }
+
 	/**
 	 * Splits compressed data into chunks suitable for streaming audio.
 	 *
 	 * @param SrcBuffer Pre-compressed data as an array of bytes.
 	 * @param OutBuffers Array of buffers that contain the chunks the original data was split into.
+	 * @param FirstChunkMaxSize The maximum size for the chunk that will be loaded inline with it's owning USoundWave asset.
 	 * @param MaxChunkSize The maximum chunk size for each chunk. The chunks will be zero-padded to match this chunk size in the bulk data serialization.
 	 * @return Whether bulk data could be split for streaming.
 	 */
-	virtual bool SplitDataForStreaming(const TArray<uint8>& SrcBuffer, TArray<TArray<uint8>>& OutBuffers, const int32 MaxChunkSize) const {return false;}
+	virtual bool SplitDataForStreaming(const TArray<uint8>& SrcBuffer, TArray<TArray<uint8>>& OutBuffers, const int32 FirstChunkMaxSize, const int32 MaxChunkSize) const {return false;}
 
 public:
 
diff --git a/Engine/Source/Developer/Windows/WindowsTargetPlatform/Classes/WindowsTargetSettings.h b/Engine/Source/Developer/Windows/WindowsTargetPlatform/Classes/WindowsTargetSettings.h
--- a/Engine/Source/Developer/Windows/WindowsTargetPlatform/Classes/WindowsTargetSettings.h
+++ b/Engine/Source/Developer/Windows/WindowsTargetPlatform/Classes/WindowsTargetSettings.h
@@ -8,6 +8,7 @@
 
 #include "CoreMinimal.h"
 #include "UObject/Object.h"
+#include "AudioCompressionSettings.h"
 #include "WindowsTargetSettings.generated.h"
 
 UENUM()
@@ -108,4 +109,46 @@ public:
 	/** Which of the currently enabled occlusion plugins to use on Windows. */
 	UPROPERTY(config, EditAnywhere, Category = "Audio")
 	FString OcclusionPlugin;
+
+	/** Audio Cook Settings */
+
+	/** Various overrides for how this platform should handle compression and decompression */
+	UPROPERTY(GlobalConfig, EditAnywhere, Category = "Audio")
+	FPlatformRuntimeAudioCompressionOverrides CompressionOverrides;
+
+	/** When this is enabled, Actual compressed data will be separated from the USoundWave, and loaded into a cache. */
+	UPROPERTY(GlobalConfig, EditAnywhere, Category = "Audio|CookOverrides", meta = (DisplayName = "Use Stream Caching (Experimental)"))
+	bool bUseAudioStreamCaching;
+
+	/** This determines the max amount of memory that should be used for the cache at any given time. If set low (<= 8 MB), it lowers the size of individual chunks of audio during cook. */
+	UPROPERTY(GlobalConfig, EditAnywhere, Category = "Audio|CookOverrides|Stream Caching", meta = (DisplayName = "Max Cache Size (KB)"))
+	int32 CacheSizeKB;
+
+	UPROPERTY(GlobalConfig, EditAnywhere, Category = "Audio|CookOverrides")
+	bool bResampleForDevice;
+
+	/** Mapping of which sample rates are used for each sample rate quality for a specific platform. */
+
+	UPROPERTY(GlobalConfig, EditAnywhere, Category = "Audio|CookOverrides|ResamplingQuality", meta = (DisplayName = "Max"))
+	float MaxSampleRate;
+
+	UPROPERTY(GlobalConfig, EditAnywhere, Category = "Audio|CookOverrides|ResamplingQuality", meta = (DisplayName = "High"))
+	float HighSampleRate;
+
+	UPROPERTY(GlobalConfig, EditAnywhere, Category = "Audio|CookOverrides|ResamplingQuality", meta = (DisplayName = "Medium"))
+	float MedSampleRate;
+
+	UPROPERTY(GlobalConfig, EditAnywhere, Category = "Audio|CookOverrides|ResamplingQuality", meta = (DisplayName = "Low"))
+	float LowSampleRate;
+
+	UPROPERTY(GlobalConfig, EditAnywhere, Category = "Audio|CookOverrides|ResamplingQuality", meta = (DisplayName = "Min"))
+	float MinSampleRate;
+
+	/** Scales all compression qualities when cooking to this platform. For example, 0.5 will halve all compression qualities, and 1.0 will leave them unchanged. */
+	UPROPERTY(GlobalConfig, EditAnywhere, Category = "Audio|CookOverrides")
+	float CompressionQualityModifier;
+
+	/** When set to anything beyond 0, this will ensure any SoundWaves longer than this value, in seconds, to stream directly off of the disk. */
+	UPROPERTY(GlobalConfig, EditAnywhere, Category = "Audio|CookOverrides", meta = (DisplayName = "Stream All Soundwaves Longer Than: "))
+	float AutoStreamingThreshold;
 };
diff --git a/Engine/Source/Developer/Windows/WindowsTargetPlatform/Private/GenericWindowsTargetPlatform.h b/Engine/Source/Developer/Windows/WindowsTargetPlatform/Private/GenericWindowsTargetPlatform.h
--- a/Engine/Source/Developer/Windows/WindowsTargetPlatform/Private/GenericWindowsTargetPlatform.h
+++ b/Engine/Source/Developer/Windows/WindowsTargetPlatform/Private/GenericWindowsTargetPlatform.h
@@ -17,6 +17,120 @@
 
 #define LOCTEXT_NAMESPACE "TGenericWindowsTargetPlatform"
 
+namespace Windows
+{
+	void CachePlatformAudioCookOverrides(FPlatformAudioCookOverrides& OutOverrides)
+	{
+		const TCHAR* CategoryName = TEXT("/Script/WindowsTargetPlatform.WindowsTargetSettings");
+
+		GConfig->GetBool(CategoryName, TEXT("bUseAudioStreamCaching"), OutOverrides.bUseStreamCaching, GEngineIni);
+
+		/** Memory Load On Demand Settings */
+		if (OutOverrides.bUseStreamCaching)
+		{
+			// Cache size:
+			int32 RetrievedCacheSize = 32 * 1024;
+			GConfig->GetInt(CategoryName, TEXT("CacheSizeKB"), RetrievedCacheSize, GEngineIni);
+			OutOverrides.StreamCachingSettings.CacheSizeKB = RetrievedCacheSize;
+		}
+
+		GConfig->GetBool(CategoryName, TEXT("bResampleForDevice"), OutOverrides.bResampleForDevice, GEngineIni);
+
+		GConfig->GetFloat(CategoryName, TEXT("CompressionQualityModifier"), OutOverrides.CompressionQualityModifier, GEngineIni);
+
+		GConfig->GetFloat(CategoryName, TEXT("AutoStreamingThreshold"), OutOverrides.AutoStreamingThreshold, GEngineIni);
+
+		//Cache sample rate map:
+		float RetrievedSampleRate = -1.0f;
+
+		GConfig->GetFloat(CategoryName, TEXT("MaxSampleRate"), RetrievedSampleRate, GEngineIni);
+		float* FoundSampleRate = OutOverrides.PlatformSampleRates.Find(ESoundwaveSampleRateSettings::Max);
+
+		if (FoundSampleRate)
+		{
+			if (!FMath::IsNearlyEqual(*FoundSampleRate, RetrievedSampleRate))
+			{
+				*FoundSampleRate = RetrievedSampleRate;
+			}
+
+		}
+		else
+		{
+			OutOverrides.PlatformSampleRates.Add(ESoundwaveSampleRateSettings::Max, RetrievedSampleRate);
+		}
+
+		RetrievedSampleRate = -1.0f;
+
+		GConfig->GetFloat(CategoryName, TEXT("HighSampleRate"), RetrievedSampleRate, GEngineIni);
+		FoundSampleRate = OutOverrides.PlatformSampleRates.Find(ESoundwaveSampleRateSettings::High);
+
+		if (FoundSampleRate)
+		{
+			if (!FMath::IsNearlyEqual(*FoundSampleRate, RetrievedSampleRate))
+			{
+				*FoundSampleRate = RetrievedSampleRate;
+			}
+
+		}
+		else
+		{
+			OutOverrides.PlatformSampleRates.Add(ESoundwaveSampleRateSettings::High, RetrievedSampleRate);
+		}
+
+
+		RetrievedSampleRate = -1.0f;
+
+		GConfig->GetFloat(CategoryName, TEXT("MedSampleRate"), RetrievedSampleRate, GEngineIni);
+		FoundSampleRate = OutOverrides.PlatformSampleRates.Find(ESoundwaveSampleRateSettings::Medium);
+
+		if (FoundSampleRate)
+		{
+			if (!FMath::IsNearlyEqual(*FoundSampleRate, RetrievedSampleRate))
+			{
+				*FoundSampleRate = RetrievedSampleRate;
+			}
+		}
+		else
+		{
+			OutOverrides.PlatformSampleRates.Add(ESoundwaveSampleRateSettings::Medium, RetrievedSampleRate);
+		}
+
+		RetrievedSampleRate = -1.0f;
+
+		GConfig->GetFloat(CategoryName, TEXT("LowSampleRate"), RetrievedSampleRate, GEngineIni);
+		FoundSampleRate = OutOverrides.PlatformSampleRates.Find(ESoundwaveSampleRateSettings::Low);
+
+		if (FoundSampleRate)
+		{
+			if (!FMath::IsNearlyEqual(*FoundSampleRate, RetrievedSampleRate))
+			{
+				*FoundSampleRate = RetrievedSampleRate;
+			}
+		}
+		else
+		{
+			OutOverrides.PlatformSampleRates.Add(ESoundwaveSampleRateSettings::Low, RetrievedSampleRate);
+		}
+
+		RetrievedSampleRate = -1.0f;
+
+		GConfig->GetFloat(CategoryName, TEXT("MinSampleRate"), RetrievedSampleRate, GEngineIni);
+		FoundSampleRate = OutOverrides.PlatformSampleRates.Find(ESoundwaveSampleRateSettings::Min);
+
+		if (FoundSampleRate)
+		{
+			if (!FMath::IsNearlyEqual(*FoundSampleRate, RetrievedSampleRate))
+			{
+				*FoundSampleRate = RetrievedSampleRate;
+			}
+		}
+		else
+		{
+			OutOverrides.PlatformSampleRates.Add(ESoundwaveSampleRateSettings::Min, RetrievedSampleRate);
+		}
+	}
+}
+
 /**
  * Template for Windows target platforms
  */
@@ -438,7 +552,21 @@ public:
 
 	virtual FPlatformAudioCookOverrides* GetAudioCompressionSettings() const override
 	{
-		return nullptr;
+		static FPlatformAudioCookOverrides Settings;
+
+#if !WITH_EDITOR
+		static bool bCachedPlatformSettings = false;
+
+		if (!bCachedPlatformSettings)
+		{
+			Windows::CachePlatformAudioCookOverrides(Settings);
+			bCachedPlatformSettings = true;
+		}
+#else
+		Windows::CachePlatformAudioCookOverrides(Settings);
+#endif
+
+		return &Settings;
 	}
 
 #endif //WITH_ENGINE
diff --git a/Engine/Source/Developer/Windows/WindowsTargetPlatform/WindowsTargetPlatform.Build.cs b/Engine/Source/Developer/Windows/WindowsTargetPlatform/WindowsTargetPlatform.Build.cs
--- a/Engine/Source/Developer/Windows/WindowsTargetPlatform/WindowsTargetPlatform.Build.cs
+++ b/Engine/Source/Developer/Windows/WindowsTargetPlatform/WindowsTargetPlatform.Build.cs
@@ -12,7 +12,8 @@ public class WindowsTargetPlatform : ModuleRules
 				"CoreUObject",
 				"TargetPlatform",
 				"DesktopPlatform",
-			}
+                "AudioPlatformConfiguration",
+            }
 		);
 
 		PrivateIncludePathModuleNames.AddRange(
diff --git a/Engine/Source/Editor/DetailCustomizations/Private/WindowsTargetSettingsDetails.cpp b/Engine/Source/Editor/DetailCustomizations/Private/WindowsTargetSettingsDetails.cpp
--- a/Engine/Source/Editor/DetailCustomizations/Private/WindowsTargetSettingsDetails.cpp
+++ b/Engine/Source/Editor/DetailCustomizations/Private/WindowsTargetSettingsDetails.cpp
@@ -9,6 +9,7 @@
 #include "Widgets/SNullWidget.h"
 #include "Widgets/DeclarativeSyntaxSupport.h"
 #include "Widgets/SBoxPanel.h"
+#include "Widgets/Input/SCheckBox.h"
 #include "Styling/SlateTypes.h"
 #include "Textures/SlateIcon.h"
 #include "Framework/Commands/UIAction.h"
@@ -31,6 +32,7 @@
 
 #if WITH_ENGINE
 #include "AudioDevice.h"
+#include "ContentStreaming.h"
 #endif 
 
 #define LOCTEXT_NAMESPACE "WindowsTargetSettingsDetails"
@@ -335,7 +337,33 @@ void FWindowsTargetSettingsDetails::CustomizeDetails( IDetailLayoutBuilder& Deta
 				.ToolTipText(LOCTEXT("AudioDevicesButtonToolTip", "Pick from the list of available audio devices"))
 			]
 	];
+
 	AudioPluginWidgetManager.BuildAudioCategory(DetailBuilder, EAudioPlatform::Windows);
+
+	// Here we add a callback when the 
+	TSharedPtr<IPropertyHandle> AudioStreamCachingPropertyHandle = DetailBuilder.GetProperty("bUseAudioStreamCaching");
+	IDetailCategoryBuilder& AudioStreamCachingCategory = DetailBuilder.EditCategory("Audio");
+	IDetailPropertyRow& AudioStreamCachingPropertyRow = AudioDeviceCategory.AddProperty(AudioStreamCachingPropertyHandle);
+	AudioStreamCachingPropertyRow.CustomWidget()
+		.NameContent()
+		[
+			AudioStreamCachingPropertyHandle->CreatePropertyNameWidget()
+		]
+		.ValueContent()
+		.MaxDesiredWidth(500.0f)
+		.MinDesiredWidth(100.0f)
+		[
+			SNew(SHorizontalBox)
+
+			+ SHorizontalBox::Slot()
+			.FillWidth(1.0f)
+			[
+				SNew(SCheckBox)
+				.OnCheckStateChanged(this, &FWindowsTargetSettingsDetails::HandleAudioStreamCachingToggled, AudioStreamCachingPropertyHandle)
+				.IsChecked(this, &FWindowsTargetSettingsDetails::GetAudioStreamCachingToggled, AudioStreamCachingPropertyHandle)
+				.ToolTipText(AudioStreamCachingPropertyHandle->GetToolTipText())
+			]
+		];
 }
 
 bool FWindowsTargetSettingsDetails::HandlePreExternalIconCopy(const FString& InChosenImage)
@@ -361,6 +389,30 @@ void FWindowsTargetSettingsDetails::HandleAudioDeviceSelected(FString AudioDevic
 	PropertyHandle->SetValue(AudioDeviceName);
 }
 
+void FWindowsTargetSettingsDetails::HandleAudioStreamCachingToggled(ECheckBoxState EnableStreamCaching, TSharedPtr<IPropertyHandle> PropertyHandle)
+{
+	PropertyHandle->SetValue(EnableStreamCaching == ECheckBoxState::Checked);
+
+#if WITH_ENGINE
+	IStreamingManager::Get().OnAudioStreamingParamsChanged();
+#endif
+}
+
+ECheckBoxState FWindowsTargetSettingsDetails::GetAudioStreamCachingToggled(TSharedPtr<IPropertyHandle> PropertyHandle) const
+{
+	bool bEnabled = false;
+	PropertyHandle->GetValue(bEnabled);
+
+	if (bEnabled)
+	{
+		return ECheckBoxState::Checked;
+	}
+	else
+	{
+		return ECheckBoxState::Unchecked;
+	}
+}
+
 FSlateColor FWindowsTargetSettingsDetails::HandleAudioDeviceBoxForegroundColor(TSharedPtr<IPropertyHandle> PropertyHandle) const
 {
 	FString Value;
diff --git a/Engine/Source/Editor/DetailCustomizations/Private/WindowsTargetSettingsDetails.h b/Engine/Source/Editor/DetailCustomizations/Private/WindowsTargetSettingsDetails.h
--- a/Engine/Source/Editor/DetailCustomizations/Private/WindowsTargetSettingsDetails.h
+++ b/Engine/Source/Editor/DetailCustomizations/Private/WindowsTargetSettingsDetails.h
@@ -38,6 +38,12 @@ private:
 	/** Handles when a new audio device is selected from list of available audio devices. */
 	void HandleAudioDeviceSelected(FString AudioDeviceName, TSharedPtr<IPropertyHandle> PropertyHandle);
 
+	/** Handles when Stream Caching is toggled. */
+	void HandleAudioStreamCachingToggled(ECheckBoxState bEnableStreamCaching, TSharedPtr<IPropertyHandle> PropertyHandle);
+
+	/** This gets the current value of the audio stream caching bool property. */
+	ECheckBoxState GetAudioStreamCachingToggled(TSharedPtr<IPropertyHandle> PropertyHandle) const;
+
 	/** Handles changing the foreground color of the audio device box. */
 	FSlateColor HandleAudioDeviceBoxForegroundColor(TSharedPtr<IPropertyHandle> PropertyHandle) const;
 
diff --git a/Engine/Source/Runtime/Android/AndroidRuntimeSettings/Classes/AndroidRuntimeSettings.h b/Engine/Source/Runtime/Android/AndroidRuntimeSettings/Classes/AndroidRuntimeSettings.h
--- a/Engine/Source/Runtime/Android/AndroidRuntimeSettings/Classes/AndroidRuntimeSettings.h
+++ b/Engine/Source/Runtime/Android/AndroidRuntimeSettings/Classes/AndroidRuntimeSettings.h
@@ -489,6 +489,14 @@ public:
 	UPROPERTY(config, EditAnywhere, Category = "Audio")
 	FPlatformRuntimeAudioCompressionOverrides CompressionOverrides;
 
+	/** When this is enabled, Actual compressed data will be separated from the USoundWave, and loaded into a cache. */
+	UPROPERTY(GlobalConfig, EditAnywhere, Category = "Audio|CookOverrides", meta = (DisplayName = "Use Stream Caching (Experimental)"))
+	bool bUseAudioStreamCaching;
+
+	/** This determines the max amount of memory that should be used for the cache at any given time. If set low (<= 8 MB), it lowers the size of individual chunks of audio during cook. */
+	UPROPERTY(GlobalConfig, EditAnywhere, Category = "Audio|CookOverrides|Stream Caching", meta = (DisplayName = "Max Cache Size (KB)"))
+	int32 CacheSizeKB;
+
 	UPROPERTY(config, EditAnywhere, Category = "Audio|CookOverrides")
 	bool bResampleForDevice;
 
diff --git a/Engine/Source/Runtime/AudioMixer/Private/AudioMixerBlueprintLibrary.cpp b/Engine/Source/Runtime/AudioMixer/Private/AudioMixerBlueprintLibrary.cpp
--- a/Engine/Source/Runtime/AudioMixer/Private/AudioMixerBlueprintLibrary.cpp
+++ b/Engine/Source/Runtime/AudioMixer/Private/AudioMixerBlueprintLibrary.cpp
@@ -6,7 +6,9 @@
 #include "AudioMixerDevice.h"
 #include "CoreMinimal.h"
 #include "DSP/SpectrumAnalyzer.h"
-
+#include "ContentStreaming.h"
+#include "AudioCompressionSettingsUtils.h"
+#include "Async/Async.h"
 
 // This is our global recording task:
 static TUniquePtr<Audio::FAudioRecordingData> RecordingData;
@@ -347,6 +349,43 @@ int32 UAudioMixerBlueprintLibrary::GetNumberOfEntriesInSourceEffectChain(const U
 	return 0;
 }
 
+void UAudioMixerBlueprintLibrary::PrimeSoundForPlayback(USoundWave* SoundWave, const FOnSoundLoadComplete OnLoadCompletion)
+{
+	if (!SoundWave)
+	{
+		UE_LOG(LogAudioMixer, Warning, TEXT("Prime Sound For Playback called with a null SoundWave pointer."));
+	}
+	else if (!FPlatformCompressionUtilities::IsCurrentPlatformUsingStreamCaching())
+	{
+		UE_LOG(LogAudioMixer, Warning, TEXT("Prime Sound For Playback doesn't do anything unless Audio Load On Demand is enabled."));
+		
+		OnLoadCompletion.ExecuteIfBound(SoundWave, false);
+	}
+	else
+	{
+		IStreamingManager::Get().GetAudioStreamingManager().RequestChunk(SoundWave, 1, [OnLoadCompletion, SoundWave](EAudioChunkLoadResult InResult) 
+		{
+			AsyncTask(ENamedThreads::GameThread, [OnLoadCompletion, SoundWave, InResult]() {
+				if (InResult == EAudioChunkLoadResult::Completed || InResult == EAudioChunkLoadResult::AlreadyLoaded)
+				{
+					OnLoadCompletion.ExecuteIfBound(SoundWave, false);
+				}
+				else
+				{
+					OnLoadCompletion.ExecuteIfBound(SoundWave, true);
+				}
+			});
+		});
+	}
+}
+
+float UAudioMixerBlueprintLibrary::TrimAudioCache(float InMegabytesToFree)
+{
+	uint64 NumBytesToFree = (uint64) (((double)InMegabytesToFree) * 1024.0 * 1024.0);
+	uint64 NumBytesFreed = IStreamingManager::Get().GetAudioStreamingManager().TrimMemory(NumBytesToFree);
+	return (float)(((double) NumBytesFreed / 1024) / 1024.0);
+}
+
 void UAudioMixerBlueprintLibrary::PopulateSpectrumAnalyzerSettings(EFFTSize FFTSize, EFFTPeakInterpolationMethod InterpolationMethod, EFFTWindowType WindowType, float HopSize, Audio::FSpectrumAnalyzerSettings &OutSettings)
 {
 	switch (FFTSize)
diff --git a/Engine/Source/Runtime/AudioMixer/Public/AudioMixerBlueprintLibrary.h b/Engine/Source/Runtime/AudioMixer/Public/AudioMixerBlueprintLibrary.h
--- a/Engine/Source/Runtime/AudioMixer/Public/AudioMixerBlueprintLibrary.h
+++ b/Engine/Source/Runtime/AudioMixer/Public/AudioMixerBlueprintLibrary.h
@@ -58,6 +58,12 @@ enum class EFFTWindowType : uint8
 	Blackman
 };
 
+/** 
+* Called when a load request for a sound has completed.
+*/
+DECLARE_DYNAMIC_DELEGATE_TwoParams(FOnSoundLoadComplete, const class USoundWave*, LoadedSoundWave, const bool, WasCancelled);
+
+
 UCLASS(meta=(ScriptName="AudioMixerLibrary"))
 class AUDIOMIXER_API UAudioMixerBlueprintLibrary : public UBlueprintFunctionLibrary
 {
@@ -125,6 +131,13 @@ public:
 	UFUNCTION(BlueprintCallable, Category = "Audio|Effects", meta = (WorldContext = "WorldContextObject"))
 	static int32 GetNumberOfEntriesInSourceEffectChain(const UObject* WorldContextObject, USoundEffectSourcePresetChain* PresetChain);
 
+	/** Begin loading a sound into the cache so that it can be played immediately. */
+	UFUNCTION(BlueprintCallable, Category = "Sound")
+	static void PrimeSoundForPlayback(USoundWave* SoundWave, const FOnSoundLoadComplete OnLoadCompletion);
+
+	/** Trim memory used by the audio cache. Returns the number of megabytes freed. */
+	UFUNCTION(BlueprintCallable, Category = "Sound")
+	static float TrimAudioCache(float InMegabytesToFree);
 
 private:
 	static void PopulateSpectrumAnalyzerSettings(EFFTSize FFTSize, EFFTPeakInterpolationMethod InterpolationMethod, EFFTWindowType WindowType, float HopSize, Audio::FSpectrumAnalyzerSettings &OutSettings);
diff --git a/Engine/Source/Runtime/AudioPlatformConfiguration/Public/AudioCompressionSettings.h b/Engine/Source/Runtime/AudioPlatformConfiguration/Public/AudioCompressionSettings.h
--- a/Engine/Source/Runtime/AudioPlatformConfiguration/Public/AudioCompressionSettings.h
+++ b/Engine/Source/Runtime/AudioPlatformConfiguration/Public/AudioCompressionSettings.h
@@ -15,6 +15,26 @@ enum class ESoundwaveSampleRateSettings : uint8
 	MatchDevice
 };
 
+/************************************************************************/
+/* FAudioStreamCachingSettings                                           */
+/* Properties used to determine chunk sizes for the two caches used     */
+/* when the experimental Stream Caching feature is used.                */
+/************************************************************************/
+struct FAudioStreamCachingSettings
+{
+	static constexpr int32 DefaultCacheSize = 32 * 1024;
+
+	// Target memory usage, in kilobytes.
+	// In the future settings for the cache can be more complex, but for now
+	// we set the max chunk size to 256 kilobytes, then set the number of elements in our cache as
+	// CacheSizeKB / 256.
+	int32 CacheSizeKB;
+
+	FAudioStreamCachingSettings()
+		: CacheSizeKB(DefaultCacheSize)
+	{
+	}
+};
 
 /************************************************************************/
 /* FPlatformAudioCookOverrides                                          */
@@ -23,6 +43,11 @@ enum class ESoundwaveSampleRateSettings : uint8
 /************************************************************************/
 struct FPlatformAudioCookOverrides
 {
+	// Increment this to force a recook on all Stream Caching assets.
+	// For testing, it's useful to set this to either a negative number or
+	// absurdly large number, to ensure you do not pollute the DDC.
+	static const int32 StreamCachingVersion = 5017;
+
 	bool bResampleForDevice;
 
 	// Mapping of which sample rates are used for each sample rate quality for a specific platform.
@@ -32,12 +57,20 @@ struct FPlatformAudioCookOverrides
 	float CompressionQualityModifier;
 
 	// When set to any platform > 0.0, this will automatically set any USoundWave beyond this value to be streamed from disk.
+	// If StreamCaching is set to true, this will be used 
 	float AutoStreamingThreshold;
 
+	// Whether to use the experimental Load on Demand feature, which uses as little memory at runtime as possible.
+	bool bUseStreamCaching;
+
+	// If Load On Demand is enabled, these settings are used to determine chunks and cache sizes.
+	FAudioStreamCachingSettings StreamCachingSettings;
+
 	FPlatformAudioCookOverrides()
 		: bResampleForDevice(false)
 		, CompressionQualityModifier(1.0f)
 		, AutoStreamingThreshold(0.0f)
+		, bUseStreamCaching(false)
 	{
 		PlatformSampleRates.Add(ESoundwaveSampleRateSettings::Max, 48000);
 		PlatformSampleRates.Add(ESoundwaveSampleRateSettings::High, 32000);
@@ -46,7 +79,7 @@ struct FPlatformAudioCookOverrides
 		PlatformSampleRates.Add(ESoundwaveSampleRateSettings::Min, 8000);
 	}
 
-	// This is used to invalidate compressed audio for a specific platform
+	// This is used to invalidate compressed audio for a specific platform.
 	static void GetHashSuffix(const FPlatformAudioCookOverrides* InOverrides, FString& OutSuffix)
 	{
 		if (InOverrides == nullptr)
@@ -60,6 +93,18 @@ struct FPlatformAudioCookOverrides
 		int32 AutoStreamingThresholdHash = FMath::FloorToInt(InOverrides->AutoStreamingThreshold * 100.0f);
 		OutSuffix.AppendInt(AutoStreamingThresholdHash);
 
+		if (InOverrides->bUseStreamCaching)
+		{
+			OutSuffix.Append(TEXT("_StreamCache_Ver"));
+			OutSuffix.AppendInt(StreamCachingVersion);
+			OutSuffix.AppendChar('_');
+
+			// cache info:
+			OutSuffix.Append(TEXT("MEM_"));
+			OutSuffix.AppendInt(InOverrides->StreamCachingSettings.CacheSizeKB);
+		}
+		
+
 		int32 ResampleBoolHash = (int32)InOverrides->bResampleForDevice;
 		OutSuffix.AppendInt(ResampleBoolHash);
 
diff --git a/Engine/Source/Runtime/Engine/Classes/Engine/Engine.h b/Engine/Source/Runtime/Engine/Classes/Engine/Engine.h
--- a/Engine/Source/Runtime/Engine/Classes/Engine/Engine.h
+++ b/Engine/Source/Runtime/Engine/Classes/Engine/Engine.h
@@ -3397,6 +3397,7 @@ private:
 	bool ToggleStatSoundWaves(UWorld* World, FCommonViewportClient* ViewportClient, const TCHAR* Stream = nullptr);
 	bool ToggleStatSoundCues(UWorld* World, FCommonViewportClient* ViewportClient, const TCHAR* Stream = nullptr);
 	bool ToggleStatSounds(UWorld* World, FCommonViewportClient* ViewportClient, const TCHAR* Stream = nullptr);
+	bool ToggleStatAudioStreaming(UWorld* World, FCommonViewportClient* ViewportClient, const TCHAR* Stream = nullptr);
 	bool ToggleStatSoundMixes(UWorld* World, FCommonViewportClient* ViewportClient, const TCHAR* Stream = nullptr);
 	bool ToggleStatSoundModulators(UWorld* World, FCommonViewportClient* ViewportClient, const TCHAR* Stream = nullptr);
 #endif
@@ -3428,6 +3429,7 @@ private:
  	int32 RenderStatSoundMixes(UWorld* World, FViewport* Viewport, FCanvas* Canvas, int32 X, int32 Y, const FVector* ViewLocation = nullptr, const FRotator* ViewRotation = nullptr);
 	int32 RenderStatSoundModulators(UWorld* World, FViewport* Viewport, FCanvas* Canvas, int32 X, int32 Y, const FVector* ViewLocation = nullptr, const FRotator* ViewRotation = nullptr);
 	int32 RenderStatSoundWaves(UWorld* World, FViewport* Viewport, FCanvas* Canvas, int32 X, int32 Y, const FVector* ViewLocation = nullptr, const FRotator* ViewRotation = nullptr);
+	int32 RenderStatAudioStreaming(UWorld* World, FViewport* Viewport, FCanvas* Canvas, int32 X, int32 Y, const FVector* ViewLocation = nullptr, const FRotator* ViewRotation = nullptr);
 	int32 RenderStatSoundCues(UWorld* World, FViewport* Viewport, FCanvas* Canvas, int32 X, int32 Y, const FVector* ViewLocation = nullptr, const FRotator* ViewRotation = nullptr);
 	int32 RenderStatSounds(UWorld* World, FViewport* Viewport, FCanvas* Canvas, int32 X, int32 Y, const FVector* ViewLocation = nullptr, const FRotator* ViewRotation = nullptr);
 #endif // !UE_BUILD_SHIPPING
diff --git a/Engine/Source/Runtime/Engine/Classes/Sound/SoundWave.h b/Engine/Source/Runtime/Engine/Classes/Sound/SoundWave.h
--- a/Engine/Source/Runtime/Engine/Classes/Sound/SoundWave.h
+++ b/Engine/Source/Runtime/Engine/Classes/Sound/SoundWave.h
@@ -16,6 +16,7 @@
 #include "Sound/SoundGroups.h"
 #include "AudioMixerTypes.h"
 #include "AudioCompressionSettings.h"
+#include "PerPlatformProperties.h"
 #include "SoundWave.generated.h"
 
 class ITargetPlatform;
@@ -108,13 +109,15 @@ struct FStreamedAudioPlatformData
 	~FStreamedAudioPlatformData();
 
 	/**
-	 * Try to load audio chunk from the derived data cache.
+	 * Try to load audio chunk from the derived data cache or build it if it isn't there.
 	 * @param ChunkIndex	The Chunk index to load.
 	 * @param OutChunkData	Address of pointer that will store chunk data - should
 	 *						either be NULL or have enough space for the chunk
-	 * @returns true if requested chunk has been loaded.
+	 * @returns if > 0, the size of the chunk in bytes. If 0, the chunk failed to load.
 	 */
-	bool TryLoadChunk(int32 ChunkIndex, uint8** OutChunkData, bool bMakeSureChunkIsLoaded = false);
+	int32 GetChunkFromDDC(int32 ChunkIndex, uint8** OutChunkData, bool bMakeSureChunkIsLoaded = false);
+
+	
 
 	/** Serialization. */
 	void Serialize(FArchive& Ar, class USoundWave* Owner);
@@ -127,6 +130,18 @@ struct FStreamedAudioPlatformData
 	bool AreDerivedChunksAvailable() const;
 #endif
 
+private:
+
+	/**
+	 * Takes the results of a DDC operation and deserializes it into an FStreamedAudioChunk struct.
+	 * @param SerializedData Serialized data resulting from DDC.GetAsynchronousResults or DDC.GetSynchronous.
+	 * @param ChunkToDeserializeInto is the chunk to fill with the deserialized data.
+	 * @param ChunkIndex is the index of the chunk in this instance of FStreamedAudioPlatformData.
+	 * @param bCachedChunk is true if the chunk was successfully cached, false otherwise.
+	 * @param OutChunkData is a pointer to a pointer to populate with the chunk itself, or if pointing to nullptr, returns an allocated buffer.
+	 * @returns the size of the chunk loaded in bytes, or zero if the chunk didn't load.
+	 */
+	int32 DeserializeChunkFromDDC(TArray<uint8> SerializedData, FStreamedAudioChunk &ChunkToDeserializeInto, int32 ChunkIndex, uint8** &OutChunkData);
 };
 
 USTRUCT(BlueprintType)
@@ -294,14 +309,22 @@ class ENGINE_API USoundWave : public USoundBase
 	uint8 bLooping:1;
 
 	/** Whether this sound can be streamed to avoid increased memory usage */
-	UPROPERTY(EditAnywhere, Category="Playback|Streaming")
+	UPROPERTY(EditAnywhere, Category="Playback|Streaming", meta = (DisplayName = "Force Streaming"))
 	uint8 bStreaming:1;
 
 	/** Whether this sound supports seeking. This requires recooking with a codec which supports seekability and streaming. */
 	UPROPERTY(EditAnywhere, Category = "Playback|Streaming", meta = (DisplayName = "Seekable", EditCondition = "bStreaming"))
 	uint8 bSeekableStreaming:1;
 
-	/** Set to true for programatically-generated, streamed audio. */
+	/** If Load On Demand or Auto-Streaming is enabled, use this to force the entire compressed audio file to be loaded when this USoundWave is loaded. */
+	UPROPERTY(EditAnywhere, AdvancedDisplay, Category = "Loading", meta = (DisplayName = "Ensure Sound is Always Loaded (Inlined)"))
+	uint8 bForceInline : 1;
+
+	/** If Load On Demand or Auto-Streaming is enabled, This will immediately kick off a load for this sound once postload is complete. */
+	UPROPERTY(EditAnywhere, AdvancedDisplay, Category = "Loading", meta = (DisplayName = "Load Compressed Audio When This Sound Wave Is Loaded"))
+	uint8 bLoadCompressedAudioWhenSoundWaveIsLoaded : 1;
+
+	/** Set to true for programmatically-generated, streamed audio. */
 	uint8 bProcedural:1;
 
 	/** Whether this sound wave is beginning to be destroyed by GC. */
@@ -340,6 +363,10 @@ class ENGINE_API USoundWave : public USoundBase
 	/** Whether this SoundWave was decompressed from OGG. */
 	uint8 bDecompressedFromOgg : 1;
 
+	/** Use this to override how much audio data is loaded when this USoundWave is loaded. */
+	UPROPERTY(EditAnywhere, AdvancedDisplay, Category = "Loading")
+	int32 InitialChunkSize;
+
 #if WITH_EDITORONLY_DATA
 	/** Specify a sound to use for the baked analysis. Will default to this USoundWave if not sete. */
 	UPROPERTY(EditAnywhere, Category = "Analysis")
@@ -542,6 +569,9 @@ public:
 	FOwnedBulkDataPtr* OwnedBulkDataPtr;
 	const uint8* ResourceData;
 
+	/** Zeroth Chunk of audio for sources that use Load On Demand. */
+	TArray<uint8> ZerothChunkData;
+
 	/** Uncompressed wav data 16 bit in mono or stereo - stereo not allowed for multichannel data */
 	FByteBulkData RawData;
 
@@ -579,6 +609,16 @@ public:
 	virtual bool IsReadyForFinishDestroy() override;
 	virtual void FinishDestroy() override;
 	virtual void PostLoad() override;
+
+	// When stream caching is enabled, this is called after we've successfully compressed and split the streamed audio for this file.
+	void EnsureZerothChunkIsLoaded();
+
+	// Returns the amount of chunks this soundwave contains if it's streaming,
+	// or zero if it is not a streaming source.
+	uint32 GetNumChunks() const;
+
+	uint32 GetSizeOfChunk(uint32 ChunkIndex);
+
 	virtual void BeginDestroy() override;
 #if WITH_EDITOR
 	virtual void PostEditChangeProperty(FPropertyChangedEvent& PropertyChangedEvent) override;
@@ -716,6 +756,7 @@ public:
 #endif
 
 private:
+
 	FName GetPlatformSpecificFormat(FName Format, const FPlatformAudioCookOverrides* CompressionOverrides);
 
 #if WITH_EDITOR
@@ -774,6 +815,15 @@ public:
 
 	/** Checks whether sound has seekable streaming enabled. */
 	bool IsSeekableStreaming() const;
+	/**
+	 * Checks whether we should use the load on demand cache.
+	 */
+	bool ShouldUseStreamCaching() const;
+
+	/**
+	 * This returns the initial chunk of compressed data for streaming data sources.
+	 */
+	TArrayView<const uint8> GetZerothChunk();
 
 	/**
 	 * Attempts to update the cached platform data after any changes that might affect it
diff --git a/Engine/Source/Runtime/Engine/Private/ADPCMAudioInfo.cpp b/Engine/Source/Runtime/Engine/Private/ADPCMAudioInfo.cpp
--- a/Engine/Source/Runtime/Engine/Private/ADPCMAudioInfo.cpp
+++ b/Engine/Source/Runtime/Engine/Private/ADPCMAudioInfo.cpp
@@ -13,8 +13,6 @@
 
 namespace ADPCM
 {
-	const uint32 MaxChunkSize = 256 * 1024;
-
 	void DecodeBlock(const uint8* EncodedADPCMBlock, int32 BlockSize, int16* DecodedPCMData);
 	void DecodeBlockStereo(const uint8* EncodedADPCMBlockLeft, const uint8* EncodedADPCMBlockRight, int32 BlockSize, int16* DecodedPCMData);
 }
@@ -44,7 +42,8 @@ FADPCMAudioInfo::~FADPCMAudioInfo(void)
 
 void FADPCMAudioInfo::SeekToTime(const float SeekTime)
 {
-	CurCompressedChunkData = nullptr;
+	// Reset chunk handle in preperation for a new chunk.
+	CurrentCompressedChunkData = FAudioChunkHandle();
 
 	if(SeekTime <= 0.0f)
 	{
@@ -60,49 +59,77 @@ void FADPCMAudioInfo::SeekToTime(const float SeekTime)
 	// Calculate block index & force SeekTime to be in bounds.
 	check(WaveInfo.pSamplesPerSec != nullptr);
 	const uint32 SamplesPerSec = *WaveInfo.pSamplesPerSec;
-	const uint32 SeekedSamples = static_cast<uint32>(SeekTime * SamplesPerSec);
+	uint32 SeekedSamples = static_cast<uint32>(SeekTime * SamplesPerSec);
 	TotalSamplesStreamed = FMath::Min<uint32>(SeekedSamples, TotalSamplesPerChannel - 1);
 
 	const uint32 HeaderOffset = static_cast<uint32>(WaveInfo.SampleDataStart - SrcBufferData);
-
-	if (Format == WAVE_FORMAT_ADPCM)
+	
+	if (!StreamingSoundWave)
 	{
-		CurrentCompressedBlockIndex = TotalSamplesStreamed / SamplesPerBlock; // Compute the block index that where SeekTime resides.
-		CurrentChunkIndex = 0;
-		CurrentChunkBufferOffset = HeaderOffset;
-
-		const int32 ChannelBlockSize = BlockSize * NumChannels;
-		for (uint32 BlockIndex = 0; BlockIndex < CurrentCompressedBlockIndex; ++BlockIndex)
+		// For the non streaming case:
+		if (Format == WAVE_FORMAT_LPCM)
 		{
-			if (CurrentChunkBufferOffset + ChannelBlockSize >= ADPCM::MaxChunkSize)
-			{
-				++CurrentChunkIndex;
-				CurrentChunkBufferOffset = 0;
-			}
+			// There are no "blocks" on LPCM, so only update the total samples streamed (which is based off sample rate).
+			// Note that TotalSamplesStreamed is per-channel in the ReadCompressedInfo. Channels are taken into account there.
+			TotalSamplesStreamed = FMath::Clamp<uint32>(SeekedSamples, 0, TotalSamplesPerChannel - 1);
+		}
+		else
+		{
+			// Clamp to the end of memory in case we have an invalid seek time.
+			SeekedSamples = FMath::Clamp<uint32>(SeekedSamples, 0, TotalSamplesPerChannel - 1);
+
+			// Compute the block index that we're seeked to
+			CurrentCompressedBlockIndex = SeekedSamples / SamplesPerBlock;
 
-			// Always add chunks in NumChannels pairs
-			CurrentChunkBufferOffset += ChannelBlockSize;
+			// Update the samples streamed to the current block index and the samples per block
+			TotalSamplesStreamed = CurrentCompressedBlockIndex * SamplesPerBlock;
 		}
 	}
-	else if (Format == WAVE_FORMAT_LPCM)
+	else
 	{
-		const int32 ChannelBlockSize = sizeof(int16) * NumChannels;
+		if (Format == WAVE_FORMAT_ADPCM)
+		{
+			CurrentCompressedBlockIndex = TotalSamplesStreamed / SamplesPerBlock; // Compute the block index that where SeekTime resides.
+			CurrentChunkIndex = 0;
+			CurrentChunkBufferOffset = HeaderOffset;
 
-		// 1. Find total offset
-		CurrentChunkBufferOffset = HeaderOffset + (TotalSamplesStreamed * ChannelBlockSize);
+			const int32 ChannelBlockSize = BlockSize * NumChannels;
+			for (uint32 BlockIndex = 0; BlockIndex < CurrentCompressedBlockIndex; ++BlockIndex)
+			{
+				if (CurrentChunkBufferOffset + ChannelBlockSize >= StreamingSoundWave->GetSizeOfChunk(CurrentChunkIndex))
+				{
+					++CurrentChunkIndex;
+					CurrentChunkBufferOffset = 0;
+				}
 
-		// 2. Calculate index and remove from offset
-		CurrentChunkIndex = CurrentChunkBufferOffset / ADPCM::MaxChunkSize;
-		CurrentChunkBufferOffset %= ADPCM::MaxChunkSize;
+				// Always add chunks in NumChannels pairs
+				CurrentChunkBufferOffset += ChannelBlockSize;
+			}
+		}
+		else if (Format == WAVE_FORMAT_LPCM)
+		{
+			const int32 ChannelBlockSize = sizeof(int16) * NumChannels;
 
-		// 3. Trim remainder of block size, effectively aligning the block to a channel pair boundary
-		CurrentChunkBufferOffset -= CurrentChunkBufferOffset % ChannelBlockSize;
-	}
-	else
-	{
-		return;
-	}
+			// 1. Find total offset
+			CurrentChunkBufferOffset = HeaderOffset + (TotalSamplesStreamed * ChannelBlockSize);
+			
+			// 2. Calculate index and remove from offset
+			while (CurrentChunkBufferOffset >= StreamingSoundWave->GetSizeOfChunk(CurrentChunkIndex))
+			{
+				CurrentChunkBufferOffset -= StreamingSoundWave->GetSizeOfChunk(CurrentChunkIndex);
+				CurrentChunkIndex++;
+			}
 
+			// 3. Trim remainder of block size, effectively aligning the block to a channel pair boundary
+			CurrentChunkBufferOffset -= CurrentChunkBufferOffset % ChannelBlockSize;
+		}
+		else
+		{
+			// If we hit this, Format was invalid:
+			checkNoEntry();
+			return;
+		}
+	}
 	bSeekPending = true;
 }
 
@@ -314,29 +341,27 @@ bool FADPCMAudioInfo::StreamCompressedInfoInternal(USoundWave* Wave, struct FSou
 	check(StreamingSoundWave == Wave);
 
 	// Get the first chunk of audio data (should already be loaded)
-	uint8 const* FirstChunk = IStreamingManager::Get().GetAudioStreamingManager().GetLoadedChunk(Wave, 0, &CurrentChunkDataSize);
+	// uint8 const* firstChunk = IStreamingManager::Get().GetAudioStreamingManager().GetLoadedChunk(Wave, 0, &CurrentChunkDataSize);
+	TArrayView<const uint8> ZerothChunk = Wave->GetZerothChunk();
 
-	if (FirstChunk == nullptr)
+	if (ZerothChunk.Num() <= 0)
 	{
 		return false;
 	}
 
 	SrcBufferData = nullptr;
 	SrcBufferDataSize = 0;
-
-	void* FormatHeader;
-
-	if (!WaveInfo.ReadWaveInfo((uint8*)FirstChunk, Wave->RunningPlatformData->Chunks[0].AudioDataSize, nullptr, true, &FormatHeader))
+	
+	void*	FormatHeader;
+	
+	if (!WaveInfo.ReadWaveInfo(ZerothChunk.GetData(), ZerothChunk.Num(), NULL, true, &FormatHeader))
 	{
 		UE_LOG(LogAudio, Warning, TEXT("WaveInfo.ReadWaveInfo Failed"));
 		return false;
 	}
 
-	SrcBufferData = FirstChunk;
-
-	FirstChunkSampleDataOffset = WaveInfo.SampleDataStart - FirstChunk;
-	CurrentChunkBufferOffset = 0;
-	CurCompressedChunkData = nullptr;
+	FirstChunkSampleDataOffset = WaveInfo.SampleDataStart - ZerothChunk.GetData();
+	CurrentChunkBufferOffset = FirstChunkSampleDataOffset;
 	CurrentUncompressedBlockSampleIndex = 0;
 	CurrentChunkIndex = 0;
 	TotalSamplesStreamed = 0;
@@ -363,6 +388,7 @@ bool FADPCMAudioInfo::StreamCompressedInfoInternal(USoundWave* Wave, struct FSou
 
 		UncompressedBlockData = (uint8*)FMemory::Realloc(UncompressedBlockData, NumChannels * UncompressedBlockSize);
 		check(UncompressedBlockData != nullptr);
+		FMemory::Memzero(UncompressedBlockData, NumChannels * UncompressedBlockSize);
 	}
 	else if (Format == WAVE_FORMAT_LPCM)
 	{
@@ -428,54 +454,36 @@ bool FADPCMAudioInfo::StreamCompressedData(uint8* Destination, bool bLooping, ui
 		// We need to loop over the number of samples requested since an uncompressed block will not match the number of frames requested
 		while(BufferSize > 0)
 		{
-			if(CurCompressedChunkData == nullptr || CurrentUncompressedBlockSampleIndex >= UncompressedBlockSize / sizeof(uint16))
-			{
-				// we need to decompress another block of compressed data from the current chunk
-
-				if(CurCompressedChunkData == nullptr || CurrentChunkBufferOffset >= CurrentChunkDataSize)
-				{
-					// Get another chunk of compressed data from the streaming engine
-
-					// CurrentChunkIndex is used to keep track of the current chunk for loading/unloading by the streaming engine, but chunk 0 is
-					// preloaded so don't increment this when getting chunk 0, only later chunks. If previous chunk retrieval failed because it was
-					// not ready, don't re-increment the CurrentChunkIndex.  Only increment if seek not pending as seek determines the current chunk index
-					// and invalidates CurCompressedChunkData.
-					if(CurCompressedChunkData != nullptr)
-					{
-						++CurrentChunkIndex;
-					}
-
-					// Request the next chunk of data from the streaming engine
-					CurCompressedChunkData = IStreamingManager::Get().GetAudioStreamingManager().GetLoadedChunk(StreamingSoundWave, CurrentChunkIndex, &CurrentChunkDataSize);
+			const uint8* CompressedChunkPtr = nullptr;
+			uint32 CompressedChunkSize = 0;
 
-					if(CurCompressedChunkData == nullptr)
-					{
-						// We only need to worry about missing the stream chunk if we were seeking. Seeking might cause a bit of latency with chunk loading. That is expected.
-						if (!bSeekPending)
-						{
-							// If we did not get it then just bail, CurrentChunkIndex will not get incremented on the next callback so in effect another attempt will be made to fetch the chunk
-							// Since audio streaming depends on the general data streaming mechanism used by other parts of the engine and new data is prefectched on the game tick thread its possible a game hickup can cause this
-							UE_LOG(LogAudio, Verbose, TEXT("Missed Deadline chunk %d"), CurrentChunkIndex);
-						}
-
-						// zero out remaining data and bail
-						FMemory::Memset(OutData, 0, BufferSize);
-						return false;
-					}
+			GetChunkPtr(CompressedChunkPtr, CompressedChunkSize);
 
-					// Set the current buffer offset accounting for the header in the first chunk
+			if(CurrentUncompressedBlockSampleIndex >= UncompressedBlockSize / sizeof(uint16))
+			{
+				// If GetChunkPtr returns a nullptr, the buffer wasn't loaded in time.
+				if(CompressedChunkPtr == nullptr)
+				{
+					// We only need to worry about missing the stream chunk if we were seeking. Seeking might cause a bit of latency with chunk loading. That is expected.
 					if (!bSeekPending)
 					{
-						CurrentChunkBufferOffset = CurrentChunkIndex == 0 ? FirstChunkSampleDataOffset : 0;
+						// If we did not get it then just bail, CurrentChunkIndex will not get incremented on the next callback so in effect another attempt will be made to fetch the chunk
+						// Since audio streaming depends on the general data streaming mechanism used by other parts of the engine and new data is prefectched on the game tick thread its possible a game hickup can cause this
+						UE_LOG(LogAudio, Verbose, TEXT("Missed Deadline for chunk %d"), CurrentChunkIndex);
 					}
-					bSeekPending = false;
+
+					// zero out remaining data and bail
+					FMemory::Memset(OutData, 0, BufferSize);
+					return false;
 				}
 
 				// Decompress one block for each channel and store it in UncompressedBlockData
 				for(int32 ChannelItr = 0; ChannelItr < NumChannels; ++ChannelItr)
 				{
+					check(CompressedChunkPtr && (CurrentChunkBufferOffset + ChannelItr * CompressedBlockSize) <= CompressedChunkSize)
+
 					ADPCM::DecodeBlock(
-						CurCompressedChunkData + CurrentChunkBufferOffset + ChannelItr * CompressedBlockSize,
+						CompressedChunkPtr + CurrentChunkBufferOffset + ChannelItr * CompressedBlockSize,
 						CompressedBlockSize,
 						(int16*)(UncompressedBlockData + ChannelItr * UncompressedBlockSize));
 				}
@@ -483,6 +491,19 @@ bool FADPCMAudioInfo::StreamCompressedData(uint8* Destination, bool bLooping, ui
 				// Update some bookkeeping
 				CurrentUncompressedBlockSampleIndex = 0;
 				CurrentChunkBufferOffset += NumChannels * CompressedBlockSize;
+
+				if (CurrentChunkBufferOffset >= CompressedChunkSize)
+				{
+					// Move on to the next streamed chunk.
+					++CurrentChunkIndex;
+
+					// Set the current buffer offset accounting for the header in the first chunk
+					if (!bSeekPending)
+					{
+						CurrentChunkBufferOffset = CurrentChunkIndex == 0 ? FirstChunkSampleDataOffset : 0;
+					}
+					bSeekPending = false;
+				}
 			}
 
 			// Only copy over the number of samples we currently have available, we will loop around if needed
@@ -518,9 +539,9 @@ bool FADPCMAudioInfo::StreamCompressedData(uint8* Destination, bool bLooping, ui
 				ReachedEndOfSamples = true;
 				CurrentUncompressedBlockSampleIndex = 0;
 				CurrentChunkIndex = 0;
-				CurrentChunkBufferOffset = 0;
+				CurrentChunkBufferOffset = FirstChunkSampleDataOffset;
 				TotalSamplesStreamed = 0;
-				CurCompressedChunkData = nullptr;
+				CurrentCompressedChunkData = FAudioChunkHandle();
 				if(!bLooping)
 				{
 					// Set the remaining buffer to 0
@@ -534,37 +555,55 @@ bool FADPCMAudioInfo::StreamCompressedData(uint8* Destination, bool bLooping, ui
 	{
 		while(BufferSize > 0)
 		{
-			if(CurCompressedChunkData == nullptr || CurrentChunkBufferOffset >= CurrentChunkDataSize)
-			{
-				// Get another chunk of compressed data from the streaming engine
+			const uint8* CompressedChunkPtr = nullptr;
+			uint32 CompressedChunkSize = 0;
+
+			GetChunkPtr(CompressedChunkPtr, CompressedChunkSize);
+
 
-				// CurrentChunkIndex is used to keep track of the current chunk for loading/unloading by the streaming engine but chunk 0
-				// is preloaded so we don't want to increment this when getting chunk 0, only later chunks
-				// Also, if we failed to get a previous chunk because it was not ready we don't want to re-increment the CurrentChunkIndex
-				if(CurCompressedChunkData != nullptr)
+			if (CompressedChunkPtr == nullptr)
+			{
+				// Only report missing the stream chunk if we were seeking. Seeking
+				// may cause a bit of latency with chunk loading, which is expected.
+				if (!bSeekPending)
 				{
-					++CurrentChunkIndex;
+					// CurrentChunkIndex will not get incremented on the next callback, effectively causing another attempt to fetch the chunk.
+					// Since audio streaming depends on the general data streaming mechanism used by other parts of the engine and new data is
+					// prefetched on the game tick thread, this may be caused by a game hitch.
+					UE_LOG(LogAudio, Verbose, TEXT("Missed streaming ADPCM deadline chunk %d"), CurrentChunkIndex);
 				}
 
-				// Request the next chunk of data from the streaming engine
-				CurCompressedChunkData = IStreamingManager::Get().GetAudioStreamingManager().GetLoadedChunk(StreamingSoundWave, CurrentChunkIndex, &CurrentChunkDataSize);
+				FMemory::Memset(OutData, 0, BufferSize);
+				return false;
+			}
 
-				if(CurCompressedChunkData == nullptr)
-				{
-					// Only report missing the stream chunk if we were seeking. Seeking
-					// may cause a bit of latency with chunk loading, which is expected.
-					if (!bSeekPending)
-					{
-						// CurrentChunkIndex will not get incremented on the next callback, effectively causing another attempt to fetch the chunk.
-						// Since audio streaming depends on the general data streaming mechanism used by other parts of the engine and new data is
-						// prefetched on the game tick thread, this may be caused by a game hitch.
-						UE_LOG(LogAudio, Verbose, TEXT("Missed streaming ADPCM deadline chunk %d"), CurrentChunkIndex);
-					}
+			uint32 DecompressedSamplesToCopy = FMath::Min<uint32>(
+				(CompressedChunkSize - CurrentChunkBufferOffset) / ChannelSampleSize,
+				BufferSize / ChannelSampleSize);
 
-					FMemory::Memset(OutData, 0, BufferSize);
-					return false;
+			if (DecompressedSamplesToCopy > 0)
+			{
+				// Ensure we don't go over the number of samples left in the audio data
+				if (DecompressedSamplesToCopy > TotalSamplesPerChannel - TotalSamplesStreamed)
+				{
+					DecompressedSamplesToCopy = TotalSamplesPerChannel - TotalSamplesStreamed;
 				}
 
+				const uint32 SampleSize = DecompressedSamplesToCopy * ChannelSampleSize;
+
+				check(CompressedChunkPtr && CurrentChunkBufferOffset < CompressedChunkSize);
+				FMemory::Memcpy(OutData, CompressedChunkPtr + CurrentChunkBufferOffset, SampleSize);
+
+				OutData += DecompressedSamplesToCopy * NumChannels;
+				CurrentChunkBufferOffset += SampleSize;
+				BufferSize -= SampleSize;
+				TotalSamplesStreamed += DecompressedSamplesToCopy;
+			}
+
+			if (CurrentChunkBufferOffset >= CompressedChunkSize)
+			{
+				++CurrentChunkIndex;
+
 				// Set the current buffer offset accounting for the header in the first chunk
 				if (!bSeekPending)
 				{
@@ -573,25 +612,6 @@ bool FADPCMAudioInfo::StreamCompressedData(uint8* Destination, bool bLooping, ui
 				bSeekPending = false;
 			}
 
-			uint32 DecompressedSamplesToCopy = FMath::Min<uint32>(
-				(CurrentChunkDataSize - CurrentChunkBufferOffset) / ChannelSampleSize,
-				BufferSize / ChannelSampleSize);
-
-			check(DecompressedSamplesToCopy > 0);
-
-			// Ensure we don't go over the number of samples left in the audio data
-			if(DecompressedSamplesToCopy > TotalSamplesPerChannel - TotalSamplesStreamed)
-			{
-				DecompressedSamplesToCopy = TotalSamplesPerChannel - TotalSamplesStreamed;
-			}
-
-			const uint32 SampleSize = DecompressedSamplesToCopy * ChannelSampleSize;
-			FMemory::Memcpy(OutData, CurCompressedChunkData + CurrentChunkBufferOffset, SampleSize);
-
-			OutData += DecompressedSamplesToCopy * NumChannels;
-			CurrentChunkBufferOffset += SampleSize;
-			BufferSize -= SampleSize;
-			TotalSamplesStreamed += DecompressedSamplesToCopy;
 
 			// Check for the end of the audio samples and loop if needed
 			if(TotalSamplesStreamed >= TotalSamplesPerChannel)
@@ -600,7 +620,7 @@ bool FADPCMAudioInfo::StreamCompressedData(uint8* Destination, bool bLooping, ui
 				CurrentChunkIndex = 0;
 				CurrentChunkBufferOffset = 0;
 				TotalSamplesStreamed = 0;
-				CurCompressedChunkData = nullptr;
+				CurrentCompressedChunkData = FAudioChunkHandle();
 				if(!bLooping)
 				{
 					// Set the remaining buffer to 0
@@ -613,3 +633,21 @@ bool FADPCMAudioInfo::StreamCompressedData(uint8* Destination, bool bLooping, ui
 
 	return ReachedEndOfSamples;
 }
+
+void FADPCMAudioInfo::GetChunkPtr(const uint8* &CompressedChunkPtr, uint32 &CompressedChunkSize)
+{
+	// Retrieve the current chunk:
+	if (CurrentChunkIndex == 0)
+	{
+		TArrayView<const uint8> ZerothChunk = StreamingSoundWave->GetZerothChunk();
+		CompressedChunkPtr = ZerothChunk.GetData();
+		CompressedChunkSize = ZerothChunk.Num();
+	}
+	else
+	{
+		// Request the next chunk of data from the streaming engine
+		CurrentCompressedChunkData = IStreamingManager::Get().GetAudioStreamingManager().GetLoadedChunk(StreamingSoundWave, CurrentChunkIndex);
+		CompressedChunkPtr = CurrentCompressedChunkData.GetData();
+		CompressedChunkSize = CurrentCompressedChunkData.Num();
+	}
+}
diff --git a/Engine/Source/Runtime/Engine/Private/AudioCompressionSettingsUtils.cpp b/Engine/Source/Runtime/Engine/Private/AudioCompressionSettingsUtils.cpp
--- a/Engine/Source/Runtime/Engine/Private/AudioCompressionSettingsUtils.cpp
+++ b/Engine/Source/Runtime/Engine/Private/AudioCompressionSettingsUtils.cpp
@@ -4,6 +4,7 @@
 #include "AudioCompressionSettings.h"
 
 #define ENABLE_PLATFORM_COMPRESSION_OVERRIDES 1
+#define PLATFORM_SUPPORTS_COMPRESSION_OVERRIDES (PLATFORM_ANDROID && !PLATFORM_LUMIN) || PLATFORM_IOS || PLATFORM_SWITCH || PLATFORM_XBOXONE || PLATFORM_PS4 || PLATFORM_WINDOWS
 
 #if PLATFORM_ANDROID && !PLATFORM_LUMIN && ENABLE_PLATFORM_COMPRESSION_OVERRIDES
 #include "AndroidRuntimeSettings.h"
@@ -19,6 +20,14 @@
 
 #include "Misc/ConfigCacheIni.h"
 
+static float CookOverrideCachingIntervalCvar = 1.0f;
+FAutoConsoleVariableRef CVarCookOverrideCachingIntervalCVar(
+	TEXT("au.editor.CookOverrideCachingInterval"),
+	CookOverrideCachingIntervalCvar,
+	TEXT("This sets the max latency between when a cook override is changed in the project settings and when it is applied to new audio sources.\n")
+	TEXT("n: Time between caching intervals, in seconds."),
+	ECVF_Default);
+
 const FPlatformRuntimeAudioCompressionOverrides* FPlatformCompressionUtilities::GetRuntimeCompressionOverridesForCurrentPlatform()
 {
 #if PLATFORM_ANDROID && !PLATFORM_LUMIN && ENABLE_PLATFORM_COMPRESSION_OVERRIDES
@@ -56,6 +65,12 @@ void CacheCurrentPlatformAudioCookOverrides(FPlatformAudioCookOverrides& OutOver
 	const TCHAR* CategoryName = TEXT("/Script/IOSRuntimeSettings.IOSRuntimeSettings");
 #elif PLATFORM_SWITCH
 	const TCHAR* CategoryName = TEXT("/Script/SwitchRuntimeSettings.SwitchRuntimeSettings");
+#elif PLATFORM_WINDOWS
+	const TCHAR* CategoryName = TEXT("/Script/WindowsTargetPlatform.WindowsTargetSettings");
+#elif PLATFORM_PS4
+	const TCHAR* CategoryName = TEXT("/Script/PS4TargetPlatform.PS4TargetSettings");
+#elif PLATFORM_XBOXONE
+	const TCHAR* CategoryName = TEXT("/Script/XboxOneTargetPlatform.XboxOneTargetSettings");
 #else
 	const TCHAR* CategoryName = TEXT("");
 #endif
@@ -66,6 +81,17 @@ void CacheCurrentPlatformAudioCookOverrides(FPlatformAudioCookOverrides& OutOver
 
 	GConfig->GetFloat(CategoryName, TEXT("AutoStreamingThreshold"), OutOverrides.AutoStreamingThreshold, GEngineIni);
 
+	GConfig->GetBool(CategoryName, TEXT("bUseAudioStreamCaching"), OutOverrides.bUseStreamCaching, GEngineIni);
+
+	/** Memory Load On Demand Settings */
+	if (OutOverrides.bUseStreamCaching)
+	{
+		// Cache size:
+		int32 RetrievedCacheSize = 32 * 1024;
+		GConfig->GetInt(CategoryName, TEXT("CacheSizeKB"), RetrievedCacheSize, GEngineIni);
+		OutOverrides.StreamCachingSettings.CacheSizeKB = RetrievedCacheSize;
+	}
+
 	//Cache sample rate map.
 	OutOverrides.PlatformSampleRates.Reset();
 
@@ -95,17 +121,45 @@ void CacheCurrentPlatformAudioCookOverrides(FPlatformAudioCookOverrides& OutOver
 	OutOverrides.PlatformSampleRates.Add(ESoundwaveSampleRateSettings::Min, RetrievedSampleRate);
 }
 
-#if (PLATFORM_ANDROID && !PLATFORM_LUMIN) || PLATFORM_IOS || PLATFORM_SWITCH
+#if PLATFORM_SUPPORTS_COMPRESSION_OVERRIDES
 static FPlatformAudioCookOverrides OutOverrides = FPlatformAudioCookOverrides();
 #endif
 
 
-const FPlatformAudioCookOverrides* FPlatformCompressionUtilities::GetCookOverridesForCurrentPlatform()
+
+#if PLATFORM_SUPPORTS_COMPRESSION_OVERRIDES
+static FCriticalSection CookOverridesCriticalSection;
+
+void FPlatformCompressionUtilities::RecacheCookOverrides()
+{
+	FScopeLock ScopeLock(&CookOverridesCriticalSection);
+	CacheCurrentPlatformAudioCookOverrides(OutOverrides);
+}
+
+#endif // WITH_EDITOR && PLATFORM_SUPPORTS_COMPRESSION_OVERRIDES
+
+const FPlatformAudioCookOverrides* FPlatformCompressionUtilities::GetCookOverridesForCurrentPlatform(bool bForceRecache)
 {
-#if (PLATFORM_ANDROID && !PLATFORM_LUMIN) || PLATFORM_IOS || PLATFORM_SWITCH
+#if PLATFORM_SUPPORTS_COMPRESSION_OVERRIDES
 	static bool bCachedCookOverrides = false;
+
+#if !WITH_EDITOR
+	// In non-editor situations, we only need to cache the cook overrides once.
 	if (!bCachedCookOverrides)
 	{
+#else
+	// In editor situations, the settings can change at any time, so we need to retrieve them.
+	FScopeLock ScopeLock(&CookOverridesCriticalSection);
+	
+	static double LastCacheTime = 0.0;
+	double CurrentTime = FPlatformTime::Seconds();
+	double TimeSinceLastCache = CurrentTime - LastCacheTime;
+
+	if (bForceRecache || TimeSinceLastCache > CookOverrideCachingIntervalCvar)
+	{
+		LastCacheTime = CurrentTime;
+#endif // WITH_EDITOR
+	
 		CacheCurrentPlatformAudioCookOverrides(OutOverrides);
 		bCachedCookOverrides = true;
 	}
@@ -115,6 +169,69 @@ const FPlatformAudioCookOverrides* FPlatformCompressionUtilities::GetCookOverrid
 #endif
 }
 
+bool FPlatformCompressionUtilities::IsCurrentPlatformUsingStreamCaching()
+{
+	const FPlatformAudioCookOverrides* Settings = GetCookOverridesForCurrentPlatform();
+	return Settings && Settings->bUseStreamCaching;
+}
+
+const FAudioStreamCachingSettings& FPlatformCompressionUtilities::GetStreamCachingSettingsForCurrentPlatform()
+{
+	const FPlatformAudioCookOverrides* Settings = GetCookOverridesForCurrentPlatform();
+	checkf(Settings, TEXT("Please only use this function if FPlatformCompressionUtilities::IsCurrentPlatformUsingLoadOnDemand() returns true."));
+	return Settings->StreamCachingSettings;
+}
+
+FCachedAudioStreamingManagerParams FPlatformCompressionUtilities::BuildCachedStreamingManagerParams()
+{
+	const FAudioStreamCachingSettings& CacheSettings = GetStreamCachingSettingsForCurrentPlatform();
+	int32 ChunkSize = GetChunkSizeForCookOverrides(GetCookOverridesForCurrentPlatform());
+
+	// Our number of elements in the cache is CacheSize in kilobytes divided by our chunk size in kilobytes:
+	int32 NumElements = CacheSettings.CacheSizeKB / (ChunkSize / 1024);
+
+	FCachedAudioStreamingManagerParams Params;
+	FCachedAudioStreamingManagerParams::FCacheDimensions CacheDimensions;
+
+	// Primary cache defined here:
+	CacheDimensions.MaxChunkSize = ChunkSize;
+	CacheDimensions.NumElements = NumElements;
+	Params.Caches.Add(CacheDimensions);
+
+	// TODO: When settings are added to support multiple sub-caches, add it here.
+
+	return Params;
+}
+
+uint32 FPlatformCompressionUtilities::GetChunkSizeForCookOverrides(const FPlatformAudioCookOverrides* InCompressionOverrides)
+{
+	check(InCompressionOverrides);
+
+	// TODO: We can fine tune this to platform-specific voice counts, but in the meantime we target 32 voices as an average-case.
+	// If the game runs with higher than 32 voices, that means we will potentially have a larger cache than what was set in the target settings.
+	// In that case we log a warning on application launch.
+	const int32 MinimumNumChunks = 32;
+	const int32 DefaultMaxChunkSizeKB = 256;
+	
+	int32 CacheSizeKB = InCompressionOverrides->StreamCachingSettings.CacheSizeKB;
+	
+	if (CacheSizeKB == 0)
+	{
+		CacheSizeKB = FAudioStreamCachingSettings::DefaultCacheSize;
+	}
+
+	// If we won't have a large enough cache size to fit enough chunks to play 32 different sources at once, 
+	// we truncate the chunk size to fit at least that many 
+	if (CacheSizeKB / DefaultMaxChunkSizeKB < MinimumNumChunks)
+	{
+		return (CacheSizeKB / MinimumNumChunks) * 1024;
+	}
+	else
+	{
+		return DefaultMaxChunkSizeKB * 1024;
+	}
+}
+
 float FPlatformCompressionUtilities::GetCompressionDurationForCurrentPlatform()
 {
 	float Threshold = -1.0f;
@@ -175,4 +292,4 @@ int32 FPlatformCompressionUtilities::GetQualityIndexOverrideForCurrentPlatform()
 	{
 		return INDEX_NONE;
 	}
-}
+}
\ No newline at end of file
diff --git a/Engine/Source/Runtime/Engine/Private/AudioDecompress.cpp b/Engine/Source/Runtime/Engine/Private/AudioDecompress.cpp
--- a/Engine/Source/Runtime/Engine/Private/AudioDecompress.cpp
+++ b/Engine/Source/Runtime/Engine/Private/AudioDecompress.cpp
@@ -8,12 +8,13 @@
 #include "HAL/LowLevelMemTracker.h"
 #include "ADPCMAudioInfo.h"
 
+PRAGMA_DISABLE_OPTIMIZATION
+
 IStreamedCompressedInfo::IStreamedCompressedInfo()
 	: bIsStreaming(false)
-	,SrcBufferData(nullptr)
-	, SrcBufferDataSize(0)
 	, SrcBufferOffset(0)
 	, AudioDataOffset(0)
+	, AudioDataOffsetChunkIndex(0)
 	, SampleRate(0)
 	, TrueSampleCount(0)
 	, CurrentSampleCount(0)
@@ -26,15 +27,27 @@ IStreamedCompressedInfo::IStreamedCompressedInfo()
 	, CurrentChunkIndex(0)
 	, bPrintChunkFailMessage(true)
 	, SrcBufferPadding(0)
+	, NonStreamingAudio(nullptr)
+	, NonStreamingAudioSize(0)
 {
 }
 
 uint32 IStreamedCompressedInfo::Read(void *OutBuffer, uint32 DataSize)
 {
-	uint32 BytesToRead = FMath::Min(DataSize, SrcBufferDataSize - SrcBufferOffset);
+	const uint8* SrcBufferPtr = nullptr;
+	uint32 SrcBufferSize = 0;
+
+	GetChunkPtr(SrcBufferPtr, SrcBufferSize);
+
+	uint32 BytesToRead = 0;
+	if (SrcBufferOffset < SrcBufferSize)
+	{
+		BytesToRead = FMath::Min(DataSize, SrcBufferSize - SrcBufferOffset);
+	}
+
 	if (BytesToRead > 0)
 	{
-		FMemory::Memcpy(OutBuffer, SrcBufferData + SrcBufferOffset, BytesToRead);
+		FMemory::Memcpy(OutBuffer, SrcBufferPtr + SrcBufferOffset, BytesToRead);
 		SrcBufferOffset += BytesToRead;
 	}
 	return BytesToRead;
@@ -42,19 +55,28 @@ uint32 IStreamedCompressedInfo::Read(void *OutBuffer, uint32 DataSize)
 
 bool IStreamedCompressedInfo::ReadCompressedInfo(const uint8* InSrcBufferData, uint32 InSrcBufferDataSize, FSoundQualityInfo* QualityInfo)
 {
-	check(!SrcBufferData);
+	check(!SrcBufferHandle.IsValid());
 
 	SCOPE_CYCLE_COUNTER(STAT_AudioStreamedDecompressTime);
 
+	// ReadCompressedInfo is called from StreamCompressedInfo(), in which case bIsStreaming will be true.
+	if (!bIsStreaming)
+	{
+		NonStreamingAudio = InSrcBufferData;
+		NonStreamingAudioSize = InSrcBufferDataSize;
+	}
+
+	const uint8* SrcBufferPtr = nullptr;
+	uint32 SrcBufferSize = 0;
+
+	GetChunkPtr(SrcBufferPtr, SrcBufferSize);
+
 	// Parse the format header, this is done different for each format
-	if (!ParseHeader(InSrcBufferData, InSrcBufferDataSize, QualityInfo))
+	if (!ParseHeader(SrcBufferPtr, SrcBufferSize, QualityInfo))
 	{
 		return false;
 	}
 
-	// After parsing the header, the SrcBufferData should be none-null
-	check(SrcBufferData != nullptr);
-
 	// Sample Stride is 
 	SampleStride = NumChannels * sizeof(int16);
 
@@ -75,10 +97,15 @@ bool IStreamedCompressedInfo::ReadCompressedData(uint8* Destination, bool bLoopi
 	bool bFinished = false;
 	uint32 TotalBytesDecoded = 0;
 
+	const uint8* SrcBufferPtr = nullptr;
+	uint32 SrcBufferSize = 0;
+
+	GetChunkPtr(SrcBufferPtr, SrcBufferSize);
+
 	while (TotalBytesDecoded < BufferSize)
 	{
-		const uint8* EncodedSrcPtr = SrcBufferData + SrcBufferOffset;
-		const uint32 RemainingEncodedSrcSize = SrcBufferDataSize - SrcBufferOffset;
+		const uint8* EncodedSrcPtr = SrcBufferPtr + SrcBufferOffset;
+		const uint32 RemainingEncodedSrcSize = SrcBufferSize - SrcBufferOffset;
 
 		const FDecodeResult DecodeResult = Decode(EncodedSrcPtr, RemainingEncodedSrcSize, Destination, BufferSize - TotalBytesDecoded);
 		if (!DecodeResult.NumPcmBytesProduced)
@@ -147,13 +174,36 @@ bool IStreamedCompressedInfo::StreamCompressedInfoInternal(USoundWave* Wave, str
 
 	// Get the first chunk of audio data (should always be loaded)
 	CurrentChunkIndex = 0;
-	const uint8* FirstChunk = IStreamingManager::Get().GetAudioStreamingManager().GetLoadedChunk(StreamingSoundWave, CurrentChunkIndex);
-
 	bIsStreaming = true;
-	if (FirstChunk)
+
+	const uint8* SrcBufferPtr = nullptr;
+	uint32 SrcBufferSize = 0;
+
+	GetChunkPtr(SrcBufferPtr, SrcBufferSize);
+
+	if (ReadCompressedInfo(SrcBufferPtr, SrcBufferSize, QualityInfo))
 	{
-		return ReadCompressedInfo(FirstChunk, Wave->RunningPlatformData->Chunks[0].AudioDataSize, QualityInfo);
+		while (AudioDataOffset >= SrcBufferSize)
+		{
+			// scrub through chunks until we find the beginning of the audio data.
+			CurrentChunkIndex++;
+			
+			AudioDataOffset = AudioDataOffset - SrcBufferSize;
+			SrcBufferSize = Wave->GetSizeOfChunk(CurrentChunkIndex);
+			// Setting SrcBufferData to null will ensure that we try to load this chunk on the next iteration.
+			SrcBufferHandle = FAudioChunkHandle();
+		}
+
+		AudioDataOffsetChunkIndex = CurrentChunkIndex;
+		SrcBufferOffset = AudioDataOffset;
+		
+		return true;
 	}
+	else
+	{
+		UE_LOG(LogAudio, Log, TEXT("Read Compressed Info failed for SoundWave'%s'"), *StreamingSoundWave->GetName());
+	}
+	
 
 	return false;
 }
@@ -165,32 +215,28 @@ bool IStreamedCompressedInfo::StreamCompressedData(uint8* Destination, bool bLoo
 
 	SCOPE_CYCLE_COUNTER(STAT_AudioStreamedDecompressTime);
 
-	UE_LOG(LogAudio, Log, TEXT("Streaming compressed data from SoundWave'%s' - Chunk %d, Offset %d"), *StreamingSoundWave->GetName(), CurrentChunkIndex, SrcBufferOffset);
+	UE_LOG(LogAudio, Verbose, TEXT("Streaming compressed data from SoundWave'%s' - Chunk %d, Offset %d"), *StreamingSoundWave->GetName(), CurrentChunkIndex, SrcBufferOffset);
 
 	// Write out any PCM data that was decoded during the last request
 	uint32 RawPCMOffset = WriteFromDecodedPCM(Destination, BufferSize);
 
+	const uint8* ChunkPtr = nullptr;
+	uint32 ChunkSize = 0;
+
+	GetChunkPtr(ChunkPtr, ChunkSize);
+
 	// If next chunk wasn't loaded when last one finished reading, try to get it again now
-	if (SrcBufferData == NULL)
+	if (!ChunkPtr)
 	{
-		SrcBufferData = IStreamingManager::Get().GetAudioStreamingManager().GetLoadedChunk(StreamingSoundWave, CurrentChunkIndex);
-		if (SrcBufferData)
+		// Still not loaded, zero remainder of current buffer
+		if (bPrintChunkFailMessage)
 		{
-			bPrintChunkFailMessage = true;
-			SrcBufferDataSize = StreamingSoundWave->RunningPlatformData->Chunks[CurrentChunkIndex].AudioDataSize;
-			SrcBufferOffset = CurrentChunkIndex == 0 ? AudioDataOffset : 0;
-		}
-		else
-		{
-			// Still not loaded, zero remainder of current buffer
-			if (bPrintChunkFailMessage)
-			{
-				UE_LOG(LogAudio, Verbose, TEXT("Chunk %d not loaded from streaming manager for SoundWave '%s'. Likely due to stall on game thread."), CurrentChunkIndex, *StreamingSoundWave->GetName());
-				bPrintChunkFailMessage = false;
-			}
-			ZeroBuffer(Destination + RawPCMOffset, BufferSize - RawPCMOffset);
-			return false;
+			UE_LOG(LogAudio, Verbose, TEXT("Chunk %d not loaded from streaming manager for SoundWave '%s'. Likely due to stall on game thread or Disc IO issue."), CurrentChunkIndex, *StreamingSoundWave->GetName());
+			bPrintChunkFailMessage = false;
 		}
+
+		ZeroBuffer(Destination + RawPCMOffset, BufferSize - RawPCMOffset);
+		return false;
 	}
 
 	bool bLooped = false;
@@ -204,6 +250,19 @@ bool IStreamedCompressedInfo::StreamCompressedData(uint8* Destination, bool bLoo
 
 	while (RawPCMOffset < BufferSize)
 	{
+		GetChunkPtr(ChunkPtr, ChunkSize);
+
+		if (ChunkPtr)
+		{
+			UE_LOG(LogAudio, Verbose, TEXT("Incremented current chunk from SoundWave'%s' - Chunk %d, Offset %d"), *StreamingSoundWave->GetName(), CurrentChunkIndex, SrcBufferOffset);
+		}
+		else
+		{
+			// Zero the rest of the buffer and break when we fail to load the chunk.
+			RawPCMOffset += ZeroBuffer(Destination + RawPCMOffset, BufferSize - RawPCMOffset);
+			break;
+		}
+
 		// Get the platform-dependent size of the current "frame" of encoded audio (note: frame is used here differently than audio frame/sample)
 		uint16 FrameSize = GetFrameSize();
 
@@ -223,10 +282,10 @@ bool IStreamedCompressedInfo::StreamCompressedData(uint8* Destination, bool bLoo
 			RawPCMOffset += WriteFromDecodedPCM(Destination + RawPCMOffset, BufferSize - RawPCMOffset);
 
 			// Have we reached the end of buffer?
-			if (SrcBufferOffset >= SrcBufferDataSize - SrcBufferPadding)
+			if (SrcBufferOffset >= (ChunkSize - SrcBufferPadding))
 			{
 				// Special case for the last chunk of audio
-				if (CurrentChunkIndex == StreamingSoundWave->RunningPlatformData->NumChunks - 1)
+				if (CurrentChunkIndex == StreamingSoundWave->GetNumChunks() - 1)
 				{
 					// check whether all decoded PCM was written
 					if (LastPCMByteSize == 0)
@@ -240,7 +299,7 @@ bool IStreamedCompressedInfo::StreamCompressedData(uint8* Destination, bool bLoo
 
 					if (bLooping)
 					{
-						CurrentChunkIndex = 0;
+						CurrentChunkIndex = AudioDataOffsetChunkIndex;
 						SrcBufferOffset = AudioDataOffset;
 						CurrentSampleCount = 0;
 
@@ -257,18 +316,6 @@ bool IStreamedCompressedInfo::StreamCompressedData(uint8* Destination, bool bLoo
 					CurrentChunkIndex++;
 					SrcBufferOffset = 0;
 				}
-
-				SrcBufferData = IStreamingManager::Get().GetAudioStreamingManager().GetLoadedChunk(StreamingSoundWave, CurrentChunkIndex);
-				if (SrcBufferData)
-				{
-					UE_LOG(LogAudio, Log, TEXT("Incremented current chunk from SoundWave'%s' - Chunk %d, Offset %d"), *StreamingSoundWave->GetName(), CurrentChunkIndex, SrcBufferOffset);
-					SrcBufferDataSize = StreamingSoundWave->RunningPlatformData->Chunks[CurrentChunkIndex].AudioDataSize;
-				}
-				else
-				{
-					SrcBufferDataSize = 0;
-					RawPCMOffset += ZeroBuffer(Destination + RawPCMOffset, BufferSize - RawPCMOffset);
-				}
 			}
 		}
 	}
@@ -278,13 +325,32 @@ bool IStreamedCompressedInfo::StreamCompressedData(uint8* Destination, bool bLoo
 
 int32 IStreamedCompressedInfo::DecompressToPCMBuffer(uint16 FrameSize)
 {
-	if (SrcBufferOffset + FrameSize > SrcBufferDataSize)
+
+	const uint8* ChunkPtr = nullptr;
+	uint32 ChunkSize = 0;
+
+	GetChunkPtr(ChunkPtr, ChunkSize);
+
+	if (ChunkSize == 0)
+	{
+		// We hit an underrun.
+		return 0;
+	}
+	else if (SrcBufferOffset == ChunkSize)
+	{
+		// We reached the end of the buffer.
+		return 0;
+	}
+	else if (SrcBufferOffset + FrameSize > ChunkSize)
 	{
 		// if frame size is too large, something has gone wrong
+		check(false);
 		return -1;
 	}
 
-	const uint8* SrcPtr = SrcBufferData + SrcBufferOffset;
+	check(ChunkPtr && SrcBufferOffset < ChunkSize);
+
+	const uint8* SrcPtr = ChunkPtr + SrcBufferOffset;
 	SrcBufferOffset += FrameSize;
 	LastPCMOffset = 0;
 	
@@ -351,6 +417,27 @@ uint32 IStreamedCompressedInfo::ZeroBuffer(uint8* Destination, uint32 BufferSize
 }
 
 
+void IStreamedCompressedInfo::GetChunkPtr(const uint8*& OutPtr, uint32& OutSize)
+{
+	if (!bIsStreaming)
+	{
+		OutPtr = NonStreamingAudio;
+		OutSize = NonStreamingAudioSize;
+	}
+	else if (CurrentChunkIndex == 0)
+	{
+		TArrayView<const uint8> ZerothChunk = StreamingSoundWave->GetZerothChunk();
+		OutPtr = ZerothChunk.GetData();
+		OutSize = ZerothChunk.Num();
+	}
+	else
+	{
+		SrcBufferHandle = IStreamingManager::Get().GetAudioStreamingManager().GetLoadedChunk(StreamingSoundWave, CurrentChunkIndex);
+		OutPtr = SrcBufferHandle.GetData();
+		OutSize = SrcBufferHandle.Num();
+	}
+}
+
 //////////////////////////////////////////////////////////////////////////
 // Copied from IOS - probably want to split and share
 //////////////////////////////////////////////////////////////////////////
@@ -718,5 +805,5 @@ bool ShouldUseBackgroundPoolFor_FAsyncRealtimeAudioTask()
 	return !!CVarShouldUseBackgroundPoolFor_FAsyncRealtimeAudioTask.GetValueOnAnyThread();
 }
 
-
+PRAGMA_ENABLE_OPTIMIZATION
 // end
diff --git a/Engine/Source/Runtime/Engine/Private/AudioDerivedData.cpp b/Engine/Source/Runtime/Engine/Private/AudioDerivedData.cpp
--- a/Engine/Source/Runtime/Engine/Private/AudioDerivedData.cpp
+++ b/Engine/Source/Runtime/Engine/Private/AudioDerivedData.cpp
@@ -17,6 +17,7 @@
 #include "AudioResampler.h"
 #include "AudioCompressionSettingsUtils.h"
 #include "Sound/SoundSourceBus.h"
+#include "Sound/SoundWave.h"
 
 DEFINE_LOG_CATEGORY_STATIC(LogAudioDerivedData, Log, All);
 
@@ -43,7 +44,7 @@ Derived data key generation.
 
 // If you want to bump this version, generate a new guid using
 // VS->Tools->Create GUID and paste it here. https://www.guidgen.com works too.
-#define STREAMEDAUDIO_DERIVEDDATA_VER		TEXT("31F97D9AF03C476B943F885DAB70E772")
+#define STREAMEDAUDIO_DERIVEDDATA_VER		TEXT("BC6E92FBBD314E3B9B9EC6778749EA7F")
 
 /**
  * Computes the derived data key suffix for a SoundWave's Streamed Audio.
@@ -178,7 +179,7 @@ static FName GetWaveFormatForRunningPlatform(USoundWave& SoundWave)
  * After this returns, all bulk data from streaming chunks will be sent separately to the DDC and the BulkData for those chunks removed.
  * @param DerivedData - The data to store in the DDC.
  * @param DerivedDataKeySuffix - The key suffix at which to store derived data.
- * @return number of bytes put to the DDC (total, including all chunnks)
+ * @return number of bytes put to the DDC (total, including all chunks)
  */
 static uint32 PutDerivedDataInCache(
 	FStreamedAudioPlatformData* DerivedData,
@@ -303,15 +304,77 @@ class FStreamedAudioCacheDerivedDataWorker : public FNonAbandonableTask
 				CompressedData->GetCopy(&BufferData, false);
 				TArray<TArray<uint8>> ChunkBuffers;
 
-				// Set the ideal chunk size to be 256k to optimize for data reads on console
-				static const int32 MaxChunkSize = 256 * 1024;
+				// Set the ideal chunk size to be 256k to optimize for data reads on console.
+				int32 MaxChunkSize = 256 * 1024;
+				
+				// By default, the first chunk's max size is the same as the other chunks.
+				int32 FirstChunkSize = MaxChunkSize;
+
+				const int32 MinimumChunkSize = AudioFormat->GetMinimumSizeForInitialChunk(AudioFormatName, CompressedBuffer);
+				const bool bUseStreamCaching = CompressionOverrides && CompressionOverrides->bUseStreamCaching;
 
-				if (AudioFormat->SplitDataForStreaming(CompressedBuffer, ChunkBuffers, MaxChunkSize))
+				// If the initial chunk  for this sound wave was overridden, use that:
+				if (SoundWave.InitialChunkSize > 0)
+				{
+					FirstChunkSize = FMath::Max(MinimumChunkSize, SoundWave.InitialChunkSize);
+				}
+				else if (bUseStreamCaching)
 				{
-					for (int32 ChunkIndex = 0; ChunkIndex < ChunkBuffers.Num(); ++ChunkIndex)
+					checkf(MinimumChunkSize != 0, TEXT("To use Load On Demand, please override GetMinimumSizeForInitialChunk"));
+
+					// Otherwise if we're using Load On Demand, the first chunk should be as small as possible:
+					FirstChunkSize = MinimumChunkSize;
+				}
+
+				if (bUseStreamCaching)
+				{
+					// Use the chunk size for this duration:
+					MaxChunkSize = FPlatformCompressionUtilities::GetChunkSizeForCookOverrides(CompressionOverrides);
+					UE_LOG(LogAudio, Display, TEXT("Chunk size for %s: %d"), *SoundWave.GetFullName(), MaxChunkSize);
+				}
+
+				check(FirstChunkSize != 0 && MaxChunkSize != 0);
+				check(FirstChunkSize <= (256 * 1024) && MaxChunkSize <= (256 * 1024));
+
+				if (AudioFormat->SplitDataForStreaming(CompressedBuffer, ChunkBuffers, FirstChunkSize, MaxChunkSize))
+				{
+					if (ChunkBuffers.Num() > 32)
+					{
+						UE_LOG(LogAudio, Display, TEXT("Sound Wave %s is very large, requiring %d chunks."), *SoundWave.GetFullName(), ChunkBuffers.Num());
+					}
+
+					if (ChunkBuffers.Num() > 0)
+					{
+						// The zeroth chunk should not be zero-padded.
+						const int32 AudioDataSize = ChunkBuffers[0].Num();
+
+						//FStreamedAudioChunk* NewChunk = new(DerivedData->Chunks) FStreamedAudioChunk();
+						int32 ChunkIndex = DerivedData->Chunks.Add(new FStreamedAudioChunk());
+						FStreamedAudioChunk* NewChunk = &(DerivedData->Chunks[ChunkIndex]);
+						// Store both the audio data size and the data size so decoders will know what portion of the bulk data is real audio
+						NewChunk->AudioDataSize = AudioDataSize;
+						NewChunk->DataSize = AudioDataSize;
+
+#if WITH_EDITORONLY_DATA
+						if (NewChunk->BulkData.IsLocked())
+						{
+							UE_LOG(LogAudioDerivedData, Warning, TEXT("While building split chunk for streaming: Raw PCM data already being written to. Chunk Index: 0 SoundWave: %s "), *SoundWave.GetFullName());
+						}
+#endif
+
+						NewChunk->BulkData.Lock(LOCK_READ_WRITE);
+						void* NewChunkData = NewChunk->BulkData.Realloc(NewChunk->DataSize);
+						FMemory::Memcpy(NewChunkData, ChunkBuffers[0].GetData(), AudioDataSize);
+						NewChunk->BulkData.Unlock();
+					}
+
+					// Zero-pad the rest of the chunks here:
+					for (int32 ChunkIndex = 1; ChunkIndex < ChunkBuffers.Num(); ++ChunkIndex)
 					{
 						// Zero pad the reallocation if the chunk isn't precisely the max chunk size to keep the reads aligned to MaxChunkSize
 						const int32 AudioDataSize = ChunkBuffers[ChunkIndex].Num();
+						check(AudioDataSize != 0 && AudioDataSize <= MaxChunkSize);
+						int32 AlignedChunkSize = FMath::RoundUpToPowerOfTwo(MaxChunkSize);
 						const int32 ZeroPadBytes = FMath::Max(MaxChunkSize - AudioDataSize, 0);
 
 						FStreamedAudioChunk* NewChunk = new FStreamedAudioChunk();
@@ -365,6 +428,8 @@ class FStreamedAudioCacheDerivedDataWorker : public FNonAbandonableTask
 				// Should ideally avoid this redundant work, but it only happens when we actually have
 				// to build the texture, which should only ever be once.
 				this->BytesCached = PutDerivedDataInCache(DerivedData, KeySuffix);
+
+				check(this->BytesCached != 0);
 			}
 			else
 			{
@@ -428,7 +493,7 @@ public:
 			{
 				for (int32 Index = 0; Index < DerivedData->NumChunks; ++Index)
 				{
-					if (!DerivedData->TryLoadChunk(Index, NULL))
+					if (!DerivedData->GetChunkFromDDC(Index, NULL))
 					{
 						bSucceeded = false;
 						break;
@@ -627,8 +692,9 @@ bool FStreamedAudioPlatformData::TryInlineChunkData()
 #endif //WITH_EDITORONLY_DATA
 
 FStreamedAudioPlatformData::FStreamedAudioPlatformData()
+	: NumChunks(0)
 #if WITH_EDITORONLY_DATA
-: AsyncTask(NULL)
+	, AsyncTask(NULL)
 #endif // #if WITH_EDITORONLY_DATA
 {
 }
@@ -645,75 +711,102 @@ FStreamedAudioPlatformData::~FStreamedAudioPlatformData()
 #endif
 }
 
-bool FStreamedAudioPlatformData::TryLoadChunk(int32 ChunkIndex, uint8** OutChunkData, bool bMakeSureChunkIsLoaded /* = false */)
+int32 FStreamedAudioPlatformData::DeserializeChunkFromDDC(TArray<uint8> TempData, FStreamedAudioChunk &Chunk, int32 ChunkIndex, uint8** &OutChunkData)
+{
+	int32 ChunkSize = 0;
+	FMemoryReader Ar(TempData, /*bIsPersistent=*/ true);
+	int32 AudioDataSize = 0;
+	Ar << ChunkSize;
+	Ar << AudioDataSize;
+
+#if WITH_EDITORONLY_DATA
+	ensureAlwaysMsgf((ChunkSize == Chunk.DataSize && AudioDataSize == Chunk.AudioDataSize),
+		TEXT("Chunk %d of %s SoundWave has invalid data in the DDC. Got %d bytes, expected %d. Audio Data was %d bytes but we expected %d bytes. Key=%s"),
+		ChunkIndex,
+		*AudioFormat.ToString(),
+		ChunkSize,
+		Chunk.DataSize,
+		AudioDataSize,
+		Chunk.AudioDataSize,
+		*Chunk.DerivedDataKey
+	);
+#endif
+
+	if (OutChunkData)
+	{
+		if (*OutChunkData == NULL)
+		{
+			*OutChunkData = static_cast<uint8*>(FMemory::Malloc(ChunkSize));
+		}
+		Ar.Serialize(*OutChunkData, ChunkSize);
+	}
+
+	if (FPlatformCompressionUtilities::IsCurrentPlatformUsingStreamCaching())
+	{
+		return AudioDataSize;
+	}
+	else
+	{
+		return ChunkSize;
+	}
+}
+
+int32 FStreamedAudioPlatformData::GetChunkFromDDC(int32 ChunkIndex, uint8** OutChunkData, bool bMakeSureChunkIsLoaded /* = false */)
 {
 	// if bMakeSureChunkIsLoaded is true, we don't actually know the size of the chunk's bulk data,
 	// so it will need to be allocated in GetCopy.
 	check(!bMakeSureChunkIsLoaded || (OutChunkData && (*OutChunkData == nullptr)));
 
 	bool bCachedChunk = false;
+	check(ChunkIndex < Chunks.Num());
 	FStreamedAudioChunk& Chunk = Chunks[ChunkIndex];
 
+	int32 ChunkDataSize = 0;
+
 #if WITH_EDITORONLY_DATA
+	TArray<uint8> TempData;
+
 	// Begin async DDC retrieval
 	FDerivedDataCacheInterface& DDC = GetDerivedDataCacheRef();
 	uint32 AsyncHandle = 0;
 	if (!Chunk.DerivedDataKey.IsEmpty())
 	{
-		AsyncHandle = DDC.GetAsynchronous(*Chunk.DerivedDataKey);
+		if (bMakeSureChunkIsLoaded)
+		{
+			if (DDC.GetSynchronous(*Chunk.DerivedDataKey, TempData))
+			{
+				ChunkDataSize = DeserializeChunkFromDDC(TempData, Chunk, ChunkIndex, OutChunkData);
+			}
+		}
+		else
+		{
+			AsyncHandle = DDC.GetAsynchronous(*Chunk.DerivedDataKey);
+		}
 	}
-#endif // #if WITH_EDITORONLY_DATA
 
+	// Wait for async DDC to complete
+	// TODO: Not necessary since bMakeSureChunkIsLoaded was introduced. 
+	// Also, if you do this- why not just call DDC.GetSynchronous?
+// 	if (Chunk.DerivedDataKey.IsEmpty() == false)
+// 	{
+// 		DDC.WaitAsynchronousCompletion(AsyncHandle);
+// 		if (DDC.GetAsynchronousResults(AsyncHandle, TempData))
+// 		{
+// 			ChunkDataSize = DeserializeChunkFromDDC(TempData, Chunk, ChunkIndex, OutChunkData);
+// 		}
+// 	}
+#else // #if WITH_EDITORONLY_DATA
 	// Load chunk from bulk data if available. If the chunk is not loaded, GetCopy will load it synchronously.
 	if (Chunk.BulkData.IsBulkDataLoaded() || bMakeSureChunkIsLoaded)
 	{
 		if (OutChunkData)
 		{
+			ChunkDataSize = Chunk.BulkData.GetBulkDataSize();
 			Chunk.BulkData.GetCopy((void**)OutChunkData);
 		}
-		bCachedChunk = true;
-	}
-
-
-#if WITH_EDITORONLY_DATA
-	// Wait for async DDC to complete
-	if (Chunk.DerivedDataKey.IsEmpty() == false)
-	{
-		TArray<uint8> TempData;
-		DDC.WaitAsynchronousCompletion(AsyncHandle);
-		if (DDC.GetAsynchronousResults(AsyncHandle, TempData))
-		{
-			int32 ChunkSize = 0;
-			FMemoryReader Ar(TempData, /*bIsPersistent=*/ true);
-			Ar << ChunkSize;
-
-			if (ChunkSize != Chunk.DataSize)
-			{
-				UE_LOG(LogAudio, Warning,
-					TEXT("Chunk %d of %s SoundWave has invalid data in the DDC. Got %d bytes, expected %d. Key=%s"),
-					ChunkIndex,
-					*AudioFormat.ToString(),
-					ChunkSize,
-					Chunk.DataSize,
-					*Chunk.DerivedDataKey
-					);
-			}
-
-			bCachedChunk = true;
-
-			if (OutChunkData)
-			{
-				if (*OutChunkData == NULL)
-				{
-					*OutChunkData = static_cast<uint8*>(FMemory::Malloc(ChunkSize));
-				}
-				Ar.Serialize(*OutChunkData, ChunkSize);
-			}
-		}
 	}
 #endif // #if WITH_EDITORONLY_DATA
-
-	return bCachedChunk;
+	return ChunkDataSize;
 }
 
 #if WITH_EDITORONLY_DATA
@@ -753,10 +846,11 @@ void FStreamedAudioPlatformData::Serialize(FArchive& Ar, USoundWave* Owner)
 			Chunks.Add(new FStreamedAudioChunk());
 		}
 	}
-	for (int32 ChunkIndex = 0; ChunkIndex < NumChunks; ++ChunkIndex)
+	for (int32 ChunkIndex = 0; ChunkIndex < Chunks.Num(); ++ChunkIndex)
 	{
 		Chunks[ChunkIndex].Serialize(Ar, Owner, ChunkIndex);
 	}
+	
 
 #if WITH_EDITORONLY_DATA
 	if (Owner)
diff --git a/Engine/Source/Runtime/Engine/Private/AudioStreaming.cpp b/Engine/Source/Runtime/Engine/Private/AudioStreaming.cpp
--- a/Engine/Source/Runtime/Engine/Private/AudioStreaming.cpp
+++ b/Engine/Source/Runtime/Engine/Private/AudioStreaming.cpp
@@ -17,6 +17,7 @@ AudioStreaming.cpp: Implementation of audio streaming classes.
 #include "HAL/IConsoleManager.h"
 #include "HAL/LowLevelMemTracker.h"
 #include "AudioDecompress.h"
+#include "AudioCompressionSettingsUtils.h"
 
 static int32 SpoofFailedStreamChunkLoad = 0;
 FAutoConsoleVariableRef CVarSpoofFailedStreamChunkLoad(
@@ -46,13 +47,15 @@ FAsyncStreamDerivedChunkWorker::FAsyncStreamDerivedChunkWorker(
 	const FString& InDerivedDataKey,
 	void* InDestChunkData,
 	int32 InChunkSize,
-	FThreadSafeCounter* InThreadSafeCounter
+	FThreadSafeCounter* InThreadSafeCounter,
+	TFunction<void(bool)> InOnLoadCompleted
 	)
 	: DerivedDataKey(InDerivedDataKey)
 	, DestChunkData(InDestChunkData)
 	, ExpectedChunkSize(InChunkSize)
 	, bRequestFailed(false)
 	, ThreadSafeCounter(InThreadSafeCounter)
+	, OnLoadCompleted(InOnLoadCompleted)
 {
 }
 
@@ -69,9 +72,13 @@ void FAsyncStreamDerivedChunkWorker::DoWork()
 	{
 		FMemoryReader Ar(DerivedChunkData, true);
 		int32 ChunkSize = 0;
+		int32 AudioDataSize = 0;
 		Ar << ChunkSize;
-		checkf(ChunkSize == ExpectedChunkSize, TEXT("ChunkSize(%d) != ExpectedSize(%d)"), ChunkSize, ExpectedChunkSize);
-		Ar.Serialize(DestChunkData, ChunkSize);
+		Ar << AudioDataSize;
+
+		// Currently, the legacy streaming manager loads in the entire, zero padded chunk, while the cached streaming manager only reads the audio data itself.
+		checkf(AudioDataSize == ExpectedChunkSize || ChunkSize == ExpectedChunkSize, TEXT("Neither the padded chunk size (%d) nor the actual audio data size (%d) was equivalent to the ExpectedSize(%d)"), ChunkSize, AudioDataSize, ExpectedChunkSize);
+		Ar.Serialize(DestChunkData, ExpectedChunkSize);
 	}
 	else
 	{
@@ -80,7 +87,8 @@ void FAsyncStreamDerivedChunkWorker::DoWork()
 	FPlatformMisc::MemoryBarrier();
 	ThreadSafeCounter->Decrement();
 
-	UE_LOG(LogAudio, Verbose, TEXT("End of ASync DDC Chunk read for key: %s"), *DerivedDataKey);
+	OnLoadCompleted(bRequestFailed);
+	UE_LOG(LogAudio, Verbose, TEXT("End of Async DDC Chunk Load. DDC Key: %s"), *DerivedDataKey);
 }
 
 #endif // #if WITH_EDITORONLY_DATA
@@ -125,7 +133,7 @@ void FStreamingWaveData::FreeResources()
 	}
 }
 
-bool FStreamingWaveData::Initialize(USoundWave* InSoundWave, FAudioStreamingManager* InAudioStreamingManager)
+bool FStreamingWaveData::Initialize(USoundWave* InSoundWave, FLegacyAudioStreamingManager* InAudioStreamingManager)
 {
 	check(!IORequestHandle);
 
@@ -330,11 +338,16 @@ void FStreamingWaveData::BeginPendingRequests(const TArray<uint32>& IndicesToLoa
 				INC_DWORD_STAT_BY(STAT_AudioMemorySize, ChunkSize);
 				INC_DWORD_STAT_BY(STAT_AudioMemory, ChunkSize);
 
+				// This streaming manager does not use the callback when the DDC load
+				// is completed:
+				TFunction<void(bool)> NullOnLoadCompletedCallback = [](bool) {};
+
 				FAsyncStreamDerivedChunkTask* Task = new FAsyncStreamDerivedChunkTask(
 					Chunk.DerivedDataKey,
 					ChunkStorage->Data,
 					ChunkSize,
-					&PendingChunkChangeRequestStatus
+					&PendingChunkChangeRequestStatus,
+					MoveTemp(NullOnLoadCompletedCallback)
 					);
 				PendingAsyncStreamDerivedChunkTasks.Add(Task);
 				Task->StartBackgroundTask();
@@ -479,15 +492,15 @@ void FStreamingWaveData::FreeLoadedChunk(FLoadedAudioChunk& LoadedChunk)
 // FAudioStreamingManager //
 ////////////////////////////
 
-FAudioStreamingManager::FAudioStreamingManager()
+FLegacyAudioStreamingManager::FLegacyAudioStreamingManager()
 {
 }
 
-FAudioStreamingManager::~FAudioStreamingManager()
+FLegacyAudioStreamingManager::~FLegacyAudioStreamingManager()
 {
 }
 
-void FAudioStreamingManager::OnAsyncFileCallback(FStreamingWaveData* StreamingWaveData, int32 LoadedAudioChunkIndex, IAsyncReadRequest* ReadRequest)
+void FLegacyAudioStreamingManager::OnAsyncFileCallback(FStreamingWaveData* StreamingWaveData, int32 LoadedAudioChunkIndex, IAsyncReadRequest* ReadRequest)
 {
 	// Check to see if we successfully managed to load anything
 	uint8* Mem = ReadRequest->GetReadResults();
@@ -511,7 +524,7 @@ void FAudioStreamingManager::OnAsyncFileCallback(FStreamingWaveData* StreamingWa
 	}
 }
 
-void FAudioStreamingManager::ProcessPendingAsyncFileResults()
+void FLegacyAudioStreamingManager::ProcessPendingAsyncFileResults()
 {
 	// Pump the results of any async file loads in a protected critical section 
 	FScopeLock StreamChunkResults(&ChunkResultCriticalSection);
@@ -550,7 +563,7 @@ void FAudioStreamingManager::ProcessPendingAsyncFileResults()
 	AsyncAudioStreamChunkResults.Reset();
 }
 
-void FAudioStreamingManager::UpdateResourceStreaming(float DeltaTime, bool bProcessEverything /*= false*/)
+void FLegacyAudioStreamingManager::UpdateResourceStreaming(float DeltaTime, bool bProcessEverything /*= false*/)
 {
 	LLM_SCOPE(ELLMTag::Audio);
 
@@ -639,7 +652,7 @@ void FAudioStreamingManager::UpdateResourceStreaming(float DeltaTime, bool bProc
 	ProcessPendingAsyncFileResults();
 }
 
-int32 FAudioStreamingManager::BlockTillAllRequestsFinished(float TimeLimit, bool)
+int32 FLegacyAudioStreamingManager::BlockTillAllRequestsFinished(float TimeLimit, bool)
 {
 	{
 		FScopeLock Lock(&CriticalSection);
@@ -679,35 +692,36 @@ int32 FAudioStreamingManager::BlockTillAllRequestsFinished(float TimeLimit, bool
 	return 0;
 }
 
-void FAudioStreamingManager::CancelForcedResources()
+void FLegacyAudioStreamingManager::CancelForcedResources()
 {
 }
 
-void FAudioStreamingManager::NotifyLevelChange()
+void FLegacyAudioStreamingManager::NotifyLevelChange()
 {
 }
 
-void FAudioStreamingManager::SetDisregardWorldResourcesForFrames(int32 NumFrames)
+void FLegacyAudioStreamingManager::SetDisregardWorldResourcesForFrames(int32 NumFrames)
 {
 }
 
-void FAudioStreamingManager::AddLevel(class ULevel* Level)
+void FLegacyAudioStreamingManager::AddLevel(class ULevel* Level)
 {
 }
 
-void FAudioStreamingManager::RemoveLevel(class ULevel* Level)
+void FLegacyAudioStreamingManager::RemoveLevel(class ULevel* Level)
 {
 }
 
-void FAudioStreamingManager::NotifyLevelOffset(class ULevel* Level, const FVector& Offset)
+void FLegacyAudioStreamingManager::NotifyLevelOffset(class ULevel* Level, const FVector& Offset)
 {
 }
 
-void FAudioStreamingManager::AddStreamingSoundWave(USoundWave* SoundWave)
+void FLegacyAudioStreamingManager::AddStreamingSoundWave(USoundWave* SoundWave)
 {
 	if (FPlatformProperties::SupportsAudioStreaming() && SoundWave->IsStreaming())
 	{
 		FScopeLock Lock(&CriticalSection);
+
 		if (StreamingSoundWaves.FindRef(SoundWave) == nullptr)
 		{
 			FStreamingWaveData* NewStreamingWaveData = new FStreamingWaveData;
@@ -724,7 +738,7 @@ void FAudioStreamingManager::AddStreamingSoundWave(USoundWave* SoundWave)
 	}
 }
 
-void FAudioStreamingManager::RemoveStreamingSoundWave(USoundWave* SoundWave)
+void FLegacyAudioStreamingManager::RemoveStreamingSoundWave(USoundWave* SoundWave)
 {
 	FScopeLock Lock(&CriticalSection);
 	FStreamingWaveData* WaveData = StreamingSoundWaves.FindRef(SoundWave);
@@ -755,25 +769,25 @@ void FAudioStreamingManager::RemoveStreamingSoundWave(USoundWave* SoundWave)
 	WaveRequests.Remove(SoundWave);
 }
 
-void FAudioStreamingManager::AddDecoder(ICompressedAudioInfo* InCompressedAudioInfo)
+void FLegacyAudioStreamingManager::AddDecoder(ICompressedAudioInfo* InCompressedAudioInfo)
 {
 	FScopeLock Lock(&CriticalSection);
 	CompressedAudioInfos.AddUnique(InCompressedAudioInfo);
 }
 
-void FAudioStreamingManager::RemoveDecoder(ICompressedAudioInfo* InCompressedAudioInfo)
+void FLegacyAudioStreamingManager::RemoveDecoder(ICompressedAudioInfo* InCompressedAudioInfo)
 {
 	FScopeLock Lock(&CriticalSection);
 	CompressedAudioInfos.Remove(InCompressedAudioInfo);
 }
 
-bool FAudioStreamingManager::IsManagedStreamingSoundWave(const USoundWave* SoundWave) const
+bool FLegacyAudioStreamingManager::IsManagedStreamingSoundWave(const USoundWave* SoundWave) const
 {
 	FScopeLock Lock(&CriticalSection);
 	return StreamingSoundWaves.FindRef(SoundWave) != NULL;
 }
 
-bool FAudioStreamingManager::IsStreamingInProgress(const USoundWave* SoundWave)
+bool FLegacyAudioStreamingManager::IsStreamingInProgress(const USoundWave* SoundWave)
 {
 	FScopeLock Lock(&CriticalSection);
 	FStreamingWaveData* WaveData = StreamingSoundWaves.FindRef(SoundWave);
@@ -784,7 +798,7 @@ bool FAudioStreamingManager::IsStreamingInProgress(const USoundWave* SoundWave)
 	return false;
 }
 
-bool FAudioStreamingManager::CanCreateSoundSource(const FWaveInstance* WaveInstance) const
+bool FLegacyAudioStreamingManager::CanCreateSoundSource(const FWaveInstance* WaveInstance) const
 {
 	check(WaveInstance);
 	check(WaveInstance->IsStreaming());
@@ -822,7 +836,7 @@ bool FAudioStreamingManager::CanCreateSoundSource(const FWaveInstance* WaveInsta
 	return false;
 }
 
-void FAudioStreamingManager::AddStreamingSoundSource(FSoundSource* SoundSource)
+void FLegacyAudioStreamingManager::AddStreamingSoundSource(FSoundSource* SoundSource)
 {
 	const FWaveInstance* WaveInstance = SoundSource->GetWaveInstance();
 	if (WaveInstance && WaveInstance->IsStreaming())
@@ -865,10 +879,10 @@ void FAudioStreamingManager::AddStreamingSoundSource(FSoundSource* SoundSource)
 	}
 }
 
-void FAudioStreamingManager::RemoveStreamingSoundSource(FSoundSource* SoundSource)
+void FLegacyAudioStreamingManager::RemoveStreamingSoundSource(FSoundSource* SoundSource)
 {
 	const FWaveInstance* WaveInstance = SoundSource->GetWaveInstance();
-	if (WaveInstance && WaveInstance->WaveData && WaveInstance->WaveData->IsStreaming())
+	if (WaveInstance && WaveInstance->WaveData && WaveInstance->WaveData->IsStreaming() && !WaveInstance->WaveData->ShouldUseStreamCaching())
 	{
 		FScopeLock Lock(&CriticalSection);
 
@@ -879,24 +893,30 @@ void FAudioStreamingManager::RemoveStreamingSoundSource(FSoundSource* SoundSourc
 	}
 }
 
-bool FAudioStreamingManager::IsManagedStreamingSoundSource(const FSoundSource* SoundSource) const
+bool FLegacyAudioStreamingManager::IsManagedStreamingSoundSource(const FSoundSource* SoundSource) const
 {
 	FScopeLock Lock(&CriticalSection);
 	return StreamingSoundSources.FindByKey(SoundSource) != NULL;
 }
 
-const uint8* FAudioStreamingManager::GetLoadedChunk(const USoundWave* SoundWave, uint32 ChunkIndex, uint32* OutChunkSize) const
+bool FLegacyAudioStreamingManager::RequestChunk(USoundWave* SoundWave, uint32 ChunkIndex, TFunction<void(EAudioChunkLoadResult)> OnLoadCompleted)
+{
+	UE_LOG(LogAudio, Warning, TEXT("RequestChunk is only supported in Stream Caching."));
+	return false;
+}
+
+FAudioChunkHandle FLegacyAudioStreamingManager::GetLoadedChunk(const USoundWave* SoundWave, uint32 ChunkIndex, bool bBlockForLoad) const
 {
 	// Check for the spoof of failing to load a stream chunk
 	if (SpoofFailedStreamChunkLoad > 0)
 	{
-		return nullptr;
+		return FAudioChunkHandle();
 	}
 
 	// If we fail at getting the critical section here, early out. 
 	if (!CriticalSection.TryLock())
 	{
-		return nullptr;
+		return FAudioChunkHandle();
 	}
 
 	const FStreamingWaveData* WaveData = StreamingSoundWaves.FindRef(SoundWave);
@@ -908,24 +928,18 @@ const uint8* FAudioStreamingManager::GetLoadedChunk(const USoundWave* SoundWave,
 			{
 				if (WaveData->LoadedChunks[Index].Index == ChunkIndex)
 				{
-					if (OutChunkSize != NULL)
-					{
-						// Return the size of the audio data within the chunk, not the chunk itself since this chunk could be zero-padded
-						*OutChunkSize = WaveData->LoadedChunks[Index].AudioDataSize;
-					}
-
 					CriticalSection.Unlock();
-					return WaveData->LoadedChunks[Index].Data;
+					return BuildChunkHandle(WaveData->LoadedChunks[Index].Data, WaveData->LoadedChunks[Index].AudioDataSize, SoundWave, SoundWave->GetFName(), ChunkIndex);
 				}
 			}
 		}
 	}
 
 	CriticalSection.Unlock();
-	return NULL;
+	return FAudioChunkHandle();
 }
 
-FWaveRequest& FAudioStreamingManager::GetWaveRequest(USoundWave* SoundWave)
+FWaveRequest& FLegacyAudioStreamingManager::GetWaveRequest(USoundWave* SoundWave)
 {
 	FWaveRequest* WaveRequest = WaveRequests.Find(SoundWave);
 	if (!WaveRequest)
@@ -937,3 +951,15 @@ FWaveRequest& FAudioStreamingManager::GetWaveRequest(USoundWave* SoundWave)
 	}
 	return *WaveRequest;
 }
+
+uint64 FLegacyAudioStreamingManager::TrimMemory(uint64 NumBytesToFree)
+{
+	UE_LOG(LogAudio, Warning, TEXT("IAudioStreamingManager::TrimMemory is only supported when Stream Caching is enabled."));
+	return 0;
+}
+
+int32 FLegacyAudioStreamingManager::RenderStatAudioStreaming(UWorld* World, FViewport* Viewport, FCanvas* Canvas, int32 X, int32 Y, const FVector* ViewLocation, const FRotator* ViewRotation)
+{
+	// Only supported by stream caching.
+	return Y;
+}
diff --git a/Engine/Source/Runtime/Engine/Private/AudioStreamingCache.cpp b/Engine/Source/Runtime/Engine/Private/AudioStreamingCache.cpp
new file mode 100644
--- /dev/null
+++ b/Engine/Source/Runtime/Engine/Private/AudioStreamingCache.cpp
@@ -0,0 +1,1050 @@
+// Copyright 1998-2019 Epic Games, Inc. All Rights Reserved.
+
+/*=============================================================================
+AudioStreaming.cpp: Implementation of audio streaming classes.
+=============================================================================*/
+
+#include "AudioStreamingCache.h"
+#include "Misc/CoreStats.h"
+#include "Sound/SoundWave.h"
+#include "Sound/AudioSettings.h"
+#include "DerivedDataCacheInterface.h"
+#include "Serialization/MemoryReader.h"
+#include "HAL/PlatformFile.h"
+#include "HAL/PlatformFilemanager.h"
+#include "Async/AsyncFileHandle.h"
+#include "Misc/ScopeLock.h"
+#include "HAL/IConsoleManager.h"
+#include "HAL/LowLevelMemTracker.h"
+#include "AudioDecompress.h"
+#include "AudioCompressionSettingsUtils.h"
+
+static int32 DebugMaxElementsDisplayCVar = 128;
+FAutoConsoleVariableRef CVarDebugDisplayCaches(
+	TEXT("au.streamcaching.MaxCachesToDisplay"),
+	DebugMaxElementsDisplayCVar,
+	TEXT("Sets the max amount of stream chunks to display on screen.\n")
+	TEXT("n: Number of elements to display on screen."),
+	ECVF_Default);
+
+FCachedAudioStreamingManager::FCachedAudioStreamingManager(const FCachedAudioStreamingManagerParams& InitParams)
+{
+	check(FPlatformCompressionUtilities::IsCurrentPlatformUsingStreamCaching());
+	checkf(InitParams.Caches.Num() > 0, TEXT("FCachedAudioStreamingManager should be initialized with dimensions for at least one cache."));
+
+	// const FAudioStreamCachingSettings& CacheSettings = FPlatformCompressionUtilities::GetStreamCachingSettingsForCurrentPlatform();
+	for (const FCachedAudioStreamingManagerParams::FCacheDimensions& CacheDimensions : InitParams.Caches)
+	{
+		CacheArray.Emplace(CacheDimensions.MaxChunkSize, CacheDimensions.NumElements);
+	}
+
+	// Here we make sure our CacheArray is sorted from smallest MaxChunkSize to biggest, so that GetCacheForWave can scan through these caches to find the appropriate cache for the chunk size.
+	CacheArray.Sort();
+}
+
+FCachedAudioStreamingManager::~FCachedAudioStreamingManager()
+{
+}
+
+void FCachedAudioStreamingManager::UpdateResourceStreaming(float DeltaTime, bool bProcessEverything /*= false*/)
+{
+	// The cached audio streaming manager doesn't tick.
+}
+
+int32 FCachedAudioStreamingManager::BlockTillAllRequestsFinished(float TimeLimit, bool bLogResults)
+{
+	// TODO: Honor TimeLimit and bLogResults. Since we cancel any in flight read requests, this should not spin out.
+	for (FAudioChunkCache& Cache : CacheArray)
+	{
+		Cache.CancelAllPendingLoads();
+	}
+
+	return 0;
+}
+
+void FCachedAudioStreamingManager::CancelForcedResources()
+{
+	// Unused.
+}
+
+void FCachedAudioStreamingManager::NotifyLevelChange()
+{
+	// Unused.
+}
+
+void FCachedAudioStreamingManager::SetDisregardWorldResourcesForFrames(int32 NumFrames)
+{
+	// Unused.
+}
+
+void FCachedAudioStreamingManager::AddLevel(class ULevel* Level)
+{
+	// Unused.
+}
+
+void FCachedAudioStreamingManager::RemoveLevel(class ULevel* Level)
+{
+	// Unused.
+}
+
+void FCachedAudioStreamingManager::NotifyLevelOffset(class ULevel* Level, const FVector& Offset)
+{
+	// Unused.
+}
+
+void FCachedAudioStreamingManager::AddStreamingSoundWave(USoundWave* SoundWave)
+{
+	// Unused.
+}
+
+void FCachedAudioStreamingManager::RemoveStreamingSoundWave(USoundWave* SoundWave)
+{
+	// Unused.
+}
+
+void FCachedAudioStreamingManager::AddDecoder(ICompressedAudioInfo* InCompressedAudioInfo)
+{
+	// Unused.
+}
+
+void FCachedAudioStreamingManager::RemoveDecoder(ICompressedAudioInfo* InCompressedAudioInfo)
+{
+	//Unused.
+}
+
+bool FCachedAudioStreamingManager::IsManagedStreamingSoundWave(const USoundWave* SoundWave) const
+{
+	// Unused. The concept of a sound wave being "managed" doesn't apply here.
+	checkf(false, TEXT("Not Implemented!"));
+	return true;
+}
+
+bool FCachedAudioStreamingManager::IsStreamingInProgress(const USoundWave* SoundWave)
+{
+	// This function is used in USoundWave cleanup.
+	// Since this manager owns the binary data we are streaming off of,
+	// It's safe to delete the USoundWave as long as
+	// There are NO sound sources playing with this Sound Wave.
+	//
+	// This is because a playing sound source might kick off a load for a new chunk,
+	// which dereferences the corresponding USoundWave.
+	//
+	// As of right now, this is handled by USoundWave::FreeResources(), called
+	// by USoundWave::IsReadyForFinishDestroy.
+	return false;
+}
+
+bool FCachedAudioStreamingManager::CanCreateSoundSource(const FWaveInstance* WaveInstance) const
+{
+	return true;
+}
+
+void FCachedAudioStreamingManager::AddStreamingSoundSource(FSoundSource* SoundSource)
+{
+	// Unused.
+}
+
+void FCachedAudioStreamingManager::RemoveStreamingSoundSource(FSoundSource* SoundSource)
+{
+	// Unused.
+}
+
+bool FCachedAudioStreamingManager::IsManagedStreamingSoundSource(const FSoundSource* SoundSource) const
+{
+	// Unused. The concept of a sound wave being "managed" doesn't apply here.
+	checkf(false, TEXT("Not Implemented!"));
+	return true;
+}
+
+FAudioChunkHandle FCachedAudioStreamingManager::GetLoadedChunk(const USoundWave* SoundWave, uint32 ChunkIndex, bool bBlockForLoad) const
+{
+	// If this sound wave is managed by a cache, use that to get the chunk:
+	FAudioChunkCache* Cache = GetCacheForWave(SoundWave);
+	if (Cache)
+	{
+		// With this code, the zeroth chunk should never get hit.
+		checkf(ChunkIndex != 0, TEXT("Decoder tried to access the zeroth chunk through the streaming manager. Use USoundWave::GetZerothChunk() instead."));
+
+		// TODO:  See if we can avoid non-const calls on the USoundWave here.
+		USoundWave* MutableWave = const_cast<USoundWave*>(SoundWave);
+		const FAudioChunkCache::FChunkKey ChunkKey = { MutableWave, SoundWave->GetFName(), ChunkIndex };
+
+		if (!FAudioChunkCache::IsKeyValid(ChunkKey))
+		{
+			UE_LOG(LogAudio, Warning, TEXT("Invalid Chunk Index %d Requested for Wave %s!"), ChunkIndex, *SoundWave->GetName());
+			return FAudioChunkHandle();
+		}
+
+		// The function call below increments the reference count to the internal chunk.
+		TArrayView<uint8> LoadedChunk = Cache->GetChunk(ChunkKey, bBlockForLoad);
+
+		const int32 NextChunk = GetNextChunkIndex(SoundWave, ChunkIndex);
+
+		if (NextChunk != INDEX_NONE)
+		{
+			const FAudioChunkCache::FChunkKey NextChunkKey = { MutableWave, SoundWave->GetFName(), ((uint32)NextChunk) };
+			bool bIsValidChunk = Cache->AddOrTouchChunk(NextChunkKey, [](EAudioChunkLoadResult) {});
+			if (!bIsValidChunk)
+			{
+				UE_LOG(LogAudio, Warning, TEXT("Cache overflow!!! couldn't load chunk %d for sound %s!"), ChunkIndex, *SoundWave->GetName());
+			}
+		}
+
+		
+		return BuildChunkHandle(LoadedChunk.GetData(), LoadedChunk.Num(), SoundWave, SoundWave->GetFName(), ChunkIndex);
+	}
+	else
+	{
+		ensureMsgf(false, TEXT("Failed to find cache for wave %s. Are you sure this is a streaming wave?"), *SoundWave->GetName());
+		return FAudioChunkHandle();
+	}
+}
+
+FAudioChunkCache* FCachedAudioStreamingManager::GetCacheForWave(const USoundWave* InSoundWave) const
+{
+	check(InSoundWave);
+
+	// We only cache chunks beyond the zeroth chunk of audio (which is inlined directly on the asset)
+	if (InSoundWave->RunningPlatformData && InSoundWave->RunningPlatformData->Chunks.Num() > 1)
+	{
+		const int32 SoundWaveChunkSize = InSoundWave->RunningPlatformData->Chunks[1].AudioDataSize;
+
+		// Iterate over our caches until we find the lowest MaxChunkSize cache this sound's chunks will fit into. 
+		for (int32 CacheIndex = 0; CacheIndex < CacheArray.Num(); CacheIndex++)
+		{
+			if (SoundWaveChunkSize <= CacheArray[CacheIndex].MaxChunkSize)
+			{
+				return const_cast<FAudioChunkCache*>(&CacheArray[CacheIndex]);
+			}
+		}
+
+		// If we ever hit this, something went wrong during cook.
+		// Please check to make sure this platform's implementation of IAudioFormat honors the MaxChunkSize parameter passed into SplitDataForStreaming,
+		// or that FStreamedAudioCacheDerivedDataWorker::BuildStreamedAudio() is passing the correct MaxChunkSize to IAudioFormat::SplitDataForStreaming.
+		ensureMsgf(false, TEXT("Chunks in SoundWave are too large: %s, Chunk size: %d bytes"), *InSoundWave->GetName(), SoundWaveChunkSize);
+		return nullptr;
+	}
+	else
+	{
+		return nullptr;
+	}
+}
+
+int32 FCachedAudioStreamingManager::GetNextChunkIndex(const USoundWave* InSoundWave, uint32 CurrentChunkIndex) const
+{
+	check(InSoundWave);
+	// TODO: Figure out a way to tell whether this wave is looping or not.
+	// if(bNotLooping) return ((int32) CurrentChunkIndex) < (InSoundWave->RunningPlatformData->Chunks.Num() - 1);
+	
+	const int32 NumChunksTotal = InSoundWave->GetNumChunks();
+	if (NumChunksTotal <= 2)
+	{
+		// If there's only one chunk to cache (besides the zeroth chunk, which is inlined),
+		// We don't need to load anything.
+		return INDEX_NONE;
+	}
+	else if(CurrentChunkIndex == (NumChunksTotal - 1))
+	{
+		// if we're on the last chunk, load the first chunk after the zeroth chunk.
+		return 1;
+	}
+	else
+	{
+		// Otherwise, there's another chunk of audio after this one before the end of this file.
+		return CurrentChunkIndex + 1;
+	}
+}
+
+void FCachedAudioStreamingManager::AddReferenceToChunk(const FAudioChunkHandle& InHandle)
+{
+	FAudioChunkCache* Cache = GetCacheForWave(InHandle.CorrespondingWave);
+	check(Cache);
+
+	FAudioChunkCache::FChunkKey ChunkKey =
+	{
+		const_cast<USoundWave*>(InHandle.CorrespondingWave),
+		InHandle.CorrespondingWaveName,
+		((uint32) InHandle.ChunkIndex)
+	};
+
+	Cache->AddNewReferenceToChunk(ChunkKey);
+}
+
+void FCachedAudioStreamingManager::RemoveReferenceToChunk(const FAudioChunkHandle& InHandle)
+{
+	FAudioChunkCache* Cache = GetCacheForWave(InHandle.CorrespondingWave);
+	check(Cache);
+
+	FAudioChunkCache::FChunkKey ChunkKey =
+	{
+		const_cast<USoundWave*>(InHandle.CorrespondingWave),
+		InHandle.CorrespondingWaveName,
+		((uint32) InHandle.ChunkIndex)
+	};
+
+	Cache->RemoveReferenceToChunk(ChunkKey);
+}
+
+bool FCachedAudioStreamingManager::RequestChunk(USoundWave* SoundWave, uint32 ChunkIndex, TFunction<void(EAudioChunkLoadResult)> OnLoadCompleted)
+{
+	FAudioChunkCache* Cache = GetCacheForWave(SoundWave);
+	if (Cache)
+	{
+		FAudioChunkCache::FChunkKey ChunkKey = { SoundWave, SoundWave->GetFName(), ChunkIndex };
+		return Cache->AddOrTouchChunk(ChunkKey, OnLoadCompleted);
+	}
+	else
+	{
+		// This can hit if an out of bounds chunk was requested, or the zeroth chunk was requested from the streaming manager.
+		ensureMsgf(false, TEXT("GetCacheForWave failed for SoundWave %s!"), *SoundWave->GetName());
+		return false;
+	}
+}
+
+FAudioChunkCache::FAudioChunkCache(uint32 InMaxChunkSize, uint32 NumChunks)
+	: MaxChunkSize(InMaxChunkSize)
+	, MostRecentElement(nullptr)
+	, LeastRecentElement(nullptr)
+	, ChunksInUse(0)
+	, bLogCacheMisses(false)
+{
+	CachePool.Reset(NumChunks);
+	for (uint32 Index = 0; Index < NumChunks; Index++)
+	{
+		CachePool.Emplace(MaxChunkSize, Index);
+	}
+}
+
+FAudioChunkCache::~FAudioChunkCache()
+{
+	// While this is handled by the default destructor, we do this to ensure that we don't leak async read operations.
+	CachePool.Reset();
+	check(NumberOfLoadsInFlight.GetValue() == 0);
+}
+
+bool FAudioChunkCache::AddOrTouchChunk(const FChunkKey& InKey, TFunction<void(EAudioChunkLoadResult)> OnLoadCompleted)
+{
+	if (!IsKeyValid(InKey))
+	{
+		ensure(false);
+		OnLoadCompleted(EAudioChunkLoadResult::ChunkOutOfBounds);
+		return false;
+	}
+
+	FCacheElement* FoundElement = FindElementForKey(InKey);
+	
+	if (FoundElement)
+	{
+		TouchElement(FoundElement);
+		if (FoundElement->bIsLoaded)
+		{
+			OnLoadCompleted(EAudioChunkLoadResult::AlreadyLoaded);
+		}
+
+#if DEBUG_STREAM_CACHE
+		FoundElement->DebugInfo.NumTimesTouched++;
+#endif
+
+		return true;
+	}
+	else
+	{
+		FCacheElement* CacheElement = InsertChunk(InKey);
+
+		if (!CacheElement)
+		{
+			OnLoadCompleted(EAudioChunkLoadResult::CacheBlown);
+			return false;
+		}
+
+		KickOffAsyncLoad(CacheElement, InKey, OnLoadCompleted);
+		return true;
+	}
+}
+
+TArrayView<uint8> FAudioChunkCache::GetChunk(const FChunkKey& InKey, bool bBlockForLoadCompletion)
+{
+	FCacheElement* FoundElement = FindElementForKey(InKey);
+	if (FoundElement)
+	{
+		TouchElement(FoundElement);
+		if (bBlockForLoadCompletion)
+		{
+			FoundElement->WaitForAsyncLoadCompletion(false);
+		}
+		else if (FoundElement->IsLoadInProgress())
+		{
+			return TArrayView<uint8>();
+		}
+
+		// If this value is ever negative, it means that we're decrementing more than we're incrementing:
+		check(FoundElement->NumConsumers.GetValue() >= 0);
+		FoundElement->NumConsumers.Increment();
+		return TArrayView<uint8>(FoundElement->ChunkData.GetData(), FoundElement->ChunkDataSize);
+	}
+	else
+	{
+		// If we missed it, kick off a new load with it.
+		FoundElement = InsertChunk(InKey);
+
+		if (!FoundElement)
+		{
+			return TArrayView<uint8>();
+		}
+
+		KickOffAsyncLoad(FoundElement, InKey, [](EAudioChunkLoadResult InResult) {});
+
+		if (bBlockForLoadCompletion)
+		{
+			// If bBlockForLoadCompletion was true and we don't have an element present, we have to load the element into the cache:
+			FoundElement->WaitForAsyncLoadCompletion(false);
+
+			FoundElement->NumConsumers.Increment();
+			return TArrayView<uint8>(FoundElement->ChunkData.GetData(), FoundElement->ChunkDataSize);
+		}
+		else if (bLogCacheMisses)
+		{
+			// Chunks missing. Log this as a miss.
+			const uint32 TotalNumChunksInWave = InKey.SoundWave->RunningPlatformData->Chunks.Num();
+
+			FCacheMissInfo CacheMissInfo = { InKey.SoundWaveName, InKey.ChunkIndex, TotalNumChunksInWave, false };
+			CacheMissQueue.Enqueue(MoveTemp(CacheMissInfo));
+		}
+
+		// We missed, return an empty array view.
+		return TArrayView<uint8>();
+	}
+}
+
+void FAudioChunkCache::AddNewReferenceToChunk(const FChunkKey& InKey)
+{
+	FCacheElement* FoundElement = FindElementForKey(InKey);
+	check(FoundElement);
+
+	// If this value is ever negative, it means that we're decrementing more than we're incrementing:
+	check(FoundElement->NumConsumers.GetValue() >= 0);
+	FoundElement->NumConsumers.Increment();
+}
+
+void FAudioChunkCache::RemoveReferenceToChunk(const FChunkKey& InKey)
+{
+	FCacheElement* FoundElement = FindElementForKey(InKey);
+	check(FoundElement);
+
+	// If this value is ever less than 1 when we hit this code, it means that we're decrementing more than we're incrementing:
+	check(FoundElement->NumConsumers.GetValue() >= 1);
+	FoundElement->NumConsumers.Decrement();
+}
+
+void FAudioChunkCache::ClearCache()
+{
+	const uint32 NumChunks = CachePool.Num();
+
+	CachePool.Reset(NumChunks);
+	check(NumberOfLoadsInFlight.GetValue() == 0);
+
+	for (uint32 Index = 0; Index < NumChunks; Index++)
+	{
+		CachePool.Emplace(MaxChunkSize, Index);
+	}
+
+	MostRecentElement = nullptr;
+	LeastRecentElement = nullptr;
+	ChunksInUse = 0;
+}
+
+uint64 FAudioChunkCache::TrimMemory(uint64 BytesToFree)
+{
+	FScopeLock ScopeLock(&CacheMutationCriticalSection);
+
+	if (!MostRecentElement || MostRecentElement->LessRecentElement == nullptr)
+	{
+		return 0;
+	}
+
+	FCacheElement* CurrentElement = LeastRecentElement;
+
+	// In order to avoid cycles, we always leave at least two chunks in the cache.
+	const FCacheElement* ElementToStopAt = MostRecentElement->LessRecentElement;
+
+	uint64 BytesFreed = 0;
+	while (CurrentElement != ElementToStopAt && BytesFreed < BytesToFree)
+	{
+		if (CurrentElement->CanEvictChunk())
+		{
+			BytesFreed += CurrentElement->ChunkData.Num();
+
+			// Empty the chunk data and invalidate the key.
+			CurrentElement->ChunkData.Empty();
+			CurrentElement->ChunkDataSize = 0;
+			CurrentElement->Key = FChunkKey();
+
+#if DEBUG_STREAM_CACHE
+			// Reset debug info:
+			CurrentElement->DebugInfo = FCacheElementDebugInfo();
+#endif
+		}
+
+		// Important to note that we don't actually relink chunks here,
+		// So by trimming memory we are not moving chunks up the recency list.
+		CurrentElement = CurrentElement->MoreRecentElement;
+	}
+
+	return BytesFreed;
+}
+
+void FAudioChunkCache::BlockForAllPendingLoads() const
+{
+	bool bLoadInProgress = false;
+
+	float TimeStarted = FPlatformTime::Seconds();
+
+	do 
+	{
+		// If we did find an in flight async load,
+		// sleep to let other threads complete this task.
+		if (bLoadInProgress)
+		{
+			float TimeSinceStarted = FPlatformTime::Seconds() - TimeStarted;
+			UE_LOG(LogAudio, Display, TEXT("Waited %f seconds for async audio chunk loads."), TimeSinceStarted);
+			FPlatformProcess::Sleep(0.0f);
+		}
+
+		// Iterate through every element until we find one with a load in progress.
+		const FCacheElement* CurrentElement = MostRecentElement;
+		while (CurrentElement != nullptr)
+		{
+			bLoadInProgress |= CurrentElement->IsLoadInProgress();
+			CurrentElement = CurrentElement->LessRecentElement;
+		}
+	} while (bLoadInProgress);
+}
+
+void FAudioChunkCache::CancelAllPendingLoads()
+{
+	FCacheElement* CurrentElement = MostRecentElement;
+	while (CurrentElement != nullptr)
+	{
+		CurrentElement->WaitForAsyncLoadCompletion(true);
+		CurrentElement = CurrentElement->LessRecentElement;
+	}
+}
+
+uint64 FAudioChunkCache::ReportCacheSize()
+{
+	const uint32 NumChunks = CachePool.Num();
+
+	return MaxChunkSize * NumChunks;
+}
+
+void FAudioChunkCache::BeginLoggingCacheMisses()
+{
+	bLogCacheMisses = true;
+}
+
+void FAudioChunkCache::StopLoggingCacheMisses()
+{
+	bLogCacheMisses = false;
+}
+
+FString FAudioChunkCache::FlushCacheMissLog()
+{
+	FString ConcatenatedCacheMisses;
+	ConcatenatedCacheMisses.Append(TEXT("SoundWave:\tChunkIndex\n"));
+
+	FCacheMissInfo CacheMissInfo;
+	while (CacheMissQueue.Dequeue(CacheMissInfo))
+	{
+		ConcatenatedCacheMisses.Append(CacheMissInfo.SoundWaveName.ToString());
+		ConcatenatedCacheMisses.Append(TEXT("\t"));
+		ConcatenatedCacheMisses.AppendInt(CacheMissInfo.ChunkIndex);
+		ConcatenatedCacheMisses.Append(TEXT("\n"));
+	}
+
+	return ConcatenatedCacheMisses;
+}
+
+FAudioChunkCache::FCacheElement* FAudioChunkCache::FindElementForKey(const FChunkKey& InKey)
+{
+	FCacheElement* CurrentElement = MostRecentElement;
+
+#if DEBUG_STREAM_CACHE
+	// In debuggable situations, we breadcrumb how far down the cache the cache we were.
+	int32 ElementPosition = 0;
+#endif
+
+	while (CurrentElement != nullptr)
+	{
+		if (InKey == CurrentElement->Key)
+		{
+			
+#if DEBUG_STREAM_CACHE
+			CurrentElement->DebugInfo.PreviousLocationsBeforeBeingTouched.Add(ElementPosition);
+#endif
+
+			return CurrentElement;
+		}
+		else
+		{
+			CurrentElement = CurrentElement->LessRecentElement;
+
+#if DEBUG_STREAM_CACHE
+			ElementPosition++;
+#endif
+		}
+	}
+
+	return CurrentElement;
+}
+
+void FAudioChunkCache::TouchElement(FCacheElement* InElement)
+{
+	checkSlow(InElement);
+
+	// Check to ensure we do not have any cycles in our list.
+	// If this first check is hit, try to ensure that EvictLeastRecent chunk isn't evicting the top two chunks.
+	check(MostRecentElement == nullptr || MostRecentElement != LeastRecentElement);
+	check(InElement->LessRecentElement != InElement);
+
+	FScopeLock ScopeLock(&CacheMutationCriticalSection);
+
+	// If this is already the most recent element, we don't need to do anything.
+	if (InElement == MostRecentElement)
+	{
+		return;
+	}
+
+	// If this was previously the least recent chunk, update LeastRecentElement.
+	if (LeastRecentElement == InElement)
+	{
+		LeastRecentElement = InElement->MoreRecentElement;
+	}
+
+	FCacheElement* PreviousLessRecent = InElement->LessRecentElement;
+	FCacheElement* PreviousMoreRecent = InElement->MoreRecentElement;
+	FCacheElement* PreviousMostRecent = MostRecentElement;
+
+	check(PreviousMostRecent != InElement);
+
+	// Move this element to the top:
+	MostRecentElement = InElement;
+	InElement->MoreRecentElement = nullptr;
+	InElement->LessRecentElement = PreviousMostRecent;
+
+	if (PreviousMostRecent != nullptr)
+	{
+		PreviousMostRecent->MoreRecentElement = InElement;
+	}
+
+	if (PreviousLessRecent == PreviousMoreRecent)
+	{
+		return;
+	}
+	else
+	{
+		// Link InElement's previous neighbors together:
+		if (PreviousLessRecent != nullptr)
+		{
+			PreviousLessRecent->MoreRecentElement = PreviousMoreRecent;
+		}
+
+		if (PreviousMoreRecent != nullptr)
+		{
+			PreviousMoreRecent->LessRecentElement = PreviousLessRecent;
+		}
+	}
+}
+
+FAudioChunkCache::FCacheElement* FAudioChunkCache::InsertChunk(const FChunkKey& InKey)
+{
+	FCacheElement* CacheElement = nullptr;
+	
+	{
+		FScopeLock ScopeLock(&CacheMutationCriticalSection);
+
+		if (ChunksInUse < CachePool.Num())
+		{
+			// We haven't filled up the pool yet, so we don't need to evict anything.
+			CacheElement = &CachePool[ChunksInUse];
+			ChunksInUse++;
+		}
+		else
+		{
+			// The pools filled, so we're going to need to evict.
+			CacheElement = EvictLeastRecentChunk();
+
+			if (!CacheElement)
+			{
+				return nullptr;
+			}
+		}
+	}
+
+	check(CacheElement);
+	CacheElement->bIsLoaded = false;
+	CacheElement->Key = InKey;
+	TouchElement(CacheElement);
+
+	// If we've got multiple chunks, we can not cache the least recent chunk
+	// without worrying about a circular dependency.
+	if (LeastRecentElement == nullptr && ChunksInUse > 1)
+	{
+		SetUpLeastRecentChunk();
+	}
+
+	return CacheElement;
+}
+
+void FAudioChunkCache::SetUpLeastRecentChunk()
+{
+	FScopeLock ScopeLock(&CacheMutationCriticalSection);
+
+	FCacheElement* CacheElement = MostRecentElement;
+	while (CacheElement->LessRecentElement != nullptr)
+	{
+		CacheElement = CacheElement->LessRecentElement;
+	}
+
+	LeastRecentElement = CacheElement;
+}
+
+FAudioChunkCache::FCacheElement* FAudioChunkCache::EvictLeastRecentChunk()
+{
+	FCacheElement* CacheElement = LeastRecentElement;
+
+	// If the least recent chunk is evictable, evict it.
+	if (CacheElement->CanEvictChunk())
+	{
+		FCacheElement* NewLeastRecentElement = CacheElement->MoreRecentElement;
+		check(NewLeastRecentElement);
+
+		LeastRecentElement = NewLeastRecentElement;
+	}
+	else
+	{
+		// We should never hit this code path unless we have at least two chunks active.
+		check(MostRecentElement && MostRecentElement->LessRecentElement);
+
+		// In order to avoid cycles, we always leave at least two chunks in the cache.
+		const FCacheElement* ElementToStopAt = MostRecentElement->LessRecentElement;
+
+		// Otherwise, we need to crawl up the cache from least recent used to most to find a chunk that is not in use:
+		while (CacheElement != ElementToStopAt)
+		{
+			if (CacheElement->CanEvictChunk())
+			{
+				// Link the two neighboring chunks:
+				if (CacheElement->MoreRecentElement)
+				{
+					CacheElement->MoreRecentElement->LessRecentElement = CacheElement->LessRecentElement;
+				}
+				
+				// If we ever hit this while loop it means that CacheElement is not the least recently used element.
+				check(CacheElement->LessRecentElement);
+				CacheElement->LessRecentElement->MoreRecentElement = CacheElement->MoreRecentElement;
+				break;
+			}
+			else
+			{
+				CacheElement = CacheElement->MoreRecentElement;
+			}
+
+			// If we ever hit this, it means that we couldn't find any cache elements that aren't in use.
+			if (CacheElement != MostRecentElement)
+			{
+				ensureMsgf(false, TEXT("Cache blown! Please increase the cache size or load less audio."));
+				return nullptr;
+			}
+		}
+	}
+
+#if DEBUG_STREAM_CACHE
+	// Reset debug information:
+	CacheElement->DebugInfo.PreviousLocationsBeforeBeingTouched.Reset();
+	CacheElement->DebugInfo.NumTimesTouched = 0;
+#endif
+
+	return CacheElement;
+}
+
+void FAudioChunkCache::KickOffAsyncLoad(FCacheElement* CacheElement, const FChunkKey& InKey, TFunction<void(EAudioChunkLoadResult)> OnLoadCompleted)
+{
+	LLM_SCOPE(ELLMTag::Audio);
+	check(CacheElement);
+	
+	const FStreamedAudioChunk& Chunk = InKey.SoundWave->RunningPlatformData->Chunks[InKey.ChunkIndex];
+	int32 ChunkDataSize = Chunk.AudioDataSize;
+
+	EAsyncIOPriorityAndFlags AsyncIOPriority = AIOP_High;
+
+	// Reallocate our chunk data This allows us to shrink if possible.
+	CacheElement->ChunkData.SetNumUninitialized(Chunk.AudioDataSize, true);
+
+#if DEBUG_STREAM_CACHE
+	CacheElement->DebugInfo.NumTotalChunks = InKey.SoundWave->GetNumChunks() - 1;
+#endif
+
+	// In editor, we retrieve from the DDC. In non-editor situations, we read the chunk async from the pak file.
+#if WITH_EDITORONLY_DATA
+	if (Chunk.DerivedDataKey.IsEmpty() == false)
+	{
+		CacheElement->ChunkDataSize = ChunkDataSize;
+
+		INC_DWORD_STAT_BY(STAT_AudioMemorySize, ChunkDataSize);
+		INC_DWORD_STAT_BY(STAT_AudioMemory, ChunkDataSize);
+
+		if (CacheElement->DDCTask.IsValid())
+		{
+			check(CacheElement->DDCTask->IsDone());
+		}
+		
+#if DEBUG_STREAM_CACHE
+		CacheElement->DebugInfo.TimeLoadStarted = FPlatformTime::Seconds();
+#endif
+
+
+		TFunction<void(bool)> OnLoadComplete = [OnLoadCompleted, CacheElement, InKey, ChunkDataSize](bool bRequestFailed)
+		{
+			// Populate key and DataSize. The async read request was set up to write directly into CacheElement->ChunkData.
+			CacheElement->Key = InKey;
+			CacheElement->ChunkDataSize = ChunkDataSize;
+			CacheElement->bIsLoaded = true;
+
+#if DEBUG_STREAM_CACHE
+			CacheElement->DebugInfo.TimeToLoad = (FPlatformTime::Seconds() - CacheElement->DebugInfo.TimeLoadStarted) * 1000.0f;
+#endif
+
+			OnLoadCompleted(bRequestFailed ? EAudioChunkLoadResult::Interrupted : EAudioChunkLoadResult::Completed);
+		};
+
+		NumberOfLoadsInFlight.Increment();
+
+		CacheElement->DDCTask.Reset(new FAsyncStreamDerivedChunkTask(
+			Chunk.DerivedDataKey,
+			CacheElement->ChunkData.GetData(),
+			ChunkDataSize,
+			&NumberOfLoadsInFlight,
+			MoveTemp(OnLoadComplete)
+		));
+
+		CacheElement->DDCTask->StartBackgroundTask();
+	}
+	else
+#endif // #if WITH_EDITORONLY_DATA
+	{
+		check(Chunk.BulkData.GetFilename().Len());
+		UE_CLOG(Chunk.BulkData.IsStoredCompressedOnDisk(), LogAudio, Fatal, TEXT("Package level compression is no longer supported."));
+
+		// TODO: Figure out how to cache this file handle somewhere, so that we are not reopening a file on every individual load operation.
+		CacheElement->FileHandle.Reset(FPlatformFileManager::Get().GetPlatformFile().OpenAsyncRead(*Chunk.BulkData.GetFilename()));
+		check(CacheElement->FileHandle.IsValid());
+
+		// Sanity check our bulk data against
+		const int32 ChunkBulkDataSize = Chunk.BulkData.GetBulkDataSize();
+		check(ChunkDataSize <= ChunkBulkDataSize);
+		check(((int32)ChunkDataSize) <= CacheElement->ChunkData.Num());
+
+		// If we ever want to eliminate zero-padding in chunks, that could be done here:
+		//ensureAlwaysMsgf(AudioChunkSize == ChunkBulkDataSize, TEXT("For memory load on demand, we do not zero-pad to page sizes."));
+		
+		NumberOfLoadsInFlight.Increment();
+
+		FAsyncFileCallBack AsyncFileCallBack = [this, OnLoadCompleted, CacheElement, InKey, ChunkDataSize](bool bWasCancelled, IAsyncReadRequest* Request)
+		{
+			// Populate key and DataSize. The async read request was set up to write directly into CacheElement->ChunkData.
+			CacheElement->Key = InKey;
+			CacheElement->ChunkDataSize = ChunkDataSize;
+			CacheElement->bIsLoaded = true;
+
+#if DEBUG_STREAM_CACHE
+			CacheElement->DebugInfo.TimeToLoad = FPlatformTime::Seconds() - CacheElement->DebugInfo.TimeLoadStarted;
+#endif
+
+			OnLoadCompleted(bWasCancelled ? EAudioChunkLoadResult::Interrupted : EAudioChunkLoadResult::Completed);
+
+			NumberOfLoadsInFlight.Decrement();
+		};
+
+#if DEBUG_STREAM_CACHE
+		CacheElement->DebugInfo.TimeLoadStarted = FPlatformTime::Seconds();
+#endif
+
+		CacheElement->ReadRequest.Reset(CacheElement->FileHandle->ReadRequest(Chunk.BulkData.GetBulkDataOffsetInFile(), ChunkDataSize, AsyncIOPriority, &AsyncFileCallBack, CacheElement->ChunkData.GetData()));
+		if (!CacheElement->ReadRequest.IsValid())
+		{
+			UE_LOG(LogAudio, Error, TEXT("Chunk load in audio LRU cache failed."));
+			OnLoadCompleted(EAudioChunkLoadResult::ChunkOutOfBounds);
+			NumberOfLoadsInFlight.Decrement();
+		}
+	}
+}
+
+bool FAudioChunkCache::IsKeyValid(const FChunkKey& InKey)
+{
+	return InKey.ChunkIndex < TNumericLimits<uint32>::Max() && ((int32)InKey.ChunkIndex) < InKey.SoundWave->RunningPlatformData->Chunks.Num();
+}
+
+#include "UnrealEngine.h"
+
+int32 FCachedAudioStreamingManager::RenderStatAudioStreaming(UWorld* World, FViewport* Viewport, FCanvas* Canvas, int32 X, int32 Y, const FVector* ViewLocation, const FRotator* ViewRotation)
+{
+	Canvas->DrawShadowedString(X, Y, TEXT("Stream Caches:"), UEngine::GetSmallFont(), FLinearColor::White);
+	Y += 12;
+
+	int32 CacheIndex = 0;
+	int32 Height = Y;
+	for (const FAudioChunkCache& Cache : CacheArray)
+	{
+		FString CacheTitle = *FString::Printf(TEXT("Cache %d"), CacheIndex);
+		Canvas->DrawShadowedString(X, Y, *CacheTitle, UEngine::GetSmallFont(), FLinearColor::White);
+		Y += 12;
+
+		TPair<int, int> Size = Cache.DebugDisplay(World, Viewport, Canvas, X, Y, ViewLocation, ViewRotation);
+		
+		// Separate caches are laid out horizontally across the screen, so the total height is equal to our tallest cache panel:
+		X += Size.Key;
+		Height = FMath::Max(Height, Size.Value);
+	}
+
+	return Y + Height;
+}
+
+uint64 FCachedAudioStreamingManager::TrimMemory(uint64 NumBytesToFree)
+{
+	uint64 NumBytesLeftToFree = NumBytesToFree;
+	
+	// TODO: When we support multiple caches, it's probably best to do this in reverse,
+	// since the caches are sorted from shortest sounds to longest.
+	// Freeing longer chunks will get us bigger gains and (presumably) have lower churn.
+	for (FAudioChunkCache& Cache : CacheArray)
+	{
+		uint64 NumBytesFreed = Cache.TrimMemory(NumBytesLeftToFree);
+
+		// NumBytesFreed could potentially be more than what we requested to free (since we delete whole chunks at once).
+		NumBytesLeftToFree -= FMath::Min(NumBytesFreed, NumBytesLeftToFree);
+
+		// If we've freed all the memory we needed to, exit.
+		if (NumBytesLeftToFree == 0)
+		{
+			break;
+		}
+	}
+
+	check(NumBytesLeftToFree <= NumBytesToFree);
+	uint64 TotalBytesFreed = NumBytesToFree - NumBytesLeftToFree;
+
+	UE_LOG(LogAudio, Display, TEXT("Call to IAudioStreamingManager::TrimMemory successfully freed %lu of the requested %lu bytes."), TotalBytesFreed, NumBytesToFree);
+	return TotalBytesFreed;
+}
+
+#include "Engine/Font.h"
+
+TPair<int, int> FAudioChunkCache::DebugDisplay(UWorld* World, FViewport* Viewport, FCanvas* Canvas, int32 X, int32 Y, const FVector* ViewLocation, const FRotator* ViewRotation) const
+{
+	const int32 InitialX = X;
+	const int32 InitialY = Y;
+
+	FString NumElementsDetail = *FString::Printf(TEXT("Number of chunks loaded: %d of %d"), ChunksInUse, CachePool.Num());
+
+	// Offset our number of elements loaded horizontally to the right next to the cache title:
+	int32 CacheTitleOffsetX = 0;
+	int32 CacheTitleOffsetY = 0;
+	UEngine::GetSmallFont()->GetStringHeightAndWidth(TEXT("Cache XX "), CacheTitleOffsetY, CacheTitleOffsetX);
+
+	Canvas->DrawShadowedString(X + CacheTitleOffsetX, Y - 12, *NumElementsDetail, UEngine::GetSmallFont(), FLinearColor::Green);
+	Y += 10;
+
+	// First pass: We run through and get a snap shot of the amount of memory currently in use.
+	FCacheElement* CurrentElement = MostRecentElement;
+	uint32 NumBytesCounter = 0;
+
+	while (CurrentElement != nullptr)
+	{
+		// Note: this is potentially a stale value if we're in the middle of FCacheElement::KickOffAsyncLoad.
+		NumBytesCounter += CurrentElement->ChunkData.Num();
+		CurrentElement = CurrentElement->LessRecentElement;
+	}
+	
+	// Convert to megabytes and print the total size:
+	const double NumMegabytesInUse = (double)NumBytesCounter / (1024 * 1024);
+	const double MaxCacheSizeMB = ((double)MaxChunkSize * CachePool.Num()) / (1024 * 1024);
+
+	FString CacheMemoryUsage = *FString::Printf(TEXT("Using: %.4f Megabytes. Max Potential Usage: %.4f Megabytes."), NumMegabytesInUse, MaxCacheSizeMB);
+
+	// We're going to align this horizontally with the number of elements right above it.
+	Canvas->DrawShadowedString(X + CacheTitleOffsetX, Y, *CacheMemoryUsage, UEngine::GetSmallFont(), FLinearColor::Green);
+	Y += 12;
+
+	// Second Pass: We're going to list the actual chunks in the cache.
+	CurrentElement = MostRecentElement;
+	int32 Index = 0;
+
+	float ColorLerpAmount = 0.0f;
+	const float ColorLerpStep = 0.04f;
+
+	// More detailed info about individual chunks here:
+	while (CurrentElement != nullptr)
+	{
+		// We use a CVar to clamp the max amount of chunks we display.
+		if (Index > DebugMaxElementsDisplayCVar)
+		{
+			break;
+		}
+
+		int32 NumTotalChunks = -1;
+		int32 NumTimesTouched = -1;
+		double TimeToLoad = -1.0;
+
+#if DEBUG_STREAM_CACHE
+		NumTotalChunks = CurrentElement->DebugInfo.NumTotalChunks;
+		NumTimesTouched = CurrentElement->DebugInfo.NumTimesTouched;
+		TimeToLoad = CurrentElement->DebugInfo.TimeToLoad;
+#endif
+
+		const bool bWasTrimmed = CurrentElement->ChunkData.Num() == 0;
+
+		FString ElementInfo = *FString::Printf(TEXT("%4i. Size: %6.2f KB   Chunk: %d of %d   Request Count: %d    Number of Decoders Using Chunk: %d     Chunk Load Time: %6.4fms      Name: %s "),
+			Index,
+			CurrentElement->ChunkData.Num() / 1024.0f,
+			CurrentElement->Key.ChunkIndex,
+			NumTotalChunks,
+			NumTimesTouched,
+			CurrentElement->NumConsumers.GetValue(),
+			TimeToLoad,
+			bWasTrimmed ? TEXT("TRIMMED CHUNK") : *CurrentElement->Key.SoundWaveName.ToString()
+			);
+
+		// Since there's a lot of info here,
+		// Subtly fading the chunk info to gray seems to help as a visual indicator of how far down on the list things are.
+		ColorLerpAmount = FMath::Min(ColorLerpAmount + ColorLerpStep, 1.0f);
+		FLinearColor TextColor = FLinearColor::LerpUsingHSV(FLinearColor::Blue, FLinearColor::Gray, ColorLerpAmount);
+
+		// If there's a load in flight, paint this element yellow.
+		if (CurrentElement->IsLoadInProgress())
+		{
+			TextColor = FLinearColor::Yellow;
+		}
+		else if (CurrentElement->IsInUse())
+		{
+			TextColor = FLinearColor::Green;
+		}
+		else if (bWasTrimmed)
+		{
+			TextColor = FLinearColor::Red;
+		}
+
+		Canvas->DrawShadowedString(X, Y, *ElementInfo, UEngine::GetSmallFont(), TextColor);
+		Y += 12;
+
+		CurrentElement = CurrentElement->LessRecentElement;
+		Index++;
+	}
+
+	// The largest element of our debug panel is the initial memory details.
+	int32 CacheMemoryTextOffsetX = 0;
+	int32 CacheMemoryTextOffsetY = 0;
+	UEngine::GetSmallFont()->GetStringHeightAndWidth(*CacheMemoryUsage, CacheMemoryTextOffsetX, CacheMemoryTextOffsetY);
+
+	return TPair<int, int>(X + CacheTitleOffsetX + CacheMemoryTextOffsetX - InitialX, Y - InitialY);
+}
+
diff --git a/Engine/Source/Runtime/Engine/Private/ContentStreaming.cpp b/Engine/Source/Runtime/Engine/Private/ContentStreaming.cpp
--- a/Engine/Source/Runtime/Engine/Private/ContentStreaming.cpp
+++ b/Engine/Source/Runtime/Engine/Private/ContentStreaming.cpp
@@ -18,6 +18,8 @@
 #include "Streaming/StreamingManagerTexture.h"
 #include "AudioStreaming.h"
 #include "Animation/AnimationStreaming.h"
+#include "AudioStreamingCache.h"
+#include "AudioCompressionSettingsUtils.h"
 
 /*-----------------------------------------------------------------------------
 	Globals.
@@ -753,7 +755,16 @@ FStreamingManagerCollection::FStreamingManagerCollection()
 
 	AddOrRemoveTextureStreamingManagerIfNeeded(true);
 
-	AudioStreamingManager = new FAudioStreamingManager();
+	if (FPlatformCompressionUtilities::IsCurrentPlatformUsingStreamCaching())
+	{
+		FCachedAudioStreamingManagerParams Params = FPlatformCompressionUtilities::BuildCachedStreamingManagerParams();
+		AudioStreamingManager = new FCachedAudioStreamingManager(Params);
+	}
+	else
+	{
+		AudioStreamingManager = new FLegacyAudioStreamingManager();
+	}
+	
 	AddStreamingManager( AudioStreamingManager );
 
 	AnimationStreamingManager = new FAnimationStreamingManager();
@@ -952,6 +963,10 @@ IRenderAssetStreamingManager& FStreamingManagerCollection::GetTextureStreamingMa
 
 IAudioStreamingManager& FStreamingManagerCollection::GetAudioStreamingManager() const
 {
+#if WITH_EDITOR
+	FScopeLock ScopeLock(&AudioStreamingManagerCriticalSection);
+#endif
+
 	check(AudioStreamingManager);
 	return *AudioStreamingManager;
 }
@@ -1129,6 +1144,48 @@ void FStreamingManagerCollection::PropagateLightingScenarioChange()
 	}
 }
 
+#if WITH_EDITOR
+#include "AudioDevice.h"
+#include "AudioDeviceManager.h"
+
+void FStreamingManagerCollection::OnAudioStreamingParamsChanged()
+{
+	// Before we swap out the audio streaming manager, we'll need to stop all sounds running on all audio devices:
+	FAudioDeviceManager* DeviceManager = GEngine->GetAudioDeviceManager();
+	TArray<FAudioDevice*>& AudioDevices = DeviceManager->GetAudioDevices();
+	for (FAudioDevice* AudioDevice : AudioDevices)
+	{
+		if (AudioDevice)
+		{
+			AudioDevice->StopAllSounds(true);
+		}
+	}
+
+	FScopeLock ScopeLock(&AudioStreamingManagerCriticalSection);
+
+	RemoveStreamingManager(AudioStreamingManager);
+
+	delete AudioStreamingManager;
+	AudioStreamingManager = nullptr;
+
+	// Lastly, make sure we have the most up to date cook overrides cached:
+	FPlatformCompressionUtilities::RecacheCookOverrides();
+
+	// Finally, reinitialize the streaming manager.
+	if (FPlatformCompressionUtilities::IsCurrentPlatformUsingStreamCaching())
+	{
+		FCachedAudioStreamingManagerParams Params = FPlatformCompressionUtilities::BuildCachedStreamingManagerParams();
+		AudioStreamingManager = new FCachedAudioStreamingManager(Params);
+	}
+	else
+	{
+		AudioStreamingManager = new FLegacyAudioStreamingManager();
+	}
+
+	AddStreamingManager(AudioStreamingManager);
+}
+#endif
+
 void FStreamingManagerCollection::AddOrRemoveTextureStreamingManagerIfNeeded(bool bIsInit)
 {
 	bool bUseTextureStreaming = false;
@@ -1275,3 +1332,82 @@ FArchive& operator<<( FArchive& Ar, FDynamicTextureInstance& TextureInstance )
 	Ar << TextureInstance.OriginalRadius;
 	return Ar;
 }
+
+
+/*-----------------------------------------------------------------------------
+	Audio chunk handles. Used by the cached audio streaming manager.
+-----------------------------------------------------------------------------*/
+FAudioChunkHandle::FAudioChunkHandle()
+	: CachedData(nullptr)
+	, CachedDataNumBytes(0)
+	, CorrespondingWave(nullptr)
+	, CorrespondingWaveName()
+	, ChunkIndex(INDEX_NONE)
+{
+}
+
+FAudioChunkHandle::FAudioChunkHandle(const uint8* InData, uint32 NumBytes, const USoundWave* InSoundWave, const FName& SoundWaveName, uint32 InChunkIndex)
+	: CachedData(InData)
+	, CachedDataNumBytes(NumBytes)
+	, CorrespondingWave(InSoundWave)
+	, CorrespondingWaveName(SoundWaveName)
+	, ChunkIndex(InChunkIndex)
+{
+}
+
+FAudioChunkHandle::FAudioChunkHandle(const FAudioChunkHandle& Other)
+{
+	*this = Other;
+}
+
+FAudioChunkHandle& FAudioChunkHandle::operator=(const FAudioChunkHandle& Other)
+{
+	// If this chunk was previously referencing another chunk, remove that chunk here.
+	if (IsValid())
+	{
+		IStreamingManager::Get().GetAudioStreamingManager().RemoveReferenceToChunk(*this);
+	}
+
+	CachedData = Other.CachedData;
+	CachedDataNumBytes = Other.CachedDataNumBytes;
+	CorrespondingWave = Other.CorrespondingWave;
+	CorrespondingWaveName = Other.CorrespondingWaveName;
+	ChunkIndex = Other.ChunkIndex;
+
+
+	if (IsValid())
+	{
+		// Increment the reference count for the new streaming chunk.
+		IStreamingManager::Get().GetAudioStreamingManager().AddReferenceToChunk(*this);
+	}
+
+	return *this;
+}
+
+FAudioChunkHandle::~FAudioChunkHandle()
+{
+	if (IsValid())
+	{
+		IStreamingManager::Get().GetAudioStreamingManager().RemoveReferenceToChunk(*this);
+	}
+}
+
+const uint8* FAudioChunkHandle::GetData() const
+{
+	return CachedData;
+}
+
+uint32 FAudioChunkHandle::Num() const
+{
+	return CachedDataNumBytes;
+}
+
+bool FAudioChunkHandle::IsValid() const
+{
+	return GetData() != nullptr;
+}
+
+FAudioChunkHandle IAudioStreamingManager::BuildChunkHandle(const uint8* InData, uint32 NumBytes, const USoundWave* InSoundWave, const FName& SoundWaveName, uint32 InChunkIndex)
+{
+	return FAudioChunkHandle(InData, NumBytes, InSoundWave, SoundWaveName, InChunkIndex);
+}
diff --git a/Engine/Source/Runtime/Engine/Private/OpusAudioInfo.cpp b/Engine/Source/Runtime/Engine/Private/OpusAudioInfo.cpp
--- a/Engine/Source/Runtime/Engine/Private/OpusAudioInfo.cpp
+++ b/Engine/Source/Runtime/Engine/Private/OpusAudioInfo.cpp
@@ -113,13 +113,13 @@ FOpusAudioInfo::~FOpusAudioInfo()
 
 bool FOpusAudioInfo::ParseHeader(const uint8* InSrcBufferData, uint32 InSrcBufferDataSize, struct FSoundQualityInfo* QualityInfo)
 {
-	SrcBufferData = InSrcBufferData;
-	SrcBufferDataSize = InSrcBufferDataSize;
 	SrcBufferOffset = 0;
 	CurrentSampleCount = 0;
 
+	check(InSrcBufferDataSize >= sizeof(OPUS_ID_STRING));
+
 	// Read Identifier, True Sample Count, Number of channels and Frames to Encode first
-	if (!SrcBufferData || FCStringAnsi::Strcmp((char*)SrcBufferData, OPUS_ID_STRING) != 0)
+	if (!InSrcBufferData || FMemory::Memcmp((char*)InSrcBufferData, OPUS_ID_STRING, sizeof(OPUS_ID_STRING) - 1) != 0)
 	{
 		return false;
 	}
diff --git a/Engine/Source/Runtime/Engine/Private/SoundWave.cpp b/Engine/Source/Runtime/Engine/Private/SoundWave.cpp
--- a/Engine/Source/Runtime/Engine/Private/SoundWave.cpp
+++ b/Engine/Source/Runtime/Engine/Private/SoundWave.cpp
@@ -39,6 +39,14 @@ FAutoConsoleVariableRef CVarBypassPlayWhenSilent(
 	TEXT("0: Honor the Play When Silent flag, 1: stop all silent non-procedural sources."),
 	ECVF_Default);
 
+static int32 LoadIntoCacheOnPostLoadCVar = 0;
+FAutoConsoleVariableRef CVarLoadIntoCacheOnPostLoad(
+	TEXT("au.streamcache.LoadIntoCacheOnPostLoad"),
+	LoadIntoCacheOnPostLoadCVar,
+	TEXT("When set to 1, immediately loads any sounds that are loaded into the cache on load.\n")
+	TEXT("0: Honor the Play When Silent flag, 1: stop all silent non-procedural sources."),
+	ECVF_Default);
+
 #if !UE_BUILD_SHIPPING
 static void DumpBakedAnalysisData(const TArray<FString>& Args)
 {
@@ -141,6 +149,8 @@ uint32 FStreamedAudioChunk::StoreInDerivedDataCache(const FString& InDerivedData
 	TArray<uint8> DerivedData;
 	FMemoryWriter Ar(DerivedData, /*bIsPersistent=*/ true);
 	Ar << BulkDataSizeInBytes;
+	Ar << AudioDataSize;
+
 	{
 		void* BulkChunkData = BulkData.Lock(LOCK_READ_ONLY);
 		Ar.Serialize(BulkChunkData, BulkDataSizeInBytes);
@@ -713,7 +723,7 @@ void USoundWave::PostLoad()
 		return;
 	}
 
-	// In case any code accesses bStreaming directly, we update it based on the current platform's cook overrides.
+	// In case any code accesses bStreaming directly, we fix up bStreaming based on the current platform's cook overrides.
 	bStreaming = IsStreaming();
 
 	// Compress to whatever formats the active target platforms want
@@ -729,6 +739,8 @@ void USoundWave::PostLoad()
 		}
 	}
 
+	
+
 	// We don't precache default objects and we don't precache in the Editor as the latter will
 	// most likely cause us to run out of memory.
 	if (!GIsEditor && !IsTemplate( RF_ClassDefaultObject ) && GEngine)
@@ -755,6 +767,19 @@ void USoundWave::PostLoad()
 		IStreamingManager::Get().GetAudioStreamingManager().AddStreamingSoundWave(this);
 	}
 
+	if (ShouldUseStreamCaching() && IsStreaming())
+	{
+		EnsureZerothChunkIsLoaded();
+	}
+
+	const bool bShouldPrimeSound = LoadIntoCacheOnPostLoadCVar || bLoadCompressedAudioWhenSoundWaveIsLoaded;
+
+	if (bShouldPrimeSound && FPlatformCompressionUtilities::IsCurrentPlatformUsingStreamCaching())
+	{
+		// Load rest of the audio into cache.
+		IStreamingManager::Get().GetAudioStreamingManager().RequestChunk(this, 1, [](EAudioChunkLoadResult) {});
+	}
+
 #if WITH_EDITORONLY_DATA
 	if (!SourceFilePath_DEPRECATED.IsEmpty() && AssetImportData)
 	{
@@ -770,6 +795,78 @@ void USoundWave::PostLoad()
 	INC_FLOAT_STAT_BY( STAT_AudioBufferTimeChannels, NumChannels * Duration );
 }
 
+void USoundWave::EnsureZerothChunkIsLoaded()
+{
+	// If the zeroth chunk is already loaded, early exit.
+	if (ZerothChunkData.Num() > 0)
+	{
+		return;
+	}
+
+#if WITH_EDITOR 
+	if (RunningPlatformData == nullptr)
+	{
+		CachePlatformData(false);
+	}
+
+	// If we're running the editor, we'll need to retrieve the chunked audio from the DDC:
+	uint8* TempChunkBuffer = nullptr;
+	int32 ChunkSizeInBytes = RunningPlatformData->GetChunkFromDDC(0, &TempChunkBuffer, true);
+	// Since we block for the DDC in the previous call we should always have the chunk loaded.
+	check(ChunkSizeInBytes > 0);
+
+	// TODO: Support passing a TArray by ref into FStreamedAudioPlatformData::GetChunkFromDDC.
+	// Currently not feasible unless FUntypedBulkData::GetCopy API was changed.
+	// For now in editor, we have an extra allocated buffer in this scope.
+	ZerothChunkData.Reset();
+	ZerothChunkData.AddUninitialized(ChunkSizeInBytes);
+	FMemory::Memcpy(ZerothChunkData.GetData(), TempChunkBuffer, ChunkSizeInBytes);
+	FMemory::Free(TempChunkBuffer);
+#else // WITH_EDITOR
+	// Otherwise, the zeroth chunk is cooked out to RunningPlatformData, and we just need to retrieve it.
+	check(RunningPlatformData && RunningPlatformData->Chunks.Num() > 0);
+	FStreamedAudioChunk& ZerothChunk = RunningPlatformData->Chunks[0];
+	// Some sanity checks to ensure that the bulk size set up
+	check(ZerothChunk.BulkData.GetBulkDataSize() == ZerothChunk.DataSize && ZerothChunk.DataSize >= ZerothChunk.AudioDataSize);
+
+	ZerothChunkData.Reset();
+	ZerothChunkData.AddUninitialized(ZerothChunk.AudioDataSize);
+	uint8* Destination = ZerothChunkData.GetData();
+	ZerothChunk.BulkData.GetCopy((void**)&Destination, true);
+#endif // WITH_EDITOR
+}
+
+uint32 USoundWave::GetNumChunks() const
+{
+	if (RunningPlatformData)
+	{
+		return RunningPlatformData->Chunks.Num();
+	}
+	else
+	{
+		return 0;
+	}
+}
+
+uint32 USoundWave::GetSizeOfChunk(uint32 ChunkIndex)
+{
+	check(ChunkIndex < GetNumChunks());
+
+	if (ChunkIndex == 0)
+	{
+		EnsureZerothChunkIsLoaded();
+		return ZerothChunkData.Num();
+	}
+	else if(RunningPlatformData)
+	{
+		return RunningPlatformData->Chunks[ChunkIndex].AudioDataSize;
+	}
+	else
+	{
+		return 0;
+	}
+}
+
 void USoundWave::BeginDestroy()
 {
 	Super::BeginDestroy();
@@ -862,7 +959,7 @@ void USoundWave::RemoveAudioResource()
 	delete OwnedBulkDataPtr;
 	OwnedBulkDataPtr = nullptr;
 	ResourceData = nullptr;
-		ResourceSize = 0;
+	ResourceSize = 0;
 #endif
 }
 
@@ -1717,16 +1814,71 @@ float USoundWave::GetDuration()
 
 bool USoundWave::IsStreaming(const FPlatformAudioCookOverrides* Overrides /* = nullptr */) const
 {
-	// TODO: add in check on whether it's part of a streaming SoundGroup
+	// TODO: add in check on whether it's part of a streaming SoundGroup.
 	if (!Overrides)
 	{
 		Overrides = GetPlatformCompressionOverridesForCurrentPlatform();
 	}
 
-	return bStreaming
-		|| (Overrides != nullptr
-			&& Overrides->AutoStreamingThreshold > SMALL_NUMBER
-			&& Duration > Overrides->AutoStreamingThreshold);
+	// We stream if (A) bStreaming is set to true, (B) bForceInline is false and either bUseLoadOnDemand was set to true in
+	// our cook overrides, or the AutoStreamingThreshold was set and this sound is longer than the auto streaming threshold.
+	
+	if (bStreaming)
+	{
+		return true;
+	}
+	else if (bForceInline)
+	{
+		return false;
+	}
+	else if (bProcedural)
+	{
+		return false;
+	}
+	else if (Overrides)
+	{
+		// For stream caching, the auto streaming threshold is used to force sounds to be inlined:
+		const bool bUsesStreamCache = (Overrides->bUseStreamCaching && Duration > Overrides->AutoStreamingThreshold);
+		const bool bOverAutoStreamingThreshold = (Overrides->AutoStreamingThreshold > SMALL_NUMBER  && Duration > Overrides->AutoStreamingThreshold);
+
+		return bUsesStreamCache || bOverAutoStreamingThreshold;
+	}
+	else
+	{
+		return false;
+	}
+}
+
+bool USoundWave::ShouldUseStreamCaching() const
+{
+	return  FPlatformCompressionUtilities::IsCurrentPlatformUsingStreamCaching() && IsStreaming();
+}
+
+TArrayView<const uint8> USoundWave::GetZerothChunk()
+{
+	if (ShouldUseStreamCaching())
+	{
+		// In editor, we actually don't have a zeroth chunk until we try to play an audio file.
+		if (GIsEditor)
+		{
+			EnsureZerothChunkIsLoaded();
+		}
+
+		check(ZerothChunkData.Num() > 0);
+
+		if (GetNumChunks() > 1)
+		{
+			// Prime first chunk for playback.
+			IStreamingManager::Get().GetAudioStreamingManager().RequestChunk(this, 1, [](EAudioChunkLoadResult InResult) {});
+		}
+
+		return TArrayView<const uint8>(ZerothChunkData);
+	}
+	else
+	{
+		FAudioChunkHandle ChunkHandle = IStreamingManager::Get().GetAudioStreamingManager().GetLoadedChunk(this, 0);
+		return TArrayView<const uint8>(ChunkHandle.GetData(), ChunkHandle.Num());
+	}
 }
 
 bool USoundWave::IsSeekableStreaming() const
@@ -1826,14 +1978,14 @@ float USoundWave::GetSampleRateForCompressionOverrides(const FPlatformAudioCookO
 
 bool USoundWave::GetChunkData(int32 ChunkIndex, uint8** OutChunkData, bool bMakeSureChunkIsLoaded /* = false */)
 {
-	if (RunningPlatformData->TryLoadChunk(ChunkIndex, OutChunkData, bMakeSureChunkIsLoaded) == false)
+	if (RunningPlatformData->GetChunkFromDDC(ChunkIndex, OutChunkData, bMakeSureChunkIsLoaded) == 0)
 	{
 #if WITH_EDITORONLY_DATA
 		// Unable to load chunks from the cache. Rebuild the sound and attempt to recache it.
 		UE_LOG(LogAudio, Display, TEXT("GetChunkData failed, rebuilding %s"), *GetPathName());
 
 		ForceRebuildPlatformData();
-		if (RunningPlatformData->TryLoadChunk(ChunkIndex, OutChunkData, bMakeSureChunkIsLoaded) == false)
+		if (RunningPlatformData->GetChunkFromDDC(ChunkIndex, OutChunkData, bMakeSureChunkIsLoaded) == 0)
 		{
 			UE_LOG(LogAudio, Display, TEXT("Failed to build sound %s."), *GetPathName());
 		}
diff --git a/Engine/Source/Runtime/Engine/Private/UnrealEngine.cpp b/Engine/Source/Runtime/Engine/Private/UnrealEngine.cpp
--- a/Engine/Source/Runtime/Engine/Private/UnrealEngine.cpp
+++ b/Engine/Source/Runtime/Engine/Private/UnrealEngine.cpp
@@ -1603,6 +1603,7 @@ void UEngine::Init(IEngineLoop* InEngineLoop)
 	EngineStats.Add(FEngineStatFuncs(TEXT("STAT_Sounds"), TEXT("STATCAT_Engine"), FText::GetEmpty(), &UEngine::RenderStatSounds, &UEngine::ToggleStatSounds));
 	EngineStats.Add(FEngineStatFuncs(TEXT("STAT_SoundCues"), TEXT("STATCAT_Engine"), FText::GetEmpty(), &UEngine::RenderStatSoundCues, &UEngine::ToggleStatSoundCues));
 	EngineStats.Add(FEngineStatFuncs(TEXT("STAT_SoundMixes"), TEXT("STATCAT_Engine"), FText::GetEmpty(), &UEngine::RenderStatSoundMixes, &UEngine::ToggleStatSoundMixes));
+	EngineStats.Add(FEngineStatFuncs(TEXT("STAT_AudioStreaming"), TEXT("STATCAT_Engine"), FText::GetEmpty(), &UEngine::RenderStatAudioStreaming, &UEngine::ToggleStatAudioStreaming));
 	EngineStats.Add(FEngineStatFuncs(TEXT("STAT_SoundModulators"), TEXT("STATCAT_Engine"), FText::GetEmpty(), &UEngine::RenderStatSoundModulators, &UEngine::ToggleStatSoundModulators));
 	EngineStats.Add(FEngineStatFuncs(TEXT("STAT_SoundModulatorsHelp"), TEXT("STATCAT_Engine"), FText::GetEmpty(), nullptr, &UEngine::PostStatSoundModulatorHelp));
 	EngineStats.Add(FEngineStatFuncs(TEXT("STAT_SoundReverb"), TEXT("STATCAT_Engine"), FText::GetEmpty(), &UEngine::RenderStatSoundReverb, nullptr));
@@ -15514,6 +15515,12 @@ bool UEngine::ToggleStatSoundCues(UWorld* World, FCommonViewportClient* Viewport
 #endif // !ENABLE_AUDIO_DEBUG
 }
 
+bool UEngine::ToggleStatAudioStreaming(UWorld* World, FCommonViewportClient* ViewportClient, const TCHAR* Stream /*= nullptr*/)
+{
+	// Noop for now:
+	return true;
+}
+
 // SOUNDMIXES
 bool UEngine::ToggleStatSoundMixes(UWorld* World, FCommonViewportClient* ViewportClient, const TCHAR* Stream)
 {
@@ -15554,6 +15561,12 @@ int32 UEngine::RenderStatSoundWaves(UWorld* World, FViewport* Viewport, FCanvas*
 #endif // !ENABLE_AUDIO_DEBUG
 }
 
+// AUDIOSTREAMING
+int32 UEngine::RenderStatAudioStreaming(UWorld* World, FViewport* Viewport, FCanvas* Canvas, int32 X, int32 Y, const FVector* ViewLocation /*= nullptr*/, const FRotator* ViewRotation /*= nullptr*/)
+{
+	return IStreamingManager::Get().GetAudioStreamingManager().RenderStatAudioStreaming(World, Viewport, Canvas, X, Y, ViewLocation, ViewRotation);
+}
+
 // SOUNDCUES
 int32 UEngine::RenderStatSoundCues(UWorld* World, FViewport* Viewport, FCanvas* Canvas, int32 X, int32 Y, const FVector* ViewLocation, const FRotator* ViewRotation)
 {
diff --git a/Engine/Source/Runtime/Engine/Private/VorbisAudioInfo.cpp b/Engine/Source/Runtime/Engine/Private/VorbisAudioInfo.cpp
--- a/Engine/Source/Runtime/Engine/Private/VorbisAudioInfo.cpp
+++ b/Engine/Source/Runtime/Engine/Private/VorbisAudioInfo.cpp
@@ -98,10 +98,8 @@ FVorbisAudioInfo::FVorbisAudioInfo()
 	, SrcBufferDataSize(0)
 	, BufferOffset(0)
 	, CurrentBufferChunkOffset(0)
-	, CurrentStreamingChunkData(nullptr)
 	, CurrentStreamingChunkIndex(INDEX_NONE)
 	, NextStreamingChunkIndex(0)
-	, CurrentStreamingChunksSize(0)
 	, bHeaderParsed(false)
 {
 	// Make sure we have properly allocated a VFWrapper
@@ -194,17 +192,16 @@ size_t FVorbisAudioInfo::ReadStreaming(void *Ptr, uint32 Size )
 
 	while (NumBytesRead < Size)
 	{
-		if (!CurrentStreamingChunkData || CurrentStreamingChunkIndex != NextStreamingChunkIndex)
+		const uint8* CurrentStreamingChunkData = nullptr;
+		uint32 CurrentStreamingChunksSize = 0;
+
+		if (CurrentStreamingChunkIndex != NextStreamingChunkIndex)
 		{
 			CurrentStreamingChunkIndex = NextStreamingChunkIndex;
-			CurrentStreamingChunkData = IStreamingManager::Get().GetAudioStreamingManager().GetLoadedChunk(StreamingSoundWave, CurrentStreamingChunkIndex);
-			if (CurrentStreamingChunkData)
-			{
-				check(CurrentStreamingChunkIndex < StreamingSoundWave->RunningPlatformData->Chunks.Num());
-				CurrentStreamingChunksSize = StreamingSoundWave->RunningPlatformData->Chunks[CurrentStreamingChunkIndex].AudioDataSize;
-				CurrentBufferChunkOffset = 0;
-			}
+			CurrentBufferChunkOffset = 0;
 		}
+		
+		GetChunkPtr(CurrentStreamingChunkData, CurrentStreamingChunksSize);
 
 		// No chunk data -- either looping or something else happened with stream
 		if (!CurrentStreamingChunkData)
@@ -212,9 +209,11 @@ size_t FVorbisAudioInfo::ReadStreaming(void *Ptr, uint32 Size )
 			return NumBytesRead;
 		}
 
+		check(CurrentStreamingChunksSize >= CurrentBufferChunkOffset);
 		// How many bytes left in the current chunk
 		uint32 BytesLeftInCurrentChunk = CurrentStreamingChunksSize - CurrentBufferChunkOffset;
 
+		check(Size >= NumBytesRead);
 		// How many more bytes we want to read
 		uint32 NumBytesLeftToRead = Size - NumBytesRead;
 
@@ -541,6 +540,8 @@ bool FVorbisAudioInfo::StreamCompressedInfoInternal(USoundWave* Wave, struct FSo
 
 	FScopeLock ScopeLock(&VorbisCriticalSection);
 
+	TotalNumChunks = Wave->GetNumChunks();
+
 	if (!VFWrapper)
 	{
 		UE_LOG(LogAudio, Error, TEXT("FVorbisAudioInfo::StreamCompressedInfoInternal failed due to no vorbis wrapper for sound '%s'."), *Wave->GetName());
@@ -569,6 +570,29 @@ bool FVorbisAudioInfo::StreamCompressedInfoInternal(USoundWave* Wave, struct FSo
 	return bHeaderParsed;
 }
 
+void FVorbisAudioInfo::GetChunkPtr(const uint8*& OutPtr, uint32& OutSize)
+{
+	check(StreamingSoundWave);
+
+	if (CurrentStreamingChunkIndex == 0)
+	{
+		TArrayView<const uint8> ZerothChunk = StreamingSoundWave->GetZerothChunk();
+		OutPtr = ZerothChunk.GetData();
+		OutSize = ZerothChunk.Num();
+	}
+	else if (((uint32)CurrentStreamingChunkIndex) < StreamingSoundWave->GetNumChunks())
+	{
+		CurrentStreamingChunkHandle = IStreamingManager::Get().GetAudioStreamingManager().GetLoadedChunk(StreamingSoundWave, CurrentStreamingChunkIndex);
+		OutPtr = CurrentStreamingChunkHandle.GetData();
+		OutSize = CurrentStreamingChunkHandle.Num();
+	}
+	else
+	{
+		OutPtr = nullptr;
+		OutSize = 0;
+	}
+}
+
 bool FVorbisAudioInfo::StreamCompressedData(uint8* InDestination, bool bLooping, uint32 BufferSize)
 {
 	if (!bDllLoaded)
@@ -603,7 +627,9 @@ bool FVorbisAudioInfo::StreamCompressedData(uint8* InDestination, bool bLooping,
 		{
 			// if we read 0 bytes or there was an error, instead of assuming we looped, lets write out zero's.
 			// this means that the chunk wasn't loaded in time
-			if (NextStreamingChunkIndex < StreamingSoundWave->RunningPlatformData->Chunks.Num())
+			const bool bHitEOF = ((uint32) CurrentStreamingChunkIndex) >= StreamingSoundWave->GetNumChunks();
+
+			if (!bHitEOF)
 			{
 				// zero out the rest of the buffer
 				FMemory::Memzero(InDestination, BufferSize);
diff --git a/Engine/Source/Runtime/Engine/Public/ADPCMAudioInfo.h b/Engine/Source/Runtime/Engine/Public/ADPCMAudioInfo.h
--- a/Engine/Source/Runtime/Engine/Public/ADPCMAudioInfo.h
+++ b/Engine/Source/Runtime/Engine/Public/ADPCMAudioInfo.h
@@ -10,6 +10,7 @@
 #include "Audio.h"
 #include "AudioDecompress.h"
 #include "Sound/SoundWave.h"
+#include "ContentStreaming.h"
 
 #define NUM_ADAPTATION_TABLE 16
 #define NUM_ADAPTATION_COEFF 7
@@ -107,6 +108,7 @@ public:
 	virtual bool SupportsStreaming() const override {return true;}
 	virtual bool StreamCompressedInfoInternal(USoundWave* Wave, struct FSoundQualityInfo* QualityInfo) override;
 	virtual bool StreamCompressedData(uint8* Destination, bool bLooping, uint32 BufferSize) override;
+
 	virtual int32 GetCurrentChunkIndex() const override
 	{
 		return CurrentChunkIndex;
@@ -114,6 +116,26 @@ public:
 	virtual int32 GetCurrentChunkOffset() const override {return CurrentChunkBufferOffset;}
 	// End of ICompressedAudioInfo Interface
 
+	uint16 GetFormatTag()
+	{
+		if (WaveInfo.pFormatTag)
+		{
+			return *WaveInfo.pFormatTag;
+		}
+		else
+		{
+			return 0;
+		}
+	}
+
+	int32 GetNumChannels()
+	{
+		return NumChannels;
+	}
+
+private:
+	void GetChunkPtr(const uint8* &CompressedChunkPtr, uint32 &CompressedChunkSize);
+
 	FWaveModInfo WaveInfo;
 	const uint8*	SrcBufferData;
 	uint32			SrcBufferDataSize;
@@ -130,12 +152,11 @@ public:
 	uint32			CurrentUncompressedBlockSampleIndex;	// This is the sample index within the current uncompressed block data
 	uint32			CurrentChunkIndex;				// This is the index that is currently being used, needed by streaming engine to make sure it stays loaded and the next chunk gets preloaded
 	uint32			CurrentChunkBufferOffset;		// This is this byte offset within the current chunk, used by streaming engine to prioritize a load if more then half way through current chunk
-	uint32			CurrentChunkDataSize;			// The size of the current chunk, the first chunk is bigger to accommodate the header info
 	uint32			TotalSamplesStreamed;			// The number of samples streamed so far (per channel)
 	uint32			TotalSamplesPerChannel;			// Number of samples per channel, used to detect when an audio waveform has ended
 	uint32			SamplesPerBlock;				// The number of samples per block
 	uint32			FirstChunkSampleDataOffset;		// The size of the header in the first chunk, used to skip over it when looping or starting the sample over
-	const uint8*	CurCompressedChunkData;			// A pointer to the current chunk of data
+	FAudioChunkHandle	CurrentCompressedChunkData;			// A pointer to the current chunk of data
 
 	uint32			CurrentCompressedBlockIndex;		// For non disk streaming - the current compressed block in the compressed source data
 	uint32			TotalCompressedBlocksPerChannel;	// For non disk streaming - the total number of compressed blocks per channel
diff --git a/Engine/Source/Runtime/Engine/Public/AudioCompressionSettingsUtils.h b/Engine/Source/Runtime/Engine/Public/AudioCompressionSettingsUtils.h
--- a/Engine/Source/Runtime/Engine/Public/AudioCompressionSettingsUtils.h
+++ b/Engine/Source/Runtime/Engine/Public/AudioCompressionSettingsUtils.h
@@ -5,6 +5,7 @@
 #include "IAudioExtensionPlugin.h"
 #include "AudioPluginUtilities.h"
 #include "AudioCompressionSettings.h"
+#include "AudioStreamingCache.h"
 
 class ENGINE_API FPlatformCompressionUtilities
 {
@@ -19,7 +20,19 @@ public:
 
 	static int32 GetQualityIndexOverrideForCurrentPlatform();
 
-	static const FPlatformAudioCookOverrides* GetCookOverridesForCurrentPlatform();
+	static void RecacheCookOverrides();
+
+	static const FPlatformAudioCookOverrides* GetCookOverridesForCurrentPlatform(bool bForceRecache = false);
+
+	static bool IsCurrentPlatformUsingStreamCaching();
+
+	static const FAudioStreamCachingSettings& GetStreamCachingSettingsForCurrentPlatform();
+
+	/** This is used at runtime to initialize FCachedAudioStreamingManager. */
+	static FCachedAudioStreamingManagerParams BuildCachedStreamingManagerParams();
+
+	/** This is used at runtime in BuildCachedStreamingManagerParams, as well as cooktime in FStreamedAudioCacheDerivedDataWorker::BuildStreamedAudio to split compressed audio.  */
+	static uint32 GetChunkSizeForCookOverrides(const FPlatformAudioCookOverrides* InCompressionOverrides);
 
 private:
 	static const FPlatformRuntimeAudioCompressionOverrides* GetRuntimeCompressionOverridesForCurrentPlatform();
diff --git a/Engine/Source/Runtime/Engine/Public/AudioDecompress.h b/Engine/Source/Runtime/Engine/Public/AudioDecompress.h
--- a/Engine/Source/Runtime/Engine/Public/AudioDecompress.h
+++ b/Engine/Source/Runtime/Engine/Public/AudioDecompress.h
@@ -12,6 +12,7 @@
 #include "Sound/SoundWave.h"
 #include "Misc/ScopeLock.h"
 #include "HAL/LowLevelMemTracker.h"
+#include "ContentStreaming.h"
 
 // 186ms of 44.1KHz data
 // 372ms of 22KHz data
@@ -178,7 +179,7 @@ public:
 	ENGINE_API virtual void SeekToTime(const float SeekTime) override {};
 	ENGINE_API virtual void ExpandFile(uint8* DstBuffer, struct FSoundQualityInfo* QualityInfo) override;
 	ENGINE_API virtual void EnableHalfRate(bool HalfRate) override {};
-	virtual uint32 GetSourceBufferSize() const override { return SrcBufferDataSize; }
+	virtual uint32 GetSourceBufferSize() const override { return SrcBufferHandle.Num(); }
 	virtual bool UsesVorbisChannelOrdering() const override { return false; }
 	virtual int GetStreamBufferSize() const override { return  MONO_PCM_BUFFER_SIZE; }
 	virtual bool SupportsStreaming() const override { return true; }
@@ -245,16 +246,20 @@ protected:
 	*/
 	uint32	ZeroBuffer(uint8* Destination, uint32 BufferSize);
 
+	// this returns a pointer to the current chunk and the size of the audio in streaming situations,
+	// or a pointer to the entire audio buffer for non-streaming situations.
+	void GetChunkPtr(const uint8*& OutPtr, uint32& OutSize);
+
 	/** bool set before ParseHeader. Whether we are streaming a file or not. */
 	bool bIsStreaming;
 	/** Ptr to the current streamed chunk. */
-	const uint8* SrcBufferData;
-	/** Size of the current streamed chunk. */
-	uint32 SrcBufferDataSize;
+	FAudioChunkHandle SrcBufferHandle;
 	/** What byte we're currently reading in the streamed chunk. */
 	uint32 SrcBufferOffset;
 	/** Where the actual audio data starts in the current streamed chunk. Accounts for header offset. */
 	uint32 AudioDataOffset;
+	/** The chunk index where the audio data begins. */
+	uint32 AudioDataOffsetChunkIndex;
 	/** Sample rate of the source file */
 	uint16 SampleRate;
 	/** The total sample count of the source file. */
@@ -281,6 +286,9 @@ protected:
 	bool bPrintChunkFailMessage;
 	/** Number of bytes of padding used, overridden in some implementations. Defaults to 0. */
 	uint32 SrcBufferPadding;
+
+	const uint8* NonStreamingAudio;
+	uint32 NonStreamingAudioSize;
 };
 
 
@@ -483,6 +491,11 @@ public:
 	void EnsureCompletion(bool bDoWorkOnThisThreadIfNotStarted = true)
 	{
 		FScopeLock Lock(&CritSect);
+
+		if (!bDoWorkOnThisThreadIfNotStarted)
+		{
+			Task->Cancel();
+		}
 		Task->EnsureCompletion(bDoWorkOnThisThreadIfNotStarted);
 	}
 
diff --git a/Engine/Source/Runtime/Engine/Public/AudioStreaming.h b/Engine/Source/Runtime/Engine/Public/AudioStreaming.h
--- a/Engine/Source/Runtime/Engine/Public/AudioStreaming.h
+++ b/Engine/Source/Runtime/Engine/Public/AudioStreaming.h
@@ -14,11 +14,12 @@ AudioStreaming.h: Definitions of classes used for audio streaming.
 #include "ContentStreaming.h"
 #include "Async/AsyncWork.h"
 #include "Async/AsyncFileHandle.h"
+#include "HAL/ThreadSafeBool.h"
 
 class FSoundSource;
 class USoundWave;
 class ICompressedAudioInfo;
-struct FAudioStreamingManager;
+struct FLegacyAudioStreamingManager;
 struct FWaveInstance;
 
 /** Lists possible states used by Thread-safe counter. */
@@ -46,7 +47,8 @@ public:
 		const FString& InDerivedDataKey,
 		void* InDestChunkData,
 		int32 InChunkSize,
-		FThreadSafeCounter* InThreadSafeCounter
+		FThreadSafeCounter* InThreadSafeCounter,
+		TFunction<void(bool)> InOnLoadComplete
 		);
 	
 	/**
@@ -78,6 +80,8 @@ private:
 	bool bRequestFailed;
 	/** Thread-safe counter to decrement when data has been copied. */
 	FThreadSafeCounter* ThreadSafeCounter;
+	/** This function is called when the load is completed */
+	TFunction<void(bool)> OnLoadCompleted;
 };
 
 /** Async task to stream chunks from the derived data cache. */
@@ -136,7 +140,7 @@ struct FStreamingWaveData final
 	 *
 	 * @param SoundWave	The SoundWave we are managing
 	 */
-	bool Initialize(USoundWave* SoundWave, FAudioStreamingManager* InStreamingManager);
+	bool Initialize(USoundWave* SoundWave, FLegacyAudioStreamingManager* InStreamingManager);
 
 	/**
 	 * Updates the streaming status of the sound wave and performs finalization when appropriate. The function returns
@@ -218,7 +222,7 @@ public:
 #endif // #if WITH_EDITORONLY_DATA
 
 	/** Ptr to owning audio streaming manager. */
-	FAudioStreamingManager* AudioStreamingManager;
+	FLegacyAudioStreamingManager* AudioStreamingManager;
 };
 
 /** Struct used to store results of an async file load. */
@@ -240,16 +244,15 @@ struct FASyncAudioChunkLoadResult
 	{}
 };
 
-
 /**
 * Streaming manager dealing with audio.
 */
-struct FAudioStreamingManager : public IAudioStreamingManager
+struct FLegacyAudioStreamingManager : public IAudioStreamingManager
 {
 	/** Constructor, initializing all members */
-	FAudioStreamingManager();
+	FLegacyAudioStreamingManager();
 
-	virtual ~FAudioStreamingManager();
+	virtual ~FLegacyAudioStreamingManager();
 
 	// IStreamingManager interface
 	virtual void UpdateResourceStreaming( float DeltaTime, bool bProcessEverything=false ) override;
@@ -273,7 +276,10 @@ struct FAudioStreamingManager : public IAudioStreamingManager
 	virtual void AddStreamingSoundSource(FSoundSource* SoundSource) override;
 	virtual void RemoveStreamingSoundSource(FSoundSource* SoundSource) override;
 	virtual bool IsManagedStreamingSoundSource(const FSoundSource* SoundSource) const override;
-	virtual const uint8* GetLoadedChunk(const USoundWave* SoundWave, uint32 ChunkIndex, uint32* OutChunkSize = NULL) const override;
+	virtual bool RequestChunk(USoundWave* SoundWave, uint32 ChunkIndex, TFunction<void(EAudioChunkLoadResult)> OnLoadCompleted) override;
+	virtual FAudioChunkHandle GetLoadedChunk(const USoundWave* SoundWave, uint32 ChunkIndex, bool bBlockForLoad = false) const override;
+	virtual uint64 TrimMemory(uint64 NumBytesToFree) override;
+	virtual int32 RenderStatAudioStreaming(UWorld* World, FViewport* Viewport, FCanvas* Canvas, int32 X, int32 Y, const FVector* ViewLocation, const FRotator* ViewRotation) override;
 	// End IAudioStreamingManager interface
 
 	/** Called when an async callback is made on an async loading audio chunk request. */
@@ -284,6 +290,10 @@ struct FAudioStreamingManager : public IAudioStreamingManager
 
 protected:
 
+	/** Unused by the legacy audio streaming manager. */
+	virtual void AddReferenceToChunk(const FAudioChunkHandle& InHandle) override {};
+	virtual void RemoveReferenceToChunk(const FAudioChunkHandle& InHandle) override {};
+
 	/**
 	 * Gets Wave request associated with a specific wave
 	 *
diff --git a/Engine/Source/Runtime/Engine/Public/AudioStreamingCache.h b/Engine/Source/Runtime/Engine/Public/AudioStreamingCache.h
new file mode 100644
--- /dev/null
+++ b/Engine/Source/Runtime/Engine/Public/AudioStreamingCache.h
@@ -0,0 +1,339 @@
+// Copyright 1998-2019 Epic Games, Inc. All Rights Reserved.
+
+/*=============================================================================
+AudioStreaming.h: Definitions of classes used for audio streaming.
+=============================================================================*/
+
+#pragma once
+
+#include "CoreMinimal.h"
+#include "HAL/ThreadSafeCounter.h"
+#include "Containers/IndirectArray.h"
+#include "Containers/Queue.h"
+#include "Stats/Stats.h"
+#include "ContentStreaming.h"
+#include "Async/AsyncWork.h"
+#include "Async/AsyncFileHandle.h"
+#include "HAL/ThreadSafeBool.h"
+#include "AudioStreaming.h"
+
+
+#define DEBUG_STREAM_CACHE !UE_BUILD_SHIPPING
+
+// Basic fixed-size LRU cache for retaining chunks of compressed audio data.
+class FAudioChunkCache
+{
+public:
+	struct FChunkKey
+	{
+		USoundWave* SoundWave = nullptr;
+		FName SoundWaveName = FName();
+		uint32 ChunkIndex = INDEX_NONE;
+
+		inline bool operator==(const FChunkKey& Other) const
+		{
+			return (SoundWaveName == Other.SoundWaveName) && (ChunkIndex == Other.ChunkIndex);
+		}
+	};
+
+	FAudioChunkCache(uint32 InMaxChunkSize, uint32 NumChunks);
+	
+	~FAudioChunkCache();
+
+	// Places chunk in cache, or puts this chunk back at the top of the cache if it's already loaded. Returns false on failure.
+	bool AddOrTouchChunk(const FChunkKey& InKey, TFunction<void(EAudioChunkLoadResult) > OnLoadCompleted);
+
+	// Returns the chunk asked for, or an empty TArrayView if that chunk is not loaded.
+	TArrayView<uint8> GetChunk(const FChunkKey& InKey, bool bBlockForLoadCompletion);
+
+	// add an additional reference for a chunk.
+	void AddNewReferenceToChunk(const FChunkKey& InKey);
+	void RemoveReferenceToChunk(const FChunkKey& InKey);
+
+	// Evict all sounds from the cache.
+	void ClearCache();
+
+	// This function will reclaim memory by freeing as many chunks as needed to free BytesToFree.
+	// returns the amount of bytes we were actually able to free.
+	// It's important to note that this will block any chunk requests.
+	uint64 TrimMemory(uint64 BytesToFree);
+
+	// This function will continue to lock until any async file loads are finished.
+	void BlockForAllPendingLoads() const;
+
+	// This function will cancel any in-flight loads and wait for their completion.
+	void CancelAllPendingLoads();
+
+	// Reports the size of this cache's memory pool, in bytes.
+	uint64 ReportCacheSize();
+
+	// Debug tools:
+	// Call this to start enqueing reports on any cache misses to a queue.
+	// This queue will continue to grow until FlushCacheMissLog is called.
+	void BeginLoggingCacheMisses();
+
+	// This will stop enqueueing reports of cache misses.
+	void StopLoggingCacheMisses();
+
+	// When called, flushes the entire queue of cache misses that has accumulated
+	// And prints them to a formatted 
+	FString FlushCacheMissLog();
+
+	// Static helper function to make sure a chunk is withing the bounds of a USoundWave.
+	static bool IsKeyValid(const FChunkKey& InKey);
+
+	const int32 MaxChunkSize;
+
+	// This is for debugging purposes only. Prints the elements in the cache from most recently used to least.
+	// Returns the dimensions of this debug log so that multiple caches can be tiled across the screen.
+	TPair<int, int> DebugDisplay(UWorld* World, FViewport* Viewport, FCanvas* Canvas, int32 X, int32 Y, const FVector* ViewLocation, const FRotator* ViewRotation) const;
+
+private:
+
+#if DEBUG_STREAM_CACHE
+	// This struct lets us breadcrumb debug information.
+	struct FCacheElementDebugInfo
+	{
+		// The total number of chunks in the sound wave.
+		int32 NumTotalChunks;
+
+		// Number of times this chunk was requested during its time in the cache.
+		int32 NumTimesTouched;
+
+		double TimeLoadStarted;
+		// Amount of time spent loading the audio file.
+		double TimeToLoad;
+
+		// This allows every call to AddOrTouchChunk to report how far down the cache
+		// it was when it was needed, or INDEX_NONE if it wasn't in the cache when it was needed.
+		TArray<int32> PreviousLocationsBeforeBeingTouched;
+
+		FCacheElementDebugInfo()
+			: NumTotalChunks(0)
+			, NumTimesTouched(0)
+			, TimeLoadStarted(0.0)
+			, TimeToLoad(0.0)
+		{
+		}
+	};
+#endif
+
+	// Struct containing a single element in our LRU Cache.  
+	struct FCacheElement
+	{
+		FChunkKey Key;
+		TArray<uint8> ChunkData;
+		uint32 ChunkDataSize;
+		FCacheElement* MoreRecentElement;
+		FCacheElement* LessRecentElement;
+		uint32 CacheIndex;
+
+		FThreadSafeBool bIsLoaded;
+		
+		// How many disparate consumers have called GetLoadedChunk.
+		FThreadSafeCounter NumConsumers;
+
+#if WITH_EDITORONLY_DATA
+		TUniquePtr<FAsyncStreamDerivedChunkTask> DDCTask;
+#endif
+
+		TUniquePtr<IAsyncReadFileHandle> FileHandle;
+		// Handle to our async read request operation.
+		TUniquePtr<IAsyncReadRequest> ReadRequest;
+
+#if DEBUG_STREAM_CACHE
+		FCacheElementDebugInfo DebugInfo;
+#endif
+
+		FCacheElement(uint32 MaxChunkSize, uint32 InCacheIndex)
+			: ChunkDataSize(0)
+			, MoreRecentElement(nullptr)
+			, LessRecentElement(nullptr)
+			, CacheIndex(InCacheIndex)
+			, bIsLoaded(false)
+		{
+		}
+
+		void WaitForAsyncLoadCompletion(bool bCancel)
+		{
+#if WITH_EDITORONLY_DATA
+			if (DDCTask.IsValid() && !DDCTask->IsDone())
+			{
+				DDCTask->Cancel();
+				DDCTask->EnsureCompletion(false);
+			}
+#endif
+
+			if (ReadRequest.IsValid())
+			{
+				if (bCancel)
+				{
+					ReadRequest->Cancel();
+				}
+				
+				ReadRequest->WaitCompletion();
+				ReadRequest.Reset();
+			}
+		}
+
+		bool IsLoadInProgress() const
+		{
+			return !bIsLoaded;
+		}
+
+		bool IsInUse() const
+		{
+			return NumConsumers.GetValue() > 0;
+		}
+
+		bool CanEvictChunk() const
+		{
+			return !IsInUse() && !IsLoadInProgress();
+		}
+
+		~FCacheElement()
+		{
+			WaitForAsyncLoadCompletion(true);
+			checkf(NumConsumers.GetValue() == 0, TEXT("Tried to destroy streaming cache while the cached data was in use!"));
+		}
+	};
+
+	// Our actual memory pool.
+	TArray<FCacheElement> CachePool;
+	FCacheElement* MostRecentElement;
+	FCacheElement* LeastRecentElement;
+
+	// This is incremented on every call of InsertChunk until we hit CachePool.Num().
+	int32 ChunksInUse;
+
+	// Number of async load operations we have currently in flight.
+	FThreadSafeCounter NumberOfLoadsInFlight;
+
+	// Critical section: only used when we are modifying element positions in the cache. This only happens in TouchElement, EvictLeastRecentChunk, and TrimMemory.
+	// Individual cache elements should be thread safe to access.
+	FCriticalSection CacheMutationCriticalSection;
+
+	// This struct is used for logging cache misses.
+	struct FCacheMissInfo
+	{
+		FName SoundWaveName;
+		uint32 ChunkIndex;
+		uint32 TotalChunksInWave;
+		bool bBlockedForLoad;
+	};
+
+	// This queue is pushed to anytime GetChunk fails to get the chunk. and bLogCacheMisses is true.
+	TQueue<FCacheMissInfo> CacheMissQueue;
+
+	// This is set to true when BeginLoggingCacheMisses is called. 
+	bool bLogCacheMisses;
+
+	// Returns cached element if it exists in our cache, nullptr otherwise.
+	FCacheElement* FindElementForKey(const FChunkKey& InKey);
+
+	// Puts this element at the front of the linked list.
+	void TouchElement(FCacheElement* InElement);
+
+	// Inserts a new element into the cache, potentially evicting the oldest element in the cache.
+	FCacheElement* InsertChunk(const FChunkKey& InKey);
+
+	// This is called once we have more than one chunk in our cache:
+	void SetUpLeastRecentChunk();
+
+	// Returns the least recent chunk and fixes up the linked list accordingly.
+	FCacheElement* EvictLeastRecentChunk();
+
+	void KickOffAsyncLoad(FCacheElement* CacheElement, const FChunkKey& InKey, TFunction<void(EAudioChunkLoadResult)> OnLoadCompleted);
+};
+
+// This is used to sort the cache array from smallest chunk size to biggest.
+inline bool operator< (const FAudioChunkCache& Element1, const FAudioChunkCache& Element2)
+{
+	return Element1.MaxChunkSize < Element2.MaxChunkSize;
+}
+
+struct FCachedAudioStreamingManagerParams
+{
+	struct FCacheDimensions
+	{
+		// The max size, in bytes, of a single chunk of compressed audio.
+		// During cook, compressed audio assets will be chunked based on this amount.
+		int32 MaxChunkSize;
+
+		// The maximum number of elements stored in a single cache before it is evicted.
+		// At runtime, this will be clamped to ensure that it is greater than the amount of
+		// sources that can be playing simultaneously.
+		int32 NumElements;
+	};
+
+	// Most use cases will only use a single cache, but applications can optionally
+	// use multiple LRU caches to reduce churn for specific types of sounds.
+	// For example, an application can have
+	// a cache for short sounds with room for many elements, and a separate cache
+	// for longer sounds with fewer elements.
+	TArray<FCacheDimensions> Caches;
+};
+
+/**
+ * This implementation of the audio streaming manager uses an internal LRU cache (or in more advanced applications, a bank of parallel LRU caches)
+ */
+struct FCachedAudioStreamingManager : public IAudioStreamingManager
+{
+public:
+	/** Constructor, initializing all members */
+	FCachedAudioStreamingManager(const FCachedAudioStreamingManagerParams& InitParams);
+
+	virtual ~FCachedAudioStreamingManager();
+
+	// IStreamingManager interface
+	virtual void UpdateResourceStreaming( float DeltaTime, bool bProcessEverything=false ) override;
+	virtual int32 BlockTillAllRequestsFinished( float TimeLimit = 0.0f, bool bLogResults = false ) override;
+	virtual void CancelForcedResources() override;
+	virtual void NotifyLevelChange() override;
+	virtual void SetDisregardWorldResourcesForFrames( int32 NumFrames ) override;
+	virtual void AddLevel( class ULevel* Level ) override;
+	virtual void RemoveLevel( class ULevel* Level ) override;
+	virtual void NotifyLevelOffset( class ULevel* Level, const FVector& Offset ) override;
+	// End IStreamingManager interface
+
+	// IAudioStreamingManager interface
+	virtual void AddStreamingSoundWave(USoundWave* SoundWave) override;
+	virtual void RemoveStreamingSoundWave(USoundWave* SoundWave) override;
+	virtual void AddDecoder(ICompressedAudioInfo* CompressedAudioInfo) override;
+	virtual void RemoveDecoder(ICompressedAudioInfo* CompressedAudioInfo) override;
+	virtual bool IsManagedStreamingSoundWave(const USoundWave* SoundWave) const override;
+	virtual bool IsStreamingInProgress(const USoundWave* SoundWave) override;
+	virtual bool CanCreateSoundSource(const FWaveInstance* WaveInstance) const override;
+	virtual void AddStreamingSoundSource(FSoundSource* SoundSource) override;
+	virtual void RemoveStreamingSoundSource(FSoundSource* SoundSource) override;
+	virtual bool IsManagedStreamingSoundSource(const FSoundSource* SoundSource) const override;
+	virtual bool RequestChunk(USoundWave* SoundWave, uint32 ChunkIndex, TFunction<void(EAudioChunkLoadResult)> OnLoadCompleted) override;
+	virtual FAudioChunkHandle GetLoadedChunk(const USoundWave* SoundWave, uint32 ChunkIndex, bool bBlockForLoad = false) const override;
+	virtual uint64 TrimMemory(uint64 NumBytesToFree) override;
+	virtual int32 RenderStatAudioStreaming(UWorld* World, FViewport* Viewport, FCanvas* Canvas, int32 X, int32 Y, const FVector* ViewLocation, const FRotator* ViewRotation) override;
+	// End IAudioStreamingManager interface
+
+protected:
+
+	// These are used to reference count consumers of audio chunks.
+	virtual void AddReferenceToChunk(const FAudioChunkHandle& InHandle) override;
+	virtual void RemoveReferenceToChunk(const FAudioChunkHandle& InHandle) override;
+
+	/**
+	 * Returns which cache this sound wave should be in,
+	 * based on the size of this sound wave's chunk,
+	 * or nullptr if MemoryLoadOnDemand is disabled.
+	 */
+	FAudioChunkCache* GetCacheForWave(const USoundWave* InSoundWave) const;
+
+	/**
+	 * Returns the next chunk to kick off a load for, or INDEX_NONE if there is only one chunk to cache.
+	 */
+	int32 GetNextChunkIndex(const USoundWave* InSoundWave, uint32 CurrentChunkIndex) const;
+
+	/** Audio chunk caches. These are set up on initialization. */
+	TArray<FAudioChunkCache> CacheArray;
+
+	
+
+};
diff --git a/Engine/Source/Runtime/Engine/Public/ContentStreaming.h b/Engine/Source/Runtime/Engine/Public/ContentStreaming.h
--- a/Engine/Source/Runtime/Engine/Public/ContentStreaming.h
+++ b/Engine/Source/Runtime/Engine/Public/ContentStreaming.h
@@ -8,6 +8,8 @@
 
 #include "CoreMinimal.h"
 #include "UObject/WeakObjectPtr.h"
+#include "CanvasTypes.h"
+#include "UnrealClient.h"
 
 class AActor;
 class FSoundSource;
@@ -82,14 +84,52 @@ struct FStreamingViewInfo
 	TWeakObjectPtr<AActor> ActorToBoost;
 };
 
+/**
+ * This structure allows audio chunk data to be accessed, and guarantees that the chunk in question will not be deleted
+ * during it's lifecycle.
+ */
+class FAudioChunkHandle
+{
+public:
+	FAudioChunkHandle();
+	FAudioChunkHandle(const FAudioChunkHandle& Other);
+
+	FAudioChunkHandle& operator=(const FAudioChunkHandle& Other);
+
+	~FAudioChunkHandle();
+
+	// gets a pointer to the compressed chunk.
+	const uint8* GetData() const;
+
+	// Returns the num bytes pointed to by GetData().
+	uint32 Num() const;
+
+	// Checks whether this points to a valid compressed chunk.
+	bool IsValid() const;
+
+private:
+	// This constructor should only be called by an implementation of IAudioStreamingManager.
+	FAudioChunkHandle(const uint8* InData, uint32 NumBytes, const USoundWave* InSoundWave, const FName& SoundWaveName, uint32 InChunkIndex);
+
+	const uint8*  CachedData;
+	int32 CachedDataNumBytes;
+
+	const USoundWave* CorrespondingWave;
+	FName CorrespondingWaveName;
+	int32 ChunkIndex;
+
+	friend struct IAudioStreamingManager;
+	friend struct FCachedAudioStreamingManager;
+};
+
 /**
  * Pure virtual base class of a streaming manager.
  */
 struct ENGINE_VTABLE IStreamingManager
 {
 	IStreamingManager()
-	:	NumWantingResources(0)
-	,	NumWantingResourcesCounter(0)
+		: NumWantingResources(0)
+		, NumWantingResourcesCounter(0)
 	{
 	}
 
@@ -113,7 +153,7 @@ struct ENGINE_VTABLE IStreamingManager
 	 * @param DeltaTime				Time since last call in seconds
 	 * @param bProcessEverything	[opt] If true, process all resources with no throttling limits
 	 */
-	ENGINE_API virtual void Tick( float DeltaTime, bool bProcessEverything=false );
+	ENGINE_API virtual void Tick(float DeltaTime, bool bProcessEverything = false);
 
 	/**
 	 * Updates streaming, taking into account all current view infos. Can be called multiple times per frame.
@@ -121,7 +161,7 @@ struct ENGINE_VTABLE IStreamingManager
 	 * @param DeltaTime				Time since last call in seconds
 	 * @param bProcessEverything	[opt] If true, process all resources with no throttling limits
 	 */
-	virtual void UpdateResourceStreaming( float DeltaTime, bool bProcessEverything=false ) = 0;
+	virtual void UpdateResourceStreaming(float DeltaTime, bool bProcessEverything = false) = 0;
 
 	/**
 	 * Streams in/out all resources that wants to and blocks until it's done.
@@ -129,7 +169,7 @@ struct ENGINE_VTABLE IStreamingManager
 	 * @param TimeLimit					Maximum number of seconds to wait for streaming I/O. If zero, uses .ini setting
 	 * @return							Number of streaming requests still in flight, if the time limit was reached before they were finished.
 	 */
-	virtual int32 StreamAllResources( float TimeLimit=0.0f )
+	virtual int32 StreamAllResources(float TimeLimit = 0.0f)
 	{
 		return 0;
 	}
@@ -141,7 +181,7 @@ struct ENGINE_VTABLE IStreamingManager
 	 * @param bLogResults	Whether to dump the results to the log.
 	 * @return				Number of streaming requests still in flight, if the time limit was reached before they were finished.
 	 */
-	virtual int32 BlockTillAllRequestsFinished( float TimeLimit = 0.0f, bool bLogResults = false ) = 0;
+	virtual int32 BlockTillAllRequestsFinished(float TimeLimit = 0.0f, bool bLogResults = false) = 0;
 
 	/**
 	 * Cancels the timed Forced resources (i.e used the Kismet action "Stream In Textures").
@@ -158,7 +198,7 @@ struct ENGINE_VTABLE IStreamingManager
 	 *
 	 * @param RemovalType	What types of views to remove (all or just the normal views)
 	 */
-	void RemoveStreamingViews( ERemoveStreamingViews RemovalType );
+	void RemoveStreamingViews(ERemoveStreamingViews RemovalType);
 
 	/**
 	 * Adds the passed in view information to the static array.
@@ -170,7 +210,7 @@ struct ENGINE_VTABLE IStreamingManager
 	 * @param Duration				How long the streaming system should keep checking this location, in seconds. 0 means just for the next Tick.
 	 * @param InActorToBoost		Optional pointer to an actor who's textures should have their streaming priority boosted
 	 */
-	ENGINE_API void AddViewInformation( const FVector& ViewOrigin, float ScreenSize, float FOVScreenSize, float BoostFactor=1.0f, bool bOverrideLocation=false, float Duration=0.0f, TWeakObjectPtr<AActor> InActorToBoost = NULL );
+	ENGINE_API void AddViewInformation(const FVector& ViewOrigin, float ScreenSize, float FOVScreenSize, float BoostFactor = 1.0f, bool bOverrideLocation = false, float Duration = 0.0f, TWeakObjectPtr<AActor> InActorToBoost = NULL);
 
 	/**
 	 * Queue up view "slave" locations to the streaming system. These locations will be added properly at the next call to AddViewInformation,
@@ -181,10 +221,10 @@ struct ENGINE_VTABLE IStreamingManager
 	 * @param bOverrideLocation		Whether this is an override location, which forces the streaming system to ignore all other locations
 	 * @param Duration				How long the streaming system should keep checking this location, in seconds. 0 means just for the next Tick.
 	 */
-	ENGINE_API void AddViewSlaveLocation( const FVector& SlaveLocation, float BoostFactor=1.0f, bool bOverrideLocation=false, float Duration=0.0f );
+	ENGINE_API void AddViewSlaveLocation(const FVector& SlaveLocation, float BoostFactor = 1.0f, bool bOverrideLocation = false, float Duration = 0.0f);
 
 	/** Don't stream world resources for the next NumFrames. */
-	virtual void SetDisregardWorldResourcesForFrames( int32 NumFrames ) = 0;
+	virtual void SetDisregardWorldResourcesForFrames(int32 NumFrames) = 0;
 
 	/**
 	 * Allows the streaming manager to process exec commands.
@@ -194,39 +234,39 @@ struct ENGINE_VTABLE IStreamingManager
 	 * @param Ar	Output device for feedback
 	 * @return		true if the command was handled
 	 */
-	virtual bool Exec( UWorld* InWorld, const TCHAR* Cmd, FOutputDevice& Ar ) 
+	virtual bool Exec(UWorld* InWorld, const TCHAR* Cmd, FOutputDevice& Ar)
 	{
 		return false;
 	}
 
 	/** Adds a ULevel that has already prepared StreamingData to the streaming manager. */
-	virtual void AddLevel( class ULevel* Level ) = 0;
+	virtual void AddLevel(class ULevel* Level) = 0;
 
 	/** Removes a ULevel from the streaming manager. */
-	virtual void RemoveLevel( class ULevel* Level ) = 0;
+	virtual void RemoveLevel(class ULevel* Level) = 0;
 
 	/**
 	 * Notifies manager that level primitives were shifted
 	 */
-	virtual void NotifyLevelOffset( class ULevel* Level, const FVector& Offset ) = 0;
+	virtual void NotifyLevelOffset(class ULevel* Level, const FVector& Offset) = 0;
 
 	/** Called when a spawned actor is destroyed. */
-	virtual void NotifyActorDestroyed( AActor* Actor )
+	virtual void NotifyActorDestroyed(AActor* Actor)
 	{
 	}
 
 	/** Called when a primitive is detached from an actor or another component. */
-	virtual void NotifyPrimitiveDetached( const UPrimitiveComponent* Primitive )
+	virtual void NotifyPrimitiveDetached(const UPrimitiveComponent* Primitive)
 	{
 	}
 
 	/** Called when a primitive streaming data needs to be updated. */
-	virtual void NotifyPrimitiveUpdated( const UPrimitiveComponent* Primitive )
+	virtual void NotifyPrimitiveUpdated(const UPrimitiveComponent* Primitive)
 	{
 	}
 
 	/**  Called when a primitive streaming data needs to be updated in the last stage of the frame. */
-	virtual void NotifyPrimitiveUpdated_Concurrent( const UPrimitiveComponent* Primitive )
+	virtual void NotifyPrimitiveUpdated_Concurrent(const UPrimitiveComponent* Primitive)
 	{
 	}
 
@@ -237,9 +277,9 @@ struct ENGINE_VTABLE IStreamingManager
 	}
 
 	/** Returns the view info by the specified index. */
-	ENGINE_API const FStreamingViewInfo& GetViewInformation( int32 ViewIndex ) const
+	ENGINE_API const FStreamingViewInfo& GetViewInformation(int32 ViewIndex) const
 	{
-		return CurrentViewInfos[ ViewIndex ];
+		return CurrentViewInfos[ViewIndex];
 	}
 
 	/** Returns the number of resources that currently wants to be streamed in. */
@@ -260,10 +300,14 @@ struct ENGINE_VTABLE IStreamingManager
 	}
 
 	/** Propagates a change to the active lighting scenario. */
-	virtual void PropagateLightingScenarioChange() 
+	virtual void PropagateLightingScenarioChange()
 	{
 	}
 
+#if WITH_EDITOR
+	ENGINE_API virtual void OnAudioStreamingParamsChanged() {};
+#endif
+
 protected:
 
 	/**
@@ -401,6 +445,15 @@ struct IRenderAssetStreamingManager : public IStreamingManager
 	//END: APIs for backward compatibility
 };
 
+enum class EAudioChunkLoadResult : uint8
+{
+	Completed,
+	AlreadyLoaded,
+	Interrupted,
+	ChunkOutOfBounds,
+	CacheBlown
+};
+
 /**
  * Interface to add functions specifically related to audio streaming
  */
@@ -435,14 +488,46 @@ struct IAudioStreamingManager : public IStreamingManager
 	/** Returns true if this is a streaming Sound Source that is managed by the streaming manager. */
 	virtual bool IsManagedStreamingSoundSource(const FSoundSource* SoundSource) const = 0;
 
+	/** Manually prepare a chunk to start playing back. This should only be used when the Load On Demand feature is enabled, and returns false on failure. */
+	virtual bool RequestChunk(USoundWave* SoundWave, uint32 ChunkIndex, TFunction<void(EAudioChunkLoadResult)> OnLoadCompleted) = 0;
+
 	/**
 	 * Gets a pointer to a chunk of audio data
 	 *
 	 * @param SoundWave		SoundWave we want a chunk from
 	 * @param ChunkIndex	Index of the chunk we want
-	 * @return Either the desired chunk or NULL if it's not loaded
+	 * @return a handle to the loaded chunk. Can return a default constructed FAudioChunkHandle if the chunk is not loaded yet.
+	 */
+	virtual FAudioChunkHandle GetLoadedChunk(const USoundWave* SoundWave, uint32 ChunkIndex,  bool bBlockForLoad = false) const = 0;
+
+	/**
+	 * This will start evicting elements from the cache until either hit our target of bytes or run out of chunks we can free.
+	 *
+	 * @param NumBytesToFree	The amount of memory we would like to free, in bytes.
+	 * @return the amount of bytes we managed to free.
 	 */
-	virtual const uint8* GetLoadedChunk(const USoundWave* SoundWave, uint32 ChunkIndex, uint32* OutChunkSize = NULL) const = 0;
+	virtual uint64 TrimMemory(uint64 NumBytesToFree) = 0;
+
+	/**
+	 * Used for rendering debug info:
+	 */
+	virtual int32 RenderStatAudioStreaming(UWorld* World, FViewport* Viewport, FCanvas* Canvas, int32 X, int32 Y, const FVector* ViewLocation, const FRotator* ViewRotation) = 0;
+
+protected:
+	friend FAudioChunkHandle;
+
+	/** This can be called by implementers of IAudioStreamingManager to construct an FAudioChunkHandle using an otherwise inaccessible constructor. */
+	static FAudioChunkHandle BuildChunkHandle(const uint8* InData, uint32 NumBytes, const USoundWave* InSoundWave, const FName& SoundWaveName, uint32 InChunkIndex);
+
+	/**
+	 * This can be used to increment reference counted handles to audio chunks. Called by the copy constructor of FAudioChunkHandle.
+	 */
+	virtual void AddReferenceToChunk(const FAudioChunkHandle& InHandle) = 0;
+
+	/**
+	 * This can be used to decrement reference counted handles to audio chunks. Called by the destructor of FAudioChunkHandle.
+	 */
+	virtual void RemoveReferenceToChunk(const FAudioChunkHandle& InHandle) = 0;
 };
 
 /**
@@ -624,6 +709,10 @@ struct ENGINE_VTABLE FStreamingManagerCollection : public IStreamingManager
 	/** Propagates a change to the active lighting scenario. */
 	void PropagateLightingScenarioChange() override;
 
+#if WITH_EDITOR
+	virtual void OnAudioStreamingParamsChanged() override;
+#endif
+
 protected:
 
 	virtual void AddOrRemoveTextureStreamingManagerIfNeeded(bool bIsInit=false);
@@ -647,5 +736,10 @@ protected:
 
 	/** The animation streaming manager, should always exist */
 	IAnimationStreamingManager* AnimationStreamingManager;
+
+#if WITH_EDITOR
+	// Locks out any audio streaming manager call when we are re-initializing the audio streaming manager.
+	mutable FCriticalSection AudioStreamingManagerCriticalSection;
+#endif
 };
 
diff --git a/Engine/Source/Runtime/Engine/Public/VorbisAudioInfo.h b/Engine/Source/Runtime/Engine/Public/VorbisAudioInfo.h
--- a/Engine/Source/Runtime/Engine/Public/VorbisAudioInfo.h
+++ b/Engine/Source/Runtime/Engine/Public/VorbisAudioInfo.h
@@ -9,6 +9,7 @@
 #include "CoreMinimal.h"
 #include "HAL/ThreadSafeBool.h"
 #include "AudioDecompress.h"
+#include "ContentStreaming.h"
 
 /**
  * Whether to use OggVorbis audio format.
@@ -46,7 +47,8 @@ public:
 
 	/** Emulate read from streaming functionality */
 	size_t			ReadStreaming( void *ptr, uint32 size );
-	int				CloseStreaming( void );
+
+	int				CloseStreaming(void);
 
 	/** Common info and data functions between ReadCompressedInfo/ReadCompressedData and StreamCompressedInfo/StreamCompressedData */
 	/** Can't use ov_callbacks here so we have to pass in a void* for the callbacks instead */
@@ -93,7 +95,19 @@ public:
 	virtual bool SupportsStreaming() const override {return true;}
 	virtual bool StreamCompressedData(uint8* InDestination, bool bLooping, uint32 BufferSize) override;
 	virtual int32 GetCurrentChunkIndex() const override {return CurrentStreamingChunkIndex;}
-	virtual int32 GetCurrentChunkOffset() const override {return BufferOffset % CurrentStreamingChunksSize;}
+
+	virtual int32 GetCurrentChunkOffset() const override
+	{
+		const uint32 ChunkSize = StreamingSoundWave ? StreamingSoundWave->GetSizeOfChunk(CurrentStreamingChunkIndex) : 0;
+		if (ChunkSize)
+		{
+			return CurrentBufferChunkOffset % ChunkSize;
+		}
+		else
+		{
+			return 0;
+		}
+	}
 	// End of ICompressedAudioInfo Interface
 
 protected:
@@ -101,6 +115,8 @@ protected:
 
 private:
 
+	void GetChunkPtr(const uint8*& OutPtr, uint32& OutSize);
+
 	struct FVorbisFileWrapper* VFWrapper;
 	const uint8* SrcBufferData;
 	uint32 SrcBufferDataSize;
@@ -110,10 +126,10 @@ private:
 	/** Critical section used to prevent multiple threads accessing same ogg-vorbis file handles at the same time */
 	FCriticalSection VorbisCriticalSection;
 
-	uint8 const* CurrentStreamingChunkData;
+	FAudioChunkHandle CurrentStreamingChunkHandle;
 	int32 CurrentStreamingChunkIndex;
+	int32 TotalNumChunks;
 	int32 NextStreamingChunkIndex;
-	uint32 CurrentStreamingChunksSize;
 	bool bHeaderParsed;
 };
 #endif
diff --git a/Engine/Source/Runtime/IOS/IOSAudio/Private/IOSAudioBuffer.cpp b/Engine/Source/Runtime/IOS/IOSAudio/Private/IOSAudioBuffer.cpp
--- a/Engine/Source/Runtime/IOS/IOSAudio/Private/IOSAudioBuffer.cpp
+++ b/Engine/Source/Runtime/IOS/IOSAudio/Private/IOSAudioBuffer.cpp
@@ -38,7 +38,7 @@ FIOSAudioSoundBuffer::FIOSAudioSoundBuffer(FIOSAudioDevice* InAudioDevice, USoun
 			return;
 		}
 
-		SoundFormat = static_cast<ESoundFormat>(*DecompressionState->WaveInfo.pFormatTag);
+		SoundFormat = static_cast<ESoundFormat>(DecompressionState->GetFormatTag());
 	}
     else
     {
diff --git a/Engine/Source/Runtime/IOS/IOSAudio/Private/IOSAudioSource.cpp b/Engine/Source/Runtime/IOS/IOSAudio/Private/IOSAudioSource.cpp
--- a/Engine/Source/Runtime/IOS/IOSAudio/Private/IOSAudioSource.cpp
+++ b/Engine/Source/Runtime/IOS/IOSAudio/Private/IOSAudioSource.cpp
@@ -483,7 +483,7 @@ OSStatus FIOSAudioSoundSource::IOSAudioRenderCallback(void* RefCon, AudioUnitRen
 		}
 		else
 		{
-			Source->IOSBuffer->RenderCallbackBufferSize = NumFrames * sizeof(uint16) * Source->IOSBuffer->DecompressionState->NumChannels;
+			Source->IOSBuffer->RenderCallbackBufferSize = NumFrames * sizeof(uint16) * Source->IOSBuffer->DecompressionState->GetNumChannels();
 
 			// Since StreamCompressedData returns interlaced samples we need to decompress all frames(samples) for all channels here so we don't end up decompressing multiple times
 			// Ensure we have enough memory to do this. If needed we could realloc here but that is bad practice inside the audio callback since it has a hard deadline
diff --git a/Engine/Source/Runtime/IOS/IOSRuntimeSettings/Classes/IOSRuntimeSettings.h b/Engine/Source/Runtime/IOS/IOSRuntimeSettings/Classes/IOSRuntimeSettings.h
--- a/Engine/Source/Runtime/IOS/IOSRuntimeSettings/Classes/IOSRuntimeSettings.h
+++ b/Engine/Source/Runtime/IOS/IOSRuntimeSettings/Classes/IOSRuntimeSettings.h
@@ -511,6 +511,13 @@ public:
 	UPROPERTY(config, EditAnywhere, Category = "Audio")
 	FPlatformRuntimeAudioCompressionOverrides CompressionOverrides;
 
+	/** When this is enabled, Actual compressed data will be separated from the USoundWave, and loaded into a cache. */
+	UPROPERTY(GlobalConfig, EditAnywhere, Category = "Audio|CookOverrides", meta = (DisplayName = "Use Stream Caching (Experimental)"))
+	bool bUseAudioStreamCaching;
+
+	/** This determines the max amount of memory that should be used for the cache at any given time. If set low (<= 8 MB), it lowers the size of individual chunks of audio during cook. */
+	UPROPERTY(GlobalConfig, EditAnywhere, Category = "Audio|CookOverrides|Stream Caching", meta = (DisplayName = "Max Cache Size (KB)"))
+	int32 CacheSizeKB;
 
 	UPROPERTY(config, EditAnywhere, Category = "Audio|CookOverrides")
 	bool bResampleForDevice;
