Extra logging for the streaming manager when we run out of playable data

diff --git a/Engine/Source/Runtime/Engine/Private/Animation/AnimationStreaming.cpp b/Engine/Source/Runtime/Engine/Private/Animation/AnimationStreaming.cpp
--- a/Engine/Source/Runtime/Engine/Private/Animation/AnimationStreaming.cpp
+++ b/Engine/Source/Runtime/Engine/Private/Animation/AnimationStreaming.cpp
@@ -197,7 +197,7 @@ void FStreamingAnimationData::BeginPendingRequests(const TArray<uint32>& Indices
 
 		if(!ExistingCompressedData)
 		{
-			check(!ChunkStorage.CompressedAnimData);
+			check(ChunkStorage.CompressedAnimData == nullptr);
 			check(Chunk.BulkData.GetFilename().Len());
 			UE_CLOG(Chunk.BulkData.IsStoredCompressedOnDisk(), LogAnimation, Fatal, TEXT("Package level compression is not supported for streaming animation."));
 			check(!ChunkStorage.IORequest);
@@ -210,14 +210,14 @@ void FStreamingAnimationData::BeginPendingRequests(const TArray<uint32>& Indices
 
 			int64 ChunkSize = Chunk.BulkData.GetBulkDataSize();
 			ChunkStorage.RequestStart = FPlatformTime::Seconds();
-			FPlatformMisc::LowLevelOutputDebugStringf(TEXT("Request Started %.2f\n"), ChunkStorage.RequestStart);
+			UE_LOG(LogAnimation, Log, TEXT("Request Started %s Chunk:%i At:%.3f\n"), *StreamableAnim->GetName(), ChunkIndex, ChunkStorage.RequestStart);
 			FAsyncFileCallBack AsyncFileCallBack =
 				[this, ChunkIndex, ChunkSize](bool bWasCancelled, IAsyncReadRequest* Req)
 			{
 				AnimationStreamingManager->OnAsyncFileCallback(this, ChunkIndex, ChunkSize, Req);
 			};
 
-			FPlatformMisc::LowLevelOutputDebugStringf(TEXT("Loading Stream Chunk %s Chunk:%i Length: %.3f Offset:%i Size:%i File:%i\n"), *StreamableAnim->GetName(), ChunkIndex, Chunk.SequenceLength, Chunk.BulkData.GetBulkDataOffsetInFile(), Chunk.BulkData.GetBulkDataSize(), *Chunk.BulkData.GetFilename());
+			UE_LOG(LogAnimation, Verbose, TEXT("Loading Stream Chunk %s Chunk:%i Length: %.3f Offset:%i Size:%i File:%i\n"), *StreamableAnim->GetName(), ChunkIndex, Chunk.SequenceLength, Chunk.BulkData.GetBulkDataOffsetInFile(), Chunk.BulkData.GetBulkDataSize(), *Chunk.BulkData.GetFilename());
 			ChunkStorage.IORequest = IORequestHandle->ReadRequest(Chunk.BulkData.GetBulkDataOffsetInFile(), Chunk.BulkData.GetBulkDataSize(), AsyncIOPriority, &AsyncFileCallBack);
 			if (!ChunkStorage.IORequest)
 			{
@@ -281,7 +281,7 @@ void FStreamingAnimationData::FreeLoadedChunk(FLoadedAnimationChunk& LoadedChunk
 
 	if (LoadedChunk.bOwnsCompressedData)
 	{
-		delete LoadedChunk.CompressedAnimData;
+		delete LoadedChunk.CompressedAnimData.Load();
 	}
 
 	LoadedChunk.CompressedAnimData = NULL;
@@ -317,17 +317,29 @@ void FAnimationStreamingManager::OnAsyncFileCallback(FStreamingAnimationData* St
 		int32 LoadedChunkIndex = StreamingAnimData->LoadedChunks.IndexOfByPredicate([ChunkIndex](const FLoadedAnimationChunk& Chunk) { return Chunk.Index == ChunkIndex; });
 		FLoadedAnimationChunk& ChunkStorage = StreamingAnimData->LoadedChunks[LoadedChunkIndex];
 
-		checkf(!ChunkStorage.CompressedAnimData, TEXT("Chunk storage already has data. (0x%p)"), ChunkStorage.CompressedAnimData);
+		checkf(ChunkStorage.CompressedAnimData == nullptr, TEXT("Chunk storage already has data. (0x%p)"), ChunkStorage.CompressedAnimData.Load());
 
-		ChunkStorage.CompressedAnimData = new FCompressedAnimSequence();
+		FCompressedAnimSequence* NewCompressedData = new FCompressedAnimSequence();
 
 		TArrayView<const uint8> MemView(Mem, ReadSize);
 		FMemoryReaderView Reader(MemView);
 
 		UAnimStreamable* Anim = StreamingAnimData->StreamableAnim;
-		ChunkStorage.CompressedAnimData->SerializeCompressedData(Reader, false, Anim, Anim->GetSkeleton(), Anim->CurveCompressionSettings);
+		NewCompressedData->SerializeCompressedData(Reader, false, Anim, Anim->GetSkeleton(), Anim->CurveCompressionSettings);
 
-		FPlatformMisc::LowLevelOutputDebugStringf(TEXT("Request Finished %.2f\nAnim Chunk Streamed %.4f\n"), FPlatformTime::Seconds(), FPlatformTime::Seconds() - ChunkStorage.RequestStart);
+		const double CurrentTime = FPlatformTime::Seconds();
+		const double RequestDuration = CurrentTime - ChunkStorage.RequestStart;
+
+		const float SequenceLength = StreamingAnimData->StreamableAnim->GetRunningPlatformData().Chunks[ChunkIndex].SequenceLength;
+		if (SequenceLength < RequestDuration)
+		{
+			UE_LOG(LogAnimation, Warning, TEXT("Streaming anim loading slower than needed ChunkIndex %i - Length: %.3f Load Duration:%.3f Anim:%s\n"), ChunkIndex, SequenceLength, RequestDuration, *StreamingAnimData->StreamableAnim->GetName());
+		}
+
+		ChunkStorage.CompressedAnimData = NewCompressedData;
+		ChunkStorage.RequestStart = -2.0; //Signify we have finished loading
+
+		UE_LOG(LogAnimation, Log, TEXT("Request Finished %.2f\n Anim Chunk Streamed %.4f\n"), CurrentTime, RequestDuration);
 	}
 }
 
@@ -457,8 +469,23 @@ const FCompressedAnimSequence* FAnimationStreamingManager::GetLoadedChunk(const
 
 		if (AnimData->LoadedChunkIndices.Contains(ChunkIndex))
 		{
-			const FLoadedAnimationChunk* Chunk = Algo::FindBy(AnimData->LoadedChunks, ChunkIndex, &FLoadedAnimationChunk::Index);
-			return Chunk ? Chunk->CompressedAnimData : nullptr;
+			if (const FLoadedAnimationChunk* Chunk = Algo::FindBy(AnimData->LoadedChunks, ChunkIndex, &FLoadedAnimationChunk::Index))
+			{
+				if (Chunk->CompressedAnimData == nullptr)
+				{
+					const double RequestTimer = Chunk->RequestStart < 0.f ? Chunk->RequestStart : FPlatformTime::Seconds() - Chunk->RequestStart;
+					UE_LOG(LogAnimation, Warning, TEXT("No Animation Data for loaded chunk: %i, Anim: %s Request timer : %.3f"), ChunkIndex, *Anim->GetFullName(), RequestTimer);
+				}
+				return Chunk->CompressedAnimData;
+			}
+			else
+			{
+				UE_LOG(LogAnimation, Warning, TEXT("Unable to find requested Chunk: %i, Anim: %s - Is in LoadedChunkIndices however"), ChunkIndex, *Anim->GetFullName());
+			}
+		}
+		else
+		{
+			UE_LOG(LogAnimation, Warning, TEXT("Requested Previously Unknown Chunk: %i, Anim: %s"),  ChunkIndex, *Anim->GetFullName());
 		}
 	}
 
diff --git a/Engine/Source/Runtime/Engine/Public/Animation/AnimationStreaming.h b/Engine/Source/Runtime/Engine/Public/Animation/AnimationStreaming.h
--- a/Engine/Source/Runtime/Engine/Public/Animation/AnimationStreaming.h
+++ b/Engine/Source/Runtime/Engine/Public/Animation/AnimationStreaming.h
@@ -14,14 +14,14 @@ AudioStreaming.h: Definitions of classes used for audio streaming.
 #include "ContentStreaming.h"
 #include "Async/AsyncWork.h"
 #include "Async/AsyncFileHandle.h"
-
+#include "Templates/Atomic.h"
 class UAnimStreamable;
 struct FAnimationStreamingManager;
 
 // 
 struct FLoadedAnimationChunk
 {
-	FCompressedAnimSequence* CompressedAnimData;
+	TAtomic<FCompressedAnimSequence*> CompressedAnimData;
 
 	class IAsyncReadRequest* IORequest;
 	double RequestStart;
@@ -32,6 +32,7 @@ struct FLoadedAnimationChunk
 	FLoadedAnimationChunk()
 		: CompressedAnimData(nullptr)
 		, IORequest(nullptr)
+		, RequestStart(-1.0)
 		, Index(0)
 		, bOwnsCompressedData(false)
 	{
@@ -39,7 +40,7 @@ struct FLoadedAnimationChunk
 
 	~FLoadedAnimationChunk()
 	{
-		checkf(CompressedAnimData == nullptr, TEXT("Animation chunk compressed data ptr not null (%p), DataSize: %d"));
+		checkf(CompressedAnimData == nullptr, TEXT("Animation chunk compressed data ptr not null (%p), Index: %u"), CompressedAnimData.Load(), Index);
 	}
 
 	void CleanUpIORequest();
