Fix for issues relating to RTT happening in the middle of a tick when there is other async work happening.

diff --git a/Engine/Source/Runtime/Engine/Classes/Components/ActorComponent.h b/Engine/Source/Runtime/Engine/Classes/Components/ActorComponent.h
--- a/Engine/Source/Runtime/Engine/Classes/Components/ActorComponent.h
+++ b/Engine/Source/Runtime/Engine/Classes/Components/ActorComponent.h
@@ -454,6 +454,9 @@ public:
 	/** This should only be called by the engine in ULevel::InitializeNetworkActors to initialize bIsNetStartupComponent. */
 	void SetIsNetStartupComponent(const bool bInIsNetStartupComponent) { bIsNetStartupComponent = bInIsNetStartupComponent; }
 
+	/** Allows components to handle an EOF update happening mid tick. Can be used to block on in-flight async tasks etc. This should ensure the the component's tick is complete so that it's render update is correct. */
+	virtual void OnEndOfFrameUpdateDuringTick() {}
+
 private:
 	/** Cached pointer to owning actor */
 	mutable AActor* OwnerPrivate;
diff --git a/Engine/Source/Runtime/Engine/Classes/Particles/ParticleSystemComponent.h b/Engine/Source/Runtime/Engine/Classes/Particles/ParticleSystemComponent.h
--- a/Engine/Source/Runtime/Engine/Classes/Particles/ParticleSystemComponent.h
+++ b/Engine/Source/Runtime/Engine/Classes/Particles/ParticleSystemComponent.h
@@ -1216,6 +1216,7 @@ public:
 	virtual void OnChildAttached(USceneComponent* ChildComponent)override;
 	virtual void OnChildDetached(USceneComponent* ChildComponent)override;
 
+	virtual void OnEndOfFrameUpdateDuringTick() override;
 protected:
 	virtual void CreateRenderState_Concurrent() override;
 	virtual void SendRenderTransform_Concurrent() override;
diff --git a/Engine/Source/Runtime/Engine/Private/LevelTick.cpp b/Engine/Source/Runtime/Engine/Private/LevelTick.cpp
--- a/Engine/Source/Runtime/Engine/Private/LevelTick.cpp
+++ b/Engine/Source/Runtime/Engine/Private/LevelTick.cpp
@@ -106,6 +106,8 @@ DEFINE_STAT(STAT_PostTickComponentRecreate);
 DEFINE_STAT(STAT_PostTickComponentUpdate);
 DEFINE_STAT(STAT_PostTickComponentUpdateWait);
 
+DECLARE_CYCLE_STAT(TEXT("OnEndOfFrameUpdateDuringTick"), STAT_OnEndOfFrameUpdateDuringTick, STATGROUP_Game);
+
 DEFINE_STAT(STAT_TickTime);
 DEFINE_STAT(STAT_WorldTickTime);
 DEFINE_STAT(STAT_UpdateCameraTime);
@@ -980,6 +982,28 @@ void UWorld::SendAllEndOfFrameUpdates()
 		return;
 	}
 
+	//If we call SendAllEndOfFrameUpdates during a tick, we must ensure that all marked objects have completed any async work etc before doing the updates.
+	if (bInTick)
+	{
+		SCOPE_CYCLE_COUNTER(STAT_OnEndOfFrameUpdateDuringTick);
+
+		//If this proves too slow we can possibly have the mark set a complimentary bit array that marks which components need this call? Or just another component array?
+		for (UActorComponent* Component : ComponentsThatNeedEndOfFrameUpdate)
+		{
+			if (Component)
+			{
+				Component->OnEndOfFrameUpdateDuringTick();
+			}
+		}
+		for (UActorComponent* Component : ComponentsThatNeedEndOfFrameUpdate_OnGameThread)
+		{
+			if (Component)
+			{
+				Component->OnEndOfFrameUpdateDuringTick();
+			}
+		}
+	}
+
 	// Issue a GPU event to wrap GPU work done during SendAllEndOfFrameUpdates, like skin cache updates
 	FSendAllEndOfFrameUpdates* SendAllEndOfFrameUpdates = BeginSendEndOfFrameUpdatesDrawEvent(Scene ? Scene->GetGPUSkinCache() : nullptr);
 
diff --git a/Engine/Source/Runtime/Engine/Private/Particles/ParticleComponents.cpp b/Engine/Source/Runtime/Engine/Private/Particles/ParticleComponents.cpp
--- a/Engine/Source/Runtime/Engine/Private/Particles/ParticleComponents.cpp
+++ b/Engine/Source/Runtime/Engine/Private/Particles/ParticleComponents.cpp
@@ -3890,6 +3890,11 @@ void UParticleSystemComponent::OnUnregister()
 	check(FXSystem == NULL);
 }
 
+void UParticleSystemComponent::OnEndOfFrameUpdateDuringTick()
+{
+	WaitForAsyncAndFinalize(STALL);
+}
+
 void UParticleSystemComponent::CreateRenderState_Concurrent()
 {
 	LLM_SCOPE(ELLMTag::Particles);
