#UE4 Mac editor CIS fix

diff --git a/Engine/Source/Runtime/Engine/Private/AudioCompressionSettingsUtils.cpp b/Engine/Source/Runtime/Engine/Private/AudioCompressionSettingsUtils.cpp
--- a/Engine/Source/Runtime/Engine/Private/AudioCompressionSettingsUtils.cpp
+++ b/Engine/Source/Runtime/Engine/Private/AudioCompressionSettingsUtils.cpp
@@ -129,15 +129,16 @@ static FPlatformAudioCookOverrides OutOverrides = FPlatformAudioCookOverrides();
 
 #if PLATFORM_SUPPORTS_COMPRESSION_OVERRIDES
 static FCriticalSection CookOverridesCriticalSection;
+#endif // PLATFORM_SUPPORTS_COMPRESSION_OVERRIDES
 
 void FPlatformCompressionUtilities::RecacheCookOverrides()
 {
+#if PLATFORM_SUPPORTS_COMPRESSION_OVERRIDES
 	FScopeLock ScopeLock(&CookOverridesCriticalSection);
 	CacheCurrentPlatformAudioCookOverrides(OutOverrides);
+#endif // PLATFORM_SUPPORTS_COMPRESSION_OVERRIDES
 }
 
-#endif // WITH_EDITOR && PLATFORM_SUPPORTS_COMPRESSION_OVERRIDES
-
 const FPlatformAudioCookOverrides* FPlatformCompressionUtilities::GetCookOverridesForCurrentPlatform(bool bForceRecache)
 {
 #if PLATFORM_SUPPORTS_COMPRESSION_OVERRIDES
