Mesh attachment groups on the scene are now only filled by primitives which are part of an attachment group with a root which has "bLightAttachmentsAsGroup" enabled. It is now possible to select any primitive component in the attachment hierarchy to be used as a root for grouping inset shadows and indirect lighting. Added BP setters to manipulate the inset shadow state of an entire group of attachments. Saves a little bit of memory by not creating Attachment Groups on FScene when not needed.

diff --git a/Engine/Source/Runtime/Engine/Classes/Components/PrimitiveComponent.h b/Engine/Source/Runtime/Engine/Classes/Components/PrimitiveComponent.h
--- a/Engine/Source/Runtime/Engine/Classes/Components/PrimitiveComponent.h
+++ b/Engine/Source/Runtime/Engine/Classes/Components/PrimitiveComponent.h
@@ -888,6 +888,12 @@ public:
 	bool ComponentOverlapMulti(TArray<struct FOverlapResult>& OutOverlaps, const class UWorld* InWorld, const FVector& Pos, const FQuat& Rot, ECollisionChannel TestChannel, const struct FComponentQueryParams& Params = FComponentQueryParams::DefaultComponentQueryParams, const struct FCollisionObjectQueryParams& ObjectQueryParams = FCollisionObjectQueryParams::DefaultObjectQueryParam) const;
 	bool ComponentOverlapMulti(TArray<struct FOverlapResult>& OutOverlaps, const class UWorld* InWorld, const FVector& Pos, const FRotator& Rot, ECollisionChannel TestChannel, const struct FComponentQueryParams& Params = FComponentQueryParams::DefaultComponentQueryParams, const struct FCollisionObjectQueryParams& ObjectQueryParams = FCollisionObjectQueryParams::DefaultObjectQueryParam) const;
 
+	/**
+	 *	Walks up the attachment tree until a primitive component with LightAttachmentsAsGroup enabled is found. This component will effectively act as the root of the attachment group.
+	 *	Return nullptr if none is found. 
+	 */
+	const UPrimitiveComponent* GetLightingAttachmentRoot() const;
+
 protected:
 	/** Override this method for custom behavior for ComponentOverlapMulti() */
 	virtual bool ComponentOverlapMultiImpl(TArray<struct FOverlapResult>& OutOverlaps, const class UWorld* InWorld, const FVector& Pos, const FQuat& Rot, ECollisionChannel TestChannel, const struct FComponentQueryParams& Params, const struct FCollisionObjectQueryParams& ObjectQueryParams = FCollisionObjectQueryParams::DefaultObjectQueryParam) const;
@@ -1400,6 +1406,14 @@ public:
 	UFUNCTION(BlueprintCallable, Category="Rendering")
 	void SetCastShadow(bool NewCastShadow);
 
+	/** Changes the value of CastInsetShadow. */
+	UFUNCTION(BlueprintCallable, Category="Rendering")
+	void SetCastInsetShadow(UPARAM(DisplayName="CastInsetShadow") bool bInCastInsetShadow);
+
+	/** Changes the value of LightAttachmentsAsGroup. */
+	UFUNCTION(BlueprintCallable, Category="Rendering")
+	void SetLightAttachmentsAsGroup(UPARAM(DisplayName="LightAttachmentsAsGroup") bool bInLightAttachmentsAsGroup);
+
 	/** Changes the value of bSingleSampleShadowFromStationaryLights. */
 	UFUNCTION(BlueprintCallable, Category="Rendering")
 	void SetSingleSampleShadowFromStationaryLights(bool bNewSingleSampleShadowFromStationaryLights);
diff --git a/Engine/Source/Runtime/Engine/Classes/Kismet/KismetRenderingLibrary.h b/Engine/Source/Runtime/Engine/Classes/Kismet/KismetRenderingLibrary.h
--- a/Engine/Source/Runtime/Engine/Classes/Kismet/KismetRenderingLibrary.h
+++ b/Engine/Source/Runtime/Engine/Classes/Kismet/KismetRenderingLibrary.h
@@ -151,4 +151,10 @@ class UKismetRenderingLibrary : public UBlueprintFunctionLibrary
 	/** Break FSkelMeshSkinWeightInfo */
 	UFUNCTION(BlueprintPure, Category = "Rendering", meta=(NativeBreakFunc))
 	static ENGINE_API void BreakSkinWeightInfo(FSkelMeshSkinWeightInfo InWeight, int32& Bone0, uint8& Weight0, int32& Bone1, uint8& Weight1, int32& Bone2, uint8& Weight2, int32& Bone3, uint8& Weight3);
+
+	/** Set the inset shadow casting state of the given component and all its child attachments. 
+	 *	Also choose if all attachments should be grouped for the inset shadow rendering. If enabled, one depth target will be shared for all attachments.
+	 */
+	UFUNCTION(BlueprintCallable, Category="Rendering", meta=(Keywords="SetCastShadowForAllAttachments", WorldContext="WorldContextObject", UnsafeDuringActorConstruction="true"))
+	static ENGINE_API void SetCastInsetShadowForAllAttachments(UObject* WorldContextObject, UPrimitiveComponent* PrimitiveComponent, bool bCastInsetShadow, bool bLightAttachmentsAsGroup);
 };
diff --git a/Engine/Source/Runtime/Engine/Private/Components/PrimitiveComponent.cpp b/Engine/Source/Runtime/Engine/Private/Components/PrimitiveComponent.cpp
--- a/Engine/Source/Runtime/Engine/Private/Components/PrimitiveComponent.cpp
+++ b/Engine/Source/Runtime/Engine/Private/Components/PrimitiveComponent.cpp
@@ -1392,6 +1392,24 @@ void UPrimitiveComponent::SetCastShadow(bool NewCastShadow)
 	}
 }
 
+void UPrimitiveComponent::SetCastInsetShadow(bool bInCastInsetShadow)
+{
+	if(bInCastInsetShadow != bCastInsetShadow)
+	{
+		bCastInsetShadow = bInCastInsetShadow;
+		MarkRenderStateDirty();
+	}
+}
+
+void UPrimitiveComponent::SetLightAttachmentsAsGroup(bool bInLightAttachmentsAsGroup)
+{
+	if(bInLightAttachmentsAsGroup != bLightAttachmentsAsGroup)
+	{
+		bLightAttachmentsAsGroup = bInLightAttachmentsAsGroup;
+		MarkRenderStateDirty();
+	}
+}
+
 void UPrimitiveComponent::SetSingleSampleShadowFromStationaryLights(bool bNewSingleSampleShadowFromStationaryLights)
 {
 	if (bNewSingleSampleShadowFromStationaryLights != bSingleSampleShadowFromStationaryLights)
@@ -3210,6 +3228,27 @@ bool UPrimitiveComponent::ComponentOverlapMultiImpl(TArray<struct FOverlapResult
 	return BodyInstance.OverlapMulti(OutOverlaps, World, /*pWorldToComponent=*/ nullptr, Pos, Quat, TestChannel, ParamsWithSelf, FCollisionResponseParams(GetCollisionResponseToChannels()), ObjectQueryParams);
 }
 
+const UPrimitiveComponent* UPrimitiveComponent::GetLightingAttachmentRoot() const
+{
+	const USceneComponent* CurrentHead = this;
+
+	while (CurrentHead)
+	{
+		// If the component has been marked to light itself and child attachments as a group, return it as root
+		if (const UPrimitiveComponent* CurrentHeadPrim = Cast<UPrimitiveComponent>(CurrentHead))
+		{
+			if (CurrentHeadPrim->bLightAttachmentsAsGroup)
+			{
+				return CurrentHeadPrim;
+			}
+		}
+
+		CurrentHead = CurrentHead->GetAttachParent();
+	}
+
+	return nullptr;
+}
+
 #if WITH_EDITOR
 void UPrimitiveComponent::UpdateBounds()
 {
diff --git a/Engine/Source/Runtime/Engine/Private/KismetRenderingLibrary.cpp b/Engine/Source/Runtime/Engine/Private/KismetRenderingLibrary.cpp
--- a/Engine/Source/Runtime/Engine/Private/KismetRenderingLibrary.cpp
+++ b/Engine/Source/Runtime/Engine/Private/KismetRenderingLibrary.cpp
@@ -654,4 +654,42 @@ void UKismetRenderingLibrary::BreakSkinWeightInfo(FSkelMeshSkinWeightInfo InWeig
 	Weight3 = InWeight.Weights[3];
 }
 
+void UKismetRenderingLibrary::SetCastInsetShadowForAllAttachments(UObject* WorldContextObject, UPrimitiveComponent* PrimitiveComponent, bool bCastInsetShadow, bool bLightAttachmentsAsGroup)
+{
+	if (PrimitiveComponent)
+	{
+		// Update this primitive
+		PrimitiveComponent->SetCastInsetShadow(bCastInsetShadow);
+		PrimitiveComponent->SetLightAttachmentsAsGroup(bLightAttachmentsAsGroup);
+
+		// Go through all potential children and update them 
+		TArray<USceneComponent*, TInlineAllocator<8>> ProcessStack;
+		ProcessStack.Append(PrimitiveComponent->GetAttachChildren());
+
+		// Walk down the tree updating
+		while (ProcessStack.Num() > 0)
+		{
+			USceneComponent* Current = ProcessStack.Pop(/*bAllowShrinking=*/ false);
+			UPrimitiveComponent* CurrentPrimitive = Cast<UPrimitiveComponent>(Current);
+
+			if (CurrentPrimitive && CurrentPrimitive->ShouldComponentAddToScene())
+			{
+				if (bLightAttachmentsAsGroup)
+				{
+					// Clear all the children if the root primitive wants to light attachments as group
+					// This is to make sure no child attachment in the chain overrides its parent
+					CurrentPrimitive->SetLightAttachmentsAsGroup(false);
+				}
+
+				CurrentPrimitive->SetCastInsetShadow(bCastInsetShadow);
+			}
+
+			ProcessStack.Append(Current->GetAttachChildren());
+		}
+	}
+	else
+	{
+		FMessageLog("Blueprint").Warning(LOCTEXT("SetCastInsetShadowForAllAttachments_InvalidPrimitiveComponent", "SetCastInsetShadowForAllAttachments: PrimitiveComponent must be non-null."));
+	}
+}
 #undef LOCTEXT_NAMESPACE
diff --git a/Engine/Source/Runtime/Renderer/Private/PrimitiveSceneInfo.cpp b/Engine/Source/Runtime/Renderer/Private/PrimitiveSceneInfo.cpp
--- a/Engine/Source/Runtime/Renderer/Private/PrimitiveSceneInfo.cpp
+++ b/Engine/Source/Runtime/Renderer/Private/PrimitiveSceneInfo.cpp
@@ -143,7 +143,7 @@ FPrimitiveSceneInfo::FPrimitiveSceneInfo(UPrimitiveComponent* InComponent,FScene
 	check(PrimitiveComponentId.IsValid());
 	check(Proxy);
 
-	UPrimitiveComponent* SearchParentComponent = Cast<UPrimitiveComponent>(InComponent->GetAttachmentRoot());
+	const UPrimitiveComponent* SearchParentComponent = InComponent->GetLightingAttachmentRoot();
 
 	if (SearchParentComponent && SearchParentComponent != InComponent)
 	{
diff --git a/Engine/Source/Runtime/Renderer/Private/RendererScene.cpp b/Engine/Source/Runtime/Renderer/Private/RendererScene.cpp
--- a/Engine/Source/Runtime/Renderer/Private/RendererScene.cpp
+++ b/Engine/Source/Runtime/Renderer/Private/RendererScene.cpp
@@ -1407,7 +1407,7 @@ void FScene::UpdatePrimitiveTransform(UPrimitiveComponent* Primitive)
 
 void FScene::UpdatePrimitiveLightingAttachmentRoot(UPrimitiveComponent* Primitive)
 {
-	const UPrimitiveComponent* NewLightingAttachmentRoot = Cast<UPrimitiveComponent>(Primitive->GetAttachmentRoot());
+	const UPrimitiveComponent* NewLightingAttachmentRoot = Primitive->GetLightingAttachmentRoot();
 
 	if (NewLightingAttachmentRoot == Primitive)
 	{
