Anim Streaming
- Fix memory leak with streamed chunks
- Fix issue with More chunks than needed being loaded
- Add proper memory tracking
- Increased default buffer size (2 seconds seemed too optimistic in practice)
- Change missing chunk logic to be closest chunk rather than ref pose, should minimize artifacts
- Add extra logging for case where streaming hasn't caught up

diff --git a/Engine/Source/Runtime/Engine/Classes/Animation/AnimStreamable.h b/Engine/Source/Runtime/Engine/Classes/Animation/AnimStreamable.h
--- a/Engine/Source/Runtime/Engine/Classes/Animation/AnimStreamable.h
+++ b/Engine/Source/Runtime/Engine/Classes/Animation/AnimStreamable.h
@@ -45,14 +45,16 @@ public:
 	// Bulk data if stored in the package.
 	FByteBulkData BulkData;
 
-	int32 GetMemorySize() const
+	SIZE_T GetMemorySize() const
 	{
-		return	sizeof(FAnimStreamableChunk)
-			+ sizeof(FCompressedAnimSequence)
-			+ CompressedAnimSequence->CompressedTrackToSkeletonMapTable.GetAllocatedSize()
-			+ CompressedAnimSequence->CompressedCurveNames.GetAllocatedSize()
-			+ CompressedAnimSequence->CompressedCurveByteStream.GetAllocatedSize()
-			+ CompressedAnimSequence->CompressedDataStructure.GetApproxBoneCompressedSize();
+		static const SIZE_T ClassSize = sizeof(FAnimStreamableChunk);
+		SIZE_T CurrentSize = ClassSize;
+
+		if (CompressedAnimSequence)
+		{
+			CurrentSize += CompressedAnimSequence->GetMemorySize();
+		}
+		return CurrentSize;
 	}
 
 	/** Serialization. */
@@ -70,6 +72,16 @@ public:
 	{
 		Chunks.Reset();
 	}
+
+	SIZE_T GetMemorySize() const
+	{
+		SIZE_T ChunkSize = 0;
+		for (const FAnimStreamableChunk& Chunk : Chunks)
+		{
+			ChunkSize += Chunk.GetMemorySize();
+		}
+		return sizeof(FStreamableAnimPlatformData) + ChunkSize;
+	}
 };
 
 
@@ -189,6 +201,7 @@ public:
 	virtual void Serialize(FArchive& Ar) override;
 	virtual void PostLoad() override;
 	virtual void FinishDestroy() override;
+	virtual void GetResourceSizeEx(FResourceSizeEx& CumulativeResourceSize) override;
 	//~ End UObject Interface
 
 	//~ Begin UAnimSequenceBase Interface
diff --git a/Engine/Source/Runtime/Engine/Private/Animation/AnimCompressionTypes.cpp b/Engine/Source/Runtime/Engine/Private/Animation/AnimCompressionTypes.cpp
--- a/Engine/Source/Runtime/Engine/Private/Animation/AnimCompressionTypes.cpp
+++ b/Engine/Source/Runtime/Engine/Private/Animation/AnimCompressionTypes.cpp
@@ -619,6 +619,15 @@ void FCompressedAnimSequence::SerializeCompressedData(FArchive& Ar, bool bDDCDat
 #endif
 }
 
+SIZE_T FCompressedAnimSequence::GetMemorySize() const
+{
+	return	  CompressedTrackToSkeletonMapTable.GetAllocatedSize()
+			+ CompressedCurveNames.GetAllocatedSize()
+			+ CompressedCurveByteStream.GetAllocatedSize()
+			+ CompressedDataStructure.GetApproxBoneCompressedSize()
+			+ sizeof(FCompressedAnimSequence);
+}
+
 struct FGetBonePoseScratchArea : public TThreadSingleton<FGetBonePoseScratchArea>
 {
 	BoneTrackArray RotationScalePairs;
diff --git a/Engine/Source/Runtime/Engine/Private/Animation/AnimStreamable.cpp b/Engine/Source/Runtime/Engine/Private/Animation/AnimStreamable.cpp
--- a/Engine/Source/Runtime/Engine/Private/Animation/AnimStreamable.cpp
+++ b/Engine/Source/Runtime/Engine/Private/Animation/AnimStreamable.cpp
@@ -25,7 +25,7 @@ DECLARE_CYCLE_STAT(TEXT("AnimStreamable GetAnimationPose"), STAT_AnimStreamable_
 // set it here
 const TCHAR* StreamingAnimChunkVersion = TEXT("1F1656B9E10142729AB16650D9821B1F");
 
-const float MINIMUM_CHUNK_SIZE = 2.0f;
+const float MINIMUM_CHUNK_SIZE = 4.0f;
 
 float GChunkSizeSeconds = MINIMUM_CHUNK_SIZE;
 
@@ -168,11 +168,16 @@ void UAnimStreamable::HandleAssetPlayerTickedInternal(FAnimAssetTickContext &Con
 	
 	const int32 ChunkIndex = GetChunkIndexForTime(GetRunningPlatformData().Chunks, PreviousTime);
 	IAnimationStreamingManager& StreamingManager = IStreamingManager::Get().GetAnimationStreamingManager();
-	const FCompressedAnimSequence* CurveCompressedDataChunk = StreamingManager.GetLoadedChunk(this, ChunkIndex);
+	const FCompressedAnimSequence* CurveCompressedDataChunk = StreamingManager.GetLoadedChunk(this, ChunkIndex, true);
 	
 	//ExtractRootMotionFromTrack(AnimationTrack, PreviousTime, PreviousTime + MoveDelta, Context.RootMotionMovementParams);
 }
 
+FORCEINLINE int32 PreviousChunkIndex(int32 ChunkIndex, int32 NumChunks)
+{
+	return (ChunkIndex + NumChunks - 1) % NumChunks;
+}
+
 void UAnimStreamable::GetAnimationPose(FCompactPose& OutPose, FBlendedCurve& OutCurve, const FAnimExtractContext& ExtractionContext) const
 {
 	OutPose.ResetToRefPose();
@@ -270,7 +275,7 @@ void UAnimStreamable::GetAnimationPose(FCompactPose& OutPose, FBlendedCurve& Out
 	{
 		IAnimationStreamingManager& StreamingManager = IStreamingManager::Get().GetAnimationStreamingManager();
 
-		const FCompressedAnimSequence* CurveCompressedDataChunk = StreamingManager.GetLoadedChunk(this, 0); //Curve Data stored in chunk 0 till it is properly cropped
+		const FCompressedAnimSequence* CurveCompressedDataChunk = StreamingManager.GetLoadedChunk(this, 0, false); //Curve Data stored in chunk 0 till it is properly cropped
 
 		if (!CurveCompressedDataChunk)
 		{
@@ -284,15 +289,32 @@ void UAnimStreamable::GetAnimationPose(FCompactPose& OutPose, FBlendedCurve& Out
 
 		CurveCompressedDataChunk->CurveCompressionCodec->DecompressCurves(*CurveCompressedDataChunk, OutCurve, ExtractionContext.CurrentTime);
 
-		const FCompressedAnimSequence* CompressedData = (ChunkIndex == 0) ? CurveCompressedDataChunk : StreamingManager.GetLoadedChunk(this, ChunkIndex);
+		const FCompressedAnimSequence* CompressedData = (ChunkIndex == 0) ? CurveCompressedDataChunk : StreamingManager.GetLoadedChunk(this, ChunkIndex, true);
+
+		float ChunkCurrentTime = ExtractionContext.CurrentTime - GetRunningPlatformData().Chunks[ChunkIndex].StartTime;
 
 		if (!CompressedData)
 		{
 #if WITH_EDITOR
 			CompressedData = GetRunningPlatformData().Chunks[ChunkIndex].CompressedAnimSequence;
 #else
-			UE_LOG(LogAnimation, Warning, TEXT("Failed to get streamed compressed data Time: %.2f, ChunkIndex:%i, Anim: %s"), ExtractionContext.CurrentTime, ChunkIndex, *GetFullName());
-			return;
+			const int32 NumChunks = GetRunningPlatformData().Chunks.Num();
+
+			int32 FallbackChunkIndex = ChunkIndex;
+			while (!CompressedData)
+			{
+				FallbackChunkIndex = PreviousChunkIndex(FallbackChunkIndex, NumChunks);
+				if (FallbackChunkIndex == ChunkIndex)
+				{
+					//Cannot get fallback
+					UE_LOG(LogAnimation, Warning, TEXT("Failed to get ANY streamed compressed data Time: %.2f, ChunkIndex:%i, Anim: %s"), ExtractionContext.CurrentTime, ChunkIndex, *GetFullName());
+					return;
+				}
+				CompressedData = StreamingManager.GetLoadedChunk(this, FallbackChunkIndex, false);
+				ChunkCurrentTime = GetRunningPlatformData().Chunks[FallbackChunkIndex].SequenceLength;
+			}
+
+			UE_LOG(LogAnimation, Warning, TEXT("Failed to get streamed compressed data Time: %.2f, ChunkIndex:%i - Using Chunk %i Anim: %s"), ExtractionContext.CurrentTime, ChunkIndex, FallbackChunkIndex, *GetFullName());
 #endif
 		}
 
@@ -302,7 +324,6 @@ void UAnimStreamable::GetAnimationPose(FCompactPose& OutPose, FBlendedCurve& Out
 			return;
 		}
 
-		float ChunkCurrentTime = ExtractionContext.CurrentTime - GetRunningPlatformData().Chunks[ChunkIndex].StartTime;
 		FAnimExtractContext ChunkExtractionContext(ChunkCurrentTime, ExtractionContext.bExtractRootMotion);
 		ChunkExtractionContext.BonesRequired = ExtractionContext.BonesRequired;
 		ChunkExtractionContext.PoseCurves = ExtractionContext.PoseCurves;
@@ -361,6 +382,20 @@ void UAnimStreamable::FinishDestroy()
 	IStreamingManager::Get().GetAnimationStreamingManager().RemoveStreamingAnim(this);
 }
 
+void UAnimStreamable::GetResourceSizeEx(FResourceSizeEx& CumulativeResourceSize)
+{
+	Super::GetResourceSizeEx(CumulativeResourceSize);
+#if WITH_EDITOR
+	for (auto& AnimData : StreamableAnimPlatformData)
+	{
+		CumulativeResourceSize.AddDedicatedSystemMemoryBytes(AnimData.Value->GetMemorySize());
+	}
+#else
+	CumulativeResourceSize.AddDedicatedSystemMemoryBytes(GetRunningPlatformData().GetMemorySize());
+	CumulativeResourceSize.AddDedicatedSystemMemoryBytes(IStreamingManager::Get().GetAnimationStreamingManager().GetMemorySizeForAnim(this));
+#endif
+}
+
 int32 UAnimStreamable::GetChunkIndexForTime(const TArray<FAnimStreamableChunk>& Chunks, float CurrentTime) const
 {
 	for (int32 ChunkIndex = 0; ChunkIndex < Chunks.Num(); ++ChunkIndex)
diff --git a/Engine/Source/Runtime/Engine/Private/Animation/AnimationStreaming.cpp b/Engine/Source/Runtime/Engine/Private/Animation/AnimationStreaming.cpp
--- a/Engine/Source/Runtime/Engine/Private/Animation/AnimationStreaming.cpp
+++ b/Engine/Source/Runtime/Engine/Private/Animation/AnimationStreaming.cpp
@@ -142,7 +142,7 @@ bool FStreamingAnimationData::HasPendingRequests(TArray<uint32>& IndicesToLoad,
 	IndicesToFree.Reset();
 
 	// Find indices that aren't loaded
-	for (auto NeededIndex : RequestedChunks)
+	for (uint32 NeededIndex : RequestedChunks)
 	{
 		if (!LoadedChunkIndices.Contains(NeededIndex))
 		{
@@ -151,7 +151,7 @@ bool FStreamingAnimationData::HasPendingRequests(TArray<uint32>& IndicesToLoad,
 	}
 
 	// Find indices that aren't needed anymore
-	for (auto CurrentIndex : LoadedChunkIndices)
+	for (uint32 CurrentIndex : LoadedChunkIndices)
 	{
 		if (!RequestedChunks.Contains(CurrentIndex))
 		{
@@ -168,7 +168,7 @@ void FStreamingAnimationData::BeginPendingRequests(const TArray<uint32>& Indices
 
 	// Mark Chunks for removal in case they can be reused
 	{
-		for (auto Index : IndicesToFree)
+		for (uint32 Index : IndicesToFree)
 		{
 			for (int32 ChunkIndex = 0; ChunkIndex < LoadedChunks.Num(); ++ChunkIndex)
 			{
@@ -187,7 +187,7 @@ void FStreamingAnimationData::BeginPendingRequests(const TArray<uint32>& Indices
 
 	const EAsyncIOPriorityAndFlags AsyncIOPriority = AIOP_High;
 
-	for (auto ChunkIndex : IndicesToLoad)
+	for (uint32 ChunkIndex : IndicesToLoad)
 	{
 		const FAnimStreamableChunk& Chunk = StreamableAnim->GetRunningPlatformData().Chunks[ChunkIndex];
 
@@ -295,6 +295,19 @@ void FStreamingAnimationData::ResetRequestedChunks()
 	RequestedChunks.Add(0); //Always want 1
 }
 
+SIZE_T FStreamingAnimationData::GetMemorySize() const
+{
+	SIZE_T CurrentSize = 0;
+	for (const FLoadedAnimationChunk& Chunk : LoadedChunks)
+	{
+		if (Chunk.CompressedAnimData && Chunk.bOwnsCompressedData)
+		{
+			CurrentSize += Chunk.CompressedAnimData.Load()->GetMemorySize();
+		}
+	}
+	return CurrentSize;
+}
+
 ////////////////////////////
 // FAnimationStreamingManager //
 ////////////////////////////
@@ -337,19 +350,35 @@ void FAnimationStreamingManager::OnAsyncFileCallback(FStreamingAnimationData* St
 		}
 
 		ChunkStorage.CompressedAnimData = NewCompressedData;
+		ChunkStorage.bOwnsCompressedData = true;
 		ChunkStorage.RequestStart = -2.0; //Signify we have finished loading
 
 		UE_LOG(LogAnimation, Log, TEXT("Request Finished %.2f\n Anim Chunk Streamed %.4f\n"), CurrentTime, RequestDuration);
 	}
 }
 
+SIZE_T FAnimationStreamingManager::GetMemorySizeForAnim(const UAnimStreamable* Anim)
+{
+	SIZE_T AnimSize = 0;
+
+	FScopeLock Lock(&CriticalSection);
+
+	FStreamingAnimationData* AnimData = StreamingAnimations.FindRef(Anim);
+	if (AnimData)
+	{
+		return AnimData->GetMemorySize();
+	}
+
+	return AnimSize;
+}
+
 void FAnimationStreamingManager::UpdateResourceStreaming(float DeltaTime, bool bProcessEverything /*= false*/)
 {
 	LLM_SCOPE(ELLMTag::Audio);
 
 	FScopeLock Lock(&CriticalSection);
 
-	for (auto& AnimData : StreamingAnimations)
+	for (TPair<UAnimStreamable*, FStreamingAnimationData*>& AnimData : StreamingAnimations)
 	{
 		AnimData.Value->UpdateStreamingStatus();
 	}
@@ -365,7 +394,7 @@ int32 FAnimationStreamingManager::BlockTillAllRequestsFinished(float TimeLimit,
 
 		if (TimeLimit == 0.0f)
 		{
-			for (auto& AnimPair : StreamingAnimations)
+			for (TPair<UAnimStreamable*, FStreamingAnimationData*>& AnimPair : StreamingAnimations)
 			{
 				AnimPair.Value->BlockTillAllRequestsFinished();
 			}
@@ -373,7 +402,7 @@ int32 FAnimationStreamingManager::BlockTillAllRequestsFinished(float TimeLimit,
 		else
 		{
 			double EndTime = FPlatformTime::Seconds() + TimeLimit;
-			for (auto& AnimPair : StreamingAnimations)
+			for (TPair<UAnimStreamable*, FStreamingAnimationData*>& AnimPair : StreamingAnimations)
 			{
 				float ThisTimeLimit = EndTime - FPlatformTime::Seconds();
 				if (ThisTimeLimit < .001f || // one ms is the granularity of the platform event system
@@ -450,7 +479,7 @@ bool FAnimationStreamingManager::RemoveStreamingAnim(UAnimStreamable* Anim)
 	return false;
 }
 
-const FCompressedAnimSequence* FAnimationStreamingManager::GetLoadedChunk(const UAnimStreamable* Anim, uint32 ChunkIndex) const
+const FCompressedAnimSequence* FAnimationStreamingManager::GetLoadedChunk(const UAnimStreamable* Anim, uint32 ChunkIndex, bool bRequestNextChunk) const
 {
 	// Check for the spoof of failing to load a stream chunk
 	if (SpoofFailedAnimationChunkLoad > 0)
@@ -465,7 +494,10 @@ const FCompressedAnimSequence* FAnimationStreamingManager::GetLoadedChunk(const
 	if (AnimData)
 	{
 		AnimData->RequestedChunks.AddUnique(ChunkIndex);
-		AnimData->RequestedChunks.AddUnique((ChunkIndex + 1) % Anim->GetRunningPlatformData().Chunks.Num());
+		if (bRequestNextChunk)
+		{
+			AnimData->RequestedChunks.AddUnique((ChunkIndex + 1) % Anim->GetRunningPlatformData().Chunks.Num());
+		}
 
 		if (AnimData->LoadedChunkIndices.Contains(ChunkIndex))
 		{
diff --git a/Engine/Source/Runtime/Engine/Public/Animation/AnimCompressionTypes.h b/Engine/Source/Runtime/Engine/Public/Animation/AnimCompressionTypes.h
--- a/Engine/Source/Runtime/Engine/Public/Animation/AnimCompressionTypes.h
+++ b/Engine/Source/Runtime/Engine/Public/Animation/AnimCompressionTypes.h
@@ -699,6 +699,9 @@ public:
 	{
 		return CompressedTrackToSkeletonMapTable[TrackIndex].BoneTreeIndex;
 	}
+
+	// Return the number of bytes used
+	SIZE_T GetMemorySize() const;
 };
 
 struct FRootMotionReset
diff --git a/Engine/Source/Runtime/Engine/Public/Animation/AnimationStreaming.h b/Engine/Source/Runtime/Engine/Public/Animation/AnimationStreaming.h
--- a/Engine/Source/Runtime/Engine/Public/Animation/AnimationStreaming.h
+++ b/Engine/Source/Runtime/Engine/Public/Animation/AnimationStreaming.h
@@ -73,14 +73,6 @@ struct FStreamingAnimationData final
 	bool UpdateStreamingStatus();
 
 	/**
-	 * Tells the SoundWave which chunks are currently required so that it can start loading any needed
-	 *
-	 * @param InChunkIndices The Chunk Indices that are currently needed by all sources using this sound
-	 * @param bShouldPrioritizeAsyncIORequest Whether request should have higher priority than usual
-	 */
-	//void UpdateChunkRequests(FWaveRequest& InWaveRequest);
-
-		/**
 	 * Checks whether the requested chunk indices differ from those loaded
 	 *
 	 * @param IndicesToLoad		List of chunk indices that should be loaded
@@ -102,14 +94,8 @@ struct FStreamingAnimationData final
 	*/
 	bool BlockTillAllRequestsFinished(float TimeLimit = 0.0f);
 
-#if WITH_EDITORONLY_DATA
-	/**
-	 * Finishes any Derived Data Cache requests that may be in progress
-	 *
-	 * @return Whether any of the requests failed.
-	 */
-	bool FinishDDCRequests();
-#endif //WITH_EDITORONLY_DATA
+	// Return the number of bytes used
+	SIZE_T GetMemorySize() const;
 
 private:
 	// Don't allow copy construction as it could free shared memory
@@ -135,37 +121,10 @@ public:
 
 	TArray<uint32> RequestedChunks;
 
-	/** Indices of chunks we want to have loaded */
-	//FWaveRequest	CurrentRequest;
-
-#if WITH_EDITORONLY_DATA
-	/** Pending async derived data streaming tasks */
-	//TIndirectArray<FAsyncStreamDerivedChunkTask> PendingAsyncStreamDerivedChunkTasks;
-#endif // #if WITH_EDITORONLY_DATA
-
 	/** Ptr to owning audio streaming manager. */
 	FAnimationStreamingManager* AnimationStreamingManager;
 };
 
-/** Struct used to store results of an async file load. */
-/*struct FASyncAnimationChunkLoadResult
-{
-	// Place to safely copy the ptr of a loaded audio chunk when load result is finished
-	uint8* DataResults;
-
-	// Actual storage of the loaded audio chunk, will be filled on audio thread.
-	FStreamingAnimationData* StreamingAnimData;
-
-	// Loaded audio chunk index
-	int32 LoadedChunkIndex;
-
-	FASyncAnimationChunkLoadResult()
-		: DataResults(nullptr)
-		, StreamingAnimData(nullptr)
-		, LoadedChunkIndex(INDEX_NONE)
-	{}
-};*/
-
 
 /**
 * Streaming manager dealing with audio.
@@ -191,15 +150,8 @@ struct FAnimationStreamingManager : public IAnimationStreamingManager
 	// IAudioStreamingManager interface
 	virtual void AddStreamingAnim(UAnimStreamable* Anim) override;
 	virtual bool RemoveStreamingAnim(UAnimStreamable* Anim) override;
-	/*void AddDecoder(ICompressedAudioInfo* CompressedAudioInfo);
-	void RemoveDecoder(ICompressedAudioInfo* CompressedAudioInfo);
-	bool IsManagedStreamingSoundWave(const USoundWave* SoundWave) const;
-	bool IsStreamingInProgress(const USoundWave* SoundWave);
-	bool CanCreateSoundSource(const FWaveInstance* WaveInstance) const;
-	void AddStreamingSoundSource(FSoundSource* SoundSource);
-	void RemoveStreamingSoundSource(FSoundSource* SoundSource);
-	bool IsManagedStreamingSoundSource(const FSoundSource* SoundSource) const;*/
-	virtual const FCompressedAnimSequence* GetLoadedChunk(const UAnimStreamable* Anim, uint32 ChunkIndex) const override;
+	virtual SIZE_T GetMemorySizeForAnim(const UAnimStreamable* Anim) override;
+	virtual const FCompressedAnimSequence* GetLoadedChunk(const UAnimStreamable* Anim, uint32 ChunkIndex, bool bRequestNextChunk) const override;
 	// End IAudioStreamingManager interface
 
 	/** Called when an async callback is made on an async loading audio chunk request. */
diff --git a/Engine/Source/Runtime/Engine/Public/ContentStreaming.h b/Engine/Source/Runtime/Engine/Public/ContentStreaming.h
--- a/Engine/Source/Runtime/Engine/Public/ContentStreaming.h
+++ b/Engine/Source/Runtime/Engine/Public/ContentStreaming.h
@@ -541,6 +541,9 @@ struct IAnimationStreamingManager : public IStreamingManager
 	/** Removes a Streamable Anim from the streaming manager. */
 	virtual bool RemoveStreamingAnim(UAnimStreamable* Anim) = 0;
 
+	/** Returns the memory a Streamable Anim is currently using */
+	virtual SIZE_T GetMemorySizeForAnim(const UAnimStreamable* Anim) = 0;
+
 	/**
 	 * Gets a pointer to a chunk of animation data
 	 *
@@ -548,7 +551,7 @@ struct IAnimationStreamingManager : public IStreamingManager
 	 * @param ChunkIndex	Index of the chunk we want
 	 * @return Either the desired chunk or NULL if it's not loaded
 	 */
-	virtual const FCompressedAnimSequence* GetLoadedChunk(const UAnimStreamable* Anim, uint32 ChunkIndex) const = 0;
+	virtual const FCompressedAnimSequence* GetLoadedChunk(const UAnimStreamable* Anim, uint32 ChunkIndex, bool bRequestNextChunk) const = 0;
 };
 
 /**
