Sub anim instance state is invalid when calling immediate update before the owning Anim Instance has been updated
#fix Ensure that Sub-Anim instances are always updated after the main Anim Instance, and enforce immediate update of the main instance and post-update when the sub-instance requires so.

diff --git a/Engine/Source/Runtime/Engine/Classes/Animation/AnimInstance.h b/Engine/Source/Runtime/Engine/Classes/Animation/AnimInstance.h
--- a/Engine/Source/Runtime/Engine/Classes/Animation/AnimInstance.h
+++ b/Engine/Source/Runtime/Engine/Classes/Animation/AnimInstance.h
@@ -988,6 +988,17 @@ public:
 	 */
 	virtual void OnUROPreInterpolation_AnyThread(FAnimationEvaluationContext& InOutContext) {}
 
+	/** Flag passed to UpdateAnimation, determines the path we follow */
+	enum class EUpdateAnimationFlag : uint8
+	{
+		/** Enforce an immediate update, regardless of state*/
+		ForceImmediateUpdate,
+		/** Enforces a parallel update, regardless of state */
+		ForceParallelUpdate,
+		/** Use state to determine whether or not to immediately or update in parallel */
+		Default
+	};
+
 	// Animation phase trigger
 	// start with initialize
 	// update happens in every tick. Can happen in parallel with others if conditions are right.
@@ -995,7 +1006,9 @@ public:
 	// post eval happens after evaluation is done
 	// uninitialize happens when owner is unregistered
 	void InitializeAnimation();
-	void UpdateAnimation(float DeltaSeconds, bool bNeedsValidRootMotion);
+
+	/** Update Animation code-paths, updates and advances animation state, returns whether or not the actual update should have been called immediately */
+	bool UpdateAnimation(float DeltaSeconds, bool bNeedsValidRootMotion, EUpdateAnimationFlag UpdateFlag = EUpdateAnimationFlag::Default );
 
 	/** Run update animation work on a worker thread */
 	void ParallelUpdateAnimation();
diff --git a/Engine/Source/Runtime/Engine/Private/Animation/AnimInstance.cpp b/Engine/Source/Runtime/Engine/Private/Animation/AnimInstance.cpp
--- a/Engine/Source/Runtime/Engine/Private/Animation/AnimInstance.cpp
+++ b/Engine/Source/Runtime/Engine/Private/Animation/AnimInstance.cpp
@@ -337,7 +337,7 @@ void UAnimInstance::UpdateMontageSyncGroup()
 	}
 }
 
-void UAnimInstance::UpdateAnimation(float DeltaSeconds, bool bNeedsValidRootMotion)
+bool UAnimInstance::UpdateAnimation(float DeltaSeconds, bool bNeedsValidRootMotion, EUpdateAnimationFlag UpdateFlag)
 {
 	LLM_SCOPE(ELLMTag::Animation);
 
@@ -401,7 +401,7 @@ void UAnimInstance::UpdateAnimation(float DeltaSeconds, bool bNeedsValidRootMoti
 				when we also update the AnimGraph as well.
 				This means that calls to 'Evaluation' without a call to 'Update' prior will render stale data, but that's to be expected.
 			*/
-			return;
+			return false;
 		}
 	}
 
@@ -427,7 +427,7 @@ void UAnimInstance::UpdateAnimation(float DeltaSeconds, bool bNeedsValidRootMoti
 
 		if (UpdateSnapshotAndSkipRemainingUpdate())
 		{
-			return;
+			return false;
 		}
 	}
 #endif
@@ -460,13 +460,34 @@ PRAGMA_ENABLE_DEPRECATION_WARNINGS
 		SCOPE_CYCLE_COUNTER(STAT_BlueprintUpdateAnimation);
 		BlueprintUpdateAnimation(DeltaSeconds);
 	}
+	
+	// Determine whether or not the animation should be immediately updated according to current state
+	const bool bWantsImmediateUpdate = bNeedsValidRootMotion || NeedsImmediateUpdate(DeltaSeconds);
 
-	if(bNeedsValidRootMotion || NeedsImmediateUpdate(DeltaSeconds))
+	// Determine whether or not we can or should actually immediately update the animation state
+	bool bShouldImmediateUpdate = bWantsImmediateUpdate;
+	switch (UpdateFlag)
+	{
+		case EUpdateAnimationFlag::ForceImmediateUpdate:
+		{
+			bShouldImmediateUpdate = true;
+			break;
+		}
+		case EUpdateAnimationFlag::ForceParallelUpdate:
+		{
+			bShouldImmediateUpdate = false;
+			break;
+		}
+	}
+
+	if(bShouldImmediateUpdate)
 	{
 		// cant use parallel update, so just do the work here (we call this function here to do the work on the game thread)
 		ParallelUpdateAnimation();
 		PostUpdateAnimation();
 	}
+
+	return bWantsImmediateUpdate;
 }
 
 void UAnimInstance::PreUpdateAnimation(float DeltaSeconds)
diff --git a/Engine/Source/Runtime/Engine/Private/Components/SkeletalMeshComponent.cpp b/Engine/Source/Runtime/Engine/Private/Components/SkeletalMeshComponent.cpp
--- a/Engine/Source/Runtime/Engine/Private/Components/SkeletalMeshComponent.cpp
+++ b/Engine/Source/Runtime/Engine/Private/Components/SkeletalMeshComponent.cpp
@@ -1011,15 +1011,31 @@ void USkeletalMeshComponent::TickAnimation(float DeltaTime, bool bNeedsValidRoot
 
 		// We update sub instances first incase we're using either root motion or non-threaded update.
 		// This ensures that we go through the pre update process and initialize the proxies correctly.
+
+		// Accumulate whether or not any of the sub-instances wanted to do an immediate Animation Update
+		bool bWantsImmediateUpdate = false;
 		for(UAnimInstance* SubInstance : SubInstances)
 		{
-			SubInstance->UpdateAnimation(DeltaTime * GlobalAnimRateScale, false);
+			// Sub anim instances are always forced to do a parallel update 
+			bWantsImmediateUpdate |= SubInstance->UpdateAnimation(DeltaTime * GlobalAnimRateScale, false, UAnimInstance::EUpdateAnimationFlag::ForceParallelUpdate);
 		}
 
 		if (AnimScriptInstance != nullptr)
 		{
+			// In case any of the sub-instances required an immediate animation update, make sure we immediately do so for the main anim instance
+			const UAnimInstance::EUpdateAnimationFlag UpdateFlag = bWantsImmediateUpdate ? UAnimInstance::EUpdateAnimationFlag::ForceImmediateUpdate : UAnimInstance::EUpdateAnimationFlag::Default;
+
 			// Tick the animation
-			AnimScriptInstance->UpdateAnimation(DeltaTime * GlobalAnimRateScale, bNeedsValidRootMotion);
+			bWantsImmediateUpdate |= AnimScriptInstance->UpdateAnimation(DeltaTime * GlobalAnimRateScale, bNeedsValidRootMotion, UpdateFlag);
+		}
+
+		// Make sure PostUpdateAnimation is called on the subinstances if they, or the main instance did an immediate Animation Update
+		if (bWantsImmediateUpdate)
+		{
+			for (UAnimInstance* SubInstance : SubInstances)
+			{
+				SubInstance->PostUpdateAnimation();
+			}
 		}
 
 		if(ShouldUpdatePostProcessInstance())
