#ue4 - UpdateOverlaps() does nothing until after BeginPlay() has been called on the Actor. Fixed inconsistent ordering of UpdateOverlaps() relative to BeginPlay() in the cases of server / standalone SpawnActor, networked replicated actor, actor streamed as part of a level. BeginPlay() now always precedes UpdateOverlaps(), which was the previous behavior in standalone or on the server when spawning or as part of persistent map load.

Added UpdateOverlapsMethodDuringLevelStreaming to Actor, which allows you to disable the entire UpdateOverlaps() query run after BeginPlay(), or make it conditional based on mobility (skip for static objects). There is also an option to defer to the default specified in INI, which allows each class to have a separate default that is changeable also for engine classes without modifying engine source, and which makes it easily hotfixable.

The new default option is OnlyUpdateMovable, which means to runs overlaps only if the root component of the actor has Movable (non-static, non-stationary) mobility. Previous behavior was equivalent to AlwaysUpdate, however this can be horrible for streaming performance for many actor types.

diff --git a/Engine/Config/BaseEngine.ini b/Engine/Config/BaseEngine.ini
--- a/Engine/Config/BaseEngine.ini
+++ b/Engine/Config/BaseEngine.ini
@@ -1760,6 +1760,13 @@ bPromptForRemoteDebugOnEnsure=false
 [Niagara]
 EnableNiagara=false
 
+[/Script/Engine.Actor]
+; option used when UpdateOverlapsMethodDuringLevelStreaming is set to UseConfigDefault. Options are: AlwaysUpdate, OnlyUpdateMovable, NeverUpdate.
+DefaultUpdateOverlapsMethodDuringLevelStreaming=OnlyUpdateMovable
+
+[/Script/Engine.TriggerVolume]
+DefaultUpdateOverlapsMethodDuringLevelStreaming=AlwaysUpdate
+
 [/Script/Engine.CollisionProfile]
 +Profiles=(Name="NoCollision",CollisionEnabled=NoCollision,ObjectTypeName="WorldStatic",CustomResponses=((Channel="Visibility",Response=ECR_Ignore),(Channel="Camera",Response=ECR_Ignore)),HelpMessage="No collision",bCanModify=False)
 +Profiles=(Name="BlockAll",CollisionEnabled=QueryAndPhysics,ObjectTypeName="WorldStatic",CustomResponses=,HelpMessage="WorldStatic object that blocks all actors by default. All new custom channels will use its own default response. ",bCanModify=False)
diff --git a/Engine/Source/Runtime/Engine/Classes/GameFramework/Actor.h b/Engine/Source/Runtime/Engine/Classes/GameFramework/Actor.h
--- a/Engine/Source/Runtime/Engine/Classes/GameFramework/Actor.h
+++ b/Engine/Source/Runtime/Engine/Classes/GameFramework/Actor.h
@@ -33,6 +33,16 @@ struct FAttachedActorInfo;
 struct FNetViewer;
 struct FNetworkObjectInfo;
 
+/** Chooses a method for actors to update overlap state (objects it is touching) on initialization, currently only used during level streaming. */
+UENUM(BlueprintType)
+enum class EActorUpdateOverlapsMethod : uint8
+{
+	UseConfigDefault,	// Use the default value specified by the native class or config value.
+	AlwaysUpdate,		// Always update overlap state on initialization.
+	OnlyUpdateMovable,	// Only update if root component has Movable mobility.
+	NeverUpdate			// Never update overlap state on initialization.
+};
+
 ENGINE_API DECLARE_LOG_CATEGORY_EXTERN(LogActor, Log, Warning);
 
 // Delegate signatures
@@ -249,8 +259,12 @@ public:
 	UPROPERTY(EditAnywhere, BlueprintReadWrite, Category=Actor, AdvancedDisplay)
 	uint8 bFindCameraComponentWhenViewTarget:1;
 	
-    /** If true, this actor will generate overlap events when spawned as part of level streaming. You might enable this is in the case where a streaming level loads around an actor and you want overlaps to trigger. */
-	UPROPERTY(EditAnywhere, BlueprintReadWrite, Category=Actor)
+    /**
+	 * If true, this actor will generate overlap Begin/End events when spawned as part of level streaming.
+	 * You might enable this is in the case where a streaming level loads around an actor and you want Begin/End overlaps to trigger.
+	 * @see UpdateOverlapsMethodDuringLevelStreaming
+	 */
+	UPROPERTY(EditAnywhere, BlueprintReadWrite, Category=Collision)
 	uint8 bGenerateOverlapEventsDuringLevelStreaming:1;
 
 	/** Whether this actor should not be affected by world origin shifting. */
@@ -322,6 +336,9 @@ private:
 	 */
 	uint8 bActorInitialized:1;
 
+	/** Set when DispatchBeginPlay() triggers from level streaming, and cleared afterwards. @see IsActorBeginningPlayFromLevelStreaming(). */
+	uint8 bActorBeginningPlayFromLevelStreaming:1;
+
 	/** Whether we've tried to register tick functions. Reset when they are unregistered. */
 	uint8 bTickFunctionsRegistered:1;
 
@@ -361,9 +378,54 @@ private:
 
 	static uint32 BeginPlayCallDepth;
 
+protected:
+
+	/**
+	 * Condition for calling UpdateOverlaps() to initialize overlap state when loaded in during level streaming.
+	 * If set to 'UseConfigDefault', the default specified in ini (displayed in 'DefaultUpdateOverlapsMethodDuringLevelStreaming') will be used.
+	 * If overlaps are not initialized, this actor and attached components will not have an initial state of what objects are touching it,
+	 * and overlap events may only come in once one of those objects update overlaps themselves (for example when moving).
+	 * However if an object touching it *does* initialize state, both objects will know about their touching state with each other.
+	 * This can be a potentially large performance savings during level streaming, and is safe if the object and others initially
+	 * overlapping it do not need the overlap state because they will not trigger overlap notifications.
+	 * 
+	 * Note that if 'bGenerateOverlapEventsDuringLevelStreaming' is true, overlaps are always updated in this case, but that flag
+	 * determines whether the Begin/End overlap events are triggered.
+	 * 
+	 * @see bGenerateOverlapEventsDuringLevelStreaming, DefaultUpdateOverlapsMethodDuringLevelStreaming, GetUpdateOverlapsMethodDuringLevelStreaming()
+	 */
+	UPROPERTY(Category=Collision, EditAnywhere)
+	EActorUpdateOverlapsMethod UpdateOverlapsMethodDuringLevelStreaming;
+
+public:
+
+	/** Get the method used to UpdateOverlaps() when loaded via level streaming. Resolves the 'UseConfigDefault' option to the class default specified in config. */
+	EActorUpdateOverlapsMethod GetUpdateOverlapsMethodDuringLevelStreaming() const;
+
+private:
+
+	/**
+	 * Default value taken from config file for this class when 'UseConfigDefault' is chosen for
+	 * 'UpdateOverlapsMethodDuringLevelStreaming'. This allows a default to be chosen per class in the matching config.
+	 * For example, for Actor it could be specified in DefaultEngine.ini as:
+	 * 
+	 * [/Script/Engine.Actor]
+	 * DefaultUpdateOverlapsMethodDuringLevelStreaming = OnlyUpdateMovable
+	 *
+	 * Another subclass could set their default to something different, such as:
+	 *
+	 * [/Script/Engine.BlockingVolume]
+	 * DefaultUpdateOverlapsMethodDuringLevelStreaming = NeverUpdate
+	 * 
+	 * @see UpdateOverlapsMethodDuringLevelStreaming
+	 */
+	UPROPERTY(Config, Category = Collision, VisibleAnywhere)
+	EActorUpdateOverlapsMethod DefaultUpdateOverlapsMethodDuringLevelStreaming;
+
 	/** Describes how much control the remote machine has over the actor. */
 	UPROPERTY(Replicated, Transient)
-	TEnumAsByte<enum ENetRole> RemoteRole;
+	TEnumAsByte<enum ENetRole> RemoteRole;	
+
 
 public:
 	/**
@@ -1296,9 +1358,16 @@ protected:
 	/** Overridable native event for when play begins for this actor. */
 	virtual void BeginPlay();
 
+	/** Event to notify blueprints this actor is being deleted or removed from a level. */
+	UFUNCTION(BlueprintImplementableEvent, meta=(Keywords = "delete", DisplayName = "End Play"))
+	void ReceiveEndPlay(EEndPlayReason::Type EndPlayReason);
+
+	/** Overridable function called whenever this actor is being removed from a level */
+	virtual void EndPlay(const EEndPlayReason::Type EndPlayReason);
+
 public:
 	/** Initiate a begin play call on this Actor, will handle calling in the correct order. */
-	void DispatchBeginPlay();
+	void DispatchBeginPlay(bool bFromLevelStreaming = false);
 
 	/** Returns whether an actor has been initialized for gameplay */
 	bool IsActorInitialized() const { return bActorInitialized; }
@@ -1309,6 +1378,9 @@ public:
 	/** Returns whether an actor has had BeginPlay called on it (and not subsequently had EndPlay called) */
 	bool HasActorBegunPlay() const { return ActorHasBegunPlay == EActorBeginPlayState::HasBegunPlay; }
 
+	/** Returns whether an actor has called BeginPlay(). */
+	bool IsActorBeginningPlayFromLevelStreaming() const { return bActorBeginningPlayFromLevelStreaming; }
+
 	/** Returns true if this actor is currently being destroyed, some gameplay events may be unsafe */
 	UFUNCTION(BlueprintCallable, Category="Game")
 	bool IsActorBeingDestroyed() const 
@@ -1483,10 +1555,6 @@ public:
 	UPROPERTY(BlueprintAssignable, Category="Game")
 	FActorDestroyedSignature OnDestroyed;
 
-	/** Event to notify blueprints this actor is being deleted or removed from a level. */
-	UFUNCTION(BlueprintImplementableEvent, meta=(Keywords = "delete", DisplayName = "End Play"))
-	void ReceiveEndPlay(EEndPlayReason::Type EndPlayReason);
-
 	/** Event triggered when the actor is being deleted or removed from a level. */
 	UPROPERTY(BlueprintAssignable, Category="Game")
 	FActorEndPlaySignature OnEndPlay;
@@ -2349,9 +2417,6 @@ public:
 	/** Non-virtual function to evaluate which portions of the EndPlay process should be dispatched for each actor */
 	void RouteEndPlay(const EEndPlayReason::Type EndPlayReason);
 
-	/** Overridable function called whenever this actor is being removed from a level */
-	virtual void EndPlay(const EEndPlayReason::Type EndPlayReason);
-
 	/**
 	 * Iterates up the movement base chain to see whether or not this Actor is based on the given Actor, defaults to checking attachment
 	 * @param Other the Actor to test for
diff --git a/Engine/Source/Runtime/Engine/Private/Actor.cpp b/Engine/Source/Runtime/Engine/Private/Actor.cpp
--- a/Engine/Source/Runtime/Engine/Private/Actor.cpp
+++ b/Engine/Source/Runtime/Engine/Private/Actor.cpp
@@ -130,7 +130,12 @@ void AActor::InitializeDefaults()
 	bAllowReceiveTickEventOnDedicatedServer = true;
 	bRelevantForNetworkReplays = true;
 	bRelevantForLevelBounds = true;
+	
+	// Overlap collision settings
 	bGenerateOverlapEventsDuringLevelStreaming = false;
+	UpdateOverlapsMethodDuringLevelStreaming = EActorUpdateOverlapsMethod::UseConfigDefault;
+	DefaultUpdateOverlapsMethodDuringLevelStreaming = EActorUpdateOverlapsMethod::OnlyUpdateMovable;
+	
 	bHasDeferredComponentRegistration = false;
 #if WITH_EDITORONLY_DATA
 	PivotOffset = FVector::ZeroVector;
@@ -3197,15 +3202,6 @@ void AActor::PostActorConstruction()
 		Modify(false);
 		ClearPendingKill();
 	}
-
-	if (!IsPendingKill())
-	{
-		// Components are all there and we've begun play, init overlapping state
-		if (!bDeferBeginPlayAndUpdateOverlaps)
-		{
-			UpdateOverlaps();
-		}
-	}
 }
 
 void AActor::SetReplicates(bool bInReplicates)
@@ -3292,8 +3288,6 @@ void AActor::PostNetInit()
 			DispatchBeginPlay();
 		}
 	}
-
-	UpdateOverlaps();
 }
 
 void AActor::ExchangeNetRoles(bool bRemoteOwned)
@@ -3317,14 +3311,27 @@ void AActor::SwapRoles()
 	ForcePropertyCompare();
 }
 
-void AActor::DispatchBeginPlay()
+EActorUpdateOverlapsMethod AActor::GetUpdateOverlapsMethodDuringLevelStreaming() const
+{
+	if (UpdateOverlapsMethodDuringLevelStreaming == EActorUpdateOverlapsMethod::UseConfigDefault)
+	{
+		// In the case of a default value saying "use defaults", pick something else.
+		return (DefaultUpdateOverlapsMethodDuringLevelStreaming != EActorUpdateOverlapsMethod::UseConfigDefault) ? DefaultUpdateOverlapsMethodDuringLevelStreaming : EActorUpdateOverlapsMethod::AlwaysUpdate;
+	}
+	return UpdateOverlapsMethodDuringLevelStreaming;
+}
+
+void AActor::DispatchBeginPlay(bool bFromLevelStreaming)
 {
 	UWorld* World = (!HasActorBegunPlay() && !IsPendingKill() ? GetWorld() : nullptr);
 
 	if (World)
 	{
+		ensureMsgf(ActorHasBegunPlay == EActorBeginPlayState::HasNotBegunPlay, TEXT("BeginPlay was called on actor %s which was in state %d"), *GetPathName(), (int32)ActorHasBegunPlay);
 		const uint32 CurrentCallDepth = BeginPlayCallDepth++;
 
+		bActorBeginningPlayFromLevelStreaming = bFromLevelStreaming;
+		ActorHasBegunPlay = EActorBeginPlayState::BeginningPlay;
 		BeginPlay();
 
 		ensure(BeginPlayCallDepth - 1 == CurrentCallDepth);
@@ -3336,19 +3343,63 @@ void AActor::DispatchBeginPlay()
 			// get to the point we set bActorWantsDestroyDuringBeginPlay to true
 			World->DestroyActor(this, true); 
 		}
+		
+		if (!IsPendingKill())
+		{
+			// Initialize overlap state
+			if (!bFromLevelStreaming)
+			{
+				UpdateOverlaps();
+			}
+			else
+			{
+				// Note: Conditionally doing notifies here since loading or streaming in isn't actually conceptually beginning a touch.
+				//	     Rather, it was always touching and the mechanics of loading is just an implementation detail.
+				if (bGenerateOverlapEventsDuringLevelStreaming)
+				{
+					UpdateOverlaps(bGenerateOverlapEventsDuringLevelStreaming);
+				}
+				else
+				{
+					bool bUpdateOverlaps = true;
+					const EActorUpdateOverlapsMethod UpdateMethod = GetUpdateOverlapsMethodDuringLevelStreaming();
+					switch (UpdateMethod)
+					{
+					case EActorUpdateOverlapsMethod::OnlyUpdateMovable:
+						bUpdateOverlaps = IsRootComponentMovable();
+						break;
+
+					case EActorUpdateOverlapsMethod::NeverUpdate:
+						bUpdateOverlaps = false;
+						break;
+
+					case EActorUpdateOverlapsMethod::AlwaysUpdate:
+					default:
+						bUpdateOverlaps = true;
+						break;
+					}
+
+					if (bUpdateOverlaps)
+					{
+						UpdateOverlaps(bGenerateOverlapEventsDuringLevelStreaming);
+					}
+				}
+			}
+		}
+
+		bActorBeginningPlayFromLevelStreaming = false;
 	}
 }
 
 void AActor::BeginPlay()
 {
-	ensureMsgf(ActorHasBegunPlay == EActorBeginPlayState::HasNotBegunPlay, TEXT("BeginPlay was called on actor %s which was in state %d"), *GetPathName(), (int32)ActorHasBegunPlay);
+	ensureMsgf(ActorHasBegunPlay == EActorBeginPlayState::BeginningPlay, TEXT("BeginPlay was called on actor %s which was in state %d"), *GetPathName(), (int32)ActorHasBegunPlay);
 	SetLifeSpan( InitialLifeSpan );
 	RegisterAllActorTickFunctions(true, false); // Components are done below.
 
 	TInlineComponentArray<UActorComponent*> Components;
 	GetComponents(Components);
 
-	ActorHasBegunPlay = EActorBeginPlayState::BeginningPlay;
 	for (UActorComponent* Component : Components)
 	{
 		// bHasBegunPlay will be true for the component if the component was renamed and moved to a new outer during initialization
diff --git a/Engine/Source/Runtime/Engine/Private/Components/ChildActorComponent.cpp b/Engine/Source/Runtime/Engine/Private/Components/ChildActorComponent.cpp
--- a/Engine/Source/Runtime/Engine/Private/Components/ChildActorComponent.cpp
+++ b/Engine/Source/Runtime/Engine/Private/Components/ChildActorComponent.cpp
@@ -689,6 +689,8 @@ void UChildActorComponent::BeginPlay()
 
 	if (ChildActor && !ChildActor->HasActorBegunPlay())
 	{
-		ChildActor->DispatchBeginPlay();
+		const AActor* Owner = GetOwner();
+		const bool bFromLevelStreaming = Owner ? Owner->IsActorBeginningPlayFromLevelStreaming() : false;
+		ChildActor->DispatchBeginPlay(bFromLevelStreaming);
 	}
 }
diff --git a/Engine/Source/Runtime/Engine/Private/Components/PrimitiveComponent.cpp b/Engine/Source/Runtime/Engine/Private/Components/PrimitiveComponent.cpp
--- a/Engine/Source/Runtime/Engine/Private/Components/PrimitiveComponent.cpp
+++ b/Engine/Source/Runtime/Engine/Private/Components/PrimitiveComponent.cpp
@@ -3001,16 +3001,21 @@ bool UPrimitiveComponent::UpdateOverlapsImpl(const TOverlapArrayView* NewPending
 	SCOPE_CYCLE_COUNTER(STAT_UpdateOverlaps); 
 	SCOPE_CYCLE_UOBJECT(ComponentScope, this);
 
+	// if we haven't begun play, we're still setting things up (e.g. we might be inside one of the construction scripts)
+	// so we don't want to generate overlaps yet. There is no need to update children yet either, they will update once we are allowed to as well.
+	const AActor* const MyActor = GetOwner();
+	if (MyActor && !MyActor->HasActorBegunPlay() && !MyActor->IsActorBeginningPlay())
+	{
+		return false;
+	}
+
 	bool bCanSkipUpdateOverlaps = true;
 
 	// first, dispatch any pending overlaps
 	if (GetGenerateOverlapEvents() && IsQueryCollisionEnabled())	//TODO: should modifying query collision remove from mayoverlapevents?
 	{
 		bCanSkipUpdateOverlaps = false;
-		// if we haven't begun play, we're still setting things up (e.g. we might be inside one of the construction scripts)
-		// so we don't want to generate overlaps yet.
-		AActor* const MyActor = GetOwner();
-		if ( MyActor && MyActor->IsActorInitialized() )
+		if (MyActor)
 		{
 			const FTransform PrevTransform = GetComponentTransform();
 			// If we are the root component we ignore child components. Those children will update their overlaps when we descend into the child tree.
diff --git a/Engine/Source/Runtime/Engine/Private/Level.cpp b/Engine/Source/Runtime/Engine/Private/Level.cpp
--- a/Engine/Source/Runtime/Engine/Private/Level.cpp
+++ b/Engine/Source/Runtime/Engine/Private/Level.cpp
@@ -1836,11 +1836,6 @@ void ULevel::RouteActorInitialize()
 					ActorsToBeginPlay.Add(Actor);
 				}
 			}
-
-			// Components are all set up, init touching state.
-			// Note: Not doing notifies here since loading or streaming in isn't actually conceptually beginning a touch.
-			//	     Rather, it was always touching and the mechanics of loading is just an implementation detail.
-			Actor->UpdateOverlaps(Actor->bGenerateOverlapEventsDuringLevelStreaming);
 		}
 	}
 
@@ -1849,7 +1844,7 @@ void ULevel::RouteActorInitialize()
 	{
 		AActor* Actor = ActorsToBeginPlay[ActorIndex];
 		SCOPE_CYCLE_COUNTER(STAT_ActorBeginPlay);
-		Actor->DispatchBeginPlay();
+		Actor->DispatchBeginPlay(/*bFromLevelStreaming*/ true);
 	}
 }
 
