Fix incorrect warning log

diff --git a/Engine/Source/Runtime/Engine/Private/Animation/AnimationUtils.cpp b/Engine/Source/Runtime/Engine/Private/Animation/AnimationUtils.cpp
--- a/Engine/Source/Runtime/Engine/Private/Animation/AnimationUtils.cpp
+++ b/Engine/Source/Runtime/Engine/Private/Animation/AnimationUtils.cpp
@@ -1599,7 +1599,7 @@ void FAnimationUtils::CompressAnimSequenceExplicit(
 						OriginalSize, AfterOriginalRecompression, CompressorStats.CurrentSize,
 						CompressorStats.WinningCompressorError);
 
-					UE_LOG(LogAnimation, Warning, TEXT("  CompressedTrackOffsets(%d) CompressedByteStream(%d) CompressedScaleOffsets(%d) CompressedSegments(%d)"),
+					UE_LOG(LogAnimationCompression, Log, TEXT("  CompressedTrackOffsets(%d) CompressedByteStream(%d) CompressedScaleOffsets(%d) CompressedSegments(%d)"),
 						OutCompressedData.CompressedTrackOffsets.Num(),
 						OutCompressedData.CompressedByteStream.Num(),
 						OutCompressedData.CompressedScaleOffsets.GetMemorySize(),
