Fixed divide by zero case in UTimelineComponent::TickComponent
[RN] fix general

diff --git a/Engine/Source/Runtime/Engine/Private/Timeline.cpp b/Engine/Source/Runtime/Engine/Private/Timeline.cpp
--- a/Engine/Source/Runtime/Engine/Private/Timeline.cpp
+++ b/Engine/Source/Runtime/Engine/Private/Timeline.cpp
@@ -17,6 +17,7 @@
 #include "Components/TimelineComponent.h"
 #include "Engine/World.h"
 #include "ProfilingDebugging/CsvProfiler.h"
+#include "Misc/App.h"
 
 DEFINE_LOG_CATEGORY_STATIC(LogTimeline, Log, All);
 
@@ -669,19 +670,14 @@ void UTimelineComponent::TickComponent(float DeltaTime, enum ELevelTick TickType
 
 	if (bIgnoreTimeDilation)
 	{
-		AActor* const OwnerActor = GetOwner();
-		if (OwnerActor)
+		// Get the raw, undilated delta time.
+		DeltaTime = FApp::GetDeltaTime();
+		const UWorld* World = GetWorld();
+		if (const AWorldSettings* WorldSettings = World ? World->GetWorldSettings() : nullptr)
 		{
-			DeltaTime /= OwnerActor->GetActorTimeDilation();
-		}
-		else
-		{
-			// no Actor for some reason, use the world time dilation as fallback
-			UWorld* const W = GetWorld();
-			if (W)
-			{
-				DeltaTime /= W->GetWorldSettings()->GetEffectiveTimeDilation();
-			}
+			// Clamp DeltaTime in the same way as before.
+			// UWorld::Tick called AWorldSettings::FixupDeltaSeconds, which clamped between Min and MaxUndilatedFrameTime.
+			DeltaTime = FMath::Clamp(DeltaTime, WorldSettings->MinUndilatedFrameTime, WorldSettings->MaxUndilatedFrameTime);
 		}
 	}
 
