Merging using //UE4/Dev-Framework_to_/UE4/Release-4.23

Make AActor::IsDataValid call IsDataValid on the actor's components.

diff --git a/Engine/Source/Runtime/Engine/Private/ActorEditor.cpp b/Engine/Source/Runtime/Engine/Private/ActorEditor.cpp
--- a/Engine/Source/Runtime/Engine/Private/ActorEditor.cpp
+++ b/Engine/Source/Runtime/Engine/Private/ActorEditor.cpp
@@ -1001,7 +1001,24 @@ EDataValidationResult AActor::IsDataValid(TArray<FText>& ValidationErrors)
 		bSuccess = false;
 	}
 
-	return bSuccess ? EDataValidationResult::Valid : EDataValidationResult::Invalid;
+	EDataValidationResult Result = bSuccess ? EDataValidationResult::Valid : EDataValidationResult::Invalid;
+
+	// check the components
+	for (UActorComponent* Component : GetComponents())
+	{
+		if (Component)
+		{
+			// if any component is invalid, our result is invalid
+			// in the future we may want to update this to say that the actor was not validated if any of its components returns EDataValidationResult::NotValidated
+			EDataValidationResult ComponentResult = Component->IsDataValid(ValidationErrors);
+			if (ComponentResult == EDataValidationResult::Invalid)
+			{
+				Result = EDataValidationResult::Invalid;
+			}
+		}
+	}
+
+	return Result;
 }
 #undef LOCTEXT_NAMESPACE
 
