Fix crashing when calling SetPhysicalMaterialOverride in constructor. #jira UE-62367 #rb none

diff --git a/Engine/Source/Runtime/Engine/Private/PhysicsEngine/BodyInstance.cpp b/Engine/Source/Runtime/Engine/Private/PhysicsEngine/BodyInstance.cpp
--- a/Engine/Source/Runtime/Engine/Private/PhysicsEngine/BodyInstance.cpp
+++ b/Engine/Source/Runtime/Engine/Private/PhysicsEngine/BodyInstance.cpp
@@ -2424,7 +2424,13 @@ void FBodyInstance::GetComplexPhysicalMaterials(TArray<UPhysicalMaterial*>& Phys
 
 void FBodyInstance::GetComplexPhysicalMaterials(const FBodyInstance*, TWeakObjectPtr<UPrimitiveComponent> OwnerComp, TArray<UPhysicalMaterial*>& OutPhysicalMaterials)
 {
-	check(GEngine->DefaultPhysMaterial != NULL);
+	if(!GEngine || !GEngine->DefaultPhysMaterial)
+	{
+		UE_LOG(LogPhysics, Error, TEXT("FBodyInstance::GetComplexPhysicalMaterials : GEngine not initialized! Cannot call this during native CDO construction, wrap with if(!HasAnyFlags(RF_ClassDefaultObject)) or move out of constructor, material parameters will not be correct."));
+
+		return;
+	}
+
 	// See if the Material has a PhysicalMaterial
 	UPrimitiveComponent* PrimComp = OwnerComp.Get();
 	if (PrimComp)
