Extend AudioComponent Stop API to provide the ability to issue a delayed stop request (provides ability to fix sequencer scrub issue)

diff --git a/Engine/Source/Runtime/Engine/Classes/Components/AudioComponent.h b/Engine/Source/Runtime/Engine/Classes/Components/AudioComponent.h
--- a/Engine/Source/Runtime/Engine/Classes/Components/AudioComponent.h
+++ b/Engine/Source/Runtime/Engine/Classes/Components/AudioComponent.h
@@ -36,7 +36,7 @@ DECLARE_DYNAMIC_MULTICAST_DELEGATE_TwoParams(FOnAudioPlaybackPercent, const clas
 /** shadow delegate declaration for above */
 DECLARE_MULTICAST_DELEGATE_ThreeParams(FOnAudioPlaybackPercentNative, const class UAudioComponent*, const class USoundWave*, const float);
 
-/** 
+/**
 * Called while a sound plays and returns the sound's envelope value (using an envelope follower in the audio renderer).
 * This only works in the audio mixer.
 */
@@ -74,7 +74,7 @@ struct FAudioComponentParam
 	// Value of the parameter when used as a boolean
 	UPROPERTY(EditAnywhere, BlueprintReadWrite, Category=AudioComponentParam)
 	bool BoolParam;
-	
+
 	// Value of the parameter when used as an integer
 	UPROPERTY(EditAnywhere, BlueprintReadWrite, Category=AudioComponentParam)
 	int32 IntParam;
@@ -379,12 +379,16 @@ public:
 
 	/** Start a sound playing on an audio component */
 	UFUNCTION(BlueprintCallable, Category="Audio|Components|Audio")
-	virtual void Play(float StartTime = 0.f);
+	virtual void Play(float StartTime = 0.0f);
 
-	/** Stop an audio component playing its sound cue, issue any delegates if needed */
+	/** Stop an audio component's sound, issue any delegates if needed */
 	UFUNCTION(BlueprintCallable, Category="Audio|Components|Audio")
 	virtual void Stop();
 
+	/** Cues request to stop sound after the provided delay, stopping immediately if delay is zero or negative */
+	UFUNCTION(BlueprintCallable, Category="Audio|Components|Audio")
+	void StopDelayed(float DelayTime);
+
 	/** Pause an audio component playing its sound cue, issue any delegates if needed */
 	UFUNCTION(BlueprintCallable, Category = "Audio|Components|Audio")
 	void SetPaused(bool bPause);
@@ -408,7 +412,7 @@ public:
 	/** Set a boolean instance parameter for use in sound cues played by this audio component */
 	UFUNCTION(BlueprintCallable, Category="Audio|Components|Audio", meta=(DisplayName="Set Boolean Parameter"))
 	void SetBoolParameter(FName InName, bool InBool);
-	
+
 	/** Set an integer instance parameter for use in sound cues played by this audio component */
 	UFUNCTION(BlueprintCallable, Category="Audio|Components|Audio", meta=(DisplayName="Set Integer Parameter"))
 	void SetIntParameter(FName InName, int32 InInt);
@@ -459,26 +463,26 @@ public:
 	UFUNCTION(BlueprintCallable, Category = "Audio|Components|Audio")
 	bool HasCookedAmplitudeEnvelopeData() const;
 
-	/** 
-	* Retrieves the current-time cooked spectral data of the sounds playing on the audio component. 
-	* Spectral data is averaged and interpolated for all playing sounds on this audio component. 
-	* Returns true if there is data and the audio component is playing. 
+	/**
+	* Retrieves the current-time cooked spectral data of the sounds playing on the audio component.
+	* Spectral data is averaged and interpolated for all playing sounds on this audio component.
+	* Returns true if there is data and the audio component is playing.
 	*/
 	UFUNCTION(BlueprintCallable, Category = "Audio|Components|Audio")
 	bool GetCookedFFTData(const TArray<float>& FrequenciesToGet, TArray<FSoundWaveSpectralData>& OutSoundWaveSpectralData);
 
-	/** 
-	* Retrieves the current-time cooked spectral data of the sounds playing audio component. 
+	/**
+	* Retrieves the current-time cooked spectral data of the sounds playing audio component.
 	* Spectral data is not averaged or interpolated. Instead an array of data with all playing sound waves with cooked data is returned.
-	* Returns true if there is data and the audio component is playing. 
+	* Returns true if there is data and the audio component is playing.
 	*/
 	UFUNCTION(BlueprintCallable, Category = "Audio|Components|Audio")
 	bool GetCookedFFTDataForAllPlayingSounds(TArray<FSoundWaveSpectralDataPerSound>& OutSoundWaveSpectralData);
 
-	/** 
-	* Retrieves the current-time cooked envelope data of the playing audio component. 
+	/**
+	* Retrieves the current-time cooked envelope data of the playing audio component.
 	* Cooked data is interpolated and averaged across all playing sound waves.
-	* Returns true if there is data and the audio component is playing. 
+	* Returns true if there is data and the audio component is playing.
 	*/
 	UFUNCTION(BlueprintCallable, Category = "Audio|Components|Audio")
 	bool GetCookedEnvelopeData(float& OutEnvelopeData);
@@ -616,9 +620,6 @@ protected:
 	/** Utility function called by Play and FadeIn to start a sound playing. */
 	void PlayInternal(const float StartTime = 0.f, const float FadeInDuration = 0.f, const float FadeVolumeLevel = 1.f);
 
-private:
-	void StopInternal();
-
 #if WITH_EDITORONLY_DATA
 	/** Utility function that updates which texture is displayed on the sprite dependent on the properties of the Audio Component. */
 	void UpdateSpriteTexture();
diff --git a/Engine/Source/Runtime/Engine/Private/Components/AudioComponent.cpp b/Engine/Source/Runtime/Engine/Private/Components/AudioComponent.cpp
--- a/Engine/Source/Runtime/Engine/Private/Components/AudioComponent.cpp
+++ b/Engine/Source/Runtime/Engine/Private/Components/AudioComponent.cpp
@@ -575,7 +575,7 @@ void UAudioComponent::AdjustVolumeInternal(float AdjustVolumeDuration, float Adj
 		return;
 	}
 
-	if (AdjustVolumeDuration == 0.0f && AdjustVolumeLevel == 0.0f)
+	if (FMath::IsNearlyZero(AdjustVolumeDuration) && FMath::IsNearlyZero(AdjustVolumeLevel))
 	{
 		Stop();
 		return;
@@ -597,30 +597,31 @@ void UAudioComponent::AdjustVolumeInternal(float AdjustVolumeDuration, float Adj
 			return;
 		}
 
+		const float NewTargetStopTime = ActiveSound->PlaybackTime + AdjustVolumeDuration;
 		if (ActiveSound->FadeOut == FActiveSound::EFadeOut::Concurrency)
 		{
 			// Ignore adjust volume request if non-zero and currently voice stealing.
-			if (AdjustVolumeLevel != 0.0f)
+			if (!FMath::IsNearlyZero(AdjustVolumeLevel))
 			{
 				return;
 			}
 
 			// Ignore request of longer fade out than active target if active is concurrency (voice stealing) fade.
-			if (AdjustVolumeDuration > ActiveSound->TargetAdjustVolumeStopTime)
+			if (NewTargetStopTime > ActiveSound->TargetAdjustVolumeStopTime)
 			{
 				return;
 			}
 		}
 		else
 		{
-			ActiveSound->FadeOut = AdjustVolumeLevel == 0.0f ? FActiveSound::EFadeOut::User : FActiveSound::EFadeOut::None;
+			ActiveSound->FadeOut = FMath::IsNearlyZero(AdjustVolumeLevel) ? FActiveSound::EFadeOut::User : FActiveSound::EFadeOut::None;
 		}
 
 		ActiveSound->TargetAdjustVolumeMultiplier = AdjustVolumeLevel;
 
 		if (AdjustVolumeDuration > 0.0f)
 		{
-			ActiveSound->TargetAdjustVolumeStopTime = ActiveSound->PlaybackTime + AdjustVolumeDuration;
+			ActiveSound->TargetAdjustVolumeStopTime = NewTargetStopTime;
 		}
 		else
 		{
@@ -632,23 +633,82 @@ void UAudioComponent::AdjustVolumeInternal(float AdjustVolumeDuration, float Adj
 
 void UAudioComponent::Stop()
 {
-	StopInternal();
+	if (!bIsActive)
+	{
+		return;
+	}
+
+	FAudioDevice* AudioDevice = GetAudioDevice();
+	if (!AudioDevice)
+	{
+		return;
+	}
+
+	// Set this to immediately be inactive
+	bIsActive = false;
+
+	UE_LOG(LogAudio, Verbose, TEXT("%g: Stopping AudioComponent : '%s' with Sound: '%s'"),
+		GetWorld() ? GetWorld()->GetAudioTimeSeconds() : 0.0f, *GetFullName(),
+		Sound ? *Sound->GetName() : TEXT("nullptr"));
+
+	AudioDevice->StopActiveSound(AudioComponentID);
 }
 
-void UAudioComponent::StopInternal()
+void UAudioComponent::StopDelayed(float DelayTime)
 {
-	if (bIsActive)
+	// 1. Stop immediately if no delay time
+	if (DelayTime < 0.0f || FMath::IsNearlyZero(DelayTime))
 	{
-		// Set this to immediately be inactive
-		bIsActive = false;
+		Stop();
+		return;
+	}
 
-		UE_LOG(LogAudio, Verbose, TEXT( "%g: Stopping AudioComponent : '%s' with Sound: '%s'" ), GetWorld() ? GetWorld()->GetAudioTimeSeconds() : 0.0f, *GetFullName(), Sound ? *Sound->GetName() : TEXT( "nullptr" ) );
+	if (!bIsActive)
+	{
+		return;
+	}
 
-		if (FAudioDevice* AudioDevice = GetAudioDevice())
+	FAudioDevice* AudioDevice = GetAudioDevice();
+	if (!AudioDevice)
+	{
+		return;
+	}
+
+	// 2. Performs delayed stop with no fade
+	const uint64 InAudioComponentID = AudioComponentID;
+	DECLARE_CYCLE_STAT(TEXT("FAudioThreadTask.StopDelayed"), STAT_AudioAdjustVolume, STATGROUP_AudioThreadCommands);
+	FAudioThread::RunCommandOnAudioThread([AudioDevice, InAudioComponentID, DelayTime]()
+	{
+		FActiveSound* ActiveSound = AudioDevice->FindActiveSound(InAudioComponentID);
+		if (!ActiveSound)
 		{
-			AudioDevice->StopActiveSound(AudioComponentID);
+			return;
 		}
-	}
+
+		const float NewTargetStopTime = ActiveSound->PlaybackTime + DelayTime;
+		if (ActiveSound->FadeOut == FActiveSound::EFadeOut::Concurrency)
+		{
+			// Ignore request of longer fade out than active target if active is concurrency (voice stealing) fade.
+			if (NewTargetStopTime > ActiveSound->TargetAdjustVolumeStopTime)
+			{
+				return;
+			}
+		}
+		else
+		{
+			// Set fade to user, but don't adjust target volume, which will cause sound to stop abruptly as intended.
+			ActiveSound->FadeOut = FActiveSound::EFadeOut::User;
+		}
+
+		if (const USoundBase* StoppingSound = ActiveSound->GetSound())
+		{
+			UE_LOG(LogAudio, Verbose, TEXT("%g: Delayed Stop requested for sound '%s'"),
+				ActiveSound->GetWorld() ? ActiveSound->GetWorld()->GetAudioTimeSeconds() : 0.0f,
+				*StoppingSound->GetName());
+		}
+
+		ActiveSound->TargetAdjustVolumeStopTime = NewTargetStopTime;
+	}, GET_STATID(STAT_AudioAdjustVolume));
 }
 
 void UAudioComponent::SetPaused(bool bPause)
