Fix SkeletalMeshComponent LineTraceComponent and SweepComponent not working with bEnablePerPolyCollision. #jira UE-70652 #rb none

diff --git a/Engine/Source/Runtime/Engine/Private/SkeletalMeshComponentPhysics.cpp b/Engine/Source/Runtime/Engine/Private/SkeletalMeshComponentPhysics.cpp
--- a/Engine/Source/Runtime/Engine/Private/SkeletalMeshComponentPhysics.cpp
+++ b/Engine/Source/Runtime/Engine/Private/SkeletalMeshComponentPhysics.cpp
@@ -2192,17 +2192,26 @@ bool USkeletalMeshComponent::LineTraceComponent(struct FHitResult& OutHit, const
 	UWorld* const World = GetWorld();
 	bool bHaveHit = false;
 
-	float MinTime = MAX_FLT;
-	FHitResult Hit;
-	for (int32 BodyIdx=0; BodyIdx < Bodies.Num(); ++BodyIdx)
+	if (bEnablePerPolyCollision)
 	{
-		if (Bodies[BodyIdx] && Bodies[BodyIdx]->LineTrace(Hit, Start, End, Params.bTraceComplex, Params.bReturnPhysicalMaterial))
+		// Using PrimitiveComponent implementation
+		//as it intersects against mesh polys.
+		bHaveHit = UPrimitiveComponent::LineTraceComponent(OutHit, Start, End, Params);
+	}
+	else
+	{
+		float MinTime = MAX_FLT;
+		FHitResult Hit;
+		for (int32 BodyIdx = 0; BodyIdx < Bodies.Num(); ++BodyIdx)
 		{
-			bHaveHit = true;
-			if(MinTime > Hit.Time)
+			if (Bodies[BodyIdx] && Bodies[BodyIdx]->LineTrace(Hit, Start, End, Params.bTraceComplex, Params.bReturnPhysicalMaterial))
 			{
-				MinTime = Hit.Time;
-				OutHit = Hit;
+				bHaveHit = true;
+				if (MinTime > Hit.Time)
+				{
+					MinTime = Hit.Time;
+					OutHit = Hit;
+				}
 			}
 		}
 	}
@@ -2226,16 +2235,25 @@ bool USkeletalMeshComponent::SweepComponent( FHitResult& OutHit, const FVector S
 {
 	bool bHaveHit = false;
 
-	FHitResult Hit;
-	for (int32 BodyIdx=0; BodyIdx < Bodies.Num(); ++BodyIdx)
+	if (bEnablePerPolyCollision)
 	{
-		if (Bodies[BodyIdx] && Bodies[BodyIdx]->Sweep(Hit, Start, End, ShapeWorldRotation, CollisionShape, bTraceComplex))
+		// Using PrimitiveComponent implementation
+		//as it intersects against mesh polys.
+		bHaveHit =  UPrimitiveComponent::SweepComponent(OutHit, Start, End, ShapeWorldRotation, CollisionShape, bTraceComplex);
+	}
+	else
+	{
+		FHitResult Hit;
+		for (int32 BodyIdx = 0; BodyIdx < Bodies.Num(); ++BodyIdx)
 		{
-			if (!bHaveHit || Hit.Time < OutHit.Time)
+			if (Bodies[BodyIdx] && Bodies[BodyIdx]->Sweep(Hit, Start, End, ShapeWorldRotation, CollisionShape, bTraceComplex))
 			{
-				OutHit = Hit;
+				if (!bHaveHit || Hit.Time < OutHit.Time)
+				{
+					OutHit = Hit;
+				}
+				bHaveHit = true;
 			}
-			bHaveHit = true;
 		}
 	}
 
