Fix replicating physics twice in PhysX Scene. #jira UE-66893 #rb none

diff --git a/Engine/Source/Runtime/Engine/Private/PhysicsEngine/PhysScene_PhysX.cpp b/Engine/Source/Runtime/Engine/Private/PhysicsEngine/PhysScene_PhysX.cpp
--- a/Engine/Source/Runtime/Engine/Private/PhysicsEngine/PhysScene_PhysX.cpp
+++ b/Engine/Source/Runtime/Engine/Private/PhysicsEngine/PhysScene_PhysX.cpp
@@ -1105,14 +1105,6 @@ void FPhysScene_PhysX::TickPhysScene(FGraphEventRef& InOutCompletionEvent)
 	}
 #endif
 
-	// Replicate physics
-#if WITH_PHYSX
-	if (bSimulateScene && PhysicsReplication)
-	{
-		PhysicsReplication->Tick(AveragedFrameTime);
-	}
-#endif
-
 	float PreTickTime = IsSubstepping() ? UseDelta : AveragedFrameTime;
 
 	// Broadcast 'pre tick' delegate
