Partially reverting CurveAtlas material expression Compile() behavior. The default data being used was compiling with legacy data from pre-refactor instead of respecting the actually selected curve. The change here restores re-evaluation of the curve during material translation.

diff --git a/Engine/Shaders/Public/ShaderVersion.ush b/Engine/Shaders/Public/ShaderVersion.ush
--- a/Engine/Shaders/Public/ShaderVersion.ush
+++ b/Engine/Shaders/Public/ShaderVersion.ush
@@ -3,4 +3,4 @@
 // in Platform.ush (which should be included in any shader) it allows to invalidate the shader DDC.
 // 
 // If you are merging streams and there is a conflict with this GUID you should make a new GUID rather than taking one or the other.
-// GUID = 676CDC95922E43F19A67C4089A159E89
+// GUID = EDD1AC120A734A1DABC42835826F4B17
\ No newline at end of file
diff --git a/Engine/Source/Runtime/Engine/Private/Materials/MaterialExpressions.cpp b/Engine/Source/Runtime/Engine/Private/Materials/MaterialExpressions.cpp
--- a/Engine/Source/Runtime/Engine/Private/Materials/MaterialExpressions.cpp
+++ b/Engine/Source/Runtime/Engine/Private/Materials/MaterialExpressions.cpp
@@ -15733,6 +15733,10 @@ int32 UMaterialExpressionCurveAtlasRowParameter::Compile(class FMaterialCompiler
 {
 	if (Atlas && Curve)
 	{
+		// Retrieve the curve index directly from the atlas rather than relying on the scalar parameter defaults
+		int32 CurveIndex = 0;
+		Atlas->GetCurveIndex(Curve, CurveIndex);
+		DefaultValue = (float)CurveIndex;
 		int32 Slot = Compiler->ScalarParameter(ParameterName, DefaultValue);
 
 		// Get Atlas texture object and texture size
