Make the nak packet received log message VeryVerbose - it has the potential to be very spammy.

diff --git a/Engine/Source/Runtime/Engine/Private/DataReplication.cpp b/Engine/Source/Runtime/Engine/Private/DataReplication.cpp
--- a/Engine/Source/Runtime/Engine/Private/DataReplication.cpp
+++ b/Engine/Source/Runtime/Engine/Private/DataReplication.cpp
@@ -866,7 +866,7 @@ void FObjectReplicator::ReceivedNak( int32 NakPacketId )
 					}
 					else if (NakPacketId >= Rec->OutPacketIdRange.First && NakPacketId <= Rec->OutPacketIdRange.Last)
 					{
-						UE_LOG(LogNet, Verbose, TEXT("Restoring Previous Base State of dynamic property. Channel: %s, NakId: %d, First: %d, Last: %d, Address: %s)"), *OwningChannel->Describe(), NakPacketId, Rec->OutPacketIdRange.First, Rec->OutPacketIdRange.Last, *Connection->LowLevelGetRemoteAddress(true));
+						UE_LOG(LogNet, VeryVerbose, TEXT("Restoring Previous Base State of dynamic property. Channel: %s, NakId: %d, First: %d, Last: %d, Address: %s)"), *OwningChannel->Describe(), NakPacketId, Rec->OutPacketIdRange.First, Rec->OutPacketIdRange.Last, *Connection->LowLevelGetRemoteAddress(true));
 
 						// The Nack'd packet did update this property, so we need to replace the buffer in RecentDynamic
 						// with the buffer we used to create this update (which was dropped), so that the update will be recreated on the next replicate actor
