Don't chunk streamable animations on a dedicated server

diff --git a/Engine/Source/Runtime/Engine/Private/Animation/AnimStreamable.cpp b/Engine/Source/Runtime/Engine/Private/Animation/AnimStreamable.cpp
--- a/Engine/Source/Runtime/Engine/Private/Animation/AnimStreamable.cpp
+++ b/Engine/Source/Runtime/Engine/Private/Animation/AnimStreamable.cpp
@@ -522,7 +522,7 @@ void UAnimStreamable::RequestCompressedData(const ITargetPlatform* Platform)
 	float ChunkSizeSeconds = GetChunkSizeSeconds(Platform);
 
 	uint32 NumChunks = 1;
-	if (ChunkSizeSeconds > 0.f) // <= 0.f signifies to not chunk
+	if (!Platform->IsServerOnly() && ChunkSizeSeconds > 0.f) // <= 0.f signifies to not chunk & don't have chunks on a server
 	{
 		ChunkSizeSeconds = FMath::Max(ChunkSizeSeconds, MINIMUM_CHUNK_SIZE);
 		const int32 InitialNumChunks = FMath::FloorToInt(SequenceLength / ChunkSizeSeconds);
