Reverting back some of the changes to the ADPCM decoder that I checked in with Stream Caching, because this source is so brittle. This version of the code is identical to the previous version, save for the fixes for non-streaming seeking and fixups to support the new GetLoadedChunk API.

diff --git a/Engine/Source/Runtime/Engine/Private/ADPCMAudioInfo.cpp b/Engine/Source/Runtime/Engine/Private/ADPCMAudioInfo.cpp
--- a/Engine/Source/Runtime/Engine/Private/ADPCMAudioInfo.cpp
+++ b/Engine/Source/Runtime/Engine/Private/ADPCMAudioInfo.cpp
@@ -43,7 +43,7 @@ FADPCMAudioInfo::~FADPCMAudioInfo(void)
 void FADPCMAudioInfo::SeekToTime(const float SeekTime)
 {
 	// Reset chunk handle in preperation for a new chunk.
-	CurrentCompressedChunkData = FAudioChunkHandle();
+	CurCompressedChunkData = nullptr;
 
 	if(SeekTime <= 0.0f)
 	{
@@ -341,27 +341,29 @@ bool FADPCMAudioInfo::StreamCompressedInfoInternal(USoundWave* Wave, struct FSou
 	check(StreamingSoundWave == Wave);
 
 	// Get the first chunk of audio data (should already be loaded)
-	// uint8 const* firstChunk = IStreamingManager::Get().GetAudioStreamingManager().GetLoadedChunk(Wave, 0, &CurrentChunkDataSize);
-	TArrayView<const uint8> ZerothChunk = Wave->GetZerothChunk();
+	uint8 const* FirstChunk = GetLoadedChunk(Wave, 0, CurrentChunkDataSize);
 
-	if (ZerothChunk.Num() <= 0)
+	if (FirstChunk == nullptr)
 	{
 		return false;
 	}
 
 	SrcBufferData = nullptr;
 	SrcBufferDataSize = 0;
-	
-	void*	FormatHeader;
-	
-	if (!WaveInfo.ReadWaveInfo(ZerothChunk.GetData(), ZerothChunk.Num(), NULL, true, &FormatHeader))
+
+	void* FormatHeader;
+
+	if (!WaveInfo.ReadWaveInfo((uint8*)FirstChunk, Wave->RunningPlatformData->Chunks[0].AudioDataSize, nullptr, true, &FormatHeader))
 	{
 		UE_LOG(LogAudio, Warning, TEXT("WaveInfo.ReadWaveInfo Failed"));
 		return false;
 	}
 
-	FirstChunkSampleDataOffset = WaveInfo.SampleDataStart - ZerothChunk.GetData();
-	CurrentChunkBufferOffset = FirstChunkSampleDataOffset;
+	SrcBufferData = FirstChunk;
+
+	FirstChunkSampleDataOffset = WaveInfo.SampleDataStart - FirstChunk;
+	CurrentChunkBufferOffset = 0;
+	CurCompressedChunkData = nullptr;
 	CurrentUncompressedBlockSampleIndex = 0;
 	CurrentChunkIndex = 0;
 	TotalSamplesStreamed = 0;
@@ -388,7 +390,6 @@ bool FADPCMAudioInfo::StreamCompressedInfoInternal(USoundWave* Wave, struct FSou
 
 		UncompressedBlockData = (uint8*)FMemory::Realloc(UncompressedBlockData, NumChannels * UncompressedBlockSize);
 		check(UncompressedBlockData != nullptr);
-		FMemory::Memzero(UncompressedBlockData, NumChannels * UncompressedBlockSize);
 	}
 	else if (Format == WAVE_FORMAT_LPCM)
 	{
@@ -454,36 +455,54 @@ bool FADPCMAudioInfo::StreamCompressedData(uint8* Destination, bool bLooping, ui
 		// We need to loop over the number of samples requested since an uncompressed block will not match the number of frames requested
 		while(BufferSize > 0)
 		{
-			const uint8* CompressedChunkPtr = nullptr;
-			uint32 CompressedChunkSize = 0;
-
-			GetChunkPtr(CompressedChunkPtr, CompressedChunkSize);
-
-			if(CurrentUncompressedBlockSampleIndex >= UncompressedBlockSize / sizeof(uint16))
+			if(CurCompressedChunkData == nullptr || CurrentUncompressedBlockSampleIndex >= UncompressedBlockSize / sizeof(uint16))
 			{
-				// If GetChunkPtr returns a nullptr, the buffer wasn't loaded in time.
-				if(CompressedChunkPtr == nullptr)
+				// we need to decompress another block of compressed data from the current chunk
+
+				if(CurCompressedChunkData == nullptr || CurrentChunkBufferOffset >= CurrentChunkDataSize)
 				{
-					// We only need to worry about missing the stream chunk if we were seeking. Seeking might cause a bit of latency with chunk loading. That is expected.
-					if (!bSeekPending)
+					// Get another chunk of compressed data from the streaming engine
+
+					// CurrentChunkIndex is used to keep track of the current chunk for loading/unloading by the streaming engine, but chunk 0 is
+					// preloaded so don't increment this when getting chunk 0, only later chunks. If previous chunk retrieval failed because it was
+					// not ready, don't re-increment the CurrentChunkIndex.  Only increment if seek not pending as seek determines the current chunk index
+					// and invalidates CurCompressedChunkData.
+					if(CurCompressedChunkData != nullptr)
 					{
-						// If we did not get it then just bail, CurrentChunkIndex will not get incremented on the next callback so in effect another attempt will be made to fetch the chunk
-						// Since audio streaming depends on the general data streaming mechanism used by other parts of the engine and new data is prefectched on the game tick thread its possible a game hickup can cause this
-						UE_LOG(LogAudio, Verbose, TEXT("Missed Deadline for chunk %d"), CurrentChunkIndex);
+						++CurrentChunkIndex;
 					}
 
-					// zero out remaining data and bail
-					FMemory::Memset(OutData, 0, BufferSize);
-					return false;
+					// Request the next chunk of data from the streaming engine
+					CurCompressedChunkData = GetLoadedChunk(StreamingSoundWave, CurrentChunkIndex, CurrentChunkDataSize);
+
+					if(CurCompressedChunkData == nullptr)
+					{
+						// We only need to worry about missing the stream chunk if we were seeking. Seeking might cause a bit of latency with chunk loading. That is expected.
+						if (!bSeekPending)
+						{
+							// If we did not get it then just bail, CurrentChunkIndex will not get incremented on the next callback so in effect another attempt will be made to fetch the chunk
+							// Since audio streaming depends on the general data streaming mechanism used by other parts of the engine and new data is prefectched on the game tick thread its possible a game hickup can cause this
+							UE_LOG(LogAudio, Verbose, TEXT("Missed Deadline chunk %d"), CurrentChunkIndex);
+						}
+
+						// zero out remaining data and bail
+						FMemory::Memset(OutData, 0, BufferSize);
+						return false;
+					}
+
+					// Set the current buffer offset accounting for the header in the first chunk
+					if (!bSeekPending)
+					{
+						CurrentChunkBufferOffset = CurrentChunkIndex == 0 ? FirstChunkSampleDataOffset : 0;
+					}
+					bSeekPending = false;
 				}
 
 				// Decompress one block for each channel and store it in UncompressedBlockData
 				for(int32 ChannelItr = 0; ChannelItr < NumChannels; ++ChannelItr)
 				{
-					check(CompressedChunkPtr && (CurrentChunkBufferOffset + ChannelItr * CompressedBlockSize) <= CompressedChunkSize)
-
 					ADPCM::DecodeBlock(
-						CompressedChunkPtr + CurrentChunkBufferOffset + ChannelItr * CompressedBlockSize,
+						CurCompressedChunkData + CurrentChunkBufferOffset + ChannelItr * CompressedBlockSize,
 						CompressedBlockSize,
 						(int16*)(UncompressedBlockData + ChannelItr * UncompressedBlockSize));
 				}
@@ -491,19 +510,6 @@ bool FADPCMAudioInfo::StreamCompressedData(uint8* Destination, bool bLooping, ui
 				// Update some bookkeeping
 				CurrentUncompressedBlockSampleIndex = 0;
 				CurrentChunkBufferOffset += NumChannels * CompressedBlockSize;
-
-				if (CurrentChunkBufferOffset >= CompressedChunkSize)
-				{
-					// Move on to the next streamed chunk.
-					++CurrentChunkIndex;
-
-					// Set the current buffer offset accounting for the header in the first chunk
-					if (!bSeekPending)
-					{
-						CurrentChunkBufferOffset = CurrentChunkIndex == 0 ? FirstChunkSampleDataOffset : 0;
-					}
-					bSeekPending = false;
-				}
 			}
 
 			// Only copy over the number of samples we currently have available, we will loop around if needed
@@ -539,9 +545,9 @@ bool FADPCMAudioInfo::StreamCompressedData(uint8* Destination, bool bLooping, ui
 				ReachedEndOfSamples = true;
 				CurrentUncompressedBlockSampleIndex = 0;
 				CurrentChunkIndex = 0;
-				CurrentChunkBufferOffset = FirstChunkSampleDataOffset;
+				CurrentChunkBufferOffset = 0;
 				TotalSamplesStreamed = 0;
-				CurrentCompressedChunkData = FAudioChunkHandle();
+				CurCompressedChunkData = nullptr;
 				if(!bLooping)
 				{
 					// Set the remaining buffer to 0
@@ -555,54 +561,36 @@ bool FADPCMAudioInfo::StreamCompressedData(uint8* Destination, bool bLooping, ui
 	{
 		while(BufferSize > 0)
 		{
-			const uint8* CompressedChunkPtr = nullptr;
-			uint32 CompressedChunkSize = 0;
-
-			GetChunkPtr(CompressedChunkPtr, CompressedChunkSize);
-
-
-			if (CompressedChunkPtr == nullptr)
+			if(CurCompressedChunkData == nullptr || CurrentChunkBufferOffset >= CurrentChunkDataSize)
 			{
-				// Only report missing the stream chunk if we were seeking. Seeking
-				// may cause a bit of latency with chunk loading, which is expected.
-				if (!bSeekPending)
-				{
-					// CurrentChunkIndex will not get incremented on the next callback, effectively causing another attempt to fetch the chunk.
-					// Since audio streaming depends on the general data streaming mechanism used by other parts of the engine and new data is
-					// prefetched on the game tick thread, this may be caused by a game hitch.
-					UE_LOG(LogAudio, Verbose, TEXT("Missed streaming ADPCM deadline chunk %d"), CurrentChunkIndex);
-				}
-
-				FMemory::Memset(OutData, 0, BufferSize);
-				return false;
-			}
+				// Get another chunk of compressed data from the streaming engine
 
-			uint32 DecompressedSamplesToCopy = FMath::Min<uint32>(
-				(CompressedChunkSize - CurrentChunkBufferOffset) / ChannelSampleSize,
-				BufferSize / ChannelSampleSize);
-
-			if (DecompressedSamplesToCopy > 0)
-			{
-				// Ensure we don't go over the number of samples left in the audio data
-				if (DecompressedSamplesToCopy > TotalSamplesPerChannel - TotalSamplesStreamed)
+				// CurrentChunkIndex is used to keep track of the current chunk for loading/unloading by the streaming engine but chunk 0
+				// is preloaded so we don't want to increment this when getting chunk 0, only later chunks
+				// Also, if we failed to get a previous chunk because it was not ready we don't want to re-increment the CurrentChunkIndex
+				if(CurCompressedChunkData != nullptr)
 				{
-					DecompressedSamplesToCopy = TotalSamplesPerChannel - TotalSamplesStreamed;
+					++CurrentChunkIndex;
 				}
 
-				const uint32 SampleSize = DecompressedSamplesToCopy * ChannelSampleSize;
+				// Request the next chunk of data from the streaming engine
+				CurCompressedChunkData = GetLoadedChunk(StreamingSoundWave, CurrentChunkIndex, CurrentChunkDataSize);
 
-				check(CompressedChunkPtr && CurrentChunkBufferOffset < CompressedChunkSize);
-				FMemory::Memcpy(OutData, CompressedChunkPtr + CurrentChunkBufferOffset, SampleSize);
-
-				OutData += DecompressedSamplesToCopy * NumChannels;
-				CurrentChunkBufferOffset += SampleSize;
-				BufferSize -= SampleSize;
-				TotalSamplesStreamed += DecompressedSamplesToCopy;
-			}
+				if(CurCompressedChunkData == nullptr)
+				{
+					// Only report missing the stream chunk if we were seeking. Seeking
+					// may cause a bit of latency with chunk loading, which is expected.
+					if (!bSeekPending)
+					{
+						// CurrentChunkIndex will not get incremented on the next callback, effectively causing another attempt to fetch the chunk.
+						// Since audio streaming depends on the general data streaming mechanism used by other parts of the engine and new data is
+						// prefetched on the game tick thread, this may be caused by a game hitch.
+						UE_LOG(LogAudio, Verbose, TEXT("Missed streaming ADPCM deadline chunk %d"), CurrentChunkIndex);
+					}
 
-			if (CurrentChunkBufferOffset >= CompressedChunkSize)
-			{
-				++CurrentChunkIndex;
+					FMemory::Memset(OutData, 0, BufferSize);
+					return false;
+				}
 
 				// Set the current buffer offset accounting for the header in the first chunk
 				if (!bSeekPending)
@@ -612,6 +600,25 @@ bool FADPCMAudioInfo::StreamCompressedData(uint8* Destination, bool bLooping, ui
 				bSeekPending = false;
 			}
 
+			uint32 DecompressedSamplesToCopy = FMath::Min<uint32>(
+				(CurrentChunkDataSize - CurrentChunkBufferOffset) / ChannelSampleSize,
+				BufferSize / ChannelSampleSize);
+
+			check(DecompressedSamplesToCopy > 0);
+
+			// Ensure we don't go over the number of samples left in the audio data
+			if(DecompressedSamplesToCopy > TotalSamplesPerChannel - TotalSamplesStreamed)
+			{
+				DecompressedSamplesToCopy = TotalSamplesPerChannel - TotalSamplesStreamed;
+			}
+
+			const uint32 SampleSize = DecompressedSamplesToCopy * ChannelSampleSize;
+			FMemory::Memcpy(OutData, CurCompressedChunkData + CurrentChunkBufferOffset, SampleSize);
+
+			OutData += DecompressedSamplesToCopy * NumChannels;
+			CurrentChunkBufferOffset += SampleSize;
+			BufferSize -= SampleSize;
+			TotalSamplesStreamed += DecompressedSamplesToCopy;
 
 			// Check for the end of the audio samples and loop if needed
 			if(TotalSamplesStreamed >= TotalSamplesPerChannel)
@@ -620,7 +627,7 @@ bool FADPCMAudioInfo::StreamCompressedData(uint8* Destination, bool bLooping, ui
 				CurrentChunkIndex = 0;
 				CurrentChunkBufferOffset = 0;
 				TotalSamplesStreamed = 0;
-				CurrentCompressedChunkData = FAudioChunkHandle();
+				CurCompressedChunkData = nullptr;
 				if(!bLooping)
 				{
 					// Set the remaining buffer to 0
@@ -634,20 +641,24 @@ bool FADPCMAudioInfo::StreamCompressedData(uint8* Destination, bool bLooping, ui
 	return ReachedEndOfSamples;
 }
 
-void FADPCMAudioInfo::GetChunkPtr(const uint8* &CompressedChunkPtr, uint32 &CompressedChunkSize)
+const uint8* FADPCMAudioInfo::GetLoadedChunk(USoundWave* InSoundWave, uint32 ChunkIndex, uint32& OutChunkSize)
 {
-	// Retrieve the current chunk:
-	if (CurrentChunkIndex == 0)
+	if (!InSoundWave || ChunkIndex >= InSoundWave->GetNumChunks())
+	{
+		UE_LOG(LogAudio, Error, TEXT("Error calling GetLoadedChunk for ChunkIndex %d!"), ChunkIndex);
+		OutChunkSize = 0;
+		return nullptr;
+	}
+	else if (ChunkIndex == 0)
 	{
-		TArrayView<const uint8> ZerothChunk = StreamingSoundWave->GetZerothChunk();
-		CompressedChunkPtr = ZerothChunk.GetData();
-		CompressedChunkSize = ZerothChunk.Num();
+		TArrayView<const uint8> ZerothChunk = InSoundWave->GetZerothChunk();
+		OutChunkSize = ZerothChunk.Num();
+		return ZerothChunk.GetData();
 	}
 	else
 	{
-		// Request the next chunk of data from the streaming engine
-		CurrentCompressedChunkData = IStreamingManager::Get().GetAudioStreamingManager().GetLoadedChunk(StreamingSoundWave, CurrentChunkIndex);
-		CompressedChunkPtr = CurrentCompressedChunkData.GetData();
-		CompressedChunkSize = CurrentCompressedChunkData.Num();
+		CurCompressedChunkHandle = IStreamingManager::Get().GetAudioStreamingManager().GetLoadedChunk(InSoundWave, ChunkIndex);
+		OutChunkSize = CurCompressedChunkHandle.Num();
+		return CurCompressedChunkHandle.GetData();
 	}
 }
diff --git a/Engine/Source/Runtime/Engine/Public/ADPCMAudioInfo.h b/Engine/Source/Runtime/Engine/Public/ADPCMAudioInfo.h
--- a/Engine/Source/Runtime/Engine/Public/ADPCMAudioInfo.h
+++ b/Engine/Source/Runtime/Engine/Public/ADPCMAudioInfo.h
@@ -108,7 +108,6 @@ public:
 	virtual bool SupportsStreaming() const override {return true;}
 	virtual bool StreamCompressedInfoInternal(USoundWave* Wave, struct FSoundQualityInfo* QualityInfo) override;
 	virtual bool StreamCompressedData(uint8* Destination, bool bLooping, uint32 BufferSize) override;
-
 	virtual int32 GetCurrentChunkIndex() const override
 	{
 		return CurrentChunkIndex;
@@ -134,7 +133,11 @@ public:
 	}
 
 private:
-	void GetChunkPtr(const uint8* &CompressedChunkPtr, uint32 &CompressedChunkSize);
+
+	// Wrapper function that returns a pointer to the currently used compressed data.
+	// If a non-zero chunk is requested, this function also aquires a reference to that chunk
+	// until we move on to a different chunk.
+	const uint8* GetLoadedChunk(USoundWave* InSoundWave, uint32 ChunkIndex, uint32& OutChunkSize);
 
 	FWaveModInfo WaveInfo;
 	const uint8*	SrcBufferData;
@@ -152,11 +155,13 @@ private:
 	uint32			CurrentUncompressedBlockSampleIndex;	// This is the sample index within the current uncompressed block data
 	uint32			CurrentChunkIndex;				// This is the index that is currently being used, needed by streaming engine to make sure it stays loaded and the next chunk gets preloaded
 	uint32			CurrentChunkBufferOffset;		// This is this byte offset within the current chunk, used by streaming engine to prioritize a load if more then half way through current chunk
+	uint32			CurrentChunkDataSize;			// The size of the current chunk, the first chunk is bigger to accommodate the header info
 	uint32			TotalSamplesStreamed;			// The number of samples streamed so far (per channel)
 	uint32			TotalSamplesPerChannel;			// Number of samples per channel, used to detect when an audio waveform has ended
 	uint32			SamplesPerBlock;				// The number of samples per block
 	uint32			FirstChunkSampleDataOffset;		// The size of the header in the first chunk, used to skip over it when looping or starting the sample over
-	FAudioChunkHandle	CurrentCompressedChunkData;			// A pointer to the current chunk of data
+	const uint8*	CurCompressedChunkData;			// A pointer to the current chunk of data
+	FAudioChunkHandle CurCompressedChunkHandle;     // Shared reference to the current chunk of data.
 
 	uint32			CurrentCompressedBlockIndex;		// For non disk streaming - the current compressed block in the compressed source data
 	uint32			TotalCompressedBlocksPerChannel;	// For non disk streaming - the total number of compressed blocks per channel
