Improved the assignation of UV slots for custom interpolators.
- Now reusing unused slots instead of simply adding them after the last used slot.
- Also added missing calls to GatherExpressionsForCustomInterpolators()

diff --git a/Engine/Shaders/Private/MaterialTemplate.ush b/Engine/Shaders/Private/MaterialTemplate.ush
--- a/Engine/Shaders/Private/MaterialTemplate.ush
+++ b/Engine/Shaders/Private/MaterialTemplate.ush
@@ -22,6 +22,7 @@
 #define NUM_MATERIAL_TEXCOORDS %s
 #define NUM_CUSTOM_VERTEX_INTERPOLATORS %s
 #define NUM_TEX_COORD_INTERPOLATORS %s
+%s
 
 #if NUM_VIRTUALTEXTURE_SAMPLES || LIGHTMAP_VT_ENABLED
 	#include "/Engine/Private/VirtualTextureCommon.ush"
diff --git a/Engine/Source/Developer/MaterialBaking/Private/ExportMaterialProxy.h b/Engine/Source/Developer/MaterialBaking/Private/ExportMaterialProxy.h
--- a/Engine/Source/Developer/MaterialBaking/Private/ExportMaterialProxy.h
+++ b/Engine/Source/Developer/MaterialBaking/Private/ExportMaterialProxy.h
@@ -576,6 +576,15 @@ public:
 	{
 		return true; 
 	}
+
+	virtual void GatherExpressionsForCustomInterpolators(TArray<UMaterialExpression*>& OutExpressions) const override
+	{
+		if(Material)
+		{
+			Material->GetAllExpressionsForCustomInterpolators(OutExpressions);
+		}
+	}
+
 private:
 	/** The material interface for this proxy */
 	UMaterialInterface* MaterialInterface;
diff --git a/Engine/Source/Developer/MaterialUtilities/Private/MaterialUtilities.cpp b/Engine/Source/Developer/MaterialUtilities/Private/MaterialUtilities.cpp
--- a/Engine/Source/Developer/MaterialUtilities/Private/MaterialUtilities.cpp
+++ b/Engine/Source/Developer/MaterialUtilities/Private/MaterialUtilities.cpp
@@ -754,6 +754,15 @@ public:
 	{
 		return Material && Material->MaterialDomain == MD_Volume;
 	}
+
+	virtual void GatherExpressionsForCustomInterpolators(TArray<UMaterialExpression*>& OutExpressions) const override
+	{
+		if(Material)
+		{
+			Material->GetAllExpressionsForCustomInterpolators(OutExpressions);
+		}
+	}
+
 private:
 	/** The material interface for this proxy */
 	UMaterialInterface* MaterialInterface;
diff --git a/Engine/Source/Editor/MaterialEditor/Private/MaterialEditor.h b/Engine/Source/Editor/MaterialEditor/Private/MaterialEditor.h
--- a/Engine/Source/Editor/MaterialEditor/Private/MaterialEditor.h
+++ b/Engine/Source/Editor/MaterialEditor/Private/MaterialEditor.h
@@ -177,6 +177,14 @@ public:
 		return Ar << V.Expression;
 	}
 
+	virtual void GatherExpressionsForCustomInterpolators(TArray<UMaterialExpression*>& OutExpressions) const override
+	{
+		if(Expression.IsValid() && Expression->Material)
+		{
+			Expression->Material->GetAllExpressionsForCustomInterpolators(OutExpressions);
+		}
+	}
+
 	float UnrelatedNodesOpacity;
 
 private:
diff --git a/Engine/Source/Runtime/Engine/Private/Materials/HLSLMaterialTranslator.h b/Engine/Source/Runtime/Engine/Private/Materials/HLSLMaterialTranslator.h
--- a/Engine/Source/Runtime/Engine/Private/Materials/HLSLMaterialTranslator.h
+++ b/Engine/Source/Runtime/Engine/Private/Materials/HLSLMaterialTranslator.h
@@ -7,6 +7,7 @@
 
 #include "CoreMinimal.h"
 #include "Stats/Stats.h"
+#include "Algo/Transform.h"
 #include "Misc/Guid.h"
 #include "HAL/IConsoleManager.h"
 #include "ShaderParameters.h"
@@ -230,6 +231,8 @@ protected:
 
 	/** Custom vertex interpolators */
 	TArray<UMaterialExpressionVertexInterpolator*> CustomVertexInterpolators;
+	/** Index to assign to next vertex interpolator. */
+	int32 NextVertexInterpolatorIndex;
 	/** Current float-width offset for custom vertex interpolators */
 	int32 CurrentCustomVertexInterpolatorOffset;
 
@@ -298,10 +301,10 @@ protected:
 	uint32 bIsFullyRough : 1;
 	/** true if allowed to generate code chunks. Translator operates in two phases; generate all code chunks & query meta data based on generated code chunks. */
 	uint32 bAllowCodeChunkGeneration : 1;
-	/** Tracks the number of texture coordinates used by this material. */
-	uint32 NumUserTexCoords;
-	/** Tracks the number of texture coordinates used by the vertex shader in this material. */
-	uint32 NumUserVertexTexCoords;
+	/** Tracks the texture coordinates used by this material. */
+	TBitArray<> AllocatedUserTexCoords;
+	/** Tracks the texture coordinates used by the vertex shader in this material. */
+	TBitArray<> AllocatedUserVertexTexCoords;
 
 	uint32 DynamicParticleParameterMask;
 
@@ -334,6 +337,7 @@ public:
 	,	FeatureLevel(InFeatureLevel)
 	,	MaterialTemplateLineNumber(INDEX_NONE)
 	,	NextSymbolIndex(INDEX_NONE)
+	,	NextVertexInterpolatorIndex(0)
 	,	CurrentCustomVertexInterpolatorOffset(0)
 	,	CompileErrorsSink(nullptr)
 	,	CompileErrorExpressionsSink(nullptr)
@@ -367,8 +371,8 @@ public:
 	,	bUsesDistanceCullFade(false)
 	,	bIsFullyRough(0)
 	,	bAllowCodeChunkGeneration(true)
-	,	NumUserTexCoords(0)
-	,	NumUserVertexTexCoords(0)
+	,	AllocatedUserTexCoords()
+	,	AllocatedUserVertexTexCoords()
 	,	DynamicParticleParameterMask(0)
 	,	NumVtSamples(0)
 	,	TargetPlatform(InTargetPlatform)
@@ -416,6 +420,26 @@ public:
 		ClearAllFunctionStacks();
 	}
 
+	int32 GetNumUserTexCoords() const
+	{
+		int32 NumUserTexCoords = 0;
+		for(TConstSetBitIterator<> It(AllocatedUserTexCoords); It; ++It)
+		{
+			NumUserTexCoords++;
+		}
+		return NumUserTexCoords;
+	}
+
+	int32 GetNumUserVertexTexCoords() const
+	{
+		int32 NumUserVertexTexCoords = 0;
+		for(TConstSetBitIterator<> It(AllocatedUserVertexTexCoords); It; ++It)
+		{
+			NumUserVertexTexCoords++;
+		}
+		return NumUserVertexTexCoords;
+	}
+
 	void ClearAllFunctionStacks()
 	{
 		for (uint32 Frequency = 0; Frequency < SF_NumFrequencies; ++Frequency)
@@ -466,10 +490,11 @@ public:
 				CompileErrorExpressionsSink = &Interpolator->CompileErrorExpressions;
 
 				// Compile node and store those successfully translated
-				int32 Ret = Interpolator->CompileInput(this, CustomVertexInterpolators.Num());
+				int32 Ret = Interpolator->CompileInput(this, NextVertexInterpolatorIndex);
 				if (Ret != INDEX_NONE)
 				{
-					CustomVertexInterpolators.Add(Interpolator);
+					CustomVertexInterpolators.AddUnique(Interpolator);
+					NextVertexInterpolatorIndex++;
 				}
 
 				// Restore error handling
@@ -822,6 +847,7 @@ public:
 			{
 				CustomVertexInterpolators.Empty();
 				CurrentCustomVertexInterpolatorOffset = 0;
+				NextVertexInterpolatorIndex = 0;
 				MaterialProperty = MP_MAX;
 				ShaderFrequency = SF_Vertex;
 
@@ -1030,16 +1056,15 @@ ResourcesString = TEXT("");
 			}
 
 			// No more calls to non-vertex shader CompilePropertyAndSetMaterialProperty beyond this point
-			const uint32 SavedNumUserTexCoords = NumUserTexCoords;
+			const uint32 SavedNumUserTexCoords = GetNumUserTexCoords();
 
-			for (uint32 CustomUVIndex = MP_CustomizedUVs0; CustomUVIndex <= MP_CustomizedUVs7; CustomUVIndex++)
+			// Only compile custom UV inputs for UV channels requested by the pixel shader inputs
+			for(TConstSetBitIterator<> It(AllocatedUserTexCoords); It; ++It)
 			{
-				// Only compile custom UV inputs for UV channels requested by the pixel shader inputs
-				// Any unconnected inputs will have a texcoord generated for them in Material->CompileProperty, which will pass through the vertex (uncustomized) texture coordinates
-				// Note: this is using NumUserTexCoords, which is set by translating all the pixel properties above
-				if (CustomUVIndex - MP_CustomizedUVs0 < SavedNumUserTexCoords)
+				EMaterialProperty CustomUVIndex = (EMaterialProperty)(MP_CustomizedUVs0 + It.GetIndex());
+				if(CustomUVIndex <= MP_CustomizedUVs7)
 				{
-					Chunk[CustomUVIndex] = Material->CompilePropertyAndSetMaterialProperty((EMaterialProperty)CustomUVIndex, this);
+					Chunk[CustomUVIndex] = Material->CompilePropertyAndSetMaterialProperty(CustomUVIndex, this);
 				}
 			}
 
@@ -1161,22 +1186,25 @@ ResourcesString = TEXT("");
 			}
 
 			// Catch any modifications to NumUserTexCoords that will not seen by customized UVs
-			check(SavedNumUserTexCoords == NumUserTexCoords);
+			check(SavedNumUserTexCoords == GetNumUserTexCoords());
+
+			FString InterpolatorsOffsetsDefinitionCode;
+			TBitArray<> FinalAllocatedCoords = GetVertexInterpolatorsOffsets(InterpolatorsOffsetsDefinitionCode);
 
 			// Finished compilation, verify final interpolator count restrictions
 			if (CurrentCustomVertexInterpolatorOffset > 0)
 			{
 				const int32 MaxNumScalars = (FeatureLevel == ERHIFeatureLevel::ES2) ? 3 * 2 : 8 * 2;
-				const int32 TotalUsedScalars = CurrentCustomVertexInterpolatorOffset + NumUserTexCoords * 2;
+				const int32 TotalUsedScalars = FinalAllocatedCoords.FindLast(true) + 1;
 
  				if (TotalUsedScalars > MaxNumScalars)
 				{
 					Errorf(TEXT("Maximum number of custom vertex interpolators exceeded. (%i / %i scalar values) (TexCoord: %i scalars, Custom: %i scalars)"),
-						TotalUsedScalars, MaxNumScalars, NumUserTexCoords * 2, CurrentCustomVertexInterpolatorOffset);
+						TotalUsedScalars, MaxNumScalars, GetNumUserTexCoords() * 2, CurrentCustomVertexInterpolatorOffset);
 				}
 			}
 
-			MaterialCompilationOutput.NumUsedUVScalars = NumUserTexCoords * 2;
+			MaterialCompilationOutput.NumUsedUVScalars = GetNumUserTexCoords();
 			MaterialCompilationOutput.NumUsedCustomInterpolatorScalars = CurrentCustomVertexInterpolatorOffset;
 
 			// Do Normal Chunk first
@@ -1494,6 +1522,74 @@ ResourcesString = TEXT("");
 		}
 	}
 
+	// Assign custom interpolators to slots, packing them as much as possible in unused slots.
+	TBitArray<> GetVertexInterpolatorsOffsets(FString& VertexInterpolatorsOffsetsDefinitionCode) const
+	{
+		TBitArray<> AllocatedCoords = AllocatedUserTexCoords; // Don't mess with the already assigned sets of UV coords
+
+		int32 CurrentSlot = INDEX_NONE;
+		int32 EndAllocatedSlot = INDEX_NONE;
+
+		auto GetNextUVSlot = [&CurrentSlot, &EndAllocatedSlot, &AllocatedCoords]() -> int32
+		{
+			if (CurrentSlot == EndAllocatedSlot)
+			{
+				CurrentSlot = AllocatedCoords.FindAndSetFirstZeroBit();
+				if (CurrentSlot == INDEX_NONE)
+				{
+					CurrentSlot = AllocatedCoords.Add(true);
+				}
+
+				// Track one slot per component (u,v)
+				const int32 NUM_COMPONENTS = 2;
+				CurrentSlot *= NUM_COMPONENTS;
+				EndAllocatedSlot = CurrentSlot + NUM_COMPONENTS;
+			}
+
+			int32 ResultUVSlot = CurrentSlot / 2;
+			CurrentSlot++;
+
+			return ResultUVSlot;
+		};
+
+		TArray<UMaterialExpressionVertexInterpolator*> SortedInterpolators;
+		Algo::TransformIf(CustomVertexInterpolators, 
+						  SortedInterpolators, 
+						  [](const UMaterialExpressionVertexInterpolator* Interpolator) { return Interpolator && Interpolator->InterpolatorIndex != INDEX_NONE && Interpolator->InterpolatorOffset != INDEX_NONE; },
+						  [](UMaterialExpressionVertexInterpolator* Interpolator) { return Interpolator; });
+						
+		SortedInterpolators.Sort([](const UMaterialExpressionVertexInterpolator& LHS, const UMaterialExpressionVertexInterpolator& RHS)  { return LHS.InterpolatorOffset < RHS.InterpolatorOffset; });
+		
+		for (UMaterialExpressionVertexInterpolator* Interpolator : SortedInterpolators)
+		{
+			int32 Index = Interpolator->InterpolatorIndex;
+
+			const EMaterialValueType Type = Interpolator->InterpolatedType == MCT_Float ? MCT_Float1 : Interpolator->InterpolatedType;
+
+			VertexInterpolatorsOffsetsDefinitionCode += LINE_TERMINATOR;
+			VertexInterpolatorsOffsetsDefinitionCode += FString::Printf(TEXT("#define VERTEX_INTERPOLATOR_%i_TEXCOORDS_X\t%i") LINE_TERMINATOR, Index, GetNextUVSlot());
+
+			if (Type >= MCT_Float2)
+			{
+				VertexInterpolatorsOffsetsDefinitionCode += FString::Printf(TEXT("#define VERTEX_INTERPOLATOR_%i_TEXCOORDS_Y\t%i") LINE_TERMINATOR, Index, GetNextUVSlot());
+
+				if (Type >= MCT_Float3)
+				{
+					VertexInterpolatorsOffsetsDefinitionCode += FString::Printf(TEXT("#define VERTEX_INTERPOLATOR_%i_TEXCOORDS_Z\t%i") LINE_TERMINATOR, Index, GetNextUVSlot());
+
+					if (Type == MCT_Float4)
+					{
+						VertexInterpolatorsOffsetsDefinitionCode += FString::Printf(TEXT("#define VERTEX_INTERPOLATOR_%i_TEXCOORDS_W\t%i") LINE_TERMINATOR, Index, GetNextUVSlot());
+					}
+				}
+			}
+			
+			VertexInterpolatorsOffsetsDefinitionCode += LINE_TERMINATOR;
+		}
+
+		return AllocatedCoords;
+	}
+
 	void GetSharedInputsMaterialCode(FString& PixelMembersDeclaration, FString& NormalAssignment, FString& PixelMembersInitializationEpilog)
 	{
 		{
@@ -1556,14 +1652,22 @@ ResourcesString = TEXT("");
 		// use "/Engine/Private/MaterialTemplate.ush" to create the functions to get data (e.g. material attributes) and code (e.g. material expressions to create specular color) from C++
 		FLazyPrintf LazyPrintf(*MaterialTemplate);
 
+		// Assign slots to vertex interpolators
+		FString VertexInterpolatorsOffsetsDefinition;
+		TBitArray<> FinalAllocatedCoords = GetVertexInterpolatorsOffsets(VertexInterpolatorsOffsetsDefinition);
+
+		const uint32 NumUserVertexTexCoords = AllocatedUserVertexTexCoords.FindLast(true) + 1;
+		const uint32 NumUserTexCoords = AllocatedUserTexCoords.FindLast(true) + 1;
 		const uint32 NumCustomVectors = FMath::DivideAndRoundUp((uint32)CurrentCustomVertexInterpolatorOffset, 2u);
-		const uint32 NumTexCoordVectors = NumUserTexCoords + NumCustomVectors;
+		const uint32 NumTexCoordVectors = FinalAllocatedCoords.FindLast(true) + 1;
 
 		LazyPrintf.PushParam(*FString::Printf(TEXT("%u"),NumUserVertexTexCoords));
 		LazyPrintf.PushParam(*FString::Printf(TEXT("%u"),NumUserTexCoords));
 		LazyPrintf.PushParam(*FString::Printf(TEXT("%u"),NumCustomVectors));
 		LazyPrintf.PushParam(*FString::Printf(TEXT("%u"),NumTexCoordVectors));
 
+		LazyPrintf.PushParam(*VertexInterpolatorsOffsetsDefinition);
+
 		// Stores the shared shader results member declarations
 		FString PixelMembersDeclaration;
 
@@ -1614,23 +1718,25 @@ ResourcesString = TEXT("");
 		// Print custom texture coordinate assignments
 		FString CustomUVAssignments;
 
-		int32 LastProperty = -1;
-		for (uint32 CustomUVIndex = 0; CustomUVIndex < NumUserTexCoords; CustomUVIndex++)
+		int32 LastProperty = INDEX_NONE;
+		for (TConstSetBitIterator<> It(AllocatedUserTexCoords); It; ++It)
 		{
-			if (CustomUVIndex == 0)
-			{
-				CustomUVAssignments += TranslatedCodeChunkDefinitions[MP_CustomizedUVs0 + CustomUVIndex];
-			}
+			EMaterialProperty CustomUVIndex = (EMaterialProperty)(MP_CustomizedUVs0 + It.GetIndex());
 
-			if (TranslatedCodeChunkDefinitions[MP_CustomizedUVs0 + CustomUVIndex].Len() > 0)
+			if (TranslatedCodeChunkDefinitions[CustomUVIndex].Len() > 0)
 			{
-				if (LastProperty >= 0)
+				if (LastProperty != INDEX_NONE)
+				{
+					check(TranslatedCodeChunkDefinitions[LastProperty].Len() == TranslatedCodeChunkDefinitions[CustomUVIndex].Len());
+				}
+				else
 				{
-					check(TranslatedCodeChunkDefinitions[LastProperty].Len() == TranslatedCodeChunkDefinitions[MP_CustomizedUVs0 + CustomUVIndex].Len());
+					CustomUVAssignments += TranslatedCodeChunkDefinitions[CustomUVIndex];
 				}
-				LastProperty = MP_CustomizedUVs0 + CustomUVIndex;
+
+				LastProperty = CustomUVIndex;
 			}
-			CustomUVAssignments += FString::Printf(TEXT("\tOutTexCoords[%u] = %s;") LINE_TERMINATOR, CustomUVIndex, *TranslatedCodeChunks[MP_CustomizedUVs0 + CustomUVIndex]);
+			CustomUVAssignments += FString::Printf(TEXT("\tOutTexCoords[%u] = %s;") LINE_TERMINATOR, It.GetIndex(), *TranslatedCodeChunks[CustomUVIndex]);
 		}
 
 		LazyPrintf.PushParam(*CustomUVAssignments);
@@ -1638,32 +1744,32 @@ ResourcesString = TEXT("");
 		// Print custom vertex shader interpolator assignments
 		FString CustomInterpolatorAssignments;
 
-		for (int32 Index = 0; Index < CustomVertexInterpolators.Num(); ++Index)
+		for (UMaterialExpressionVertexInterpolator* Interpolator : CustomVertexInterpolators)
 		{
-			UMaterialExpressionVertexInterpolator* Interpolator = CustomVertexInterpolators[Index];
-			check(Interpolator && Interpolator->InterpolatorIndex != INDEX_NONE);
-			check(Interpolator->InterpolatedType & MCT_Float);
-
 			if (Interpolator->InterpolatorOffset != INDEX_NONE)
 			{
+				check(Interpolator->InterpolatorIndex != INDEX_NONE);
+				check(Interpolator->InterpolatedType & MCT_Float);
+
 				const EMaterialValueType Type = Interpolator->InterpolatedType == MCT_Float ? MCT_Float1 : Interpolator->InterpolatedType;
 				const TCHAR* Swizzle[2] = { TEXT("x"), TEXT("y") };
 				const int32 Offset = Interpolator->InterpolatorOffset;
+				const int32 Index = Interpolator->InterpolatorIndex;
 
 				// Note: We reference the UV define directly to avoid having to pre-accumulate UV counts before property translation
-				CustomInterpolatorAssignments += FString::Printf(TEXT("\tOutTexCoords[%i + NUM_MATERIAL_TEXCOORDS].%s = VertexInterpolator%i(Parameters).x;") LINE_TERMINATOR, Offset/2, Swizzle[Offset%2], Index);
+				CustomInterpolatorAssignments += FString::Printf(TEXT("\tOutTexCoords[VERTEX_INTERPOLATOR_%i_TEXCOORDS_X].%s = VertexInterpolator%i(Parameters).x;") LINE_TERMINATOR, Index, Swizzle[Offset%2], Index);
 				
 				if (Type >= MCT_Float2)
 				{
-					CustomInterpolatorAssignments += FString::Printf(TEXT("\tOutTexCoords[%i + NUM_MATERIAL_TEXCOORDS].%s = VertexInterpolator%i(Parameters).y;") LINE_TERMINATOR, (Offset+1)/2, Swizzle[(Offset+1)%2], Index);
+					CustomInterpolatorAssignments += FString::Printf(TEXT("\tOutTexCoords[VERTEX_INTERPOLATOR_%i_TEXCOORDS_Y].%s = VertexInterpolator%i(Parameters).y;") LINE_TERMINATOR, Index, Swizzle[(Offset+1)%2], Index);
 
 					if (Type >= MCT_Float3)
 					{
-						CustomInterpolatorAssignments += FString::Printf(TEXT("\tOutTexCoords[%i + NUM_MATERIAL_TEXCOORDS].%s = VertexInterpolator%i(Parameters).z;") LINE_TERMINATOR, (Offset+2)/2, Swizzle[(Offset+2)%2], Index);
+						CustomInterpolatorAssignments += FString::Printf(TEXT("\tOutTexCoords[VERTEX_INTERPOLATOR_%i_TEXCOORDS_Z].%s = VertexInterpolator%i(Parameters).z;") LINE_TERMINATOR, Index, Swizzle[(Offset+2)%2], Index);
 
 						if (Type == MCT_Float4)
 						{
-							CustomInterpolatorAssignments += FString::Printf(TEXT("\tOutTexCoords[%i + NUM_MATERIAL_TEXCOORDS].%s = VertexInterpolator%i(Parameters).w;") LINE_TERMINATOR, (Offset+3)/2, Swizzle[(Offset+3)%2], Index);
+							CustomInterpolatorAssignments += FString::Printf(TEXT("\tOutTexCoords[VERTEX_INTERPOLATOR_%i_TEXCOORDS_W].%s = VertexInterpolator%i(Parameters).w;") LINE_TERMINATOR, Index, Swizzle[(Offset+3)%2], Index);
 						}
 					}
 				}
@@ -3867,16 +3973,29 @@ protected:
 		}
 	}
 
+	void AllocateSlot(TBitArray<>& InBitArray, int32 InSlotIndex) const
+	{
+		// Grow as needed
+		while (InBitArray.Num() <= InSlotIndex)
+		{
+			InBitArray.Add(false);
+		}
+
+		InBitArray[InSlotIndex] = true;
+	}
+
 #if WITH_EDITOR
 	virtual int32 MaterialBakingWorldPosition() override
 	{
 		if (ShaderFrequency == SF_Vertex)
 		{
-			NumUserVertexTexCoords = FMath::Max((uint32)8, NumUserVertexTexCoords);
+			AllocateSlot(AllocatedUserVertexTexCoords, 6);
+			AllocateSlot(AllocatedUserVertexTexCoords, 7);
 		}
 		else
 		{
-			NumUserTexCoords = FMath::Max((uint32)8, NumUserTexCoords);
+			AllocateSlot(AllocatedUserTexCoords, 6);
+			AllocateSlot(AllocatedUserTexCoords, 7);
 		}
 
 		// Note: inlining is important so that on ES2 devices, where half precision is used in the pixel shader, 
@@ -3901,11 +4020,11 @@ protected:
 
 		if (ShaderFrequency == SF_Vertex)
 		{
-			NumUserVertexTexCoords = FMath::Max(CoordinateIndex + 1, NumUserVertexTexCoords);
+			AllocateSlot(AllocatedUserVertexTexCoords, CoordinateIndex);
 		}
 		else
 		{
-			NumUserTexCoords = FMath::Max(CoordinateIndex + 1, NumUserTexCoords);
+			AllocateSlot(AllocatedUserTexCoords, CoordinateIndex);
 		}
 
 		FString	SampleCode;
@@ -5212,13 +5331,15 @@ protected:
 		{
 			return Errorf(TEXT("Custom interpolator outputs only available in pixel shaders."));
 		}
-		else if (InterpolatorIndex >= (uint32)CustomVertexInterpolators.Num())
+
+		UMaterialExpressionVertexInterpolator** InterpolatorPtr = CustomVertexInterpolators.FindByPredicate([InterpolatorIndex](const UMaterialExpressionVertexInterpolator* Item) { return Item && Item->InterpolatorIndex == InterpolatorIndex; });
+		if (InterpolatorPtr == nullptr)
 		{
 			return Errorf(TEXT("Invalid custom interpolator index."));
 		}
 
-		UMaterialExpressionVertexInterpolator* Interpolator = CustomVertexInterpolators[InterpolatorIndex];
-		check(Interpolator && Interpolator->InterpolatorIndex == InterpolatorIndex);
+		UMaterialExpressionVertexInterpolator* Interpolator = *InterpolatorPtr;
+		check(Interpolator->InterpolatorIndex == InterpolatorIndex);
 		check(Interpolator->InterpolatedType & MCT_Float);
 
 		// Assign interpolator offset and accumulate size
@@ -5245,19 +5366,19 @@ protected:
 		const int32 Offset = Interpolator->InterpolatorOffset;
 	
 		// Note: We reference the UV define directly to avoid having to pre-accumulate UV counts before property translation
-		FString GetValueCode = FString::Printf(TEXT("%s(Parameters.TexCoords[%i + NUM_MATERIAL_TEXCOORDS].%s"), TypeName, (Offset/2), Swizzle[Offset%2]);
+		FString GetValueCode = FString::Printf(TEXT("%s(Parameters.TexCoords[VERTEX_INTERPOLATOR_%i_TEXCOORDS_X].%s"), TypeName, InterpolatorIndex, Swizzle[Offset%2]);
 
 		if (Type >= MCT_Float2)
 		{
-			GetValueCode += FString::Printf(TEXT(", Parameters.TexCoords[%i + NUM_MATERIAL_TEXCOORDS].%s"), (Offset+1)/2, Swizzle[(Offset+1)%2]);
+			GetValueCode += FString::Printf(TEXT(", Parameters.TexCoords[VERTEX_INTERPOLATOR_%i_TEXCOORDS_Y].%s"), InterpolatorIndex, Swizzle[(Offset+1)%2]);
 
 			if (Type >= MCT_Float3)
 			{
-				GetValueCode += FString::Printf(TEXT(", Parameters.TexCoords[%i + NUM_MATERIAL_TEXCOORDS].%s"), (Offset+2)/2, Swizzle[(Offset+2)%2]);
+				GetValueCode += FString::Printf(TEXT(", Parameters.TexCoords[VERTEX_INTERPOLATOR_%i_TEXCOORDS_Z].%s"), InterpolatorIndex, Swizzle[(Offset+2)%2]);
 
 				if (Type == MCT_Float4)
 				{
-					GetValueCode += FString::Printf(TEXT(", Parameters.TexCoords[%i + NUM_MATERIAL_TEXCOORDS].%s"), (Offset+3)/2, Swizzle[(Offset+3)%2]);
+					GetValueCode += FString::Printf(TEXT(", Parameters.TexCoords[VERTEX_INTERPOLATOR_%i_TEXCOORDS_W].%s"), InterpolatorIndex, Swizzle[(Offset+3)%2]);
 				}
 			}
 		}
@@ -6801,7 +6922,8 @@ protected:
 		{
 			bUsesSpeedTree = true;
 
-			NumUserVertexTexCoords = FMath::Max<uint32>(NumUserVertexTexCoords, 8);
+			AllocateSlot(AllocatedUserVertexTexCoords, 7);
+
 			// Only generate previous frame's computations if required and opted-in
 			const bool bEnablePreviousFrameInformation = bCompilingPreviousFrame && bAccurateWindVelocities;
 			return AddCodeChunk(MCT_Float3, TEXT("GetSpeedTreeVertexOffset(Parameters, %s, %s, %s, %g, %s, %s, %s)"), *GetParameterCode(GeometryArg), *GetParameterCode(WindArg), *GetParameterCode(LODArg), BillboardThreshold, bEnablePreviousFrameInformation ? TEXT("true") : TEXT("false"), bExtraBend ? TEXT("true") : TEXT("false"), *GetParameterCode(ExtraBendArg, TEXT("float3(0,0,0)")));
diff --git a/Engine/Source/Runtime/Engine/Private/Materials/MaterialShared.cpp b/Engine/Source/Runtime/Engine/Private/Materials/MaterialShared.cpp
--- a/Engine/Source/Runtime/Engine/Private/Materials/MaterialShared.cpp
+++ b/Engine/Source/Runtime/Engine/Private/Materials/MaterialShared.cpp
@@ -3264,7 +3264,7 @@ void UMaterialInterface::AnalyzeMaterialProperty(EMaterialProperty InProperty, i
 		{}
 		int32 GetTextureCoordsCount() const
 		{
-			return NumUserTexCoords;
+			return GetNumUserTexCoords();
 		}
 		bool UsesVertexColor() const
 		{
