Guard against crash when cooking UModelComponent physics data. #jira UE-69488 #rb none

diff --git a/Engine/Source/Runtime/Engine/Private/Components/ModelComponent.cpp b/Engine/Source/Runtime/Engine/Private/Components/ModelComponent.cpp
--- a/Engine/Source/Runtime/Engine/Private/Components/ModelComponent.cpp
+++ b/Engine/Source/Runtime/Engine/Private/Components/ModelComponent.cpp
@@ -683,6 +683,13 @@ bool UModelComponent::GetPhysicsTriMeshData(struct FTriMeshCollisionData* Collis
 
 		IndexBufferSize = IndexBuffer->Indices.Num();
 
+		if (IndexBufferSize < (int32)(Element.FirstIndex + Element.NumTriangles * 3 - 1))
+		{
+			UE_LOG(LogPhysics, Warning, TEXT("Found bad index buffer when cooking UModelComponent physics data! Not enough indices in buffer for Element triangles. Component: %s, Element: %d, Index Start: %d, Tri Count: %d, Index Buffer Size: %d"),
+				*GetPathName(), ElementIndex, Element.FirstIndex, Element.NumTriangles, IndexBufferSize);
+			continue;
+		}
+
 		for (uint32 TriIdx = 0; TriIdx < Element.NumTriangles; TriIdx++)
 		{
 			FTriIndices Triangle;
