Prevent Crash in USkinnedMeshComponent::RebuildVisibilityArray.

diff --git a/Engine/Source/Runtime/Engine/Private/Components/SkinnedMeshComponent.cpp b/Engine/Source/Runtime/Engine/Private/Components/SkinnedMeshComponent.cpp
--- a/Engine/Source/Runtime/Engine/Private/Components/SkinnedMeshComponent.cpp
+++ b/Engine/Source/Runtime/Engine/Private/Components/SkinnedMeshComponent.cpp
@@ -943,7 +943,12 @@ void USkinnedMeshComponent::RebuildVisibilityArray()
 
 		// The following code relies on a complete hierarchy sorted from parent to children
 		TArray<uint8>& EditableBoneVisibilityStates = GetEditableBoneVisibilityStates();
-		check(EditableBoneVisibilityStates.Num() == SkeletalMesh->RefSkeleton.GetNum());
+		if (EditableBoneVisibilityStates.Num() != SkeletalMesh->RefSkeleton.GetNum())
+		{
+			UE_LOG(LogSkinnedMeshComp, Warning, TEXT("RebuildVisibilityArray() failed because EditableBoneVisibilityStates size: %d not equal to RefSkeleton bone count: %d."), EditableBoneVisibilityStates.Num(), SkeletalMesh->RefSkeleton.GetNum());
+			return;
+		}
+
 		for (int32 BoneId=0; BoneId < EditableBoneVisibilityStates.Num(); ++BoneId)
 		{
 			uint8 VisState = EditableBoneVisibilityStates[BoneId];
