Workaround for UE-76848: Editor Crashes when Isolating a Material Element on Skeletal Meshes with Ray Tracing Enabled

diff --git a/Engine/Source/Runtime/Engine/Private/SkeletalMesh.cpp b/Engine/Source/Runtime/Engine/Private/SkeletalMesh.cpp
--- a/Engine/Source/Runtime/Engine/Private/SkeletalMesh.cpp
+++ b/Engine/Source/Runtime/Engine/Private/SkeletalMesh.cpp
@@ -4170,6 +4170,10 @@ void FSkeletalMeshSceneProxy::GetDynamicRayTracingInstances(FRayTracingMaterialG
 				const FLODSectionElements& LODSection = LODSections[LODIndex];
 				check(LODSection.SectionElements.Num() == LODData.RenderSections.Num());
 
+				int32 SectionIndexPreview = MeshObject->SectionIndexPreview;
+				int32 MaterialIndexPreview = MeshObject->MaterialIndexPreview;
+				MeshObject->SectionIndexPreview = INDEX_NONE;
+				MeshObject->MaterialIndexPreview = INDEX_NONE;
 				for (FSkeletalMeshSectionIter Iter(LODIndex, *MeshObject, LODData, LODSection); Iter; ++Iter)
 				{
 					const FSkelMeshRenderSection& Section = Iter.GetSection();
@@ -4181,6 +4185,8 @@ void FSkeletalMeshSceneProxy::GetDynamicRayTracingInstances(FRayTracingMaterialG
 
 					RayTracingInstance.Materials.Add(MeshBatch);
 				}
+				MeshObject->SectionIndexPreview = SectionIndexPreview;
+				MeshObject->MaterialIndexPreview = MaterialIndexPreview;
 
 				if (bAnySegmentUsesWorldPositionOffset)
 				{
