Fix UE-76933

diff --git a/Engine/Source/Runtime/Engine/Private/SkeletalMesh.cpp b/Engine/Source/Runtime/Engine/Private/SkeletalMesh.cpp
--- a/Engine/Source/Runtime/Engine/Private/SkeletalMesh.cpp
+++ b/Engine/Source/Runtime/Engine/Private/SkeletalMesh.cpp
@@ -4170,10 +4170,12 @@ void FSkeletalMeshSceneProxy::GetDynamicRayTracingInstances(FRayTracingMaterialG
 				const FLODSectionElements& LODSection = LODSections[LODIndex];
 				check(LODSection.SectionElements.Num() == LODData.RenderSections.Num());
 
+			#if WITH_EDITORONLY_DATA
 				int32 SectionIndexPreview = MeshObject->SectionIndexPreview;
 				int32 MaterialIndexPreview = MeshObject->MaterialIndexPreview;
 				MeshObject->SectionIndexPreview = INDEX_NONE;
 				MeshObject->MaterialIndexPreview = INDEX_NONE;
+			#endif
 				for (FSkeletalMeshSectionIter Iter(LODIndex, *MeshObject, LODData, LODSection); Iter; ++Iter)
 				{
 					const FSkelMeshRenderSection& Section = Iter.GetSection();
@@ -4185,9 +4187,10 @@ void FSkeletalMeshSceneProxy::GetDynamicRayTracingInstances(FRayTracingMaterialG
 
 					RayTracingInstance.Materials.Add(MeshBatch);
 				}
+			#if WITH_EDITORONLY_DATA
 				MeshObject->SectionIndexPreview = SectionIndexPreview;
 				MeshObject->MaterialIndexPreview = MaterialIndexPreview;
-
+			#endif
 				if (bAnySegmentUsesWorldPositionOffset)
 				{
 					RayTracingInstance.InstanceTransforms.Add(FMatrix::Identity);
