Fix duplicate assignment in if statement
#rnx

diff --git a/Engine/Source/Runtime/Engine/Private/Subsystems/SubsystemBlueprintLibrary.cpp b/Engine/Source/Runtime/Engine/Private/Subsystems/SubsystemBlueprintLibrary.cpp
--- a/Engine/Source/Runtime/Engine/Private/Subsystems/SubsystemBlueprintLibrary.cpp
+++ b/Engine/Source/Runtime/Engine/Private/Subsystems/SubsystemBlueprintLibrary.cpp
@@ -16,7 +16,7 @@
 {
 	if (const UWorld* World = ThisClass::GetWorldFrom(ContextObject))
 	{
-		if (const UGameInstance* GameInstance = GameInstance = World->GetGameInstance())
+		if (const UGameInstance* GameInstance = World->GetGameInstance())
 		{
 			return GameInstance->GetSubsystemBase(Class);
 		}
