Fix bug with focus not interpolating correctly when enabled

diff --git a/Engine/Source/Runtime/Engine/Private/ActiveSound.cpp b/Engine/Source/Runtime/Engine/Private/ActiveSound.cpp
--- a/Engine/Source/Runtime/Engine/Private/ActiveSound.cpp
+++ b/Engine/Source/Runtime/Engine/Private/ActiveSound.cpp
@@ -1515,6 +1515,9 @@ void FActiveSound::UpdateAttenuation(float DeltaTime, FSoundParseParameters& Par
 
 	if (Settings->bSpatialize || Settings->bEnableListenerFocus)
 	{
+		// Feed prior focus factor on update to allow for proper interpolation.
+		FocusDataToApply.FocusFactor = FocusData.FocusFactor;
+
 		// Update azimuth angles prior to updating focus as it uses this in calculating
 		// in and out of focus values.
 		UpdateFocusData(DeltaTime, ListenerData, &FocusDataToApply);
