Reduce memory usage for importing lots of meshes (56.17GB to 45.75GB -> ~18% for DATASET-0010a)
  - Clear MeshDescription once serialized in BulkData after UStaticMesh Build (56.17GB peak to 51.28GB)
  - Clear SourceModels from DatasmithMesh as soon as possible (51.28GB peak to 48.8GB)
  - Remove CommitMeshDescription from FillUStaticMesh since already in PreBuildStaticMesh phase (48.8GB peak to 45.75GB)

From dev-enterprise[REDACTED_HANDLE]

diff --git a/Engine/Source/Runtime/Engine/Classes/Engine/StaticMesh.h b/Engine/Source/Runtime/Engine/Classes/Engine/StaticMesh.h
--- a/Engine/Source/Runtime/Engine/Classes/Engine/StaticMesh.h
+++ b/Engine/Source/Runtime/Engine/Classes/Engine/StaticMesh.h
@@ -824,6 +824,7 @@ public:
 	ENGINE_API FMeshDescription* CreateMeshDescription(int32 LodIndex, FMeshDescription MeshDescription);
 	ENGINE_API void CommitMeshDescription(int32 LodIndex);
 	ENGINE_API void ClearMeshDescription(int32 LodIndex);
+	ENGINE_API void ClearMeshDescriptions();
 
 	UE_DEPRECATED(4.22, "Please use GetMeshDescription().")
 	FMeshDescription* GetOriginalMeshDescription(int32 LodIndex) const { return GetMeshDescription(LodIndex); }
diff --git a/Engine/Source/Runtime/Engine/Private/StaticMesh.cpp b/Engine/Source/Runtime/Engine/Private/StaticMesh.cpp
--- a/Engine/Source/Runtime/Engine/Private/StaticMesh.cpp
+++ b/Engine/Source/Runtime/Engine/Private/StaticMesh.cpp
@@ -3858,6 +3858,14 @@ void UStaticMesh::ClearMeshDescription(int32 LodIndex)
 	}
 }
 
+void UStaticMesh::ClearMeshDescriptions()
+{
+	for (int LODIndex = 0; LODIndex < SourceModels.Num(); LODIndex++)
+	{
+		ClearMeshDescription(LODIndex);
+	}
+}
+
 void UStaticMesh::FixupMaterialSlotName()
 {
 	TArray<FName> UniqueMaterialSlotName;
@@ -4834,10 +4842,7 @@ void UStaticMesh::PostLoad()
 
 #if WITH_EDITOR
 	// Release cached mesh descriptions until they are loaded on demand
-	for (int LODIndex = 0; LODIndex < SourceModels.Num(); LODIndex++)
-	{
-		ClearMeshDescription(LODIndex);
-	}
+	ClearMeshDescriptions();
 #endif
 
 	CreateNavCollision();
diff --git a/Engine/Source/Runtime/MeshDescription/Private/MeshDescription.cpp b/Engine/Source/Runtime/MeshDescription/Private/MeshDescription.cpp
--- a/Engine/Source/Runtime/MeshDescription/Private/MeshDescription.cpp
+++ b/Engine/Source/Runtime/MeshDescription/Private/MeshDescription.cpp
@@ -1276,6 +1276,9 @@ void FMeshDescriptionBulkData::SaveMeshDescription( FMeshDescription& MeshDescri
 		const bool bIsPersistent = true;
 		FBulkDataWriter Ar( BulkData, bIsPersistent );
 		Ar << MeshDescription;
+
+		// Preserve CustomVersions at save time so we can reuse the same ones when reloading direct from memory
+		CustomVersions = Ar.GetCustomVersions();
 	}
 
 	FPlatformMisc::CreateGuid( Guid );
