Improve deselect performance
  - 25 seconds to deselect 60 000 actors to 400 ms

From dev-enterprise[REDACTED_HANDLE]

diff --git a/Engine/Source/Runtime/Engine/Private/Selection.cpp b/Engine/Source/Runtime/Engine/Private/Selection.cpp
--- a/Engine/Source/Runtime/Engine/Private/Selection.cpp
+++ b/Engine/Source/Runtime/Engine/Private/Selection.cpp
@@ -79,7 +79,7 @@ void USelection::Deselect(UObject* InObject)
 	SelectionAnnotation->Clear(InObject);
 
 	// Remove from selected list.
-	SelectedObjects.Remove( InObject );
+	SelectedObjects.RemoveSingle( InObject );
 
 	if( !IsBatchSelecting() )
 	{
