Revised implementation of CL 7224811, which had a possible issue when FFXSystemSet subsystems would be null while pending kill.

diff --git a/Engine/Plugins/FX/Niagara/Source/Niagara/Classes/NiagaraEmitterInstanceBatcher.h b/Engine/Plugins/FX/Niagara/Source/Niagara/Classes/NiagaraEmitterInstanceBatcher.h
--- a/Engine/Plugins/FX/Niagara/Source/Niagara/Classes/NiagaraEmitterInstanceBatcher.h
+++ b/Engine/Plugins/FX/Niagara/Source/Niagara/Classes/NiagaraEmitterInstanceBatcher.h
@@ -81,7 +81,7 @@ public:
 	virtual void PreInitViews(FRHICommandListImmediate& RHICmdList) override;
 	virtual bool UsesGlobalDistanceField() const override;
 	virtual void PreRender(FRHICommandListImmediate& RHICmdList, const class FGlobalDistanceFieldParameterData* GlobalDistanceFieldParameterData, bool bAllowGPUParticleSceneUpdate) override;
-	virtual void Destroy() override; // Called on the gamethread to delete the batcher on the renderthread.
+	virtual void OnDestroy() override; // Called on the gamethread to delete the batcher on the renderthread.
 
 	virtual void Tick(float DeltaTime) override
 	{
diff --git a/Engine/Plugins/FX/Niagara/Source/Niagara/Private/NiagaraEmitterInstanceBatcher.cpp b/Engine/Plugins/FX/Niagara/Source/Niagara/Private/NiagaraEmitterInstanceBatcher.cpp
--- a/Engine/Plugins/FX/Niagara/Source/Niagara/Private/NiagaraEmitterInstanceBatcher.cpp
+++ b/Engine/Plugins/FX/Niagara/Source/Niagara/Private/NiagaraEmitterInstanceBatcher.cpp
@@ -588,10 +588,10 @@ void NiagaraEmitterInstanceBatcher::PreRender(FRHICommandListImmediate& RHICmdLi
 
 }
 
-void NiagaraEmitterInstanceBatcher::Destroy()
+void NiagaraEmitterInstanceBatcher::OnDestroy()
 {
 	INiagaraModule::OnBatcherDestroyed(this);
-	FFXSystemInterface::Destroy();
+	FFXSystemInterface::OnDestroy();
 }
 
 int32 NiagaraEmitterInstanceBatcher::AddSortedGPUSimulation(const FNiagaraGPUSortInfo& SortInfo)
diff --git a/Engine/Source/Editor/UnrealEd/Private/ThumbnailHelpers.cpp b/Engine/Source/Editor/UnrealEd/Private/ThumbnailHelpers.cpp
--- a/Engine/Source/Editor/UnrealEd/Private/ThumbnailHelpers.cpp
+++ b/Engine/Source/Editor/UnrealEd/Private/ThumbnailHelpers.cpp
@@ -172,11 +172,7 @@ FParticleSystemThumbnailScene::FParticleSystemThumbnailScene()
 
 FParticleSystemThumbnailScene::~FParticleSystemThumbnailScene()
 {
-	if (ThumbnailFXSystem)
-	{
-		ThumbnailFXSystem->Destroy();
-		ThumbnailFXSystem = nullptr;
-	}
+	FFXSystemInterface::Destroy( ThumbnailFXSystem );
 }
 
 void FParticleSystemThumbnailScene::SetParticleSystem(UParticleSystem* ParticleSystem)
diff --git a/Engine/Source/Runtime/Engine/Private/Particles/FXSystem.cpp b/Engine/Source/Runtime/Engine/Private/Particles/FXSystem.cpp
--- a/Engine/Source/Runtime/Engine/Private/Particles/FXSystem.cpp
+++ b/Engine/Source/Runtime/Engine/Private/Particles/FXSystem.cpp
@@ -46,13 +46,11 @@ FFXSystemInterface* FFXSystemInterface::Create(ERHIFeatureLevel::Type InFeatureL
 	}
 }
 
-void FFXSystemInterface::Destroy()
+void FFXSystemInterface::Destroy( FFXSystemInterface* FXSystem )
 {
-	check(!bIsPendingKill);
+	check(FXSystem);
+	FXSystem->OnDestroy();
 
-	bIsPendingKill = true;
-
-	FFXSystemInterface* FXSystem = this;
 	ENQUEUE_RENDER_COMMAND(FDestroyFXSystemCommand)(
 		[FXSystem](FRHICommandList& RHICmdList)
 		{
diff --git a/Engine/Source/Runtime/Engine/Private/Particles/FXSystemSet.cpp b/Engine/Source/Runtime/Engine/Private/Particles/FXSystemSet.cpp
--- a/Engine/Source/Runtime/Engine/Private/Particles/FXSystemSet.cpp
+++ b/Engine/Source/Runtime/Engine/Private/Particles/FXSystemSet.cpp
@@ -10,13 +10,11 @@ FFXSystemInterface* FFXSystemSet::GetInterface(const FName& InName)
 {
 	for (FFXSystemInterface* FXSystem : FXSystems)
 	{
-		if (FXSystem) // Can be null when called from ::Destroy()
+		check(FXSystem);
+		FXSystem = FXSystem->GetInterface(InName);
+		if (FXSystem)
 		{
-			FXSystem = FXSystem->GetInterface(InName);
-			if (FXSystem)
-			{
-				return FXSystem;
-			}
+			return FXSystem;
 		}
 	}
 	return nullptr;
@@ -138,16 +136,22 @@ void FFXSystemSet::PostRenderOpaque(
 	}
 }
 
-void FFXSystemSet::Destroy()
+void FFXSystemSet::OnDestroy()
 {
 	for (FFXSystemInterface*& FXSystem : FXSystems)
 	{
-		// Here it is important that the FXSystem stays in the FXSystems array while destroy is being called.
-		// This is used to figure out from which world the destroyed FXSystem is owned (from UWorld::FXSystem->GetInterface())
 		check(FXSystem);
-		FXSystem->Destroy();
-		FXSystem = nullptr;
+		FXSystem->OnDestroy();
 	}
 
-	FFXSystemInterface::Destroy();
+	FFXSystemInterface::OnDestroy();
+}
+
+FFXSystemSet::~FFXSystemSet()
+{
+	for (FFXSystemInterface* FXSystem : FXSystems)
+	{
+		check(FXSystem);
+		delete FXSystem;
+	}
 }
diff --git a/Engine/Source/Runtime/Engine/Private/Particles/FXSystemSet.h b/Engine/Source/Runtime/Engine/Private/Particles/FXSystemSet.h
--- a/Engine/Source/Runtime/Engine/Private/Particles/FXSystemSet.h
+++ b/Engine/Source/Runtime/Engine/Private/Particles/FXSystemSet.h
@@ -39,10 +39,10 @@ public:
 		const class FShaderParametersMetadata* SceneTexturesUniformBufferStruct,
 		FRHIUniformBuffer* SceneTexturesUniformBuffer) override;
 
-	virtual void Destroy() override;
+	virtual void OnDestroy() override;
 
 protected:
 
 	/** By making the destructor protected, an instance must be destroyed via FFXSystemInterface::Destroy. */
-	virtual ~FFXSystemSet(){}
+	virtual ~FFXSystemSet();
 };
diff --git a/Engine/Source/Runtime/Engine/Private/RendererSupport.cpp b/Engine/Source/Runtime/Engine/Private/RendererSupport.cpp
--- a/Engine/Source/Runtime/Engine/Private/RendererSupport.cpp
+++ b/Engine/Source/Runtime/Engine/Private/RendererSupport.cpp
@@ -60,7 +60,7 @@ static void ClearReferencesToRendererModuleClasses(
 
 			if (World->FXSystem)
 			{
-				World->FXSystem->Destroy();
+				FFXSystemInterface::Destroy(World->FXSystem);
 				World->FXSystem = NULL;
 			}
 
diff --git a/Engine/Source/Runtime/Engine/Private/World.cpp b/Engine/Source/Runtime/Engine/Private/World.cpp
--- a/Engine/Source/Runtime/Engine/Private/World.cpp
+++ b/Engine/Source/Runtime/Engine/Private/World.cpp
@@ -854,7 +854,7 @@ void UWorld::FinishDestroy()
 
 		if (FXSystem)
 		{
-			FXSystem->Destroy();
+			FFXSystemInterface::Destroy( FXSystem );
 			FXSystem = NULL;
 		}
 
diff --git a/Engine/Source/Runtime/Engine/Public/FXSystem.h b/Engine/Source/Runtime/Engine/Public/FXSystem.h
--- a/Engine/Source/Runtime/Engine/Public/FXSystem.h
+++ b/Engine/Source/Runtime/Engine/Public/FXSystem.h
@@ -129,6 +129,11 @@ public:
 	 */
 	ENGINE_API static FFXSystemInterface* Create(ERHIFeatureLevel::Type InFeatureLevel, EShaderPlatform InShaderPlatform);
 
+	/**
+	 * Destroy an effects system instance.
+	 */
+	ENGINE_API static void Destroy(FFXSystemInterface* FXSystem);
+
 	/**
 	 * Register a custom FX system implementation.
 	 */
@@ -145,9 +150,9 @@ public:
 	virtual FFXSystemInterface* GetInterface(const FName& InName) { return nullptr; };
 
 	/**
-	 * Destroy an effects system instance.
+	 * Gamethread callback when destroy gets called, allows to clean up references.
 	 */
-	ENGINE_API virtual void Destroy();
+	ENGINE_API virtual void OnDestroy() { bIsPendingKill = true; }
 
 
 	/**
@@ -225,6 +230,8 @@ public:
 	bool IsPendingKill() const { return bIsPendingKill; }
 
 protected:
+	
+	friend class FFXSystemSet;
 
 	/** By making the destructor protected, an instance must be destroyed via FFXSystemInterface::Destroy. */
 	ENGINE_API virtual ~FFXSystemInterface() {}
diff --git a/Engine/Source/Runtime/Renderer/Private/RendererScene.cpp b/Engine/Source/Runtime/Renderer/Private/RendererScene.cpp
--- a/Engine/Source/Runtime/Renderer/Private/RendererScene.cpp
+++ b/Engine/Source/Runtime/Renderer/Private/RendererScene.cpp
@@ -1097,7 +1097,7 @@ FScene::FScene(UWorld* InWorld, bool bInRequiresHitProxies, bool bInIsEditorScen
 
 	if (World->FXSystem)
 	{
-		World->FXSystem->Destroy();
+		FFXSystemInterface::Destroy(World->FXSystem);
 	}
 
 	if (bCreateFXSystem)
