Add warning clarifying why FindTeleportSpot fails if root component is not collider, or movement component does not target collider, instead of silently failing.

diff --git a/Engine/Source/Runtime/Engine/Private/LevelActor.cpp b/Engine/Source/Runtime/Engine/Private/LevelActor.cpp
--- a/Engine/Source/Runtime/Engine/Private/LevelActor.cpp
+++ b/Engine/Source/Runtime/Engine/Private/LevelActor.cpp
@@ -806,6 +806,32 @@ bool UWorld::FindTeleportSpot(const AActor* TestActor, FVector& TestLocation, FR
 
 	if ( Adjust.IsNearlyZero() )
 	{
+		// EncroachingBlockingGeometry fails to produce adjustment if movement component's target, or root component,
+		// fallback to a box when calling GetCollisionShape (UPrimitiveComponent's behaviour).
+		// This would occur if SkeletalMeshComponent was root, instead of capsule, for example.
+		// Here we test for these cases and warn user why this function is failing.
+		UMovementComponent* const MoveComponent = TestActor->FindComponentByClass<UMovementComponent>();
+		if (MoveComponent && MoveComponent->UpdatedPrimitive)
+		{
+			UPrimitiveComponent* const PrimComponent = MoveComponent->UpdatedPrimitive;
+			FCollisionShape shape = PrimComponent->GetCollisionShape();
+			if (shape.IsBox() && (Cast<UBoxComponent>(PrimComponent) == nullptr))
+			{
+				UE_LOG(LogPhysics, Warning, TEXT("UWorld::FindTeleportSpot called with an actor that is intersecting geometry. Failed to find new location likely due to "
+					"movement component's 'UpdatedComponent' not being a collider component."));
+			}
+		}
+		else
+		{
+			USceneComponent* const RootComponent = TestActor->GetRootComponent();
+			FCollisionShape shape =  Cast<UPrimitiveComponent>(RootComponent)->GetCollisionShape();
+			if (shape.IsBox() && (Cast<UBoxComponent>(RootComponent) == nullptr))
+			{
+				UE_LOG(LogPhysics, Warning, TEXT("UWorld::FindTeleportSpot called with an actor that is intersecting geometry. Failed to find new location likely due to "
+					" actor's root component not being a collider component."));
+			}
+		}
+
 		// Reset in case Adjust is not actually zero
 		TestLocation = OriginalTestLocation;
 		return false;
