Fixed missing shader when rendering shadercomplexity with decals in SM4

diff --git a/Engine/Source/Runtime/Engine/Private/DebugViewModeHelpers.cpp b/Engine/Source/Runtime/Engine/Private/DebugViewModeHelpers.cpp
--- a/Engine/Source/Runtime/Engine/Private/DebugViewModeHelpers.cpp
+++ b/Engine/Source/Runtime/Engine/Private/DebugViewModeHelpers.cpp
@@ -44,7 +44,7 @@ bool AllowDebugViewShaderMode(EDebugViewShaderMode ShaderMode, EShaderPlatform P
 	case DVSM_None:
 		return false;
 	case DVSM_ShaderComplexity:
-		return FeatureLevel != ERHIFeatureLevel::SM4; // SM4 doesn't apply post-processes correctly in the deferred path.
+		return true;
 	case DVSM_ShaderComplexityContainedQuadOverhead:
 	case DVSM_ShaderComplexityBleedingQuadOverhead:
 	case DVSM_QuadComplexity:
