Prevent lazy component template creation logic from replacing references to existing templates with references to the new templates. This could cause the UI to make it look like a parent blueprint was being altered by edits to its children

diff --git a/Engine/Source/Runtime/Engine/Private/InheritableComponentHandler.cpp b/Engine/Source/Runtime/Engine/Private/InheritableComponentHandler.cpp
--- a/Engine/Source/Runtime/Engine/Private/InheritableComponentHandler.cpp
+++ b/Engine/Source/Runtime/Engine/Private/InheritableComponentHandler.cpp
@@ -140,6 +140,8 @@ UActorComponent* UInheritableComponentHandler::CreateOverridenComponentTemplate(
 		NewComponentTemplate->ClearPendingKill();
 		UEngine::FCopyPropertiesForUnrelatedObjectsParams CopyParams;
 		CopyParams.bDoDelta = false;
+		// No good can come of replacing references to BestArchetype with references to NewComponentTemplate
+		CopyParams.bNotifyObjectReplacement = false;
 		UEngine::CopyPropertiesForUnrelatedObjects(BestArchetype, NewComponentTemplate, CopyParams);
 	}
 
@@ -154,6 +156,8 @@ UActorComponent* UInheritableComponentHandler::CreateOverridenComponentTemplate(
 
 			UEngine::FCopyPropertiesForUnrelatedObjectsParams CopyParams;
 			CopyParams.bDoDelta = false;
+			// No good can come of replacing references to BestArchetype with references to NewComponentTemplate
+			CopyParams.bNotifyObjectReplacement = false;
 			UEngine::CopyPropertiesForUnrelatedObjects(BestArchetype, NewComponentTemplate, CopyParams);
 		}
 	}
