Revert CL 4677151
Don't change texture DDC key (later commits have changed DDC key specifically for VTs, so no reason to change key for all textures)

diff --git a/Engine/Source/Runtime/Engine/Private/TextureDerivedData.cpp b/Engine/Source/Runtime/Engine/Private/TextureDerivedData.cpp
--- a/Engine/Source/Runtime/Engine/Private/TextureDerivedData.cpp
+++ b/Engine/Source/Runtime/Engine/Private/TextureDerivedData.cpp
@@ -69,7 +69,7 @@ static TAutoConsoleVariable<int32> CVarVTTileBorderSize(
 // In case of merge conflicts with DDC versions, you MUST generate a new GUID and set this new
 // guid as version
 
-#define TEXTURE_DERIVEDDATA_VER		TEXT("5A375C643EC24E508E41D564657A46BF")
+#define TEXTURE_DERIVEDDATA_VER		TEXT("68A083899C6F4316B8CE0E2958EDE2C2")
 
 // This GUID is mixed into DDC version for virtual textures only, this allows updating DDC version for VT without invalidating DDC for all textures
 // This is useful during development, but once large numbers of VT are present in shipped content, it will have the same problem as TEXTURE_DERIVEDDATA_VER
