Fixup Code for missing Metadata Points + Fix Metadata Serialization

#rnx

diff --git a/Engine/Source/Runtime/Engine/Classes/Components/SplineComponent.h b/Engine/Source/Runtime/Engine/Classes/Components/SplineComponent.h
--- a/Engine/Source/Runtime/Engine/Classes/Components/SplineComponent.h
+++ b/Engine/Source/Runtime/Engine/Classes/Components/SplineComponent.h
@@ -48,6 +48,7 @@ public:
 	virtual void RemovePoint(int32 Index) PURE_VIRTUAL(USplineMetadata::RemovePoint, );
 	virtual void DuplicatePoint(int32 Index) PURE_VIRTUAL(USplineMetadata::DuplicatePoint, );
 	virtual void Reset(int32 NumPoints) PURE_VIRTUAL(USplineMetadata::Reset, );
+	virtual void Fixup(int32 NumPoints) PURE_VIRTUAL(USplineMetadata::Fixup, );
 };
 
 USTRUCT()
@@ -71,8 +72,8 @@ struct ENGINE_API FSplineCurves
 	UPROPERTY()
 	FInterpCurveFloat ReparamTable;
 
-	UPROPERTY()
-	USplineMetadata* Metadata = nullptr;
+	UPROPERTY(Instanced)
+	USplineMetadata* Metadata;
 
 	bool operator==(const FSplineCurves& Other) const
 	{
@@ -285,6 +286,7 @@ public:
 
 	//~ Begin UObject Interface
 	virtual void Serialize(FArchive& Ar) override;
+	virtual void PostLoad() override;
 #if WITH_EDITOR
 	virtual void PostEditChangeChainProperty(FPropertyChangedChainEvent& PropertyChangedEvent) override;
 #endif
diff --git a/Engine/Source/Runtime/Engine/Private/Components/SplineComponent.cpp b/Engine/Source/Runtime/Engine/Private/Components/SplineComponent.cpp
--- a/Engine/Source/Runtime/Engine/Private/Components/SplineComponent.cpp
+++ b/Engine/Source/Runtime/Engine/Private/Components/SplineComponent.cpp
@@ -61,12 +61,11 @@ USplineComponent::USplineComponent(const FObjectInitializer& ObjectInitializer,
 	SplineCurves.Position.Points.Emplace(1.0f, FVector(100, 0, 0), FVector::ZeroVector, FVector::ZeroVector, CIM_CurveAuto);
 	SplineCurves.Rotation.Points.Emplace(1.0f, FQuat::Identity, FQuat::Identity, FQuat::Identity, CIM_CurveAuto);
 	SplineCurves.Scale.Points.Emplace(1.0f, FVector(1.0f), FVector::ZeroVector, FVector::ZeroVector, CIM_CurveAuto);
-
+		
 	if (SplineCurves.Metadata)
 	{
-		SplineCurves.Metadata->Reset(10);
-		SplineCurves.Metadata->AddPoint(0.0f);
-		SplineCurves.Metadata->AddPoint(1.0f);
+		const int32 NumPoints = GetNumberOfSplinePoints();
+		SplineCurves.Metadata->Fixup(NumPoints);
 	}
 
 	UpdateSpline();
@@ -152,6 +151,16 @@ void USplineComponent::Serialize(FArchive& Ar)
 	}
 }
 
+void USplineComponent::PostLoad()
+{
+	Super::PostLoad();
+
+	if (USplineMetadata* MetaData = GetSplinePointsMetadata())
+	{
+		const int32 NumPoints = GetNumberOfSplinePoints();
+		MetaData->Fixup(NumPoints);
+	}
+}
 
 void FSplineCurves::UpdateSpline(bool bClosedLoop, bool bStationaryEndpoints, int32 ReparamStepsPerSegment, bool bLoopPositionOverride, float LoopPosition, const FVector& Scale3D)
 {
