Set dynamic actors whose dormancy setting is DORM_Initial to DORM_DormantAll before they are added to the network drivers.
DORM_Initial should be used exclusively with actors placed in a level.
This fixes undiscovered actors on clients when a dynamic actor was placed in a streamed level with the DORM_Initial flag.

diff --git a/Engine/Source/Runtime/Engine/Private/World.cpp b/Engine/Source/Runtime/Engine/Private/World.cpp
--- a/Engine/Source/Runtime/Engine/Private/World.cpp
+++ b/Engine/Source/Runtime/Engine/Private/World.cpp
@@ -4420,6 +4420,14 @@ void UWorld::AddNetworkActor( AActor* Actor )
 		return;
 	}
 
+    // Remove DORM_Initial from dynamic actors
+	if (Actor->NetDormancy == DORM_Initial && Actor->IsNetStartupActor() == false)
+	{
+		UE_LOG(LogNet, VeryVerbose, TEXT("Dynamic actor %s dormancy was changed from DORM_Initial to DORM_DormantAll"), *Actor->GetFullName());
+        // Changing dormancy directly here since the actor was not yet registered to the network drivers
+		Actor->NetDormancy = DORM_DormantAll;
+	}
+
 	ForEachNetDriver(GEngine, this, [Actor](UNetDriver* const Driver)
 	{
 		if (Driver != nullptr)
