Added a way to reset FNavigationSystem's delegates

diff --git a/Engine/Source/Runtime/Engine/Classes/AI/NavigationSystemBase.h b/Engine/Source/Runtime/Engine/Classes/AI/NavigationSystemBase.h
--- a/Engine/Source/Runtime/Engine/Classes/AI/NavigationSystemBase.h
+++ b/Engine/Source/Runtime/Engine/Classes/AI/NavigationSystemBase.h
@@ -156,7 +156,6 @@ namespace FNavigationSystem
 
 	ENGINE_API bool WantsComponentChangeNotifies();
 
-	//ENGINE_API INavigationDataInterface* GetNavDataForProps(const FNavAgentProperties& AgentProperties);
 	ENGINE_API INavigationDataInterface* GetNavDataForActor(const AActor& Actor);
 	ENGINE_API TSubclassOf<AActor> GetDefaultNavDataClass();
 
@@ -249,6 +248,7 @@ protected:
 	static void SetDefaultWalkableArea(TSubclassOf<UNavAreaBase> InAreaClass);
 	static void SetDefaultObstacleArea(TSubclassOf<UNavAreaBase> InAreaClass);
 
+	static void ResetEventDelegates();
 	static FNavigationSystem::FActorBasedSignature& UpdateActorDataDelegate();
 	static FNavigationSystem::FActorComponentBasedSignature& UpdateComponentDataDelegate();
 	static FNavigationSystem::FSceneComponentBasedSignature& UpdateComponentDataAfterMoveDelegate();
diff --git a/Engine/Source/Runtime/Engine/Private/AI/NavigationSystemBase.cpp b/Engine/Source/Runtime/Engine/Private/AI/NavigationSystemBase.cpp
--- a/Engine/Source/Runtime/Engine/Private/AI/NavigationSystemBase.cpp
+++ b/Engine/Source/Runtime/Engine/Private/AI/NavigationSystemBase.cpp
@@ -149,6 +149,8 @@ namespace FNavigationSystem
 
 	FDelegates Delegates;
 
+	void ResetDelegates() { new(&Delegates)FDelegates(); }
+
 	void UpdateActorData(AActor& Actor) { Delegates.UpdateActorData.Execute(Actor); }
 	void UpdateComponentData(UActorComponent& Comp) { Delegates.UpdateComponentData.Execute(Comp); }
 	void UpdateActorAndComponentData(AActor& Actor, bool bUpdateAttachedActors) { Delegates.UpdateActorAndComponentData.Execute(Actor, bUpdateAttachedActors); }
@@ -343,6 +345,7 @@ void UNavigationSystemBase::SetDefaultObstacleArea(TSubclassOf<UNavAreaBase> InA
 }
 
 
+void UNavigationSystemBase::ResetEventDelegates() { FNavigationSystem::ResetDelegates(); }
 FNavigationSystem::FActorBasedSignature& UNavigationSystemBase::UpdateActorDataDelegate() { return FNavigationSystem::Delegates.UpdateActorData; }
 FNavigationSystem::FActorComponentBasedSignature& UNavigationSystemBase::UpdateComponentDataDelegate() { return FNavigationSystem::Delegates.UpdateComponentData; }
 FNavigationSystem::FSceneComponentBasedSignature& UNavigationSystemBase::UpdateComponentDataAfterMoveDelegate() { return FNavigationSystem::Delegates.UpdateComponentDataAfterMove; }
