Deprecation-related linux cis fixes

diff --git a/Engine/Source/Runtime/Engine/Classes/AI/Navigation/NavigationTypes.h b/Engine/Source/Runtime/Engine/Classes/AI/Navigation/NavigationTypes.h
--- a/Engine/Source/Runtime/Engine/Classes/AI/Navigation/NavigationTypes.h
+++ b/Engine/Source/Runtime/Engine/Classes/AI/Navigation/NavigationTypes.h
@@ -518,6 +518,7 @@ inline uint32 GetTypeHash(const FNavAgentProperties& A)
 	return ((int16(A.AgentRadius) << 16) | int16(A.AgentHeight)) ^ int32(A.AgentStepHeight);
 }
 
+PRAGMA_DISABLE_DEPRECATION_WARNINGS
 USTRUCT(BlueprintType)
 struct ENGINE_API FNavDataConfig : public FNavAgentProperties
 {
@@ -565,6 +566,7 @@ public:
 	}
 #endif // WITH_EDITOR
 };
+PRAGMA_ENABLE_DEPRECATION_WARNINGS
 
 struct FNavigationProjectionWork
 {
diff --git a/Engine/Source/Runtime/Engine/Private/AI/Navigation/NavigationTypes.cpp b/Engine/Source/Runtime/Engine/Private/AI/Navigation/NavigationTypes.cpp
--- a/Engine/Source/Runtime/Engine/Private/AI/Navigation/NavigationTypes.cpp
+++ b/Engine/Source/Runtime/Engine/Private/AI/Navigation/NavigationTypes.cpp
@@ -35,6 +35,7 @@ uint32 FNavPathType::NextUniqueId = 0;
 //----------------------------------------------------------------------//
 // FNavDataConfig
 //----------------------------------------------------------------------//
+PRAGMA_DISABLE_DEPRECATION_WARNINGS
 FNavDataConfig::FNavDataConfig(float Radius, float Height)
 	: FNavAgentProperties(Radius, Height)
 	, Name(TEXT("Default"))
@@ -52,6 +53,7 @@ FNavDataConfig::FNavDataConfig(const FNavDataConfig& Other)
 	, NavDataClass(Other.NavDataClass)
 {
 }
+PRAGMA_ENABLE_DEPRECATION_WARNINGS
 
 void FNavDataConfig::SetNavDataClass(UClass* InNavDataClass)
 {
diff --git a/Engine/Source/Runtime/NavigationSystem/Private/NavigationDirtyAreasController.cpp b/Engine/Source/Runtime/NavigationSystem/Private/NavigationDirtyAreasController.cpp
--- a/Engine/Source/Runtime/NavigationSystem/Private/NavigationDirtyAreasController.cpp
+++ b/Engine/Source/Runtime/NavigationSystem/Private/NavigationDirtyAreasController.cpp
@@ -8,7 +8,9 @@
 //----------------------------------------------------------------------//
 FNavigationDirtyAreasController::FNavigationDirtyAreasController()
 	: bCanAccumulateDirtyAreas(true)
+#if !UE_BUILD_SHIPPING
 	, bDirtyAreasReportedWhileAccumulationLocked(false)
+#endif // !UE_BUILD_SHIPPING
 {
 
 }
diff --git a/Engine/Source/Runtime/NavigationSystem/Public/NavigationDirtyAreasController.h b/Engine/Source/Runtime/NavigationSystem/Public/NavigationDirtyAreasController.h
--- a/Engine/Source/Runtime/NavigationSystem/Public/NavigationDirtyAreasController.h
+++ b/Engine/Source/Runtime/NavigationSystem/Public/NavigationDirtyAreasController.h
@@ -34,5 +34,5 @@ struct NAVIGATIONSYSTEM_API FNavigationDirtyAreasController
 
 #if !UE_BUILD_SHIPPING
 	bool HadDirtyAreasReportedWhileAccumulationLocked() const { return bCanAccumulateDirtyAreas == false && bDirtyAreasReportedWhileAccumulationLocked; }
-#endif UE_BUILD_SHIPPING
+#endif // UE_BUILD_SHIPPING
 };
\ No newline at end of file
