Fix no in game UI being visible.

diff --git a/Engine/Source/Runtime/Engine/Private/Slate/SGameLayerManager.cpp b/Engine/Source/Runtime/Engine/Private/Slate/SGameLayerManager.cpp
--- a/Engine/Source/Runtime/Engine/Private/Slate/SGameLayerManager.cpp
+++ b/Engine/Source/Runtime/Engine/Private/Slate/SGameLayerManager.cpp
@@ -256,7 +256,7 @@ void SGameLayerManager::ClearWidgets()
 
 void SGameLayerManager::Tick(const FGeometry& AllottedGeometry, const double InCurrentTime, const float InDeltaTime)
 {
-	if (AllottedGeometry != CachedGeometry)
+	//if (AllottedGeometry != CachedGeometry)
 	{
 		CachedGeometry = AllottedGeometry;
 		UpdateLayout();
