Simplify and validate instanced static mesh vertex factory creation to mitigate and catch rare crash.

diff --git a/Engine/Shaders/Private/LocalVertexFactory.ush b/Engine/Shaders/Private/LocalVertexFactory.ush
--- a/Engine/Shaders/Private/LocalVertexFactory.ush
+++ b/Engine/Shaders/Private/LocalVertexFactory.ush
@@ -87,9 +87,13 @@ float3 SplineMeshY;
 #define MANUAL_VERTEX_FETCH 0
 #endif
 
-
-
 #if MANUAL_VERTEX_FETCH
+	#define VF_ColorIndexMask_Index 0
+	#define VF_NumTexcoords_Index 1
+	#define FV_LightMapIndex_Index 2
+	#define VF_VertexOffset 3
+
+#if !USE_INSTANCING_EMULATED
 	Buffer<float4> VertexFetch_InstanceOriginBuffer;
 	Buffer<float4> VertexFetch_InstanceTransformBuffer;
 	Buffer<float4> VertexFetch_InstanceLightmapBuffer;
@@ -99,13 +103,9 @@ float3 SplineMeshY;
 		Buffer<uint> VertexFetch_InstanceBoneMapBuffer;
 	#endif
 
-	uint InstanceOffset; 
-
-	#define VF_ColorIndexMask_Index 0
-	#define VF_NumTexcoords_Index 1
-	#define FV_LightMapIndex_Index 2
-	#define VF_VertexOffset 3
-#endif
+	uint InstanceOffset;
+#endif //! !USE_INSTANCING_EMULATED
+#endif //! MANUAL_VERTEX_FETCH
 
 /**
  * Per-vertex inputs from bound vertex buffers
diff --git a/Engine/Source/Runtime/Engine/Private/InstancedStaticMesh.cpp b/Engine/Source/Runtime/Engine/Private/InstancedStaticMesh.cpp
--- a/Engine/Source/Runtime/Engine/Private/InstancedStaticMesh.cpp
+++ b/Engine/Source/Runtime/Engine/Private/InstancedStaticMesh.cpp
@@ -369,6 +369,10 @@ void FStaticMeshInstanceBuffer::CreateVertexBuffer(FResourceArrayInterface* InRe
 
 void FStaticMeshInstanceBuffer::BindInstanceVertexBuffer(const class FVertexFactory* VertexFactory, FInstancedStaticMeshDataType& InstancedStaticMeshData) const
 {
+	check(InstanceOriginSRV);
+	check(InstanceTransformSRV);
+	check(InstanceLightmapSRV);
+
 	{
 		InstancedStaticMeshData.InstanceOriginSRV = InstanceOriginSRV;
 		InstancedStaticMeshData.InstanceTransformSRV = InstanceTransformSRV;
@@ -642,47 +646,47 @@ FVertexFactoryShaderParameters* FInstancedStaticMeshVertexFactory::ConstructShad
 IMPLEMENT_VERTEX_FACTORY_TYPE_EX(FInstancedStaticMeshVertexFactory,"/Engine/Private/LocalVertexFactory.ush",true,true,true,true,true,true,false);
 IMPLEMENT_VERTEX_FACTORY_TYPE_EX(FEmulatedInstancedStaticMeshVertexFactory,"/Engine/Private/LocalVertexFactory.ush",true,true,true,true,true,true,false);
 
-void FInstancedStaticMeshRenderData::InitStaticMeshVertexFactories(
-		TIndirectArray<FInstancedStaticMeshVertexFactory>* VertexFactories,
-		FInstancedStaticMeshRenderData* InstancedRenderData,
-		UStaticMesh* Parent)
+void FInstancedStaticMeshRenderData::InitVertexFactories()
 {
 	const bool bInstanced = GRHISupportsInstancing;
 
-	for( int32 LODIndex=0;LODIndex<VertexFactories->Num(); LODIndex++ )
+	// Allocate the vertex factories for each LOD
+	for (int32 LODIndex = 0; LODIndex < LODModels.Num(); LODIndex++)
 	{
-		const FStaticMeshLODResources* RenderData = &InstancedRenderData->LODModels[LODIndex];
-						
-		FInstancedStaticMeshVertexFactory::FDataType Data;
-		// Assign to the vertex factory for this LOD.
-		FInstancedStaticMeshVertexFactory& VertexFactory = (*VertexFactories)[LODIndex];
+		VertexFactories.Add(bInstanced ? new FInstancedStaticMeshVertexFactory(FeatureLevel) : new FEmulatedInstancedStaticMeshVertexFactory(FeatureLevel));
+	}
 
-		RenderData->VertexBuffers.PositionVertexBuffer.BindPositionVertexBuffer(&VertexFactory, Data);
-		RenderData->VertexBuffers.StaticMeshVertexBuffer.BindTangentVertexBuffer(&VertexFactory, Data);
-		RenderData->VertexBuffers.StaticMeshVertexBuffer.BindPackedTexCoordVertexBuffer(&VertexFactory, Data);
-		if (Parent->LightMapCoordinateIndex < (int32)RenderData->VertexBuffers.StaticMeshVertexBuffer.GetNumTexCoords() && Parent->LightMapCoordinateIndex >= 0)
+	const int32 LightMapCoordinateIndex = Component->GetStaticMesh()->LightMapCoordinateIndex;
+	ENQUEUE_RENDER_COMMAND(InstancedStaticMeshRenderData_InitVertexFactories)(
+		[this, LightMapCoordinateIndex, bInstanced](FRHICommandListImmediate& RHICmdList)
+	{
+		for (int32 LODIndex = 0; LODIndex < VertexFactories.Num(); LODIndex++)
 		{
-			RenderData->VertexBuffers.StaticMeshVertexBuffer.BindLightMapVertexBuffer(&VertexFactory, Data, Parent->LightMapCoordinateIndex);
-		}
-		RenderData->VertexBuffers.ColorVertexBuffer.BindColorVertexBuffer(&VertexFactory, Data);
+			const FStaticMeshLODResources* RenderData = &LODModels[LODIndex];
 
-		if (bInstanced && InstancedRenderData->PerInstanceRenderData.IsValid())
-		{
-			const FStaticMeshInstanceBuffer& InstanceBuffer = InstancedRenderData->PerInstanceRenderData->InstanceBuffer;
-			InstanceBuffer.BindInstanceVertexBuffer(&VertexFactory, Data);
-		}
+			FInstancedStaticMeshVertexFactory::FDataType Data;
+			// Assign to the vertex factory for this LOD.
+			FInstancedStaticMeshVertexFactory& VertexFactory = VertexFactories[LODIndex];
 
-		VertexFactory.SetData(Data);
+			RenderData->VertexBuffers.PositionVertexBuffer.BindPositionVertexBuffer(&VertexFactory, Data);
+			RenderData->VertexBuffers.StaticMeshVertexBuffer.BindTangentVertexBuffer(&VertexFactory, Data);
+			RenderData->VertexBuffers.StaticMeshVertexBuffer.BindPackedTexCoordVertexBuffer(&VertexFactory, Data);
+			if (LightMapCoordinateIndex < (int32)RenderData->VertexBuffers.StaticMeshVertexBuffer.GetNumTexCoords() && LightMapCoordinateIndex >= 0)
+			{
+				RenderData->VertexBuffers.StaticMeshVertexBuffer.BindLightMapVertexBuffer(&VertexFactory, Data, LightMapCoordinateIndex);
+			}
+			RenderData->VertexBuffers.ColorVertexBuffer.BindColorVertexBuffer(&VertexFactory, Data);
 
-		if (VertexFactory.IsInitialized())
-		{
-			VertexFactory.UpdateRHI();
-		}
-		else
-		{
+			if (bInstanced)
+			{
+				check(PerInstanceRenderData);
+				PerInstanceRenderData->InstanceBuffer.BindInstanceVertexBuffer(&VertexFactory, Data);
+			}
+
+			VertexFactory.SetData(Data);
 			VertexFactory.InitResource();
 		}
-	}
+	});
 }
 
 FPerInstanceRenderData::FPerInstanceRenderData(FStaticMeshInstanceData& Other, ERHIFeatureLevel::Type InFeaureLevel, bool InRequireCPUAccess)
@@ -2613,14 +2617,16 @@ void UInstancedStaticMeshComponent::OnPostLoadPerInstanceData()
 
 	if (PerInstanceRenderData.IsValid())
 	{
-		AActor* Owner = GetOwner();
-		ULevel* OwnerLevel = Owner->GetLevel();
-		UWorld* OwnerWorld = OwnerLevel ? OwnerLevel->OwningWorld : nullptr;
-
-		if (OwnerWorld && OwnerWorld->GetActiveLightingScenario() != nullptr && OwnerWorld->GetActiveLightingScenario() != OwnerLevel)
+		if (AActor* Owner = GetOwner())
 		{
-			//update the instance data if the lighting scenario isn't the owner level
-			InstanceUpdateCmdBuffer.Edit();
+			ULevel* OwnerLevel = Owner->GetLevel();
+			UWorld* OwnerWorld = OwnerLevel ? OwnerLevel->OwningWorld : nullptr;
+
+			if (OwnerWorld && OwnerWorld->GetActiveLightingScenario() != nullptr && OwnerWorld->GetActiveLightingScenario() != OwnerLevel)
+			{
+				//update the instance data if the lighting scenario isn't the owner level
+				InstanceUpdateCmdBuffer.Edit();
+			}
 		}
 	}
 }
@@ -2894,6 +2900,8 @@ void FInstancedStaticMeshVertexFactoryShaderParameters::GetElementShaderBindings
 	FVertexInputStreamArray& VertexStreams
 	) const
 {
+	const bool bInstanced = GRHISupportsInstancing;
+
 	// Decode VertexFactoryUserData as VertexFactoryUniformBuffer
 	FRHIUniformBuffer* VertexFactoryUniformBuffer = static_cast<FRHIUniformBuffer*>(BatchElement.VertexFactoryUserData);
 	FLocalVertexFactoryShaderParametersBase::GetElementShaderBindingsBase(Scene, View, Shader, InputStreamType, FeatureLevel, VertexFactory, BatchElement, VertexFactoryUniformBuffer, ShaderBindings, VertexStreams);
@@ -2902,17 +2910,35 @@ void FInstancedStaticMeshVertexFactoryShaderParameters::GetElementShaderBindings
 	const auto* InstancedVertexFactory = static_cast<const FInstancedStaticMeshVertexFactory*>(VertexFactory);
 	const int32 InstanceOffsetValue = BatchElement.UserIndex;
 
-	if (InstancedVertexFactory->SupportsManualVertexFetch(FeatureLevel))
+	if (bInstanced)
 	{
-		ShaderBindings.Add(VertexFetch_InstanceOriginBufferParameter, InstancedVertexFactory->GetInstanceOriginSRV());
-		ShaderBindings.Add(VertexFetch_InstanceTransformBufferParameter, InstancedVertexFactory->GetInstanceTransformSRV());
-		ShaderBindings.Add(VertexFetch_InstanceLightmapBufferParameter, InstancedVertexFactory->GetInstanceLightmapSRV());
-		ShaderBindings.Add(InstanceOffset, InstanceOffsetValue);
-	}
+		if (InstancedVertexFactory->SupportsManualVertexFetch(FeatureLevel))
+		{
+			ShaderBindings.Add(VertexFetch_InstanceOriginBufferParameter, InstancedVertexFactory->GetInstanceOriginSRV());
+			ShaderBindings.Add(VertexFetch_InstanceTransformBufferParameter, InstancedVertexFactory->GetInstanceTransformSRV());
+			ShaderBindings.Add(VertexFetch_InstanceLightmapBufferParameter, InstancedVertexFactory->GetInstanceLightmapSRV());
+			ShaderBindings.Add(InstanceOffset, InstanceOffsetValue);
+		}
 
-	if (GRHISupportsInstancing && InstanceOffsetValue > 0 && VertexStreams.Num() > 0)
+		if (InstanceOffsetValue > 0 && VertexStreams.Num() > 0)
+		{
+			VertexFactory->OffsetInstanceStreams(InstanceOffsetValue, InputStreamType, VertexStreams);
+		}
+	}
+	else if (CPUInstanceOrigin.IsBound())
 	{
-		VertexFactory->OffsetInstanceStreams(InstanceOffsetValue, InputStreamType, VertexStreams);
+		const float ShortScale = 1.0f / 32767.0f;
+		auto* InstancingData = (const FInstancingUserData*)BatchElement.UserData;
+		check(InstancingData);
+
+		FVector4 InstanceTransform[3];
+		FVector4 InstanceLightmapAndShadowMapUVBias;
+		FVector4 InstanceOrigin;
+		InstancingData->RenderData->PerInstanceRenderData->InstanceBuffer.GetInstanceShaderValues(BatchElement.UserIndex, InstanceTransform, InstanceLightmapAndShadowMapUVBias, InstanceOrigin);
+
+		ShaderBindings.Add(CPUInstanceOrigin, InstanceOrigin);
+		ShaderBindings.Add(CPUInstanceTransform, InstanceTransform);
+		ShaderBindings.Add(CPUInstanceLightmapAndShadowMapBias, InstanceLightmapAndShadowMapUVBias);
 	}
 
 	if( InstancingWorldViewOriginOneParameter.IsBound() )
@@ -3059,24 +3085,4 @@ void FInstancedStaticMeshVertexFactoryShaderParameters::GetElementShaderBindings
 		ShaderBindings.Add(InstancingFadeOutParamsParameter, InstancingFadeOutParams);
 
 	}
-
-	const bool bInstanced = GRHISupportsInstancing;
-	if (!bInstanced)
-	{
-		if (CPUInstanceOrigin.IsBound())
-		{
-			const float ShortScale = 1.0f / 32767.0f;
-			auto* InstancingData = (const FInstancingUserData*)BatchElement.UserData;
-			check(InstancingData);
-
-			FVector4 InstanceTransform[3];
-			FVector4 InstanceLightmapAndShadowMapUVBias;
-			FVector4 InstanceOrigin;
-			InstancingData->RenderData->PerInstanceRenderData->InstanceBuffer.GetInstanceShaderValues(BatchElement.UserIndex, InstanceTransform, InstanceLightmapAndShadowMapUVBias, InstanceOrigin);
-
-			ShaderBindings.Add(CPUInstanceOrigin, InstanceOrigin);
-			ShaderBindings.Add(CPUInstanceTransform, InstanceTransform);
-			ShaderBindings.Add(CPUInstanceLightmapAndShadowMapBias, InstanceLightmapAndShadowMapUVBias);
-		}
-	}
 }
diff --git a/Engine/Source/Runtime/Engine/Private/InstancedStaticMesh.h b/Engine/Source/Runtime/Engine/Private/InstancedStaticMesh.h
--- a/Engine/Source/Runtime/Engine/Private/InstancedStaticMesh.h
+++ b/Engine/Source/Runtime/Engine/Private/InstancedStaticMesh.h
@@ -431,39 +431,9 @@ public:
 		check(PerInstanceRenderData.IsValid());
 		// Allocate the vertex factories for each LOD
 		InitVertexFactories();
-		ReInitVertexFactories();
 		RegisterSpeedTreeWind();
 	}
 
-	~FInstancedStaticMeshRenderData()
-	{
-	}
-
-	void ReInitVertexFactories()
-	{
-		// Initialize the static mesh's vertex factory.
-		TIndirectArray<FInstancedStaticMeshVertexFactory>* InVertexFactories = &VertexFactories;
-		FInstancedStaticMeshRenderData* InstancedRenderData = this;
-		UStaticMesh* Parent = Component->GetStaticMesh();
-		ENQUEUE_RENDER_COMMAND(CallInitStaticMeshVertexFactory)(
-			[InVertexFactories, InstancedRenderData, Parent](FRHICommandListImmediate& RHICmdList)
-		{
-			InitStaticMeshVertexFactories(InVertexFactories, InstancedRenderData, Parent );
-		});
-	}
-
-	void RegisterSpeedTreeWind()
-	{
-		// register SpeedTree wind with the scene
-		if (Component->GetStaticMesh()->SpeedTreeWind.IsValid())
-		{
-			for (int32 LODIndex = 0; LODIndex < LODModels.Num(); LODIndex++)
-			{
-				Component->GetScene()->AddSpeedTreeWind(&VertexFactories[LODIndex], Component->GetStaticMesh());
-			}
-		}
-	}
-
 	void ReleaseResources(FSceneInterface* Scene, const UStaticMesh* StaticMesh)
 	{
 		// unregister SpeedTree wind with the scene
@@ -475,17 +445,12 @@ public:
 			}
 		}
 
-		for( int32 LODIndex=0;LODIndex<VertexFactories.Num();LODIndex++ )
+		for (int32 LODIndex = 0; LODIndex < VertexFactories.Num(); LODIndex++)
 		{
 			VertexFactories[LODIndex].ReleaseResource();
 		}
 	}
 
-	static void InitStaticMeshVertexFactories(
-		TIndirectArray<FInstancedStaticMeshVertexFactory>* VertexFactories,
-		FInstancedStaticMeshRenderData* InstancedRenderData,
-		UStaticMesh* Parent);
-
 	/** Source component */
 	UInstancedStaticMeshComponent* Component;
 
@@ -502,24 +467,17 @@ public:
 	ERHIFeatureLevel::Type FeatureLevel;
 
 private:
+	void InitVertexFactories();
 
-	void InitVertexFactories()
+	void RegisterSpeedTreeWind()
 	{
-		const bool bEmulatedInstancing = !GRHISupportsInstancing;
-		
-		// Allocate the vertex factories for each LOD
-		for( int32 LODIndex=0;LODIndex<LODModels.Num();LODIndex++ )
+		// register SpeedTree wind with the scene
+		if (Component->GetStaticMesh()->SpeedTreeWind.IsValid())
 		{
-			FInstancedStaticMeshVertexFactory* VertexFactoryPtr;
-			if (bEmulatedInstancing)
-			{
-				VertexFactoryPtr = new FEmulatedInstancedStaticMeshVertexFactory(FeatureLevel);
-			}
-			else
+			for (int32 LODIndex = 0; LODIndex < LODModels.Num(); LODIndex++)
 			{
-				VertexFactoryPtr = new FInstancedStaticMeshVertexFactory(FeatureLevel);
+				Component->GetScene()->AddSpeedTreeWind(&VertexFactories[LODIndex], Component->GetStaticMesh());
 			}
-			VertexFactories.Add(VertexFactoryPtr);
 		}
 	}
 };
