Undo change to GetCollisionEnabled_CheckOwner in 4678617 affecting the results of bodies on skeletal mesh.

diff --git a/Engine/Source/Runtime/Engine/Private/PhysicsEngine/BodyInstance.cpp b/Engine/Source/Runtime/Engine/Private/PhysicsEngine/BodyInstance.cpp
--- a/Engine/Source/Runtime/Engine/Private/PhysicsEngine/BodyInstance.cpp
+++ b/Engine/Source/Runtime/Engine/Private/PhysicsEngine/BodyInstance.cpp
@@ -748,20 +748,20 @@ ECollisionEnabled::Type FBodyInstance::GetCollisionEnabled_CheckOwner() const
 	// Check actor override
 	const UPrimitiveComponent* OwnerComponentInst = OwnerComponent.Get();
 	AActor* Owner = OwnerComponentInst ? OwnerComponentInst->GetOwner() : nullptr;
-		if (Owner && !Owner->GetActorEnableCollision())
-		{
-			return ECollisionEnabled::NoCollision;
-		}
-		else if(const USkeletalMeshComponent* SkelMeshComp = Cast<USkeletalMeshComponent>(OwnerComponentInst))
-{
-			// Check component override (skel mesh case)
-			return (ECollisionEnabled::Type)FMath::Max(SkelMeshComp->BodyInstance.CollisionEnabled, CollisionEnabled);
-}
+	if (Owner && !Owner->GetActorEnableCollision())
+	{
+		return ECollisionEnabled::NoCollision;
+	}
+	else if(const USkeletalMeshComponent* SkelMeshComp = Cast<USkeletalMeshComponent>(OwnerComponentInst))
+	{
+		// Check component override (skel mesh case)
+		return SkelMeshComp->BodyInstance.CollisionEnabled;
+	}
 	else
 	{
 		return CollisionEnabled;
 	}
-			}
+}
 
 void FBodyInstance::SetMaskFilter(FMaskFilter InMaskFilter)
 {
