Create FInternetAddrDemo (which is essentially a dummy internet addr) to be returned when GetInternetAddr is called on UDemoNetConnection. This fixes an issue with network profiling where demo traffic would be recorded on a client connection instead.

If FNetworkProfiler::SetCurrentConnection is passed a connection that has an invalid FInternetAddr shared pointer, it will write the data out to the "Invalid Connection" index instead of whatever the previous connection was.

diff --git a/Engine/Source/Runtime/Engine/Classes/Engine/DemoNetConnection.h b/Engine/Source/Runtime/Engine/Classes/Engine/DemoNetConnection.h
--- a/Engine/Source/Runtime/Engine/Classes/Engine/DemoNetConnection.h
+++ b/Engine/Source/Runtime/Engine/Classes/Engine/DemoNetConnection.h
@@ -75,7 +75,7 @@ public:
 	virtual int32 IsNetReady( bool Saturate ) override;
 	virtual void FlushNet( bool bIgnoreSimulation = false ) override;
 	virtual void HandleClientPlayer( APlayerController* PC, class UNetConnection* NetConnection ) override;
-	virtual TSharedPtr<const FInternetAddr> GetRemoteAddr() override { return nullptr; }
+	virtual TSharedPtr<const FInternetAddr> GetRemoteAddr() override;
 	virtual bool ClientHasInitializedLevelFor( const AActor* TestActor ) const override;
 	virtual TSharedPtr<FObjectReplicator> CreateReplicatorForNewActorChannel(UObject* Object);
 	virtual FString RemoteAddressToString() override { return TEXT("Demo"); }
diff --git a/Engine/Source/Runtime/Engine/Private/DemoNetDriver.cpp b/Engine/Source/Runtime/Engine/Private/DemoNetDriver.cpp
--- a/Engine/Source/Runtime/Engine/Private/DemoNetDriver.cpp
+++ b/Engine/Source/Runtime/Engine/Private/DemoNetDriver.cpp
@@ -5717,6 +5717,113 @@ void UDemoNetDriver::SetSpectatorController(APlayerController* PC)
 	}
 }
 
+/*------------------------------------------------------------------------------------------
+	FInternetAddrDemo - dummy internet addr that can be used for anything that requests it.
+--------------------------------------------------------------------------------------------*/
+class FInternetAddrDemo : public FInternetAddr
+{
+public:
+
+	FInternetAddrDemo()
+	{
+	}
+
+	virtual TArray<uint8> GetRawIp() const override
+	{
+		return TArray<uint8>();
+	}
+
+	virtual void SetRawIp(const TArray<uint8>& RawAddr) override
+	{
+	}
+
+	void SetIp(uint32 InAddr) override
+	{
+	}
+
+
+	void SetIp(const TCHAR* InAddr, bool& bIsValid) override
+	{
+	}
+
+	void GetIp(uint32& OutAddr) const override
+	{
+		OutAddr = 0;
+	}
+
+	void SetPort(int32 InPort) override
+	{
+	}
+
+
+	void GetPort(int32& OutPort) const override
+	{
+		OutPort = 0;
+	}
+
+
+	int32 GetPort() const override
+	{
+		return 0;
+	}
+
+	void SetAnyAddress() override
+	{
+	}
+
+	void SetBroadcastAddress() override
+	{
+	}
+
+	void SetLoopbackAddress() override
+	{
+	}
+
+	FString ToString(bool bAppendPort) const override
+	{
+		return FString(TEXT("Demo Internet Address"));
+	}
+
+	virtual bool operator==(const FInternetAddr& Other) const override
+	{
+		return Other.ToString(true) == ToString(true);
+	}
+
+	bool operator!=(const FInternetAddrDemo& Other) const
+	{
+		return !(FInternetAddrDemo::operator==(Other));
+	}
+
+	virtual uint32 GetTypeHash() const override
+	{
+		return GetConstTypeHash();
+	}
+
+	uint32 GetConstTypeHash() const
+	{
+		return ::GetTypeHash(ToString(true));
+	}
+
+	friend uint32 GetTypeHash(const FInternetAddrDemo& A)
+	{
+		return A.GetConstTypeHash();
+	}
+
+	virtual bool IsValid() const override
+	{
+		return true;
+	}
+
+	virtual TSharedRef<FInternetAddr> Clone() const override
+	{
+		return DemoInternetAddr.ToSharedRef();
+	}
+
+	static TSharedPtr<FInternetAddr> DemoInternetAddr;
+};
+
+TSharedPtr<FInternetAddr> FInternetAddrDemo::DemoInternetAddr = MakeShareable(new FInternetAddrDemo);
+
 /*-----------------------------------------------------------------------------
 	UDemoNetConnection.
 -----------------------------------------------------------------------------*/
@@ -5859,6 +5966,11 @@ void UDemoNetConnection::HandleClientPlayer(APlayerController* PC, UNetConnectio
 	}
 }
 
+TSharedPtr<const FInternetAddr> UDemoNetConnection::GetRemoteAddr()
+{
+	return FInternetAddrDemo::DemoInternetAddr;
+}
+
 bool UDemoNetConnection::ClientHasInitializedLevelFor(const AActor* TestActor) const
 {
 	// We save all currently streamed levels into the demo stream so we can force the demo playback client
diff --git a/Engine/Source/Runtime/Engine/Private/NetworkProfiler.cpp b/Engine/Source/Runtime/Engine/Private/NetworkProfiler.cpp
--- a/Engine/Source/Runtime/Engine/Private/NetworkProfiler.cpp
+++ b/Engine/Source/Runtime/Engine/Private/NetworkProfiler.cpp
@@ -229,19 +229,29 @@ void FNetworkProfiler::SetCurrentConnection( UNetConnection* Connection )
 {
 	if ( bIsTrackingEnabled && Connection != nullptr )
 	{
+		uint32 Index = INDEX_NONE;
+
 		const TSharedPtr<const FInternetAddr> ConnectionAddr = Connection->GetRemoteAddr();
 		if ( ConnectionAddr.IsValid() )
 		{
 			if ( LastAddress != ConnectionAddr )
 			{
-				uint32 Index = GetAddressTableIndex(ConnectionAddr->ToString(true));
+				Index = GetAddressTableIndex(ConnectionAddr->ToString(true));
+			}
+		}
+		else
+		{
+			static const FString InvalidConnection = FString(TEXT("Invalid Connection"));
+			Index = GetAddressTableIndex(InvalidConnection);
+		}
 
-				uint8 Type = NPTYPE_ConnectionChanged;
-				(*FileWriter) << Type;
-				(*FileWriter).SerializeIntPacked(Index);
+		if (Index != INDEX_NONE)
+		{
+			uint8 Type = NPTYPE_ConnectionChanged;
+			(*FileWriter) << Type;
+			(*FileWriter).SerializeIntPacked(Index);
 
-				LastAddress = ConnectionAddr;
-			}
+			LastAddress = ConnectionAddr;
 		}
 	}
 }
