Fixed serialization / versioning issue with MeshDescription, and fixed unnecessary memory bloat by not retaining MeshDescriptionBulkData which is generated to be written to the DDC when a legacy static mesh is first loaded. This goes alongside CL 7248033

diff --git a/Engine/Source/Editor/StaticMeshEditor/Private/StaticMeshEditorTools.cpp b/Engine/Source/Editor/StaticMeshEditor/Private/StaticMeshEditorTools.cpp
--- a/Engine/Source/Editor/StaticMeshEditor/Private/StaticMeshEditorTools.cpp
+++ b/Engine/Source/Editor/StaticMeshEditor/Private/StaticMeshEditorTools.cpp
@@ -3365,7 +3365,7 @@ void FLevelOfDetailSettingsLayout::AddLODLevelCategories( IDetailLayoutBuilder&
 			CategoryName.AppendInt( LODIndex );
 
 			FText LODLevelString = FText::FromString(FString(TEXT("LOD ")) + FString::FromInt(LODIndex) );
-			bool bHasBeenSimplified = StaticMesh->GetMeshDescription(LODIndex) == nullptr || StaticMesh->IsReductionActive(LODIndex);
+			bool bHasBeenSimplified = !StaticMesh->IsMeshDescriptionValid(LODIndex) || StaticMesh->IsReductionActive(LODIndex);
 			FText GeneratedString = FText::FromString(bHasBeenSimplified ? TEXT("[generated]") : TEXT(""));
 
 			IDetailCategoryBuilder& LODCategory = DetailBuilder.EditCategory( *CategoryName, LODLevelString, ECategoryPriority::Important );
diff --git a/Engine/Source/Runtime/Engine/Private/StaticMesh.cpp b/Engine/Source/Runtime/Engine/Private/StaticMesh.cpp
--- a/Engine/Source/Runtime/Engine/Private/StaticMesh.cpp
+++ b/Engine/Source/Runtime/Engine/Private/StaticMesh.cpp
@@ -3992,7 +3992,8 @@ void UStaticMesh::CacheMeshData()
 						FillMaterialName(StaticMaterials, MaterialMap);
 						FMeshDescriptionOperations::ConvertFromRawMesh(TempRawMesh, *SourceModel.MeshDescription, MaterialMap);
 
-						// Pack MeshDescription into temporary bulk data
+						// Pack MeshDescription into temporary bulk data, ready to write out to DDC.
+						// This will be reloaded from the DDC when needed if a MeshDescription is requested from the static mesh.
 						FMeshDescriptionBulkData MeshDescriptionBulkData;
 						MeshDescriptionBulkData.SaveMeshDescription(*SourceModel.MeshDescription);
 
@@ -4002,13 +4003,6 @@ void UStaticMesh::CacheMeshData()
 						FMemoryWriter Ar(DerivedData, bIsPersistent);
 						MeshDescriptionBulkData.Serialize(Ar, this);
 						GetDerivedDataCacheRef().Put(*MeshDataKey, DerivedData);
-
-						// Pack MeshDescription into the bulk data
-						if (!SourceModel.MeshDescriptionBulkData.IsValid())
-						{
-							SourceModel.MeshDescriptionBulkData = MakeUnique<FMeshDescriptionBulkData>();
-						}
-						SourceModel.MeshDescriptionBulkData->SaveMeshDescription(*SourceModel.MeshDescription);
 					}
 				}
 			}
