Fix component instance data cache not getting applied correctly to a root component that is attached to another actor.

diff --git a/Engine/Source/Runtime/Engine/Private/ComponentInstanceDataCache.cpp b/Engine/Source/Runtime/Engine/Private/ComponentInstanceDataCache.cpp
--- a/Engine/Source/Runtime/Engine/Private/ComponentInstanceDataCache.cpp
+++ b/Engine/Source/Runtime/Engine/Private/ComponentInstanceDataCache.cpp
@@ -561,11 +561,11 @@ void FComponentInstanceDataCache::ApplyToActor(AActor* Actor, const ECacheApplyP
 							// components within a child actor are handled by applying the instance data to the child actor component
 							if (ChildComponent->GetOwner() == Actor)
 							{
-							Components.Add(ChildComponent);
+								Components.Add(ChildComponent);
+							}
 						}
 					}
 				}
-				}
 
 				// next loop start with the nodes we just added
 				FirstProcessIndex = StartingProcessedCount;
@@ -582,16 +582,19 @@ void FComponentInstanceDataCache::ApplyToActor(AActor* Actor, const ECacheApplyP
 		{
 			if (USceneComponent* SceneComponent = Cast<USceneComponent>(Component))
 			{
-				// Scene components that aren't attached to the root component hierarchy won't already have been processed so process them now.
-				// * If there is an unattached scene component 
-				// * If there is a scene component attached to another Actor's hierarchy
-				// * If the scene is not registered (likely because bAutoRegister is false or component is marked pending kill), then we may not have successfully attached to our parent and properly been handled
-				USceneComponent* ParentComponent = SceneComponent->GetAttachParent();
-				if (   (ParentComponent == nullptr && SceneComponent != Actor->GetRootComponent()) 			
-					|| (ParentComponent && ParentComponent->GetOwner() != Actor) 					
-					|| (ParentComponent && !SceneComponent->IsRegistered() && !ParentComponent->GetAttachChildren().Contains(SceneComponent)))
+				if (SceneComponent != Actor->GetRootComponent())
 				{
-					AddComponentHierarchy(SceneComponent);
+					// Scene components that aren't attached to the root component hierarchy won't already have been processed so process them now.
+					// * If there is an unattached scene component 
+					// * If there is a scene component attached to another Actor's hierarchy
+					// * If the scene is not registered (likely because bAutoRegister is false or component is marked pending kill), then we may not have successfully attached to our parent and properly been handled
+					USceneComponent* ParentComponent = SceneComponent->GetAttachParent();
+					if (   (ParentComponent == nullptr)
+					    || (ParentComponent->GetOwner() != Actor) 					
+					    || (!SceneComponent->IsRegistered() && !ParentComponent->GetAttachChildren().Contains(SceneComponent)))
+					{
+						AddComponentHierarchy(SceneComponent);
+					}
 				}
 			}
 			else if (Component)
