Being able to hide parameters that are not connected in the material graph
#rnx

diff --git a/Engine/Source/Runtime/Engine/Classes/Materials/MaterialExpression.h b/Engine/Source/Runtime/Engine/Classes/Materials/MaterialExpression.h
--- a/Engine/Source/Runtime/Engine/Classes/Materials/MaterialExpression.h
+++ b/Engine/Source/Runtime/Engine/Classes/Materials/MaterialExpression.h
@@ -418,6 +418,8 @@ class ENGINE_API UMaterialExpression : public UObject
 	virtual FName GetParameterName() const { return NAME_None; }
 	virtual void SetParameterName(const FName& Name) {}
 
+	virtual bool HasConnectedOutputs() const;
+
 #endif // WITH_EDITOR
 
 	/** Checks whether any inputs to this expression create a loop */
@@ -435,4 +437,3 @@ protected:
 };
 
 
-
diff --git a/Engine/Source/Runtime/Engine/Private/Materials/MaterialExpressions.cpp b/Engine/Source/Runtime/Engine/Private/Materials/MaterialExpressions.cpp
--- a/Engine/Source/Runtime/Engine/Private/Materials/MaterialExpressions.cpp
+++ b/Engine/Source/Runtime/Engine/Private/Materials/MaterialExpressions.cpp
@@ -1278,6 +1278,30 @@ bool UMaterialExpression::HasClassAndNameCollision(UMaterialExpression* OtherExp
 	return GetClass() == OtherExpression->GetClass();
 }
 
+
+bool UMaterialExpression::HasConnectedOutputs() const
+{
+	bool bIsConnected = !GraphNode;
+	if (GraphNode)
+	{
+		UMaterialGraphNode* MatGraphNode = Cast<UMaterialGraphNode>(GraphNode);
+		if (MatGraphNode)
+		{
+			TArray<UEdGraphPin*> OutputPins;
+			MatGraphNode->GetOutputPins(OutputPins);
+			for (UEdGraphPin* Pin : OutputPins)
+			{
+				if (Pin->LinkedTo.Num() > 0)
+				{
+					bIsConnected = true;
+				}
+			}
+		}
+	}
+	return bIsConnected;
+}
+
+
 #endif // WITH_EDITOR
 
 
@@ -2398,13 +2422,16 @@ void UMaterialExpressionTextureSampleParameter::GetAllParameterInfo(TArray<FMate
 	{
 		NewParameter.ParameterLocation = Function;
 	}
-#endif
-
-	OutParameterInfo.AddUnique(NewParameter);
 
-	if(CurrentSize != OutParameterInfo.Num())
+	if (HasConnectedOutputs())
+#endif
 	{
-		OutParameterIds.Add(ExpressionGUID);
+		OutParameterInfo.AddUnique(NewParameter);
+
+		if (CurrentSize != OutParameterInfo.Num())
+		{
+			OutParameterIds.Add(ExpressionGUID);
+		}
 	}
 }
 
@@ -6679,12 +6706,15 @@ void UMaterialExpressionParameter::GetAllParameterInfo(TArray<FMaterialParameter
 	{
 		NewParameter.ParameterLocation = Function;
 	}
-#endif
 
-	OutParameterInfo.AddUnique(NewParameter);
-	if(CurrentSize != OutParameterInfo.Num())
+	if (HasConnectedOutputs())
+#endif
 	{
-		OutParameterIds.Add(ExpressionGUID);
+		OutParameterInfo.AddUnique(NewParameter);
+		if (CurrentSize != OutParameterInfo.Num())
+		{
+			OutParameterIds.Add(ExpressionGUID);
+		}
 	}
 }
 
@@ -9795,11 +9825,14 @@ void UMaterialExpressionFontSampleParameter::GetAllParameterInfo(TArray<FMateria
 	{
 		NewParameter.ParameterLocation = Function;
 	}
+	if (HasConnectedOutputs())
 #endif
-	OutParameterInfo.AddUnique(NewParameter);
-	if(CurrentSize != OutParameterInfo.Num())
 	{
-		OutParameterIds.Add(ExpressionGUID);
+		OutParameterInfo.AddUnique(NewParameter);
+		if (CurrentSize != OutParameterInfo.Num())
+		{
+			OutParameterIds.Add(ExpressionGUID);
+		}
 	}
 }
 
