Fixed automated test crash and patched issue where render state was not being invalidated on instanced static mesh.

diff --git a/Engine/Source/Runtime/Engine/Private/InstancedStaticMesh.cpp b/Engine/Source/Runtime/Engine/Private/InstancedStaticMesh.cpp
--- a/Engine/Source/Runtime/Engine/Private/InstancedStaticMesh.cpp
+++ b/Engine/Source/Runtime/Engine/Private/InstancedStaticMesh.cpp
@@ -369,9 +369,12 @@ void FStaticMeshInstanceBuffer::CreateVertexBuffer(FResourceArrayInterface* InRe
 
 void FStaticMeshInstanceBuffer::BindInstanceVertexBuffer(const class FVertexFactory* VertexFactory, FInstancedStaticMeshDataType& InstancedStaticMeshData) const
 {
-	check(InstanceOriginSRV);
-	check(InstanceTransformSRV);
-	check(InstanceLightmapSRV);
+	if (InstanceData->GetNumInstances())
+	{
+		check(InstanceOriginSRV);
+		check(InstanceTransformSRV);
+		check(InstanceLightmapSRV);
+	}
 
 	{
 		InstancedStaticMeshData.InstanceOriginSRV = InstanceOriginSRV;
@@ -1928,6 +1931,7 @@ void UInstancedStaticMeshComponent::SerializeRenderData(FArchive& Ar)
 					FStaticMeshInstanceData RenderInstanceData = FStaticMeshInstanceData(GVertexElementTypeSupport.IsSupported(VET_Half2));
 					BuildRenderData(RenderInstanceData, PerInstanceRenderData->HitProxies);
 					PerInstanceRenderData->UpdateFromPreallocatedData(RenderInstanceData);
+					MarkRenderStateDirty();
 				}
 			}
 		
