Since the child classes now mark their components as VisibleAnywhere, the base class should not or else you get duplicate property panel entries

diff --git a/Engine/Source/Runtime/Engine/Classes/Engine/Light.h b/Engine/Source/Runtime/Engine/Classes/Engine/Light.h
--- a/Engine/Source/Runtime/Engine/Classes/Engine/Light.h
+++ b/Engine/Source/Runtime/Engine/Classes/Engine/Light.h
@@ -16,7 +16,7 @@ class ENGINE_API ALight : public AActor
 
 private:
 	/** @todo document */
-	UPROPERTY(Category = Light, VisibleAnywhere, BlueprintReadOnly, meta = (ExposeFunctionCategories = "Light,Rendering,Rendering|Components|Light", AllowPrivateAccess = "true"))
+	UPROPERTY(Category = Light, BlueprintReadOnly, meta = (ExposeFunctionCategories = "Light,Rendering,Rendering|Components|Light", AllowPrivateAccess = "true"))
 	class ULightComponent* LightComponent;
 public:
 
