Fixed: support instanced static meshes for software occlusion cullling

diff --git a/Engine/Source/Runtime/Engine/Private/InstancedStaticMesh.cpp b/Engine/Source/Runtime/Engine/Private/InstancedStaticMesh.cpp
--- a/Engine/Source/Runtime/Engine/Private/InstancedStaticMesh.cpp
+++ b/Engine/Source/Runtime/Engine/Private/InstancedStaticMesh.cpp
@@ -833,6 +833,28 @@ int32 FInstancedStaticMeshSceneProxy::GetNumMeshBatches() const
 	}
 }
 
+int32 FInstancedStaticMeshSceneProxy::CollectOccluderElements(FOccluderElementsCollector& Collector) const
+{
+	if (OccluderData)
+	{	
+		FStaticMeshInstanceBuffer& InstanceBuffer = InstancedRenderData.PerInstanceRenderData->InstanceBuffer;
+		const int32 NumInstances = InstanceBuffer.GetNumInstances();
+		
+		for (int32 InstanceIndex = 0; InstanceIndex < NumInstances; ++InstanceIndex)
+		{
+			FMatrix InstanceToLocal;
+			InstanceBuffer.GetInstanceTransform(InstanceIndex, InstanceToLocal);	
+			InstanceToLocal.M[3][3] = 1.0f;
+						
+			Collector.AddElements(OccluderData->VerticesSP, OccluderData->IndicesSP, InstanceToLocal * GetLocalToWorld());
+		}
+		
+		return NumInstances;
+	}
+	
+	return 0;
+}
+
 void FInstancedStaticMeshSceneProxy::SetupProxy(UInstancedStaticMeshComponent* InComponent)
 {
 #if WITH_EDITOR
diff --git a/Engine/Source/Runtime/Engine/Private/InstancedStaticMesh.h b/Engine/Source/Runtime/Engine/Private/InstancedStaticMesh.h
--- a/Engine/Source/Runtime/Engine/Private/InstancedStaticMesh.h
+++ b/Engine/Source/Runtime/Engine/Private/InstancedStaticMesh.h
@@ -568,6 +568,8 @@ public:
 
 	virtual void GetDistanceFieldInstanceInfo(int32& NumInstances, float& BoundsSurfaceArea) const override;
 
+	virtual int32 CollectOccluderElements(FOccluderElementsCollector& Collector) const override;
+
 	/**
 	 * Creates the hit proxies are used when DrawDynamicElements is called.
 	 * Called in the game thread.
diff --git a/Engine/Source/Runtime/Engine/Private/StaticMeshRender.cpp b/Engine/Source/Runtime/Engine/Private/StaticMeshRender.cpp
--- a/Engine/Source/Runtime/Engine/Private/StaticMeshRender.cpp
+++ b/Engine/Source/Runtime/Engine/Private/StaticMeshRender.cpp
@@ -510,15 +510,15 @@ bool FStaticMeshSceneProxy::GetMeshElement(
 	}
 }
 
-bool FStaticMeshSceneProxy::CollectOccluderElements(FOccluderElementsCollector& Collector) const
+int32 FStaticMeshSceneProxy::CollectOccluderElements(FOccluderElementsCollector& Collector) const
 {
 	if (OccluderData)
 	{	
 		Collector.AddElements(OccluderData->VerticesSP, OccluderData->IndicesSP, GetLocalToWorld());
-		return true;
+		return 1;
 	}
 	
-	return false;
+	return 0;
 }
 
 /** Sets up a wireframe FMeshBatch for a specific LOD. */
diff --git a/Engine/Source/Runtime/Engine/Public/PrimitiveSceneProxy.h b/Engine/Source/Runtime/Engine/Public/PrimitiveSceneProxy.h
--- a/Engine/Source/Runtime/Engine/Public/PrimitiveSceneProxy.h
+++ b/Engine/Source/Runtime/Engine/Public/PrimitiveSceneProxy.h
@@ -215,7 +215,7 @@ public:
 #endif // RHI_RAYTRACING
 
 	/** Collects occluder geometry for software occlusion culling */
-	virtual bool CollectOccluderElements(class FOccluderElementsCollector& Collector) const { return false; }
+	virtual int32 CollectOccluderElements(class FOccluderElementsCollector& Collector) const { return 0; }
 
 	/** 
 	 * Gathers the primitive's dynamic mesh elements.  This will only be called if GetViewRelevance declares dynamic relevance.
diff --git a/Engine/Source/Runtime/Engine/Public/StaticMeshResources.h b/Engine/Source/Runtime/Engine/Public/StaticMeshResources.h
--- a/Engine/Source/Runtime/Engine/Public/StaticMeshResources.h
+++ b/Engine/Source/Runtime/Engine/Public/StaticMeshResources.h
@@ -764,7 +764,7 @@ public:
 		bool bAllowPreCulledIndices,
 		FMeshBatch& OutMeshBatch) const;
 
-	virtual bool CollectOccluderElements(class FOccluderElementsCollector& Collector) const override;
+	virtual int32 CollectOccluderElements(class FOccluderElementsCollector& Collector) const override;
 
 	/** Sets up a wireframe FMeshBatch for a specific LOD. */
 	virtual bool GetWireframeMeshElement(int32 LODIndex, int32 BatchIndex, const FMaterialRenderProxy* WireframeRenderProxy, uint8 InDepthPriorityGroup, bool bAllowPreCulledIndices, FMeshBatch& OutMeshBatch) const;
diff --git a/Engine/Source/Runtime/Landscape/Private/LandscapeRender.cpp b/Engine/Source/Runtime/Landscape/Private/LandscapeRender.cpp
--- a/Engine/Source/Runtime/Landscape/Private/LandscapeRender.cpp
+++ b/Engine/Source/Runtime/Landscape/Private/LandscapeRender.cpp
@@ -3256,10 +3256,10 @@ void FLandscapeComponentSceneProxy::GetDynamicRayTracingInstances(FRayTracingMat
 }
 #endif
 
-bool FLandscapeComponentSceneProxy::CollectOccluderElements(FOccluderElementsCollector& Collector) const
+int32 FLandscapeComponentSceneProxy::CollectOccluderElements(FOccluderElementsCollector& Collector) const
 {
 	// TODO: implement
-	return false;
+	return 0;
 }
 
 //
diff --git a/Engine/Source/Runtime/Landscape/Private/LandscapeRenderMobile.cpp b/Engine/Source/Runtime/Landscape/Private/LandscapeRenderMobile.cpp
--- a/Engine/Source/Runtime/Landscape/Private/LandscapeRenderMobile.cpp
+++ b/Engine/Source/Runtime/Landscape/Private/LandscapeRenderMobile.cpp
@@ -351,15 +351,15 @@ FLandscapeComponentSceneProxyMobile::FLandscapeComponentSceneProxyMobile(ULandsc
 #endif
 }
 
-bool FLandscapeComponentSceneProxyMobile::CollectOccluderElements(FOccluderElementsCollector& Collector) const
+int32 FLandscapeComponentSceneProxyMobile::CollectOccluderElements(FOccluderElementsCollector& Collector) const
 {
 	if (MobileRenderData->OccluderVerticesSP.IsValid() && SharedBuffers->OccluderIndicesSP.IsValid())
 	{
 		Collector.AddElements(MobileRenderData->OccluderVerticesSP, SharedBuffers->OccluderIndicesSP, GetLocalToWorld());
-		return true;
+		return 1;
 	}
 
-	return false;
+	return 0;
 }
 
 FLandscapeComponentSceneProxyMobile::~FLandscapeComponentSceneProxyMobile()
diff --git a/Engine/Source/Runtime/Landscape/Public/LandscapeRender.h b/Engine/Source/Runtime/Landscape/Public/LandscapeRender.h
--- a/Engine/Source/Runtime/Landscape/Public/LandscapeRender.h
+++ b/Engine/Source/Runtime/Landscape/Public/LandscapeRender.h
@@ -730,7 +730,7 @@ public:
 	virtual void ApplyWorldOffset(FVector InOffset) override;
 	virtual void DrawStaticElements(FStaticPrimitiveDrawInterface* PDI) override;
 	virtual void GetDynamicMeshElements(const TArray<const FSceneView*>& Views, const FSceneViewFamily& ViewFamily, uint32 VisibilityMap, FMeshElementCollector& Collector) const override;
-	virtual bool CollectOccluderElements(FOccluderElementsCollector& Collector) const override;
+	virtual int32 CollectOccluderElements(FOccluderElementsCollector& Collector) const override;
 	virtual uint32 GetMemoryFootprint() const override { return(sizeof(*this) + GetAllocatedSize()); }
 	virtual FPrimitiveViewRelevance GetViewRelevance(const FSceneView* View) const override;
 	virtual bool CanBeOccluded() const override;
diff --git a/Engine/Source/Runtime/Landscape/Public/LandscapeRenderMobile.h b/Engine/Source/Runtime/Landscape/Public/LandscapeRenderMobile.h
--- a/Engine/Source/Runtime/Landscape/Public/LandscapeRenderMobile.h
+++ b/Engine/Source/Runtime/Landscape/Public/LandscapeRenderMobile.h
@@ -128,7 +128,7 @@ public:
 	FLandscapeComponentSceneProxyMobile(ULandscapeComponent* InComponent);
 
 	virtual void CreateRenderThreadResources() override;
-	virtual bool CollectOccluderElements(FOccluderElementsCollector& Collector) const override;
+	virtual int32 CollectOccluderElements(FOccluderElementsCollector& Collector) const override;
 
 	uint8 BlendableLayerMask;
 
diff --git a/Engine/Source/Runtime/Renderer/Private/SceneSoftwareOcclusion.cpp b/Engine/Source/Runtime/Renderer/Private/SceneSoftwareOcclusion.cpp
--- a/Engine/Source/Runtime/Renderer/Private/SceneSoftwareOcclusion.cpp
+++ b/Engine/Source/Runtime/Renderer/Private/SceneSoftwareOcclusion.cpp
@@ -1048,10 +1048,7 @@ static FGraphEventRef SubmitScene(const FScene* Scene, FViewInfo& View, FOcclusi
 			{
 				Collector.SetPrimitiveID(PrimitiveComponentId);
 				// Collect occluder geometry
-				if (Proxy->CollectOccluderElements(Collector))
-				{
-					NumCollectedOccluders++;
-				}
+				NumCollectedOccluders+= Proxy->CollectOccluderElements(Collector);
 			}
 
 			if (NumCollectedOccluders >= GSOMaxOccluderNum)
