Fix a bug when DiscardUnusedQualityLevels.

diff --git a/Engine/Source/Runtime/Engine/Private/Materials/Material.cpp b/Engine/Source/Runtime/Engine/Private/Materials/Material.cpp
--- a/Engine/Source/Runtime/Engine/Private/Materials/Material.cpp
+++ b/Engine/Source/Runtime/Engine/Private/Materials/Material.cpp
@@ -706,7 +706,7 @@ void ProcessSerializedInlineShaderMaps(UMaterialInterface* Owner, TArray<FMateri
 		FMaterialResource* DesiredQLMaterialResource[ERHIFeatureLevel::Num] = { 0 };
 
 		// Prioritize generic or common qualities if desired is unavailable
-		int32 BestQL[EMaterialQualityLevel::Num] = {0};
+		int32 BestQLPriority = 0;
 		int32 QLPriorities[EMaterialQualityLevel::Num+1] = {1,3,2,4}; // Num > High > Med > Low
 
 		for (FMaterialResource& Resource : LoadedResources)
@@ -717,14 +717,13 @@ void ProcessSerializedInlineShaderMaps(UMaterialInterface* Owner, TArray<FMateri
 				const EMaterialQualityLevel::Type LoadedQualityLevel = LoadedShaderMap->GetShaderMapId().QualityLevel;
 				const ERHIFeatureLevel::Type LoadedFeatureLevel = LoadedShaderMap->GetShaderMapId().FeatureLevel;
 				const int32 LoadedQLPriority = QLPriorities[LoadedQualityLevel];
-				int32& BestQLRef = BestQL[LoadedQualityLevel];
 
-				if ((LoadedQualityLevel == DesiredQL)								// We always choose an exact match
-				|| (BestQLRef != DesiredQL && LoadedQLPriority > BestQLRef)			// Favor higher quality if still searching
-				|| (DesiredQLMaterialResource[LoadedFeatureLevel] == nullptr))		// If there's nothing, take anything valid
+				if ((LoadedQualityLevel == DesiredQL)												// We always choose an exact match
+				|| (BestQLPriority != QLPriorities[DesiredQL] && LoadedQLPriority > BestQLPriority)	// Favor higher quality if still searching
+				|| (DesiredQLMaterialResource[LoadedFeatureLevel] == nullptr))						// If there's nothing, take anything valid
 				{	
 					DesiredQLMaterialResource[LoadedFeatureLevel] = &Resource;
-					BestQLRef = LoadedQLPriority;
+					BestQLPriority = LoadedQLPriority;
 				}
 			}
 		}
