AltSkinWeights - Text referencing "LOD Index" does not contain a space
#fix Add meta data to alter property display name

diff --git a/Engine/Source/Runtime/Engine/Public/Animation/SkinWeightProfile.h b/Engine/Source/Runtime/Engine/Public/Animation/SkinWeightProfile.h
--- a/Engine/Source/Runtime/Engine/Public/Animation/SkinWeightProfile.h
+++ b/Engine/Source/Runtime/Engine/Public/Animation/SkinWeightProfile.h
@@ -38,7 +38,7 @@ struct FSkinWeightProfileInfo
 	FPerPlatformBool DefaultProfile;
 
 	/** When DefaultProfile is set any LOD below this LOD Index will override the Skin Weights of the Skeletal Mesh with the Skin Weights from this Profile */
-	UPROPERTY(EditAnywhere, Category = SkinWeights, meta=(EditCondition="DefaultProfile", ClampMin=0))
+	UPROPERTY(EditAnywhere, Category = SkinWeights, meta=(EditCondition="DefaultProfile", ClampMin=0, DisplayName = "Default Profile from LOD Index"))
 	FPerPlatformInt DefaultProfileFromLODIndex;
 	
 #if WITH_EDITORONLY_DATA
