When low on memory AND there are levels pending unload speed up GC interval to 10 seconds.  Mainly used when s.ForceGCAfterLevelStreamedOut and s.ContinuouslyIncrementalGCWhileLevelsPendingPurge are disabled.

diff --git a/Engine/Source/Runtime/Engine/Private/UnrealEngine.cpp b/Engine/Source/Runtime/Engine/Private/UnrealEngine.cpp
--- a/Engine/Source/Runtime/Engine/Private/UnrealEngine.cpp
+++ b/Engine/Source/Runtime/Engine/Private/UnrealEngine.cpp
@@ -1153,6 +1153,15 @@ static FAutoConsoleVariableRef CVarLowMemoryTimeBetweenPurgingPendingKillObjects
 	ECVF_Default
 );
 
+// see also: s.ForceGCAfterLevelStreamedOut, s.ContinuouslyIncrementalGCWhileLevelsPendingPurge
+static float GLowMemoryTimeBetweenPurgingPendingLevels = 10.0f;
+static FAutoConsoleVariableRef CVarLowMemoryTimeBetweenPurgingPendingLevels(
+	TEXT("gc.LowMemory.TimeBetweenPurgingPendingLevels"),
+	GLowMemoryTimeBetweenPurgingPendingLevels,
+	TEXT("Time in seconds (game time) we should wait between GC when we're low on memory and there are levels pending unload"),
+	ECVF_Default
+);
+
 static float GLowMemoryMemoryThresholdMB = 0.0f;
 static FAutoConsoleVariableRef CVarLowMemoryThresholdMB(
 	TEXT("gc.LowMemory.MemoryThresholdMB"),
@@ -1198,12 +1207,21 @@ float UEngine::GetTimeBetweenGarbageCollectionPasses() const
 	}
 
 	// Do more frequent GCs if we're under the low memory threshold (if enabled)
-	if ( GLowMemoryMemoryThresholdMB > 0.0 && GLowMemoryTimeBetweenPurgingPendingKillObjects < TimeBetweenGC )
+	if ( GLowMemoryMemoryThresholdMB > 0.0 )
 	{
 		float MBFree = float(FPlatformMemory::GetStats().AvailablePhysical / 1024 / 1024);
 		if ( MBFree <= GLowMemoryMemoryThresholdMB)
 		{
-			TimeBetweenGC = GLowMemoryTimeBetweenPurgingPendingKillObjects;
+			if (GLowMemoryTimeBetweenPurgingPendingKillObjects < TimeBetweenGC)
+			{
+				TimeBetweenGC = GLowMemoryTimeBetweenPurgingPendingKillObjects;
+			}
+
+			// Go even faster if there are levels that need to be unloaded
+			if (GLowMemoryTimeBetweenPurgingPendingLevels < TimeBetweenGC && FLevelStreamingGCHelper::GetNumLevelsPendingPurge() > 0)
+			{
+				TimeBetweenGC = GLowMemoryTimeBetweenPurgingPendingLevels;
+			}
 		}
 	}
 
