Fixed crash in FSceneViewport::OnFocusReceived(...) when ViewportClient is nullptr.

#ReleaseNote above^^

#UE4 #ReleaseNote

diff --git a/Engine/Source/Runtime/Engine/Private/Slate/SceneViewport.cpp b/Engine/Source/Runtime/Engine/Private/Slate/SceneViewport.cpp
--- a/Engine/Source/Runtime/Engine/Private/Slate/SceneViewport.cpp
+++ b/Engine/Source/Runtime/Engine/Private/Slate/SceneViewport.cpp
@@ -1079,11 +1079,14 @@ FReply FSceneViewport::OnFocusReceived(const FFocusEvent& InFocusEvent)
 		{
 			FSlateApplication& SlateApp = FSlateApplication::Get();
 
-			const bool bPermanentCapture = (!GIsEditor || InFocusEvent.GetCause() == EFocusCause::Mouse) && 
-				(( ViewportClient->CaptureMouseOnClick() == EMouseCaptureMode::CapturePermanently ) ||
-				( ViewportClient->CaptureMouseOnClick() == EMouseCaptureMode::CapturePermanently_IncludingInitialMouseDown ));
-
-			if ( SlateApp.IsActive() && !ViewportClient->IgnoreInput() && bPermanentCapture )
+			const bool bPermanentCapture = (!GIsEditor || InFocusEvent.GetCause() == EFocusCause::Mouse)	&&
+										   (ViewportClient != nullptr)										&&
+					(( ViewportClient->CaptureMouseOnClick() == EMouseCaptureMode::CapturePermanently ) ||
+					 ( ViewportClient->CaptureMouseOnClick() == EMouseCaptureMode::CapturePermanently_IncludingInitialMouseDown )
+					);
+
+			// if bPermanentCapture is true, then ViewportClient != nullptr, so it's ok to dereference it.  But the permanent capture must be tested first.
+			if ( SlateApp.IsActive() && bPermanentCapture && !ViewportClient->IgnoreInput() )
 			{
 				TSharedRef<SViewport> ViewportWidgetRef = ViewportWidget.Pin().ToSharedRef();
 
