Add support for predefined network profiles

* You can now configure network emulation profiles in Engine.ini using [PacketSimulationProfile.MyProfileName]
* Profiles can be applied via console or cmdline using "Net EmulationProfile=MyProfileName"
* Added PktLagMin/PktLagMax setting:  use this to have a random latency between Min and Max
* No longer capping PktLagVariance between 0 and 100

[RN] minor engine/network

diff --git a/Engine/Source/Runtime/Engine/Classes/Engine/NetDriver.h b/Engine/Source/Runtime/Engine/Classes/Engine/NetDriver.h
--- a/Engine/Source/Runtime/Engine/Classes/Engine/NetDriver.h
+++ b/Engine/Source/Runtime/Engine/Classes/Engine/NetDriver.h
@@ -437,13 +437,22 @@ struct ENGINE_API FPacketSimulationSettings
 	/**
 	 * When set, will cause PktLag to use variable lag instead of constant.
 	 * Value is treated as millisecond lag range (e.g. -GivenVariance <= 0 <= GivenVariance).
-	 * Clamped between 0 and 100.
 	 *
 	 * Can only be used when PktLag is enabled.
 	 */
 	UPROPERTY(EditAnywhere, Category="Simulation Settings")
 	int32	PktLagVariance;
 
+	/**
+	 * If set lag values will randomly fluctuate between Min and Max.
+	 * Ignored if PktLag value is set
+	 */
+	UPROPERTY(EditAnywhere, Category = "Simulation Settings")
+	int32	PktLagMin = 0;
+	
+	UPROPERTY(EditAnywhere, Category = "Simulation Settings")
+	int32	PktLagMax = 0;
+
 	/** Ctor. Zeroes the settings */
 	FPacketSimulationSettings() : 
 		PktLoss(0),
@@ -460,6 +469,11 @@ struct ENGINE_API FPacketSimulationSettings
 	 * @note: overwrites all previous settings
 	 */
 	void LoadConfig(const TCHAR* OptionalQualifier = nullptr);
+	
+	/** 
+	 * Load a preconfigured emulation profile from the .ini
+	 */
+	void LoadEmulationProfile(const TCHAR* ProfileName);
 
 	/**
 	 * Registers commands for auto-completion, etc.
@@ -471,6 +485,12 @@ struct ENGINE_API FPacketSimulationSettings
 	 */
 	void UnregisterCommands();
 
+	/**
+	 * Ensure that settings have proper values
+	 */
+	void ValidateSettings();
+	void ResetSettings();
+
 	/**
 	 * Reads the settings from a string: command line or an exec
 	 *
diff --git a/Engine/Source/Runtime/Engine/Private/NetConnection.cpp b/Engine/Source/Runtime/Engine/Private/NetConnection.cpp
--- a/Engine/Source/Runtime/Engine/Private/NetConnection.cpp
+++ b/Engine/Source/Runtime/Engine/Private/NetConnection.cpp
@@ -1326,7 +1326,22 @@ void UNetConnection::FlushNet(bool bIgnoreSimulation)
 			{
 				DelayedPacket& B = *(new(Delayed)DelayedPacket(SendBuffer.GetData(), SendBuffer.GetNumBits(), Traits));
 
-				B.SendTime = FPlatformTime::Seconds() + (double(PacketSimulationSettings.PktLag)  + 2.0f * (FMath::FRand() - 0.5f) * double(PacketSimulationSettings.PktLagVariance))/ 1000.f;
+				// ExtraLag goes from PktLag + [-PktLagVariance, PktLagVariance]
+				const double LagVariance = 2.0f * (FMath::FRand() - 0.5f) * double(PacketSimulationSettings.PktLagVariance);
+				const double ExtraLag = (double(PacketSimulationSettings.PktLag) + LagVariance) / 1000.f;
+				B.SendTime = FPlatformTime::Seconds() + ExtraLag;
+			}
+		}
+		else if (PacketSimulationSettings.PktLagMin > 0 && PacketSimulationSettings.PktLagMax > 0)
+		{
+			if (!ShouldDropOutgoingPacketForLossSimulation(SendBuffer.GetNumBits()))
+			{
+				DelayedPacket& B = *(new(Delayed)DelayedPacket(SendBuffer.GetData(), SendBuffer.GetNumBits(), Traits));
+
+				// ExtraLag goes from [PktLagMin, PktLagMax]
+				const double LagVariance = FMath::FRand() * double(PacketSimulationSettings.PktLagMax - PacketSimulationSettings.PktLagMin);
+				const double ExtraLag = (double(PacketSimulationSettings.PktLagMin) + LagVariance) / 1000.f;
+				B.SendTime = FPlatformTime::Seconds() + ExtraLag;
 			}
 		}
 		else if (!ShouldDropOutgoingPacketForLossSimulation(SendBuffer.GetNumBits()))
@@ -2687,7 +2702,7 @@ void UNetConnection::Tick()
 
 	// Lag simulation.
 #if DO_ENABLE_NET_TEST
-	if( PacketSimulationSettings.PktLag )
+	if(Delayed.Num() > 0)
 	{
 		for( int32 i=0; i < Delayed.Num(); i++ )
 		{
@@ -2699,7 +2714,7 @@ void UNetConnection::Tick()
 			}
 			else
 			{
-				// Break now instead of continuing to iterate through the list. Otherwise LagVariance may cause out of order sends
+				// Break now instead of continuing to iterate through the list. Otherwise may cause out of order sends
 				break;
 			}
 		}
@@ -3168,13 +3183,7 @@ void UChildConnection::HandleClientPlayer(APlayerController* PC, UNetConnection*
 void UNetConnection::UpdatePacketSimulationSettings(void)
 {
 	check(Driver);
-	PacketSimulationSettings.PktLoss = Driver->PacketSimulationSettings.PktLoss;
-	PacketSimulationSettings.PktLossMinSize = Driver->PacketSimulationSettings.PktLossMinSize;
-	PacketSimulationSettings.PktLossMaxSize = Driver->PacketSimulationSettings.PktLossMaxSize;
-	PacketSimulationSettings.PktOrder = Driver->PacketSimulationSettings.PktOrder;
-	PacketSimulationSettings.PktDup = Driver->PacketSimulationSettings.PktDup;
-	PacketSimulationSettings.PktLag = Driver->PacketSimulationSettings.PktLag;
-	PacketSimulationSettings.PktLagVariance = Driver->PacketSimulationSettings.PktLagVariance;
+	PacketSimulationSettings = Driver->PacketSimulationSettings;
 }
 #endif
 
diff --git a/Engine/Source/Runtime/Engine/Private/NetworkDriver.cpp b/Engine/Source/Runtime/Engine/Private/NetworkDriver.cpp
--- a/Engine/Source/Runtime/Engine/Private/NetworkDriver.cpp
+++ b/Engine/Source/Runtime/Engine/Private/NetworkDriver.cpp
@@ -3248,10 +3248,7 @@ public:
  */
 void FPacketSimulationSettings::LoadConfig(const TCHAR* OptionalQualifier)
 {
-	if (ConfigHelperInt(TEXT("PktLoss"), PktLoss, OptionalQualifier))
-	{
-		PktLoss = FMath::Clamp<int32>(PktLoss, 0, 100);
-	}
+	ConfigHelperInt(TEXT("PktLoss"), PktLoss, OptionalQualifier);
 	
 	ConfigHelperInt(TEXT("PktLossMinSize"), PktLossMinSize, OptionalQualifier);
 	ConfigHelperInt(TEXT("PktLossMaxSize"), PktLossMaxSize, OptionalQualifier);
@@ -3261,16 +3258,72 @@ void FPacketSimulationSettings::LoadConfig(const TCHAR* OptionalQualifier)
 	PktOrder = int32(InPktOrder);
 	
 	ConfigHelperInt(TEXT("PktLag"), PktLag, OptionalQualifier);
+	ConfigHelperInt(TEXT("PktLagVariance"), PktLagVariance, OptionalQualifier);
+
+	ConfigHelperInt(TEXT("PktLagMin"), PktLagMin, OptionalQualifier);
+	ConfigHelperInt(TEXT("PktLagMax"), PktLagMax, OptionalQualifier);
 	
-	if (ConfigHelperInt(TEXT("PktDup"), PktDup, OptionalQualifier))
+	ConfigHelperInt(TEXT("PktDup"), PktDup, OptionalQualifier);
+
+	ValidateSettings();
+}
+
+void FPacketSimulationSettings::LoadEmulationProfile(const TCHAR* ProfileName)
+{
+	const FString SectionName = FString::Printf(TEXT("%s.%s"), TEXT("PacketSimulationProfile"), ProfileName);
+
+	TArray<FString> SectionConfigs;
+	bool bSectionExists = GConfig->GetSection(*SectionName, SectionConfigs, GEngineIni);
+	if (!bSectionExists)
 	{
-		PktDup = FMath::Clamp<int32>(PktDup, 0, 100);
+		UE_LOG(LogNet, Warning, TEXT("EmulationProfile [%s] was not found in %s. Packet settings were not changed"), *SectionName, *GEngineIni);
+		return;
 	}
-	
-	if (ConfigHelperInt(TEXT("PktLagVariance"), PktLagVariance, OptionalQualifier))
+
+	ResetSettings();
+
+	UScriptStruct* ThisStruct = FPacketSimulationSettings::StaticStruct();
+
+	for (const FString& ConfigVar : SectionConfigs)
 	{
-		PktLagVariance = FMath::Clamp<int32>(PktLagVariance, 0, 100);
+		FString VarName;
+		FString VarValue;
+		if (ConfigVar.Split(TEXT("="), &VarName, &VarValue))
+		{
+			// If using the one line struct definition
+			if (VarName.Equals(TEXT("PacketSimulationSettings"), ESearchCase::IgnoreCase))
+			{
+				ThisStruct->ImportText(*VarValue, this, nullptr, 0, (FOutputDevice*)GWarn, TEXT("FPacketSimulationSettings"));
+			}
+			else if (UProperty* StructProperty = ThisStruct->FindPropertyByName(FName(*VarName, FNAME_Find)))
+			{
+				StructProperty->ImportText(*VarValue, StructProperty->ContainerPtrToValuePtr<void>(this), 0, nullptr);
+			}
+			else
+			{
+				UE_LOG(LogNet, Warning, TEXT("FPacketSimulationSettings::LoadEmulationProfile could not find property named %s"), *VarName);
+			}
+		}
 	}
+	
+	ValidateSettings();
+}
+
+void FPacketSimulationSettings::ResetSettings()
+{
+	*this = FPacketSimulationSettings();
+}
+
+void FPacketSimulationSettings::ValidateSettings()
+{
+	PktLoss = FMath::Clamp<int32>(PktLoss, 0, 100);
+
+	PktOrder = FMath::Clamp<int32>(PktOrder, 0, 1);
+
+	PktLagMin = FMath::Max(PktLagMin, 0);
+	PktLagMax = FMath::Max(PktLagMin, PktLagMax);
+
+	PktDup = FMath::Clamp<int32>(PktDup, 0, 100);
 }
 
 bool FPacketSimulationSettings::ConfigHelperInt(const TCHAR* Name, int32& Value, const TCHAR* OptionalQualifier)
@@ -3321,6 +3374,9 @@ void FPacketSimulationSettings::RegisterCommands()
 		ConsoleManager.RegisterConsoleCommand(TEXT("Net PktDup="), TEXT("PktDup=<n> (simulates sending/receiving duplicate network packets)"));
 		ConsoleManager.RegisterConsoleCommand(TEXT("Net PktLag="), TEXT("PktLag=<n> (simulates network packet lag)"));
 		ConsoleManager.RegisterConsoleCommand(TEXT("Net PktLagVariance="), TEXT("PktLagVariance=<n> (simulates variable network packet lag)"));
+		ConsoleManager.RegisterConsoleCommand(TEXT("Net PktLagMin="), TEXT("PktLagMin=<n> (minimum packet latency)"));
+		ConsoleManager.RegisterConsoleCommand(TEXT("Net PktLagMax="), TEXT("PktLagMax=<n> (maximum packet latency)"));
+		ConsoleManager.RegisterConsoleCommand(TEXT("Net EmulationProfile="), TEXT("EmulationProfile=<NAME> (apply an emulation profile)"));
 	}
 }
 
@@ -3359,10 +3415,16 @@ bool FPacketSimulationSettings::ParseSettings(const TCHAR* Cmd, const TCHAR* Opt
 	// this is because the same function will be used to parse the command line as well
 	bool bParsed = false;
 
+	FString EmulationProfileName;
+	if (FParse::Value(Cmd, TEXT("EmulationProfile="), EmulationProfileName))
+	{
+		bParsed = true;
+		UE_LOG(LogNet, Log, TEXT("Applying EmulationProfile %s"), *EmulationProfileName);
+		LoadEmulationProfile(*EmulationProfileName);
+	}
 	if( ParseHelper(Cmd, TEXT("PktLoss="), PktLoss, OptionalQualifier) )
 	{
 		bParsed = true;
-		PktLoss = FMath::Clamp<int32>( PktLoss, 0, 100 );
 		UE_LOG(LogNet, Log, TEXT("PktLoss set to %d"), PktLoss);
 	}
 	if (ParseHelper(Cmd, TEXT("PktLossMinSize="), PktLossMinSize, OptionalQualifier))
@@ -3378,7 +3440,6 @@ bool FPacketSimulationSettings::ParseSettings(const TCHAR* Cmd, const TCHAR* Opt
 	if( ParseHelper(Cmd, TEXT("PktOrder="), PktOrder, OptionalQualifier) )
 	{
 		bParsed = true;
-		PktOrder = FMath::Clamp<int32>( PktOrder, 0, 1 );
 		UE_LOG(LogNet, Log, TEXT("PktOrder set to %d"), PktOrder);
 	}
 	if( ParseHelper(Cmd, TEXT("PktLag="), PktLag, OptionalQualifier) )
@@ -3389,15 +3450,25 @@ bool FPacketSimulationSettings::ParseSettings(const TCHAR* Cmd, const TCHAR* Opt
 	if( ParseHelper(Cmd, TEXT("PktDup="), PktDup, OptionalQualifier) )
 	{
 		bParsed = true;
-		PktDup = FMath::Clamp<int32>( PktDup, 0, 100 );
 		UE_LOG(LogNet, Log, TEXT("PktDup set to %d"), PktDup);
 	}	
 	if (ParseHelper(Cmd, TEXT("PktLagVariance="), PktLagVariance, OptionalQualifier))
 	{
 		bParsed = true;
-		PktLagVariance = FMath::Clamp<int32>( PktLagVariance, 0, 100 );
 		UE_LOG(LogNet, Log, TEXT("PktLagVariance set to %d"), PktLagVariance);
 	}
+	if (ParseHelper(Cmd, TEXT("PktLagMin="), PktLagMin, OptionalQualifier))
+	{
+		bParsed = true;
+		UE_LOG(LogNet, Log, TEXT("PktLagMin set to %d"), PktLagMin);
+	}
+	if (ParseHelper(Cmd, TEXT("PktLagMax="), PktLagMax, OptionalQualifier))
+	{
+		bParsed = true;
+		UE_LOG(LogNet, Log, TEXT("PktLagMax set to %d"), PktLagMax);
+	}
+
+	ValidateSettings();
 	return bParsed;
 }
 
