Fixed parsing errors in RendererSettings

diff --git a/Engine/Source/Runtime/Engine/Classes/Engine/RendererSettings.h b/Engine/Source/Runtime/Engine/Classes/Engine/RendererSettings.h
--- a/Engine/Source/Runtime/Engine/Classes/Engine/RendererSettings.h
+++ b/Engine/Source/Runtime/Engine/Classes/Engine/RendererSettings.h
@@ -494,7 +494,7 @@ class ENGINE_API URendererSettings : public UDeveloperSettings
 	TEnumAsByte<EEarlyZPass::Type> EarlyZPass;
 
 	UPROPERTY(config, EditAnywhere, Category = Optimizations, meta = (
-		EditCondition = "EarlyZPass == OpaqueAndMasked && bEarlyZPassMovable",
+		EditCondition = "EarlyZPass == EEarlyZPass::OpaqueAndMasked",
 		ConsoleVariable = "r.EarlyZPassOnlyMaterialMasking", DisplayName = "Mask material only in early Z-pass",
 		ToolTip = "Whether to compute materials' mask opacity only in early Z pass. Changing this setting requires restarting the editor.",
 		ConfigRestartRequired = true))
