Fix UE-75802: Instances Created From Merge Actor Have Corrupt Lightmaps

* Was caused by race condition in UInstancedStaticMeshComponent::BuildRenderData()
* (Hierarchical) instanced static meshes now use GetMeshBuildDataDuringBuild, as they only want lightmap/shadowmap UV scale and bias rather than real data

diff --git a/Engine/Source/Runtime/Engine/Classes/Components/StaticMeshComponent.h b/Engine/Source/Runtime/Engine/Classes/Components/StaticMeshComponent.h
--- a/Engine/Source/Runtime/Engine/Classes/Components/StaticMeshComponent.h
+++ b/Engine/Source/Runtime/Engine/Classes/Components/StaticMeshComponent.h
@@ -665,7 +665,7 @@ public:
 
 	virtual void PropagateLightingScenarioChange() override;
 
-	const FMeshMapBuildData* GetMeshMapBuildData(const FStaticMeshComponentLODInfo& LODInfo) const;
+	const FMeshMapBuildData* GetMeshMapBuildData(const FStaticMeshComponentLODInfo& LODInfo, bool bCheckForResourceCluster = true) const;
 
 
 #if WITH_EDITOR
diff --git a/Engine/Source/Runtime/Engine/Private/Components/StaticMeshComponent.cpp b/Engine/Source/Runtime/Engine/Private/Components/StaticMeshComponent.cpp
--- a/Engine/Source/Runtime/Engine/Private/Components/StaticMeshComponent.cpp
+++ b/Engine/Source/Runtime/Engine/Private/Components/StaticMeshComponent.cpp
@@ -499,7 +499,7 @@ void UStaticMeshComponent::PropagateLightingScenarioChange()
 	FComponentRecreateRenderStateContext Context(this);
 }
 
-const FMeshMapBuildData* UStaticMeshComponent::GetMeshMapBuildData(const FStaticMeshComponentLODInfo& LODInfo) const
+const FMeshMapBuildData* UStaticMeshComponent::GetMeshMapBuildData(const FStaticMeshComponentLODInfo& LODInfo, bool bCheckForResourceCluster) const
 {
 	if (!GetStaticMesh() || !GetStaticMesh()->RenderData)
 	{
@@ -544,7 +544,7 @@ const FMeshMapBuildData* UStaticMeshComponent::GetMeshMapBuildData(const FStatic
 
 			if (MapBuildData)
 			{
-				return MapBuildData->GetMeshBuildData(LODInfo.MapBuildDataId);
+				return bCheckForResourceCluster ? MapBuildData->GetMeshBuildData(LODInfo.MapBuildDataId) : MapBuildData->GetMeshBuildDataDuringBuild(LODInfo.MapBuildDataId);
 			}
 		}
 	}
diff --git a/Engine/Source/Runtime/Engine/Private/HierarchicalInstancedStaticMesh.cpp b/Engine/Source/Runtime/Engine/Private/HierarchicalInstancedStaticMesh.cpp
--- a/Engine/Source/Runtime/Engine/Private/HierarchicalInstancedStaticMesh.cpp
+++ b/Engine/Source/Runtime/Engine/Private/HierarchicalInstancedStaticMesh.cpp
@@ -2888,7 +2888,7 @@ void UHierarchicalInstancedStaticMeshComponent::PropagateLightingScenarioChange(
 			const FMeshMapBuildData* MeshMapBuildData = nullptr;
 			if (LODData.Num() > 0)
 			{
-				MeshMapBuildData = GetMeshMapBuildData(LODData[0]);
+				MeshMapBuildData = GetMeshMapBuildData(LODData[0], false);
 			}
 
 			if (MeshMapBuildData != nullptr)
@@ -2918,7 +2918,7 @@ void UHierarchicalInstancedStaticMeshComponent::SetPerInstanceLightMapAndEditorD
 	const FMeshMapBuildData* MeshMapBuildData = nullptr;
 	if (LODData.Num() > 0)
 	{
-		MeshMapBuildData = GetMeshMapBuildData(LODData[0]);
+		MeshMapBuildData = GetMeshMapBuildData(LODData[0], false);
 	}
 
 	if (MeshMapBuildData != nullptr || GIsEditor)
diff --git a/Engine/Source/Runtime/Engine/Private/InstancedStaticMesh.cpp b/Engine/Source/Runtime/Engine/Private/InstancedStaticMesh.cpp
--- a/Engine/Source/Runtime/Engine/Private/InstancedStaticMesh.cpp
+++ b/Engine/Source/Runtime/Engine/Private/InstancedStaticMesh.cpp
@@ -1434,7 +1434,7 @@ void UInstancedStaticMeshComponent::BuildRenderData(FStaticMeshInstanceData& Out
 	const FMeshMapBuildData* MeshMapBuildData = nullptr;
 	if (LODData.Num() > 0)
 	{
-		MeshMapBuildData = GetMeshMapBuildData(LODData[0]);
+		MeshMapBuildData = GetMeshMapBuildData(LODData[0], false);
 	}
 	
 	check(InstancingRandomSeed != 0);
