PR #6021: Fixed USplineComponent::GetDirectionAtSplineInputKey result for scaled objects (Contributed by [REDACTED_PERSON])

diff --git a/Engine/Source/Runtime/Engine/Private/Components/SplineComponent.cpp b/Engine/Source/Runtime/Engine/Private/Components/SplineComponent.cpp
--- a/Engine/Source/Runtime/Engine/Private/Components/SplineComponent.cpp
+++ b/Engine/Source/Runtime/Engine/Private/Components/SplineComponent.cpp
@@ -382,7 +382,8 @@ FVector USplineComponent::GetDirectionAtSplineInputKey(float InKey, ESplineCoord
 
 	if (CoordinateSpace == ESplineCoordinateSpace::World)
 	{
-		Direction = GetComponentTransform().TransformVectorNoScale(Direction);
+		Direction = GetComponentTransform().TransformVector(Direction);
+		Direction.Normalize();
 	}
 
 	return Direction;
