PR #6026: Add const qualifier to AActor::ActorLineTraceSingle (Contributed by [REDACTED_PERSON])

diff --git a/Engine/Source/Runtime/Engine/Classes/GameFramework/Actor.h b/Engine/Source/Runtime/Engine/Classes/GameFramework/Actor.h
--- a/Engine/Source/Runtime/Engine/Classes/GameFramework/Actor.h
+++ b/Engine/Source/Runtime/Engine/Classes/GameFramework/Actor.h
@@ -2389,7 +2389,7 @@ public:
 	 * @param  Params          Additional parameters used for the trace
 	 * @return TRUE if a blocking hit is found
 	 */
-	bool ActorLineTraceSingle(struct FHitResult& OutHit, const FVector& Start, const FVector& End, ECollisionChannel TraceChannel, const struct FCollisionQueryParams& Params);
+	bool ActorLineTraceSingle(struct FHitResult& OutHit, const FVector& Start, const FVector& End, ECollisionChannel TraceChannel, const struct FCollisionQueryParams& Params) const;
 
 	/** 
 	 * returns Distance to closest Body Instance surface. 
diff --git a/Engine/Source/Runtime/Engine/Private/Actor.cpp b/Engine/Source/Runtime/Engine/Private/Actor.cpp
--- a/Engine/Source/Runtime/Engine/Private/Actor.cpp
+++ b/Engine/Source/Runtime/Engine/Private/Actor.cpp
@@ -4555,7 +4555,7 @@ void AActor::InvalidateLightingCacheDetailed(bool bTranslationOnly)
 
  // COLLISION
 
-bool AActor::ActorLineTraceSingle(struct FHitResult& OutHit, const FVector& Start, const FVector& End, ECollisionChannel TraceChannel, const struct FCollisionQueryParams& Params)
+bool AActor::ActorLineTraceSingle(struct FHitResult& OutHit, const FVector& Start, const FVector& End, ECollisionChannel TraceChannel, const struct FCollisionQueryParams& Params) const
 {
 	OutHit = FHitResult(1.f);
 	OutHit.TraceStart = Start;
