Restore default aperture to 4 in post process settings, and instead disable DOF by setting default focus distance to 0.

Note: Post process volume saved between 5993259 and this CL will loose their focus distance. But that is OK, because configuring a focus distance in a post process volume where camera came move arround doesn't make any sens anyway.

diff --git a/Engine/Plugins/Runtime/Nvidia/Ansel/Source/Ansel/Private/Ansel.cpp b/Engine/Plugins/Runtime/Nvidia/Ansel/Source/Ansel/Private/Ansel.cpp
--- a/Engine/Plugins/Runtime/Nvidia/Ansel/Source/Ansel/Private/Ansel.cpp
+++ b/Engine/Plugins/Runtime/Nvidia/Ansel/Source/Ansel/Private/Ansel.cpp
@@ -364,7 +364,7 @@ void FNVAnselCameraPhotographyPrivate::DoCustomUIControls(FPostProcessSettings&
 		PRAGMA_DISABLE_DEPRECATION_WARNINGS
 		if (bEffectUIAllowed[DepthOfField])
 		{
-			bool bAnyDofVisible = InOutPPSettings.DepthOfFieldFstop >= 0;
+			const bool bAnyDofVisible = InOutPPSettings.DepthOfFieldFstop > 0 && InOutPPSettings.DepthOfFieldFocalDistance > 0;
 
 			if (bAnyDofVisible)
 			{
diff --git a/Engine/Source/Runtime/Core/Public/UObject/ReleaseObjectVersion.h b/Engine/Source/Runtime/Core/Public/UObject/ReleaseObjectVersion.h
--- a/Engine/Source/Runtime/Core/Public/UObject/ReleaseObjectVersion.h
+++ b/Engine/Source/Runtime/Core/Public/UObject/ReleaseObjectVersion.h
@@ -77,6 +77,9 @@ struct CORE_API FReleaseObjectVersion
 		// Make sure the Blueprint Replicated Property Conditions are actually serialized properly.
 		PropertiesSerializeRepCondition,
 
+		// DepthOfFieldFocalDistance at 0 now disables DOF instead of DepthOfFieldFstop at 0.
+		FocalDistanceDisablesDOF,
+
 		// -----<new versions can be added above this line>-------------------------------------------------
 		VersionPlusOne,
 		LatestVersion = VersionPlusOne - 1
diff --git a/Engine/Source/Runtime/Engine/Private/Scene.cpp b/Engine/Source/Runtime/Engine/Private/Scene.cpp
--- a/Engine/Source/Runtime/Engine/Private/Scene.cpp
+++ b/Engine/Source/Runtime/Engine/Private/Scene.cpp
@@ -3,6 +3,7 @@
 #include "Engine/Scene.h"
 #include "HAL/IConsoleManager.h"
 #include "UObject/RenderingObjectVersion.h"
+#include "UObject/ReleaseObjectVersion.h"
 
 void FColorGradingSettings::ExportToPostProcessSettings(FPostProcessSettings* OutPostProcessSettings) const
 {
@@ -515,8 +516,8 @@ FPostProcessSettings::FPostProcessSettings()
 	RayTracingGIMaxBounces = 1;
 	RayTracingGISamplesPerPixel = 1;
 
-	DepthOfFieldFocalDistance = 1000.0f;
-	DepthOfFieldFstop = 0.0f; // Intentionally invalid to disable DOF by default.
+	DepthOfFieldFocalDistance = 0; // Intentionally invalid to disable DOF by default.
+	DepthOfFieldFstop = 4.0f; 
 	DepthOfFieldMinFstop = 1.2f;
 	DepthOfFieldBladeCount = FPostProcessSettings::kDefaultDepthOfFieldBladeCount;
 	DepthOfFieldSensorWidth = 24.576f;			// APS-C
@@ -942,6 +943,7 @@ FPostProcessSettings::FPostProcessSettings(const FPostProcessSettings& Settings)
 bool FPostProcessSettings::Serialize(FArchive& Ar)
 {
 	Ar.UsingCustomVersion(FRenderingObjectVersion::GUID);
+	Ar.UsingCustomVersion(FReleaseObjectVersion::GUID);
 
 	// Don't actually serialize, just write the custom version for PostSerialize
 	return false;
@@ -949,18 +951,32 @@ bool FPostProcessSettings::Serialize(FArchive& Ar)
 
 void FPostProcessSettings::PostSerialize(const FArchive& Ar)
 {
-	if (Ar.IsLoading() && Ar.CustomVer(FRenderingObjectVersion::GUID) < FRenderingObjectVersion::DiaphragmDOFOnlyForDeferredShadingRenderer)
+	if (Ar.IsLoading())
 	{
-		// Make sure the DOF of the deferred shading renderer is enabled if the circle DOF method was used before with previous default setting for DepthOfFieldFstop.
-		if (DepthOfFieldFstop == 0.0f && DepthOfFieldMethod_DEPRECATED == DOFM_CircleDOF)
+		const int32 RenderingObjectVersion = Ar.CustomVer(FRenderingObjectVersion::GUID);
+		const int32 ReleaseObjectVersion = Ar.CustomVer(FReleaseObjectVersion::GUID);
+
+		if (RenderingObjectVersion < FRenderingObjectVersion::DiaphragmDOFOnlyForDeferredShadingRenderer)
 		{
-			DepthOfFieldFstop = 4.0f;
-		}
+			// This is impossible because FocalDistanceDisablesDOF was in Release 4.23 branch where DiaphragmDOFOnlyForDeferredShadingRenderer was already existing.
+			check(ReleaseObjectVersion < FReleaseObjectVersion::FocalDistanceDisablesDOF);
+
+			// Make sure the DOF of the deferred shading renderer is enabled if the circle DOF method was used before with previous default setting for DepthOfFieldFstop.
+			if (DepthOfFieldFocalDistance == 0.0f && DepthOfFieldMethod_DEPRECATED == DOFM_CircleDOF)
+			{
+				DepthOfFieldFocalDistance = 1000.0f;
+			}
 
-		// Make sure gaussian DOF is disabled on mobile if the DOF method was set to something else.
-		if (DepthOfFieldMethod_DEPRECATED != DOFM_Gaussian)
+			// Make sure gaussian DOF is disabled on mobile if the DOF method was set to something else.
+			if (DepthOfFieldMethod_DEPRECATED != DOFM_Gaussian)
+			{
+				DepthOfFieldScale = 0.0f;
+			}
+		}
+		else if (ReleaseObjectVersion < FReleaseObjectVersion::FocalDistanceDisablesDOF)
 		{
-			DepthOfFieldScale = 0.0f;
+			// This is only for assets saved in the the window DiaphragmDOFOnlyForDeferredShadingRenderer -> FocalDistanceDisablesDOF
+			DepthOfFieldFocalDistance = 0.0f;
 		}
 	}
 }
diff --git a/Engine/Source/Runtime/Renderer/Private/PostProcess/DiaphragmDOFUtils.cpp b/Engine/Source/Runtime/Renderer/Private/PostProcess/DiaphragmDOFUtils.cpp
--- a/Engine/Source/Runtime/Renderer/Private/PostProcess/DiaphragmDOFUtils.cpp
+++ b/Engine/Source/Runtime/Renderer/Private/PostProcess/DiaphragmDOFUtils.cpp
@@ -47,7 +47,7 @@ float DiaphragmDOF::ComputeFocalLengthFromFov(const FSceneView& View)
 FVector4 DiaphragmDOF::CircleDofHalfCoc(const FViewInfo& View)
 {
 	static const auto CVar = IConsoleManager::Get().FindTConsoleVariableDataInt(TEXT("r.DepthOfFieldQuality"));
-	bool bDepthOfField = View.Family->EngineShowFlags.DepthOfField && CVar->GetValueOnRenderThread() > 0 && View.FinalPostProcessSettings.DepthOfFieldFstop > 0;
+	bool bDepthOfField = View.Family->EngineShowFlags.DepthOfField && CVar->GetValueOnRenderThread() > 0 && View.FinalPostProcessSettings.DepthOfFieldFstop > 0 && View.FinalPostProcessSettings.DepthOfFieldFocalDistance > 0;
 
 	FVector4 Ret(0, 1, 0, 0);
 
diff --git a/Engine/Source/Runtime/Renderer/Private/PostProcess/PostProcessing.cpp b/Engine/Source/Runtime/Renderer/Private/PostProcess/PostProcessing.cpp
--- a/Engine/Source/Runtime/Renderer/Private/PostProcess/PostProcessing.cpp
+++ b/Engine/Source/Runtime/Renderer/Private/PostProcess/PostProcessing.cpp
@@ -1401,7 +1401,8 @@ void FPostProcessing::Process(FRHICommandListImmediate& RHICmdList, const FViewI
 			bool bDepthOfField = 
 				View.Family->EngineShowFlags.DepthOfField &&
 				CVar->GetValueOnRenderThread() > 0 &&
-				View.FinalPostProcessSettings.DepthOfFieldFstop > 0;
+				View.FinalPostProcessSettings.DepthOfFieldFstop > 0 &&
+				View.FinalPostProcessSettings.DepthOfFieldFocalDistance > 0;
 
 			// Applies DOF and separate translucency,
 			{
