Moves FPostProcessSettings::bMobileHQGaussian to Mobile DOF category.

diff --git a/Engine/Source/Runtime/Engine/Classes/Engine/Scene.h b/Engine/Source/Runtime/Engine/Classes/Engine/Scene.h
--- a/Engine/Source/Runtime/Engine/Classes/Engine/Scene.h
+++ b/Engine/Source/Runtime/Engine/Classes/Engine/Scene.h
@@ -1095,7 +1095,7 @@ struct FPostProcessSettings
 	// -----------------------------------------------------------------------
 
 	/** Enable HQ Gaussian on high end mobile platforms. (ES3_1) */
-	UPROPERTY(EditAnywhere, BlueprintReadWrite, Category = "Lens|Depth of Field", meta = (editcondition = "bOverride_MobileHQGaussian", DisplayName = "High Quality Gaussian DoF on Mobile"))
+	UPROPERTY(EditAnywhere, BlueprintReadWrite, Category = "Lens|Mobile Depth of Field", meta = (editcondition = "bOverride_MobileHQGaussian", DisplayName = "High Quality Gaussian DoF on Mobile"))
 	uint8 bMobileHQGaussian:1;
 
 	/** Bloom algorithm */
