#UE4 Avoid spamming cvar sync callbacks when running in -nullrhi

diff --git a/Engine/Source/Runtime/Engine/Private/ContentStreaming.cpp b/Engine/Source/Runtime/Engine/Private/ContentStreaming.cpp
--- a/Engine/Source/Runtime/Engine/Private/ContentStreaming.cpp
+++ b/Engine/Source/Runtime/Engine/Private/ContentStreaming.cpp
@@ -1197,10 +1197,13 @@ void FStreamingManagerCollection::AddOrRemoveTextureStreamingManagerIfNeeded(boo
 
 	if( !GRHISupportsTextureStreaming || IsRunningDedicatedServer() )
 	{
-		bUseTextureStreaming = false;
+		if (bUseTextureStreaming)
+		{
+			bUseTextureStreaming = false;
 
-		// some code relies on r.TextureStreaming so we're going to disable it here to reflect the hardware capabilities and system needs
-		CVarSetTextureStreaming.AsVariable()->Set(0, ECVF_SetByCode);
+			// some code relies on r.TextureStreaming so we're going to disable it here to reflect the hardware capabilities and system needs
+			CVarSetTextureStreaming.AsVariable()->Set(0, ECVF_SetByCode);
+		}
 	}
 #if TEXTURE2DMIPMAP_USE_COMPACT_BULKDATA
 	else if (!bUseTextureStreaming)
