Fixed a crash when importing cooked materials within the editor that came from an optimization to cooked shader map ID size.
#tests Verified importing cooked material and mesh content, re-cooking and packaging the derivative data.

diff --git a/Engine/Source/Runtime/Engine/Private/Materials/MaterialInstance.cpp b/Engine/Source/Runtime/Engine/Private/Materials/MaterialInstance.cpp
--- a/Engine/Source/Runtime/Engine/Private/Materials/MaterialInstance.cpp
+++ b/Engine/Source/Runtime/Engine/Private/Materials/MaterialInstance.cpp
@@ -3073,11 +3073,13 @@ void UMaterialInstance::Serialize(FArchive& Ar)
 #if WITH_EDITOR
 		else
 		{
+			const bool bLoadedByCookedMaterial = FPlatformProperties::RequiresCookedData() || GetOutermost()->bIsCookedForEditor;
+
 			FMaterialResource LegacyResource;
 			LegacyResource.LegacySerialize(Ar);
 
 			FMaterialShaderMapId LegacyId;
-			LegacyId.Serialize(Ar);
+			LegacyId.Serialize(Ar, bLoadedByCookedMaterial);
 
 			const FStaticParameterSet& IdParameterSet = LegacyId.GetParameterSet();
 			StaticParameters.StaticSwitchParameters = IdParameterSet.StaticSwitchParameters;
diff --git a/Engine/Source/Runtime/Engine/Private/Materials/MaterialShader.cpp b/Engine/Source/Runtime/Engine/Private/Materials/MaterialShader.cpp
--- a/Engine/Source/Runtime/Engine/Private/Materials/MaterialShader.cpp
+++ b/Engine/Source/Runtime/Engine/Private/Materials/MaterialShader.cpp
@@ -503,7 +503,7 @@ bool FStaticParameterSet::Equivalent(const FStaticParameterSet& ReferenceSet) co
 	return false;
 }
 
-void FMaterialShaderMapId::Serialize(FArchive& Ar)
+void FMaterialShaderMapId::Serialize(FArchive& Ar, bool bLoadedByCookedMaterial)
 {
 	// Note: FMaterialShaderMapId is saved both in packages (legacy UMaterialInstance) and the DDC (FMaterialShaderMap)
 	// Backwards compatibility only works with FMaterialShaderMapId's stored in packages.  
@@ -513,8 +513,9 @@ void FMaterialShaderMapId::Serialize(FArchive& Ar)
 
 #if WITH_EDITOR
 	const bool bIsSavingCooked = Ar.IsSaving() && Ar.IsCooking();
+	bIsCookedId = bLoadedByCookedMaterial;
 
-	if (!bIsSavingCooked)
+	if (!bIsSavingCooked && !bLoadedByCookedMaterial)
 	{
 		uint32 UsageInt = Usage;
 		Ar << UsageInt;
@@ -536,7 +537,7 @@ void FMaterialShaderMapId::Serialize(FArchive& Ar)
 	}
 
 #if WITH_EDITOR
-	if (!bIsSavingCooked)
+	if (!bIsSavingCooked && !bLoadedByCookedMaterial)
 	{
 		ParameterSet.Serialize(Ar);
 
@@ -578,11 +579,13 @@ void FMaterialShaderMapId::Serialize(FArchive& Ar)
 	}
 	else
 	{
-		// Saving cooked data, this should be valid
-		FSHAHash CookedShaderMapIdHash;
-		GetMaterialHash(CookedShaderMapIdHash);
-
-		checkf(CookedShaderMapIdHash != FSHAHash(), TEXT("Tried to save an invalid shadermap id hash during cook"));
+		if (bIsSavingCooked)
+		{
+			// Saving cooked data, this should be valid
+			GetMaterialHash(CookedShaderMapIdHash);
+			checkf(CookedShaderMapIdHash != FSHAHash(), TEXT("Tried to save an invalid shadermap id hash during cook"));
+		}
+		
 		Ar << CookedShaderMapIdHash;
 	}
 #else
@@ -639,17 +642,25 @@ void FMaterialShaderMapId::GetMaterialHash(FSHAHash& OutHash) const
 bool FMaterialShaderMapId::operator==(const FMaterialShaderMapId& ReferenceSet) const
 {
 #if WITH_EDITOR
-	if (Usage != ReferenceSet.Usage
-		|| BaseMaterialId != ReferenceSet.BaseMaterialId)
+	if (IsCookedId() != ReferenceSet.IsCookedId())
 	{
 		return false;
 	}
-#else
+
+	if (!IsCookedId())
+	{
+		if (Usage != ReferenceSet.Usage
+			|| BaseMaterialId != ReferenceSet.BaseMaterialId)
+		{
+			return false;
+		}
+	}
+	else
+#endif
 	if (CookedShaderMapIdHash != ReferenceSet.CookedShaderMapIdHash)
 	{
 		return false;
 	}
-#endif
 
 	if (QualityLevel != ReferenceSet.QualityLevel
 		|| FeatureLevel != ReferenceSet.FeatureLevel)
@@ -658,75 +669,78 @@ bool FMaterialShaderMapId::operator==(const FMaterialShaderMapId& ReferenceSet)
 	}
 
 #if WITH_EDITOR
-	if (ParameterSet != ReferenceSet.ParameterSet
-		|| ReferencedFunctions.Num() != ReferenceSet.ReferencedFunctions.Num()
-		|| ReferencedParameterCollections.Num() != ReferenceSet.ReferencedParameterCollections.Num()
-		|| ShaderTypeDependencies.Num() != ReferenceSet.ShaderTypeDependencies.Num()
-		|| ShaderPipelineTypeDependencies.Num() != ReferenceSet.ShaderPipelineTypeDependencies.Num()
-		|| VertexFactoryTypeDependencies.Num() != ReferenceSet.VertexFactoryTypeDependencies.Num())
-	{
-		return false;
-	}
-
-	for (int32 RefFunctionIndex = 0; RefFunctionIndex < ReferenceSet.ReferencedFunctions.Num(); RefFunctionIndex++)
+	if (!IsCookedId())
 	{
-		const FGuid& ReferenceGuid = ReferenceSet.ReferencedFunctions[RefFunctionIndex];
-
-		if (ReferencedFunctions[RefFunctionIndex] != ReferenceGuid)
+		if (ParameterSet != ReferenceSet.ParameterSet
+			|| ReferencedFunctions.Num() != ReferenceSet.ReferencedFunctions.Num()
+			|| ReferencedParameterCollections.Num() != ReferenceSet.ReferencedParameterCollections.Num()
+			|| ShaderTypeDependencies.Num() != ReferenceSet.ShaderTypeDependencies.Num()
+			|| ShaderPipelineTypeDependencies.Num() != ReferenceSet.ShaderPipelineTypeDependencies.Num()
+			|| VertexFactoryTypeDependencies.Num() != ReferenceSet.VertexFactoryTypeDependencies.Num())
 		{
 			return false;
 		}
-	}
 
-	for (int32 RefCollectionIndex = 0; RefCollectionIndex < ReferenceSet.ReferencedParameterCollections.Num(); RefCollectionIndex++)
-	{
-		const FGuid& ReferenceGuid = ReferenceSet.ReferencedParameterCollections[RefCollectionIndex];
+		for (int32 RefFunctionIndex = 0; RefFunctionIndex < ReferenceSet.ReferencedFunctions.Num(); RefFunctionIndex++)
+		{
+			const FGuid& ReferenceGuid = ReferenceSet.ReferencedFunctions[RefFunctionIndex];
 
-		if (ReferencedParameterCollections[RefCollectionIndex] != ReferenceGuid)
+			if (ReferencedFunctions[RefFunctionIndex] != ReferenceGuid)
+			{
+				return false;
+			}
+		}
+
+		for (int32 RefCollectionIndex = 0; RefCollectionIndex < ReferenceSet.ReferencedParameterCollections.Num(); RefCollectionIndex++)
 		{
-			return false;
+			const FGuid& ReferenceGuid = ReferenceSet.ReferencedParameterCollections[RefCollectionIndex];
+
+			if (ReferencedParameterCollections[RefCollectionIndex] != ReferenceGuid)
+			{
+				return false;
+			}
 		}
-	}
 
-	for (int32 ShaderIndex = 0; ShaderIndex < ShaderTypeDependencies.Num(); ShaderIndex++)
-	{
-		const FShaderTypeDependency& ShaderTypeDependency = ShaderTypeDependencies[ShaderIndex];
+		for (int32 ShaderIndex = 0; ShaderIndex < ShaderTypeDependencies.Num(); ShaderIndex++)
+		{
+			const FShaderTypeDependency& ShaderTypeDependency = ShaderTypeDependencies[ShaderIndex];
 
-		if (ShaderTypeDependency != ReferenceSet.ShaderTypeDependencies[ShaderIndex])
+			if (ShaderTypeDependency != ReferenceSet.ShaderTypeDependencies[ShaderIndex])
+			{
+				return false;
+			}
+		}
+
+		for (int32 ShaderPipelineIndex = 0; ShaderPipelineIndex < ShaderPipelineTypeDependencies.Num(); ShaderPipelineIndex++)
 		{
-			return false;
+			const FShaderPipelineTypeDependency& ShaderPipelineTypeDependency = ShaderPipelineTypeDependencies[ShaderPipelineIndex];
+
+			if (ShaderPipelineTypeDependency != ReferenceSet.ShaderPipelineTypeDependencies[ShaderPipelineIndex])
+			{
+				return false;
+			}
 		}
-	}
 
-	for (int32 ShaderPipelineIndex = 0; ShaderPipelineIndex < ShaderPipelineTypeDependencies.Num(); ShaderPipelineIndex++)
-	{
-		const FShaderPipelineTypeDependency& ShaderPipelineTypeDependency = ShaderPipelineTypeDependencies[ShaderPipelineIndex];
+		for (int32 VFIndex = 0; VFIndex < VertexFactoryTypeDependencies.Num(); VFIndex++)
+		{
+			const FVertexFactoryTypeDependency& VFDependency = VertexFactoryTypeDependencies[VFIndex];
+
+			if (VFDependency != ReferenceSet.VertexFactoryTypeDependencies[VFIndex])
+			{
+				return false;
+			}
+		}
 
-		if (ShaderPipelineTypeDependency != ReferenceSet.ShaderPipelineTypeDependencies[ShaderPipelineIndex])
+		if (TextureReferencesHash != ReferenceSet.TextureReferencesHash)
 		{
 			return false;
 		}
-	}
-
-	for (int32 VFIndex = 0; VFIndex < VertexFactoryTypeDependencies.Num(); VFIndex++)
-	{
-		const FVertexFactoryTypeDependency& VFDependency = VertexFactoryTypeDependencies[VFIndex];
 
-		if (VFDependency != ReferenceSet.VertexFactoryTypeDependencies[VFIndex])
+		if( BasePropertyOverridesHash != ReferenceSet.BasePropertyOverridesHash )
 		{
 			return false;
 		}
 	}
-
-	if (TextureReferencesHash != ReferenceSet.TextureReferencesHash)
-	{
-		return false;
-	}
-
-	if( BasePropertyOverridesHash != ReferenceSet.BasePropertyOverridesHash )
-	{
-		return false;
-	}
 #endif // WITH_EDITOR
 
 	return true;
@@ -2410,7 +2424,7 @@ void FMaterialShaderMap::Serialize(FArchive& Ar, bool bInlineShaderResources, bo
 	// Backwards compatibility therefore will not work based on the version of Ar
 	// Instead, just bump MATERIALSHADERMAP_DERIVEDDATA_VER
 
-	ShaderMapId.Serialize(Ar);
+	ShaderMapId.Serialize(Ar, bLoadedByCookedMaterial);
 
 	// serialize the platform enum as a uint8
 	int32 TempPlatform = (int32)GetShaderPlatform();
diff --git a/Engine/Source/Runtime/Engine/Public/MaterialShared.h b/Engine/Source/Runtime/Engine/Public/MaterialShared.h
--- a/Engine/Source/Runtime/Engine/Public/MaterialShared.h
+++ b/Engine/Source/Runtime/Engine/Public/MaterialShared.h
@@ -667,9 +667,7 @@ namespace EMaterialShaderMapUsage
 class FMaterialShaderMapId
 {
 public:
-#if !WITH_EDITOR
 	FSHAHash CookedShaderMapIdHash;
-#endif
 
 #if WITH_EDITOR
 	/** 
@@ -698,10 +696,12 @@ public:
 	EMaterialShaderMapUsage::Type Usage;
 
 private:
+	/** Was the shadermap Id loaded in from a cooked resource. */
+	bool bIsCookedId;
+
 	/** Static parameters and base Id. */
 	FStaticParameterSet ParameterSet;
 	FString ParameterSetLayerParametersKey;
-
 public:
 	/** Guids of any functions the material was dependent on. */
 	TArray<FGuid> ReferencedFunctions;
@@ -734,6 +734,7 @@ public:
 		, FeatureLevel(ERHIFeatureLevel::SM4)
 #if WITH_EDITOR
 		, Usage(EMaterialShaderMapUsage::Default)
+		, bIsCookedId(false)
 #endif
 	{ }
 
@@ -744,12 +745,21 @@ public:
 	ENGINE_API void SetShaderDependencies(const TArray<FShaderType*>& ShaderTypes, const TArray<const FShaderPipelineType*>& ShaderPipelineTypes, const TArray<FVertexFactoryType*>& VFTypes, EShaderPlatform ShaderPlatform);
 #endif
 
-	void Serialize(FArchive& Ar);
+	void Serialize(FArchive& Ar, bool bLoadedByCookedMaterial);
+
+	bool IsCookedId() const
+	{
+#if WITH_EDITOR
+		return bIsCookedId;
+#else
+		return true;
+#endif
+	}
 
 	bool IsValid() const
 	{
 #if WITH_EDITOR
-		return BaseMaterialId.IsValid();
+		return !IsCookedId() ? BaseMaterialId.IsValid() : (CookedShaderMapIdHash != FSHAHash());
 #else
 		return (CookedShaderMapIdHash != FSHAHash());
 #endif
@@ -758,7 +768,7 @@ public:
 	friend uint32 GetTypeHash(const FMaterialShaderMapId& Ref)
 	{
 #if WITH_EDITOR
-		return Ref.BaseMaterialId.A;
+		return !Ref.IsCookedId() ? Ref.BaseMaterialId.A : (*(uint32*)&Ref.CookedShaderMapIdHash.Hash[0]);
 #else
 		// Using the hash value directly instead of FSHAHash CRC as fairly uniform distribution
 		return *(uint32*)&Ref.CookedShaderMapIdHash.Hash[0];
