Fixed: Setting MinLod and stripping unused LODs on a static mesh will make it stop rendering when used with Instanced static mesh

diff --git a/Engine/Source/Runtime/Engine/Private/Components/StaticMeshComponent.cpp b/Engine/Source/Runtime/Engine/Private/Components/StaticMeshComponent.cpp
--- a/Engine/Source/Runtime/Engine/Private/Components/StaticMeshComponent.cpp
+++ b/Engine/Source/Runtime/Engine/Private/Components/StaticMeshComponent.cpp
@@ -2043,10 +2043,12 @@ bool UStaticMeshComponent::HasLightmapTextureCoordinates() const
 		(GetStaticMesh()->LightMapCoordinateIndex >= 0) &&
 		(GetStaticMesh()->RenderData != NULL) &&
 		(GetStaticMesh()->RenderData->LODResources.Num() > 0) &&
-		(GetStaticMesh()->LightMapCoordinateIndex >= 0) &&	
-		((uint32)GetStaticMesh()->LightMapCoordinateIndex < GetStaticMesh()->RenderData->LODResources[0].VertexBuffers.StaticMeshVertexBuffer.GetNumTexCoords()))
+		(GetStaticMesh()->LightMapCoordinateIndex >= 0))
 	{
-		return true;
+		const ERHIFeatureLevel::Type FeatureLevel = GetScene() ? GetScene()->GetFeatureLevel() : GMaxRHIFeatureLevel;
+		int32 MeshMinLOD = GetStaticMesh()->MinLOD.GetValueForFeatureLevel(FeatureLevel);
+		
+		return ((uint32)GetStaticMesh()->LightMapCoordinateIndex < GetStaticMesh()->RenderData->LODResources[MeshMinLOD].VertexBuffers.StaticMeshVertexBuffer.GetNumTexCoords());
 	}
 	return false;
 }
