Fix scene queries against GeometryCollection

diff --git a/Engine/Source/Runtime/Engine/Private/Collision/SceneQueryLowLevel.cpp b/Engine/Source/Runtime/Engine/Private/Collision/SceneQueryLowLevel.cpp
--- a/Engine/Source/Runtime/Engine/Private/Collision/SceneQueryLowLevel.cpp
+++ b/Engine/Source/Runtime/Engine/Private/Collision/SceneQueryLowLevel.cpp
@@ -68,7 +68,7 @@ void LowLevelRaycast(FPhysScene& Scene, const FVector& Start, const FVector& Dir
 	if (ForceStandardSQ == 0)
 	{
 		ISQAccelerator* SQAccelerator = Scene.GetSQAccelerator();
-		SQAccelerator->Raycast(Start, Dir, HitBuffer, OutputFlags, QueryFlags, Filter, QueryFilterData, *QueryCallback);
+		SQAccelerator->Raycast(Start, Dir, DeltaMag, HitBuffer, OutputFlags, QueryFlags, Filter, QueryFilterData, *QueryCallback);
 		FinalizeQuery(HitBuffer);
 	}
 	else
@@ -100,7 +100,7 @@ void LowLevelSweep(FPhysScene& Scene, const FPhysicsGeometry& QueryGeom, const F
 	if (ForceStandardSQ == 0)
 	{
 		ISQAccelerator* SQAccelerator = Scene.GetSQAccelerator();
-		SQAccelerator->Sweep(QueryGeom, StartTM, Dir, HitBuffer, OutputFlags, QueryFlags, Filter, QueryFilterData, *QueryCallback);
+		SQAccelerator->Sweep(QueryGeom, StartTM, Dir, DeltaMag, HitBuffer, OutputFlags, QueryFlags, Filter, QueryFilterData, *QueryCallback);
 		FinalizeQuery(HitBuffer);
 	}
 	else
