Fix details not appearing for Directional Lights
#rnx

diff --git a/Engine/Source/Runtime/Engine/Classes/Engine/DirectionalLight.h b/Engine/Source/Runtime/Engine/Classes/Engine/DirectionalLight.h
--- a/Engine/Source/Runtime/Engine/Classes/Engine/DirectionalLight.h
+++ b/Engine/Source/Runtime/Engine/Classes/Engine/DirectionalLight.h
@@ -8,6 +8,7 @@
 #include "DirectionalLight.generated.h"
 
 class UArrowComponent;
+class UDirectionalLightComponent;
 
 /**
  * Implements a directional light actor.
@@ -23,6 +24,10 @@ class ADirectionalLight
 private:
 	UPROPERTY()
 	UArrowComponent* ArrowComponent;
+
+	/* EditorOnly reference to the light component to allow it to be displayed in the details panel correctly */
+	UPROPERTY(VisibleAnywhere, Category="Light")
+	UDirectionalLightComponent* DirectionalLightComponent;
 #endif
 
 public:
diff --git a/Engine/Source/Runtime/Engine/Private/Light.cpp b/Engine/Source/Runtime/Engine/Private/Light.cpp
--- a/Engine/Source/Runtime/Engine/Private/Light.cpp
+++ b/Engine/Source/Runtime/Engine/Private/Light.cpp
@@ -232,7 +232,10 @@ ADirectionalLight::ADirectionalLight(const FObjectInitializer& ObjectInitializer
 	};
 	static FConstructorStatics ConstructorStatics;
 
-	UDirectionalLightComponent* DirectionalLightComponent = CastChecked<UDirectionalLightComponent>(GetLightComponent());
+#if !WITH_EDITORONLY_DATA
+	UDirectionalLightComponent* DirectionalLightComponent;
+#endif
+	DirectionalLightComponent = CastChecked<UDirectionalLightComponent>(GetLightComponent());
 	DirectionalLightComponent->Mobility = EComponentMobility::Stationary;
 	DirectionalLightComponent->RelativeRotation = FRotator(-46.0f, 0.0f, 0.0f);
 	// Make directional light icons big since they tend to be important
