Clean up IO Request memory in animation streaming

diff --git a/Engine/Source/Runtime/Engine/Private/Animation/AnimationStreaming.cpp b/Engine/Source/Runtime/Engine/Private/Animation/AnimationStreaming.cpp
--- a/Engine/Source/Runtime/Engine/Private/Animation/AnimationStreaming.cpp
+++ b/Engine/Source/Runtime/Engine/Private/Animation/AnimationStreaming.cpp
@@ -354,6 +354,8 @@ void FAnimationStreamingManager::OnAsyncFileCallback(FStreamingAnimationData* St
 		ChunkStorage.RequestStart = -2.0; //Signify we have finished loading
 
 		UE_LOG(LogAnimation, Log, TEXT("Request Finished %.2f\n Anim Chunk Streamed %.4f\n"), CurrentTime, RequestDuration);
+
+		FMemory::Free(Mem);
 	}
 }
 
