#ue4 - Allow original move attempt in ResolvePenetration only if in the same direction as the MTD vector.

diff --git a/Engine/Source/Runtime/Engine/Private/Components/MovementComponent.cpp b/Engine/Source/Runtime/Engine/Private/Components/MovementComponent.cpp
--- a/Engine/Source/Runtime/Engine/Private/Components/MovementComponent.cpp
+++ b/Engine/Source/Runtime/Engine/Private/Components/MovementComponent.cpp
@@ -688,7 +688,7 @@ bool UMovementComponent::ResolvePenetrationImpl(const FVector& ProposedAdjustmen
 					// Finally, try the original move without MTD adjustments, but allowing depenetration along the MTD normal.
 					// This was blocked because MOVECOMP_NeverIgnoreBlockingOverlaps was true for the original move to try a better depenetration normal, but we might be running in to other geometry in the attempt.
 					// This won't necessarily get us all the way out of penetration, but can in some cases and does make progress in exiting the penetration.
-					if (!bMoved)
+					if (!bMoved && FVector::DotProduct(MoveDelta, Adjustment) > 0.f)
 					{
 						bMoved = MoveUpdatedComponent(MoveDelta, NewRotationQuat, true, nullptr, ETeleportType::TeleportPhysics);
 						UE_LOG(LogMovement, Verbose, TEXT("ResolvePenetration:   sweep by %s (Original move, attempt success = %d)"), *(MoveDelta).ToString(), bMoved);
