- Added ability to pause shader batching when we come out of embedded mode. Helps slow devices not time out when rotating when the screen resizes

diff --git a/Engine/Source/Runtime/Engine/Private/GameEngine.cpp b/Engine/Source/Runtime/Engine/Private/GameEngine.cpp
--- a/Engine/Source/Runtime/Engine/Private/GameEngine.cpp
+++ b/Engine/Source/Runtime/Engine/Private/GameEngine.cpp
@@ -31,6 +31,7 @@
 #include "GeneralProjectSettings.h"
 #include "Misc/PackageName.h"
 #include "HAL/PlatformApplicationMisc.h"
+#include "ShaderPipelineCache.h"
 
 #include "Misc/ConfigCacheIni.h"
 
@@ -783,6 +784,16 @@ public:
 				FPlatformApplicationMisc::ControlScreensaver(bEnabled ? FGenericPlatformApplicationMisc::Enable : FGenericPlatformApplicationMisc::Disable);
 				Message.OnCompleteDelegate({}, TEXT(""));
 			}
+			else if (Message.Command == TEXT("shaderresumebatching"))
+			{
+				FShaderPipelineCache::ResumeBatching();
+				Message.OnCompleteDelegate({}, TEXT(""));
+			}
+			else if (Message.Command == TEXT("shaderpausebatching"))
+			{
+				FShaderPipelineCache::PauseBatching();
+				Message.OnCompleteDelegate({}, TEXT(""));
+			}
 			else if (Message.Command == TEXT("getmemorybucket"))
 			{
 				Message.OnCompleteDelegate({ { TEXT("bucket"), LexToString((int)FPlatformMemory::GetMemorySizeBucket()) } }, TEXT(""));
