- Add console variables to limit distance tessellation (FX.Trail.MaxDistanceTessellation) and tangent tessellation (FX.Trail.MaxTangentTessellation)
- These are currently set to have no impact, and can be used to control memory usage under large displacements

diff --git a/Engine/Source/Runtime/Engine/Private/Particles/ParticleTrail2EmitterInstance.cpp b/Engine/Source/Runtime/Engine/Private/Particles/ParticleTrail2EmitterInstance.cpp
--- a/Engine/Source/Runtime/Engine/Private/Particles/ParticleTrail2EmitterInstance.cpp
+++ b/Engine/Source/Runtime/Engine/Private/Particles/ParticleTrail2EmitterInstance.cpp
@@ -48,6 +48,26 @@ DECLARE_CYCLE_STAT(TEXT("TrailEmitterInstance Init GT"), STAT_TrailEmitterInstan
 
 #define MAX_TRAIL_INDICES	65535
 
+namespace FXConsoleVariables
+{
+	int32 MaxDistanceTessellation = MAX_TRAIL_INDICES;
+	int32 MaxTangentTessellation = MAX_TRAIL_INDICES;
+
+	FAutoConsoleVariableRef CVarMaxDistanceTessellation(
+		TEXT("FX.Trail.MaxDistanceTessellation"),
+		MaxDistanceTessellation,
+		TEXT("Maximum tessellation steps allowed for distance based tessellation."),
+		ECVF_Cheat
+	);
+
+	FAutoConsoleVariableRef CVarMaxTangentTessellation(
+		TEXT("FX.Trail.MaxTangentTessellation"),
+		MaxTangentTessellation,
+		TEXT("Maximum tessellation steps allowed for tangent based tessellation."),
+		ECVF_Cheat
+	);
+}
+
 /*-----------------------------------------------------------------------------
 	FParticleTrailsEmitterInstance_Base.
 -----------------------------------------------------------------------------*/
@@ -2495,7 +2515,7 @@ void FParticleRibbonEmitterInstance::DetermineVertexAndTriangleCount()
 
 					//@todo. Need to adjust the tangent diff count when the distance is REALLY small...
 					float TangDiff = CheckTangent * TrailTypeData->TangentTessellationScalar;
-					int32 InterpCount = FMath::TruncToInt(DistDiff) + FMath::TruncToInt(TangDiff);
+					int32 InterpCount = FMath::Min(FMath::TruncToInt(DistDiff), FXConsoleVariables::MaxDistanceTessellation) + FMath::Min(FMath::TruncToInt(TangDiff), FXConsoleVariables::MaxTangentTessellation);
 
 					// There always is at least 1 point (the source particle itself)
 					InterpCount = (InterpCount > 0) ? InterpCount : 1;
@@ -3769,7 +3789,7 @@ void FParticleAnimTrailEmitterInstance::DetermineVertexAndTriangleCount()
 						// Determine the number of rendered interpolated points between these two particles
 						float CheckDistance = (CurrParticle->Location - PrevParticle->Location).Size();
 						float DistDiff = CheckDistance / TrailTypeData->DistanceTessellationStepSize;
-						InterpCount += FMath::TruncToInt(DistDiff);
+						InterpCount += FMath::Min(FMath::TruncToInt(DistDiff), FXConsoleVariables::MaxDistanceTessellation);
 					}
 					
 					if (bApplyTangentTessellation )
@@ -3788,7 +3808,7 @@ void FParticleAnimTrailEmitterInstance::DetermineVertexAndTriangleCount()
 							CheckTangent = FMath::Max( CheckTangent, NextCheckTangent );
 						}		
 						float TangDiff = CheckTangent / TangentTessellationStepSize;
-						InterpCount += FMath::TruncToInt(TangDiff);
+						InterpCount += FMath::Min(FMath::TruncToInt(TangDiff), FXConsoleVariables::MaxTangentTessellation);
 					}
 
 					if( bApplyWidthTessellation )
