- Remove cheat flag from FX.Trail.MaxDistanceTessellation / FX.Trail.MaxTangentTessellation so that .ini files can override

#rnx

diff --git a/Engine/Source/Runtime/Engine/Private/Particles/ParticleTrail2EmitterInstance.cpp b/Engine/Source/Runtime/Engine/Private/Particles/ParticleTrail2EmitterInstance.cpp
--- a/Engine/Source/Runtime/Engine/Private/Particles/ParticleTrail2EmitterInstance.cpp
+++ b/Engine/Source/Runtime/Engine/Private/Particles/ParticleTrail2EmitterInstance.cpp
@@ -57,14 +57,14 @@ namespace FXConsoleVariables
 		TEXT("FX.Trail.MaxDistanceTessellation"),
 		MaxDistanceTessellation,
 		TEXT("Maximum tessellation steps allowed for distance based tessellation."),
-		ECVF_Cheat
+		ECVF_Default
 	);
 
 	FAutoConsoleVariableRef CVarMaxTangentTessellation(
 		TEXT("FX.Trail.MaxTangentTessellation"),
 		MaxTangentTessellation,
 		TEXT("Maximum tessellation steps allowed for tangent based tessellation."),
-		ECVF_Cheat
+		ECVF_Default
 	);
 }
 
