Fix CIS

diff --git a/Engine/Source/Runtime/Engine/Public/Scalability.h b/Engine/Source/Runtime/Engine/Public/Scalability.h
--- a/Engine/Source/Runtime/Engine/Public/Scalability.h
+++ b/Engine/Source/Runtime/Engine/Public/Scalability.h
@@ -164,5 +164,5 @@ namespace Scalability
 #if WITH_EDITOR
 	/** Set an Editor preview scalability platform */
 	void ENGINE_API ChangeScalabilityPreviewPlatform(FName NewPlatformScalabilityName);
-#endif;
+#endif
 }
