The scene cubemap texture array for reflection captures is no longer bound if no reflection captures are present, instead it falls back to the default cube texture array. In addition to this, index -1 is never used anymore, instead 0 is used when no reflection captures are present as that will sample the default cubemap texture array.

diff --git a/Engine/Source/Runtime/Engine/Public/PrimitiveUniformShaderParameters.h b/Engine/Source/Runtime/Engine/Public/PrimitiveUniformShaderParameters.h
--- a/Engine/Source/Runtime/Engine/Public/PrimitiveUniformShaderParameters.h
+++ b/Engine/Source/Runtime/Engine/Public/PrimitiveUniformShaderParameters.h
@@ -38,7 +38,7 @@ BEGIN_GLOBAL_SHADER_PARAMETER_STRUCT(FPrimitiveUniformShaderParameters,ENGINE_AP
 	SHADER_PARAMETER(FVector, LocalObjectBoundsMax)		// This is used in a custom material function (ObjectLocalBounds.uasset)
 	SHADER_PARAMETER(uint32,LightmapDataIndex)
 	SHADER_PARAMETER(FVector, PreSkinnedLocalBoundsMin)	// Local space min bounds, pre-skinning
-	SHADER_PARAMETER(int32, SingleCaptureIndex)
+	SHADER_PARAMETER(int32, SingleCaptureIndex)			// Should default to 0 if no reflection captures are provided, as there will be a default black (0,0,0,0) cubemap in that slot
 	SHADER_PARAMETER(FVector, PreSkinnedLocalBoundsMax)	// Local space max bounds, pre-skinning
     SHADER_PARAMETER(uint32, OutputVelocity)
 	SHADER_PARAMETER_ARRAY(FVector4, CustomPrimitiveData, [FCustomPrimitiveData::NumCustomPrimitiveDataFloat4s]) // Custom data per primitive that can be accessed through material expression parameters and modified through UStaticMeshComponent
@@ -105,7 +105,8 @@ inline FPrimitiveUniformShaderParameters GetPrimitiveUniformShaderParameters(
 	Result.UseVolumetricLightmapShadowFromStationaryLights = bUseVolumetricLightmap && bUseSingleSampleShadowFromStationaryLights ? 1.0f : 0.0f;
 	Result.DrawsVelocity = bDrawsVelocity ? 1 : 0;
 	Result.LightmapDataIndex = LightmapDataIndex;
-	Result.SingleCaptureIndex = SingleCaptureIndex;
+	// If SingleCaptureIndex is invalid, set it to 0 since there will be a default cubemap at that slot
+	Result.SingleCaptureIndex = FMath::Max(SingleCaptureIndex, 0);
 	Result.OutputVelocity = (bOutputVelocity) ? 1 : 0;
 
 	// Clear to 0
diff --git a/Engine/Source/Runtime/Renderer/Private/ReflectionEnvironment.cpp b/Engine/Source/Runtime/Renderer/Private/ReflectionEnvironment.cpp
--- a/Engine/Source/Runtime/Renderer/Private/ReflectionEnvironment.cpp
+++ b/Engine/Source/Runtime/Renderer/Private/ReflectionEnvironment.cpp
@@ -267,7 +267,8 @@ void SetupReflectionUniformParameters(const FViewInfo& View, FReflectionUniformP
 	if (View.Family->EngineShowFlags.ReflectionEnvironment 
 		&& View.FeatureLevel >= ERHIFeatureLevel::SM5
 		&& Scene
-		&& Scene->ReflectionSceneData.CubemapArray.IsValid())
+		&& Scene->ReflectionSceneData.CubemapArray.IsValid()
+		&& Scene->ReflectionSceneData.RegisteredReflectionCaptures.Num())
 	{
 		CubeArrayTexture = Scene->ReflectionSceneData.CubemapArray.GetRenderTarget().ShaderResourceTexture;
 	}
diff --git a/Engine/Source/Runtime/Renderer/Private/RendererScene.cpp b/Engine/Source/Runtime/Renderer/Private/RendererScene.cpp
--- a/Engine/Source/Runtime/Renderer/Private/RendererScene.cpp
+++ b/Engine/Source/Runtime/Renderer/Private/RendererScene.cpp
@@ -2619,7 +2619,8 @@ void FScene::GetPrimitiveUniformShaderParameters_RenderThread(const FPrimitiveSc
 		PreviousLocalToWorld = PrimitiveSceneInfo->Proxy->GetLocalToWorld();
 	}
 
-	SingleCaptureIndex = PrimitiveSceneInfo->CachedReflectionCaptureProxy ? PrimitiveSceneInfo->CachedReflectionCaptureProxy->SortedCaptureIndex : -1;
+	// Get index if proxy exists, otherwise fall back to index 0 which will contain the default black cubemap
+	SingleCaptureIndex = PrimitiveSceneInfo->CachedReflectionCaptureProxy ? PrimitiveSceneInfo->CachedReflectionCaptureProxy->SortedCaptureIndex : 0;
 }
 
 struct FUpdateLightTransformParameters
