Fix thread safety conflict with SeekToTime being called on the AudioThread and StreamCompressedData on decompression worker thread.  Change enforces all seek requests to be processed on the worker thread.

diff --git a/Engine/Source/Runtime/Engine/Private/ADPCMAudioInfo.cpp b/Engine/Source/Runtime/Engine/Private/ADPCMAudioInfo.cpp
--- a/Engine/Source/Runtime/Engine/Private/ADPCMAudioInfo.cpp
+++ b/Engine/Source/Runtime/Engine/Private/ADPCMAudioInfo.cpp
@@ -3,11 +3,19 @@
 
 #include "ADPCMAudioInfo.h"
 #include "CoreMinimal.h"
+#include "HAL/IConsoleManager.h"
 #include "Interfaces/IAudioFormat.h"
 #include "Sound/SoundWave.h"
 #include "Audio.h"
 #include "ContentStreaming.h"
 
+static int32 bDisableADPCMSeekLockCVar = 1;
+FAutoConsoleVariableRef CVarDisableADPCMSeekLock(
+	TEXT("au.DisableADPCMSeekLock"),
+	bDisableADPCMSeekLockCVar,
+	TEXT("Disables ADPCM seek crit section fix for multiple seek requests per frame.\n"),
+	ECVF_Default);
+
 #define WAVE_FORMAT_LPCM  1
 #define WAVE_FORMAT_ADPCM 2
 
@@ -28,6 +36,8 @@ FADPCMAudioInfo::FADPCMAudioInfo(void)
 	, UncompressedBlockData(nullptr)
 	, SamplesPerBlock(0)
 	, bSeekPending(false)
+	, TargetSeekTime(0.0f)
+	, LastSeekTime(0.0f)
 {
 }
 
@@ -40,12 +50,28 @@ FADPCMAudioInfo::~FADPCMAudioInfo(void)
 	}
 }
 
-void FADPCMAudioInfo::SeekToTime(const float SeekTime)
+void FADPCMAudioInfo::SeekToTime(const float InSeekTime)
+{
+	if (bDisableADPCMSeekLockCVar)
+	{
+		SeekToTimeInternal(InSeekTime);
+	}
+	else
+	{
+		TargetSeekTime = InSeekTime;
+	}
+}
+
+void FADPCMAudioInfo::SeekToTimeInternal(const float InSeekTime)
 {
+	LastSeekTime = InSeekTime;
+
 	// Reset chunk handle in preperation for a new chunk.
 	CurCompressedChunkData = nullptr;
 
-	if(SeekTime <= 0.0f)
+	UE_LOG(LogAudio, Verbose, TEXT("Seeking ADPCM source to %.3f sec"), InSeekTime);
+
+	if (InSeekTime <= 0.0f)
 	{
 		CurrentCompressedBlockIndex = 0;
 		CurrentUncompressedBlockSampleIndex = 0;
@@ -53,13 +79,14 @@ void FADPCMAudioInfo::SeekToTime(const float SeekTime)
 		CurrentChunkBufferOffset = 0;
 		TotalSamplesStreamed = 0;
 
+		bSeekPending = false;
 		return;
 	}
 
 	// Calculate block index & force SeekTime to be in bounds.
 	check(WaveInfo.pSamplesPerSec != nullptr);
 	const uint32 SamplesPerSec = *WaveInfo.pSamplesPerSec;
-	uint32 SeekedSamples = static_cast<uint32>(SeekTime * SamplesPerSec);
+	const uint32 SeekedSamples = static_cast<uint32>(InSeekTime * SamplesPerSec);
 	TotalSamplesStreamed = FMath::Min<uint32>(SeekedSamples, TotalSamplesPerChannel - 1);
 
 	const uint32 HeaderOffset = static_cast<uint32>(WaveInfo.SampleDataStart - SrcBufferData);
@@ -222,10 +249,10 @@ bool FADPCMAudioInfo::ReadCompressedData(uint8* Destination, bool bLooping, uint
 	check(Destination);
 	check(BufferSize % ChannelSampleSize == 0);
 
-	int16* OutData = (int16*)Destination;
+	ProcessSeekRequest();
 
+	int16* OutData = (int16*)Destination;
 	bool ReachedEndOfSamples = false;
-
 	if(Format == WAVE_FORMAT_ADPCM)
 	{
 		// We need to loop over the number of samples requested since an uncompressed block will not match the number of frames requested
@@ -334,6 +361,28 @@ int FADPCMAudioInfo::GetStreamBufferSize() const
 	return StreamBufferSize;
 }
 
+void FADPCMAudioInfo::ProcessSeekRequest()
+{
+	if (bDisableADPCMSeekLockCVar)
+	{
+		return;
+	}
+
+	float NewSeekTime = -1.0f;
+	{
+		FScopeLock Lock(&StreamSeekCriticalSection);
+		if (!FMath::IsNearlyEqual(TargetSeekTime, LastSeekTime))
+		{
+			NewSeekTime = TargetSeekTime;
+		}
+	}
+
+	if (NewSeekTime >= 0.0f)
+	{
+		SeekToTimeInternal(NewSeekTime);
+	}
+}
+
 bool FADPCMAudioInfo::StreamCompressedInfoInternal(USoundWave* Wave, struct FSoundQualityInfo* QualityInfo)
 {
 	check(QualityInfo);
@@ -446,10 +495,10 @@ bool FADPCMAudioInfo::StreamCompressedData(uint8* Destination, bool bLooping, ui
 		return true;
 	}
 
-	int16* OutData = (int16*)Destination;
+	ProcessSeekRequest();
 
+	int16* OutData = (int16*)Destination;
 	bool ReachedEndOfSamples = false;
-
 	if(Format == WAVE_FORMAT_ADPCM)
 	{
 		// We need to loop over the number of samples requested since an uncompressed block will not match the number of frames requested
@@ -480,8 +529,9 @@ bool FADPCMAudioInfo::StreamCompressedData(uint8* Destination, bool bLooping, ui
 						// We only need to worry about missing the stream chunk if we were seeking. Seeking might cause a bit of latency with chunk loading. That is expected.
 						if (!bSeekPending)
 						{
-							// If we did not get it then just bail, CurrentChunkIndex will not get incremented on the next callback so in effect another attempt will be made to fetch the chunk
-							// Since audio streaming depends on the general data streaming mechanism used by other parts of the engine and new data is prefectched on the game tick thread its possible a game hickup can cause this
+							// If we did not load chunk then bail. CurrentChunkIndex will not get incremented on the next callback so in effect another attempt will be made to fetch the chunk.
+							// Since audio streaming depends on the general data streaming mechanism used by other parts of the engine and new data is pre-fetched on the game thread, its
+							// possible a game thread stall can cause this.
 							UE_LOG(LogAudio, Verbose, TEXT("Missed Deadline chunk %d"), CurrentChunkIndex);
 						}
 
@@ -495,6 +545,7 @@ bool FADPCMAudioInfo::StreamCompressedData(uint8* Destination, bool bLooping, ui
 					{
 						CurrentChunkBufferOffset = CurrentChunkIndex == 0 ? FirstChunkSampleDataOffset : 0;
 					}
+
 					bSeekPending = false;
 				}
 
@@ -597,6 +648,7 @@ bool FADPCMAudioInfo::StreamCompressedData(uint8* Destination, bool bLooping, ui
 				{
 					CurrentChunkBufferOffset = CurrentChunkIndex == 0 ? FirstChunkSampleDataOffset : 0;
 				}
+
 				bSeekPending = false;
 			}
 
diff --git a/Engine/Source/Runtime/Engine/Public/ADPCMAudioInfo.h b/Engine/Source/Runtime/Engine/Public/ADPCMAudioInfo.h
--- a/Engine/Source/Runtime/Engine/Public/ADPCMAudioInfo.h
+++ b/Engine/Source/Runtime/Engine/Public/ADPCMAudioInfo.h
@@ -166,4 +166,13 @@ private:
 	uint32			CurrentCompressedBlockIndex;		// For non disk streaming - the current compressed block in the compressed source data
 	uint32			TotalCompressedBlocksPerChannel;	// For non disk streaming - the total number of compressed blocks per channel
 	uint8			bSeekPending : 1;					// Whether or not seek has been requested and pending a read
+
+private:
+	void ProcessSeekRequest();
+	void SeekToTimeInternal(const float InSeekTime);
+
+	float TargetSeekTime;
+	float LastSeekTime;
+
+	FCriticalSection StreamSeekCriticalSection;
 };
