Anim Streaming Priority change + Logging for streaming issue

diff --git a/Engine/Source/Runtime/Engine/Private/Animation/AnimStreamable.cpp b/Engine/Source/Runtime/Engine/Private/Animation/AnimStreamable.cpp
--- a/Engine/Source/Runtime/Engine/Private/Animation/AnimStreamable.cpp
+++ b/Engine/Source/Runtime/Engine/Private/Animation/AnimStreamable.cpp
@@ -275,7 +275,9 @@ void UAnimStreamable::GetAnimationPose(FCompactPose& OutPose, FBlendedCurve& Out
 	{
 		IAnimationStreamingManager& StreamingManager = IStreamingManager::Get().GetAnimationStreamingManager();
 
-		const FCompressedAnimSequence* CurveCompressedDataChunk = StreamingManager.GetLoadedChunk(this, 0, false); //Curve Data stored in chunk 0 till it is properly cropped
+		bool bUsingFirstChunk = (ChunkIndex == 0);
+
+		const FCompressedAnimSequence* CurveCompressedDataChunk = StreamingManager.GetLoadedChunk(this, 0, bUsingFirstChunk); //Curve Data stored in chunk 0 till it is properly cropped
 
 		if (!CurveCompressedDataChunk)
 		{
@@ -289,7 +291,7 @@ void UAnimStreamable::GetAnimationPose(FCompactPose& OutPose, FBlendedCurve& Out
 
 		CurveCompressedDataChunk->CurveCompressionCodec->DecompressCurves(*CurveCompressedDataChunk, OutCurve, ExtractionContext.CurrentTime);
 
-		const FCompressedAnimSequence* CompressedData = (ChunkIndex == 0) ? CurveCompressedDataChunk : StreamingManager.GetLoadedChunk(this, ChunkIndex, true);
+		const FCompressedAnimSequence* CompressedData = bUsingFirstChunk ? CurveCompressedDataChunk : StreamingManager.GetLoadedChunk(this, ChunkIndex, true);
 
 		float ChunkCurrentTime = ExtractionContext.CurrentTime - GetRunningPlatformData().Chunks[ChunkIndex].StartTime;
 
diff --git a/Engine/Source/Runtime/Engine/Private/Animation/AnimationStreaming.cpp b/Engine/Source/Runtime/Engine/Private/Animation/AnimationStreaming.cpp
--- a/Engine/Source/Runtime/Engine/Private/Animation/AnimationStreaming.cpp
+++ b/Engine/Source/Runtime/Engine/Private/Animation/AnimationStreaming.cpp
@@ -105,6 +105,25 @@ bool FStreamingAnimationData::UpdateStreamingStatus()
 		return false;
 	}
 
+	//Handle Failed Chunks first
+	for (int32 LoadedChunkIndex = 0; LoadedChunkIndex < LoadedChunks.Num(); ++LoadedChunkIndex)
+	{
+		FLoadedAnimationChunk& LoadedChunk = LoadedChunks[LoadedChunkIndex];
+		if (LoadFailedChunks.Contains(LoadedChunk.Index))
+		{
+			// Mark as not loaded
+			LoadedChunkIndices.Remove(LoadedChunk.Index);
+
+			// Remove this chunk
+			FreeLoadedChunk(LoadedChunk);
+			
+			FScopeLock LoadedChunksLock(&LoadedChunksCritcalSection);
+			LoadedChunks.RemoveAtSwap(LoadedChunkIndex, 1, false);
+		}
+	}
+
+	LoadFailedChunks.Reset();
+
 	bool bHasPendingRequestInFlight = false;
 
 	TArray<uint32> IndicesToLoad;
@@ -176,6 +195,8 @@ void FStreamingAnimationData::BeginPendingRequests(const TArray<uint32>& Indices
 				if (LoadedChunks[ChunkIndex].Index == Index)
 				{
 					FreeLoadedChunk(LoadedChunks[ChunkIndex]);
+
+					FScopeLock LoadedChunksLock(&LoadedChunksCritcalSection);
 					LoadedChunks.RemoveAtSwap(ChunkIndex,1,false);
 					break;
 				}
@@ -185,7 +206,7 @@ void FStreamingAnimationData::BeginPendingRequests(const TArray<uint32>& Indices
 
 	// Set off all IO Requests
 
-	const EAsyncIOPriorityAndFlags AsyncIOPriority = AIOP_High;
+	const EAsyncIOPriorityAndFlags AsyncIOPriority = AIOP_CriticalPath; //Set to Crit temporarily as emergency speculative fix for streaming issue
 
 	for (uint32 ChunkIndex : IndicesToLoad)
 	{
@@ -197,27 +218,27 @@ void FStreamingAnimationData::BeginPendingRequests(const TArray<uint32>& Indices
 
 		if(!ExistingCompressedData)
 		{
-			check(ChunkStorage.CompressedAnimData == nullptr);
-			check(Chunk.BulkData.GetFilename().Len());
+			UE_CLOG(ChunkStorage.CompressedAnimData != nullptr, LogAnimation, Fatal, TEXT("Existing compressed data for streaming animation chunk."));
+			UE_CLOG(Chunk.BulkData.GetFilename().Len()==0, LogAnimation, Fatal, TEXT("Streaming animation chunk has no file name."));
 			UE_CLOG(Chunk.BulkData.IsStoredCompressedOnDisk(), LogAnimation, Fatal, TEXT("Package level compression is not supported for streaming animation."));
-			check(!ChunkStorage.IORequest);
+			UE_CLOG(ChunkStorage.IORequest, LogAnimation, Fatal, TEXT("Streaming animation chunk already has IORequest."));
 			
 			if (!IORequestHandle)
 			{
 				IORequestHandle = FPlatformFileManager::Get().GetPlatformFile().OpenAsyncRead(*Chunk.BulkData.GetFilename());
-				check(IORequestHandle); // this generally cannot fail because it is async
+				UE_CLOG(!IORequestHandle, LogAnimation, Fatal, TEXT("Failed to get IORequest Handle for streaming animation.")); // this generally cannot fail because it is async
 			}
 
 			int64 ChunkSize = Chunk.BulkData.GetBulkDataSize();
 			ChunkStorage.RequestStart = FPlatformTime::Seconds();
-			UE_LOG(LogAnimation, Log, TEXT("Request Started %s Chunk:%i At:%.3f\n"), *StreamableAnim->GetName(), ChunkIndex, ChunkStorage.RequestStart);
+			UE_LOG(LogAnimation, Warning, TEXT("Anim Streaming Request Started %s Chunk:%i At:%.3f\n"), *StreamableAnim->GetName(), ChunkIndex, ChunkStorage.RequestStart);
 			FAsyncFileCallBack AsyncFileCallBack =
 				[this, ChunkIndex, ChunkSize](bool bWasCancelled, IAsyncReadRequest* Req)
 			{
-				AnimationStreamingManager->OnAsyncFileCallback(this, ChunkIndex, ChunkSize, Req);
+				AnimationStreamingManager->OnAsyncFileCallback(this, ChunkIndex, ChunkSize, Req, bWasCancelled);
 			};
 
-			UE_LOG(LogAnimation, Verbose, TEXT("Loading Stream Chunk %s Chunk:%i Length: %.3f Offset:%i Size:%i File:%i\n"), *StreamableAnim->GetName(), ChunkIndex, Chunk.SequenceLength, Chunk.BulkData.GetBulkDataOffsetInFile(), Chunk.BulkData.GetBulkDataSize(), *Chunk.BulkData.GetFilename());
+			UE_LOG(LogAnimation, Warning, TEXT("Loading Stream Anim %s Chunk:%i Length: %.3f Offset:%i Size:%i File:%s\n"), *StreamableAnim->GetName(), ChunkIndex, Chunk.SequenceLength, Chunk.BulkData.GetBulkDataOffsetInFile(), Chunk.BulkData.GetBulkDataSize(), *Chunk.BulkData.GetFilename());
 			ChunkStorage.IORequest = IORequestHandle->ReadRequest(Chunk.BulkData.GetBulkDataOffsetInFile(), Chunk.BulkData.GetBulkDataSize(), AsyncIOPriority, &AsyncFileCallBack);
 			if (!ChunkStorage.IORequest)
 			{
@@ -260,6 +281,7 @@ bool FStreamingAnimationData::BlockTillAllRequestsFinished(float TimeLimit)
 
 FLoadedAnimationChunk& FStreamingAnimationData::AddNewLoadedChunk(uint32 ChunkIndex, FCompressedAnimSequence* ExistingData)
 {
+	FScopeLock LoadedChunksLock(&LoadedChunksCritcalSection);
 	int32 NewIndex = LoadedChunks.Num();
 	LoadedChunks.AddDefaulted();
 
@@ -320,17 +342,23 @@ FAnimationStreamingManager::~FAnimationStreamingManager()
 {
 }
 
-void FAnimationStreamingManager::OnAsyncFileCallback(FStreamingAnimationData* StreamingAnimData, int32 ChunkIndex, int64 ReadSize, IAsyncReadRequest* ReadRequest)
+void FAnimationStreamingManager::OnAsyncFileCallback(FStreamingAnimationData* StreamingAnimData, int32 ChunkIndex, int64 ReadSize, IAsyncReadRequest* ReadRequest, bool bWasCancelled)
 {
 	// Check to see if we successfully managed to load anything
 	uint8* Mem = ReadRequest->GetReadResults();
 
+	FScopeLock Lock(&StreamingAnimData->LoadedChunksCritcalSection);
+
+	const int32 LoadedChunkIndex = StreamingAnimData->LoadedChunks.IndexOfByPredicate([ChunkIndex](const FLoadedAnimationChunk& Chunk) { return Chunk.Index == ChunkIndex; });
+	FLoadedAnimationChunk& ChunkStorage = StreamingAnimData->LoadedChunks[LoadedChunkIndex];
+
+	const double CurrentTime = FPlatformTime::Seconds();
+	const double RequestDuration = CurrentTime - ChunkStorage.RequestStart;
+
 	if (Mem)
 	{
-		int32 LoadedChunkIndex = StreamingAnimData->LoadedChunks.IndexOfByPredicate([ChunkIndex](const FLoadedAnimationChunk& Chunk) { return Chunk.Index == ChunkIndex; });
-		FLoadedAnimationChunk& ChunkStorage = StreamingAnimData->LoadedChunks[LoadedChunkIndex];
 
-		checkf(ChunkStorage.CompressedAnimData == nullptr, TEXT("Chunk storage already has data. (0x%p)"), ChunkStorage.CompressedAnimData.Load());
+		checkf(ChunkStorage.CompressedAnimData == nullptr, TEXT("Chunk storage already has data. (0x%p) ChunkIndex:%i LoadChunkIndex:%i"), ChunkStorage.CompressedAnimData.Load(), ChunkIndex, LoadedChunkIndex);
 
 		FCompressedAnimSequence* NewCompressedData = new FCompressedAnimSequence();
 
@@ -340,9 +368,6 @@ void FAnimationStreamingManager::OnAsyncFileCallback(FStreamingAnimationData* St
 		UAnimStreamable* Anim = StreamingAnimData->StreamableAnim;
 		NewCompressedData->SerializeCompressedData(Reader, false, Anim, Anim->GetSkeleton(), Anim->CurveCompressionSettings);
 
-		const double CurrentTime = FPlatformTime::Seconds();
-		const double RequestDuration = CurrentTime - ChunkStorage.RequestStart;
-
 		const float SequenceLength = StreamingAnimData->StreamableAnim->GetRunningPlatformData().Chunks[ChunkIndex].SequenceLength;
 		if (SequenceLength < RequestDuration)
 		{
@@ -357,6 +382,13 @@ void FAnimationStreamingManager::OnAsyncFileCallback(FStreamingAnimationData* St
 
 		FMemory::Free(Mem);
 	}
+	else
+	{
+		const TCHAR* WasCancelledText = bWasCancelled ? TEXT("Yes") : TEXT("No");
+		UE_LOG(LogAnimation, Warning, TEXT("Streaming anim failed to load chunk: %i Load Duration:%.3f, Anim:%s WasCancelled: %s\n"), ChunkIndex, RequestDuration, *StreamingAnimData->StreamableAnim->GetName(), WasCancelledText);
+		
+		StreamingAnimData->LoadFailedChunks.AddUnique(ChunkIndex);
+	}
 }
 
 SIZE_T FAnimationStreamingManager::GetMemorySizeForAnim(const UAnimStreamable* Anim)
@@ -481,7 +513,7 @@ bool FAnimationStreamingManager::RemoveStreamingAnim(UAnimStreamable* Anim)
 	return false;
 }
 
-const FCompressedAnimSequence* FAnimationStreamingManager::GetLoadedChunk(const UAnimStreamable* Anim, uint32 ChunkIndex, bool bRequestNextChunk) const
+const FCompressedAnimSequence* FAnimationStreamingManager::GetLoadedChunk(const UAnimStreamable* Anim, uint32 ChunkIndex, bool bTrackAsRequested) const
 {
 	// Check for the spoof of failing to load a stream chunk
 	if (SpoofFailedAnimationChunkLoad > 0)
@@ -495,9 +527,9 @@ const FCompressedAnimSequence* FAnimationStreamingManager::GetLoadedChunk(const
 	FStreamingAnimationData* AnimData = StreamingAnimations.FindRef(Anim);
 	if (AnimData)
 	{
-		AnimData->RequestedChunks.AddUnique(ChunkIndex);
-		if (bRequestNextChunk)
+		if (bTrackAsRequested)
 		{
+			AnimData->RequestedChunks.AddUnique(ChunkIndex);
 			AnimData->RequestedChunks.AddUnique((ChunkIndex + 1) % Anim->GetRunningPlatformData().Chunks.Num());
 		}
 
diff --git a/Engine/Source/Runtime/Engine/Public/Animation/AnimationStreaming.h b/Engine/Source/Runtime/Engine/Public/Animation/AnimationStreaming.h
--- a/Engine/Source/Runtime/Engine/Public/Animation/AnimationStreaming.h
+++ b/Engine/Source/Runtime/Engine/Public/Animation/AnimationStreaming.h
@@ -114,6 +114,8 @@ public:
 	/* Contains pointers to Chunks of audio data that have been streamed in */
 	TArray<FLoadedAnimationChunk> LoadedChunks;
 
+	mutable FCriticalSection LoadedChunksCritcalSection;
+
 	class IAsyncReadFileHandle* IORequestHandle;
 
 	/** Indices of chunks that are currently loaded */
@@ -121,6 +123,8 @@ public:
 
 	TArray<uint32> RequestedChunks;
 
+	TArray<uint32> LoadFailedChunks;
+
 	/** Ptr to owning audio streaming manager. */
 	FAnimationStreamingManager* AnimationStreamingManager;
 };
@@ -151,11 +155,11 @@ struct FAnimationStreamingManager : public IAnimationStreamingManager
 	virtual void AddStreamingAnim(UAnimStreamable* Anim) override;
 	virtual bool RemoveStreamingAnim(UAnimStreamable* Anim) override;
 	virtual SIZE_T GetMemorySizeForAnim(const UAnimStreamable* Anim) override;
-	virtual const FCompressedAnimSequence* GetLoadedChunk(const UAnimStreamable* Anim, uint32 ChunkIndex, bool bRequestNextChunk) const override;
+	virtual const FCompressedAnimSequence* GetLoadedChunk(const UAnimStreamable* Anim, uint32 ChunkIndex, bool bTrackAsRequested) const override;
 	// End IAudioStreamingManager interface
 
 	/** Called when an async callback is made on an async loading audio chunk request. */
-	void OnAsyncFileCallback(FStreamingAnimationData* StreamingAnimData, int32 ChunkIndex, int64 ReadSize, IAsyncReadRequest* ReadRequest);
+	void OnAsyncFileCallback(FStreamingAnimationData* StreamingAnimData, int32 ChunkIndex, int64 ReadSize, IAsyncReadRequest* ReadRequest, bool bWasCancelled);
 
 protected:
 
