Fix bad merge

diff --git a/Engine/Source/Runtime/Engine/Private/ADPCMAudioInfo.cpp b/Engine/Source/Runtime/Engine/Private/ADPCMAudioInfo.cpp
--- a/Engine/Source/Runtime/Engine/Private/ADPCMAudioInfo.cpp
+++ b/Engine/Source/Runtime/Engine/Private/ADPCMAudioInfo.cpp
@@ -86,7 +86,7 @@ void FADPCMAudioInfo::SeekToTimeInternal(const float InSeekTime)
 	// Calculate block index & force SeekTime to be in bounds.
 	check(WaveInfo.pSamplesPerSec != nullptr);
 	const uint32 SamplesPerSec = *WaveInfo.pSamplesPerSec;
-	const uint32 SeekedSamples = static_cast<uint32>(InSeekTime * SamplesPerSec);
+	uint32 SeekedSamples = static_cast<uint32>(InSeekTime * SamplesPerSec);
 	TotalSamplesStreamed = FMath::Min<uint32>(SeekedSamples, TotalSamplesPerChannel - 1);
 
 	const uint32 HeaderOffset = static_cast<uint32>(WaveInfo.SampleDataStart - SrcBufferData);
