Use the CSV factory to import the csv or json file to the data table.

diff --git a/Engine/Source/Runtime/Engine/Private/DataTableFunctionLibrary.cpp b/Engine/Source/Runtime/Engine/Private/DataTableFunctionLibrary.cpp
--- a/Engine/Source/Runtime/Engine/Private/DataTableFunctionLibrary.cpp
+++ b/Engine/Source/Runtime/Engine/Private/DataTableFunctionLibrary.cpp
@@ -4,7 +4,10 @@
 #include "Engine/CurveTable.h"
 
 #if WITH_EDITOR
+#include "EditorFramework/AssetImportData.h"
+#include "Factories/CSVImportFactory.h"
 #include "Misc/FileHelper.h"
+#include "HAL/FileManager.h"
 #endif //WITH_EDITOR
 
 UDataTableFunctionLibrary::UDataTableFunctionLibrary(const FObjectInitializer& ObjectInitializer)
@@ -101,91 +104,95 @@ bool UDataTableFunctionLibrary::FillDataTableFromCSVString(UDataTable* DataTable
 {
 	if (!DataTable)
 	{
-		UE_LOG(LogDataTable, Error, TEXT("Can't fill an invalid DataTable."));
+		UE_LOG(LogDataTable, Error, TEXT("FillDataTableFromCSVString - The DataTable is invalid."));
 		return false;
 	}
 
-	bool bResult = true;
-	if (InString.Len() == 0)
-	{
-		DataTable->EmptyTable();
-	}
-	else
-	{
-		TArray<FString> Errors = DataTable->CreateTableFromCSVString(InString);
-		if (Errors.Num())
-		{
-			for (const FString& Error : Errors)
-			{
-				UE_LOG(LogDataTable, Warning, TEXT("%s"), *Error);
-			}
-		}
-		bResult = Errors.Num() == 0;
-	}
-	return bResult;
+	UCSVImportFactory* ImportFactory = NewObject<UCSVImportFactory>();
+
+	bool bWasCancelled = false;
+	const TCHAR* Buffer = *InString;
+	UObject* Result = ImportFactory->FactoryCreateText(DataTable->GetClass()
+		, DataTable->GetOuter()
+		, DataTable->GetFName()
+		, DataTable->GetFlags()
+		, nullptr
+		, TEXT("csv")
+		, Buffer
+		, Buffer + InString.Len()
+		, nullptr
+		, bWasCancelled);
+
+	return Result != nullptr && !bWasCancelled;
 }
 
 bool UDataTableFunctionLibrary::FillDataTableFromCSVFile(UDataTable* DataTable, const FString& InFilePath)
 {
 	if (!DataTable)
 	{
-		UE_LOG(LogDataTable, Error, TEXT("Can't fill an invalid DataTable."));
+		UE_LOG(LogDataTable, Error, TEXT("FillDataTableFromCSVFile - The DataTable is invalid."));
 		return false;
 	}
 
-	FString Data;
-	if (!FFileHelper::LoadFileToString(Data, *InFilePath))
+	if (!IFileManager::Get().FileExists(*InFilePath))
 	{
-		UE_LOG(LogDataTable, Error, TEXT("Can't load the file '%s'."), *InFilePath);
+		UE_LOG(LogDataTable, Error, TEXT("FillDataTableFromCSVFile - The file '%s' doesn't exist."), *InFilePath);
 		return false;
 	}
 
-	return FillDataTableFromCSVString(DataTable, Data);
+	DataTable->AssetImportData->Update(InFilePath);
+	UCSVImportFactory* ImportFactory = NewObject<UCSVImportFactory>();
+	return ImportFactory->ReimportCSV(DataTable) == EReimportResult::Succeeded;
 }
 
 bool UDataTableFunctionLibrary::FillDataTableFromJSONString(UDataTable* DataTable, const FString& InString)
 {
 	if (!DataTable)
 	{
-		UE_LOG(LogDataTable, Error, TEXT("Can't fill an invalid DataTable."));
+		UE_LOG(LogDataTable, Error, TEXT("FillDataTableFromJSONString - The DataTable is invalid."));
 		return false;
 	}
 
-	bool bResult = true;
-	if (InString.Len() == 0)
-	{
-		DataTable->EmptyTable();
-	}
-	else
-	{
-		TArray<FString> Errors = DataTable->CreateTableFromJSONString(InString);
-		if (Errors.Num())
-		{
-			for (const FString& Error : Errors)
-			{
-				UE_LOG(LogDataTable, Warning, TEXT("%s"), *Error);
-			}
-		}
-		bResult = Errors.Num() == 0;
-	}
-	return bResult;
+	UCSVImportFactory* ImportFactory = NewObject<UCSVImportFactory>();
+
+	bool bWasCancelled = false;
+	const TCHAR* Buffer = *InString;
+	UObject* Result = ImportFactory->FactoryCreateText(DataTable->GetClass()
+		, DataTable->GetOuter()
+		, DataTable->GetFName()
+		, DataTable->GetFlags()
+		, nullptr
+		, TEXT("json")
+		, Buffer
+		, Buffer + InString.Len()
+		, nullptr
+		, bWasCancelled);
+
+	return Result != nullptr && !bWasCancelled;
 }
 
 bool UDataTableFunctionLibrary::FillDataTableFromJSONFile(UDataTable* DataTable, const FString& InFilePath)
 {
 	if (!DataTable)
 	{
-		UE_LOG(LogDataTable, Error, TEXT("Can't fill an invalid DataTable."));
+		UE_LOG(LogDataTable, Error, TEXT("FillDataTableFromJSONFile - The DataTable is invalid."));
+		return false;
+	}
+
+	if (!IFileManager::Get().FileExists(*InFilePath))
+	{
+		UE_LOG(LogDataTable, Error, TEXT("FillDataTableFromJSONFile - The file '%s' doesn't exist."), *InFilePath);
 		return false;
 	}
 
-	FString Data;
-	if (!FFileHelper::LoadFileToString(Data, *InFilePath))
+	if (!InFilePath.EndsWith(TEXT(".json")))
 	{
-		UE_LOG(LogDataTable, Error, TEXT("Can't load the file '%s'."), *InFilePath);
+		UE_LOG(LogDataTable, Error, TEXT("FillDataTableFromJSONFile - The file is not a json."));
 		return false;
 	}
 
-	return FillDataTableFromJSONString(DataTable, Data);
+	DataTable->AssetImportData->Update(InFilePath);
+	UCSVImportFactory* ImportFactory = NewObject<UCSVImportFactory>();
+	return ImportFactory->ReimportCSV(DataTable) == EReimportResult::Succeeded;
 }
 #endif //WITH_EDITOR
\ No newline at end of file
