skylight component cubemap texture resolution fix.

diff --git a/Engine/Source/Runtime/Engine/Private/Components/SkyLightComponent.cpp b/Engine/Source/Runtime/Engine/Private/Components/SkyLightComponent.cpp
--- a/Engine/Source/Runtime/Engine/Private/Components/SkyLightComponent.cpp
+++ b/Engine/Source/Runtime/Engine/Private/Components/SkyLightComponent.cpp
@@ -74,6 +74,7 @@ void FSkyTextureCubeResource::InitRHI()
 	if (GetFeatureLevel() >= ERHIFeatureLevel::SM4 || GSupportsRenderTargetFormat_PF_FloatRGBA)
 	{
 		FRHIResourceCreateInfo CreateInfo;
+		checkf(FMath::IsPowerOfTwo(Size), TEXT("Size of SkyTextureCube must be a power of two; size is %d"), Size);
 		TextureCubeRHI = RHICreateTextureCube(Size, Format, NumMips, 0, CreateInfo);
 		TextureRHI = TextureCubeRHI;
 
@@ -479,15 +480,12 @@ void USkyLightComponent::PreEditChange(UProperty* PropertyAboutToChange)
 
 void USkyLightComponent::PostEditChangeProperty(FPropertyChangedEvent& PropertyChangedEvent)
 {
-	Super::PostEditChangeProperty(PropertyChangedEvent);
-
 	if (PropertyChangedEvent.Property && 
-		PropertyChangedEvent.Property->GetFName() == GET_MEMBER_NAME_CHECKED(USkyLightComponent, CubemapResolution) &&
-		PropertyChangedEvent.ChangeType != EPropertyChangeType::Interactive)
+		PropertyChangedEvent.Property->GetFName() == GET_MEMBER_NAME_CHECKED(USkyLightComponent, CubemapResolution))
 	{
 		SanitizeCubemapSize();
 	}
-
+	Super::PostEditChangeProperty(PropertyChangedEvent);
 	SetCaptureIsDirty();
 }
 
