Make FObjectFinder visibile to reference collectors and properly replace subobjects on the CDO. This fixes instability after reloading the FirstPersonCharacter blueprint

diff --git a/Engine/Source/Editor/UnrealEd/Private/Kismet2/KismetReinstanceUtilities.cpp b/Engine/Source/Editor/UnrealEd/Private/Kismet2/KismetReinstanceUtilities.cpp
--- a/Engine/Source/Editor/UnrealEd/Private/Kismet2/KismetReinstanceUtilities.cpp
+++ b/Engine/Source/Editor/UnrealEd/Private/Kismet2/KismetReinstanceUtilities.cpp
@@ -1810,6 +1810,7 @@ static void ReplaceActorHelper(UObject* OldObject, UClass* OldClass, UObject*& N
 
 	UEngine::FCopyPropertiesForUnrelatedObjectsParams Params;
 	Params.bPreserveRootComponent = bPreserveRootComponent;
+	Params.bAggressiveDefaultSubobjectReplacement = true;
 	UEngine::CopyPropertiesForUnrelatedObjects(OldActor, NewActor, Params);
 
 	// reset properties/streams
diff --git a/Engine/Source/Runtime/CoreUObject/Public/UObject/ConstructorHelpers.h b/Engine/Source/Runtime/CoreUObject/Public/UObject/ConstructorHelpers.h
--- a/Engine/Source/Runtime/CoreUObject/Public/UObject/ConstructorHelpers.h
+++ b/Engine/Source/Runtime/CoreUObject/Public/UObject/ConstructorHelpers.h
@@ -96,7 +96,7 @@ struct COREUOBJECT_API ConstructorHelpers
 {
 public:
 	template<class T>
-	struct FObjectFinder
+	struct FObjectFinder : public FGCObject
 	{
 		T* Object;
 		FObjectFinder(const TCHAR* ObjectToFind)
@@ -112,10 +112,15 @@ public:
 		{
 			return !!Object;
 		}
+
+		virtual void AddReferencedObjects( FReferenceCollector& Collector )
+		{
+			Collector.AddReferencedObject(Object);
+		}
 	};
 
 	template<class T>
-	struct FObjectFinderOptional
+	struct FObjectFinderOptional : public FGCObject
 	{
 	private:
 		T* Object;
@@ -145,10 +150,15 @@ public:
 		{
 			return !!Get();
 		}
+
+		virtual void AddReferencedObjects( FReferenceCollector& Collector )
+		{
+			Collector.AddReferencedObject(Object);
+		}
 	};
 
 	template<class T>
-	struct FClassFinder
+	struct FClassFinder : public FGCObject
 	{
 		TSubclassOf<T> Class;
 		FClassFinder(const TCHAR* ClassToFind)
@@ -163,6 +173,13 @@ public:
 		{
 			return !!*Class;
 		}
+		
+		virtual void AddReferencedObjects( FReferenceCollector& Collector )
+		{
+			UClass* ReferencedClass = Class.Get();
+			Collector.AddReferencedObject(ReferencedClass);
+			Class = ReferencedClass;
+		}
 	};
 
 public:
diff --git a/Engine/Source/Runtime/Engine/Private/UnrealEngine.cpp b/Engine/Source/Runtime/Engine/Private/UnrealEngine.cpp
--- a/Engine/Source/Runtime/Engine/Private/UnrealEngine.cpp
+++ b/Engine/Source/Runtime/Engine/Private/UnrealEngine.cpp
@@ -13542,6 +13542,26 @@ public:
 
 void UEngine::CopyPropertiesForUnrelatedObjects(UObject* OldObject, UObject* NewObject, FCopyPropertiesForUnrelatedObjectsParams Params)
 {
+	// UObject::CollectDefaultSubobjects will not return all DSOs for the CDO (only returns inherited subobjects). It also
+	// includes any subobjects that have a matching name on the archetype. This function returns all subobjects that have
+	// a matching instancing on Object's archetype and any objects tagged as RF_DefaultSubObject. Non-DSO instanced subobjects
+	// may be missing, but testing will determine that:
+	const auto CollectAllSubobjects = [](UObject* Object, TArray<UObject*>& OutSubobjectArray)
+	{
+		const bool bIncludedNestedObjects = true;
+		GetObjectsWithOuter(Object, OutSubobjectArray, bIncludedNestedObjects);
+
+		// Remove contained objects that are not subobjects.
+		for ( int32 ComponentIndex = 0; ComponentIndex < OutSubobjectArray.Num(); ComponentIndex++ )
+		{
+			UObject* PotentialComponent = OutSubobjectArray[ComponentIndex];
+			if (!PotentialComponent->IsDefaultSubobject() && !PotentialComponent->HasAnyFlags(RF_DefaultSubObject))
+			{
+				OutSubobjectArray.RemoveAtSwap(ComponentIndex--);
+			}
+		}
+	};
+
 	check(OldObject && NewObject);
 
 	// Bad idea to write data to an actor while its components are registered
@@ -13575,7 +13595,7 @@ void UEngine::CopyPropertiesForUnrelatedObjects(UObject* OldObject, UObject* New
 	{
 		// Find all instanced objects of the old CDO, and save off their modified properties to be later applied to the newly instanced objects of the new CDO
 		TArray<UObject*> Components;
-		OldObject->CollectDefaultSubobjects(Components, true);
+		CollectAllSubobjects( OldObject, Components );
 
 		for (UObject* OldInstance : Components)
 		{
@@ -13600,7 +13620,7 @@ void UEngine::CopyPropertiesForUnrelatedObjects(UObject* OldObject, UObject* New
 	TArray<UObject*> ComponentsOnNewObject;
 	{
 		TArray<UObject*> EditInlineSubobjectsOfComponents;
-		NewObject->CollectDefaultSubobjects(ComponentsOnNewObject,true);
+		CollectAllSubobjects( NewObject, ComponentsOnNewObject );
 
 		// populate the ReferenceReplacementMap 
 		for (int32 Index = 0; Index < ComponentsOnNewObject.Num(); Index++)
