Move OnPossess and OnUnPossess blueprint implementable events from AIController to Controller.
To accomplish this Possess/UnPossess are no longer virtual, instead subclasses of Controller must override OnPossess/OnUnPossess.
OnPossess/OnUnPossess blueprint implementable events renamed ReceivePossess/ReceiveUnPossess for native purposes, but use display name to keep events the same in bluepritns.

diff --git a/Engine/Config/BaseEngine.ini b/Engine/Config/BaseEngine.ini
--- a/Engine/Config/BaseEngine.ini
+++ b/Engine/Config/BaseEngine.ini
@@ -614,6 +614,9 @@ bUseStreamingPause=false
 +FunctionRedirects=(OldName="GetHandGestureConfidence",NewName="/Script/MagicLeapHandTracking.MagicLeapHandTrackingFunctionLibrary.GetCurrentGestureConfidence")
 +FunctionRedirects=(OldName="GetCurrentGesture",NewName="/Script/MagicLeapHandTracking.MagicLeapHandTrackingFunctionLibrary.GetCurrentGesture")
 
++FunctionRedirects=(OldName="Controller.OnPossess",NewName="Controller.ReceivePossess")
++FunctionRedirects=(OldName="Controller.OnUnPossess",NewName="Controller.ReceiveUnPossess")
+
 [CoreUObject.Metadata]
 ; Note: UnrealHeaderTool should be rerun after making changes to MetadataRedirects to catch any keys specified in class headers
 +MetadataRedirects=(OldKey="K2Protected", NewKey="BlueprintProtected")
diff --git a/Engine/Source/Runtime/AIModule/Classes/AIController.h b/Engine/Source/Runtime/AIModule/Classes/AIController.h
--- a/Engine/Source/Runtime/AIModule/Classes/AIController.h
+++ b/Engine/Source/Runtime/AIModule/Classes/AIController.h
@@ -150,14 +150,6 @@ public:
 
 	AAIController(const FObjectInitializer& ObjectInitializer = FObjectInitializer::Get());
 
-	/** Event called when PossessedPawn is possessed by this controller. */
-	UFUNCTION(BlueprintImplementableEvent)
-	void OnPossess(APawn* PossessedPawn);
-
-	/** Gets triggered after given pawn has been unpossesed */
-	UFUNCTION(BlueprintImplementableEvent)
-	void OnUnpossess(APawn* UnpossessedPawn);
-
 	virtual void SetPawn(APawn* InPawn) override;
 
 	/** Makes AI go toward specified Goal actor (destination will be continuously updated), aborts any active path following
@@ -334,8 +326,11 @@ public:
 	//~ End AActor Interface
 
 	//~ Begin AController Interface
-	virtual void Possess(APawn* InPawn) override;
-	virtual void UnPossess() override;
+protected:
+	virtual void OnPossess(APawn* InPawn) override;
+	virtual void OnUnPossess() override;
+
+public:
 	virtual bool ShouldPostponePathUpdates() const override;
 	virtual void DisplayDebug(UCanvas* Canvas, const FDebugDisplayInfo& DebugDisplay, float& YL, float& YPos) override;
 
diff --git a/Engine/Source/Runtime/AIModule/Private/AIController.cpp b/Engine/Source/Runtime/AIModule/Private/AIController.cpp
--- a/Engine/Source/Runtime/AIModule/Private/AIController.cpp
+++ b/Engine/Source/Runtime/AIModule/Private/AIController.cpp
@@ -454,7 +454,7 @@ void AAIController::UpdateControlRotation(float DeltaTime, bool bUpdatePawn)
 }
 
 
-void AAIController::Possess(APawn* InPawn)
+void AAIController::OnPossess(APawn* InPawn)
 {
 	// don't even try possessing pending-kill pawns
 	if (InPawn != nullptr && InPawn->IsPendingKill())
@@ -462,7 +462,7 @@ void AAIController::Possess(APawn* InPawn)
 		return;
 	}
 
-	Super::Possess(InPawn);
+	Super::OnPossess(InPawn);
 
 	if (GetPawn() == nullptr || InPawn == nullptr)
 	{
@@ -501,15 +501,13 @@ void AAIController::Possess(APawn* InPawn)
 
 		REDIRECT_OBJECT_TO_VLOG(CachedGameplayTasksComponent, this);
 	}
-
-	OnPossess(InPawn);
 }
 
-void AAIController::UnPossess()
+void AAIController::OnUnPossess()
 {
 	APawn* CurrentPawn = GetPawn();
 
-	Super::UnPossess();
+	Super::OnUnPossess();
 
 	if (PathFollowingComponent)
 	{
@@ -529,8 +527,6 @@ void AAIController::UnPossess()
 		CachedGameplayTasksComponent->OnClaimedResourcesChange.RemoveDynamic(this, &AAIController::OnGameplayTaskResourcesClaimed);
 		CachedGameplayTasksComponent = nullptr;
 	}
-
-	OnUnpossess(CurrentPawn);
 }
 
 void AAIController::SetPawn(APawn* InPawn)
diff --git a/Engine/Source/Runtime/Engine/Classes/GameFramework/Controller.h b/Engine/Source/Runtime/Engine/Classes/GameFramework/Controller.h
--- a/Engine/Source/Runtime/Engine/Classes/GameFramework/Controller.h
+++ b/Engine/Source/Runtime/Engine/Classes/GameFramework/Controller.h
@@ -93,11 +93,13 @@ protected:
 	 * update only some components of the rotation).
 	 */
 	UPROPERTY(EditDefaultsOnly, AdvancedDisplay, Category="Controller|Transform")
-	uint32 bAttachToPawn:1;
+	uint8 bAttachToPawn:1;
 
 	/** Whether this controller is a PlayerController. */
-	UPROPERTY()
-	uint32 bIsPlayerController:1;
+	uint8 bIsPlayerController:1;
+
+	/** Whether the controller must have authority to be able to call possess on a Pawn */
+	uint8 bCanPossessWithoutAuthority:1;
 
 	/** Ignores movement input. Stacked state storage, Use accessor function IgnoreMoveInput() */
 	uint8 IgnoreMoveInput;
@@ -251,12 +253,28 @@ public:
 	 * @see HasAuthority()
 	 */
 	UFUNCTION(BlueprintCallable, BlueprintAuthorityOnly, Category=Pawn, meta=(Keywords="set controller"))
-	virtual void Possess(APawn* InPawn);
+	virtual void Possess(APawn* InPawn) final; // DEPRECATED(4.22, "Posssess is marked virtual final as you should now be overriding OnPossess instead")
 
 	/** Called to unpossess our pawn for any reason that is not the pawn being destroyed (destruction handled by PawnDestroyed()). */
 	UFUNCTION(BlueprintCallable, Category=Pawn, meta=(Keywords="set controller"))
-	virtual void UnPossess();
+	virtual void UnPossess() final; // DEPRECATED(4.22, "Posssess is marked virtual final as you should now be overriding OnUnPossess instead")
+
+protected:
+	/** Blueprint implementable event to react to the controller possessing a pawn */
+	UFUNCTION(BlueprintImplementableEvent, meta = (DisplayName = "On Possess"))
+	void ReceivePossess(APawn* PossessedPawn);
 
+	/** Overridable native function for when this controller possesses a pawn. */
+	virtual void OnPossess(APawn* InPawn);
+
+	/** Blueprint implementable event to react to the controller unpossessing a pawn */
+	UFUNCTION(BlueprintImplementableEvent, meta = (DisplayName = "On UnPossess"))
+	void ReceiveUnPossess(APawn* UnpossessedPawn);
+
+	/** Overridable native function for when this controller unpossesses its pawn. */
+	virtual void OnUnPossess();
+
+public:
 	/**
 	 * Called to unpossess our pawn because it is going to be destroyed.
 	 * (other unpossession handled by UnPossess())
diff --git a/Engine/Source/Runtime/Engine/Classes/GameFramework/PlayerController.h b/Engine/Source/Runtime/Engine/Classes/GameFramework/PlayerController.h
--- a/Engine/Source/Runtime/Engine/Classes/GameFramework/PlayerController.h
+++ b/Engine/Source/Runtime/Engine/Classes/GameFramework/PlayerController.h
@@ -1497,8 +1497,10 @@ public:
 	virtual bool IsNetRelevantFor(const AActor* RealViewer, const AActor* ViewTarget, const FVector& SrcLocation) const override;
 	virtual void FellOutOfWorld(const class UDamageType& dmgType) override;
 	virtual void Reset() override;
-	virtual void Possess(APawn* aPawn) override;
-	virtual void UnPossess() override;
+protected:
+	virtual void OnPossess(APawn* aPawn) override;
+	virtual void OnUnPossess() override;
+public:
 	virtual void CleanupPlayerState() override;
 	virtual void EndPlay(const EEndPlayReason::Type EndPlayReason) override;
 	virtual void Destroyed() override;
diff --git a/Engine/Source/Runtime/Engine/Private/Controller.cpp b/Engine/Source/Runtime/Engine/Private/Controller.cpp
--- a/Engine/Source/Runtime/Engine/Private/Controller.cpp
+++ b/Engine/Source/Runtime/Engine/Private/Controller.cpp
@@ -24,6 +24,7 @@
 
 #include "GameFramework/PlayerState.h"
 
+DEFINE_LOG_CATEGORY(LogController);
 DEFINE_LOG_CATEGORY(LogPath);
 
 #define LOCTEXT_NAMESPACE "Controller"
@@ -46,6 +47,7 @@ AController::AController(const FObjectInitializer& ObjectInitializer)
 	bCanBeDamaged = false;
 	bAttachToPawn = false;
 	bIsPlayerController = false;
+	bCanPossessWithoutAuthority = false;
 
 	if (RootComponent)
 	{
@@ -274,17 +276,25 @@ void AController::PostInitializeComponents()
 
 void AController::Possess(APawn* InPawn)
 {
-	if (!HasAuthority())
+	if (!bCanPossessWithoutAuthority && !HasAuthority())
 	{
 		FMessageLog("PIE").Warning(FText::Format(
 			LOCTEXT("ControllerPossessAuthorityOnly", "Possess function should only be used by the network authority for {0}"),
 			FText::FromName(GetFName())
 			));
+		UE_LOG(LogController, Warning, TEXT("Trying to possess %s without network authority! Request will be ignored."), *GetNameSafe(InPawn));
 		return;
 	}
 
 	REDIRECT_OBJECT_TO_VLOG(InPawn, this);
 
+	OnPossess(InPawn);
+
+	ReceivePossess(InPawn);
+}
+
+void AController::OnPossess(APawn* InPawn)
+{
 	const bool bNewPawn = GetPawn() != InPawn;
 
 	if (InPawn != NULL)
@@ -315,6 +325,15 @@ void AController::Possess(APawn* InPawn)
 }
 
 void AController::UnPossess()
+{
+	APawn* CurrentPawn = GetPawn();
+
+	OnUnPossess();
+
+	ReceiveUnPossess(CurrentPawn);
+}
+
+void AController::OnUnPossess()
 {
 	if ( Pawn != NULL )
 	{
@@ -323,12 +342,11 @@ void AController::UnPossess()
 	}
 }
 
-
 void AController::PawnPendingDestroy(APawn* inPawn)
 {
 	if ( IsInState(NAME_Inactive) )
 	{
-		UE_LOG(LogPath, Log, TEXT("PawnPendingDestroy while inactive %s"), *GetName());
+		UE_LOG(LogController, Log, TEXT("PawnPendingDestroy while inactive %s"), *GetName());
 	}
 
 	if ( inPawn != Pawn )
diff --git a/Engine/Source/Runtime/Engine/Private/PlayerController.cpp b/Engine/Source/Runtime/Engine/Private/PlayerController.cpp
--- a/Engine/Source/Runtime/Engine/Private/PlayerController.cpp
+++ b/Engine/Source/Runtime/Engine/Private/PlayerController.cpp
@@ -747,18 +747,8 @@ void APlayerController::ClientRestart_Implementation(APawn* NewPawn)
 
 /// @endcond
 
-void APlayerController::Possess(APawn* PawnToPossess)
+void APlayerController::OnPossess(APawn* PawnToPossess)
 {
-	if (!HasAuthority())
-	{
-		FMessageLog("PIE").Warning(FText::Format(
-			LOCTEXT("PlayerControllerPossessAuthorityOnly", "Possess function should only be used by the network authority for {0}"),
-			FText::FromName(GetFName())
-			));
-		UE_LOG(LogPlayerController, Warning, TEXT("Trying to possess %s without network authority! Request will be ignored."), *GetNameSafe(PawnToPossess));
-		return;
-	}
-
 	if ( PawnToPossess != NULL && 
 		(PlayerState == NULL || !PlayerState->bOnlySpectator) )
 	{
@@ -1218,7 +1208,7 @@ bool APlayerController::ServerAcknowledgePossession_Validate(APawn* P)
 
 /// @endcond
 
-void APlayerController::UnPossess()
+void APlayerController::OnUnPossess()
 {
 	if (GetPawn() != NULL)
 	{
diff --git a/Engine/Source/Runtime/Engine/Public/EngineLogs.h b/Engine/Source/Runtime/Engine/Public/EngineLogs.h
--- a/Engine/Source/Runtime/Engine/Public/EngineLogs.h
+++ b/Engine/Source/Runtime/Engine/Public/EngineLogs.h
@@ -9,6 +9,7 @@
 class Error;
 
 ENGINE_API DECLARE_LOG_CATEGORY_EXTERN(LogPath, Warning, All);
+ENGINE_API DECLARE_LOG_CATEGORY_EXTERN(LogController, Warning, All);
 ENGINE_API DECLARE_LOG_CATEGORY_EXTERN(LogPhysics, Warning, All);
 ENGINE_API DECLARE_LOG_CATEGORY_EXTERN(LogBlueprint, Warning, All);
 ENGINE_API DECLARE_LOG_CATEGORY_EXTERN(LogBlueprintUserMessages, Log, All);
