Fix potential data loss in inherited SCS component default value overrides at runtime in a cooked build for an instance of a non-nativized child Blueprint class that inherits from a nativized parent Blueprint class.

Notes:
- This change fixes a regression that was introduced into 4.21 by 4286567.
- This change also ensures that all Blueprint-owned component archetype objects will get serialized before generating "fast path" component instancing data when the optimization is enabled alongside EDL; as a result, the "fast path" Blueprint component instancing feature was previously incompatible with EDL in certain cases.

Change summary:
- Renamed the defunct UClass::GetRequiredPreloadDependencies() API to GetDefaultObjectPreloadDependencies().
- Modified UObject::GetPreloadDependencies() to call GetDefaultObjectPreloadDependencies() on the UClass when the object is the CDO.
- Added/implemented UBlueprintGeneratedClass::GetDefaultObjectPreloadDependencies(); this now injects SCS, inherited SCS (ICH) and UCS (AddComponent) templates as preload dependencies on the Blueprint CDO.
- Modified UBlueprintGeneratedClass::SerializeDefaultObject() to handle generation of "fast path" component instancing data prior to serialization of instances of Blueprint-owned or inherited component archetype objects.
- Moved the CheckAndApplyTemplateOverrides() call from UBlueprintGeneratedClass::PostLoadDefaultObject() to SerializeDefaultObject(). This now occurs immediately after the CDO is serialized, but before any instances of the Blueprint class are serialized. This ensures that defaults are properly initialized before we do any delta serialization on instances of a non-nativized Blueprint class in a nativized, cooked build.
- Renamed public FBlueprintCookedComponentInstancingData APIs to be more concise/correct.

diff --git a/Engine/Source/Runtime/CoreUObject/Private/UObject/Obj.cpp b/Engine/Source/Runtime/CoreUObject/Private/UObject/Obj.cpp
--- a/Engine/Source/Runtime/CoreUObject/Private/UObject/Obj.cpp
+++ b/Engine/Source/Runtime/CoreUObject/Private/UObject/Obj.cpp
@@ -1180,7 +1180,11 @@ void UObject::GetPreloadDependencies(TArray<UObject*>& OutDeps)
 	{
 		OutDeps.Add(ObjClass);
 
-		if (!HasAnyFlags(RF_ClassDefaultObject) && ObjClass->GetDefaultsCount() > 0)
+		if (HasAnyFlags(RF_ClassDefaultObject))
+		{
+			ObjClass->GetDefaultObjectPreloadDependencies(OutDeps);
+		}
+		else if (ObjClass->GetDefaultsCount() > 0)
 		{
 			OutDeps.Add(ObjClass->GetDefaultObject());
 		}
diff --git a/Engine/Source/Runtime/CoreUObject/Private/UObject/SavePackage.cpp b/Engine/Source/Runtime/CoreUObject/Private/UObject/SavePackage.cpp
--- a/Engine/Source/Runtime/CoreUObject/Private/UObject/SavePackage.cpp
+++ b/Engine/Source/Runtime/CoreUObject/Private/UObject/SavePackage.cpp
@@ -1933,7 +1933,6 @@ class FExportReferenceSorter : public FArchiveUObject
 		FScopeLock ScopeLock(&InitializeCoreClassesCritSec);
 		check(CoreClasses.Num() == 0);
 		check(ReferencedObjects.Num() == 0);
-		check(ForceLoadObjects.Num() == 0);
 		check(SerializedObjects.Num() == 0);
 		check(bIgnoreFieldReferences == false);
 
@@ -1941,7 +1940,6 @@ class FExportReferenceSorter : public FArchiveUObject
 		static TArray<UClass*> StaticCoreClasses;
 		static TArray<UObject*> StaticCoreReferencedObjects;
 		static TArray<UObject*> StaticProcessedObjects;
-		static TArray<UObject*> StaticForceLoadObjects;
 		static TSet<UObject*> StaticSerializedObjects;
 		
 		
@@ -2031,7 +2029,6 @@ class FExportReferenceSorter : public FArchiveUObject
 				check(CoreClasses.Num() == StaticCoreClasses.Num());
 				check(ReferencedObjects.Num() == StaticCoreReferencedObjects.Num());
 				check(ProcessedObjects.Num() == StaticProcessedObjects.Num());
-				check(ForceLoadObjects.Num() == StaticForceLoadObjects.Num());
 				check(SerializedObjects.Num() == StaticSerializedObjects.Num());
 				
 				
@@ -2047,10 +2044,6 @@ class FExportReferenceSorter : public FArchiveUObject
 				{
 					check(ProcessedObject.Value == StaticProcessedObjects.Find(ProcessedObject.Key));
 				}
-				for (int I = 0; I < ForceLoadObjects.Num(); ++I)
-				{
-					check(ForceLoadObjects[I] == StaticForceLoadObjects[I]);
-				}
 				for (const auto& SerializedObject : SerializedObjects)
 				{
 					check(StaticSerializedObjects.Find(SerializedObject));
@@ -2061,7 +2054,6 @@ class FExportReferenceSorter : public FArchiveUObject
 			StaticCoreClasses = CoreClasses;
 			StaticCoreReferencedObjects = ReferencedObjects;
 			StaticProcessedObjects = ProcessedObjects;
-			StaticForceLoadObjects = ForceLoadObjects;
 			StaticSerializedObjects = SerializedObjects;
 
 			check(CurrentClass == nullptr);
@@ -2072,7 +2064,6 @@ class FExportReferenceSorter : public FArchiveUObject
 			CoreClasses = StaticCoreClasses;
 			ReferencedObjects = StaticCoreReferencedObjects;
 			ProcessedObjects = StaticProcessedObjects;
-			ForceLoadObjects = StaticForceLoadObjects;
 			SerializedObjects = StaticSerializedObjects;
 
 			CoreReferencesOffset = StaticCoreReferencedObjects.Num();
@@ -2372,19 +2363,7 @@ public:
 				}
 				else
 				{
-					// since normal references to objects aren't force-loaded, 
-					// we do not need to pass true for bProcessObject by default
-					// (true would indicate that Object must be inserted into 
-					// the sorted export list before the object that contains 
-					// this object reference - i.e. the object we're currently
-					// serializing)
-					// 
-					// sometimes (rarely) this is the case though, so we use 
-					// ForceLoadObjects to determine if the object we're 
-					// serializing would force load Object (if so, it'll come 
-					// first in the ExportMap)
-					bool const bProcessObject = ForceLoadObjects.Contains(Object);
-					HandleDependency(Object, bProcessObject);
+					HandleDependency(Object);
 				}
 			}
 		}
@@ -2483,36 +2462,11 @@ public:
 					// so we turn off field serialization so that we don't have to worry about handling this struct's fields just yet
 					bIgnoreFieldReferences = true;
 
-					// most often, we don't want/need object references getting  
-					// recorded as dependencies, but some structs (classes) 
-					// require certain non-field objects be prioritized in the 
-					// ExportMap earlier (see UClass::GetRequiredPreloadDependencies() 
-					// for more details)... this array records/holds those 
-					// required sub-objects
-					TArray<UObject*> StructForceLoadObjects;
-
 					bool const bIsClassObject = (dynamic_cast<UClass*>(StructObject) != nullptr);
-					if (bIsClassObject)
-					{
-						UClass* AsClass = (UClass*)StructObject;
-						AsClass->GetRequiredPreloadDependencies(StructForceLoadObjects);
-						check(!StructForceLoadObjects.Num()); //@todoio GetRequiredPreloadDependencies is dead code, remove
-					}
-					// append rather than replace (in case we're nested in a 
-					// recursive call)... adding these to ForceLoadObjects 
-					// ensures that any reference to a StructForceLoadObjects
-					// object gets recorded in the ExportMap before StructObject
-					// (see operator<<, where we utilize ForceLoadObjects)
-					ForceLoadObjects.Append(StructForceLoadObjects);
-					int32 const ForceLoadCount = StructForceLoadObjects.Num();
 
 					SerializedObjects.Add(StructObject);
 					StructObject->Serialize(*this);
 
-					// remove (pop) rather than empty, in case ClassForceLoadObjects
-					// had entries from a previous call to this function up that chain
-					ForceLoadObjects.RemoveAt(ForceLoadObjects.Num() - ForceLoadCount, ForceLoadCount);
-
 					// at this point, any objects which were referenced through this struct's script or defaults will be in the list of exports, and 
 					// the CurrentInsertIndex will have been advanced so that the object processed will be inserted just before this struct in the array
 					// (i.e. just after class/archetypes for any objects which were referenced by this struct's script)
@@ -2654,13 +2608,6 @@ private:
 	 */
 	UClass* CurrentClass;
 
-	/** 
-	 * This is a list of objects that would be force loaded by a struct/class 
-	 * currently being handled by ProcessStruct() (meaning that they should be
-	 * prioritized in the target ExportMap, before the struct).
-	 */
-	TArray<UObject*> ForceLoadObjects;
-
 	/** Package to constrain checks to */
 	UPackage* PackageToSort;
 };
diff --git a/Engine/Source/Runtime/CoreUObject/Public/UObject/Class.h b/Engine/Source/Runtime/CoreUObject/Public/UObject/Class.h
--- a/Engine/Source/Runtime/CoreUObject/Public/UObject/Class.h
+++ b/Engine/Source/Runtime/CoreUObject/Public/UObject/Class.h
@@ -2669,20 +2669,11 @@ public:
 	virtual UObject* GetArchetypeForCDO() const;
 
 	/**
-	 * On save, we order a package's exports in class dependency order (so that
-	 * on load, we create the class dependencies before we create the class).  
-	 * More often than not, the class doesn't require any non-struct objects 
-	 * before it is created/serialized (only super-classes, and its UField 
-	 * members, see FExportReferenceSorter::operator<<() for reference). 
-	 * However, in some special occasions, there might be an export that we 
-	 * would like force loaded prior to the class's serialization (like  
-	 * component templates for blueprint classes). This function returns a list
-	 * of those non-struct dependencies, so that FExportReferenceSorter knows to 
-	 * prioritize them earlier in the ExportMap.
-	 * 
-	 * @param  DependenciesOut	Will be filled with a list of dependencies that need to be created before this class is recreated (on load).
-	 */
-	virtual void GetRequiredPreloadDependencies(TArray<UObject*>& DependenciesOut) {}
+	* Returns all objects that should be preloaded before the class default object is serialized at load time. Only used by the EDL.
+	*
+	* @param OutDeps		All objects that should be preloaded before the class default object is serialized at load time.
+	*/
+	virtual void GetDefaultObjectPreloadDependencies(TArray<UObject*>& OutDeps) {}
 
 	/**
 	 * Initializes the ClassReps and NetFields arrays used by replication.
diff --git a/Engine/Source/Runtime/Engine/Classes/Engine/BlueprintGeneratedClass.h b/Engine/Source/Runtime/Engine/Classes/Engine/BlueprintGeneratedClass.h
--- a/Engine/Source/Runtime/Engine/Classes/Engine/BlueprintGeneratedClass.h
+++ b/Engine/Source/Runtime/Engine/Classes/Engine/BlueprintGeneratedClass.h
@@ -562,13 +562,13 @@ struct ENGINE_API FBlueprintCookedComponentInstancingData
 	}
 
 	/** Builds/returns the internal property list that's used for serialization. This is a linked list of UProperty references. */
-	const FCustomPropertyListNode* GetCachedPropertyListForSerialization() const;
+	const FCustomPropertyListNode* GetCachedPropertyList() const;
 
 	/** Called at load time to generate the internal cached property data stream from serialization of the source template object. */
-	void LoadCachedPropertyDataForSerialization(UActorComponent* SourceTemplate);
+	void BuildCachedPropertyDataFromTemplate(UActorComponent* SourceTemplate);
 
 	/** Returns the internal property data stream that's used for fast binary object serialization when instancing components at runtime. */
-	const TArray<uint8>& GetCachedPropertyDataForSerialization() const { return CachedPropertyDataForSerialization; }
+	const TArray<uint8>& GetCachedPropertyData() const { return CachedPropertyData; }
 
 protected:
 	/** Internal method used to help recursively build the cached property list for serialization. */
@@ -582,7 +582,7 @@ private:
 	mutable TIndirectArray<FCustomPropertyListNode> CachedPropertyListForSerialization;
 
 	/** Internal property data stream that's used in binary object serialization at component instancing time. */
-	TArray<uint8> CachedPropertyDataForSerialization;
+	TArray<uint8> CachedPropertyData;
 };
 
 UCLASS()
@@ -693,7 +693,7 @@ public:
 	virtual void Link(FArchive& Ar, bool bRelinkExistingProperties) override;
 	virtual void PurgeClass(bool bRecompilingOnLoad) override;
 	virtual void Bind() override;
-	virtual void GetRequiredPreloadDependencies(TArray<UObject*>& DependenciesOut) override;
+	virtual void GetDefaultObjectPreloadDependencies(TArray<UObject*>& OutDeps) override;
 	virtual UObject* FindArchetype(UClass* ArchetypeClass, const FName ArchetypeName) const override;
 
 	virtual void InitPropertiesFromCustomList(uint8* DataPtr, const uint8* DefaultDataPtr) override;
diff --git a/Engine/Source/Runtime/Engine/Private/ActorConstruction.cpp b/Engine/Source/Runtime/Engine/Private/ActorConstruction.cpp
--- a/Engine/Source/Runtime/Engine/Private/ActorConstruction.cpp
+++ b/Engine/Source/Runtime/Engine/Private/ActorConstruction.cpp
@@ -1016,7 +1016,7 @@ UActorComponent* AActor::CreateComponentFromTemplateData(const FBlueprintCookedC
 		NewActorComp->SetFlags(RF_NeedPostLoad | RF_NeedPostLoadSubobjects);
 
 		// Load cached data into the new instance.
-		FBlueprintComponentInstanceDataLoader ComponentInstanceDataLoader(TemplateData->GetCachedPropertyDataForSerialization(), TemplateData->GetCachedPropertyListForSerialization());
+		FBlueprintComponentInstanceDataLoader ComponentInstanceDataLoader(TemplateData->GetCachedPropertyData(), TemplateData->GetCachedPropertyList());
 		NewActorComp->Serialize(ComponentInstanceDataLoader);
 
 		// Handle tasks that would normally occur post-duplication w/ SDO.
diff --git a/Engine/Source/Runtime/Engine/Private/BlueprintGeneratedClass.cpp b/Engine/Source/Runtime/Engine/Private/BlueprintGeneratedClass.cpp
--- a/Engine/Source/Runtime/Engine/Private/BlueprintGeneratedClass.cpp
+++ b/Engine/Source/Runtime/Engine/Private/BlueprintGeneratedClass.cpp
@@ -119,52 +119,9 @@ void UBlueprintGeneratedClass::PostLoad()
 #endif
 #endif // WITH_EDITORONLY_DATA
 
-	// Generate "fast path" instancing data for UCS/AddComponent node templates.
-	if (CookedComponentInstancingData.Num() > 0)
-	{
-		for (int32 Index = ComponentTemplates.Num() - 1; Index >= 0; --Index)
-		{
-			if (UActorComponent* ComponentTemplate = ComponentTemplates[Index])
-			{
-				if (FBlueprintCookedComponentInstancingData* ComponentInstancingData = CookedComponentInstancingData.Find(ComponentTemplate->GetFName()))
-				{
-					ComponentInstancingData->LoadCachedPropertyDataForSerialization(ComponentTemplate);
-				}
-			}
-		}
-	}
-
 	AssembleReferenceTokenStream(true);
 }
 
-void UBlueprintGeneratedClass::GetRequiredPreloadDependencies(TArray<UObject*>& DependenciesOut)
-{
-	Super::GetRequiredPreloadDependencies(DependenciesOut);
-
-	// the component templates are no longer needed as Preload() dependencies 
-	// (FLinkerLoad now handles these with placeholder export objects instead)...
-	// this change was prompted by a cyclic case, where creating the first
-	// component-template tripped the serialization of its class outer, before 
-	// another second component-template could be created (even though the 
-	// second component was listed in the ExportMap before the class)
-// 	for (UActorComponent* Component : ComponentTemplates)
-// 	{
-// 		// because of the linker's way of handling circular dependencies (with 
-// 		// placeholder blueprint classes), we need to ensure that class owned 
-// 		// blueprint components are created before the class's  is 
-// 		// ComponentTemplates member serialized in (otherwise, the component 
-// 		// would be created as a ULinkerPlaceholderClass instance)
-// 		//
-// 		// by returning these in the DependenciesOut array, we're making it 
-// 		// known that they should be prioritized in the package's ExportMap 
-// 		// before this class
-// 		if (Component->GetClass()->HasAnyClassFlags(CLASS_CompiledFromBlueprint))
-// 		{
-// 			DependenciesOut.Add(Component);
-// 		}
-// 	}
-}
-
 FPrimaryAssetId UBlueprintGeneratedClass::GetPrimaryAssetId() const
 {
 	FPrimaryAssetId AssetId;
@@ -345,6 +302,54 @@ void UBlueprintGeneratedClass::SerializeDefaultObject(UObject* Object, FArchive&
 		// On load, build the custom property list used in post-construct initialization logic. Note that in the editor, this will be refreshed during compile-on-load.
 		// @TODO - Potentially make this serializable (or cooked data) to eliminate the slight load time cost we'll incur below to generate this list in a cooked build. For now, it's not serialized since the raw UProperty references cannot be saved out.
 		UpdateCustomPropertyListForPostConstruction();
+
+		// Generate "fast path" instancing data for inherited SCS node templates.
+		if (InheritableComponentHandler)
+		{
+			for (auto RecordIt = InheritableComponentHandler->CreateRecordIterator(); RecordIt; ++RecordIt)
+			{
+				if (RecordIt->ComponentTemplate && RecordIt->CookedComponentInstancingData.bIsValid)
+				{
+					RecordIt->CookedComponentInstancingData.BuildCachedPropertyDataFromTemplate(RecordIt->ComponentTemplate);
+				}
+			}
+		}
+
+		// Generate "fast path" instancing data for SCS node templates owned by this Blueprint class.
+		if (SimpleConstructionScript)
+		{
+			const TArray<USCS_Node*>& AllSCSNodes = SimpleConstructionScript->GetAllNodes();
+			for (USCS_Node* SCSNode : AllSCSNodes)
+			{
+				if (SCSNode->ComponentTemplate && SCSNode->CookedComponentInstancingData.bIsValid)
+				{
+					SCSNode->CookedComponentInstancingData.BuildCachedPropertyDataFromTemplate(SCSNode->ComponentTemplate);
+				}
+			}
+		}
+
+		// Generate "fast path" instancing data for UCS/AddComponent node templates.
+		if (CookedComponentInstancingData.Num() > 0)
+		{
+			for (UActorComponent* ComponentTemplate : ComponentTemplates)
+			{
+				if (ComponentTemplate)
+				{
+					if (FBlueprintCookedComponentInstancingData* ComponentInstancingData = CookedComponentInstancingData.Find(ComponentTemplate->GetFName()))
+					{
+						ComponentInstancingData->BuildCachedPropertyDataFromTemplate(ComponentTemplate);
+					}
+				}
+			}
+		}
+
+		// We may need to manually apply default value overrides to some inherited components in a cooked build
+		// scenario. This can occur if we have a nativized Blueprint class somewhere in the parent class ancestry.
+		// Note: This must occur AFTER component templates are loaded, but BEFORE component instances are serialized.
+		if (bHasNativizedParent)
+		{
+			CheckAndApplyComponentTemplateOverrides(ClassDefaultObject);
+		}
 	}
 }
 
@@ -364,13 +369,6 @@ void UBlueprintGeneratedClass::PostLoadDefaultObject(UObject* Object)
 		{
 			ClassDefaultObject->LoadConfig();
 		}
-
-		// We may need to manually apply default value overrides to some inherited components in a cooked build
-		// scenario. This can occur if we have a nativized Blueprint class somewhere in the parent class ancestry.
-		if (bHasNativizedParent)
-		{
-			CheckAndApplyComponentTemplateOverrides(ClassDefaultObject);
-		}
 	}
 }
 
@@ -1183,7 +1181,7 @@ void UBlueprintGeneratedClass::CheckAndApplyComponentTemplateOverrides(UObject*
 									ICH->ConditionalPostLoad();
 
 									// Serialize cached override data to the instanced subobject that's based on the default subobject from the nativized parent class and owned by the non-nativized child class default object.
-									FNativizedComponentOverrideDataLoader OverrideDataLoader(OverrideData->GetCachedPropertyDataForSerialization(), OverrideData->GetCachedPropertyListForSerialization());
+									FNativizedComponentOverrideDataLoader OverrideDataLoader(OverrideData->GetCachedPropertyData(), OverrideData->GetCachedPropertyList());
 									NativizedComponentSubobjectInstance->Serialize(OverrideDataLoader);
 								}
 							}
@@ -1324,6 +1322,51 @@ void UBlueprintGeneratedClass::GetPreloadDependencies(TArray<UObject*>& OutDeps)
 	}
 }
 
+void UBlueprintGeneratedClass::GetDefaultObjectPreloadDependencies(TArray<UObject*>& OutDeps)
+{
+	Super::GetDefaultObjectPreloadDependencies(OutDeps);
+
+	// Ensure that BPGC-owned component templates (archetypes) are loaded prior to CDO serialization in order to support the following use cases:
+	//
+	//	1) When the "fast path" component instancing optimization is enabled, we generate a cached delta binary at BPGC load time that we then deserialize into
+	//	   new component instances after we spawn them at runtime. Generating the cached delta requires component templates to be loaded so that we can use them
+	//	   as the basis for delta serialization. However, we cannot add them a preload dependency of the class without introducing a cycle, so we add them as a
+	//	   preload dependency on the CDO here instead.
+	//	2) When Blueprint nativization is enabled, any Blueprint class assets that are not converted to C++ may still inherit from a Blueprint class asset that is
+	//	   converted to C++. In that case, the non-nativized child Blueprint class may still inherit one or more SCS nodes from the parent class. However, when
+	//	   we nativize a Blueprint class, we convert the class-owned SCS component templates into CDO-owned default subobjects. In the non-nativized child Blueprint
+	//	   class, these remain stored in the ICH as override templates. In order to ensure that the inherited default subobject in the CDO reflects the defaults that
+	//	   are recorded into the override template, we bake out the list of changed properties at cook time and then use it to also generate a cached delta binary
+	//	   when the non-nativized BPGC child asset is loaded in the cooked build. We then use binary serialization to update the default subobject instance (see
+	//	   CheckAndApplyComponentTemplateOverrides). That must occur prior to serializing instances of the non-nativized BPGC so that delta serialization works
+	//	   correctly, so adding them as preload dependencies here ensures that the override templates will all be loaded prior to serialization of the CDO.
+
+	// Walk up the SCS inheritance hierarchy and add component templates (archetypes). This may include override templates contained in the ICH for inherited SCS nodes.
+	UBlueprintGeneratedClass* CurrentBPClass = this;
+	while (CurrentBPClass)
+	{
+		if (CurrentBPClass->SimpleConstructionScript)
+		{
+			const TArray<USCS_Node*>& AllSCSNodes = CurrentBPClass->SimpleConstructionScript->GetAllNodes();
+			for (USCS_Node* SCSNode : AllSCSNodes)
+			{
+				OutDeps.Add(SCSNode->GetActualComponentTemplate(this));
+			}
+		}
+
+		CurrentBPClass = Cast<UBlueprintGeneratedClass>(CurrentBPClass->GetSuperClass());
+	}
+
+	// Also add UCS/AddComponent node templates (archetypes).
+	for (UActorComponent* ComponentTemplate : ComponentTemplates)
+	{
+		if (ComponentTemplate)
+		{
+			OutDeps.Add(ComponentTemplate);
+		}
+	}
+}
+
 bool UBlueprintGeneratedClass::NeedsLoadForServer() const
 {
 	// This logic can't be used for targets that use editor content because UBlueprint::NeedsLoadForEditorGame
@@ -1585,7 +1628,7 @@ void FBlueprintCookedComponentInstancingData::BuildCachedArrayPropertyList(const
 	}
 }
 
-const FCustomPropertyListNode* FBlueprintCookedComponentInstancingData::GetCachedPropertyListForSerialization() const
+const FCustomPropertyListNode* FBlueprintCookedComponentInstancingData::GetCachedPropertyList() const
 {
 	FCustomPropertyListNode* PropertyListRootNode = nullptr;
 
@@ -1605,7 +1648,7 @@ const FCustomPropertyListNode* FBlueprintCookedComponentInstancingData::GetCache
 	return PropertyListRootNode;
 }
 
-void FBlueprintCookedComponentInstancingData::LoadCachedPropertyDataForSerialization(UActorComponent* SourceTemplate)
+void FBlueprintCookedComponentInstancingData::BuildCachedPropertyDataFromTemplate(UActorComponent* SourceTemplate)
 {
 	// Blueprint component instance data writer implementation.
 	class FBlueprintComponentInstanceDataWriter : public FObjectWriter
@@ -1642,10 +1685,10 @@ void FBlueprintCookedComponentInstancingData::LoadCachedPropertyDataForSerializa
 			ComponentTemplateFlags = SourceTemplate->GetFlags();
 
 			// This will also load the cached property list, if necessary.
-			const FCustomPropertyListNode* PropertyList = GetCachedPropertyListForSerialization();
+			const FCustomPropertyListNode* PropertyList = GetCachedPropertyList();
 
 			// Write template data out to the "fast path" buffer. All dependencies will be loaded at this point.
-			FBlueprintComponentInstanceDataWriter InstanceDataWriter(CachedPropertyDataForSerialization, PropertyList);
+			FBlueprintComponentInstanceDataWriter InstanceDataWriter(CachedPropertyData, PropertyList);
 			SourceTemplate->Serialize(InstanceDataWriter);
 		}
 		else
diff --git a/Engine/Source/Runtime/Engine/Private/InheritableComponentHandler.cpp b/Engine/Source/Runtime/Engine/Private/InheritableComponentHandler.cpp
--- a/Engine/Source/Runtime/Engine/Private/InheritableComponentHandler.cpp
+++ b/Engine/Source/Runtime/Engine/Private/InheritableComponentHandler.cpp
@@ -73,11 +73,6 @@ void UInheritableComponentHandler::PostLoad()
 					Record.ComponentTemplate->MarkPendingKill(); // hack needed to be able to identify if NewObject returns this back to us in the future
 					Records.RemoveAtSwap(Index);
 				}
-				else if (Record.CookedComponentInstancingData.bIsValid)
-				{
-					// Generate "fast path" instancing data. This data may also be used to override components inherited from a nativized BP parent class.
-					Record.CookedComponentInstancingData.LoadCachedPropertyDataForSerialization(Record.ComponentTemplate);
-				}
 			}
 		}
 	}
diff --git a/Engine/Source/Runtime/Engine/Private/SCS_Node.cpp b/Engine/Source/Runtime/Engine/Private/SCS_Node.cpp
--- a/Engine/Source/Runtime/Engine/Private/SCS_Node.cpp
+++ b/Engine/Source/Runtime/Engine/Private/SCS_Node.cpp
@@ -475,12 +475,6 @@ void USCS_Node::PostLoad()
 #if WITH_EDITOR
 	ValidateGuid();
 #endif
-
-	// If valid, load cooked component instancing data.
-	if (ComponentTemplate && CookedComponentInstancingData.bIsValid)
-	{
-		CookedComponentInstancingData.LoadCachedPropertyDataForSerialization(ComponentTemplate);
-	}
 }
 
 
