Add flags to track bMarkedObjectsPendingKill and bCleanedUpWorld to ensure that sublevels are cleaned up properly

diff --git a/Engine/Source/Runtime/Engine/Classes/Engine/World.h b/Engine/Source/Runtime/Engine/Classes/Engine/World.h
--- a/Engine/Source/Runtime/Engine/Classes/Engine/World.h
+++ b/Engine/Source/Runtime/Engine/Classes/Engine/World.h
@@ -1386,6 +1386,14 @@ public:
 	UPROPERTY(transient)
 	uint32 bAreConstraintsDirty:1;
 
+	/** Indicates that the world has marked contained objects as pending kill */
+	bool HasMarkedObjectsPendingKill() const { return bMarkedObjectsPendingKill; }
+private:
+	uint32 bCleanedUpWorld:1;
+
+	uint32 bMarkedObjectsPendingKill:1;
+
+public:
 #if WITH_EDITORONLY_DATA
 	/** List of DDC async requests we need to wait on before we register components. Game thread only. */
 	TArray<TSharedPtr<FAsyncPreRegisterDDCRequest>> AsyncPreRegisterDDCRequests;
diff --git a/Engine/Source/Runtime/Engine/Private/World.cpp b/Engine/Source/Runtime/Engine/Private/World.cpp
--- a/Engine/Source/Runtime/Engine/Private/World.cpp
+++ b/Engine/Source/Runtime/Engine/Private/World.cpp
@@ -1503,6 +1503,7 @@ void UWorld::MarkObjectsPendingKill()
 		Object->MarkPendingKill();
 	};
 	ForEachObjectWithOuter(this, MarkObjectPendingKill, true, RF_NoFlags, EInternalObjectFlags::PendingKill);
+	bMarkedObjectsPendingKill = true;
 }
 
 UWorld* UWorld::CreateWorld(const EWorldType::Type InWorldType, bool bInformEngineOfWorld, FName WorldName, UPackage* InWorldPackage, bool bAddToRoot, ERHIFeatureLevel::Type InFeatureLevel)
@@ -3714,6 +3715,7 @@ bool UWorld::IsNavigationRebuilt() const
 void UWorld::CleanupWorld(bool bSessionEnded, bool bCleanupResources, UWorld* NewWorld)
 {
 	check(IsVisibilityRequestPending() == false);
+	bCleanedUpWorld = true;
 
 	// Wait on current physics scenes if they are processing
 	if(FPhysScene* CurrPhysicsScene = GetPhysicsScene())
@@ -3803,12 +3805,24 @@ void UWorld::CleanupWorld(bool bSessionEnded, bool bCleanupResources, UWorld* Ne
 	}
 #endif //WITH_EDITOR
 
+	for (int32 LevelIndex = 0; LevelIndex < GetNumLevels(); ++LevelIndex)
+	{
+		UWorld* World = CastChecked<UWorld>(GetLevel(LevelIndex)->GetOuter());
+		if (!World->bCleanedUpWorld)
+		{
+			World->CleanupWorld(bSessionEnded, bCleanupResources, NewWorld);
+		}
+	}
+
 	for (ULevelStreaming* StreamingLevel : GetStreamingLevels())
 	{
 		if (ULevel* Level = StreamingLevel->GetLoadedLevel())
 		{
 			UWorld* World = CastChecked<UWorld>(Level->GetOuter());
-			World->CleanupWorld(bSessionEnded, bCleanupResources, NewWorld);
+			if (!World->bCleanedUpWorld)
+			{
+				World->CleanupWorld(bSessionEnded, bCleanupResources, NewWorld);
+			}
 		}
 	}
 
@@ -3824,7 +3838,7 @@ void UWorld::CleanupWorld(bool bSessionEnded, bool bCleanupResources, UWorld* Ne
 			}
 
 			UWorld* const LevelWorld = CastChecked<UWorld>(Level->GetOuter());
-			if (LevelWorld != this)
+			if (!LevelWorld->bCleanedUpWorld)
 			{
 				LevelWorld->CleanupWorld(bSessionEnded, bCleanupResources, NewWorld);
 			}
@@ -5782,19 +5796,26 @@ UWorld* FSeamlessTravelHandler::Tick()
 			// mark everything else contained in the world to be deleted
 			for (auto LevelIt(CurrentWorld->GetLevelIterator()); LevelIt; ++LevelIt)
 			{
-				const ULevel* Level = *LevelIt;
-				if (Level)
+				if (const ULevel* Level = *LevelIt)
 				{
-					CastChecked<UWorld>(Level->GetOuter())->MarkObjectsPendingKill();
+					UWorld* World = CastChecked<UWorld>(Level->GetOuter());
+					if (!World->HasMarkedObjectsPendingKill())
+					{
+						World->MarkObjectsPendingKill();
+					}
 				}
 			}
 
 			for (ULevelStreaming* LevelStreaming : CurrentWorld->GetStreamingLevels())
 			{
 				// If an unloaded levelstreaming still has a loaded level we need to mark its objects to be deleted as well
-				if (LevelStreaming->GetLoadedLevel() && (!LevelStreaming->ShouldBeLoaded() || !LevelStreaming->ShouldBeVisible()))
+				if (LevelStreaming->GetLoadedLevel())
 				{
-					CastChecked<UWorld>(LevelStreaming->GetLoadedLevel()->GetOuter())->MarkObjectsPendingKill();
+					UWorld* World = CastChecked<UWorld>(LevelStreaming->GetLoadedLevel()->GetOuter());
+					if (!World->HasMarkedObjectsPendingKill())
+					{
+						World->MarkObjectsPendingKill();
+					}
 				}
 			}
 
