PR #5077: Prevent TimelineComponent Update delegate from being executed 3 times when looping (Contributed by [REDACTED_PERSON])

diff --git a/Engine/Source/Runtime/Engine/Private/Timeline.cpp b/Engine/Source/Runtime/Engine/Private/Timeline.cpp
--- a/Engine/Source/Runtime/Engine/Private/Timeline.cpp
+++ b/Engine/Source/Runtime/Engine/Private/Timeline.cpp
@@ -466,8 +466,8 @@ void FTimeline::TickTimeline(float DeltaTime)
 				// If looping, play to end, jump to start, and set target to somewhere near the beginning.
 				if(bLooping)
 				{
-					SetPlaybackPosition(TimelineLength, true);
-					SetPlaybackPosition(0.f, false);
+					SetPlaybackPosition(TimelineLength, true, false);
+					SetPlaybackPosition(0.f, false, false);
 
 					if( TimelineLength > 0.f )
 					{
@@ -497,8 +497,8 @@ void FTimeline::TickTimeline(float DeltaTime)
 				// If looping, play to start, jump to end, and set target to somewhere near the end.
 				if(bLooping)
 				{
-					SetPlaybackPosition(0.f, true);
-					SetPlaybackPosition(TimelineLength, false);
+					SetPlaybackPosition(0.f, true, false);
+					SetPlaybackPosition(TimelineLength, false, false);
 
 					if( TimelineLength > 0.f )
 					{
