PR #5074: Fixed minor bugs in GameUserSettings. (Contributed by [REDACTED_PERSON])

diff --git a/Engine/Source/Runtime/Engine/Classes/GameFramework/GameUserSettings.h b/Engine/Source/Runtime/Engine/Classes/GameFramework/GameUserSettings.h
--- a/Engine/Source/Runtime/Engine/Classes/GameFramework/GameUserSettings.h
+++ b/Engine/Source/Runtime/Engine/Classes/GameFramework/GameUserSettings.h
@@ -147,7 +147,7 @@ public:
 	float GetFrameRateLimit() const;
 
 	// Changes all scalability settings at once based on a single overall quality level
-	// @param Value 0:low, 1:medium, 2:high, 3:epic
+	// @param Value 0:low, 1:medium, 2:high, 3:epic, 4:cinematic
 	UFUNCTION(BlueprintCallable, Category=Settings)
 	virtual void SetOverallScalabilityLevel(int32 Value);
 
@@ -177,59 +177,66 @@ public:
 	UFUNCTION(BlueprintCallable, Category=Settings)
 	void SetResolutionScaleNormalized(float NewScaleNormalized);
 
-	// Sets the view distance quality (0..3, higher is better)
+	// Sets the view distance quality (0..4, higher is better)
+	// @param Value 0:low, 1:medium, 2:high, 3:epic, 4:cinematic (gets clamped if needed)
 	UFUNCTION(BlueprintCallable, Category=Settings)
 	void SetViewDistanceQuality(int32 Value);
 
-	// Returns the view distance quality (0..3, higher is better)
+	// Returns the view distance quality (0..4, higher is better)
 	UFUNCTION(BlueprintCallable, Category=Settings)
 	int32 GetViewDistanceQuality() const;
 
-	// Sets the shadow quality (0..3, higher is better)
+	// Sets the shadow quality (0..4, higher is better)
+	// @param Value 0:low, 1:medium, 2:high, 3:epic, 4:cinematic (gets clamped if needed)
 	UFUNCTION(BlueprintCallable, Category=Settings)
 	void SetShadowQuality(int32 Value);
 
-	// Returns the shadow quality (0..3, higher is better)
+	// Returns the shadow quality (0..4, higher is better)
 	UFUNCTION(BlueprintCallable, Category=Settings)
 	int32 GetShadowQuality() const;
 
-	// Sets the anti-aliasing quality (0..3, higher is better)
+	// Sets the anti-aliasing quality (0..4, higher is better)
+	// @param Value 0:low, 1:medium, 2:high, 3:epic, 4:cinematic (gets clamped if needed)
 	UFUNCTION(BlueprintCallable, Category=Settings)
 	void SetAntiAliasingQuality(int32 Value);
 
-	// Returns the anti-aliasing quality (0..3, higher is better)
+	// Returns the anti-aliasing quality (0..4, higher is better)
 	UFUNCTION(BlueprintCallable, Category=Settings)
 	int32 GetAntiAliasingQuality() const;
 
-	// Sets the texture quality (0..3, higher is better)
+	// Sets the texture quality (0..4, higher is better)
+	// @param Value 0:low, 1:medium, 2:high, 3:epic, 4:cinematic  (gets clamped if needed)
 	UFUNCTION(BlueprintCallable, Category=Settings)
 	void SetTextureQuality(int32 Value);
 
-	// Returns the texture quality (0..3, higher is better)
+	// Returns the texture quality (0..4, higher is better)
 	UFUNCTION(BlueprintCallable, Category=Settings)
 	int32 GetTextureQuality() const;
 
-	// Sets the visual effects quality (0..3, higher is better)
+	// Sets the visual effects quality (0..4, higher is better)
+	// @param Value 0:low, 1:medium, 2:high, 3:epic, 4:cinematic (gets clamped if needed)
 	UFUNCTION(BlueprintCallable, Category=Settings)
 	void SetVisualEffectQuality(int32 Value);
 
-	// Returns the visual effects quality (0..3, higher is better)
+	// Returns the visual effects quality (0..4, higher is better)
 	UFUNCTION(BlueprintCallable, Category=Settings)
 	int32 GetVisualEffectQuality() const;
 
-	// Sets the post-processing quality (0..3, higher is better)
+	// Sets the post-processing quality (0..4, higher is better)
+	// @param Value 0:low, 1:medium, 2:high, 3:epic, 4:cinematic (gets clamped if needed)
 	UFUNCTION(BlueprintCallable, Category=Settings)
 	void SetPostProcessingQuality(int32 Value);
 
-	// Returns the post-processing quality (0..3, higher is better)
+	// Returns the post-processing quality (0..4, higher is better)
 	UFUNCTION(BlueprintCallable, Category=Settings)
 	int32 GetPostProcessingQuality() const;
 	
-	// Sets the post-processing quality (0..3, higher is better)
+	// Sets the foliage quality (0..4, higher is better)
+	// @param Value 0:low, 1:medium, 2:high, 3:epic, 4:cinematic (gets clamped if needed)
 	UFUNCTION(BlueprintCallable, Category=Settings)
 	void SetFoliageQuality(int32 Value);
 
-	// Returns the post-processing quality (0..3, higher is better)
+	// Returns the foliage quality (0..4, higher is better)
 	UFUNCTION(BlueprintCallable, Category=Settings)
 	int32 GetFoliageQuality() const;
 
diff --git a/Engine/Source/Runtime/Engine/Private/GameUserSettings.cpp b/Engine/Source/Runtime/Engine/Private/GameUserSettings.cpp
--- a/Engine/Source/Runtime/Engine/Private/GameUserSettings.cpp
+++ b/Engine/Source/Runtime/Engine/Private/GameUserSettings.cpp
@@ -753,7 +753,7 @@ void UGameUserSettings::SetResolutionScaleNormalized(float NewScaleNormalized)
 
 void UGameUserSettings::SetViewDistanceQuality(int32 Value)
 {
-	ScalabilityQuality.ViewDistanceQuality = Value;
+	ScalabilityQuality.SetViewDistanceQuality(Value);
 }
 
 int32 UGameUserSettings::GetViewDistanceQuality() const
@@ -763,7 +763,7 @@ int32 UGameUserSettings::GetViewDistanceQuality() const
 
 void UGameUserSettings::SetShadowQuality(int32 Value)
 {
-	ScalabilityQuality.ShadowQuality = Value;
+	ScalabilityQuality.SetShadowQuality(Value);
 }
 
 int32 UGameUserSettings::GetShadowQuality() const
@@ -773,7 +773,7 @@ int32 UGameUserSettings::GetShadowQuality() const
 
 void UGameUserSettings::SetAntiAliasingQuality(int32 Value)
 {
-	ScalabilityQuality.AntiAliasingQuality = Value;
+	ScalabilityQuality.SetAntiAliasingQuality(Value);
 }
 
 int32 UGameUserSettings::GetAntiAliasingQuality() const
@@ -783,7 +783,7 @@ int32 UGameUserSettings::GetAntiAliasingQuality() const
 
 void UGameUserSettings::SetTextureQuality(int32 Value)
 {
-	ScalabilityQuality.TextureQuality = Value;
+	ScalabilityQuality.SetTextureQuality(Value);
 }
 
 int32 UGameUserSettings::GetTextureQuality() const
@@ -793,7 +793,7 @@ int32 UGameUserSettings::GetTextureQuality() const
 
 void UGameUserSettings::SetVisualEffectQuality(int32 Value)
 {
-	ScalabilityQuality.EffectsQuality = Value;
+	ScalabilityQuality.SetEffectsQuality(Value);
 }
 
 int32 UGameUserSettings::GetVisualEffectQuality() const
@@ -803,7 +803,7 @@ int32 UGameUserSettings::GetVisualEffectQuality() const
 
 void UGameUserSettings::SetPostProcessingQuality(int32 Value)
 {
-	ScalabilityQuality.PostProcessQuality = Value;
+	ScalabilityQuality.SetPostProcessQuality(Value);
 }
 
 int32 UGameUserSettings::GetPostProcessingQuality() const
@@ -813,7 +813,7 @@ int32 UGameUserSettings::GetPostProcessingQuality() const
 
 void UGameUserSettings::SetFoliageQuality(int32 Value)
 {
-	ScalabilityQuality.FoliageQuality = FMath::Clamp(Value, 0, 3);
+	ScalabilityQuality.SetFoliageQuality(Value);
 }
 
 int32 UGameUserSettings::GetFoliageQuality() const
diff --git a/Engine/Source/Runtime/Engine/Private/Scalability.cpp b/Engine/Source/Runtime/Engine/Private/Scalability.cpp
--- a/Engine/Source/Runtime/Engine/Private/Scalability.cpp
+++ b/Engine/Source/Runtime/Engine/Private/Scalability.cpp
@@ -686,6 +686,41 @@ int32 FQualityLevels::GetSingleQualityLevel() const
 	return -1;
 }
 
+void FQualityLevels::SetViewDistanceQuality(int32 Value)
+{
+	ViewDistanceQuality = FMath::Clamp(Value, 0, CVarViewDistanceQuality_NumLevels->GetInt() - 1);
+}
+
+void FQualityLevels::SetAntiAliasingQuality(int32 Value)
+{
+	AntiAliasingQuality = FMath::Clamp(Value, 0, CVarAntiAliasingQuality_NumLevels->GetInt() - 1);
+}
+
+void FQualityLevels::SetShadowQuality(int32 Value)
+{
+	ShadowQuality = FMath::Clamp(Value, 0, CVarShadowQuality_NumLevels->GetInt() - 1);
+}
+
+void FQualityLevels::SetPostProcessQuality(int32 Value)
+{
+	PostProcessQuality = FMath::Clamp(Value, 0, CVarPostProcessQuality_NumLevels->GetInt() - 1);
+}
+
+void FQualityLevels::SetTextureQuality(int32 Value)
+{
+	TextureQuality = FMath::Clamp(Value, 0, CVarTextureQuality_NumLevels->GetInt() - 1);
+}
+
+void FQualityLevels::SetEffectsQuality(int32 Value)
+{
+	EffectsQuality = FMath::Clamp(Value, 0, CVarEffectsQuality_NumLevels->GetInt() - 1);
+}
+
+void FQualityLevels::SetFoliageQuality(int32 Value)
+{
+	FoliageQuality = FMath::Clamp(Value, 0, CVarFoliageQuality_NumLevels->GetInt() - 1);
+}
+
 void LoadState(const FString& IniName)
 {
 	check(!IniName.IsEmpty());
diff --git a/Engine/Source/Runtime/Engine/Public/Scalability.h b/Engine/Source/Runtime/Engine/Public/Scalability.h
--- a/Engine/Source/Runtime/Engine/Public/Scalability.h
+++ b/Engine/Source/Runtime/Engine/Public/Scalability.h
@@ -61,7 +61,7 @@ namespace Scalability
 		}
 
 		// Sets all other settings based on an overall value
-		// @param Value 0:low, 1:medium, 2:high, 3:epic (gets clamped if needed)
+		// @param Value 0:low, 1:medium, 2:high, 3:epic, 4:cinematic (gets clamped if needed)
 		void SetFromSingleQualityLevel(int32 Value);
 
 		// Sets all other settings based on an overall value, but relative to the maximum.
@@ -69,9 +69,37 @@ namespace Scalability
 		void SetFromSingleQualityLevelRelativeToMax(int32 Value);
 
 		// Returns the overall value if all settings are set to the same thing
-		// @param Value -1:custom, 0:low, 1:medium, 2:high, 3:epic
+		// @param Value -1:custom, 0:low, 1:medium, 2:high, 3:epic, 4:cinematic
 		int32 GetSingleQualityLevel() const;
 
+		// Sets view distance quality
+		// @param Value 0:low, 1:medium, 2:high, 3:epic, 4:cinematic (gets clamped if needed)
+		void SetViewDistanceQuality(int32 Value);
+
+		// Sets anti-aliasing quality
+		// @param Value 0:low, 1:medium, 2:high, 3:epic, 4:cinematic (gets clamped if needed)
+		void SetAntiAliasingQuality(int32 Value);
+
+		// Sets shadow quality
+		// @param Value 0:low, 1:medium, 2:high, 3:epic, 4:cinematic (gets clamped if needed)
+		void SetShadowQuality(int32 Value);
+
+		// Sets the post-processing quality
+		// @param Value 0:low, 1:medium, 2:high, 3:epic, 4:cinematic (gets clamped if needed)
+		void SetPostProcessQuality(int32 Value);
+
+		// Sets the texture quality
+		// @param Value 0:low, 1:medium, 2:high, 3:epic, 4:cinematic (gets clamped if needed)
+		void SetTextureQuality(int32 Value);
+
+		// Sets the visual effects quality
+		// @param Value 0:low, 1:medium, 2:high, 3:epic, 4:cinematic (gets clamped if needed)
+		void SetEffectsQuality(int32 Value);
+
+		// Sets the foliage quality
+		// @param Value 0:low, 1:medium, 2:high, 3:epic, 4:cinematic (gets clamped if needed)
+		void SetFoliageQuality(int32 Value);
+
 		void SetBenchmarkFallback();
 
 		void SetDefaults();
