Inline templated version of GetGameInstance
#rnx

diff --git a/Engine/Source/Runtime/Engine/Classes/GameFramework/Actor.h b/Engine/Source/Runtime/Engine/Classes/GameFramework/Actor.h
--- a/Engine/Source/Runtime/Engine/Classes/GameFramework/Actor.h
+++ b/Engine/Source/Runtime/Engine/Classes/GameFramework/Actor.h
@@ -2693,7 +2693,10 @@ public:
 	 * May return NULL if the cast fails. 
 	 */
 	template< class T >
-	T* GetGameInstance() const;
+	T* GetGameInstance() const 
+	{ 
+		return Cast<T>(GetGameInstance()); 
+	}
 
 	/** Returns true if this is a replicated actor that was placed in the map */
 	bool IsNetStartupActor() const;
diff --git a/Engine/Source/Runtime/Engine/Private/Actor.cpp b/Engine/Source/Runtime/Engine/Private/Actor.cpp
--- a/Engine/Source/Runtime/Engine/Private/Actor.cpp
+++ b/Engine/Source/Runtime/Engine/Private/Actor.cpp
@@ -352,12 +352,6 @@ UGameInstance* AActor::GetGameInstance() const
 	return GetWorld()->GetGameInstance();
 }
 
-template< class T >
-T* AActor::GetGameInstance() const
-{
-	return GetWorld()->GetGameInstance<T>();
-}
-
 bool AActor::IsNetStartupActor() const
 {
 	return bNetStartup;
