Fix "Non-void function should return a value" warning in CIS.

diff --git a/Engine/Source/Runtime/Engine/Classes/GameFramework/Pawn.h b/Engine/Source/Runtime/Engine/Classes/GameFramework/Pawn.h
--- a/Engine/Source/Runtime/Engine/Classes/GameFramework/Pawn.h
+++ b/Engine/Source/Runtime/Engine/Classes/GameFramework/Pawn.h
@@ -204,7 +204,7 @@ public:
 	template < class T >
 	T* GetController() const
 	{
-		Cast<T>(GetController());
+		return Cast<T>(GetController());
 	}
 
 	/** Get the rotation of the Controller, often the 'view' rotation of this Pawn. */
