Fixed a hang when entering PIE when Background Streaming is disabled

diff --git a/Engine/Source/Runtime/Engine/Private/LevelStreaming.cpp b/Engine/Source/Runtime/Engine/Private/LevelStreaming.cpp
--- a/Engine/Source/Runtime/Engine/Private/LevelStreaming.cpp
+++ b/Engine/Source/Runtime/Engine/Private/LevelStreaming.cpp
@@ -429,7 +429,7 @@ bool ULevelStreaming::DetermineTargetState()
 		{
 			bContinueToConsider = false;
 		}
-		else if (!GUseBackgroundLevelStreaming)
+		else if (!GUseBackgroundLevelStreaming && ShouldBeLoaded())
 		{
 			TargetState = ETargetState::LoadedNotVisible;
 		}
