Copied 4382276 Francis.Hurteau
Fix syncing streaming level visibility state with undo/redo

diff --git a/Engine/Source/Runtime/Engine/Private/LevelStreaming.cpp b/Engine/Source/Runtime/Engine/Private/LevelStreaming.cpp
--- a/Engine/Source/Runtime/Engine/Private/LevelStreaming.cpp
+++ b/Engine/Source/Runtime/Engine/Private/LevelStreaming.cpp
@@ -1539,6 +1539,10 @@ void ULevelStreaming::PreEditUndo()
 void ULevelStreaming::PostEditUndo()
 {
 	FLevelUtils::ApplyEditorTransform(this, false);
+	if (UWorld* World = GetWorld())
+	{
+		World->UpdateStreamingLevelShouldBeConsidered(this);
+	}
 }
 
 const FName& ULevelStreaming::GetFolderPath() const
