Merging using Dev-VR_to_Release-4.21

diff --git a/Engine/Source/Runtime/Engine/Private/PlayerController.cpp b/Engine/Source/Runtime/Engine/Private/PlayerController.cpp
--- a/Engine/Source/Runtime/Engine/Private/PlayerController.cpp
+++ b/Engine/Source/Runtime/Engine/Private/PlayerController.cpp
@@ -4558,7 +4558,8 @@ void APlayerController::TickActor( float DeltaSeconds, ELevelTick TickType, FAct
 			if (bIsClient && bIsLocallyControlled && GetPawn() && PlayerCameraManager->bUseClientSideCameraUpdates)
 			{
 				UPawnMovementComponent* PawnMovement = GetPawn()->GetMovementComponent();
-				if (!PawnMovement->IsMoveInputIgnored() &&
+				if (PawnMovement != nullptr &&
+					!PawnMovement->IsMoveInputIgnored() &&
 					(PawnMovement->GetLastInputVector() != FVector::ZeroVector || PawnMovement->Velocity != FVector::ZeroVector))
 				{
 					PlayerCameraManager->bShouldSendClientSideCameraUpdate = true;
