Merging using Dev-VR_to_Release-4.21

diff --git a/Engine/Source/Runtime/Engine/Classes/Kismet/BlueprintPlatformLibrary.h b/Engine/Source/Runtime/Engine/Classes/Kismet/BlueprintPlatformLibrary.h
--- a/Engine/Source/Runtime/Engine/Classes/Kismet/BlueprintPlatformLibrary.h
+++ b/Engine/Source/Runtime/Engine/Classes/Kismet/BlueprintPlatformLibrary.h
@@ -224,4 +224,9 @@ public:
 	UFUNCTION(BlueprintCallable, Category="Platform|LocalNotification")
 	static void GetLaunchNotification(bool& NotificationLaunchedApp, FString& ActivationEvent, int32& FireDate);
 
+	/**
+	 * @return the current device orientation
+	 */
+	UFUNCTION(BlueprintPure, Category="Platform|LocalNotification")
+	static EScreenOrientation::Type GetDeviceOrientation();
 };
diff --git a/Engine/Source/Runtime/Engine/Private/Kismet/BlueprintPlatformLibrary.cpp b/Engine/Source/Runtime/Engine/Private/Kismet/BlueprintPlatformLibrary.cpp
--- a/Engine/Source/Runtime/Engine/Private/Kismet/BlueprintPlatformLibrary.cpp
+++ b/Engine/Source/Runtime/Engine/Private/Kismet/BlueprintPlatformLibrary.cpp
@@ -163,4 +163,30 @@ void UBlueprintPlatformLibrary::GetLaunchNotification(bool& NotificationLaunched
 	platformService->GetLaunchNotification(NotificationLaunchedApp, ActivationEvent, FireDate);
 }
 
+EScreenOrientation::Type UBlueprintPlatformLibrary::GetDeviceOrientation()
+{
+	switch (FPlatformMisc::GetDeviceOrientation())
+	{
+		case EDeviceScreenOrientation::Portrait:
+			return EScreenOrientation::Portrait;
+			
+		case EDeviceScreenOrientation::PortraitUpsideDown:
+			return EScreenOrientation::PortraitUpsideDown;
+			
+		case EDeviceScreenOrientation::LandscapeLeft:
+			return EScreenOrientation::LandscapeLeft;
+
+		case EDeviceScreenOrientation::LandscapeRight:
+			return EScreenOrientation::LandscapeRight;
+
+		case EDeviceScreenOrientation::FaceUp:
+			return EScreenOrientation::FaceUp;
+			
+		case EDeviceScreenOrientation::FaceDown:
+			return EScreenOrientation::FaceDown;
+	}
+	return EScreenOrientation::Unknown;
+}
+
+
 ILocalNotificationService* UBlueprintPlatformLibrary::platformService = nullptr;
