Merging using Dev-VR_to_Release-4.21

diff --git a/Engine/Source/Runtime/Engine/Classes/Components/BoxComponent.h b/Engine/Source/Runtime/Engine/Classes/Components/BoxComponent.h
--- a/Engine/Source/Runtime/Engine/Classes/Components/BoxComponent.h
+++ b/Engine/Source/Runtime/Engine/Classes/Components/BoxComponent.h
@@ -23,6 +23,10 @@ protected:
 	UPROPERTY(EditAnywhere, BlueprintReadOnly, export, Category=Shape)
 	FVector BoxExtent;
 
+	/** Used to control the line thickness when rendering */
+	UPROPERTY(EditAnywhere, BlueprintReadOnly, export, Category=Shape)
+	float LineThickness;
+
 public:
 	/** 
 	 * Change the box extent size. This is the unscaled size, before component scale is applied.
diff --git a/Engine/Source/Runtime/Engine/Private/Components/BoxComponent.cpp b/Engine/Source/Runtime/Engine/Private/Components/BoxComponent.cpp
--- a/Engine/Source/Runtime/Engine/Private/Components/BoxComponent.cpp
+++ b/Engine/Source/Runtime/Engine/Private/Components/BoxComponent.cpp
@@ -130,6 +130,7 @@ FPrimitiveSceneProxy* UBoxComponent::CreateSceneProxy()
 			,	bDrawOnlyIfSelected( InComponent->bDrawOnlyIfSelected )
 			,   BoxExtents( InComponent->BoxExtent )
 			,	BoxColor( InComponent->ShapeColor )
+			,	LineThickness( InComponent->LineThickness )
 		{
 			bWillEverBeLit = false;
 		}
@@ -149,7 +150,7 @@ FPrimitiveSceneProxy* UBoxComponent::CreateSceneProxy()
 					const FLinearColor DrawColor = GetViewSelectionColor(BoxColor, *View, IsSelected(), IsHovered(), false, IsIndividuallySelected() );
 
 					FPrimitiveDrawInterface* PDI = Collector.GetPDI(ViewIndex);
-					DrawOrientedWireBox(PDI, LocalToWorld.GetOrigin(), LocalToWorld.GetScaledAxis( EAxis::X ), LocalToWorld.GetScaledAxis( EAxis::Y ), LocalToWorld.GetScaledAxis( EAxis::Z ), BoxExtents, DrawColor, SDPG_World);
+					DrawOrientedWireBox(PDI, LocalToWorld.GetOrigin(), LocalToWorld.GetScaledAxis( EAxis::X ), LocalToWorld.GetScaledAxis( EAxis::Y ), LocalToWorld.GetScaledAxis( EAxis::Z ), BoxExtents, DrawColor, SDPG_World, LineThickness);
 				}
 			}
 		}
@@ -175,6 +176,7 @@ FPrimitiveSceneProxy* UBoxComponent::CreateSceneProxy()
 		const uint32	bDrawOnlyIfSelected:1;
 		const FVector	BoxExtents;
 		const FColor	BoxColor;
+		const float LineThickness;
 	};
 
 	return new FBoxSceneProxy( this );
