Fix FUniqueNetIdRepl serialization for out of engine OSSs. Use FNetBitReader/Writer in the serialization test so it properly tests network replication

diff --git a/Engine/Plugins/Online/OnlineSubsystemUtils/Source/OnlineSubsystemUtils/Private/OnlineSubsystemUtilsModule.cpp b/Engine/Plugins/Online/OnlineSubsystemUtils/Source/OnlineSubsystemUtils/Private/OnlineSubsystemUtilsModule.cpp
--- a/Engine/Plugins/Online/OnlineSubsystemUtils/Source/OnlineSubsystemUtils/Private/OnlineSubsystemUtilsModule.cpp
+++ b/Engine/Plugins/Online/OnlineSubsystemUtils/Source/OnlineSubsystemUtils/Private/OnlineSubsystemUtilsModule.cpp
@@ -218,6 +218,9 @@ private:
 
 		ensure(SubsystemNameToHash.Num() == (HashId - 1));
 		ensure(HashToSubsystemName.Num() == (HashId - 1));
+
+		// FUniqueNetIdRepl uses 5 bits to transmit the HashId and 31 is used for OnlineSubsystems not included in this list
+		check(HashId < 31);
 	}
 
 	/** If false it will not try to do online PIE at all */
diff --git a/Engine/Source/Runtime/Engine/Private/OnlineReplStructs.cpp b/Engine/Source/Runtime/Engine/Private/OnlineReplStructs.cpp
--- a/Engine/Source/Runtime/Engine/Private/OnlineReplStructs.cpp
+++ b/Engine/Source/Runtime/Engine/Private/OnlineReplStructs.cpp
@@ -5,6 +5,7 @@
 =============================================================================*/
 
 #include "GameFramework/OnlineReplStructs.h"
+#include "UObject/CoreNet.h"
 #include "Serialization/MemoryWriter.h"
 #include "Serialization/MemoryReader.h"
 #include "UObject/PropertyPortFlags.h"
@@ -40,7 +41,10 @@ enum class EUniqueIdEncodingFlags : uint8
 };
 ENUM_CLASS_FLAGS(EUniqueIdEncodingFlags);
 
-FArchive& operator<<( FArchive& Ar, FUniqueNetIdRepl& UniqueNetId)
+/** Use highest value for type for other (out of engine) oss type */
+const uint8 TypeHash_Other = 31;
+
+FArchive& operator<<(FArchive& Ar, FUniqueNetIdRepl& UniqueNetId)
 {
 	if (!Ar.IsPersistent() || Ar.IsNetArchive())
 	{
@@ -97,7 +101,9 @@ inline uint8 GetTypeHashFromEncoding(EUniqueIdEncodingFlags inFlags)
  *    <uint8 flags> noted it is encoded and empty
  * NonEmpty:
  * - Encoded - <uint8 flags/type> <uint8 encoded size> <encoded bytes>
+ * - Encoded (out of engine oss type) - <uint8 flags/type> <serialized FName> <uint8 encoded size> <encoded bytes>
  * - Unencoded - <uint8 flags/type> <serialized FString>
+ * - Unencoded (out of engine oss type) - <uint8 flags/type> <serialized FName> <serialized FString>
  */
 void FUniqueNetIdRepl::MakeReplicationData()
 {
@@ -136,8 +142,13 @@ void FUniqueNetIdRepl::MakeReplicationData()
 		}
 
 		// Encode the unique id type
-		uint8 TypeHash = UOnlineEngineInterface::Get()->GetReplicationHashForSubsystem(GetType());
+		FName Type = GetType();
+		uint8 TypeHash = UOnlineEngineInterface::Get()->GetReplicationHashForSubsystem(Type);
 		ensure(TypeHash < 32);
+		if (TypeHash == 0 && Type != NAME_None)
+		{
+			TypeHash = TypeHash_Other;
+		}
 		EncodingFlags = static_cast<EUniqueIdEncodingFlags>((TypeHash << 3) | static_cast<uint8>(EncodingFlags));
 
 		if (EnumHasAllFlags(EncodingFlags, EUniqueIdEncodingFlags::IsEncoded))
@@ -148,6 +159,11 @@ void FUniqueNetIdRepl::MakeReplicationData()
 
 			FMemoryWriter Writer(ReplicationBytes);
 			Writer << EncodingFlags;
+			if (TypeHash == TypeHash_Other)
+			{
+				FString TypeString = Type.ToString();
+				Writer << TypeString;
+			}
 			Writer << EncodedSize;
 
 			int32 HexStartOffset = Writer.Tell();
@@ -163,6 +179,11 @@ void FUniqueNetIdRepl::MakeReplicationData()
 
 			FMemoryWriter Writer(ReplicationBytes);
 			Writer << EncodingFlags;
+			if (TypeHash == TypeHash_Other)
+			{
+				FString TypeString = Type.ToString();
+				Writer << TypeString;
+			}
 			Writer << Contents;
 			//UE_LOG(LogNet, VeryVerbose, TEXT("Normal UniqueId, serializing %d bytes"), ReplicationBytes.Num());
 		}
@@ -220,8 +241,24 @@ bool FUniqueNetIdRepl::NetSerialize(FArchive& Ar, UPackageMap* Map, bool& bOutSu
 						// If no type was encoded, assume default
 						TypeHash = UOnlineEngineInterface::Get()->GetReplicationHashForSubsystem(UOnlineEngineInterface::Get()->GetDefaultOnlineSubsystemName());
 					}
+					FName Type;
+					bool bValidTypeHash = TypeHash != 0;
+					if (TypeHash == TypeHash_Other)
+					{
+						FString TypeString;
+						Ar << TypeString;
+						Type = FName(*TypeString);
+						if (Ar.IsError() || Type == NAME_None)
+						{
+							bValidTypeHash = false;
+						}
+					}
+					else
+					{
+						Type = UOnlineEngineInterface::Get()->GetSubsystemFromReplicationHash(TypeHash);
+					}
 
-					if (TypeHash != 0)
+					if (bValidTypeHash)
 					{
 						// Get the size
 						uint8 EncodedSize = 0;
@@ -237,7 +274,6 @@ bool FUniqueNetIdRepl::NetSerialize(FArchive& Ar, UPackageMap* Map, bool& bOutSu
 									FString Contents = BytesToHex(TempBytes, EncodedSize);
 									if (Contents.Len() > 0)
 									{
-										FName Type = UOnlineEngineInterface::Get()->GetSubsystemFromReplicationHash(TypeHash);
 										if (Type != NAME_None)
 										{
 											// BytesToHex loses case
@@ -287,14 +323,29 @@ bool FUniqueNetIdRepl::NetSerialize(FArchive& Ar, UPackageMap* Map, bool& bOutSu
 					// If no type was encoded, assume default
 					TypeHash = UOnlineEngineInterface::Get()->GetReplicationHashForSubsystem(UOnlineEngineInterface::Get()->GetDefaultOnlineSubsystemName());
 				}
+				FName Type;
+				bool bValidTypeHash = TypeHash != 0;
+				if (TypeHash == TypeHash_Other)
+				{
+					FString TypeString;
+					Ar << TypeString;
+					Type = FName(*TypeString);
+					if (Ar.IsError() || Type == NAME_None)
+					{
+						bValidTypeHash = false;
+					}
+				}
+				else
+				{
+					Type = UOnlineEngineInterface::Get()->GetSubsystemFromReplicationHash(TypeHash);
+				}
 
-				if (TypeHash != 0)
+				if (bValidTypeHash)
 				{
 					FString Contents;
 					Ar << Contents;
 					if (!Ar.IsError())
 					{
-						FName Type = UOnlineEngineInterface::Get()->GetSubsystemFromReplicationHash(TypeHash);
 						if (Type != NAME_None)
 						{
 							UniqueIdFromString(Type, Contents);
@@ -488,6 +539,13 @@ void TestUniqueIdRepl(UWorld* InWorld)
 	CHECK_REPL_VALIDITY(UpperCaseStringIdIn, bSetupSuccess);
 	CHECK_REPL_VALIDITY(WayTooLongForHexEncodingIdIn, bSetupSuccess);
 
+	static FName NAME_CustomOSS(TEXT("MyCustomOSS"));
+	FUniqueNetIdRepl CustomOSSIdIn(UOnlineEngineInterface::Get()->CreateUniquePlayerId(TEXT("a8d245fc-4b97-4150-a3cd-c2c91d8fc4b3"), NAME_CustomOSS));
+	FUniqueNetIdRepl CustomOSSEncodedIdIn(UOnlineEngineInterface::Get()->CreateUniquePlayerId(TEXT("0123456789abcdef"), NAME_CustomOSS));
+
+	CHECK_REPL_VALIDITY(CustomOSSIdIn, bSetupSuccess);
+	CHECK_REPL_VALIDITY(CustomOSSEncodedIdIn, bSetupSuccess);
+
 	bool bRegularSerializationSuccess = true;
 	bool bNetworkSerializationSuccess = true;
 	if (bSetupSuccess)
@@ -507,6 +565,8 @@ void TestUniqueIdRepl(UWorld* InWorld)
 				TestUniqueIdWriter << NonHexStringIdIn;
 				TestUniqueIdWriter << UpperCaseStringIdIn;
 				TestUniqueIdWriter << WayTooLongForHexEncodingIdIn;
+				TestUniqueIdWriter << CustomOSSIdIn;
+				TestUniqueIdWriter << CustomOSSEncodedIdIn;
 			}
 
 			FUniqueNetIdRepl EmptyIdOut;
@@ -515,6 +575,8 @@ void TestUniqueIdRepl(UWorld* InWorld)
 			FUniqueNetIdRepl NonHexStringIdOut;
 			FUniqueNetIdRepl UpperCaseStringIdOut;
 			FUniqueNetIdRepl WayTooLongForHexEncodingIdOut;
+			FUniqueNetIdRepl CustomOSSIdOut;
+			FUniqueNetIdRepl CustomOSSEncodedIdOut;
 
 			// Serialize Out
 			{
@@ -525,6 +587,8 @@ void TestUniqueIdRepl(UWorld* InWorld)
 				TestUniqueIdReader << NonHexStringIdOut;
 				TestUniqueIdReader << UpperCaseStringIdOut;
 				TestUniqueIdReader << WayTooLongForHexEncodingIdOut;
+				TestUniqueIdReader << CustomOSSIdOut;
+				TestUniqueIdReader << CustomOSSEncodedIdOut;
 			}
 
 			if (EmptyIdOut.IsValid())
@@ -544,20 +608,18 @@ void TestUniqueIdRepl(UWorld* InWorld)
 			CHECK_REPL_EQUALITY(NonHexStringIdIn, NonHexStringIdOut, bRegularSerializationSuccess);
 			CHECK_REPL_EQUALITY(UpperCaseStringIdIn, UpperCaseStringIdOut, bRegularSerializationSuccess);
 			CHECK_REPL_EQUALITY(WayTooLongForHexEncodingIdIn, WayTooLongForHexEncodingIdOut, bRegularSerializationSuccess);
+			CHECK_REPL_EQUALITY(CustomOSSIdIn, CustomOSSIdOut, bRegularSerializationSuccess);
+			CHECK_REPL_EQUALITY(CustomOSSEncodedIdIn, CustomOSSEncodedIdOut, bRegularSerializationSuccess);
 		}
 
 		// Network serialization (network/transient using MakeReplicationData)
 		{
 			bool bOutSuccess = false;
 
-			TArray<uint8> Buffer;
-			Buffer.Empty();
-
 			// Serialize In
+			FNetBitWriter TestUniqueIdWriter(16 * 1024);
 			uint8 EncodingFailures = 0;
 			{
-				FMemoryWriter TestUniqueIdWriter(Buffer);
-
 				EmptyIdIn.NetSerialize(TestUniqueIdWriter, nullptr, bOutSuccess);
 				EncodingFailures += bOutSuccess ? 0 : 1;
 				ValidIdIn.NetSerialize(TestUniqueIdWriter, nullptr, bOutSuccess);
@@ -570,6 +632,10 @@ void TestUniqueIdRepl(UWorld* InWorld)
 				EncodingFailures += bOutSuccess ? 0 : 1;
 				WayTooLongForHexEncodingIdIn.NetSerialize(TestUniqueIdWriter, nullptr, bOutSuccess);
 				EncodingFailures += bOutSuccess ? 0 : 1;
+				CustomOSSIdIn.NetSerialize(TestUniqueIdWriter, nullptr, bOutSuccess);
+				EncodingFailures += bOutSuccess ? 0 : 1;
+				CustomOSSEncodedIdIn.NetSerialize(TestUniqueIdWriter, nullptr, bOutSuccess);
+				EncodingFailures += bOutSuccess ? 0 : 1;
 			}
 
 			if (EncodingFailures > 0)
@@ -586,11 +652,13 @@ void TestUniqueIdRepl(UWorld* InWorld)
 				FUniqueNetIdRepl NonHexStringIdOut;
 				FUniqueNetIdRepl UpperCaseStringIdOut;
 				FUniqueNetIdRepl WayTooLongForHexEncodingIdOut;
+				FUniqueNetIdRepl CustomOSSIdOut;
+				FUniqueNetIdRepl CustomOSSEncodedIdOut;
 
 				// Serialize Out
 				uint8 DecodingFailures = 0;
 				{
-					FMemoryReader TestUniqueIdReader(Buffer);
+					FNetBitReader TestUniqueIdReader(nullptr, TestUniqueIdWriter.GetData(), TestUniqueIdWriter.GetNumBits());
 
 					EmptyIdOut.NetSerialize(TestUniqueIdReader, nullptr, bOutSuccess);
 					DecodingFailures += bOutSuccess ? 0 : 1;
@@ -604,6 +672,10 @@ void TestUniqueIdRepl(UWorld* InWorld)
 					DecodingFailures += bOutSuccess ? 0 : 1;
 					WayTooLongForHexEncodingIdOut.NetSerialize(TestUniqueIdReader, nullptr, bOutSuccess);
 					DecodingFailures += bOutSuccess ? 0 : 1;
+					CustomOSSIdOut.NetSerialize(TestUniqueIdReader, nullptr, bOutSuccess);
+					DecodingFailures += bOutSuccess ? 0 : 1;
+					CustomOSSEncodedIdOut.NetSerialize(TestUniqueIdReader, nullptr, bOutSuccess);
+					DecodingFailures += bOutSuccess ? 0 : 1;
 				}
 
 				if (DecodingFailures > 0)
@@ -629,6 +701,8 @@ void TestUniqueIdRepl(UWorld* InWorld)
 				CHECK_REPL_EQUALITY(NonHexStringIdIn, NonHexStringIdOut, bNetworkSerializationSuccess);
 				CHECK_REPL_EQUALITY(UpperCaseStringIdIn, UpperCaseStringIdOut, bNetworkSerializationSuccess);
 				CHECK_REPL_EQUALITY(WayTooLongForHexEncodingIdIn, WayTooLongForHexEncodingIdOut, bNetworkSerializationSuccess);
+				CHECK_REPL_EQUALITY(CustomOSSIdIn, CustomOSSIdOut, bNetworkSerializationSuccess);
+				CHECK_REPL_EQUALITY(CustomOSSEncodedIdIn, CustomOSSEncodedIdOut, bNetworkSerializationSuccess);
 			}
 		}
 	}
@@ -715,7 +789,7 @@ void TestUniqueIdRepl(UWorld* InWorld)
 	UE_LOG(LogNet, Log, TEXT("	Network: %s"), bNetworkSerializationSuccess ? (bSetupSuccess ? TEXT("PASS") : TEXT("SKIPPED")) : TEXT("FAIL"));
 	UE_LOG(LogNet, Log, TEXT("	Platform: %s"), bPlatformSerializationSuccess ? (bSetupSuccess ? TEXT("PASS") : TEXT("SKIPPED")) : TEXT("FAIL"));
 	UE_LOG(LogNet, Log, TEXT("	JSON: %s"), bJSONSerializationSuccess ? (bSetupSuccess ? TEXT("PASS") : TEXT("SKIPPED")) : TEXT("FAIL"));
-	
+
 #undef CHECK_REPL_VALIDITY
 #undef CHECK_REPL_EQUALITY
 
