#DUPE MERGe
Merging using Dev-Anim<->Release4.21

fix motion blur issue on poseable mesh component

diff --git a/Engine/Source/Runtime/Engine/Classes/Components/PoseableMeshComponent.h b/Engine/Source/Runtime/Engine/Classes/Components/PoseableMeshComponent.h
--- a/Engine/Source/Runtime/Engine/Classes/Components/PoseableMeshComponent.h
+++ b/Engine/Source/Runtime/Engine/Classes/Components/PoseableMeshComponent.h
@@ -56,6 +56,7 @@ class ENGINE_API UPoseableMeshComponent : public USkinnedMeshComponent
 	//~ Begin USkinnedMeshComponent Interface
 	virtual void RefreshBoneTransforms(FActorComponentTickFunction* TickFunction = NULL) override;
 	virtual bool AllocateTransformData() override;
+	virtual bool ShouldUpdateTransform(bool bLODHasChanged) const override;
 	//~ End USkinnedMeshComponent Interface
 
 	bool IsRunningParallelEvaluation() const { return false; }
@@ -65,6 +66,11 @@ class ENGINE_API UPoseableMeshComponent : public USkinnedMeshComponent
 	 * This code also applies any per-bone rotators etc. as part of the composition process
 	 */
 	void FillComponentSpaceTransforms();
+
+private:
+	void MarkRefreshTransformDirty();
+	// this is marked if transform has to be updated
+	bool bNeedsRefreshTransform;
 };
 
 
diff --git a/Engine/Source/Runtime/Engine/Private/Components/PoseableMeshComponent.cpp b/Engine/Source/Runtime/Engine/Private/Components/PoseableMeshComponent.cpp
--- a/Engine/Source/Runtime/Engine/Private/Components/PoseableMeshComponent.cpp
+++ b/Engine/Source/Runtime/Engine/Private/Components/PoseableMeshComponent.cpp
@@ -11,6 +11,7 @@
 
 UPoseableMeshComponent::UPoseableMeshComponent(const FObjectInitializer& ObjectInitializer)
 	: Super(ObjectInitializer)
+	, bNeedsRefreshTransform (false)
 {
 
 }
@@ -69,6 +70,8 @@ void UPoseableMeshComponent::RefreshBoneTransforms(FActorComponentTickFunction*
 	UpdateBounds();
 	MarkRenderTransformDirty();
 	MarkRenderDynamicDataDirty();
+
+	bNeedsRefreshTransform = false;
 }
 
 void UPoseableMeshComponent::FillComponentSpaceTransforms()
@@ -156,8 +159,7 @@ void UPoseableMeshComponent::SetBoneTransformByName(FName BoneName, const FTrans
 				BoneSpaceTransforms[BoneIndex].SetToRelativeTransform(CSPose.GetComponentSpaceTransform(ParentIndex));
 			}
 
-			// Need to send new state to render thread
-			RefreshBoneTransforms();
+			MarkRefreshTransformDirty();
 		}
 	}
 }
@@ -308,6 +310,18 @@ void UPoseableMeshComponent::CopyPoseFromSkeletalComponent(const USkeletalMeshCo
 				}
 			}
 		}
-		RefreshBoneTransforms();
+
+		MarkRefreshTransformDirty();
 	}
 }
+
+bool UPoseableMeshComponent::ShouldUpdateTransform(bool bLODHasChanged) const
+{
+	// we don't always update transform - each function when they changed will update
+	return Super::ShouldUpdateTransform(bLODHasChanged) && bNeedsRefreshTransform;
+}
+
+void UPoseableMeshComponent::MarkRefreshTransformDirty()
+{
+	bNeedsRefreshTransform = true;
+}
\ No newline at end of file
diff --git a/Engine/Source/Runtime/Engine/Private/SkeletalMesh.cpp b/Engine/Source/Runtime/Engine/Private/SkeletalMesh.cpp
--- a/Engine/Source/Runtime/Engine/Private/SkeletalMesh.cpp
+++ b/Engine/Source/Runtime/Engine/Private/SkeletalMesh.cpp
@@ -3150,8 +3150,8 @@ FSkeletalMeshSceneProxy::FSkeletalMeshSceneProxy(const USkinnedMeshComponent* Co
 	// Force inset shadows if capsule shadows are requested, as they can't be supported with full scene shadows
 	bCastInsetShadow = bCastInsetShadow || bCastCapsuleDirectShadow;
 
-	const USkeletalMeshComponent* SkeletalMeshComponent = Cast<const USkeletalMeshComponent>(Component);
-	if(SkeletalMeshComponent && SkeletalMeshComponent->bPerBoneMotionBlur)
+	const USkinnedMeshComponent* SkinnedMeshComponent = Cast<const USkinnedMeshComponent>(Component);
+	if(SkinnedMeshComponent && SkinnedMeshComponent->bPerBoneMotionBlur)
 	{
 		bAlwaysHasVelocity = true;
 	}
@@ -3251,22 +3251,22 @@ FSkeletalMeshSceneProxy::FSkeletalMeshSceneProxy(const USkinnedMeshComponent* Co
 	SetPropertyColor(NewPropertyColor);
 
 	// Copy out shadow physics asset data
-	if(SkeletalMeshComponent)
+	if(SkinnedMeshComponent)
 	{
-		UPhysicsAsset* ShadowPhysicsAsset = SkeletalMeshComponent->SkeletalMesh->ShadowPhysicsAsset;
+		UPhysicsAsset* ShadowPhysicsAsset = SkinnedMeshComponent->SkeletalMesh->ShadowPhysicsAsset;
 
 		if (ShadowPhysicsAsset
-			&& SkeletalMeshComponent->CastShadow
-			&& (SkeletalMeshComponent->bCastCapsuleDirectShadow || SkeletalMeshComponent->bCastCapsuleIndirectShadow))
+			&& SkinnedMeshComponent->CastShadow
+			&& (SkinnedMeshComponent->bCastCapsuleDirectShadow || SkinnedMeshComponent->bCastCapsuleIndirectShadow))
 		{
 			for (int32 BodyIndex = 0; BodyIndex < ShadowPhysicsAsset->SkeletalBodySetups.Num(); BodyIndex++)
 			{
 				UBodySetup* BodySetup = ShadowPhysicsAsset->SkeletalBodySetups[BodyIndex];
-				int32 BoneIndex = SkeletalMeshComponent->GetBoneIndex(BodySetup->BoneName);
+				int32 BoneIndex = SkinnedMeshComponent->GetBoneIndex(BodySetup->BoneName);
 
 				if (BoneIndex != INDEX_NONE)
 				{
-					const FMatrix& RefBoneMatrix = SkeletalMeshComponent->SkeletalMesh->GetComposedRefPoseMatrix(BoneIndex);
+					const FMatrix& RefBoneMatrix = SkinnedMeshComponent->SkeletalMesh->GetComposedRefPoseMatrix(BoneIndex);
 
 					const int32 NumSpheres = BodySetup->AggGeom.SphereElems.Num();
 					for (int32 SphereIndex = 0; SphereIndex < NumSpheres; SphereIndex++)
