PostLoad check and additional GetClampedColorValue BP function for linear color curves

diff --git a/Engine/Source/Runtime/Core/Public/UObject/ReleaseObjectVersion.h b/Engine/Source/Runtime/Core/Public/UObject/ReleaseObjectVersion.h
--- a/Engine/Source/Runtime/Core/Public/UObject/ReleaseObjectVersion.h
+++ b/Engine/Source/Runtime/Core/Public/UObject/ReleaseObjectVersion.h
@@ -65,6 +65,9 @@ struct CORE_API FReleaseObjectVersion
 		// Change to MeshDescription serialization
 		MeshDescriptionNewSerialization,
 
+		// Change to not clamp RGB values > 1 on linear color curves
+		UnclampRGBColorCurves,
+
 		// -----<new versions can be added above this line>-------------------------------------------------
 		VersionPlusOne,
 		LatestVersion = VersionPlusOne - 1
diff --git a/Engine/Source/Runtime/Engine/Classes/Curves/CurveLinearColor.h b/Engine/Source/Runtime/Engine/Classes/Curves/CurveLinearColor.h
--- a/Engine/Source/Runtime/Engine/Classes/Curves/CurveLinearColor.h
+++ b/Engine/Source/Runtime/Engine/Classes/Curves/CurveLinearColor.h
@@ -43,6 +43,9 @@ class ENGINE_API UCurveLinearColor : public UCurveBase
 	UFUNCTION(BlueprintCallable, Category="Math|Curves")
 	virtual FLinearColor GetLinearColorValue(float InTime) const override;
 
+	UFUNCTION(BlueprintCallable, Category = "Math|Curves")
+	virtual FLinearColor GetClampedLinearColorValue(float InTime) const override;
+
 	FLinearColor GetUnadjustedLinearColorValue(float InTime) const;
 
 	bool HasAnyAlphaKeys() const override { return FloatCurves[3].GetNumKeys() > 0; }
@@ -64,6 +67,8 @@ public:
 
 	virtual void OnCurveChanged(const TArray<FRichCurveEditInfo>& ChangedCurveEditInfos) override;
 #endif
+	virtual void PostLoad() override;
+
 public:
 	// Properties for adjusting the color of the gradient
 	UPROPERTY(EditAnywhere, Category="Color", meta = (ClampMin = "0.0", ClampMax = "359.0"))
diff --git a/Engine/Source/Runtime/Engine/Classes/Curves/CurveOwnerInterface.h b/Engine/Source/Runtime/Engine/Classes/Curves/CurveOwnerInterface.h
--- a/Engine/Source/Runtime/Engine/Classes/Curves/CurveOwnerInterface.h
+++ b/Engine/Source/Runtime/Engine/Classes/Curves/CurveOwnerInterface.h
@@ -48,6 +48,12 @@ public:
 		return FLinearColor::Black;
 	}
 
+	/** Evaluate this color curve at the specified time */
+	virtual FLinearColor GetClampedLinearColorValue(float InTime) const
+	{
+		return FLinearColor::Black;
+	}
+
 	/** @return True if the curve has any alpha keys */
 	virtual bool HasAnyAlphaKeys() const
 	{
diff --git a/Engine/Source/Runtime/Engine/Private/Curves/CurveLinearColor.cpp b/Engine/Source/Runtime/Engine/Private/Curves/CurveLinearColor.cpp
--- a/Engine/Source/Runtime/Engine/Private/Curves/CurveLinearColor.cpp
+++ b/Engine/Source/Runtime/Engine/Private/Curves/CurveLinearColor.cpp
@@ -6,6 +6,8 @@
 
 #include "Curves/CurveLinearColor.h"
 #include "CanvasItem.h"
+#include "UObject/ReleaseObjectVersion.h"
+#include "Logging/MessageLog.h"
 
 FLinearColor FRuntimeCurveLinearColor::GetLinearColorValue(float InTime) const
 {
@@ -111,6 +113,75 @@ FLinearColor UCurveLinearColor::GetLinearColorValue( float InTime ) const
 	return LinearColor;
 }
 
+
+FLinearColor UCurveLinearColor::GetClampedLinearColorValue(float InTime) const
+{
+	// Logic copied from .\Engine\Source\Developer\TextureCompressor\Private\TextureCompressorModule.cpp
+	const FLinearColor OriginalColor = GetUnadjustedLinearColorValue(InTime);
+
+	// Always clamp value
+	const bool bShouldClampValue = true;
+
+	// Convert to HSV
+	FLinearColor HSVColor = OriginalColor.LinearRGBToHSV();
+	float& PixelHue = HSVColor.R;
+	float& PixelSaturation = HSVColor.G;
+	float& PixelValue = HSVColor.B;
+
+	// Apply brightness adjustment
+	PixelValue *= AdjustBrightness;
+
+	// Apply brightness power adjustment
+	if (!FMath::IsNearlyEqual(AdjustBrightnessCurve, 1.0f, (float)KINDA_SMALL_NUMBER) && AdjustBrightnessCurve != 0.0f)
+	{
+		// Raise HSV.V to the specified power
+		PixelValue = FMath::Pow(PixelValue, AdjustBrightnessCurve);
+	}
+
+	// Apply "vibrancy" adjustment
+	if (!FMath::IsNearlyZero(AdjustVibrance, (float)KINDA_SMALL_NUMBER))
+	{
+		const float SatRaisePow = 5.0f;
+		const float InvSatRaised = FMath::Pow(1.0f - PixelSaturation, SatRaisePow);
+
+		const float ClampedVibrance = FMath::Clamp(AdjustVibrance, 0.0f, 1.0f);
+		const float HalfVibrance = ClampedVibrance * 0.5f;
+
+		const float SatProduct = HalfVibrance * InvSatRaised;
+
+		PixelSaturation += SatProduct;
+	}
+
+	// Apply saturation adjustment
+	PixelSaturation *= AdjustSaturation;
+
+	// Apply hue adjustment
+	PixelHue += AdjustHue;
+
+	// Clamp HSV values
+	{
+		PixelHue = FMath::Fmod(PixelHue, 360.0f);
+		if (PixelHue < 0.0f)
+		{
+			// Keep the hue value positive as HSVToLinearRGB prefers that
+			PixelHue += 360.0f;
+		}
+		PixelSaturation = FMath::Clamp(PixelSaturation, 0.0f, 1.0f);
+
+		if (bShouldClampValue)
+		{
+			PixelValue = FMath::Clamp(PixelValue, 0.0f, 1.0f);
+		}
+	}
+
+	// Convert back to a linear color
+	FLinearColor LinearColor = HSVColor.HSVToLinearRGB();
+
+	// Remap the alpha channel
+	LinearColor.A = FMath::Lerp(AdjustMinAlpha, AdjustMaxAlpha, OriginalColor.A);
+	return LinearColor;
+}
+
 FLinearColor UCurveLinearColor::GetUnadjustedLinearColorValue(float InTime) const
 {
 	FLinearColor Result;
@@ -219,6 +290,19 @@ void UCurveLinearColor::OnCurveChanged(const TArray<FRichCurveEditInfo>& Changed
 }
 #endif
 
+void UCurveLinearColor::PostLoad()
+{
+	Super::PostLoad();
+
+#if WITH_EDITOR
+	if (GetLinkerCustomVersion(FReleaseObjectVersion::GUID) < FReleaseObjectVersion::UnclampRGBColorCurves)
+	{
+		FMessageLog("LoadErrors").Warning(FText::Format(NSLOCTEXT("CurveEditor","CurveDataUpdate", "Linear color curves now accurately handle RGB values > 1. If you were relying on HSV clamping, please update {0}"),
+			FText::FromString(GetName())));
+	}
+#endif
+}
+
 void UCurveLinearColor::WritePixel(uint8* Pixel, const FLinearColor& Color)
 {
 	Pixel[0] = FMath::FloorToInt(Color.B * 255.999f);
