Adding custom version to linear color curve assets

diff --git a/Engine/Source/Runtime/Engine/Classes/Curves/CurveLinearColor.h b/Engine/Source/Runtime/Engine/Classes/Curves/CurveLinearColor.h
--- a/Engine/Source/Runtime/Engine/Classes/Curves/CurveLinearColor.h
+++ b/Engine/Source/Runtime/Engine/Classes/Curves/CurveLinearColor.h
@@ -69,6 +69,8 @@ public:
 #endif
 	virtual void PostLoad() override;
 
+	virtual void Serialize(FArchive& Ar) override;
+
 public:
 	// Properties for adjusting the color of the gradient
 	UPROPERTY(EditAnywhere, Category="Color", meta = (ClampMin = "0.0", ClampMax = "359.0"))
diff --git a/Engine/Source/Runtime/Engine/Private/Curves/CurveLinearColor.cpp b/Engine/Source/Runtime/Engine/Private/Curves/CurveLinearColor.cpp
--- a/Engine/Source/Runtime/Engine/Private/Curves/CurveLinearColor.cpp
+++ b/Engine/Source/Runtime/Engine/Private/Curves/CurveLinearColor.cpp
@@ -303,6 +303,12 @@ void UCurveLinearColor::PostLoad()
 #endif
 }
 
+void UCurveLinearColor::Serialize(FArchive& Ar)
+{
+	Ar.UsingCustomVersion(FReleaseObjectVersion::GUID);
+	Super::Serialize(Ar);
+}
+
 void UCurveLinearColor::WritePixel(uint8* Pixel, const FLinearColor& Color)
 {
 	Pixel[0] = FMath::FloorToInt(Color.B * 255.999f);
