Deprecates Bokeh DOF, Gaussian DOF and Circle DOF for deferred shading renderer.

diff --git a/Engine/Plugins/Runtime/Nvidia/Ansel/Source/Ansel/Private/Ansel.cpp b/Engine/Plugins/Runtime/Nvidia/Ansel/Source/Ansel/Private/Ansel.cpp
--- a/Engine/Plugins/Runtime/Nvidia/Ansel/Source/Ansel/Private/Ansel.cpp
+++ b/Engine/Plugins/Runtime/Nvidia/Ansel/Source/Ansel/Private/Ansel.cpp
@@ -298,6 +298,7 @@ void FNVAnselCameraPhotographyPrivate::DoCustomUIControls(FPostProcessSettings&
 		UEPostProcessingOriginal = InOutPPSettings;
 
 		// add all relevant controls
+		PRAGMA_DISABLE_DEPRECATION_WARNINGS
 		if (bEffectUIAllowed[DepthOfField] && (InOutPPSettings.DepthOfFieldMethod == DOFM_Gaussian || InOutPPSettings.DepthOfFieldScale > 0.f)) // 'Is DoF used?'
 		{
 			if (InOutPPSettings.DepthOfFieldMethod != DOFM_Gaussian) // scale irrelevant for gaussian
@@ -323,6 +324,7 @@ void FNVAnselCameraPhotographyPrivate::DoCustomUIControls(FPostProcessSettings&
 				InOutPPSettings.DepthOfFieldFocalDistance
 			);
 		}
+		PRAGMA_ENABLE_DEPRECATION_WARNINGS
 
 		if (bEffectUIAllowed[Bloom] &&
 			InOutPPSettings.BloomIntensity > 0.f)
diff --git a/Engine/Source/Runtime/CinematicCamera/Private/CineCameraComponent.cpp b/Engine/Source/Runtime/CinematicCamera/Private/CineCameraComponent.cpp
--- a/Engine/Source/Runtime/CinematicCamera/Private/CineCameraComponent.cpp
+++ b/Engine/Source/Runtime/CinematicCamera/Private/CineCameraComponent.cpp
@@ -42,8 +42,10 @@ UCineCameraComponent::UCineCameraComponent()
 
 	bConstrainAspectRatio = true;
 
+	PRAGMA_DISABLE_DEPRECATION_WARNINGS
 	// Default to CircleDOF, but allow the user to customize it
 	PostProcessSettings.DepthOfFieldMethod = DOFM_CircleDOF;
+	PRAGMA_ENABLE_DEPRECATION_WARNINGS
 
 	RecalcDerivedData();
 
@@ -387,7 +389,9 @@ void UCineCameraComponent::UpdateCameraLens(float DeltaTime, FMinimalViewInfo& D
 {
 	if (FocusSettings.FocusMethod == ECameraFocusMethod::None)
 	{
+		PRAGMA_DISABLE_DEPRECATION_WARNINGS
 		DesiredView.PostProcessSettings.bOverride_DepthOfFieldMethod = false;
+		PRAGMA_ENABLE_DEPRECATION_WARNINGS
 		DesiredView.PostProcessSettings.bOverride_DepthOfFieldFstop = false;
 		DesiredView.PostProcessSettings.bOverride_DepthOfFieldMinFstop = false;
 		DesiredView.PostProcessSettings.bOverride_DepthOfFieldBladeCount = false;
@@ -398,8 +402,10 @@ void UCineCameraComponent::UpdateCameraLens(float DeltaTime, FMinimalViewInfo& D
 	{
 		// Update focus/DoF
 		DesiredView.PostProcessBlendWeight = 1.f;
+		PRAGMA_DISABLE_DEPRECATION_WARNINGS
 		DesiredView.PostProcessSettings.bOverride_DepthOfFieldMethod = true;
 		DesiredView.PostProcessSettings.DepthOfFieldMethod = PostProcessSettings.DepthOfFieldMethod;
+		PRAGMA_ENABLE_DEPRECATION_WARNINGS
 
 		DesiredView.PostProcessSettings.bOverride_DepthOfFieldFstop = true;
 		DesiredView.PostProcessSettings.DepthOfFieldFstop = CurrentAperture;
diff --git a/Engine/Source/Runtime/Engine/Classes/Engine/Scene.h b/Engine/Source/Runtime/Engine/Classes/Engine/Scene.h
--- a/Engine/Source/Runtime/Engine/Classes/Engine/Scene.h
+++ b/Engine/Source/Runtime/Engine/Classes/Engine/Scene.h
@@ -571,9 +571,11 @@ struct FWeightedBlendables
 // Each property consists of a bool to enable it (by default off),
 // the variable declaration and further down the default value for it.
 // The comment should include the meaning and usable range.
+PRAGMA_DISABLE_DEPRECATION_WARNINGS
 USTRUCT(BlueprintType, meta=(HiddenByDefault))
 struct FPostProcessSettings
 {
+	PRAGMA_ENABLE_DEPRECATION_WARNINGS
 	GENERATED_USTRUCT_BODY()
 
 	// first all bOverride_... as they get grouped together into bitfields
@@ -975,6 +977,7 @@ struct FPostProcessSettings
 	UPROPERTY(EditAnywhere, BlueprintReadWrite, Category=Overrides, meta=(PinHiddenByDefault, InlineEditConditionToggle))
 	uint8 bOverride_DepthOfFieldScale:1;
 
+	DEPRECATED(4.21, "This property is now deprecated, please use circle DOF method for deferred shading renderer.")
 	UPROPERTY(EditAnywhere, BlueprintReadWrite, Category=Overrides, meta=(PinHiddenByDefault, InlineEditConditionToggle))
 	uint8 bOverride_DepthOfFieldMaxBokehSize:1;
 
@@ -984,22 +987,26 @@ struct FPostProcessSettings
 	UPROPERTY(EditAnywhere, BlueprintReadWrite, Category=Overrides, meta=(PinHiddenByDefault, InlineEditConditionToggle))
 	uint8 bOverride_DepthOfFieldFarBlurSize:1;
 
+	DEPRECATED(4.21, "This property is now deprecated, please use circle DOF method for deferred shading renderer. Mobile renderer still support Gaussian method.")
 	UPROPERTY(EditAnywhere, BlueprintReadWrite, Category=Overrides, meta=(PinHiddenByDefault, InlineEditConditionToggle))
 	uint8 bOverride_DepthOfFieldMethod:1;
 
 	UPROPERTY(EditAnywhere, BlueprintReadWrite, Category=Overrides, meta=(PinHiddenByDefault, InlineEditConditionToggle))
 	uint8 bOverride_MobileHQGaussian:1;
 
+	DEPRECATED(4.21, "This property is now deprecated, please use circle DOF method for deferred shading renderer.")
 	UPROPERTY(EditAnywhere, BlueprintReadWrite, Category=Overrides, meta=(PinHiddenByDefault, InlineEditConditionToggle))
 	uint8 bOverride_DepthOfFieldBokehShape:1;
 
 	UPROPERTY(EditAnywhere, BlueprintReadWrite, Category=Overrides, meta=(PinHiddenByDefault, InlineEditConditionToggle))
 	uint8 bOverride_DepthOfFieldOcclusion:1;
 
-	UPROPERTY(EditAnywhere, BlueprintReadWrite, Category=Overrides, meta=(PinHiddenByDefault, InlineEditConditionToggle))
+	DEPRECATED(4.21, "This property is now deprecated, please use circle DOF method for deferred shading renderer.")
+	UPROPERTY(EditAnywhere, BlueprintReadWrite, Category = Overrides, meta = (PinHiddenByDefault, InlineEditConditionToggle))
 	uint8 bOverride_DepthOfFieldColorThreshold:1;
 
-	UPROPERTY(EditAnywhere, BlueprintReadWrite, Category=Overrides, meta=(PinHiddenByDefault, InlineEditConditionToggle))
+	DEPRECATED(4.21, "This property is now deprecated, please use circle DOF method for deferred shading renderer.")
+	UPROPERTY(EditAnywhere, BlueprintReadWrite, Category = Overrides, meta = (PinHiddenByDefault, InlineEditConditionToggle))
 	uint8 bOverride_DepthOfFieldSizeThreshold:1;
 
 	UPROPERTY(EditAnywhere, BlueprintReadWrite, Category=Overrides, meta=(PinHiddenByDefault, InlineEditConditionToggle))
@@ -1030,7 +1037,7 @@ struct FPostProcessSettings
 	uint8 bOverride_ScreenSpaceReflectionMaxRoughness:1;
 
 	UPROPERTY(EditAnywhere, BlueprintReadWrite, Category=Overrides, meta=(PinHiddenByDefault, InlineEditConditionToggle))
-	uint8 bOverride_ScreenSpaceReflectionRoughnessScale:1;
+	uint8 bOverride_ScreenSpaceReflectionRoughnessScale:1; // TODO: look useless...
 
 	// -----------------------------------------------------------------------
 
@@ -1046,7 +1053,8 @@ struct FPostProcessSettings
 	UPROPERTY(EditAnywhere, BlueprintReadWrite, Category="Lens|Exposure", meta=(editcondition = "bOverride_AutoExposureMethod", DisplayName = "Metering Mode"))
     TEnumAsByte<enum EAutoExposureMethod> AutoExposureMethod;
 
-	/** BokehDOF, Simple gaussian, ... Mobile supports Gaussian only. */
+	/** DEPRECATED: BokehDOF, Simple gaussian, ... Mobile supports Gaussian only. Deferred shading renderer will only support CircleDOF. */
+	DEPRECATED(4.21, "This property is now deprecated, please use circle DOF method for deferred shading renderer. Mobile renderer still support Gaussian method.")
 	UPROPERTY(EditAnywhere, BlueprintReadWrite, Category="Lens|Depth of Field", meta=(editcondition = "bOverride_DepthOfFieldMethod", DisplayName = "Method"))
 	TEnumAsByte<enum EDepthOfFieldMethod> DepthOfFieldMethod;
 
@@ -1528,7 +1536,8 @@ struct FPostProcessSettings
 	UPROPERTY(interp, BlueprintReadWrite, Category="Lens|Depth of Field", meta=(ClampMin = "0.0", ClampMax = "2.0", editcondition = "bOverride_DepthOfFieldScale", DisplayName = "Scale"))
 	float DepthOfFieldScale;
 
-	/** BokehDOF only: Maximum size of the Depth of Field blur (in percent of the view width) (note: performance cost scales with size*size) */
+	/** DEPRECATED: BokehDOF only: Maximum size of the Depth of Field blur (in percent of the view width) (note: performance cost scales with size*size) */
+	DEPRECATED(4.21, "This property is now deprecated, please use circle DOF method for deferred shading renderer.")
 	UPROPERTY(interp, BlueprintReadWrite, Category="Lens|Depth of Field", meta=(UIMin = "0.0", UIMax = "32.0", editcondition = "bOverride_DepthOfFieldMaxBokehSize", DisplayName = "Max Bokeh Size"))
 	float DepthOfFieldMaxBokehSize;
 
@@ -1544,15 +1553,18 @@ struct FPostProcessSettings
 	UPROPERTY(interp, BlueprintReadWrite, Category="Lens|Depth of Field", AdvancedDisplay, meta=(ClampMin = "0.0", ClampMax = "1.0", editcondition = "bOverride_DepthOfFieldOcclusion", DisplayName = "Occlusion"))
 	float DepthOfFieldOcclusion;
 
-	/** Defines the shape of the Bokeh when object get out of focus, cannot be blended */
+	/** DEPRECATED: Defines the shape of the Bokeh when object get out of focus, cannot be blended */
+	DEPRECATED(4.21, "This property is now deprecated, please use circle DOF method for deferred shading renderer.")
 	UPROPERTY(EditAnywhere, BlueprintReadWrite, AdvancedDisplay, Category="Lens|Depth of Field", meta=(editcondition = "bOverride_DepthOfFieldBokehShape", DisplayName = "Shape"))
 	class UTexture* DepthOfFieldBokehShape;
 
-	/** Color threshold to do full quality DOF (BokehDOF only) */
+	/** DEPRECATED: Color threshold to do full quality DOF (BokehDOF only) */
+	DEPRECATED(4.21, "This property is now deprecated, please use circle DOF method for deferred shading renderer.")
 	UPROPERTY(interp, BlueprintReadWrite, Category="Lens|Depth of Field", AdvancedDisplay, meta=(ClampMin = "0.0", ClampMax = "10.0", editcondition = "bOverride_DepthOfFieldColorThreshold", DisplayName = "Color Threshold"))
 	float DepthOfFieldColorThreshold;
 
-	/** Size threshold to do full quality DOF (BokehDOF only) */
+	/** DEPRECATED: Size threshold to do full quality DOF (BokehDOF only) */
+	DEPRECATED(4.21, "This property is now deprecated, please use circle DOF method for deferred shading renderer.")
 	UPROPERTY(interp, BlueprintReadWrite, Category="Lens|Depth of Field", AdvancedDisplay, meta=(ClampMin = "0.0", ClampMax = "1.0", editcondition = "bOverride_DepthOfFieldSizeThreshold", DisplayName = "Size Threshold"))
 	float DepthOfFieldSizeThreshold;
 	
@@ -1696,6 +1708,18 @@ struct FPostProcessSettings
 			BloomConvolutionPreFilterMax = BloomConvolutionPreFilter_DEPRECATED.Y;
 			BloomConvolutionPreFilterMult = BloomConvolutionPreFilter_DEPRECATED.Z;
 		}
+
+		/** Actually delete settings related to BokehDOF, and mark the DepthOfFieldMethod to _DEPRECATED. To keep backward compatibility,
+		 * need to have depth of field disabled by default, which means for deferred shading renderer having redicously high aperture.
+		 *
+		 * The plan here is to change the default of DepthOfFieldFstop in FPostProcessSettings's constructor to this MyMagicHightFstopValue.
+		 * And here in this OnAfterLoad() method do:
+		 *
+		 * if (DepthOfFieldFstop == MyMagicHightFstopValue && DepthOfFieldMethod_DEPRECATED == DOFM_CircleDOF)
+		 * {
+		 *		DepthOfFieldFstop = 4.0f;
+		 * }
+		 */
 	}
 #endif
 
@@ -1743,6 +1767,8 @@ struct FPostProcessSettings
 	// good start values for a new volume, by default no value is overriding
 	ENGINE_API FPostProcessSettings();
 
+	ENGINE_API FPostProcessSettings(const FPostProcessSettings&);
+
 	/**
 		* Used to define the values before any override happens.
 		* Should be as neutral as possible.
@@ -1758,7 +1784,10 @@ struct FPostProcessSettings
 
 	// Default number of blade of the diaphragm to simulate in depth of field.
 	static constexpr int32 kDefaultDepthOfFieldBladeCount = 5;
+
+	PRAGMA_DISABLE_DEPRECATION_WARNINGS
 };
+PRAGMA_ENABLE_DEPRECATION_WARNINGS
 
 UCLASS()
 class UScene : public UObject
diff --git a/Engine/Source/Runtime/Engine/Private/PostProcessVolume.cpp b/Engine/Source/Runtime/Engine/Private/PostProcessVolume.cpp
--- a/Engine/Source/Runtime/Engine/Private/PostProcessVolume.cpp
+++ b/Engine/Source/Runtime/Engine/Private/PostProcessVolume.cpp
@@ -65,6 +65,7 @@ bool APostProcessVolume::CanEditChange(const UProperty* InProperty) const
 		FString PropertyName = InProperty->GetName();
 
 		// Settings, can be shared for multiple objects types (volume, component, camera, player)
+		PRAGMA_DISABLE_DEPRECATION_WARNINGS
 		{
 			if (PropertyName == GET_MEMBER_NAME_STRING_CHECKED(FPostProcessSettings, DepthOfFieldScale))
 			{
@@ -173,6 +174,7 @@ bool APostProcessVolume::CanEditChange(const UProperty* InProperty) const
 			}
 
 		}
+		PRAGMA_ENABLE_DEPRECATION_WARNINGS
 
 		if (PropertyName == GET_MEMBER_NAME_STRING_CHECKED(APostProcessVolume, bEnabled))
 		{
diff --git a/Engine/Source/Runtime/Engine/Private/Scene.cpp b/Engine/Source/Runtime/Engine/Private/Scene.cpp
--- a/Engine/Source/Runtime/Engine/Private/Scene.cpp
+++ b/Engine/Source/Runtime/Engine/Private/Scene.cpp
@@ -462,6 +462,7 @@ FPostProcessSettings::FPostProcessSettings()
 	IndirectLightingColor = FLinearColor(1.0f, 1.0f, 1.0f);
 	IndirectLightingIntensity = 1.0f;
 	ColorGradingIntensity = 1.0f;
+	PRAGMA_DISABLE_DEPRECATION_WARNINGS
 	DepthOfFieldFocalDistance = 1000.0f;
 	DepthOfFieldFstop = 4.0f;
 	DepthOfFieldMinFstop = 1.2f;
@@ -480,6 +481,7 @@ FPostProcessSettings::FPostProcessSettings()
 	DepthOfFieldColorThreshold = 1.0f;
 	DepthOfFieldSizeThreshold = 0.08f;
 	DepthOfFieldSkyFocusDistance = 0.0f;
+	PRAGMA_ENABLE_DEPRECATION_WARNINGS
 	// 200 should be enough even for extreme aspect ratios to give the default no effect
 	DepthOfFieldVignetteSize = 200.0f;
 	LensFlareTints[0] = FLinearColor(1.0f, 0.8f, 0.4f, 0.6f);
@@ -510,6 +512,336 @@ FPostProcessSettings::FPostProcessSettings()
 #endif // DO_CHECK
 }
 
+FPostProcessSettings::FPostProcessSettings(const FPostProcessSettings& Settings)
+	: bOverride_WhiteTemp(Settings.bOverride_WhiteTemp)
+	, bOverride_WhiteTint(Settings.bOverride_WhiteTint)
+	, bOverride_ColorSaturation(Settings.bOverride_ColorSaturation)
+	, bOverride_ColorContrast(Settings.bOverride_ColorContrast)
+	, bOverride_ColorGamma(Settings.bOverride_ColorGamma)
+	, bOverride_ColorGain(Settings.bOverride_ColorGain)
+	, bOverride_ColorOffset(Settings.bOverride_ColorOffset)
+	, bOverride_ColorSaturationShadows(Settings.bOverride_ColorSaturationShadows)
+	, bOverride_ColorContrastShadows(Settings.bOverride_ColorContrastShadows)
+	, bOverride_ColorGammaShadows(Settings.bOverride_ColorGammaShadows)
+	, bOverride_ColorGainShadows(Settings.bOverride_ColorGainShadows)
+	, bOverride_ColorOffsetShadows(Settings.bOverride_ColorOffsetShadows)
+	, bOverride_ColorSaturationMidtones(Settings.bOverride_ColorSaturationMidtones)
+	, bOverride_ColorContrastMidtones(Settings.bOverride_ColorContrastMidtones)
+	, bOverride_ColorGammaMidtones(Settings.bOverride_ColorGammaMidtones)
+	, bOverride_ColorGainMidtones(Settings.bOverride_ColorGainMidtones)
+	, bOverride_ColorOffsetMidtones(Settings.bOverride_ColorOffsetMidtones)
+	, bOverride_ColorSaturationHighlights(Settings.bOverride_ColorSaturationHighlights)
+	, bOverride_ColorContrastHighlights(Settings.bOverride_ColorContrastHighlights)
+	, bOverride_ColorGammaHighlights(Settings.bOverride_ColorGammaHighlights)
+	, bOverride_ColorGainHighlights(Settings.bOverride_ColorGainHighlights)
+	, bOverride_ColorOffsetHighlights(Settings.bOverride_ColorOffsetHighlights)
+	, bOverride_ColorCorrectionShadowsMax(Settings.bOverride_ColorCorrectionShadowsMax)
+	, bOverride_ColorCorrectionHighlightsMin(Settings.bOverride_ColorCorrectionHighlightsMin)
+	, bOverride_BlueCorrection(Settings.bOverride_BlueCorrection)
+	, bOverride_ExpandGamut(Settings.bOverride_ExpandGamut)
+	, bOverride_FilmWhitePoint(Settings.bOverride_FilmWhitePoint)
+	, bOverride_FilmSaturation(Settings.bOverride_FilmSaturation)
+	, bOverride_FilmChannelMixerRed(Settings.bOverride_FilmChannelMixerRed)
+	, bOverride_FilmChannelMixerGreen(Settings.bOverride_FilmChannelMixerGreen)
+	, bOverride_FilmChannelMixerBlue(Settings.bOverride_FilmChannelMixerBlue)
+	, bOverride_FilmContrast(Settings.bOverride_FilmContrast)
+	, bOverride_FilmDynamicRange(Settings.bOverride_FilmDynamicRange)
+	, bOverride_FilmHealAmount(Settings.bOverride_FilmHealAmount)
+	, bOverride_FilmToeAmount(Settings.bOverride_FilmToeAmount)
+	, bOverride_FilmShadowTint(Settings.bOverride_FilmShadowTint)
+	, bOverride_FilmShadowTintBlend(Settings.bOverride_FilmShadowTintBlend)
+	, bOverride_FilmShadowTintAmount(Settings.bOverride_FilmShadowTintAmount)
+	, bOverride_FilmSlope(Settings.bOverride_FilmSlope)
+	, bOverride_FilmToe(Settings.bOverride_FilmToe)
+	, bOverride_FilmShoulder(Settings.bOverride_FilmShoulder)
+	, bOverride_FilmBlackClip(Settings.bOverride_FilmBlackClip)
+	, bOverride_FilmWhiteClip(Settings.bOverride_FilmWhiteClip)
+	, bOverride_SceneColorTint(Settings.bOverride_SceneColorTint)
+	, bOverride_SceneFringeIntensity(Settings.bOverride_SceneFringeIntensity)
+	, bOverride_ChromaticAberrationStartOffset(Settings.bOverride_ChromaticAberrationStartOffset)
+	, bOverride_AmbientCubemapTint(Settings.bOverride_AmbientCubemapTint)
+	, bOverride_AmbientCubemapIntensity(Settings.bOverride_AmbientCubemapIntensity)
+	, bOverride_BloomMethod(Settings.bOverride_BloomMethod)
+	, bOverride_BloomIntensity(Settings.bOverride_BloomIntensity)
+	, bOverride_BloomThreshold(Settings.bOverride_BloomThreshold)
+	, bOverride_Bloom1Tint(Settings.bOverride_Bloom1Tint)
+	, bOverride_Bloom1Size(Settings.bOverride_Bloom1Size)
+	, bOverride_Bloom2Size(Settings.bOverride_Bloom2Size)
+	, bOverride_Bloom2Tint(Settings.bOverride_Bloom2Tint)
+	, bOverride_Bloom3Tint(Settings.bOverride_Bloom3Tint)
+	, bOverride_Bloom3Size(Settings.bOverride_Bloom3Size)
+	, bOverride_Bloom4Tint(Settings.bOverride_Bloom4Tint)
+	, bOverride_Bloom4Size(Settings.bOverride_Bloom4Size)
+	, bOverride_Bloom5Tint(Settings.bOverride_Bloom5Tint)
+	, bOverride_Bloom5Size(Settings.bOverride_Bloom5Size)
+	, bOverride_Bloom6Tint(Settings.bOverride_Bloom6Tint)
+	, bOverride_Bloom6Size(Settings.bOverride_Bloom6Size)
+	, bOverride_BloomSizeScale(Settings.bOverride_BloomSizeScale)
+	, bOverride_BloomConvolutionTexture(Settings.bOverride_BloomConvolutionTexture)
+	, bOverride_BloomConvolutionSize(Settings.bOverride_BloomConvolutionSize)
+	, bOverride_BloomConvolutionCenterUV(Settings.bOverride_BloomConvolutionCenterUV)
+	//, bOverride_BloomConvolutionPreFilter_DEPRECATED(Settings.bOverride_BloomConvolutionPreFilter_DEPRECATED)
+	, bOverride_BloomConvolutionPreFilterMin(Settings.bOverride_BloomConvolutionPreFilterMin)
+	, bOverride_BloomConvolutionPreFilterMax(Settings.bOverride_BloomConvolutionPreFilterMax)
+	, bOverride_BloomConvolutionPreFilterMult(Settings.bOverride_BloomConvolutionPreFilterMult)
+	, bOverride_BloomConvolutionBufferScale(Settings.bOverride_BloomConvolutionBufferScale)
+	, bOverride_BloomDirtMaskIntensity(Settings.bOverride_BloomDirtMaskIntensity)
+	, bOverride_BloomDirtMaskTint(Settings.bOverride_BloomDirtMaskTint)
+	, bOverride_BloomDirtMask(Settings.bOverride_BloomDirtMask)
+	, bOverride_CameraShutterSpeed(Settings.bOverride_CameraShutterSpeed)
+	, bOverride_CameraISO(Settings.bOverride_CameraISO)
+	, bOverride_AutoExposureMethod(Settings.bOverride_AutoExposureMethod)
+	, bOverride_AutoExposureLowPercent(Settings.bOverride_AutoExposureLowPercent)
+	, bOverride_AutoExposureHighPercent(Settings.bOverride_AutoExposureHighPercent)
+	, bOverride_AutoExposureMinBrightness(Settings.bOverride_AutoExposureMinBrightness)
+	, bOverride_AutoExposureMaxBrightness(Settings.bOverride_AutoExposureMaxBrightness)
+	, bOverride_AutoExposureCalibrationConstant(Settings.bOverride_AutoExposureCalibrationConstant)
+	, bOverride_AutoExposureSpeedUp(Settings.bOverride_AutoExposureSpeedUp)
+	, bOverride_AutoExposureSpeedDown(Settings.bOverride_AutoExposureSpeedDown)
+	, bOverride_AutoExposureBias(Settings.bOverride_AutoExposureBias)
+	, bOverride_HistogramLogMin(Settings.bOverride_HistogramLogMin)
+	, bOverride_HistogramLogMax(Settings.bOverride_HistogramLogMax)
+	, bOverride_LensFlareIntensity(Settings.bOverride_LensFlareIntensity)
+	, bOverride_LensFlareTint(Settings.bOverride_LensFlareTint)
+	, bOverride_LensFlareTints(Settings.bOverride_LensFlareTints)
+	, bOverride_LensFlareBokehSize(Settings.bOverride_LensFlareBokehSize)
+	, bOverride_LensFlareBokehShape(Settings.bOverride_LensFlareBokehShape)
+	, bOverride_LensFlareThreshold(Settings.bOverride_LensFlareThreshold)
+	, bOverride_VignetteIntensity(Settings.bOverride_VignetteIntensity)
+	, bOverride_GrainIntensity(Settings.bOverride_GrainIntensity)
+	, bOverride_GrainJitter(Settings.bOverride_GrainJitter)
+	, bOverride_AmbientOcclusionIntensity(Settings.bOverride_AmbientOcclusionIntensity)
+	, bOverride_AmbientOcclusionStaticFraction(Settings.bOverride_AmbientOcclusionStaticFraction)
+	, bOverride_AmbientOcclusionRadius(Settings.bOverride_AmbientOcclusionRadius)
+	, bOverride_AmbientOcclusionFadeDistance(Settings.bOverride_AmbientOcclusionFadeDistance)
+	, bOverride_AmbientOcclusionFadeRadius(Settings.bOverride_AmbientOcclusionFadeRadius)
+	//, bOverride_AmbientOcclusionDistance_DEPRECATED(Settings.bOverride_AmbientOcclusionDistance_DEPRECATED)
+	, bOverride_AmbientOcclusionRadiusInWS(Settings.bOverride_AmbientOcclusionRadiusInWS)
+	, bOverride_AmbientOcclusionPower(Settings.bOverride_AmbientOcclusionPower)
+	, bOverride_AmbientOcclusionBias(Settings.bOverride_AmbientOcclusionBias)
+	, bOverride_AmbientOcclusionQuality(Settings.bOverride_AmbientOcclusionQuality)
+	, bOverride_AmbientOcclusionMipBlend(Settings.bOverride_AmbientOcclusionMipBlend)
+	, bOverride_AmbientOcclusionMipScale(Settings.bOverride_AmbientOcclusionMipScale)
+	, bOverride_AmbientOcclusionMipThreshold(Settings.bOverride_AmbientOcclusionMipThreshold)
+	, bOverride_LPVIntensity(Settings.bOverride_LPVIntensity)
+	, bOverride_LPVDirectionalOcclusionIntensity(Settings.bOverride_LPVDirectionalOcclusionIntensity)
+	, bOverride_LPVDirectionalOcclusionRadius(Settings.bOverride_LPVDirectionalOcclusionRadius)
+	, bOverride_LPVDiffuseOcclusionExponent(Settings.bOverride_WhiteTemp)
+	, bOverride_LPVSpecularOcclusionExponent(Settings.bOverride_LPVSpecularOcclusionExponent)
+	, bOverride_LPVDiffuseOcclusionIntensity(Settings.bOverride_LPVDiffuseOcclusionIntensity)
+	, bOverride_LPVSpecularOcclusionIntensity(Settings.bOverride_LPVSpecularOcclusionIntensity)
+	, bOverride_LPVSize(Settings.bOverride_LPVSize)
+	, bOverride_LPVSecondaryOcclusionIntensity(Settings.bOverride_LPVSecondaryOcclusionIntensity)
+	, bOverride_LPVSecondaryBounceIntensity(Settings.bOverride_LPVSecondaryBounceIntensity)
+	, bOverride_LPVGeometryVolumeBias(Settings.bOverride_LPVGeometryVolumeBias)
+	, bOverride_LPVVplInjectionBias(Settings.bOverride_LPVVplInjectionBias)
+	, bOverride_LPVEmissiveInjectionIntensity(Settings.bOverride_LPVEmissiveInjectionIntensity)
+	, bOverride_LPVFadeRange(Settings.bOverride_LPVFadeRange)
+	, bOverride_LPVDirectionalOcclusionFadeRange(Settings.bOverride_LPVDirectionalOcclusionFadeRange)
+	, bOverride_IndirectLightingColor(Settings.bOverride_IndirectLightingColor)
+	, bOverride_IndirectLightingIntensity(Settings.bOverride_IndirectLightingIntensity)
+	, bOverride_ColorGradingIntensity(Settings.bOverride_ColorGradingIntensity)
+	, bOverride_ColorGradingLUT(Settings.bOverride_ColorGradingLUT)
+	PRAGMA_DISABLE_DEPRECATION_WARNINGS
+	, bOverride_DepthOfFieldFocalDistance(Settings.bOverride_DepthOfFieldFocalDistance)
+	, bOverride_DepthOfFieldFstop(Settings.bOverride_DepthOfFieldFstop)
+	, bOverride_DepthOfFieldMinFstop(Settings.bOverride_DepthOfFieldMinFstop)
+	, bOverride_DepthOfFieldBladeCount(Settings.bOverride_DepthOfFieldBladeCount)
+	, bOverride_DepthOfFieldSensorWidth(Settings.bOverride_DepthOfFieldSensorWidth)
+	, bOverride_DepthOfFieldDepthBlurRadius(Settings.bOverride_DepthOfFieldDepthBlurRadius)
+	, bOverride_DepthOfFieldDepthBlurAmount(Settings.bOverride_DepthOfFieldDepthBlurAmount)
+	, bOverride_DepthOfFieldFocalRegion(Settings.bOverride_DepthOfFieldFocalRegion)
+	, bOverride_DepthOfFieldNearTransitionRegion(Settings.bOverride_DepthOfFieldNearTransitionRegion)
+	, bOverride_DepthOfFieldFarTransitionRegion(Settings.bOverride_DepthOfFieldFarTransitionRegion)
+	, bOverride_DepthOfFieldScale(Settings.bOverride_DepthOfFieldScale)
+	, bOverride_DepthOfFieldMaxBokehSize(Settings.bOverride_DepthOfFieldMaxBokehSize)
+	, bOverride_DepthOfFieldNearBlurSize(Settings.bOverride_DepthOfFieldNearBlurSize)
+	, bOverride_DepthOfFieldFarBlurSize(Settings.bOverride_DepthOfFieldFarBlurSize)
+	, bOverride_DepthOfFieldMethod(Settings.bOverride_DepthOfFieldMethod)
+	, bOverride_MobileHQGaussian(Settings.bOverride_MobileHQGaussian)
+	, bOverride_DepthOfFieldBokehShape(Settings.bOverride_DepthOfFieldBokehShape)
+	, bOverride_DepthOfFieldOcclusion(Settings.bOverride_DepthOfFieldOcclusion)
+	, bOverride_DepthOfFieldColorThreshold(Settings.bOverride_DepthOfFieldColorThreshold)
+	, bOverride_DepthOfFieldSizeThreshold(Settings.bOverride_DepthOfFieldSizeThreshold)
+	, bOverride_DepthOfFieldSkyFocusDistance(Settings.bOverride_DepthOfFieldSkyFocusDistance)
+	, bOverride_DepthOfFieldVignetteSize(Settings.bOverride_DepthOfFieldVignetteSize)
+	PRAGMA_ENABLE_DEPRECATION_WARNINGS
+	, bOverride_MotionBlurAmount(Settings.bOverride_MotionBlurAmount)
+	, bOverride_MotionBlurMax(Settings.bOverride_MotionBlurMax)
+	, bOverride_MotionBlurPerObjectSize(Settings.bOverride_MotionBlurPerObjectSize)
+	, bOverride_ScreenPercentage(Settings.bOverride_ScreenPercentage)
+	, bOverride_ScreenSpaceReflectionIntensity(Settings.bOverride_ScreenSpaceReflectionIntensity)
+	, bOverride_ScreenSpaceReflectionQuality(Settings.bOverride_ScreenSpaceReflectionQuality)
+	, bOverride_ScreenSpaceReflectionMaxRoughness(Settings.bOverride_ScreenSpaceReflectionMaxRoughness)
+	, bOverride_ScreenSpaceReflectionRoughnessScale(Settings.bOverride_ScreenSpaceReflectionRoughnessScale)
+
+	, bMobileHQGaussian(Settings.bMobileHQGaussian)
+	, BloomMethod(Settings.BloomMethod)
+	, AutoExposureMethod(Settings.AutoExposureMethod)
+	PRAGMA_DISABLE_DEPRECATION_WARNINGS
+	, DepthOfFieldMethod(Settings.DepthOfFieldMethod)
+	PRAGMA_ENABLE_DEPRECATION_WARNINGS
+	, WhiteTemp(Settings.WhiteTemp)
+	, WhiteTint(Settings.WhiteTint)
+	, ColorSaturation(Settings.ColorSaturation)
+	, ColorContrast(Settings.ColorContrast)
+	, ColorGamma(Settings.ColorGamma)
+	, ColorGain(Settings.ColorGain)
+	, ColorOffset(Settings.ColorOffset)
+	, ColorSaturationShadows(Settings.ColorSaturationShadows)
+	, ColorContrastShadows(Settings.ColorContrastShadows)
+	, ColorGammaShadows(Settings.ColorGammaShadows)
+	, ColorGainShadows(Settings.ColorGainShadows)
+	, ColorOffsetShadows(Settings.ColorOffsetShadows)
+	, ColorSaturationMidtones(Settings.ColorSaturationMidtones)
+	, ColorContrastMidtones(Settings.ColorContrastMidtones)
+	, ColorGammaMidtones(Settings.ColorGammaMidtones)
+	, ColorGainMidtones(Settings.ColorGainMidtones)
+	, ColorOffsetMidtones(Settings.ColorOffsetMidtones)
+	, ColorSaturationHighlights(Settings.ColorSaturationHighlights)
+	, ColorContrastHighlights(Settings.ColorContrastHighlights)
+	, ColorGammaHighlights(Settings.ColorGammaHighlights)
+	, ColorGainHighlights(Settings.ColorGainHighlights)
+	, ColorOffsetHighlights(Settings.ColorOffsetHighlights)
+	, ColorCorrectionHighlightsMin(Settings.ColorCorrectionHighlightsMin)
+	, ColorCorrectionShadowsMax(Settings.ColorCorrectionShadowsMax)
+	, BlueCorrection(Settings.BlueCorrection)
+	, ExpandGamut(Settings.ExpandGamut)
+	, FilmSlope(Settings.FilmSlope)
+	, FilmToe(Settings.FilmToe)
+	, FilmShoulder(Settings.FilmShoulder)
+	, FilmBlackClip(Settings.FilmBlackClip)
+	, FilmWhiteClip(Settings.FilmWhiteClip)
+	, FilmWhitePoint(Settings.FilmWhitePoint)
+	, FilmShadowTint(Settings.FilmShadowTint)
+	, FilmShadowTintBlend(Settings.FilmShadowTintBlend)
+	, FilmShadowTintAmount(Settings.FilmShadowTintAmount)
+	, FilmSaturation(Settings.FilmSaturation)
+	, FilmChannelMixerRed(Settings.FilmChannelMixerRed)
+	, FilmChannelMixerGreen(Settings.FilmChannelMixerGreen)
+	, FilmChannelMixerBlue(Settings.FilmChannelMixerBlue)
+	, FilmContrast(Settings.FilmContrast)
+	, FilmToeAmount(Settings.FilmToeAmount)
+	, FilmHealAmount(Settings.FilmHealAmount)
+	, FilmDynamicRange(Settings.FilmDynamicRange)
+	, SceneColorTint(Settings.SceneColorTint)
+	, SceneFringeIntensity(Settings.SceneFringeIntensity)
+	, ChromaticAberrationStartOffset(Settings.ChromaticAberrationStartOffset)
+	, BloomIntensity(Settings.BloomIntensity)
+	, BloomThreshold(Settings.BloomThreshold)
+	, BloomSizeScale(Settings.BloomSizeScale)
+	, Bloom1Size(Settings.Bloom1Size)
+	, Bloom2Size(Settings.Bloom2Size)
+	, Bloom3Size(Settings.Bloom3Size)
+	, Bloom4Size(Settings.Bloom4Size)
+	, Bloom5Size(Settings.Bloom5Size)
+	, Bloom6Size(Settings.Bloom6Size)
+	, Bloom1Tint(Settings.Bloom1Tint)
+	, Bloom2Tint(Settings.Bloom2Tint)
+	, Bloom3Tint(Settings.Bloom3Tint)
+	, Bloom4Tint(Settings.Bloom4Tint)
+	, Bloom5Tint(Settings.Bloom5Tint)
+	, Bloom6Tint(Settings.Bloom6Tint)
+	, BloomConvolutionSize(Settings.BloomConvolutionSize)
+	, BloomConvolutionTexture(Settings.BloomConvolutionTexture)
+	, BloomConvolutionCenterUV(Settings.BloomConvolutionCenterUV)
+	//, BloomConvolutionPreFilter_DEPRECATED(Settings.BloomConvolutionPreFilter_DEPRECATED)
+	, BloomConvolutionPreFilterMin(Settings.BloomConvolutionPreFilterMin)
+	, BloomConvolutionPreFilterMax(Settings.BloomConvolutionPreFilterMax)
+	, BloomConvolutionPreFilterMult(Settings.BloomConvolutionPreFilterMult)
+	, BloomConvolutionBufferScale(Settings.BloomConvolutionBufferScale)
+	, BloomDirtMask(Settings.BloomDirtMask)
+	, BloomDirtMaskIntensity(Settings.BloomDirtMaskIntensity)
+	, BloomDirtMaskTint(Settings.BloomDirtMaskTint)
+	, CameraShutterSpeed(Settings.CameraShutterSpeed)
+	, CameraISO(Settings.CameraISO)
+	, DepthOfFieldFstop(Settings.DepthOfFieldFstop)
+	, DepthOfFieldMinFstop(Settings.DepthOfFieldMinFstop)
+	, DepthOfFieldBladeCount(Settings.DepthOfFieldBladeCount)
+	, AutoExposureBias(Settings.AutoExposureBias)
+	, AutoExposureLowPercent(Settings.AutoExposureLowPercent)
+	, AutoExposureHighPercent(Settings.AutoExposureHighPercent)
+	, AutoExposureMinBrightness(Settings.AutoExposureMinBrightness)
+	, AutoExposureMaxBrightness(Settings.AutoExposureMaxBrightness)
+	, AutoExposureSpeedUp(Settings.AutoExposureSpeedUp)
+	, AutoExposureSpeedDown(Settings.AutoExposureSpeedDown)
+	, HistogramLogMin(Settings.HistogramLogMin)
+	, HistogramLogMax(Settings.HistogramLogMax)
+	, AutoExposureCalibrationConstant(Settings.AutoExposureCalibrationConstant)
+	, LensFlareIntensity(Settings.LensFlareIntensity)
+	, LensFlareTint(Settings.LensFlareTint)
+	, LensFlareBokehSize(Settings.LensFlareBokehSize)
+	, LensFlareThreshold(Settings.LensFlareThreshold)
+	, LensFlareBokehShape(Settings.LensFlareBokehShape)
+	, VignetteIntensity(Settings.VignetteIntensity)
+	, GrainJitter(Settings.GrainJitter)
+	, GrainIntensity(Settings.GrainIntensity)
+	, AmbientOcclusionIntensity(Settings.AmbientOcclusionIntensity)
+	, AmbientOcclusionStaticFraction(Settings.AmbientOcclusionStaticFraction)
+	, AmbientOcclusionRadius(Settings.AmbientOcclusionRadius)
+	, AmbientOcclusionRadiusInWS(Settings.AmbientOcclusionRadiusInWS)
+	, AmbientOcclusionFadeDistance(Settings.AmbientOcclusionFadeDistance)
+	, AmbientOcclusionFadeRadius(Settings.AmbientOcclusionFadeRadius)
+	//, AmbientOcclusionDistance_DEPRECATED(Settings.AmbientOcclusionDistance_DEPRECATED)
+	, AmbientOcclusionPower(Settings.AmbientOcclusionPower)
+	, AmbientOcclusionBias(Settings.AmbientOcclusionBias)
+	, AmbientOcclusionQuality(Settings.AmbientOcclusionQuality)
+	, AmbientOcclusionMipBlend(Settings.AmbientOcclusionMipBlend)
+	, AmbientOcclusionMipScale(Settings.AmbientOcclusionMipScale)
+	, AmbientOcclusionMipThreshold(Settings.AmbientOcclusionMipThreshold)
+	, IndirectLightingColor(Settings.IndirectLightingColor)
+	, IndirectLightingIntensity(Settings.IndirectLightingIntensity)
+	, ColorGradingIntensity(Settings.ColorGradingIntensity)
+	, ColorGradingLUT(Settings.ColorGradingLUT)
+	PRAGMA_DISABLE_DEPRECATION_WARNINGS
+	, DepthOfFieldSensorWidth(Settings.DepthOfFieldSensorWidth)
+	, DepthOfFieldFocalDistance(Settings.DepthOfFieldFocalDistance)
+	, DepthOfFieldDepthBlurAmount(Settings.DepthOfFieldDepthBlurAmount)
+	, DepthOfFieldDepthBlurRadius(Settings.DepthOfFieldDepthBlurRadius)
+	, DepthOfFieldFocalRegion(Settings.DepthOfFieldFocalRegion)
+	, DepthOfFieldNearTransitionRegion(Settings.DepthOfFieldNearTransitionRegion)
+	, DepthOfFieldFarTransitionRegion(Settings.DepthOfFieldFarTransitionRegion)
+	, DepthOfFieldScale(Settings.DepthOfFieldScale)
+	, DepthOfFieldMaxBokehSize(Settings.DepthOfFieldMaxBokehSize)
+	, DepthOfFieldNearBlurSize(Settings.DepthOfFieldNearBlurSize)
+	, DepthOfFieldFarBlurSize(Settings.DepthOfFieldFarBlurSize)
+	, DepthOfFieldOcclusion(Settings.DepthOfFieldOcclusion)
+	, DepthOfFieldBokehShape(Settings.DepthOfFieldBokehShape)
+	, DepthOfFieldColorThreshold(Settings.DepthOfFieldColorThreshold)
+	, DepthOfFieldSizeThreshold(Settings.DepthOfFieldSizeThreshold)
+	, DepthOfFieldSkyFocusDistance(Settings.DepthOfFieldSkyFocusDistance)
+	, DepthOfFieldVignetteSize(Settings.DepthOfFieldVignetteSize)
+	PRAGMA_ENABLE_DEPRECATION_WARNINGS
+	, MotionBlurAmount(Settings.MotionBlurAmount)
+	, MotionBlurMax(Settings.MotionBlurMax)
+	, MotionBlurPerObjectSize(Settings.MotionBlurPerObjectSize)
+	, LPVIntensity(Settings.LPVIntensity)
+	, LPVVplInjectionBias(Settings.LPVVplInjectionBias)
+	, LPVSize(Settings.LPVSize)
+	, LPVSecondaryOcclusionIntensity(Settings.LPVSecondaryOcclusionIntensity)
+	, LPVSecondaryBounceIntensity(Settings.LPVSecondaryBounceIntensity)
+	, LPVGeometryVolumeBias(Settings.LPVGeometryVolumeBias)
+	, LPVEmissiveInjectionIntensity(Settings.LPVEmissiveInjectionIntensity)
+	, LPVDirectionalOcclusionIntensity(Settings.LPVDirectionalOcclusionIntensity)
+	, LPVDirectionalOcclusionRadius(Settings.LPVDirectionalOcclusionRadius)
+	, LPVDiffuseOcclusionExponent(Settings.WhiteTemp)
+	, LPVSpecularOcclusionExponent(Settings.LPVSpecularOcclusionExponent)
+	, LPVDiffuseOcclusionIntensity(Settings.LPVDiffuseOcclusionIntensity)
+	, LPVSpecularOcclusionIntensity(Settings.LPVSpecularOcclusionIntensity)
+	, ScreenSpaceReflectionIntensity(Settings.ScreenSpaceReflectionIntensity)
+	, ScreenSpaceReflectionQuality(Settings.ScreenSpaceReflectionQuality)
+	, ScreenSpaceReflectionMaxRoughness(Settings.ScreenSpaceReflectionMaxRoughness)
+	, LPVFadeRange(Settings.LPVFadeRange)
+	, LPVDirectionalOcclusionFadeRange(Settings.LPVDirectionalOcclusionFadeRange)
+	, ScreenPercentage(Settings.ScreenPercentage)
+
+	, WeightedBlendables(Settings.WeightedBlendables)
+	//, Blendables_DEPRECATED(Settings.Blendables_DEPRECATED)
+{
+	for (int32 i = 0; i < ARRAY_COUNT(LensFlareTints); i++)
+		LensFlareTints[i] = Settings.LensFlareTints[i];
+}
+
 UScene::UScene(const FObjectInitializer& ObjectInitializer)
 	: Super(ObjectInitializer)
 {
diff --git a/Engine/Source/Runtime/Engine/Private/SceneView.cpp b/Engine/Source/Runtime/Engine/Private/SceneView.cpp
--- a/Engine/Source/Runtime/Engine/Private/SceneView.cpp
+++ b/Engine/Source/Runtime/Engine/Private/SceneView.cpp
@@ -1338,12 +1338,16 @@ void FSceneView::OverridePostProcessSettings(const FPostProcessSettings& Src, fl
 		LERP_PP(DepthOfFieldNearTransitionRegion);
 		LERP_PP(DepthOfFieldFarTransitionRegion);
 		LERP_PP(DepthOfFieldScale);
+		PRAGMA_DISABLE_DEPRECATION_WARNINGS
 		LERP_PP(DepthOfFieldMaxBokehSize);
+		PRAGMA_ENABLE_DEPRECATION_WARNINGS
 		LERP_PP(DepthOfFieldNearBlurSize);
 		LERP_PP(DepthOfFieldFarBlurSize);
 		LERP_PP(DepthOfFieldOcclusion);
+		PRAGMA_DISABLE_DEPRECATION_WARNINGS
 		LERP_PP(DepthOfFieldColorThreshold);
 		LERP_PP(DepthOfFieldSizeThreshold);
+		PRAGMA_ENABLE_DEPRECATION_WARNINGS
 		LERP_PP(DepthOfFieldSkyFocusDistance);
 		LERP_PP(DepthOfFieldVignetteSize);
 		LERP_PP(MotionBlurAmount);
@@ -1404,11 +1408,13 @@ void FSceneView::OverridePostProcessSettings(const FPostProcessSettings& Src, fl
 			Dest.BloomConvolutionBufferScale = Src.BloomConvolutionBufferScale;
 		}
 
+		PRAGMA_DISABLE_DEPRECATION_WARNINGS
 		// actual texture cannot be blended but the intensity can be blended
 		IF_PP(DepthOfFieldBokehShape)
 		{
 			Dest.DepthOfFieldBokehShape = Src.DepthOfFieldBokehShape;
 		}
+		PRAGMA_ENABLE_DEPRECATION_WARNINGS
 
 		// actual texture cannot be blended but the intensity can be blended
 		IF_PP(LensFlareBokehShape)
@@ -1424,10 +1430,12 @@ void FSceneView::OverridePostProcessSettings(const FPostProcessSettings& Src, fl
 			}
 		}
 
+		PRAGMA_DISABLE_DEPRECATION_WARNINGS
 		if (Src.bOverride_DepthOfFieldMethod)
 		{
 			Dest.DepthOfFieldMethod = Src.DepthOfFieldMethod;
 		}
+		PRAGMA_ENABLE_DEPRECATION_WARNINGS
 
 		if (Src.bOverride_MobileHQGaussian)
 		{
@@ -1790,12 +1798,14 @@ void FSceneView::EndFinalPostprocessSettings(const FSceneViewInitOptions& ViewIn
 	}
 #endif
 
+	PRAGMA_DISABLE_DEPRECATION_WARNINGS
 	if(FinalPostProcessSettings.DepthOfFieldMethod == DOFM_CircleDOF)
 	{
 		// We intentionally don't do the DepthOfFieldFocalRegion as it breaks realism.
 		// Doing this fixes DOF material expression.
 		FinalPostProcessSettings.DepthOfFieldFocalRegion = 0;
 	}
+	PRAGMA_ENABLE_DEPRECATION_WARNINGS
 
 	{
 		const bool bStereoEnabled = StereoPass != eSSP_FULL;
diff --git a/Engine/Source/Runtime/Renderer/Private/PostProcess/PostProcessBokehDOF.cpp b/Engine/Source/Runtime/Renderer/Private/PostProcess/PostProcessBokehDOF.cpp
--- a/Engine/Source/Runtime/Renderer/Private/PostProcess/PostProcessBokehDOF.cpp
+++ b/Engine/Source/Runtime/Renderer/Private/PostProcess/PostProcessBokehDOF.cpp
@@ -97,11 +97,13 @@ public:
 		{
 			// a negative values disables the cross hair feature
 			FIntPoint CursorPosValue(-100,-100);
-			
+
+			PRAGMA_DISABLE_DEPRECATION_WARNINGS
 			if(Context.View.FinalPostProcessSettings.DepthOfFieldMethod == DOFM_CircleDOF)
 			{
 				CursorPosValue = Context.View.CursorPos;
 			}
+			PRAGMA_ENABLE_DEPRECATION_WARNINGS
 
 			SetShaderValue(RHICmdList, ShaderRHI, CursorPos, CursorPosValue);
 		}
@@ -223,10 +225,13 @@ void FRCPassPostProcessVisualizeDOF::Process(FRenderingCompositePassContext& Con
 		Canvas.DrawShadowedString(20, Y += YStep, *Line, GetStatsFont(), FLinearColor(1, 1, 0));
 		Y += YStep;
 
+		PRAGMA_DISABLE_DEPRECATION_WARNINGS
 		EDepthOfFieldMethod MethodId = View.FinalPostProcessSettings.DepthOfFieldMethod;
+		PRAGMA_ENABLE_DEPRECATION_WARNINGS
 
 		if(MethodId == DOFM_BokehDOF)
 		{
+			PRAGMA_DISABLE_DEPRECATION_WARNINGS
 			Line = FString::Printf(TEXT("Method: BokehDOF (blue is far, green is near, black is in focus)"));
 			Canvas.DrawShadowedString(X, Y += YStep, *Line, GetStatsFont(), FLinearColor(1, 1, 1));
 			Y += YStep;
@@ -249,6 +254,7 @@ void FRCPassPostProcessVisualizeDOF::Process(FRenderingCompositePassContext& Con
 			Canvas.DrawShadowedString(X, Y += YStep, *Line, GetStatsFont(), FLinearColor(1, 1, 1));
 			Line = FString::Printf(TEXT("Occlusion: %.2f"), View.FinalPostProcessSettings.DepthOfFieldOcclusion);
 			Canvas.DrawShadowedString(X, Y += YStep, *Line, GetStatsFont(), FLinearColor(1, 1, 1));
+			PRAGMA_ENABLE_DEPRECATION_WARNINGS
 		}
 		else if(MethodId == DOFM_Gaussian)
 		{
@@ -665,7 +671,7 @@ public:
 
 			SetShaderValue(Context.RHICmdList, ShaderRHI, KernelSize, KernelSizeValue);
 		}
-
+		PRAGMA_DISABLE_DEPRECATION_WARNINGS
 		{
 			FVector4 Value(
 				Context.View.FinalPostProcessSettings.DepthOfFieldColorThreshold,
@@ -673,6 +679,7 @@ public:
 
 			SetShaderValue(Context.RHICmdList, ShaderRHI, DepthOfFieldThresholds, Value);
 		}
+		PRAGMA_ENABLE_DEPRECATION_WARNINGS
 
 		{
 			FVector4 DepthOfFieldParamValues[2];
@@ -751,10 +758,12 @@ public:
 				TextureRHI = GEngine->DefaultBokehTexture->Resource->TextureRHI;
 			}
 
+			PRAGMA_DISABLE_DEPRECATION_WARNINGS
 			if(Context.View.FinalPostProcessSettings.DepthOfFieldBokehShape)
 			{
 				TextureRHI = Context.View.FinalPostProcessSettings.DepthOfFieldBokehShape->Resource->TextureRHI;
 			}
+			PRAGMA_ENABLE_DEPRECATION_WARNINGS
 
 			SetTextureParameter(Context.RHICmdList, ShaderRHI, LensTexture, LensTextureSampler, TStaticSamplerState<SF_Trilinear, AM_Border, AM_Border, AM_Clamp>::GetRHI(), TextureRHI);
 		}
@@ -817,9 +826,11 @@ void FRCPassPostProcessBokehDOF::ComputeDepthOfFieldParams(const FRenderingCompo
 		Context.View.FinalPostProcessSettings.DepthOfFieldOcclusion);
 
 	FIntPoint ViewSize = Context.View.ViewRect.Size();
-
-	float MaxBokehSizeInPixel = FMath::Max(0.0f, Context.View.FinalPostProcessSettings.DepthOfFieldMaxBokehSize) / 100.0f * ViewSize.X;
 	
+	PRAGMA_DISABLE_DEPRECATION_WARNINGS
+	float MaxBokehSizeInPixel = FMath::Max(0.0f, Context.View.FinalPostProcessSettings.DepthOfFieldMaxBokehSize) / 100.0f * ViewSize.X;
+	PRAGMA_ENABLE_DEPRECATION_WARNINGS
+
 	// Scale and offset to put two views in one texture with safety border
 	float UsedYDivTextureY = HalfRes / (float)BokehLayerSizeY;
 	float YOffsetInPixel = HalfRes + SafetyBorder;
@@ -887,7 +898,9 @@ void FRCPassPostProcessBokehDOF::Process(FRenderingCompositePassContext& Context
 
 	FIntPoint TileCount = LocalViewSize / TileSize;
 
+	PRAGMA_DISABLE_DEPRECATION_WARNINGS
 	float PixelKernelSize = Context.View.FinalPostProcessSettings.DepthOfFieldMaxBokehSize / 100.0f * LocalViewSize.X;
+	PRAGMA_ENABLE_DEPRECATION_WARNINGS
 
 	FIntPoint LeftTop = LocalViewRect.Min;
 	
diff --git a/Engine/Source/Runtime/Renderer/Private/PostProcess/PostProcessCircleDOF.cpp b/Engine/Source/Runtime/Renderer/Private/PostProcess/PostProcessCircleDOF.cpp
--- a/Engine/Source/Runtime/Renderer/Private/PostProcess/PostProcessCircleDOF.cpp
+++ b/Engine/Source/Runtime/Renderer/Private/PostProcess/PostProcessCircleDOF.cpp
@@ -53,7 +53,9 @@ FVector4 CircleDofHalfCoc(const FViewInfo& View)
 
 	FVector4 Ret(0, 1, 0, 0);
 
+	PRAGMA_DISABLE_DEPRECATION_WARNINGS
 	if(bDepthOfField && View.FinalPostProcessSettings.DepthOfFieldMethod == DOFM_CircleDOF)
+	PRAGMA_ENABLE_DEPRECATION_WARNINGS
 	{
 		float FocalLengthInMM = ComputeFocalLengthFromFov(View);
 	 
diff --git a/Engine/Source/Runtime/Renderer/Private/PostProcess/PostProcessMobile.cpp b/Engine/Source/Runtime/Renderer/Private/PostProcess/PostProcessMobile.cpp
--- a/Engine/Source/Runtime/Renderer/Private/PostProcess/PostProcessMobile.cpp
+++ b/Engine/Source/Runtime/Renderer/Private/PostProcess/PostProcessMobile.cpp
@@ -24,7 +24,9 @@ static EPixelFormat GetHDRPixelFormat()
 // return Depth of Field Scale if Gaussian DoF mode is active. 0.0f otherwise.
 float GetMobileDepthOfFieldScale(const FViewInfo& View)
 {
-	return View.FinalPostProcessSettings.DepthOfFieldMethod == DOFM_Gaussian ? View.FinalPostProcessSettings.DepthOfFieldScale : 0.0f;
+	PRAGMA_DISABLE_DEPRECATION_WARNINGS
+		return View.FinalPostProcessSettings.DepthOfFieldMethod == DOFM_Gaussian ? View.FinalPostProcessSettings.DepthOfFieldScale : 0.0f;
+	PRAGMA_ENABLE_DEPRECATION_WARNINGS
 }
 
 //
diff --git a/Engine/Source/Runtime/Renderer/Private/PostProcess/PostProcessing.cpp b/Engine/Source/Runtime/Renderer/Private/PostProcess/PostProcessing.cpp
--- a/Engine/Source/Runtime/Renderer/Private/PostProcess/PostProcessing.cpp
+++ b/Engine/Source/Runtime/Renderer/Private/PostProcess/PostProcessing.cpp
@@ -1465,10 +1465,18 @@ void FPostProcessing::Process(FRHICommandListImmediate& RHICmdList, const FViewI
 
 			bool bSepTransWasApplied = false;
 
+			PRAGMA_DISABLE_DEPRECATION_WARNINGS
 			if(bDepthOfField && View.FinalPostProcessSettings.DepthOfFieldMethod != DOFM_BokehDOF)
 			{
 				if(View.FinalPostProcessSettings.DepthOfFieldMethod == DOFM_Gaussian)
 				{
+					static bool bHasWarned = false;
+					if (!bHasWarned)
+					{
+						UE_LOG(LogRenderer, Warning, TEXT("Gaussian DOF algorithm is deprecated and will be removed from deferred shading renderer. Consider using the Circle DOF method in post process settings."));
+						bHasWarned = true;
+					}
+
 					if (FPostProcessing::HasAlphaChannelSupport())
 					{
 						UE_LOG(LogRenderer, Log, TEXT("Gaussian depth of field does not have alpha channel support. Only Circle DOF has."));
@@ -1505,6 +1513,12 @@ void FPostProcessing::Process(FRHICommandListImmediate& RHICmdList, const FViewI
 					}
 					else
 					{
+						static bool bHasWarned = false;
+						if (!bHasWarned)
+						{
+							UE_LOG(LogRenderer, Warning, TEXT("Circle DOF algorithm is deprecated and will be removed. Consider using Diphragm DOF with r.DOF.Algorithm = 1."));
+							bHasWarned = true;
+						}
 						AddPostProcessDepthOfFieldCircle(Context, DOFVelocityRef);
 					}
 				}
@@ -1514,9 +1528,17 @@ void FPostProcessing::Process(FRHICommandListImmediate& RHICmdList, const FViewI
 				&& View.FinalPostProcessSettings.DepthOfFieldScale > 0
 				&& View.FinalPostProcessSettings.DepthOfFieldMethod == DOFM_BokehDOF
 				&& !Context.View.Family->EngineShowFlags.VisualizeDOF;
+			PRAGMA_ENABLE_DEPRECATION_WARNINGS
 
 			if(bBokehDOF)
 			{
+				static bool bHasWarned = false;
+				if (!bHasWarned)
+				{
+					UE_LOG(LogRenderer, Warning, TEXT("Bokeh DOF algorithm is deprecated and will be removed. Consider using the Circle DOF method in post process settings."));
+					bHasWarned = true;
+				}
+
 				if (FPostProcessing::HasAlphaChannelSupport())
 				{
 					UE_LOG(LogRenderer, Log, TEXT("Boked depth of field does not have alpha channel support. Only Circle DOF has."));
@@ -2422,12 +2444,15 @@ void FPostProcessing::ProcessES2(FRHICommandListImmediate& RHICmdList, const FVi
 						// black is how we clear the velocity buffer so this means no velocity
 						FRenderingCompositePass* NoVelocity = Context.Graph.RegisterPass(new(FMemStack::Get()) FRCPassPostProcessInput(GSystemTextures.BlackDummy));
 						FRenderingCompositeOutputRef NoVelocityRef(NoVelocity);
+
+						PRAGMA_DISABLE_DEPRECATION_WARNINGS
 						if(View.FinalPostProcessSettings.DepthOfFieldMethod == DOFM_Gaussian && IsGaussianActive(Context))
 						{
 							FDepthOfFieldStats DepthOfFieldStat;
 							FRenderingCompositeOutputRef DummySeparateTranslucency;
 							AddPostProcessDepthOfFieldGaussian(Context, DepthOfFieldStat, NoVelocityRef, DummySeparateTranslucency);
 						}
+						PRAGMA_ENABLE_DEPRECATION_WARNINGS
 					}
 				}
 
