Fix for crash when running automation tests

diff --git a/Engine/Source/Runtime/Engine/Private/Scene.cpp b/Engine/Source/Runtime/Engine/Private/Scene.cpp
--- a/Engine/Source/Runtime/Engine/Private/Scene.cpp
+++ b/Engine/Source/Runtime/Engine/Private/Scene.cpp
@@ -753,6 +753,9 @@ FPostProcessSettings::FPostProcessSettings(const FPostProcessSettings& Settings)
 	, BloomDirtMask(Settings.BloomDirtMask)
 	, BloomDirtMaskIntensity(Settings.BloomDirtMaskIntensity)
 	, BloomDirtMaskTint(Settings.BloomDirtMaskTint)
+	, AmbientCubemapTint(Settings.AmbientCubemapTint)
+	, AmbientCubemapIntensity(Settings.AmbientCubemapIntensity)
+	, AmbientCubemap(Settings.AmbientCubemap)
 	, CameraShutterSpeed(Settings.CameraShutterSpeed)
 	, CameraISO(Settings.CameraISO)
 	, DepthOfFieldFstop(Settings.DepthOfFieldFstop)
