Deprecate ABGR16 and Float RGB from EDefaultBackBufferPixelFormat. They are not supported by D3DX SwapChain. The jira has more info on which formats are supported.

diff --git a/Engine/Source/Runtime/Engine/Classes/Engine/RendererSettings.h b/Engine/Source/Runtime/Engine/Classes/Engine/RendererSettings.h
--- a/Engine/Source/Runtime/Engine/Classes/Engine/RendererSettings.h
+++ b/Engine/Source/Runtime/Engine/Classes/Engine/RendererSettings.h
@@ -144,11 +144,11 @@ namespace EDefaultBackBufferPixelFormat
 {
 	enum Type
 	{
-		DBBPF_B8G8R8A8 = 0	UMETA(DisplayName = "8bit RGBA"),
-		DBBPF_A16B16G16R16	UMETA(DisplayName = "16bit RGBA"),
-		DBBPF_FloatRGB		UMETA(DisplayName = "Float RGB"),
-		DBBPF_FloatRGBA		UMETA(DisplayName = "Float RGBA"),
-		DBBPF_A2B10G10R10	UMETA(DisplayName = "10bit RGB, 2bit Alpha"),
+		DBBPF_B8G8R8A8 = 0				UMETA(DisplayName = "8bit RGBA"),
+		DBBPF_A16B16G16R16_DEPRECATED	UMETA(DisplayName = "DEPRECATED - 16bit RGBA", Hidden),
+		DBBPF_FloatRGB_DEPRECATED		UMETA(DisplayName = "DEPRECATED - Float RGB", Hidden),
+		DBBPF_FloatRGBA					UMETA(DisplayName = "Float RGBA"),
+		DBBPF_A2B10G10R10				UMETA(DisplayName = "10bit RGB, 2bit Alpha"),
 		DBBPF_MAX,
 	};
 }
diff --git a/Engine/Source/Runtime/Engine/Private/RendererSettings.cpp b/Engine/Source/Runtime/Engine/Private/RendererSettings.cpp
--- a/Engine/Source/Runtime/Engine/Private/RendererSettings.cpp
+++ b/Engine/Source/Runtime/Engine/Private/RendererSettings.cpp
@@ -22,22 +22,22 @@ namespace EDefaultBackBufferPixelFormat
 {
 	EPixelFormat Convert2PixelFormat(EDefaultBackBufferPixelFormat::Type InDefaultBackBufferPixelFormat)
 	{
-		static EPixelFormat SPixelFormat[] = { PF_B8G8R8A8, PF_A16B16G16R16, PF_FloatRGB, PF_FloatRGBA, PF_A2B10G10R10, PF_A2B10G10R10 };
-		return SPixelFormat[(int32)InDefaultBackBufferPixelFormat];
+		const int32 ValidIndex = FMath::Clamp((int32)InDefaultBackBufferPixelFormat, 0, (int32)DBBPF_MAX - 1);
+		static EPixelFormat SPixelFormat[] = { PF_B8G8R8A8, PF_B8G8R8A8, PF_FloatRGBA, PF_FloatRGBA, PF_A2B10G10R10 };
+		return SPixelFormat[ValidIndex];
 	}
 
 	int32 NumberOfBitForAlpha(EDefaultBackBufferPixelFormat::Type InDefaultBackBufferPixelFormat)
 	{
 		switch (InDefaultBackBufferPixelFormat)
 		{
-			case DBBPF_A16B16G16R16:
-				return 16;
+			case DBBPF_A16B16G16R16_DEPRECATED:
 			case DBBPF_B8G8R8A8:
+			case DBBPF_FloatRGB_DEPRECATED:
 			case DBBPF_FloatRGBA:
 				return 8;
 			case DBBPF_A2B10G10R10:
 				return 2;
-			case DBBPF_FloatRGB:
 			default:
 				return 0;
 		}
@@ -46,7 +46,9 @@ namespace EDefaultBackBufferPixelFormat
 
 	EDefaultBackBufferPixelFormat::Type FromInt(int32 InDefaultBackBufferPixelFormat)
 	{
-		return static_cast<EDefaultBackBufferPixelFormat::Type>(FMath::Clamp(InDefaultBackBufferPixelFormat, 0, (int32)DBBPF_MAX - 1));
+		const int32 ValidIndex = FMath::Clamp(InDefaultBackBufferPixelFormat, 0, (int32)DBBPF_MAX - 1);
+		static EDefaultBackBufferPixelFormat::Type SPixelFormat[] = { DBBPF_B8G8R8A8, DBBPF_B8G8R8A8, DBBPF_FloatRGBA, DBBPF_FloatRGBA, DBBPF_A2B10G10R10 };
+		return SPixelFormat[ValidIndex];
 	}
 }
 
