Copying from Dev-Enterprise CL 4364260
SceneCapture: Don't CaptureOnMovement if CaptureEveryFrame is enabled. This prevent the capture from happening twice in 1 frame.

diff --git a/Engine/Source/Runtime/Engine/Private/Components/SceneCaptureComponent.cpp b/Engine/Source/Runtime/Engine/Private/Components/SceneCaptureComponent.cpp
--- a/Engine/Source/Runtime/Engine/Private/Components/SceneCaptureComponent.cpp
+++ b/Engine/Source/Runtime/Engine/Private/Components/SceneCaptureComponent.cpp
@@ -483,7 +483,7 @@ void USceneCaptureComponent2D::OnRegister()
 
 void USceneCaptureComponent2D::SendRenderTransform_Concurrent()
 {	
-	if (bCaptureOnMovement)
+	if (bCaptureOnMovement && !bCaptureEveryFrame)
 	{
 		CaptureSceneDeferred();
 	}
@@ -888,7 +888,7 @@ void USceneCaptureComponentCube::OnRegister()
 
 void USceneCaptureComponentCube::SendRenderTransform_Concurrent()
 {	
-	if (bCaptureOnMovement)
+	if (bCaptureOnMovement && !bCaptureEveryFrame)
 	{
 		CaptureSceneDeferred();
 	}
