Merging using Dev-VR_to_Release-4.21

Remove check(VROverlay) in FSteamVRHMD::GetStereoLayers() and just return nullptr if it is null.

Revert previous fix for UE-54565 that caused UE-63210

diff --git a/Engine/Plugins/Runtime/Steam/SteamVR/Source/SteamVR/Private/SteamVRStereoLayers.cpp b/Engine/Plugins/Runtime/Steam/SteamVR/Source/SteamVR/Private/SteamVRStereoLayers.cpp
--- a/Engine/Plugins/Runtime/Steam/SteamVR/Source/SteamVR/Private/SteamVRStereoLayers.cpp
+++ b/Engine/Plugins/Runtime/Steam/SteamVR/Source/SteamVR/Private/SteamVRStereoLayers.cpp
@@ -380,8 +380,12 @@ void FSteamVRHMD::GetAllocatedTexture(uint32 LayerId, FTextureRHIRef &Texture, F
 //=============================================================================
 IStereoLayers* FSteamVRHMD::GetStereoLayers ()
 {
-	check(VROverlay);
-	return this;
+	if (VROverlay)
+	{
+		return this;
+	}
+
+	return nullptr;
 }
 
 #endif //STEAMVR_SUPPORTED_PLATFORMS
diff --git a/Engine/Source/Runtime/Engine/Private/Components/StereoLayerComponent.cpp b/Engine/Source/Runtime/Engine/Private/Components/StereoLayerComponent.cpp
--- a/Engine/Source/Runtime/Engine/Private/Components/StereoLayerComponent.cpp
+++ b/Engine/Source/Runtime/Engine/Private/Components/StereoLayerComponent.cpp
@@ -39,7 +39,7 @@ void UStereoLayerComponent::BeginDestroy()
 	Super::BeginDestroy();
 
 	IStereoLayers* StereoLayers;
-	if (LayerId && GEngine->StereoRenderingDevice.IsValid() && GEngine->StereoRenderingDevice->IsStereoEnabled() && (StereoLayers = GEngine->StereoRenderingDevice->GetStereoLayers()) != nullptr)
+	if (LayerId && GEngine->StereoRenderingDevice.IsValid() && (StereoLayers = GEngine->StereoRenderingDevice->GetStereoLayers()) != nullptr)
 	{
 		StereoLayers->DestroyLayer(LayerId);
 		LayerId = 0;
@@ -51,7 +51,7 @@ void UStereoLayerComponent::TickComponent(float DeltaTime, enum ELevelTick TickT
 	Super::TickComponent( DeltaTime, TickType, ThisTickFunction );
 
 	IStereoLayers* StereoLayers;
-	if (!GEngine->StereoRenderingDevice.IsValid() || !GEngine->StereoRenderingDevice->IsStereoEnabled() || (StereoLayers = GEngine->StereoRenderingDevice->GetStereoLayers()) == nullptr || !Texture )
+	if (!GEngine->StereoRenderingDevice.IsValid() || (StereoLayers = GEngine->StereoRenderingDevice->GetStereoLayers()) == nullptr || !Texture )
 	{
 		return;
 	}
