Reimplementing CL #4402016 from Dev-Core:

Fixed cooked materials not working in the editor

diff --git a/Engine/Plugins/FX/Niagara/Source/NiagaraShader/Private/NiagaraShader.cpp b/Engine/Plugins/FX/Niagara/Source/NiagaraShader/Private/NiagaraShader.cpp
--- a/Engine/Plugins/FX/Niagara/Source/NiagaraShader/Private/NiagaraShader.cpp
+++ b/Engine/Plugins/FX/Niagara/Source/NiagaraShader/Private/NiagaraShader.cpp
@@ -463,7 +463,7 @@ void FNiagaraShaderMap::LoadFromDerivedDataCache(const FNiagaraShaderScript* Scr
 
 				// Deserialize from the cached data
 				InOutShaderMap->Serialize(Ar);
-				InOutShaderMap->RegisterSerializedShaders();
+				InOutShaderMap->RegisterSerializedShaders(false);
 
 				checkSlow(InOutShaderMap->GetShaderMapId() == ShaderMapId);
 
@@ -497,8 +497,8 @@ TArray<uint8>* FNiagaraShaderMap::BackupShadersToMemory()
 	TArray<uint8>* SavedShaderData = new TArray<uint8>();
 	FMemoryWriter Ar(*SavedShaderData);
 
-	SerializeInline(Ar, true, true);
-	RegisterSerializedShaders();
+	SerializeInline(Ar, true, true, false);
+	RegisterSerializedShaders(false);
 	Empty();
 
 	return SavedShaderData;
@@ -507,8 +507,8 @@ TArray<uint8>* FNiagaraShaderMap::BackupShadersToMemory()
 void FNiagaraShaderMap::RestoreShadersFromMemory(const TArray<uint8>& ShaderData)
 {
 	FMemoryReader Ar(ShaderData);
-	SerializeInline(Ar, true, true);
-	RegisterSerializedShaders();
+	SerializeInline(Ar, true, true, false);
+	RegisterSerializedShaders(false);
 }
 
 void FNiagaraShaderMap::SaveForRemoteRecompile(FArchive& Ar, const TMap<FString, TArray<TRefCountPtr<FNiagaraShaderMap> > >& CompiledShaderMaps, const TArray<FShaderResourceId>& ClientResourceIds)
@@ -564,7 +564,7 @@ void FNiagaraShaderMap::SaveForRemoteRecompile(FArchive& Ar, const TMap<FString,
 
 	for (int32 Index = 0; Index < NumUniqueResources; Index++)
 	{
-		UniqueResources[Index]->Serialize(Ar);
+		UniqueResources[Index]->Serialize(Ar, false);
 	}
 
 	// now we serialize a map (for each script)
@@ -1092,21 +1092,21 @@ void FNiagaraShaderMap::Serialize(FArchive& Ar, bool bInlineShaderResources)
 
 	if (Ar.IsSaving())
 	{
-		TShaderMap<FNiagaraShaderType>::SerializeInline(Ar, bInlineShaderResources, false);
-		RegisterSerializedShaders();
+		TShaderMap<FNiagaraShaderType>::SerializeInline(Ar, bInlineShaderResources, false, false);
+		RegisterSerializedShaders(false);
 	}
 
 	if (Ar.IsLoading())
 	{
-		TShaderMap<FNiagaraShaderType>::SerializeInline(Ar, bInlineShaderResources, false);
+		TShaderMap<FNiagaraShaderType>::SerializeInline(Ar, bInlineShaderResources, false, false);
 	}
 }
 
-void FNiagaraShaderMap::RegisterSerializedShaders()
+void FNiagaraShaderMap::RegisterSerializedShaders(bool bCookedMaterial)
 {
 	check(IsInGameThread());
 
-	TShaderMap<FNiagaraShaderType>::RegisterSerializedShaders();
+	TShaderMap<FNiagaraShaderType>::RegisterSerializedShaders(bCookedMaterial);
 }
 
 void FNiagaraShaderMap::DiscardSerializedShaders()
diff --git a/Engine/Plugins/FX/Niagara/Source/NiagaraShader/Private/NiagaraShared.cpp b/Engine/Plugins/FX/Niagara/Source/NiagaraShader/Private/NiagaraShared.cpp
--- a/Engine/Plugins/FX/Niagara/Source/NiagaraShader/Private/NiagaraShared.cpp
+++ b/Engine/Plugins/FX/Niagara/Source/NiagaraShader/Private/NiagaraShared.cpp
@@ -136,7 +136,7 @@ void FNiagaraShaderScript::RegisterShaderMap()
 {
 	if (GameThreadShaderMap)
 	{
-		GameThreadShaderMap->RegisterSerializedShaders();
+		GameThreadShaderMap->RegisterSerializedShaders(false);
 	}
 }
 
diff --git a/Engine/Plugins/FX/Niagara/Source/NiagaraShader/Public/NiagaraShared.h b/Engine/Plugins/FX/Niagara/Source/NiagaraShader/Public/NiagaraShared.h
--- a/Engine/Plugins/FX/Niagara/Source/NiagaraShader/Public/NiagaraShared.h
+++ b/Engine/Plugins/FX/Niagara/Source/NiagaraShader/Public/NiagaraShared.h
@@ -337,7 +337,7 @@ public:
 	void SaveToDerivedDataCache();
 
 	/** Registers all shaders that have been loaded in Serialize */
-	virtual void RegisterSerializedShaders() override;
+	virtual void RegisterSerializedShaders(bool bCookedMaterial) override;
 	virtual void DiscardSerializedShaders() override;
 
 	/** Backs up any FShaders in this shader map to memory through serialization and clears FShader references. */
diff --git a/Engine/Source/Runtime/Engine/Private/Materials/Material.cpp b/Engine/Source/Runtime/Engine/Private/Materials/Material.cpp
--- a/Engine/Source/Runtime/Engine/Private/Materials/Material.cpp
+++ b/Engine/Source/Runtime/Engine/Private/Materials/Material.cpp
@@ -701,9 +701,14 @@ void ProcessSerializedInlineShaderMaps(UMaterialInterface* Owner, TArray<FMateri
 	UMaterial* OwnerMaterial = Cast<UMaterial>(Owner);
 	UMaterialInstance* OwnerMaterialInstance = Cast<UMaterialInstance>(Owner);
 
+#if WITH_EDITORONLY_DATA
+	const bool bLoadedByCookedMaterial = FPlatformProperties::RequiresCookedData() || Owner->GetOutermost()->bIsCookedForEditor;
+#else
+	const bool bLoadedByCookedMaterial = FPlatformProperties::RequiresCookedData();
+#endif
 	for (FMaterialResource& Resource : LoadedResources)
 	{
-		Resource.RegisterInlineShaderMap();
+		Resource.RegisterInlineShaderMap(bLoadedByCookedMaterial);
 	}
 	
 	if (CVarDiscardUnusedQualityLevels.GetValueOnAnyThread())
diff --git a/Engine/Source/Runtime/Engine/Private/Materials/MaterialShader.cpp b/Engine/Source/Runtime/Engine/Private/Materials/MaterialShader.cpp
--- a/Engine/Source/Runtime/Engine/Private/Materials/MaterialShader.cpp
+++ b/Engine/Source/Runtime/Engine/Private/Materials/MaterialShader.cpp
@@ -1060,7 +1060,7 @@ void FMaterialShaderMap::LoadFromDerivedDataCache(const FMaterial* Material, con
 
 				// Deserialize from the cached data
 				InOutShaderMap->Serialize(Ar);
-				InOutShaderMap->RegisterSerializedShaders();
+				InOutShaderMap->RegisterSerializedShaders(false);
 
 				const FString InDataKey = GetMaterialShaderMapKeyString(InOutShaderMap->GetShaderMapId(), InPlatform);
 				checkSlow(InOutShaderMap->GetShaderMapId() == ShaderMapId);
@@ -1099,11 +1099,11 @@ TArray<uint8>* FMaterialShaderMap::BackupShadersToMemory()
 	{
 		// Serialize data needed to handle shader key changes in between the save and the load of the FShaders
 		const bool bHandleShaderKeyChanges = true;
-		MeshShaderMaps[Index].SerializeInline(Ar, true, bHandleShaderKeyChanges);
+		MeshShaderMaps[Index].SerializeInline(Ar, true, bHandleShaderKeyChanges, false);
 		MeshShaderMaps[Index].Empty();
 	}
 
-	SerializeInline(Ar, true, true);
+	SerializeInline(Ar, true, true, false);
 	Empty();
 
 	return SavedShaderData;
@@ -1117,12 +1117,12 @@ void FMaterialShaderMap::RestoreShadersFromMemory(const TArray<uint8>& ShaderDat
 	{
 		// Use the serialized shader key data to detect when the saved shader is no longer valid and skip it
 		const bool bHandleShaderKeyChanges = true;
-		MeshShaderMaps[Index].SerializeInline(Ar, true, bHandleShaderKeyChanges);
-		MeshShaderMaps[Index].RegisterSerializedShaders();
+		MeshShaderMaps[Index].SerializeInline(Ar, true, bHandleShaderKeyChanges, false);
+		MeshShaderMaps[Index].RegisterSerializedShaders(false);
 	}
 
-	SerializeInline(Ar, true, true);
-	RegisterSerializedShaders();
+	SerializeInline(Ar, true, true, false);
+	RegisterSerializedShaders(false);
 }
 
 void FMaterialShaderMap::SaveForRemoteRecompile(FArchive& Ar, const TMap<FString, TArray<TRefCountPtr<FMaterialShaderMap> > >& CompiledShaderMaps, const TArray<FShaderResourceId>& ClientResourceIds)
@@ -1191,7 +1191,7 @@ void FMaterialShaderMap::SaveForRemoteRecompile(FArchive& Ar, const TMap<FString
 
 	for (int32 Index = 0; Index < NumUniqueResources; Index++)
 	{
-		UniqueResources[Index]->Serialize(Ar);
+		UniqueResources[Index]->Serialize(Ar, false);
 	}
 
 	// now we serialize a map (for each material), but without inline the resources, since they are above
@@ -1240,7 +1240,7 @@ void FMaterialShaderMap::LoadForRemoteRecompile(FArchive& Ar, EShaderPlatform Sh
 	{
 		// Load the inlined shader resource
 		FShaderResource* Resource = new FShaderResource();
-		Resource->Serialize(Ar);
+		Resource->Serialize(Ar, false);
 
 		// if this Id is already in memory, that means that this is a repeated resource and so we skip it
 		if (FShaderResource::FindShaderResourceById(Resource->GetId()) != NULL)
@@ -1315,7 +1315,7 @@ void FMaterialShaderMap::LoadForRemoteRecompile(FArchive& Ar, EShaderPlatform Sh
 							FMaterialResource* MaterialResource = MatchingMaterial->GetMaterialResource(GetMaxSupportedFeatureLevel(ShaderPlatform), (EMaterialQualityLevel::Type)QualityLevelIndex);
 
 							MaterialResource->SetGameThreadShaderMap(LoadedShaderMap);
-							MaterialResource->RegisterInlineShaderMap();
+							MaterialResource->RegisterInlineShaderMap(false);
 
 							ENQUEUE_UNIQUE_RENDER_COMMAND_TWOPARAMETER(
 								FSetShaderMapOnMaterialResources,
@@ -2262,7 +2262,7 @@ struct FCompareMeshShaderMaps
 	}
 };
 
-void FMaterialShaderMap::Serialize(FArchive& Ar, bool bInlineShaderResources)
+void FMaterialShaderMap::Serialize(FArchive& Ar, bool bInlineShaderResources, bool bLoadedByCookedMaterial)
 {
 	LLM_SCOPE(ELLMTag::MaterialShaderMaps);
 	// Note: This is saved to the DDC, not into packages (except when cooked)
@@ -2294,7 +2294,7 @@ void FMaterialShaderMap::Serialize(FArchive& Ar, bool bInlineShaderResources)
 	if (Ar.IsSaving())
 	{
 		// Material shaders
-		TShaderMap<FMaterialShaderType>::SerializeInline(Ar, bInlineShaderResources, false);
+		TShaderMap<FMaterialShaderType>::SerializeInline(Ar, bInlineShaderResources, false, false);
 
 		// Mesh material shaders
 		int32 NumMeshShaderMaps = 0;
@@ -2334,7 +2334,7 @@ void FMaterialShaderMap::Serialize(FArchive& Ar, bool bInlineShaderResources)
 
 				Ar << VFType;
 
-				MeshShaderMap->SerializeInline(Ar, bInlineShaderResources, false);
+				MeshShaderMap->SerializeInline(Ar, bInlineShaderResources, false, false);
 			}
 		}
 
@@ -2374,7 +2374,7 @@ void FMaterialShaderMap::Serialize(FArchive& Ar, bool bInlineShaderResources)
 		InitOrderedMeshShaderMaps();
 
 		// Material shaders
-		TShaderMap<FMaterialShaderType>::SerializeInline(Ar, bInlineShaderResources, false);
+		TShaderMap<FMaterialShaderType>::SerializeInline(Ar, bInlineShaderResources, false, bLoadedByCookedMaterial);
 
 		// Mesh material shaders
 		int32 NumMeshShaderMaps = 0;
@@ -2391,7 +2391,7 @@ void FMaterialShaderMap::Serialize(FArchive& Ar, bool bInlineShaderResources)
 			check(VFType);
 			FMeshMaterialShaderMap* MeshShaderMap = OrderedMeshShaderMaps[VFType->GetId()];
 			check(MeshShaderMap);
-			MeshShaderMap->SerializeInline(Ar, bInlineShaderResources, false);
+			MeshShaderMap->SerializeInline(Ar, bInlineShaderResources, false, bLoadedByCookedMaterial);
 		}
 
 		bool bCooked;
@@ -2418,17 +2418,17 @@ void FMaterialShaderMap::Serialize(FArchive& Ar, bool bInlineShaderResources)
 	}
 }
 
-void FMaterialShaderMap::RegisterSerializedShaders()
+void FMaterialShaderMap::RegisterSerializedShaders(bool bLoadedByCookedMaterial)
 {
 	check(IsInGameThread());
 
-	TShaderMap<FMaterialShaderType>::RegisterSerializedShaders();
+	TShaderMap<FMaterialShaderType>::RegisterSerializedShaders(bLoadedByCookedMaterial);
 	
 	for (FMeshMaterialShaderMap* MeshShaderMap : OrderedMeshShaderMaps)
 	{
 		if (MeshShaderMap)
 		{
-			MeshShaderMap->RegisterSerializedShaders();
+			MeshShaderMap->RegisterSerializedShaders(bLoadedByCookedMaterial);
 		}
 	}
 
diff --git a/Engine/Source/Runtime/Engine/Private/Materials/MaterialShared.cpp b/Engine/Source/Runtime/Engine/Private/Materials/MaterialShared.cpp
--- a/Engine/Source/Runtime/Engine/Private/Materials/MaterialShared.cpp
+++ b/Engine/Source/Runtime/Engine/Private/Materials/MaterialShared.cpp
@@ -848,14 +848,14 @@ void FMaterial::SerializeInlineShaderMap(FArchive& Ar)
 			if (bValid)
 			{
 				TRefCountPtr<FMaterialShaderMap> LoadedShaderMap = new FMaterialShaderMap();
-				LoadedShaderMap->Serialize(Ar);
+				LoadedShaderMap->Serialize(Ar, true, bCooked && Ar.IsLoading());
 				GameThreadShaderMap = LoadedShaderMap;
 			}
 		}
 	}
 }
 
-void FMaterial::RegisterInlineShaderMap()
+void FMaterial::RegisterInlineShaderMap(bool bLoadedByCookedMaterial)
 {
 	if (GameThreadShaderMap)
 	{
@@ -865,7 +865,7 @@ void FMaterial::RegisterInlineShaderMap()
 		{
 			RenderingThreadShaderMap = GameThreadShaderMap;
 		}
-		GameThreadShaderMap->RegisterSerializedShaders();
+		GameThreadShaderMap->RegisterSerializedShaders(bLoadedByCookedMaterial);
 	}
 }
 
diff --git a/Engine/Source/Runtime/Engine/Private/ShaderCompiler/ShaderCompiler.cpp b/Engine/Source/Runtime/Engine/Private/ShaderCompiler/ShaderCompiler.cpp
--- a/Engine/Source/Runtime/Engine/Private/ShaderCompiler/ShaderCompiler.cpp
+++ b/Engine/Source/Runtime/Engine/Private/ShaderCompiler/ShaderCompiler.cpp
@@ -3857,9 +3857,9 @@ static void SerializeGlobalShaders(FArchive& Ar, TShaderMap<FGlobalShaderType>*
 	}
 
 	// Serialize the global shaders.
-	GlobalShaderMap->SerializeInline(Ar, true, false, ShaderKeysToSave);
+	GlobalShaderMap->SerializeInline(Ar, true, false, false, ShaderKeysToSave);
 	// And now register them.
-	GlobalShaderMap->RegisterSerializedShaders();
+	GlobalShaderMap->RegisterSerializedShaders(false);
 }
 
 /** Saves the platform's shader map to the DDC. */
diff --git a/Engine/Source/Runtime/Engine/Public/MaterialShared.h b/Engine/Source/Runtime/Engine/Public/MaterialShared.h
--- a/Engine/Source/Runtime/Engine/Public/MaterialShared.h
+++ b/Engine/Source/Runtime/Engine/Public/MaterialShared.h
@@ -811,13 +811,13 @@ public:
 	static const FMaterialShaderMap* GetShaderMapBeingCompiled(const FMaterial* Material);
 
 	/** Serializes the shader map. */
-	void Serialize(FArchive& Ar, bool bInlineShaderResources=true);
+	void Serialize(FArchive& Ar, bool bInlineShaderResources=true, bool bLoadedByCookedMaterial=false);
 
 	/** Saves this shader map to the derived data cache. */
 	void SaveToDerivedDataCache();
 
 	/** Registers all shaders that have been loaded in Serialize */
-	virtual void RegisterSerializedShaders() override;
+	virtual void RegisterSerializedShaders(bool bLoadedByCookedMaterial) override;
 	virtual void DiscardSerializedShaders() override;
 
 	/** Backs up any FShaders in this shader map to memory through serialization and clears FShader references. */
@@ -1170,7 +1170,7 @@ public:
 	void SerializeInlineShaderMap(FArchive& Ar);
 
 	/** Serializes the shader map inline in this material, including any shader dependencies. */
-	void RegisterInlineShaderMap();
+	void RegisterInlineShaderMap(bool bLoadedByCookedMaterial);
 
 	/** Releases this material's shader map.  Must only be called on materials not exposed to the rendering thread! */
 	void ReleaseShaderMap();
diff --git a/Engine/Source/Runtime/ShaderCore/Private/GlobalShader.cpp b/Engine/Source/Runtime/ShaderCore/Private/GlobalShader.cpp
--- a/Engine/Source/Runtime/ShaderCore/Private/GlobalShader.cpp
+++ b/Engine/Source/Runtime/ShaderCore/Private/GlobalShader.cpp
@@ -184,7 +184,7 @@ void BackupGlobalShaderMap(FGlobalShaderBackupData& OutGlobalShaderBackup)
 			TUniquePtr<TArray<uint8>> ShaderData = MakeUnique<TArray<uint8>>();
 			FMemoryWriter Ar(*ShaderData);
 			GGlobalShaderMap[ShaderPlatform]->SerializeInline(Ar, true, true, nullptr);
-			GGlobalShaderMap[ShaderPlatform]->RegisterSerializedShaders();
+			GGlobalShaderMap[ShaderPlatform]->RegisterSerializedShaders(false);
 			GGlobalShaderMap[ShaderPlatform]->Empty();
 			OutGlobalShaderBackup.FeatureLevelShaderData[i] = MoveTemp(ShaderData);
 		}
@@ -208,7 +208,7 @@ void RestoreGlobalShaderMap(const FGlobalShaderBackupData& GlobalShaderBackup)
 		{
 			FMemoryReader Ar(*GlobalShaderBackup.FeatureLevelShaderData[i]);
 			GGlobalShaderMap[ShaderPlatform]->SerializeInline(Ar, true, true, nullptr);
-			GGlobalShaderMap[ShaderPlatform]->RegisterSerializedShaders();
+			GGlobalShaderMap[ShaderPlatform]->RegisterSerializedShaders(false);
 		}
 	}
 }
diff --git a/Engine/Source/Runtime/ShaderCore/Private/Shader.cpp b/Engine/Source/Runtime/ShaderCore/Private/Shader.cpp
--- a/Engine/Source/Runtime/ShaderCore/Private/Shader.cpp
+++ b/Engine/Source/Runtime/ShaderCore/Private/Shader.cpp
@@ -355,7 +355,7 @@ void FShaderType::Initialize(const TMap<FString, TArray<const TCHAR*> >& ShaderF
 				TArray<uint8> TempData;
 				FMemoryWriter Ar(TempData, true);
 				FShaderSaveArchive SaveArchive(Ar, Type->SerializationHistory);
-				TempShader->SerializeBase(SaveArchive, false);
+				TempShader->SerializeBase(SaveArchive, false, false);
 
 				// Destroy the temporary shader
 				delete TempShader;
@@ -483,7 +483,7 @@ void FShaderResource::Register()
 }
 
 
-void FShaderResource::Serialize(FArchive& Ar)
+void FShaderResource::Serialize(FArchive& Ar, bool bLoadedByCookedMaterial)
 {
 	check(!(SpecificPermutationId != 0 && SpecificType == nullptr));
 
@@ -520,7 +520,10 @@ void FShaderResource::Serialize(FArchive& Ar)
 	}
 
 #if WITH_EDITORONLY_DATA
-	SerializePlatformDebugData(Ar);
+	if (!bLoadedByCookedMaterial)
+	{
+		SerializePlatformDebugData(Ar);
+	}
 #endif
 
 	if (Ar.IsLoading())
@@ -977,7 +980,7 @@ FShader::FShader(const CompiledShaderInitializerType& Initializer)
 	SetResource(Initializer.Resource);
 
 	// Register the shader now that it is valid, so that it can be reused
-	Register();
+	Register(false);
 }
 
 
@@ -1003,7 +1006,7 @@ EShaderPlatform FShader::GetShaderPlatform() const
 }
 
 
-bool FShader::SerializeBase(FArchive& Ar, bool bShadersInline)
+bool FShader::SerializeBase(FArchive& Ar, bool bShadersInline, bool bLoadedByCookedMaterial)
 {
 	Ar.UsingCustomVersion(FFortniteMainBranchObjectVersion::GUID);
 
@@ -1088,14 +1091,14 @@ bool FShader::SerializeBase(FArchive& Ar, bool bShadersInline)
 		if (Ar.IsSaving())
 		{
 			check(Resource->Target == Target);
-			Resource->Serialize(Ar);
+			Resource->Serialize(Ar, false);
 		}
 
 		if (Ar.IsLoading())
 		{
 			// Load the inlined shader resource
 			SerializedResource = new FShaderResource();
-			SerializedResource->Serialize(Ar);
+			SerializedResource->Serialize(Ar, bLoadedByCookedMaterial);
 		}
 	}
 	else
@@ -1139,11 +1142,11 @@ void FShader::Release()
 }
 
 
-void FShader::Register()
+void FShader::Register(bool bLoadedByCookedMaterial)
 {
 	FShaderId ShaderId = GetId();
 	check(ShaderId.MaterialShaderMapHash != FSHAHash());
-	check(ShaderId.SourceHash != FSHAHash() || FPlatformProperties::RequiresCookedData());
+	check(ShaderId.SourceHash != FSHAHash() || FPlatformProperties::RequiresCookedData() || bLoadedByCookedMaterial);
 	check(Resource);
 	Type->AddToShaderIdMap(ShaderId, this);
 }
diff --git a/Engine/Source/Runtime/ShaderCore/Public/Shader.h b/Engine/Source/Runtime/ShaderCore/Public/Shader.h
--- a/Engine/Source/Runtime/ShaderCore/Public/Shader.h
+++ b/Engine/Source/Runtime/ShaderCore/Public/Shader.h
@@ -194,7 +194,7 @@ public:
 
 	~FShaderResource();
 
-	SHADERCORE_API void Serialize(FArchive& Ar);
+	SHADERCORE_API void Serialize(FArchive& Ar, bool bLoadedByCookedMaterial);
 
 	// Reference counting.
 	SHADERCORE_API void AddRef();
@@ -638,7 +638,7 @@ public:
 */
 
 	/** Serializes the shader. */
-	bool SerializeBase(FArchive& Ar, bool bShadersInline);
+	bool SerializeBase(FArchive& Ar, bool bShadersInline, bool bLoadedByCookedMaterial);
 
 	virtual bool Serialize(FArchive& Ar) { return false; }
 
@@ -647,7 +647,7 @@ public:
 	void Release();
 
 	/** Registers this shader for lookup by ID. */
-	void Register();
+	void Register(bool bLoadedByCookedMaterial);
 
 	/** Removes this shader from the ID lookup map. */
 	void Deregister();
@@ -1760,7 +1760,7 @@ public:
 			Ar << SelfContainedKey;
 		}
 
-		CurrentShader->SerializeBase(Ar, bInlineShaderResource);
+		CurrentShader->SerializeBase(Ar, bInlineShaderResource, false);
 
 		// Get the offset to the end of the shader's serialized data
 		int64 EndOffset = Ar.Tell();
@@ -1772,7 +1772,7 @@ public:
 		Ar.Seek(EndOffset);
 	}
 
-	inline FShader* SerializeShaderForLoad(FShaderType* Type, FArchive& Ar, bool bHandleShaderKeyChanges, bool bInlineShaderResource)
+	inline FShader* SerializeShaderForLoad(FShaderType* Type, FArchive& Ar, bool bHandleShaderKeyChanges, bool bInlineShaderResource, bool bLoadedByCookedMaterial)
 	{
 		int64 EndOffset = 0;
 		Ar << EndOffset;
@@ -1792,7 +1792,7 @@ public:
 		{
 			Shader = Type->ConstructForDeserialization();
 			check(Shader != nullptr);
-			Shader->SerializeBase(Ar, bInlineShaderResource);
+			Shader->SerializeBase(Ar, bInlineShaderResource, bLoadedByCookedMaterial);
 		}
 		else
 		{
@@ -1809,7 +1809,7 @@ public:
 	 * @param bHandleShaderKeyChanges - whether to serialize the data necessary to detect and gracefully handle shader key changes between saving and loading
 	 * @param DependenciesToSave - array of specific ShaderTypeDepencies which should be saved.
 	 */
-	void SerializeInline(FArchive& Ar, bool bInlineShaderResource, bool bHandleShaderKeyChanges, const TArray<FShaderPrimaryKey>* ShaderKeysToSave = nullptr)
+	void SerializeInline(FArchive& Ar, bool bInlineShaderResource, bool bHandleShaderKeyChanges, bool bLoadedByCookedMaterial, const TArray<FShaderPrimaryKey>* ShaderKeysToSave = nullptr)
 	{
 		if (Ar.IsSaving())
 		{
@@ -1889,7 +1889,7 @@ public:
 				FShaderType* Type = nullptr;
 				Ar << Type;
 
-				FShader* Shader = SerializeShaderForLoad(Type, Ar, bHandleShaderKeyChanges, bInlineShaderResource);
+				FShader* Shader = SerializeShaderForLoad(Type, Ar, bHandleShaderKeyChanges, bInlineShaderResource, bLoadedByCookedMaterial);
 				if (Shader)
 				{
 					SerializedShaders.Add(Shader);
@@ -1910,7 +1910,7 @@ public:
 				{
 					FShaderType* Type = nullptr;
 					Ar << Type;
-					FShader* Shader = SerializeShaderForLoad(Type, Ar, bHandleShaderKeyChanges, bInlineShaderResource);
+					FShader* Shader = SerializeShaderForLoad(Type, Ar, bHandleShaderKeyChanges, bInlineShaderResource, bLoadedByCookedMaterial);
 					if (Shader)
 					{
 						ShaderStages.Add(Shader);
@@ -1930,7 +1930,7 @@ public:
 	}
 
 	/** Registered all shaders that have been serialized (maybe) on another thread */
-	virtual void RegisterSerializedShaders()
+	virtual void RegisterSerializedShaders(bool bCookedMaterial)
 	{
 		bHasBeenRegistered = true;
 		check(IsInGameThread());
@@ -1949,7 +1949,7 @@ public:
 			else
 			{
 				// Register the shader now that it is valid, so that it can be reused
-				Shader->Register();
+				Shader->Register(bCookedMaterial);
 			}
 			AddShader(Shader->GetType(), Shader->GetPermutationId(), Shader);
 		}
