Merging CL# 4419752 from Dev-VR - Fixes to prevent audio from crashing on exit

diff --git a/Engine/Plugins/Runtime/ResonanceAudio/Source/ResonanceAudio/Private/ResonanceAudioReverb.cpp b/Engine/Plugins/Runtime/ResonanceAudio/Source/ResonanceAudio/Private/ResonanceAudioReverb.cpp
--- a/Engine/Plugins/Runtime/ResonanceAudio/Source/ResonanceAudio/Private/ResonanceAudioReverb.cpp
+++ b/Engine/Plugins/Runtime/ResonanceAudio/Source/ResonanceAudio/Private/ResonanceAudioReverb.cpp
@@ -24,6 +24,10 @@ namespace ResonanceAudio
 
 	FResonanceAudioReverb::~FResonanceAudioReverb()
 	{
+		if (GlobalReverbPluginPreset)
+		{
+			GlobalReverbPluginPreset->RemoveFromRoot();
+		}
 	}
 
 	void FResonanceAudioReverb::Initialize(const FAudioPluginInitializationParams InitializationParams)
@@ -116,10 +120,11 @@ namespace ResonanceAudio
 
 		if (GlobalReverbPluginPreset)
 		{
+			GlobalReverbPluginPreset->AddToRoot();
 			ReverbPluginPreset = GlobalReverbPluginPreset;
 
 			FResonanceAudioReverbPlugin* Effect = static_cast<FResonanceAudioReverbPlugin*>(GlobalReverbPluginPreset->CreateNewEffect());
-			Effect->RegisterWithPreset(GlobalReverbPluginPreset);
+			Effect->SetPreset(GlobalReverbPluginPreset);
 			Effect->SetResonanceAudioReverbPlugin(this);
 			return static_cast<FSoundEffectSubmix*>(Effect);
 		}
diff --git a/Engine/Source/Runtime/AudioMixer/Private/AudioMixerSourceManager.cpp b/Engine/Source/Runtime/AudioMixer/Private/AudioMixerSourceManager.cpp
--- a/Engine/Source/Runtime/AudioMixer/Private/AudioMixerSourceManager.cpp
+++ b/Engine/Source/Runtime/AudioMixer/Private/AudioMixerSourceManager.cpp
@@ -459,7 +459,6 @@ namespace Audio
 			}
 
 			FSoundEffectSource* NewEffect = static_cast<FSoundEffectSource*>(ChainEntry.Preset->CreateNewEffect());
-			NewEffect->RegisterWithPreset(ChainEntry.Preset);
 
 			// Get this source effect presets unique id so instances can identify their originating preset object
 			const uint32 PresetUniqueId = ChainEntry.Preset->GetUniqueID();
@@ -484,7 +483,7 @@ namespace Audio
 
 		for (int32 i = 0; i < SourceInfo.SourceEffects.Num(); ++i)
 		{
-			SourceInfo.SourceEffects[i]->UnregisterWithPreset();
+			SourceInfo.SourceEffects[i]->ClearPreset();
 			delete SourceInfo.SourceEffects[i];
 			SourceInfo.SourceEffects[i] = nullptr;
 		}
@@ -1984,7 +1983,7 @@ namespace Audio
 							const FSourceEffectChainEntry& ChainEntry = InSourceEffectChain[SourceEffectId];
 
 							FSoundEffectSource* SourceEffectInstance = ThisSourceEffectChain[SourceEffectId];
-							if (!SourceEffectInstance->IsParentPreset(ChainEntry.Preset))
+							if (!SourceEffectInstance->IsPreset(ChainEntry.Preset))
 							{
 								// As soon as one of the effects change or is not the same, then we need to rebuild the effect graph
 								bReset = true;
diff --git a/Engine/Source/Runtime/AudioMixer/Private/AudioMixerSubmix.cpp b/Engine/Source/Runtime/AudioMixer/Private/AudioMixerSubmix.cpp
--- a/Engine/Source/Runtime/AudioMixer/Private/AudioMixerSubmix.cpp
+++ b/Engine/Source/Runtime/AudioMixer/Private/AudioMixerSubmix.cpp
@@ -64,9 +64,6 @@ namespace Audio
 					// Create a new effect instance using the preset
 					FSoundEffectSubmix* SubmixEffect = static_cast<FSoundEffectSubmix*>(EffectPreset->CreateNewEffect());
 
-					// Register the submix effect instance with the preset so it can be updated dynamically
-					SubmixEffect->RegisterWithPreset(EffectPreset);
-
 					FSoundEffectSubmixInitData InitData;
 					InitData.SampleRate = MixerDevice->GetSampleRate();
 					InitData.PresetSettings = nullptr;
@@ -287,7 +284,7 @@ namespace Audio
 		{
 			if (Info.EffectInstance)
 			{
-				Info.EffectInstance->UnregisterWithPreset();
+				Info.EffectInstance->ClearPreset();
 
 				delete Info.EffectInstance;
 				Info.EffectInstance = nullptr;
diff --git a/Engine/Source/Runtime/AudioMixer/Public/AudioMixerSubmix.h b/Engine/Source/Runtime/AudioMixer/Public/AudioMixerSubmix.h
--- a/Engine/Source/Runtime/AudioMixer/Public/AudioMixerSubmix.h
+++ b/Engine/Source/Runtime/AudioMixer/Public/AudioMixerSubmix.h
@@ -45,7 +45,7 @@ namespace Audio
 	{
 	public:
 		FMixerSubmix(FMixerDevice* InMixerDevice);
-		~FMixerSubmix();
+		virtual ~FMixerSubmix();
 
 		// Initialize the submix object with the USoundSubmix ptr. Sets up child and parent connects.
 		void Init(USoundSubmix* InSoundSubmix);
diff --git a/Engine/Source/Runtime/Engine/Classes/Sound/SoundEffectBase.h b/Engine/Source/Runtime/Engine/Classes/Sound/SoundEffectBase.h
--- a/Engine/Source/Runtime/Engine/Classes/Sound/SoundEffectBase.h
+++ b/Engine/Source/Runtime/Engine/Classes/Sound/SoundEffectBase.h
@@ -83,14 +83,11 @@ public:
 
 	void SetPreset(USoundEffectPreset* Inpreset);
 
-	/** Registers the parent preset and registers this instance with the preset. */
-	void RegisterWithPreset(USoundEffectPreset* InParentPreset);
-
 	/** Removes the instance from the preset. */
-	void UnregisterWithPreset();
+	void ClearPreset();
 
-	/** Queries if the given preset object is the parent preset, i.e. the preset which spawned this effect instance. */
-	bool IsParentPreset(USoundEffectPreset* InPreset) const;
+	/** Queries if the given preset object is the uobject preset for this preset instance, i.e. the preset which spawned this effect instance. */
+	bool IsPreset(USoundEffectPreset* InPreset) const;
 
 	/** Enqueues a lambda command on a thread safe queue which is pumped from the audio render thread. */
 	void EffectCommand(TFunction<void()> Command);
@@ -104,7 +101,6 @@ protected:
 
 	FThreadSafeBool bChanged;
 	USoundEffectPreset* Preset;
-	USoundEffectPreset* ParentPreset;
 
 	FThreadSafeBool bIsRunning;
 	FThreadSafeBool bIsActive;
diff --git a/Engine/Source/Runtime/Engine/Private/SoundEffectBase.cpp b/Engine/Source/Runtime/Engine/Private/SoundEffectBase.cpp
--- a/Engine/Source/Runtime/Engine/Private/SoundEffectBase.cpp
+++ b/Engine/Source/Runtime/Engine/Private/SoundEffectBase.cpp
@@ -7,7 +7,6 @@
 FSoundEffectBase::FSoundEffectBase()
 	: bChanged(false)
 	, Preset(nullptr)
-	, ParentPreset(nullptr)
 	, bIsRunning(false)
 	, bIsActive(false)
 {}
@@ -29,37 +28,33 @@ void FSoundEffectBase::SetEnabled(const bool bInIsEnabled)
 void FSoundEffectBase::SetPreset(USoundEffectPreset* Inpreset)
 {
 	Preset = Inpreset;
+	Preset->AddEffectInstance(this);
 	bChanged = true;
 }
 
-void FSoundEffectBase::Update()
+void FSoundEffectBase::ClearPreset()
 {
-	PumpPendingMessages();
-
-	if (bChanged && Preset)
+	if (Preset)
 	{
-		OnPresetChanged();
-		bChanged = false;
+		Preset->RemoveEffectInstance(this);
+		Preset = nullptr;
 	}
 }
 
-void FSoundEffectBase::RegisterWithPreset(USoundEffectPreset* InParentPreset)
+void FSoundEffectBase::Update()
 {
-	ParentPreset = InParentPreset;
-	ParentPreset->AddEffectInstance(this);
-}
+	PumpPendingMessages();
 
-void FSoundEffectBase::UnregisterWithPreset()
-{
-	if (ParentPreset)
+	if (bChanged && Preset)
 	{
-		ParentPreset->RemoveEffectInstance(this);
+		OnPresetChanged();
+		bChanged = false;
 	}
 }
 
-bool FSoundEffectBase::IsParentPreset(USoundEffectPreset* InPreset) const
+bool FSoundEffectBase::IsPreset(USoundEffectPreset* InPreset) const
 {
-	return ParentPreset == InPreset;
+	return Preset == InPreset;
 }
 
 void FSoundEffectBase::EffectCommand(TFunction<void()> Command)
diff --git a/Engine/Source/Runtime/Engine/Private/SoundEffectPreset.cpp b/Engine/Source/Runtime/Engine/Private/SoundEffectPreset.cpp
--- a/Engine/Source/Runtime/Engine/Private/SoundEffectPreset.cpp
+++ b/Engine/Source/Runtime/Engine/Private/SoundEffectPreset.cpp
@@ -13,6 +13,15 @@ USoundEffectPreset::USoundEffectPreset(const FObjectInitializer& ObjectInitializ
 	
 }
 
+USoundEffectPreset::~USoundEffectPreset()
+{
+	for (int32 i = 0; i < Instances.Num(); ++i)
+	{
+		Instances[i]->ClearPreset();
+	}
+}
+
+
 void USoundEffectPreset::EffectCommand(TFunction<void()> Command)
 {
 	for (int32 i = 0; i < Instances.Num(); ++i)
