Merging CL 4401151 from Dev-VR - Adding destructor to SoundPresets

diff --git a/Engine/Source/Runtime/Engine/Classes/Sound/SoundEffectPreset.h b/Engine/Source/Runtime/Engine/Classes/Sound/SoundEffectPreset.h
--- a/Engine/Source/Runtime/Engine/Classes/Sound/SoundEffectPreset.h
+++ b/Engine/Source/Runtime/Engine/Classes/Sound/SoundEffectPreset.h
@@ -18,7 +18,8 @@ class ENGINE_API USoundEffectPreset : public UObject
 
 public:
 	USoundEffectPreset(const FObjectInitializer& ObjectInitializer);
-
+	virtual ~USoundEffectPreset();
+	
 	virtual FText GetAssetActionName() const PURE_VIRTUAL(USoundEffectPreset::GetAssetActionName, return FText(););
 	virtual UClass* GetSupportedClass() const PURE_VIRTUAL(USoundEffectPreset::GetSupportedClass, return nullptr;);
 	virtual USoundEffectPreset* CreateNewPreset(UObject* InParent, FName Name, EObjectFlags Flags) const PURE_VIRTUAL(USoundEffectPreset::CreateNewPreset, return nullptr;);
