Harden EdGraphPin serialization

diff --git a/Engine/Source/Runtime/Engine/Private/EdGraph/EdGraphPin.cpp b/Engine/Source/Runtime/Engine/Private/EdGraph/EdGraphPin.cpp
--- a/Engine/Source/Runtime/Engine/Private/EdGraph/EdGraphPin.cpp
+++ b/Engine/Source/Runtime/Engine/Private/EdGraph/EdGraphPin.cpp
@@ -1803,12 +1803,15 @@ void UEdGraphPin::SerializePinArray(FArchive& Ar, TArray<UEdGraphPin*>& ArrayRef
 			for (int32 I = 0; I < Pin->LinkedTo.Num(); ++I)
 			{
 				UEdGraphPin* Peer = Pin->LinkedTo[I];
-				// PeerNode will be null if the pin we were linked to was already thrown away:
-				UEdGraphNode* PeerNode = Peer->GetOwningNodeUnchecked();
-				if (PeerNode && !GEditor->Trans->IsObjectTransacting(PeerNode))
+				if(ensure(Peer))
 				{
-					Pin->BreakLinkTo(Peer);
-					--I;
+					// PeerNode will be null if the pin we were linked to was already thrown away:
+					UEdGraphNode* PeerNode = Peer->GetOwningNodeUnchecked();
+					if (PeerNode && !GEditor->Trans->IsObjectTransacting(PeerNode))
+					{
+						Pin->BreakLinkTo(Peer);
+						--I;
+					}
 				}
 			}
 			// We also must clear the Pins subpin array as that pin may have been reused if this one wasn't
