PR #5123: Small typo in Engine.h (Contributed by [REDACTED_PERSON])

diff --git a/Engine/Source/Runtime/Engine/Classes/Engine/Engine.h b/Engine/Source/Runtime/Engine/Classes/Engine/Engine.h
--- a/Engine/Source/Runtime/Engine/Classes/Engine/Engine.h
+++ b/Engine/Source/Runtime/Engine/Classes/Engine/Engine.h
@@ -280,7 +280,7 @@ struct FNamedNetDriver
  *  new worlds.
  *
  *	FWorldContext should remain internal to the UEngine classes. Outside code should not keep pointers or try to manage FWorldContexts directly.
- *	Outside code can steal deal with UWorld*, and pass UWorld*s into Engine level functions. The Engine code can look up the relevant context 
+ *	Outside code can still deal with UWorld*, and pass UWorld*s into Engine level functions. The Engine code can look up the relevant context 
  *	for a given UWorld*.
  *
  *  For convenience, FWorldContext can maintain outside pointers to UWorld*s. For example, PIE can tie UWorld* UEditorEngine::PlayWorld to the PIE
