Fixing CIS

diff --git a/Engine/Source/Runtime/Engine/Private/Materials/MaterialShared.cpp b/Engine/Source/Runtime/Engine/Private/Materials/MaterialShared.cpp
--- a/Engine/Source/Runtime/Engine/Private/Materials/MaterialShared.cpp
+++ b/Engine/Source/Runtime/Engine/Private/Materials/MaterialShared.cpp
@@ -133,7 +133,7 @@ bool ReloadMaterialResource(
 	Tmp.SerializeInlineShaderMap(Ar);
 	if (Tmp.GetGameThreadShaderMap())
 	{
-		Tmp.GetGameThreadShaderMap()->RegisterSerializedShaders();
+		Tmp.GetGameThreadShaderMap()->RegisterSerializedShaders(false);
 		return true;
 	}
 	UE_LOG(LogMaterial, Warning, TEXT("Failed to reload material resources for package %s (file name: %s)."), *PackageName, *Filename);
