Resize Temp instead of Out.

diff --git a/Engine/Source/Runtime/Engine/Private/PhysicsEngine/PhysicsInterfacePhysX.cpp b/Engine/Source/Runtime/Engine/Private/PhysicsEngine/PhysicsInterfacePhysX.cpp
--- a/Engine/Source/Runtime/Engine/Private/PhysicsEngine/PhysicsInterfacePhysX.cpp
+++ b/Engine/Source/Runtime/Engine/Private/PhysicsEngine/PhysicsInterfacePhysX.cpp
@@ -779,7 +779,7 @@ int32 GetAllShapesInternal_AssumedLocked(const FPhysicsActorHandle_PhysX& InActo
 	if (InActorHandle.AsyncActor && bCollectAsync)
 	{
 		const int32 NumAsyncShapes = InActorHandle.AsyncActor->getNbShapes();
-		OutShapes.AddUninitialized(NumAsyncShapes);
+		TempShapes.AddUninitialized(NumAsyncShapes);
 		InActorHandle.AsyncActor->getShapes(TempShapes.GetData() + NumSyncShapes, NumAsyncShapes);
 	}
 
