Adding missing counter.

diff --git a/Engine/Source/Runtime/Engine/Private/LevelTick.cpp b/Engine/Source/Runtime/Engine/Private/LevelTick.cpp
--- a/Engine/Source/Runtime/Engine/Private/LevelTick.cpp
+++ b/Engine/Source/Runtime/Engine/Private/LevelTick.cpp
@@ -1463,6 +1463,7 @@ void UWorld::Tick( ELevelTick TickType, float DeltaSeconds )
 			EnsureCollisionTreeIsBuilt();
 			bInTick = true;
 			{
+				SCOPE_CYCLE_COUNTER(STAT_TG_StartPhysics);
 				SCOPE_TIME_GUARD_MS(TEXT("UWorld::Tick - TG_StartPhysics"), 10);
 				RunTickGroup(TG_StartPhysics); 
 			}
