CL 4399650 copied to Release-4.21

diff --git a/Engine/Source/Runtime/Engine/Private/Slate/SceneViewport.cpp b/Engine/Source/Runtime/Engine/Private/Slate/SceneViewport.cpp
--- a/Engine/Source/Runtime/Engine/Private/Slate/SceneViewport.cpp
+++ b/Engine/Source/Runtime/Engine/Private/Slate/SceneViewport.cpp
@@ -739,7 +739,7 @@ FReply FSceneViewport::OnTouchStarted( const FGeometry& MyGeometry, const FPoint
 			const bool bTemporaryCapture = ViewportClient->CaptureMouseOnClick() == EMouseCaptureMode::CaptureDuringMouseDown;
 			if (bTemporaryCapture)
 			{
-				CurrentReplyState = AcquireFocusAndCapture(TouchPosition.IntPoint());
+				CurrentReplyState = AcquireFocusAndCapture(FIntPoint(TouchEvent.GetScreenSpacePosition().X, TouchEvent.GetScreenSpacePosition().Y));
 			}
 		}
 		else
