Use correct enable flag.

diff --git a/Engine/Source/Runtime/Engine/Private/PhysicsEngine/ConstraintInstance.cpp b/Engine/Source/Runtime/Engine/Private/PhysicsEngine/ConstraintInstance.cpp
--- a/Engine/Source/Runtime/Engine/Private/PhysicsEngine/ConstraintInstance.cpp
+++ b/Engine/Source/Runtime/Engine/Private/PhysicsEngine/ConstraintInstance.cpp
@@ -385,7 +385,7 @@ void FConstraintProfileProperties::UpdateConstraintFlags_AssumesLocked(const FPh
 	FPhysicsInterface::SetCanVisualize(InConstraintRef, true);
 #endif
 
-	FPhysicsInterface::SetCollisionEnabled(InConstraintRef, bDisableCollision);
+	FPhysicsInterface::SetCollisionEnabled(InConstraintRef, !bDisableCollision);
 	FPhysicsInterface::SetProjectionEnabled_AssumesLocked(InConstraintRef, bEnableProjection, ProjectionLinearTolerance, ProjectionAngularTolerance);
 	FPhysicsInterface::SetParentDominates_AssumesLocked(InConstraintRef, bParentDominates);
 }
