Allow for non uniform scaling in some of the capsule functions.

diff --git a/Engine/Source/Runtime/Engine/Classes/Components/CapsuleComponent.h b/Engine/Source/Runtime/Engine/Classes/Components/CapsuleComponent.h
--- a/Engine/Source/Runtime/Engine/Classes/Components/CapsuleComponent.h
+++ b/Engine/Source/Runtime/Engine/Classes/Components/CapsuleComponent.h
@@ -201,31 +201,30 @@ FORCEINLINE void UCapsuleComponent::SetCapsuleHalfHeight(float HalfHeight, bool
 
 FORCEINLINE float UCapsuleComponent::GetScaledCapsuleRadius() const
 {
-	return CapsuleRadius * GetShapeScale();
+	const FVector& ComponentScale = GetComponentTransform().GetScale3D();
+	return CapsuleRadius * (ComponentScale.X < ComponentScale.Y ? ComponentScale.X : ComponentScale.Y);
 }
 
 FORCEINLINE float UCapsuleComponent::GetScaledCapsuleHalfHeight() const
 {
-	return CapsuleHalfHeight * GetShapeScale();
+	return CapsuleHalfHeight * GetComponentTransform().GetScale3D().Z;
 }
 
 FORCEINLINE float UCapsuleComponent::GetScaledCapsuleHalfHeight_WithoutHemisphere() const
 {
-	return (CapsuleHalfHeight - CapsuleRadius) * GetShapeScale();
+	return GetScaledCapsuleHalfHeight() - GetScaledCapsuleRadius();
 }
 
 FORCEINLINE void UCapsuleComponent::GetScaledCapsuleSize(float& OutRadius, float& OutHalfHeight) const
 {
-	const float Scale = GetShapeScale();
-	OutRadius = CapsuleRadius * Scale;
-	OutHalfHeight = CapsuleHalfHeight * Scale;
+	OutRadius = GetScaledCapsuleRadius();
+	OutHalfHeight = GetScaledCapsuleHalfHeight();
 }
 
 FORCEINLINE void UCapsuleComponent::GetScaledCapsuleSize_WithoutHemisphere(float& OutRadius, float& OutHalfHeightWithoutHemisphere) const
 {
-	const float Scale = GetShapeScale();
-	OutRadius = CapsuleRadius * Scale;
-	OutHalfHeightWithoutHemisphere = (CapsuleHalfHeight - CapsuleRadius) * Scale;
+	OutRadius = GetScaledCapsuleRadius();
+	OutHalfHeightWithoutHemisphere = GetScaledCapsuleHalfHeight_WithoutHemisphere();
 }
 
 
diff --git a/Engine/Source/Runtime/Engine/Private/Components/CapsuleComponent.cpp b/Engine/Source/Runtime/Engine/Private/Components/CapsuleComponent.cpp
--- a/Engine/Source/Runtime/Engine/Private/Components/CapsuleComponent.cpp
+++ b/Engine/Source/Runtime/Engine/Private/Components/CapsuleComponent.cpp
@@ -103,11 +103,11 @@ FBoxSphereBounds UCapsuleComponent::CalcBounds(const FTransform& LocalToWorld) c
 
 void UCapsuleComponent::CalcBoundingCylinder(float& CylinderRadius, float& CylinderHalfHeight) const 
 {
-	const float Scale = GetComponentTransform().GetMaximumAxisScale();
+	const FVector& ComponentScale = GetComponentTransform().GetScale3D();
 	const float CapsuleEndCapCenter = FMath::Max(CapsuleHalfHeight - CapsuleRadius, 0.f);
-	const FVector ZAxis = GetComponentTransform().TransformVectorNoScale(FVector(0.f, 0.f, CapsuleEndCapCenter * Scale));
+	const FVector ZAxis = GetComponentTransform().TransformVectorNoScale(FVector(0.f, 0.f, CapsuleEndCapCenter * ComponentScale.Z));
 	
-	const float ScaledRadius = CapsuleRadius * Scale;
+	const float ScaledRadius = CapsuleRadius * (ComponentScale.X > ComponentScale.Y ? ComponentScale.X : ComponentScale.Y);
 	
 	CylinderRadius		= ScaledRadius + FMath::Sqrt(FMath::Square(ZAxis.X) + FMath::Square(ZAxis.Y));
 	CylinderHalfHeight	= ScaledRadius + ZAxis.Z;
@@ -233,9 +233,8 @@ bool UCapsuleComponent::IsZeroExtent() const
 
 FCollisionShape UCapsuleComponent::GetCollisionShape(float Inflation) const
 {
-	const float ShapeScale = GetShapeScale();
-	const float Radius = FMath::Max(0.f, (CapsuleRadius * ShapeScale) + Inflation);
-	const float HalfHeight = FMath::Max(0.f, (CapsuleHalfHeight * ShapeScale) + Inflation);
+	const float Radius = FMath::Max(0.f, GetScaledCapsuleRadius() + Inflation);
+	const float HalfHeight = FMath::Max(0.f, GetScaledCapsuleHalfHeight() + Inflation);
 	return FCollisionShape::MakeCapsule(Radius, HalfHeight);
 }
 
