Account for the case where skeletal mesh doesn't exsit in the component body setup.

diff --git a/Engine/Source/Runtime/Engine/Private/SkeletalMeshComponentPhysics.cpp b/Engine/Source/Runtime/Engine/Private/SkeletalMeshComponentPhysics.cpp
--- a/Engine/Source/Runtime/Engine/Private/SkeletalMeshComponentPhysics.cpp
+++ b/Engine/Source/Runtime/Engine/Private/SkeletalMeshComponentPhysics.cpp
@@ -97,19 +97,19 @@ FString FSkeletalMeshComponentEndPhysicsTickFunction::DiagnosticMessage()
 
 void USkeletalMeshComponent::CreateBodySetup()
 {
-	check(SkeletalMesh);
-
 	if (BodySetup == NULL)
 	{
 		BodySetup = NewObject<UBodySetup>(this);
 	}
 
-	UBodySetup* OriginalBodySetup = SkeletalMesh->GetBodySetup();
-	
-	BodySetup->CopyBodyPropertiesFrom(OriginalBodySetup);
-	BodySetup->CollisionTraceFlag = CTF_UseComplexAsSimple;
+	if (SkeletalMesh)
+	{
+		UBodySetup* OriginalBodySetup = SkeletalMesh->GetBodySetup();
+		BodySetup->CopyBodyPropertiesFrom(OriginalBodySetup);
+		BodySetup->CookedFormatDataOverride = &OriginalBodySetup->CookedFormatData;
+	}
 
-	BodySetup->CookedFormatDataOverride = &OriginalBodySetup->CookedFormatData;
+	BodySetup->CollisionTraceFlag = CTF_UseComplexAsSimple;
 
 	//need to recreate meshes
 	BodySetup->ClearPhysicsMeshes();
