Fix gpu skin cache with multiple cloth sections.

diff --git a/Engine/Shaders/Private/GpuSkinCacheComputeShader.usf b/Engine/Shaders/Private/GpuSkinCacheComputeShader.usf
--- a/Engine/Shaders/Private/GpuSkinCacheComputeShader.usf
+++ b/Engine/Shaders/Private/GpuSkinCacheComputeShader.usf
@@ -49,6 +49,7 @@ void GetMorphBuffer(uint VertexIndex, out float3 DeltaPosition, out half3 DeltaT
 #if GPUSKIN_APEX_CLOTH
 Buffer<float4> ClothBuffer;
 Buffer<float2> ClothPositionsAndNormalsBuffer;
+uint ClothBufferOffset;
 float ClothBlendWeight;
 float4x4 ClothLocalToWorld;
 float4x4 ClothWorldToLocal;
@@ -216,9 +217,17 @@ void SkinCacheUpdateBatchCS(uint3 GroupID : SV_GroupID,
 	uint WeightOffset = InputWeightStart + (VertexIndex * (InputWeightStride/4));
 
 #if GPUSKIN_MORPH_BLEND
+#if GPUSKIN_APEX_CLOTH
+	FVertexUnpacked	Unpacked = UnpackedVertex(MorphBufferOffset + VertexIndex, ClothBufferOffset + VertexIndex, InputOffset, WeightOffset);
+#else
 	FVertexUnpacked	Unpacked = UnpackedVertex(MorphBufferOffset + VertexIndex, VertexIndex, InputOffset, WeightOffset);
+#endif
+#else
+#if GPUSKIN_APEX_CLOTH
+	FVertexUnpacked	Unpacked = UnpackedVertex(VertexIndex, ClothBufferOffset + VertexIndex, InputOffset, WeightOffset);
 #else
 	FVertexUnpacked	Unpacked = UnpackedVertex(VertexIndex, VertexIndex, InputOffset, WeightOffset);
+#endif
 #endif
 
 	// Perform the skinning
diff --git a/Engine/Source/Runtime/Engine/Private/GPUSkinCache.cpp b/Engine/Source/Runtime/Engine/Private/GPUSkinCache.cpp
--- a/Engine/Source/Runtime/Engine/Private/GPUSkinCache.cpp
+++ b/Engine/Source/Runtime/Engine/Private/GPUSkinCache.cpp
@@ -205,6 +205,7 @@ public:
 		uint32 MorphBufferOffset = 0;
 
         // cloth input
+		uint32 ClothBufferOffset = 0;
         float ClothBlendWeight = 0.0f;
 
         FMatrix ClothLocalToWorld = FMatrix::Identity;
@@ -298,6 +299,10 @@ public:
 
 			Data.MorphBufferOffset = Section->BaseVertexIndex;
 		}
+		if (ClothVertexBuffer)
+		{
+			Data.ClothBufferOffset = (ClothVertexBuffer->GetClothIndexMapping()[SectionIndex] & 0xFFFFFFFF);
+		}
 
 		//INC_DWORD_STAT(STAT_GPUSkinCache_TotalNumChunks);
 
@@ -386,6 +391,7 @@ public:
 
 		ClothBuffer.Bind(Initializer.ParameterMap, TEXT("ClothBuffer"));
 		ClothPositionsAndNormalsBuffer.Bind(Initializer.ParameterMap, TEXT("ClothPositionsAndNormalsBuffer"));
+		ClothBufferOffset.Bind(Initializer.ParameterMap, TEXT("ClothBufferOffset"));
 		ClothBlendWeight.Bind(Initializer.ParameterMap, TEXT("ClothBlendWeight"));
 		ClothLocalToWorld.Bind(Initializer.ParameterMap, TEXT("ClothLocalToWorld"));
 		ClothWorldToLocal.Bind(Initializer.ParameterMap, TEXT("ClothWorldToLocal"));
@@ -428,6 +434,7 @@ public:
 		{
 			SetSRVParameter(RHICmdList, ShaderRHI, ClothBuffer, Entry->ClothBuffer);
 			SetSRVParameter(RHICmdList, ShaderRHI, ClothPositionsAndNormalsBuffer, Entry->ClothPositionsAndNormalsBuffer);
+			SetShaderValue(RHICmdList, ShaderRHI, ClothBufferOffset, DispatchData.ClothBufferOffset);
 			SetShaderValue(RHICmdList, ShaderRHI, ClothBlendWeight, DispatchData.ClothBlendWeight);
 			SetShaderValue(RHICmdList, ShaderRHI, ClothLocalToWorld, DispatchData.ClothLocalToWorld);
 			SetShaderValue(RHICmdList, ShaderRHI, ClothWorldToLocal, DispatchData.ClothWorldToLocal);
@@ -454,7 +461,7 @@ public:
 
 		Ar << InputWeightStart << InputWeightStride << InputWeightStream;
 
-        Ar << ClothBuffer << ClothPositionsAndNormalsBuffer << ClothBlendWeight << ClothLocalToWorld << ClothWorldToLocal;
+        Ar << ClothBuffer << ClothPositionsAndNormalsBuffer << ClothBufferOffset << ClothBlendWeight << ClothLocalToWorld << ClothWorldToLocal;
 
 		//Ar << DebugParameter;
 		return bShaderHasOutdatedParameters;
@@ -486,6 +493,7 @@ private:
 
 	FShaderResourceParameter ClothBuffer;
 	FShaderResourceParameter ClothPositionsAndNormalsBuffer;
+	FShaderParameter ClothBufferOffset;
 	FShaderParameter ClothBlendWeight;
 	FShaderParameter ClothLocalToWorld;
 	FShaderParameter ClothWorldToLocal;
@@ -1079,6 +1087,8 @@ void FGPUSkinCache::ProcessEntry(FRHICommandListImmediate& RHICmdList, FGPUBaseS
         // Copy the vertices into the buffer.
         checkSlow(Stride*VertexAndNormalData.GetNumVertices() == sizeof(FClothSimulEntry) * SimData->Positions.Num());
         check(sizeof(FClothSimulEntry) == 6 * sizeof(float));
+		
+		InOutEntry->DispatchData[Section].ClothBufferOffset = (ClothVertexBuffer->GetClothIndexMapping()[Section] & 0xFFFFFFFF);
 
         for (int32 Index = 0;Index < SimData->Positions.Num();Index++)
         {
diff --git a/Engine/Source/Runtime/Engine/Private/SkeletalMeshLODModel.cpp b/Engine/Source/Runtime/Engine/Private/SkeletalMeshLODModel.cpp
--- a/Engine/Source/Runtime/Engine/Private/SkeletalMeshLODModel.cpp
+++ b/Engine/Source/Runtime/Engine/Private/SkeletalMeshLODModel.cpp
@@ -780,6 +780,10 @@ void FSkeletalMeshLODModel::GetClothMappingData(TArray<FMeshToMeshVertData>& Map
 			OutClothIndexMapping.Add(KeyValue);
 			MappingData += Section.ClothMappingData;
 		}
+		else
+		{
+			OutClothIndexMapping.Add(0);
+		}
 	}
 }
 
