[NO NON-PII COMMIT MESSAGE]

diff --git a/Engine/Source/Runtime/Engine/Private/AssetManager.cpp b/Engine/Source/Runtime/Engine/Private/AssetManager.cpp
--- a/Engine/Source/Runtime/Engine/Private/AssetManager.cpp
+++ b/Engine/Source/Runtime/Engine/Private/AssetManager.cpp
@@ -1488,6 +1488,7 @@ TSharedPtr<FStreamableHandle> UAssetManager::LoadAssetList(const TArray<FSoftObj
 	if (bShouldUseSynchronousLoad && MissingChunks.Num() == 0)
 	{
 		NewHandle = StreamableManager.RequestSyncLoad(AssetList, false, DebugName);
+		FStreamableHandle::ExecuteDelegate(DelegateToCall);
 	}
 	else
 	{
