PR #5073: Implement LoadGameFromMemory function. (Contributed by [REDACTED_PERSON])

diff --git a/Engine/Source/Runtime/Engine/Classes/Kismet/GameplayStatics.h b/Engine/Source/Runtime/Engine/Classes/Kismet/GameplayStatics.h
--- a/Engine/Source/Runtime/Engine/Classes/Kismet/GameplayStatics.h
+++ b/Engine/Source/Runtime/Engine/Classes/Kismet/GameplayStatics.h
@@ -835,6 +835,13 @@ public:
 	UFUNCTION(BlueprintCallable, Category="Game")
 	static bool DoesSaveGameExist(const FString& SlotName, const int32 UserIndex);
 
+	/**
+	 *	Load the contents from a given array of bytes.
+	 *	@param InSaveData		The array containing the serialized USaveGame data to load
+	 *	@return SaveGameObject	Object containing loaded game state (NULL if load fails)
+	 */
+	static USaveGame* LoadGameFromMemory(const TArray<uint8>& InSaveData);
+
 	/** 
 	 *	Load the contents from a given slot.
 	 *	@param SlotName			Name of the save game slot to load from.
diff --git a/Engine/Source/Runtime/Engine/Private/GameplayStatics.cpp b/Engine/Source/Runtime/Engine/Private/GameplayStatics.cpp
--- a/Engine/Source/Runtime/Engine/Private/GameplayStatics.cpp
+++ b/Engine/Source/Runtime/Engine/Private/GameplayStatics.cpp
@@ -1790,38 +1790,45 @@ bool UGameplayStatics::DeleteGameInSlot(const FString& SlotName, const int32 Use
 
 USaveGame* UGameplayStatics::LoadGameFromSlot(const FString& SlotName, const int32 UserIndex)
 {
-	USaveGame* OutSaveGameObject = NULL;
-
 	ISaveGameSystem* SaveSystem = IPlatformFeaturesModule::Get().GetSaveGameSystem();
 	// If we have a save system and a valid name..
-	if(SaveSystem && (SlotName.Len() > 0))
+	if (SaveSystem && (SlotName.Len() > 0))
 	{
 		// Load raw data from slot
 		TArray<uint8> ObjectBytes;
 		bool bSuccess = SaveSystem->LoadGame(false, *SlotName, UserIndex, ObjectBytes);
-		if(bSuccess)
+		if (bSuccess)
 		{
-			FMemoryReader MemoryReader(ObjectBytes, true);
+			return LoadGameFromMemory(ObjectBytes);
+		}
+	}
 
-			FSaveGameHeader SaveHeader;
-			SaveHeader.Read(MemoryReader);
+	return nullptr;
+}
 
-			// Try and find it, and failing that, load it
-			UClass* SaveGameClass = FindObject<UClass>(ANY_PACKAGE, *SaveHeader.SaveGameClassName);
-			if(SaveGameClass == NULL)
-			{
-				SaveGameClass = LoadObject<UClass>(NULL, *SaveHeader.SaveGameClassName);
-			}
+USaveGame* UGameplayStatics::LoadGameFromMemory(const TArray<uint8>& InSaveData)
+{
+	USaveGame* OutSaveGameObject = nullptr;
 
-			// If we have a class, try and load it.
-			if(SaveGameClass != NULL)
-			{
-				OutSaveGameObject = NewObject<USaveGame>(GetTransientPackage(), SaveGameClass);
+	FMemoryReader MemoryReader(InSaveData, true);
 
-				FObjectAndNameAsStringProxyArchive Ar(MemoryReader, true);
-				OutSaveGameObject->Serialize(Ar);
-			}
-		}
+	FSaveGameHeader SaveHeader;
+	SaveHeader.Read(MemoryReader);
+
+	// Try and find it, and failing that, load it
+	UClass* SaveGameClass = FindObject<UClass>(ANY_PACKAGE, *SaveHeader.SaveGameClassName);
+	if (SaveGameClass == nullptr)
+	{
+		SaveGameClass = LoadObject<UClass>(nullptr, *SaveHeader.SaveGameClassName);
+	}
+
+	// If we have a class, try and load it.
+	if (SaveGameClass != nullptr)
+	{
+		OutSaveGameObject = NewObject<USaveGame>(GetTransientPackage(), SaveGameClass);
+
+		FObjectAndNameAsStringProxyArchive Ar(MemoryReader, true);
+		OutSaveGameObject->Serialize(Ar);
 	}
 
 	return OutSaveGameObject;
