Integrate from Dev-Editor branch

diff --git a/Engine/Source/Runtime/Engine/Classes/Components/HierarchicalInstancedStaticMeshComponent.h b/Engine/Source/Runtime/Engine/Classes/Components/HierarchicalInstancedStaticMeshComponent.h
--- a/Engine/Source/Runtime/Engine/Classes/Components/HierarchicalInstancedStaticMeshComponent.h
+++ b/Engine/Source/Runtime/Engine/Classes/Components/HierarchicalInstancedStaticMeshComponent.h
@@ -243,6 +243,8 @@ public:
 	/** Will apply current density scaling, if enabled **/
 	void UpdateDensityScaling();
 
+	virtual void PropagateLightingScenarioChange() override;
+
 protected:
 	void BuildTree();
 	void BuildTreeAsync();
diff --git a/Engine/Source/Runtime/Engine/Private/HierarchicalInstancedStaticMesh.cpp b/Engine/Source/Runtime/Engine/Private/HierarchicalInstancedStaticMesh.cpp
--- a/Engine/Source/Runtime/Engine/Private/HierarchicalInstancedStaticMesh.cpp
+++ b/Engine/Source/Runtime/Engine/Private/HierarchicalInstancedStaticMesh.cpp
@@ -30,6 +30,7 @@
 #include "SceneManagement.h"
 #include "HAL/LowLevelMemTracker.h"
 #include "UObject/ReleaseObjectVersion.h"
+#include "ComponentRecreateRenderStateContext.h"
 
 static TAutoConsoleVariable<int32> CVarFoliageSplitFactor(
 	TEXT("foliage.SplitFactor"),
@@ -2129,7 +2130,7 @@ bool UHierarchicalInstancedStaticMeshComponent::UpdateInstanceTransform(int32 In
 	// if we are only updating rotation/scale we update the instance directly in the cluster tree
 	const bool bIsOmittedInstance = (RenderIndex == INDEX_NONE);
 	const bool bIsBuiltInstance = !bIsOmittedInstance && RenderIndex < NumBuiltRenderInstances;
-	const bool bDoInPlaceUpdate = bIsBuiltInstance && NewLocalLocation.Equals(OldTransform.GetOrigin());
+	const bool bDoInPlaceUpdate = bIsBuiltInstance && NewLocalLocation.Equals(OldTransform.GetOrigin()) && (PerInstanceRenderData.IsValid() && PerInstanceRenderData->InstanceBuffer.RequireCPUAccess);
 
 	bool Result = Super::UpdateInstanceTransform(InstanceIndex, NewInstanceTransform, bWorldSpace, bMarkRenderStateDirty, bTeleport);
 	
@@ -2150,7 +2151,10 @@ bool UHierarchicalInstancedStaticMeshComponent::UpdateInstanceTransform(int32 In
 			{
 				BuiltInstanceBounds += NewInstanceBounds;
 
-				MarkRenderStateDirty();
+				if (bMarkRenderStateDirty)
+				{
+					MarkRenderStateDirty();
+				}
 			}
 		}
 		else
@@ -2704,6 +2708,31 @@ void UHierarchicalInstancedStaticMeshComponent::BuildTreeAsync()
 	}
 }
 
+void UHierarchicalInstancedStaticMeshComponent::PropagateLightingScenarioChange()
+{
+	if (GIsEditor && PerInstanceRenderData.IsValid())
+	{
+		FComponentRecreateRenderStateContext Context(this);
+
+		const FMeshMapBuildData* MeshMapBuildData = nullptr;
+		if (LODData.Num() > 0)
+		{
+			MeshMapBuildData = GetMeshMapBuildData(LODData[0]);
+		}
+
+		if (MeshMapBuildData != nullptr)
+		{
+			for (int32 InstanceIndex = 0; InstanceIndex < PerInstanceSMData.Num(); ++InstanceIndex)
+			{
+				int32 RenderIndex = InstanceReorderTable.IsValidIndex(InstanceIndex) ? InstanceReorderTable[InstanceIndex] : InstanceIndex;
+
+				InstanceUpdateCmdBuffer.SetLightMapData(RenderIndex, MeshMapBuildData->PerInstanceLightmapData[InstanceIndex].LightmapUVBias);
+				InstanceUpdateCmdBuffer.SetShadowMapData(RenderIndex, MeshMapBuildData->PerInstanceLightmapData[InstanceIndex].ShadowmapUVBias);
+			}
+		}
+	}
+}
+
 void UHierarchicalInstancedStaticMeshComponent::SetPerInstanceLightMapAndEditorData(FStaticMeshInstanceData& PerInstanceData, const TArray<TRefCountPtr<HHitProxy>>& HitProxies)
 {
 	int32 NumInstances = PerInstanceData.GetNumInstances();
@@ -2732,9 +2761,9 @@ void UHierarchicalInstancedStaticMeshComponent::SetPerInstanceLightMapAndEditorD
 			{
 				LightmapUVBias = MeshMapBuildData->PerInstanceLightmapData[Index].LightmapUVBias;
 				ShadowmapUVBias = MeshMapBuildData->PerInstanceLightmapData[Index].ShadowmapUVBias;
-			}
 
-			PerInstanceData.SetInstanceLightMapData(RenderIndex, LightmapUVBias, ShadowmapUVBias);
+				PerInstanceData.SetInstanceLightMapData(RenderIndex, LightmapUVBias, ShadowmapUVBias);
+			}
 
 	#if WITH_EDITOR
 			if (GIsEditor)
diff --git a/Engine/Source/Runtime/Engine/Private/LightMap.cpp b/Engine/Source/Runtime/Engine/Private/LightMap.cpp
--- a/Engine/Source/Runtime/Engine/Private/LightMap.cpp
+++ b/Engine/Source/Runtime/Engine/Private/LightMap.cpp
@@ -283,7 +283,8 @@ struct FLightMapAllocation
 				// Need to create per-LOD instance data to fix that
 				MeshBuildData->PerInstanceLightmapData[InstanceIndex].LightmapUVBias = LightMap->GetCoordinateBias();
 
-				Component->InstanceUpdateCmdBuffer.SetLightMapData(InstanceIndex, MeshBuildData->PerInstanceLightmapData[InstanceIndex].LightmapUVBias);
+				int32 RenderIndex = Component->InstanceReorderTable.IsValidIndex(InstanceIndex) ? Component->InstanceReorderTable[InstanceIndex] : InstanceIndex;
+				Component->InstanceUpdateCmdBuffer.SetLightMapData(RenderIndex, MeshBuildData->PerInstanceLightmapData[InstanceIndex].LightmapUVBias);
 				Component->MarkRenderStateDirty();
 			}
 		}
diff --git a/Engine/Source/Runtime/Engine/Private/ShadowMap.cpp b/Engine/Source/Runtime/Engine/Private/ShadowMap.cpp
--- a/Engine/Source/Runtime/Engine/Private/ShadowMap.cpp
+++ b/Engine/Source/Runtime/Engine/Private/ShadowMap.cpp
@@ -109,8 +109,8 @@ struct FShadowMapAllocation
 				// TODO: We currently only support one LOD of static lighting in foliage
 				// Need to create per-LOD instance data to fix that
 				MeshBuildData->PerInstanceLightmapData[InstanceIndex].ShadowmapUVBias = ShadowMap->GetCoordinateBias();
-
-				Component->InstanceUpdateCmdBuffer.SetShadowMapData(InstanceIndex, MeshBuildData->PerInstanceLightmapData[InstanceIndex].ShadowmapUVBias);
+				int32 RenderIndex = Component->InstanceReorderTable.IsValidIndex(InstanceIndex) ? Component->InstanceReorderTable[InstanceIndex] : InstanceIndex;
+				Component->InstanceUpdateCmdBuffer.SetShadowMapData(RenderIndex, MeshBuildData->PerInstanceLightmapData[InstanceIndex].ShadowmapUVBias);
 				Component->MarkRenderStateDirty();
 			}
 		}
