[NO NON-PII COMMIT MESSAGE]

diff --git a/Engine/Source/Runtime/Engine/Private/InstancedStaticMesh.cpp b/Engine/Source/Runtime/Engine/Private/InstancedStaticMesh.cpp
--- a/Engine/Source/Runtime/Engine/Private/InstancedStaticMesh.cpp
+++ b/Engine/Source/Runtime/Engine/Private/InstancedStaticMesh.cpp
@@ -1733,7 +1733,7 @@ void UInstancedStaticMeshComponent::Serialize(FArchive& Ar)
 	{
 		TArray<FInstancedStaticMeshInstanceData_DEPRECATED> DeprecatedData;
 		DeprecatedData.BulkSerialize(Ar);
-		PerInstanceSMData.Reserve(DeprecatedData.Num());
+		PerInstanceSMData.Reset(DeprecatedData.Num());
 		Algo::Transform(DeprecatedData, PerInstanceSMData, [](const FInstancedStaticMeshInstanceData_DEPRECATED& OldData){ 
 			return FInstancedStaticMeshInstanceData(OldData.Transform);
 		});
