[NO NON-PII COMMIT MESSAGE]

diff --git a/Engine/Source/Runtime/Engine/Private/HierarchicalInstancedStaticMesh.cpp b/Engine/Source/Runtime/Engine/Private/HierarchicalInstancedStaticMesh.cpp
--- a/Engine/Source/Runtime/Engine/Private/HierarchicalInstancedStaticMesh.cpp
+++ b/Engine/Source/Runtime/Engine/Private/HierarchicalInstancedStaticMesh.cpp
@@ -1566,10 +1566,18 @@ void FHierarchicalStaticMeshSceneProxy::GetDynamicMeshElements(const TArray<cons
 				}
 				InstanceParams.LODPlanesMin[InstanceParams.LODs - 1] = FinalCull - LODRandom;
 				InstanceParams.LODPlanesMax[InstanceParams.LODs - 1] = FinalCull;
+
+				// Added assert guard to track issue UE-53944
+				check(InstanceParams.LODs <= 8);
+				check(RenderData != nullptr);
 			
 				for (int32 LODIndex = 0; LODIndex < InstanceParams.LODs; LODIndex++)
 				{
 					InstanceParams.MinInstancesToSplit[LODIndex] = 2;
+
+					// Added assert guard to track issue UE-53944
+					check(RenderData->LODResources.IsValidIndex(LODIndex));
+
 					int32 NumVerts = RenderData->LODResources[LODIndex].VertexBuffers.StaticMeshVertexBuffer.GetNumVertices();
 					if (NumVerts)
 					{
