PR #4892: HLOD System: Fix crash when static mesh has no valid material (Contributed by [REDACTED_PERSON])
#fix slightly modifier fix

diff --git a/Engine/Source/Runtime/Engine/Private/HLODProxy.cpp b/Engine/Source/Runtime/Engine/Private/HLODProxy.cpp
--- a/Engine/Source/Runtime/Engine/Private/HLODProxy.cpp
+++ b/Engine/Source/Runtime/Engine/Private/HLODProxy.cpp
@@ -326,13 +326,16 @@ FName UHLODProxy::GenerateKeyForActor(const ALODActor* LODActor)
 					for(int32 MaterialIndex = 0; MaterialIndex < NumMaterials; ++MaterialIndex)
 					{
 						UMaterialInterface* MaterialInterface = StaticMeshComponent->GetMaterial(MaterialIndex);
-						CRC = GetCRC(MaterialInterface, CRC);
-
-						TArray<UTexture*> Textures;
-						MaterialInterface->GetUsedTextures(Textures, EMaterialQualityLevel::High, true, ERHIFeatureLevel::SM5, true);
-						for(UTexture* Texture : Textures)
+						if (MaterialInterface)
 						{
-							CRC = GetCRC(Texture, CRC);
+							CRC = GetCRC(MaterialInterface, CRC);
+
+							TArray<UTexture*> Textures;
+							MaterialInterface->GetUsedTextures(Textures, EMaterialQualityLevel::High, true, ERHIFeatureLevel::SM5, true);
+							for (UTexture* Texture : Textures)
+							{
+								CRC = GetCRC(Texture, CRC);
+							}
 						}
 					}
 				}
