[NO NON-PII COMMIT MESSAGE]

diff --git a/Engine/Source/Runtime/Engine/Private/HierarchicalInstancedStaticMesh.cpp b/Engine/Source/Runtime/Engine/Private/HierarchicalInstancedStaticMesh.cpp
--- a/Engine/Source/Runtime/Engine/Private/HierarchicalInstancedStaticMesh.cpp
+++ b/Engine/Source/Runtime/Engine/Private/HierarchicalInstancedStaticMesh.cpp
@@ -2439,6 +2439,7 @@ void UHierarchicalInstancedStaticMeshComponent::BuildTree()
 		ClusterTreePtr = MakeShareable(new TArray<FClusterNode>);
 		NumBuiltInstances = 0;
 		NumBuiltRenderInstances = 0;
+		InstanceCountToRender = 0;
 		InstanceReorderTable.Empty();
 		SortedInstances.Empty();
 
@@ -2707,6 +2708,7 @@ void UHierarchicalInstancedStaticMeshComponent::BuildTreeAsync()
 		ClusterTreePtr = MakeShareable(new TArray<FClusterNode>);
 		NumBuiltInstances = 0;
 		NumBuiltRenderInstances = 0;
+		InstanceCountToRender = 0;
 		InstanceReorderTable.Empty();
 		SortedInstances.Empty();
 		CacheMeshExtendedBounds = FBoxSphereBounds(ForceInitToZero);
