Unshelved from pending changelist '4438353':

Ensure const correctness in USoundWave cook processing.

diff --git a/Engine/Source/Runtime/Engine/Classes/Sound/SoundWave.h b/Engine/Source/Runtime/Engine/Classes/Sound/SoundWave.h
--- a/Engine/Source/Runtime/Engine/Classes/Sound/SoundWave.h
+++ b/Engine/Source/Runtime/Engine/Classes/Sound/SoundWave.h
@@ -341,6 +341,10 @@ public:
 
 #if WITH_EDITORONLY_DATA
 	TMap<FName, uint32> AsyncLoadingDataFormats;
+
+	/** FByteBulkData doesn't currently support readonly access from multiple threads, so we limit access to RawData with a critical section on cook. */
+	FCriticalSection RawDataCriticalSection;
+
 #endif
 
 	/** Resource index to cross reference with buffers */
diff --git a/Engine/Source/Runtime/Engine/Private/Audio.cpp b/Engine/Source/Runtime/Engine/Private/Audio.cpp
--- a/Engine/Source/Runtime/Engine/Private/Audio.cpp
+++ b/Engine/Source/Runtime/Engine/Private/Audio.cpp
@@ -1029,7 +1029,7 @@ struct FExtendedFormatChunk
 //
 //	Figure out the WAVE file layout.
 //
-bool FWaveModInfo::ReadWaveInfo( uint8* WaveData, int32 WaveDataSize, FString* ErrorReason, bool InHeaderDataOnly, void** OutFormatHeader)
+bool FWaveModInfo::ReadWaveInfo( const uint8* WaveData, int32 WaveDataSize, FString* ErrorReason, bool InHeaderDataOnly, void** OutFormatHeader)
 {
 	FRiffFormatChunk* FmtChunk;
 	FExtendedFormatChunk* FmtChunkEx = nullptr;
@@ -1273,7 +1273,7 @@ bool FWaveModInfo::ReadWaveInfo( uint8* WaveData, int32 WaveDataSize, FString* E
 	return( true );
 }
 
-bool FWaveModInfo::ReadWaveHeader(uint8* RawWaveData, int32 Size, int32 Offset )
+bool FWaveModInfo::ReadWaveHeader(const uint8* RawWaveData, int32 Size, int32 Offset )
 {
 	if( Size == 0 )
 	{
diff --git a/Engine/Source/Runtime/Engine/Private/AudioDerivedData.cpp b/Engine/Source/Runtime/Engine/Private/AudioDerivedData.cpp
--- a/Engine/Source/Runtime/Engine/Private/AudioDerivedData.cpp
+++ b/Engine/Source/Runtime/Engine/Private/AudioDerivedData.cpp
@@ -848,16 +848,16 @@ static void CookSimpleWave(USoundWave* SoundWave, FName FormatName, const IAudio
 		return;
 	}
 
+#if WITH_EDITORONLY_DATA
+	FScopeLock ScopeLock(&SoundWave->RawDataCriticalSection);
+#endif
+
 	// check if there is any raw sound data
 	if( SoundWave->RawData.GetBulkDataSize() > 0 )
 	{
-		if (SoundWave->RawData.IsLocked())
-		{
-			UE_LOG(LogAudioDerivedData, Warning, TEXT("In CookSimpleWave: Raw PCM data already being written to. %s "), *SoundWave->GetFullName());
-		}
 
 		// Lock raw wave data.
-		uint8* RawWaveData = ( uint8* )SoundWave->RawData.Lock( LOCK_READ_ONLY );
+		const uint8* RawWaveData = ( const uint8* )SoundWave->RawData.LockReadOnly();
 		bWasLocked = true;
 		int32 RawDataSize = SoundWave->RawData.GetBulkDataSize();
 
@@ -969,10 +969,9 @@ static void CookSurroundWave( USoundWave* SoundWave, FName FormatName, const IAu
 	TArray<TArray<uint8> >	SourceBuffers;
 	TArray<int32>			RequiredChannels;
 
-	if (SoundWave->RawData.IsLocked())
-	{
-		UE_LOG(LogAudioDerivedData, Warning, TEXT("In CookSurroundWave: Raw PCM data already being written to. %s "), *SoundWave->GetFullName());
-	}
+#if WITH_EDITORONLY_DATA
+	FScopeLock ScopeLock(&SoundWave->RawDataCriticalSection);
+#endif
 
 	uint8* RawWaveData = ( uint8* )SoundWave->RawData.Lock( LOCK_READ_ONLY );
 	if (RawWaveData == nullptr)
diff --git a/Engine/Source/Runtime/Engine/Public/Audio.h b/Engine/Source/Runtime/Engine/Public/Audio.h
--- a/Engine/Source/Runtime/Engine/Public/Audio.h
+++ b/Engine/Source/Runtime/Engine/Public/Audio.h
@@ -856,19 +856,19 @@ class FWaveModInfo
 public:
 
 	// Pointers to variables in the in-memory WAVE file.
-	uint32* pSamplesPerSec;
-	uint32* pAvgBytesPerSec;
-	uint16* pBlockAlign;
-	uint16* pBitsPerSample;
-	uint16* pChannels;
-	uint16* pFormatTag;
-
-	uint32* pWaveDataSize;
-	uint32* pMasterSize;
-	uint8*  SampleDataStart;
-	uint8*  SampleDataEnd;
+	const uint32* pSamplesPerSec;
+	const uint32* pAvgBytesPerSec;
+	const uint16* pBlockAlign;
+	const uint16* pBitsPerSample;
+	const uint16* pChannels;
+	const uint16* pFormatTag;
+
+	const uint32* pWaveDataSize;
+	const uint32* pMasterSize;
+	const uint8*  SampleDataStart;
+	const uint8*  SampleDataEnd;
 	uint32  SampleDataSize;
-	uint8*  WaveDataEnd;
+	const uint8*  WaveDataEnd;
 
 	uint32  NewDataSize;
 
@@ -885,12 +885,12 @@ public:
 
 	// Read headers and load all info pointers in WaveModInfo. 
 	// Returns 0 if invalid data encountered.
-	ENGINE_API bool ReadWaveInfo( uint8* WaveData, int32 WaveDataSize, FString* ErrorMessage = NULL, bool InHeaderDataOnly = false, void** OutFormatHeader = NULL );
+	ENGINE_API bool ReadWaveInfo(const uint8* WaveData, int32 WaveDataSize, FString* ErrorMessage = NULL, bool InHeaderDataOnly = false, void** OutFormatHeader = NULL );
 	
 	/**
 	 * Read a wave file header from bulkdata
 	 */
-	ENGINE_API bool ReadWaveHeader(uint8* RawWaveData, int32 Size, int32 Offset );
+	ENGINE_API bool ReadWaveHeader(const uint8* RawWaveData, int32 Size, int32 Offset );
 
 	ENGINE_API void ReportImportFailure() const;
 };
