Fixed crash with VRED importer
- DataTable wasn't counting its changes properly resulting in a failed check.

diff --git a/Engine/Source/Runtime/Engine/Private/CurveTable.cpp b/Engine/Source/Runtime/Engine/Private/CurveTable.cpp
--- a/Engine/Source/Runtime/Engine/Private/CurveTable.cpp
+++ b/Engine/Source/Runtime/Engine/Private/CurveTable.cpp
@@ -27,7 +27,7 @@ namespace
 			{
 				Count = &ScopeCount.Add(Table, 0);
 			}
-			++Count;
+			++(*Count);
 		}
 		~FScopedCurveTableChange()
 		{
diff --git a/Engine/Source/Runtime/Engine/Private/DataTable.cpp b/Engine/Source/Runtime/Engine/Private/DataTable.cpp
--- a/Engine/Source/Runtime/Engine/Private/DataTable.cpp
+++ b/Engine/Source/Runtime/Engine/Private/DataTable.cpp
@@ -39,7 +39,7 @@ namespace
 			{
 				Count = &ScopeCount.Add(Table, 0);
 			}
-			++Count;
+			++(*Count);
 		}
 		~FScopedDataTableChange()
 		{
