Avoid broadcasting events that could result in invalidating the array we're iterating.

diff --git a/Engine/Source/Runtime/Engine/Private/LevelStreaming.cpp b/Engine/Source/Runtime/Engine/Private/LevelStreaming.cpp
--- a/Engine/Source/Runtime/Engine/Private/LevelStreaming.cpp
+++ b/Engine/Source/Runtime/Engine/Private/LevelStreaming.cpp
@@ -1296,18 +1296,25 @@ void ULevelStreaming::BroadcastLevelLoadedStatus(UWorld* PersistentWorld, FName
 	
 void ULevelStreaming::BroadcastLevelVisibleStatus(UWorld* PersistentWorld, FName LevelPackageName, bool bVisible)
 {
+	TArray<ULevelStreaming*, TInlineAllocator<1>> LevelsToBroadcast;
+
 	for (ULevelStreaming* StreamingLevel : PersistentWorld->GetStreamingLevels())
 	{
 		if (StreamingLevel->GetWorldAssetPackageFName() == LevelPackageName)
 		{
-			if (bVisible)
-			{
-				StreamingLevel->OnLevelShown.Broadcast();
-			}
-			else
-			{
-				StreamingLevel->OnLevelHidden.Broadcast();
-			}
+			LevelsToBroadcast.Add(StreamingLevel);
+		}
+	}
+
+	for (ULevelStreaming* StreamingLevel : LevelsToBroadcast)
+	{
+		if (bVisible)
+		{
+			StreamingLevel->OnLevelShown.Broadcast();
+		}
+		else
+		{
+			StreamingLevel->OnLevelHidden.Broadcast();
 		}
 	}
 }
