Allowing material graph nodes to have duplicate names in all cases except creating a new parameter from a constant

diff --git a/Engine/Source/Editor/MaterialEditor/Private/MaterialEditingLibrary.cpp b/Engine/Source/Editor/MaterialEditor/Private/MaterialEditingLibrary.cpp
--- a/Engine/Source/Editor/MaterialEditor/Private/MaterialEditingLibrary.cpp
+++ b/Engine/Source/Editor/MaterialEditor/Private/MaterialEditingLibrary.cpp
@@ -389,7 +389,7 @@ UMaterialExpression* UMaterialEditingLibrary::CreateMaterialExpressionEx(UMateri
 
 		if (NewExpression->HasAParameterName())
 		{
-			NewExpression->ValidateParameterName();
+			NewExpression->ValidateParameterName(false);
 		}
 
 		UMaterialExpressionComponentMask* ComponentMaskExpression = Cast<UMaterialExpressionComponentMask>(NewExpression);
diff --git a/Engine/Source/Editor/MaterialEditor/Private/MaterialEditor.cpp b/Engine/Source/Editor/MaterialEditor/Private/MaterialEditor.cpp
--- a/Engine/Source/Editor/MaterialEditor/Private/MaterialEditor.cpp
+++ b/Engine/Source/Editor/MaterialEditor/Private/MaterialEditor.cpp
@@ -2966,7 +2966,7 @@ void FMaterialEditor::OnPromoteToParameter()
 		if (MaterialNode->MaterialExpression->HasAParameterName())
 		{
 			MaterialNode->MaterialExpression->SetParameterName(TargetPin->PinName);
-			MaterialNode->MaterialExpression->ValidateParameterName();
+			MaterialNode->MaterialExpression->ValidateParameterName(false);
 		}
 	}
 	if (MaterialEditorInstance != nullptr)
diff --git a/Engine/Source/Runtime/Engine/Classes/Materials/MaterialExpression.h b/Engine/Source/Runtime/Engine/Classes/Materials/MaterialExpression.h
--- a/Engine/Source/Runtime/Engine/Classes/Materials/MaterialExpression.h
+++ b/Engine/Source/Runtime/Engine/Classes/Materials/MaterialExpression.h
@@ -409,7 +409,7 @@ class ENGINE_API UMaterialExpression : public UObject
 	* Parameter Name functions, this is requires as multiple class have ParameterName
 	* but are not UMaterialExpressionParameter due to class hierarchy. */
 	virtual bool HasAParameterName() const { return false; }
-	virtual void ValidateParameterName();
+	virtual void ValidateParameterName(const bool bAllowDuplicateName = true);
 
 	virtual FName GetParameterName() const { return NAME_None; }
 	virtual void SetParameterName(const FName& Name) {}
diff --git a/Engine/Source/Runtime/Engine/Classes/Materials/MaterialExpressionCollectionParameter.h b/Engine/Source/Runtime/Engine/Classes/Materials/MaterialExpressionCollectionParameter.h
--- a/Engine/Source/Runtime/Engine/Classes/Materials/MaterialExpressionCollectionParameter.h
+++ b/Engine/Source/Runtime/Engine/Classes/Materials/MaterialExpressionCollectionParameter.h
@@ -47,7 +47,7 @@ class UMaterialExpressionCollectionParameter : public UMaterialExpression
 	virtual bool HasAParameterName() const override { return true; }
 	virtual FName GetParameterName() const override { return ParameterName; }
 	virtual void SetParameterName(const FName& Name) override { ParameterName = Name; }
-	virtual void ValidateParameterName() override {};
+	virtual void ValidateParameterName(const bool bAllowDuplicateName) override {};
 
 	virtual bool MatchesSearchQuery( const TCHAR* SearchQuery ) override;
 #endif
diff --git a/Engine/Source/Runtime/Engine/Classes/Materials/MaterialExpressionFontSampleParameter.h b/Engine/Source/Runtime/Engine/Classes/Materials/MaterialExpressionFontSampleParameter.h
--- a/Engine/Source/Runtime/Engine/Classes/Materials/MaterialExpressionFontSampleParameter.h
+++ b/Engine/Source/Runtime/Engine/Classes/Materials/MaterialExpressionFontSampleParameter.h
@@ -47,7 +47,7 @@ class UMaterialExpressionFontSampleParameter : public UMaterialExpressionFontSam
 	virtual bool HasAParameterName() const override { return true; }
 	virtual FName GetParameterName() const override { return ParameterName; }
 	virtual void SetParameterName(const FName& Name) override { ParameterName = Name; }
-	virtual void ValidateParameterName() override;
+	virtual void ValidateParameterName(const bool bAllowDuplicateName) override;
 #endif
 	//~ End UMaterialExpression Interface
 	
diff --git a/Engine/Source/Runtime/Engine/Classes/Materials/MaterialExpressionParameter.h b/Engine/Source/Runtime/Engine/Classes/Materials/MaterialExpressionParameter.h
--- a/Engine/Source/Runtime/Engine/Classes/Materials/MaterialExpressionParameter.h
+++ b/Engine/Source/Runtime/Engine/Classes/Materials/MaterialExpressionParameter.h
@@ -47,7 +47,7 @@ class UMaterialExpressionParameter : public UMaterialExpression
 	virtual bool HasAParameterName() const override { return true; }
 	virtual FName GetParameterName() const override { return ParameterName; }
 	virtual void SetParameterName(const FName& Name) override { ParameterName = Name; }
-	virtual void ValidateParameterName() override;
+	virtual void ValidateParameterName(const bool bAllowDuplicateName) override;
 #endif
 	//~ End UMaterialExpression Interface
 
diff --git a/Engine/Source/Runtime/Engine/Classes/Materials/MaterialExpressionTextureSampleParameter.h b/Engine/Source/Runtime/Engine/Classes/Materials/MaterialExpressionTextureSampleParameter.h
--- a/Engine/Source/Runtime/Engine/Classes/Materials/MaterialExpressionTextureSampleParameter.h
+++ b/Engine/Source/Runtime/Engine/Classes/Materials/MaterialExpressionTextureSampleParameter.h
@@ -46,7 +46,7 @@ class ENGINE_API UMaterialExpressionTextureSampleParameter : public UMaterialExp
 	virtual bool HasAParameterName() const override { return true; }
 	virtual FName GetParameterName() const override { return ParameterName; }
 	virtual void SetParameterName(const FName& Name) override { ParameterName = Name; }
-	virtual void ValidateParameterName() override;
+	virtual void ValidateParameterName(const bool bAllowDuplicateName) override;
 #endif
 	//~ End UMaterialExpression Interface
 
diff --git a/Engine/Source/Runtime/Engine/Private/Materials/MaterialExpressions.cpp b/Engine/Source/Runtime/Engine/Private/Materials/MaterialExpressions.cpp
--- a/Engine/Source/Runtime/Engine/Private/Materials/MaterialExpressions.cpp
+++ b/Engine/Source/Runtime/Engine/Private/Materials/MaterialExpressions.cpp
@@ -377,16 +377,16 @@ bool CanConnectMaterialValueTypes(uint32 InputType, uint32 OutputType)
 #if WITH_EDITOR
 
 
-void ValidateParameterNameInternal(class UMaterialExpression* ExpressionToValidate, class UMaterial* OwningMaterial)
+void ValidateParameterNameInternal(class UMaterialExpression* ExpressionToValidate, class UMaterial* OwningMaterial, const bool bAllowDuplicateName)
 {
 	if (OwningMaterial != nullptr)
 	{
 		int32 NameIndex = 1;
-		bool FoundValidName = false;
+		bool bFoundValidName = false;
 		FName PotentialName;
 
 		// Find an available unique name
-		while (!FoundValidName)
+		while (!bFoundValidName)
 		{
 			PotentialName = ExpressionToValidate->GetParameterName();
 
@@ -396,27 +396,34 @@ void ValidateParameterNameInternal(class UMaterialExpression* ExpressionToValida
 				PotentialName = UMaterialExpressionParameter::ParameterDefaultName;
 			}
 
-			if (NameIndex != 1)
+			if (!bAllowDuplicateName)
 			{
-				PotentialName.SetNumber(NameIndex);
-			}
+				if (NameIndex != 1)
+				{
+					PotentialName.SetNumber(NameIndex);
+				}
 
-			FoundValidName = true;
+				bFoundValidName = true;
 
-			for (UMaterialExpression* Expression : OwningMaterial->Expressions)
-			{
-				if (Expression != nullptr && Expression->HasAParameterName())
+				for (UMaterialExpression* Expression : OwningMaterial->Expressions)
 				{
-					// Name are unique per class type
-					if (Expression != ExpressionToValidate && Expression->GetClass() == ExpressionToValidate->GetClass() && Expression->GetParameterName() == PotentialName)
+					if (Expression != nullptr && Expression->HasAParameterName())
 					{
-						FoundValidName = false;
-						break;
+						// Name are unique per class type
+						if (Expression != ExpressionToValidate && Expression->GetClass() == ExpressionToValidate->GetClass() && Expression->GetParameterName() == PotentialName)
+						{
+							bFoundValidName = false;
+							break;
+						}
 					}
 				}
-			}
 
-			++NameIndex;
+				++NameIndex;
+			}
+			else
+			{
+				bFoundValidName = true;
+			}
 		}
 
 		ExpressionToValidate->SetParameterName(PotentialName);
@@ -1208,7 +1215,7 @@ void UMaterialExpression::SetEditableName(const FString& NewName)
 	check(false);
 }
 
-void UMaterialExpression::ValidateParameterName()
+void UMaterialExpression::ValidateParameterName(const bool bAllowDuplicateName)
 {
 	// Incrementing the name is now handled in UMaterialExpressionParameter::ValidateParameterName
 }
@@ -1886,9 +1893,9 @@ void UMaterialExpressionTextureSampleParameter::GetCaption(TArray<FString>& OutC
 	OutCaptions.Add(FString::Printf(TEXT("'%s'"), *ParameterName.ToString()));
 }
 
-void UMaterialExpressionTextureSampleParameter::ValidateParameterName()
+void UMaterialExpressionTextureSampleParameter::ValidateParameterName(const bool bAllowDuplicateName)
 {
-	ValidateParameterNameInternal(this, Material);
+	ValidateParameterNameInternal(this, Material, bAllowDuplicateName);
 }
 
 #endif // WITH_EDITOR
@@ -6073,9 +6080,9 @@ void UMaterialExpressionParameter::GetAllParameterInfo(TArray<FMaterialParameter
 }
 
 #if WITH_EDITOR
-void UMaterialExpressionParameter::ValidateParameterName()
+void UMaterialExpressionParameter::ValidateParameterName(const bool bAllowDuplicateName)
 {
-	ValidateParameterNameInternal(this, Material);
+	ValidateParameterNameInternal(this, Material, bAllowDuplicateName);
 }
 #endif
 
@@ -8887,9 +8894,9 @@ void UMaterialExpressionFontSampleParameter::GetCaption(TArray<FString>& OutCapt
 	OutCaptions.Add(FString::Printf(TEXT("'%s'"), *ParameterName.ToString()));
 }
 
-void UMaterialExpressionFontSampleParameter::ValidateParameterName()
+void UMaterialExpressionFontSampleParameter::ValidateParameterName(const bool bAllowDuplicateName)
 {
-	ValidateParameterNameInternal(this, Material);
+	ValidateParameterNameInternal(this, Material, bAllowDuplicateName);
 }
 #endif // WITH_EDITOR
 
