When removing from world non-incrementally do not allow a pending visibility level block the removal
Per: https://udn.unrealengine.com/questions/457117/crash-during-level-streaming-assertion-failed-pend.html

diff --git a/Engine/Source/Runtime/Engine/Private/World.cpp b/Engine/Source/Runtime/Engine/Private/World.cpp
--- a/Engine/Source/Runtime/Engine/Private/World.cpp
+++ b/Engine/Source/Runtime/Engine/Private/World.cpp
@@ -2358,7 +2358,9 @@ void UWorld::RemoveFromWorld( ULevel* Level, bool bAllowIncrementalRemoval )
 	check(!Level->IsPendingKill());
 	check(!Level->IsUnreachable());
 
-	if ( CurrentLevelPendingVisibility == nullptr && Level->bIsVisible )
+	// To be removed from the world a world must be visible and not pending being made visible (this may be redundent, but for safety)
+	// If the level may be removed incrementally then there must also be no level pending visibility
+	if ( ((CurrentLevelPendingVisibility == nullptr) || (!bAllowIncrementalRemoval && (CurrentLevelPendingVisibility != Level))) && Level->bIsVisible )
 	{
 		// Keep track of timing.
 		double StartTime = FPlatformTime::Seconds();	
