Wrap editor-only logic in UInheritableComponentHandler::PostLoad() in WITH_EDITOR to avoid unnecessary loop iterations on load in a cooked build.

diff --git a/Engine/Source/Runtime/Engine/Private/InheritableComponentHandler.cpp b/Engine/Source/Runtime/Engine/Private/InheritableComponentHandler.cpp
--- a/Engine/Source/Runtime/Engine/Private/InheritableComponentHandler.cpp
+++ b/Engine/Source/Runtime/Engine/Private/InheritableComponentHandler.cpp
@@ -30,14 +30,12 @@ void UInheritableComponentHandler::PostLoad()
 	
 #if WITH_EDITOR
 	if (!GIsDuplicatingClassForReinstancing)
-#endif
 	{
 		for (int32 Index = Records.Num() - 1; Index >= 0; --Index)
 		{
 			FComponentOverrideRecord& Record = Records[Index];
 			if (Record.ComponentTemplate)
 			{
-#if WITH_EDITOR
 				if (GetLinkerCustomVersion(FBlueprintsObjectVersion::GUID) < FBlueprintsObjectVersion::SCSHasComponentTemplateClass)
 				{
 					// Fix up component class on load, if it's not already set.
@@ -66,7 +64,6 @@ void UInheritableComponentHandler::PostLoad()
 						FixComponentTemplateName(Record.ComponentTemplate, ExpectedTemplateName);
 					}
 				}
-#endif
 
 				if (!CastChecked<UActorComponent>(Record.ComponentTemplate->GetArchetype())->IsEditableWhenInherited())
 				{
@@ -76,6 +73,7 @@ void UInheritableComponentHandler::PostLoad()
 			}
 		}
 	}
+#endif
 }
 
 #if WITH_EDITOR
