PR #5157: UE-64803: Expose UWindDirectionalSourceComponent functions for other . (Contributed by [REDACTED_PERSON])

diff --git a/Engine/Source/Runtime/Engine/Classes/Components/WindDirectionalSourceComponent.h b/Engine/Source/Runtime/Engine/Classes/Components/WindDirectionalSourceComponent.h
--- a/Engine/Source/Runtime/Engine/Classes/Components/WindDirectionalSourceComponent.h
+++ b/Engine/Source/Runtime/Engine/Classes/Components/WindDirectionalSourceComponent.h
@@ -19,8 +19,8 @@ class FWindData;
 class FWindSourceSceneProxy;
 
 /** Component that provides a directional wind source. Only affects SpeedTree assets. */
-UCLASS(MinimalAPI,collapsecategories, hidecategories=(Object, Mobility), editinlinenew)
-class UWindDirectionalSourceComponent : public USceneComponent
+UCLASS(collapsecategories, hidecategories=(Object, Mobility), editinlinenew)
+class ENGINE_API UWindDirectionalSourceComponent : public USceneComponent
 {
 	GENERATED_UCLASS_BODY()
 
@@ -77,7 +77,7 @@ public:
 	void SetWindType(EWindSourceType InNewType);
 
 	/** Calculate wind parameters from the data on this component, safe to call on game thread */
-	ENGINE_API bool GetWindParameters(const FVector& EvaluatePosition, FWindData& OutData, float& Weight) const;
+	bool GetWindParameters(const FVector& EvaluatePosition, FWindData& OutData, float& Weight) const;
 
 protected:
 	//~ Begin UActorComponent Interface.
