#ue4 - Fix client authoritative position when client sends a null movement base when walking on a static actor. This case was optimized to send a null actor when the relative location wasn't needed, but the server would disagree about the movement base component and then constantly overwrite the movement base with null and clear the cached floor (only in client auth mode).

diff --git a/Engine/Source/Runtime/Engine/Private/Components/CharacterMovementComponent.cpp b/Engine/Source/Runtime/Engine/Private/Components/CharacterMovementComponent.cpp
--- a/Engine/Source/Runtime/Engine/Private/Components/CharacterMovementComponent.cpp
+++ b/Engine/Source/Runtime/Engine/Private/Components/CharacterMovementComponent.cpp
@@ -8528,6 +8528,14 @@ void UCharacterMovementComponent::ServerMoveHandleClientError(float ClientTimeSt
 		ClientLoc += BaseLocation;
 	}
 
+	// Client may send a null movement base when walking on bases with no relative location (to save bandwidth).
+	// In this case don't check movement base in error conditions, use the server one (which avoids an error based on differing bases). Position will still be validated.
+	if (ClientMovementBase == nullptr && ClientMovementMode == MOVE_Walking)
+	{
+		ClientMovementBase = CharacterOwner->GetBasedMovement().MovementBase;
+		ClientBaseBoneName = CharacterOwner->GetBasedMovement().BoneName;
+	}
+
 	// Compute the client error from the server's position
 	// If client has accumulated a noticeable positional error, correct them.
 	bNetworkLargeClientCorrection = ServerData->bForceClientUpdate;
