Mirror 4513604
UBlueprint::ForceLoadMembers now fully loads all objects referenced by the CDO, rather than just those within the CDO

diff --git a/Engine/Source/Runtime/Engine/Private/Blueprint.cpp b/Engine/Source/Runtime/Engine/Private/Blueprint.cpp
--- a/Engine/Source/Runtime/Engine/Private/Blueprint.cpp
+++ b/Engine/Source/Runtime/Engine/Private/Blueprint.cpp
@@ -1153,7 +1153,7 @@ void UBlueprint::ForceLoadMembers(UObject* InObject)
 {
 	// Collect a list of all things this element owns
 	TArray<UObject*> MemberReferences;
-	FReferenceFinder ComponentCollector(MemberReferences, InObject, false, true, true, true);
+	FReferenceFinder ComponentCollector(MemberReferences, nullptr, false, true, true, true);
 	ComponentCollector.FindReferences(InObject);
 
 	// Iterate over the list, and preload everything so it is valid for refreshing
