Fix for a regression that can trigger an assert in a cooked, nativized build during async load of a non-nativized Blueprint CDO that overrides one or more components inherited from a nativized parent class.

#rnx

diff --git a/Engine/Source/Runtime/Engine/Private/BlueprintGeneratedClass.cpp b/Engine/Source/Runtime/Engine/Private/BlueprintGeneratedClass.cpp
--- a/Engine/Source/Runtime/Engine/Private/BlueprintGeneratedClass.cpp
+++ b/Engine/Source/Runtime/Engine/Private/BlueprintGeneratedClass.cpp
@@ -1177,9 +1177,6 @@ void UBlueprintGeneratedClass::CheckAndApplyComponentTemplateOverrides(UObject*
 										}
 									};
 
-									// Ensure that the ICH has gotten a PostLoad() call - we need to ensure that any cooked data will have been fully processed before proceeding.
-									ICH->ConditionalPostLoad();
-
 									// Serialize cached override data to the instanced subobject that's based on the default subobject from the nativized parent class and owned by the non-nativized child class default object.
 									FNativizedComponentOverrideDataLoader OverrideDataLoader(OverrideData->GetCachedPropertyData(), OverrideData->GetCachedPropertyList());
 									NativizedComponentSubobjectInstance->Serialize(OverrideDataLoader);
