Back out 4513686 as it caused a huge cook duration regression

diff --git a/Engine/Source/Runtime/Engine/Private/Blueprint.cpp b/Engine/Source/Runtime/Engine/Private/Blueprint.cpp
--- a/Engine/Source/Runtime/Engine/Private/Blueprint.cpp
+++ b/Engine/Source/Runtime/Engine/Private/Blueprint.cpp
@@ -1159,7 +1159,7 @@ void UBlueprint::ForceLoadMembers(UObject* InObject)
 {
 	// Collect a list of all things this element owns
 	TArray<UObject*> MemberReferences;
-	FReferenceFinder ComponentCollector(MemberReferences, nullptr, false, true, true, true);
+	FReferenceFinder ComponentCollector(MemberReferences, InObject, false, true, true, true);
 	ComponentCollector.FindReferences(InObject);
 
 	// Iterate over the list, and preload everything so it is valid for refreshing
