[Duplicate CL 4351535] - Culling distance on hierarchichal instanced static meshes is prioritized over LOD swap distance.

diff --git a/Engine/Source/Runtime/Engine/Private/HierarchicalInstancedStaticMesh.cpp b/Engine/Source/Runtime/Engine/Private/HierarchicalInstancedStaticMesh.cpp
--- a/Engine/Source/Runtime/Engine/Private/HierarchicalInstancedStaticMesh.cpp
+++ b/Engine/Source/Runtime/Engine/Private/HierarchicalInstancedStaticMesh.cpp
@@ -1561,8 +1561,8 @@ void FHierarchicalStaticMeshSceneProxy::GetDynamicMeshElements(const TArray<cons
 				for (int32 LODIndex = 1; LODIndex < InstanceParams.LODs; LODIndex++)
 				{
 					float Distance = ComputeBoundsDrawDistance(RenderData->ScreenSize[LODIndex].GetValueForFeatureLevel(View->GetFeatureLevel()), SphereRadius, View->ViewMatrices.GetProjectionMatrix()) * LODScale;
-					InstanceParams.LODPlanesMin[LODIndex - 1] = Distance - LODRandom;
-					InstanceParams.LODPlanesMax[LODIndex - 1] = Distance;
+					InstanceParams.LODPlanesMin[LODIndex - 1] = FMath::Min(FinalCull - LODRandom, Distance - LODRandom);
+					InstanceParams.LODPlanesMax[LODIndex - 1] = FMath::Min(FinalCull, Distance);
 				}
 				InstanceParams.LODPlanesMin[InstanceParams.LODs - 1] = FinalCull - LODRandom;
 				InstanceParams.LODPlanesMax[InstanceParams.LODs - 1] = FinalCull;
@@ -1737,11 +1737,11 @@ void FHierarchicalStaticMeshSceneProxy::GetDynamicMeshElements(const TArray<cons
 						for (int32 LODIndex = 1; LODIndex < NumLODs; LODIndex++)
 						{
 							float Distance = ComputeBoundsDrawDistance(RenderData->ScreenSize[LODIndex].GetValueForFeatureLevel(View->GetFeatureLevel()), SphereRadius, View->ViewMatrices.GetProjectionMatrix()) * LODScale;
-							LODPlanesMin[LODIndex - 1] = Distance - LODRandom;
-							LODPlanesMax[LODIndex - 1] = Distance;
+							LODPlanesMin[LODIndex - 1] = FMath::Min(FinalCull - LODRandom, Distance - LODRandom);
+							LODPlanesMax[LODIndex - 1] = FMath::Min(FinalCull, Distance);
 						}
 						LODPlanesMin[NumLODs - 1] = FinalCull - LODRandom;
-						LODPlanesMax[NumLODs - 1] = FinalCull;				
+						LODPlanesMax[NumLODs - 1] = FinalCull;
 
 						// NOTE: in case of unbuilt we can't really apply the instance scales so the LOD won't be optimal until the build is completed
 
