UE4.21 - Copy from 4404878
Removed extra Vulkan render-pass after tonemapper
Better test if current render-pass is compatible with new one so it can be reused
Batch targets transtion to writeable and skip ending currrent render-pass if there are no actual transitions
Skip implicit transition of backbuffer to readable, backbuffer has to be explicitly transitioned when needed
Skip unsetting current targets on scene end, to avoid splitting Post->UI render-pass
#rnx

diff --git a/Engine/Source/Runtime/Engine/Private/Slate/SceneViewport.cpp b/Engine/Source/Runtime/Engine/Private/Slate/SceneViewport.cpp
--- a/Engine/Source/Runtime/Engine/Private/Slate/SceneViewport.cpp
+++ b/Engine/Source/Runtime/Engine/Private/Slate/SceneViewport.cpp
@@ -1705,8 +1705,13 @@ void FSceneViewport::EndRenderFrame(FRHICommandListImmediate& RHICmdList, bool b
 	}
 	else
 	{
-		// Set the active render target(s) to nothing to release references in the case that the viewport is resized by slate before we draw again
-		SetRenderTarget(RHICmdList,  FTexture2DRHIRef(), FTexture2DRHIRef() );
+		// Workaround: un-setting targets splits Post->UI render-pass, we should avoid this as we don't resize viewport on mobile devices
+		bool bShouldUnsetTargets = !(IsVulkanMobilePlatform(GMaxRHIShaderPlatform) && !IsPCPlatform(GMaxRHIShaderPlatform));
+		if (bShouldUnsetTargets)
+		{
+			// Set the active render target(s) to nothing to release references in the case that the viewport is resized by slate before we draw again
+			SetRenderTarget(RHICmdList,  FTexture2DRHIRef(), FTexture2DRHIRef() );
+		}
 		// Note: this releases our reference but does not release the resource as it is owned by slate (this is intended)
 		RenderTargetTextureRenderThreadRHI.SafeRelease();
 		RenderThreadSlateTexture->SetRHIRef(nullptr, 0, 0);
