UE4 - Help track down SCW crash

diff --git a/Engine/Source/Runtime/Engine/Private/Materials/MeshMaterialShader.cpp b/Engine/Source/Runtime/Engine/Private/Materials/MeshMaterialShader.cpp
--- a/Engine/Source/Runtime/Engine/Private/Materials/MeshMaterialShader.cpp
+++ b/Engine/Source/Runtime/Engine/Private/Materials/MeshMaterialShader.cpp
@@ -107,19 +107,36 @@ void FMeshMaterialShaderType::BeginCompileShaderPipeline(
 	NewJobs.Add(NewPipelineJob);
 }
 
+
+static inline FString GetJobName(const FShaderCompileJob* SingleJob, const FShaderPipelineType* ShaderPipelineType, const FString& InDebugDescription)
+{
+	FString String = SingleJob->Input.GenerateShaderName();
+	if (ShaderPipelineType)
+	{
+		String += FString::Printf(TEXT(" Pipeline '%s'"), ShaderPipelineType->GetName());
+	}
+	if (SingleJob->VFType)
+	{
+		String += FString::Printf(TEXT(" VF '%s'"), SingleJob->VFType->GetName());
+	}
+	String += FString::Printf(TEXT(" Type '%s'"), SingleJob->ShaderType->GetName());
+	String += FString::Printf(TEXT(" '%s' Entry '%s' Permutation %i %s"), *SingleJob->Input.VirtualSourceFilePath, *SingleJob->Input.EntryPointName, SingleJob->PermutationId, *InDebugDescription);
+	return String;
+}
+
 /**
  * Either creates a new instance of this type or returns an equivalent existing shader.
  * @param Material - The material to link the shader with.
  * @param CurrentJob - Compile job that was enqueued by BeginCompileShader.
  */
-FShader* FMeshMaterialShaderType::FinishCompileShader(
+FShader* FMeshMaterialShaderType::FinishCompileShader( 
 	const FUniformExpressionSet& UniformExpressionSet, 
 	const FSHAHash& MaterialShaderMapHash,
 	const FShaderCompileJob& CurrentJob,
 	const FShaderPipelineType* ShaderPipelineType,
 	const FString& InDebugDescription)
 {
-	check(CurrentJob.bSucceeded);
+	checkf(CurrentJob.bSucceeded, TEXT("Failed MeshMaterialType compilation job: %s"), *GetJobName(&CurrentJob, ShaderPipelineType, InDebugDescription));
 	check(CurrentJob.VFType);
 
 	FShaderType* SpecificType = CurrentJob.ShaderType->LimitShaderResourceToThisType() ? CurrentJob.ShaderType : NULL;
