Fixed loss of attachment information when compiling blueprint actor blueprints that have compile errors multiple times

diff --git a/Engine/Source/Editor/UnrealEd/Private/Kismet2/KismetReinstanceUtilities.cpp b/Engine/Source/Editor/UnrealEd/Private/Kismet2/KismetReinstanceUtilities.cpp
--- a/Engine/Source/Editor/UnrealEd/Private/Kismet2/KismetReinstanceUtilities.cpp
+++ b/Engine/Source/Editor/UnrealEd/Private/Kismet2/KismetReinstanceUtilities.cpp
@@ -1183,7 +1183,8 @@ void FActorReplacementHelper::Finalize(const TMap<UObject*, UObject*>& OldToNewI
 	FEditorScriptExecutionGuard ScriptGuard;
 
 	// run the construction script, which will use the properties we just copied over
-	if (NewActor->CurrentTransactionAnnotation.IsValid())
+	bool bCanReRun = UBlueprint::IsBlueprintHierarchyErrorFree(NewActor->GetClass());
+	if (NewActor->CurrentTransactionAnnotation.IsValid() && bCanReRun)
 	{
 		NewActor->CurrentTransactionAnnotation->ComponentInstanceData.FindAndReplaceInstances(OldToNewInstanceMap);
 		NewActor->RerunConstructionScripts();
diff --git a/Engine/Source/Runtime/Engine/Classes/Engine/Blueprint.h b/Engine/Source/Runtime/Engine/Classes/Engine/Blueprint.h
--- a/Engine/Source/Runtime/Engine/Classes/Engine/Blueprint.h
+++ b/Engine/Source/Runtime/Engine/Classes/Engine/Blueprint.h
@@ -851,6 +851,11 @@ public:
 	 * @return						true if there were no status errors in any of the parent blueprints, otherwise false
 	 */
 	static bool GetBlueprintHierarchyFromClass(const UClass* InClass, TArray<UBlueprint*>& OutBlueprintParents);
+	
+#if WITH_EDITOR
+	/** returns true if the class hierarchy is error free */
+	static bool IsBlueprintHierarchyErrorFree(const UClass* InClass);
+#endif
 
 #if WITH_EDITOR
 	template<class TFieldType>
diff --git a/Engine/Source/Runtime/Engine/Private/Blueprint.cpp b/Engine/Source/Runtime/Engine/Private/Blueprint.cpp
--- a/Engine/Source/Runtime/Engine/Private/Blueprint.cpp
+++ b/Engine/Source/Runtime/Engine/Private/Blueprint.cpp
@@ -1471,6 +1471,33 @@ bool UBlueprint::GetBlueprintHierarchyFromClass(const UClass* InClass, TArray<UB
 	return bNoErrors;
 }
 
+#if WITH_EDITOR
+bool UBlueprint::IsBlueprintHierarchyErrorFree(const UClass* InClass)
+{
+	const UClass* CurrentClass = InClass;
+	while (UBlueprint* BP = UBlueprint::GetBlueprintFromClass(CurrentClass))
+	{
+		if(BP->Status == BS_Error)
+		{
+			return false;
+		}
+
+		// If valid, use stored ParentClass rather than the actual UClass::GetSuperClass(); handles the case when the class has not been recompiled yet after a reparent operation.
+		if(const UClass* ParentClass = BP->ParentClass)
+		{
+			CurrentClass = ParentClass;
+		}
+		else
+		{
+			check(CurrentClass);
+			CurrentClass = CurrentClass->GetSuperClass();
+		}
+	}
+
+	return true;
+}
+#endif
+
 ETimelineSigType UBlueprint::GetTimelineSignatureForFunctionByName(const FName& FunctionName, const FName& ObjectPropertyName)
 {
 	check(SkeletonGeneratedClass != NULL);
