Avoid crash when initializing DebugDraw canvas while intraframe debugging
#rnx

diff --git a/Engine/Source/Runtime/Engine/Private/Debug/DebugDrawService.cpp b/Engine/Source/Runtime/Engine/Private/Debug/DebugDrawService.cpp
--- a/Engine/Source/Runtime/Engine/Private/Debug/DebugDrawService.cpp
+++ b/Engine/Source/Runtime/Engine/Private/Debug/DebugDrawService.cpp
@@ -64,7 +64,7 @@ void UDebugDrawService::Draw(const FEngineShowFlags Flags, FViewport* Viewport,
 		CanvasObject->AddToRoot();
 	}
 
-	if (!(Canvas->IsStereoRendering() && GEngine->XRSystem.IsValid() && GEngine->XRSystem->GetHMDDevice()))
+	if (!Canvas || !(Canvas->IsStereoRendering() && GEngine->XRSystem.IsValid() && GEngine->XRSystem->GetHMDDevice()))
 	{
 		CanvasObject->Init(View->UnscaledViewRect.Width(), View->UnscaledViewRect.Height(), View, Canvas);
 	}
