RL - Fix crash where Wireframe rendering crashes because of an empty section the regular path itterates over section and does not hit this issue. The affected Mesh and blueprint were: SM_NoTickHaze and BP_NoTickHaze36_2

diff --git a/Engine/Source/Runtime/Engine/Private/StaticMeshRender.cpp b/Engine/Source/Runtime/Engine/Private/StaticMeshRender.cpp
--- a/Engine/Source/Runtime/Engine/Private/StaticMeshRender.cpp
+++ b/Engine/Source/Runtime/Engine/Private/StaticMeshRender.cpp
@@ -1030,14 +1030,19 @@ void FStaticMeshSceneProxy::GetDynamicMeshElements(const TArray<const FSceneView
 
 							for (int32 BatchIndex = 0; BatchIndex < NumBatches; BatchIndex++)
 							{
-								FMeshBatch& Mesh = Collector.AllocateMesh();
-
-								if (GetWireframeMeshElement(LODIndex, BatchIndex, WireframeMaterialInstance, SDPG_World, true, Mesh))
+								//GetWireframeMeshElement will try SetIndexSource at sectionindex 0
+								//and GetMeshElement loops over sections, therefore does not have this issue
+								if (LODModel.Sections.Num() > 0)
 								{
-									// We implemented our own wireframe
-									Mesh.bCanApplyViewModeOverrides = false;
-									Collector.AddMesh(ViewIndex, Mesh);
-									INC_DWORD_STAT_BY(STAT_StaticMeshTriangles,Mesh.GetNumPrimitives());
+									FMeshBatch& Mesh = Collector.AllocateMesh();
+
+									if (GetWireframeMeshElement(LODIndex, BatchIndex, WireframeMaterialInstance, SDPG_World, true, Mesh))
+									{
+										// We implemented our own wireframe
+										Mesh.bCanApplyViewModeOverrides = false;
+										Collector.AddMesh(ViewIndex, Mesh);
+										INC_DWORD_STAT_BY(STAT_StaticMeshTriangles, Mesh.GetNumPrimitives());
+									}
 								}
 							}
 						}
