AJA/Blackmagic

- Fix for packaging a project with either plugin

diff --git a/Engine/Plugins/Media/MediaFrameworkUtilities/Source/MediaFrameworkUtilities/Private/Profile/MediaProfile.cpp b/Engine/Plugins/Media/MediaFrameworkUtilities/Source/MediaFrameworkUtilities/Private/Profile/MediaProfile.cpp
--- a/Engine/Plugins/Media/MediaFrameworkUtilities/Source/MediaFrameworkUtilities/Private/Profile/MediaProfile.cpp
+++ b/Engine/Plugins/Media/MediaFrameworkUtilities/Source/MediaFrameworkUtilities/Private/Profile/MediaProfile.cpp
@@ -128,7 +128,7 @@ void UMediaProfile::Apply()
 			bool bResult = GEngine->SetTimecodeProvider(TimecodeProvider);
 			if (!bResult)
 			{
-				UE_LOG(LogMediaFrameworkUtilities, Error, TEXT("The TimecodeProvider '%s' could not be initialized."), *TimecodeProvider->GetName());
+				UE_LOG(LogMediaFrameworkUtilities, Warning, TEXT("The TimecodeProvider '%s' could not be initialized."), *TimecodeProvider->GetName());
 			}
 		}
 		else
@@ -144,7 +144,7 @@ void UMediaProfile::Apply()
 			bool bResult = GEngine->SetCustomTimeStep(CustomTimeStep);
 			if (!bResult)
 			{
-				UE_LOG(LogMediaFrameworkUtilities, Error, TEXT("The Custom Time Step '%s' could not be initialized."), *CustomTimeStep->GetName());
+				UE_LOG(LogMediaFrameworkUtilities, Warning, TEXT("The Custom Time Step '%s' could not be initialized."), *CustomTimeStep->GetName());
 			}
 		}
 		else
diff --git a/Engine/Plugins/Media/MediaIOFramework/Source/MediaIOCore/Private/MediaIOCoreModule.cpp b/Engine/Plugins/Media/MediaIOFramework/Source/MediaIOCore/Private/MediaIOCoreModule.cpp
--- a/Engine/Plugins/Media/MediaIOFramework/Source/MediaIOCore/Private/MediaIOCoreModule.cpp
+++ b/Engine/Plugins/Media/MediaIOFramework/Source/MediaIOCore/Private/MediaIOCoreModule.cpp
@@ -44,6 +44,10 @@ private:
 	TArray<IMediaIOCoreDeviceProvider*> DeviceProviders;
 };
 
+bool IMediaIOCoreModule::IsAvailable()
+{
+	return FModuleManager::Get().IsModuleLoaded("MediaIOCore");
+}
 
 IMediaIOCoreModule& IMediaIOCoreModule::Get()
 {
@@ -51,4 +55,6 @@ IMediaIOCoreModule& IMediaIOCoreModule::Get()
 }
 
 
+
+
 IMPLEMENT_MODULE(FMediaIOCoreModule, MediaIOCore);
diff --git a/Engine/Plugins/Media/MediaIOFramework/Source/MediaIOCore/Public/IMediaIOCoreModule.h b/Engine/Plugins/Media/MediaIOFramework/Source/MediaIOCore/Public/IMediaIOCoreModule.h
--- a/Engine/Plugins/Media/MediaIOFramework/Source/MediaIOCore/Public/IMediaIOCoreModule.h
+++ b/Engine/Plugins/Media/MediaIOFramework/Source/MediaIOCore/Public/IMediaIOCoreModule.h
@@ -14,6 +14,7 @@ class IMediaIOCoreDeviceProvider;
 class MEDIAIOCORE_API IMediaIOCoreModule : public IModuleInterface
 {
 public:
+	static bool IsAvailable();
 	static IMediaIOCoreModule& Get();
 
 public:
diff --git a/Engine/Source/Runtime/Engine/Private/UnrealEngine.cpp b/Engine/Source/Runtime/Engine/Private/UnrealEngine.cpp
--- a/Engine/Source/Runtime/Engine/Private/UnrealEngine.cpp
+++ b/Engine/Source/Runtime/Engine/Private/UnrealEngine.cpp
@@ -2023,7 +2023,7 @@ void UEngine::ReinitializeCustomTimeStep()
 		CustomTimeStep->Shutdown(this);
 		if (!CustomTimeStep->Initialize(this))
 		{
-			UE_LOG(LogEngine, Error, TEXT("Failed reinitializing CustomTimeStep %s"), *GetPathName(CustomTimeStep));
+			UE_LOG(LogEngine, Warning, TEXT("Failed reinitializing CustomTimeStep %s"), *GetPathName(CustomTimeStep));
 		}
 	}
 }
@@ -2047,7 +2047,7 @@ bool UEngine::SetCustomTimeStep(UEngineCustomTimeStep* InCustomTimeStep)
 			bResult = CurrentCustomTimeStep->Initialize(this);
 			if (!bResult)
 			{
-				UE_LOG(LogEngine, Error, TEXT("SetCustomTimeStep - Failed to intialize CustomTimeStep %s"), *GetPathName(CurrentCustomTimeStep));
+				UE_LOG(LogEngine, Warning, TEXT("SetCustomTimeStep - Failed to intialize CustomTimeStep %s"), *GetPathName(CurrentCustomTimeStep));
 				CurrentCustomTimeStep = nullptr;
 			}
 		}
@@ -2062,7 +2062,7 @@ void UEngine::ReinitializeTimecodeProvider()
 	Provider->Shutdown(this);
 	if (!Provider->Initialize(this))
 	{
-		UE_LOG(LogEngine, Error, TEXT("Failed reinitializing TimecodeProvider %s"), *GetPathName(Provider));
+		UE_LOG(LogEngine, Warning, TEXT("Failed reinitializing TimecodeProvider %s"), *GetPathName(Provider));
 	}
 }
 
@@ -2106,7 +2106,7 @@ bool UEngine::SetTimecodeProvider(UTimecodeProvider* InTimecodeProvider)
 		{
 			if (!ensure(DefaultTimecodeProvider->Initialize(this)))
 			{
-				UE_LOG(LogEngine, Error, TEXT("SetTimecodeProvider - Failed to intialize DefaultTimecodeProvider %s"), *GetPathName(DefaultTimecodeProvider));
+				UE_LOG(LogEngine, Warning, TEXT("SetTimecodeProvider - Failed to intialize DefaultTimecodeProvider %s"), *GetPathName(DefaultTimecodeProvider));
 			}
 		}
 	}
@@ -2449,7 +2449,7 @@ void UEngine::InitializeObjectReferences()
 			DefaultTimecodeProvider = NewObject<UTimecodeProvider>(this, DefaultTimecodeProviderClass);
 			if (!ensure(DefaultTimecodeProvider->Initialize(this)))
 			{
-				UE_LOG(LogEngine, Error, TEXT("InitializeObjectReferences - Failed to intialize DefaultTimecodeProvider %s"), *GetPathName(DefaultTimecodeProvider));
+				UE_LOG(LogEngine, Warning, TEXT("InitializeObjectReferences - Failed to intialize DefaultTimecodeProvider %s"), *GetPathName(DefaultTimecodeProvider));
 			}
 
 			DefaultTimecodeProvider->AddToRoot();
