Fixed bad volume texture size in cooked when using LODBias

diff --git a/Engine/Source/Runtime/Engine/Private/TextureDerivedData.cpp b/Engine/Source/Runtime/Engine/Private/TextureDerivedData.cpp
--- a/Engine/Source/Runtime/Engine/Private/TextureDerivedData.cpp
+++ b/Engine/Source/Runtime/Engine/Private/TextureDerivedData.cpp
@@ -22,6 +22,7 @@
 #include "DeviceProfiles/DeviceProfileManager.h"
 #include "TextureDerivedDataTask.h"
 #include "Streaming/TextureStreamingHelpers.h"
+#include "Engine/VolumeTexture.h"
 
 #if WITH_EDITOR
 
@@ -30,7 +31,6 @@
 #include "Interfaces/ITargetPlatformManagerModule.h"
 #include "Interfaces/ITextureFormat.h"
 #include "Engine/TextureCube.h"
-#include "Engine/VolumeTexture.h"
 #include "ProfilingDebugging/CookStats.h"
 
 /*------------------------------------------------------------------------------
@@ -1097,6 +1097,12 @@ void FTexturePlatformData::SerializeCooked(FArchive& Ar, UTexture* Owner, bool b
 	{
 		SizeX = Mips[0].SizeX;
 		SizeY = Mips[0].SizeY;
+
+		// SizeZ is not the same as NumSlices for texture arrays and cubemaps.
+		if (Owner && Owner->IsA(UVolumeTexture::StaticClass()))
+		{
+			 NumSlices = Mips[0].SizeZ;
+		}
 	}
 }
 
