Disable GC clustering support for SoundCues, SoundNodes and SoundWaves. Sound classes have a very dynamic nature which is not compatible with GC clustering and may lead to hard to track down crashes.

diff --git a/Engine/Source/Runtime/Engine/Classes/Sound/SoundBase.h b/Engine/Source/Runtime/Engine/Classes/Sound/SoundBase.h
--- a/Engine/Source/Runtime/Engine/Classes/Sound/SoundBase.h
+++ b/Engine/Source/Runtime/Engine/Classes/Sound/SoundBase.h
@@ -142,6 +142,8 @@ public:
 	//~ Begin UObject Interface.
 	virtual void PostInitProperties() override;
 	virtual void PostLoad() override;
+	virtual bool CanBeClusterRoot() const override;
+	virtual bool CanBeInCluster() const override;
 	//~ End UObject interface.
 	
 	/** Returns whether the sound base is set up in a playable manner */
diff --git a/Engine/Source/Runtime/Engine/Classes/Sound/SoundCue.h b/Engine/Source/Runtime/Engine/Classes/Sound/SoundCue.h
--- a/Engine/Source/Runtime/Engine/Classes/Sound/SoundCue.h
+++ b/Engine/Source/Runtime/Engine/Classes/Sound/SoundCue.h
@@ -135,6 +135,8 @@ public:
 #endif
 	virtual void PostLoad() override;
 	virtual void Serialize(FStructuredArchive::FRecord Record) override;
+	virtual bool CanBeClusterRoot() const override;
+	virtual bool CanBeInCluster() const override;
 	//~ End UObject Interface.
 
 	//~ Begin USoundBase Interface.
diff --git a/Engine/Source/Runtime/Engine/Classes/Sound/SoundNode.h b/Engine/Source/Runtime/Engine/Classes/Sound/SoundNode.h
--- a/Engine/Source/Runtime/Engine/Classes/Sound/SoundNode.h
+++ b/Engine/Source/Runtime/Engine/Classes/Sound/SoundNode.h
@@ -77,6 +77,8 @@ public:
 	static void AddReferencedObjects(UObject* InThis, FReferenceCollector& Collector);
 #endif //WITH_EDITOR
 	virtual void Serialize(FArchive& Ar) override;
+	virtual bool CanBeClusterRoot() const override;
+	virtual bool CanBeInCluster() const override;
 	//~ End UObject Interface
 
 	//
diff --git a/Engine/Source/Runtime/Engine/Private/SoundBase.cpp b/Engine/Source/Runtime/Engine/Private/SoundBase.cpp
--- a/Engine/Source/Runtime/Engine/Private/SoundBase.cpp
+++ b/Engine/Source/Runtime/Engine/Private/SoundBase.cpp
@@ -184,3 +184,13 @@ void USoundBase::PostLoad()
 #endif
 }
 
+bool USoundBase::CanBeClusterRoot() const
+{
+	return false;
+}
+
+bool USoundBase::CanBeInCluster() const
+{
+	return false;
+}
+
diff --git a/Engine/Source/Runtime/Engine/Private/SoundCue.cpp b/Engine/Source/Runtime/Engine/Private/SoundCue.cpp
--- a/Engine/Source/Runtime/Engine/Private/SoundCue.cpp
+++ b/Engine/Source/Runtime/Engine/Private/SoundCue.cpp
@@ -159,6 +159,16 @@ void USoundCue::PostLoad()
 	CacheAggregateValues();
 }
 
+bool USoundCue::CanBeClusterRoot() const
+{
+	return false;
+}
+
+bool USoundCue::CanBeInCluster() const
+{
+	return false;
+}
+
 void USoundCue::OnPostEngineInit()
 {
 	FCoreDelegates::OnPostEngineInit.Remove(OnPostEngineInitHandle);
diff --git a/Engine/Source/Runtime/Engine/Private/SoundNode.cpp b/Engine/Source/Runtime/Engine/Private/SoundNode.cpp
--- a/Engine/Source/Runtime/Engine/Private/SoundNode.cpp
+++ b/Engine/Source/Runtime/Engine/Private/SoundNode.cpp
@@ -36,6 +36,16 @@ void USoundNode::Serialize(FArchive& Ar)
 #endif
 }
 
+bool USoundNode::CanBeClusterRoot() const
+{
+	return false;
+}
+
+bool USoundNode::CanBeInCluster() const
+{
+	return false;
+}
+
 #if WITH_EDITOR
 void USoundNode::AddReferencedObjects(UObject* InThis, FReferenceCollector& Collector)
 {
