PR #5169: Mark bHiddenEdLayer as Transient (Contributed by [REDACTED_PERSON])

diff --git a/Engine/Source/Runtime/Engine/Classes/GameFramework/Actor.h b/Engine/Source/Runtime/Engine/Classes/GameFramework/Actor.h
--- a/Engine/Source/Runtime/Engine/Classes/GameFramework/Actor.h
+++ b/Engine/Source/Runtime/Engine/Classes/GameFramework/Actor.h
@@ -604,7 +604,7 @@ public:
 	uint8 bIsEditorPreviewActor:1;
 
 	/** Whether this actor is hidden by the layer browser. */
-	UPROPERTY()
+	UPROPERTY(Transient)
 	uint8 bHiddenEdLayer:1;
 
 	/** Whether this actor is hidden by the level browser. */
