Omitting material function instances as allowed classes to create on the material function call dropdown

diff --git a/Engine/Source/Runtime/Engine/Classes/Materials/MaterialExpressionMaterialFunctionCall.h b/Engine/Source/Runtime/Engine/Classes/Materials/MaterialExpressionMaterialFunctionCall.h
--- a/Engine/Source/Runtime/Engine/Classes/Materials/MaterialExpressionMaterialFunctionCall.h
+++ b/Engine/Source/Runtime/Engine/Classes/Materials/MaterialExpressionMaterialFunctionCall.h
@@ -80,7 +80,7 @@ class UMaterialExpressionMaterialFunctionCall : public UMaterialExpression
 	GENERATED_UCLASS_BODY()
 
 	/** The function to call. */
-	UPROPERTY(EditAnywhere, Category=MaterialExpressionMaterialFunctionCall)
+	UPROPERTY(EditAnywhere, Category=MaterialExpressionMaterialFunctionCall, meta=(DisallowedClasses="MaterialFunctionInstance"))
 	class UMaterialFunctionInterface* MaterialFunction;
 
 #if WITH_EDITORONLY_DATA
