Disable hmd distortion showflag by default.

diff --git a/Engine/Source/Runtime/Engine/Public/ShowFlags.h b/Engine/Source/Runtime/Engine/Public/ShowFlags.h
--- a/Engine/Source/Runtime/Engine/Public/ShowFlags.h
+++ b/Engine/Source/Runtime/Engine/Public/ShowFlags.h
@@ -368,7 +368,7 @@ private:
 		SetCameraSafeFrames(false);
 		SetVisualizeOutOfBoundsPixels(false);
 		SetHighResScreenshotMask(false);
-		SetHMDDistortion(true); // only effective if a HMD is in use
+		SetHMDDistortion(false); // only effective if a HMD is in use
 		SetStereoRendering(false);
 		SetDistanceCulledPrimitives(InitMode == ESFIM_Editor);
 		SetVisualizeLightCulling(false);
