Copied from 4444191

Expose Mesh edition settings to BP:
- FMeshReductionSettings
- FMeshProxySettings
- FMeshMergingSettings
- FMeshInstancingSettings

diff --git a/Engine/Source/Runtime/Engine/Classes/Engine/MeshMerging.h b/Engine/Source/Runtime/Engine/Classes/Engine/MeshMerging.h
--- a/Engine/Source/Runtime/Engine/Classes/Engine/MeshMerging.h
+++ b/Engine/Source/Runtime/Engine/Classes/Engine/MeshMerging.h
@@ -24,67 +24,67 @@ namespace EMeshFeatureImportance
 }
 
 /** Settings used to reduce a mesh. */
-USTRUCT()
+USTRUCT(Blueprintable)
 struct FMeshReductionSettings
 {
 	GENERATED_USTRUCT_BODY()
 
 	/** Percentage of triangles to keep. 1.0 = no reduction, 0.0 = no triangles. */
-	UPROPERTY(EditAnywhere, Category = ReductionSettings)
+	UPROPERTY(BlueprintReadWrite, EditAnywhere, Category = ReductionSettings)
 	float PercentTriangles;
 
 	/** The maximum distance in object space by which the reduced mesh may deviate from the original mesh. */
-	UPROPERTY(EditAnywhere, Category = ReductionSettings)
+	UPROPERTY(BlueprintReadWrite, EditAnywhere, Category = ReductionSettings)
 	float MaxDeviation;
 
 	/** The amount of error in pixels allowed for this LOD. */
-	UPROPERTY(EditAnywhere, Category = ReductionSettings)
+	UPROPERTY(BlueprintReadWrite, EditAnywhere, Category = ReductionSettings)
 	float PixelError;
 
 	/** Threshold in object space at which vertices are welded together. */
-	UPROPERTY(EditAnywhere, Category = ReductionSettings)
+	UPROPERTY(BlueprintReadWrite, EditAnywhere, Category = ReductionSettings)
 	float WeldingThreshold;
 
 	/** Angle at which a hard edge is introduced between faces. */
-	UPROPERTY(EditAnywhere, Category = ReductionSettings)
+	UPROPERTY(BlueprintReadWrite, EditAnywhere, Category = ReductionSettings)
 	float HardAngleThreshold;
 
 	/** Higher values minimize change to border edges. */
-	UPROPERTY(EditAnywhere, Category = ReductionSettings)
+	UPROPERTY(BlueprintReadWrite, EditAnywhere, Category = ReductionSettings)
 	TEnumAsByte<EMeshFeatureImportance::Type> SilhouetteImportance;
 
 	/** Higher values reduce texture stretching. */
-	UPROPERTY(EditAnywhere, Category = ReductionSettings)
+	UPROPERTY(BlueprintReadWrite, EditAnywhere, Category = ReductionSettings)
 	TEnumAsByte<EMeshFeatureImportance::Type> TextureImportance;
 
 	/** Higher values try to preserve normals better. */
-	UPROPERTY(EditAnywhere, Category = ReductionSettings)
+	UPROPERTY(BlueprintReadWrite, EditAnywhere, Category = ReductionSettings)
 	TEnumAsByte<EMeshFeatureImportance::Type> ShadingImportance;
 
-	UPROPERTY(EditAnywhere, Category = ReductionSettings)
+	UPROPERTY(BlueprintReadWrite, EditAnywhere, Category = ReductionSettings)
 	bool bRecalculateNormals;
 
-	UPROPERTY(EditAnywhere, Category = ReductionSettings)
+	UPROPERTY(BlueprintReadWrite, EditAnywhere, Category = ReductionSettings)
 	int32 BaseLODModel;
 
-	UPROPERTY(EditAnywhere, Category = ReductionSettings)
+	UPROPERTY(BlueprintReadWrite, EditAnywhere, Category = ReductionSettings)
 	bool bGenerateUniqueLightmapUVs;
 
-	UPROPERTY(EditAnywhere, Category = ReductionSettings)
+	UPROPERTY(BlueprintReadWrite, EditAnywhere, Category = ReductionSettings)
 	bool bKeepSymmetry;
 
-	UPROPERTY(EditAnywhere, Category = ReductionSettings)
+	UPROPERTY(BlueprintReadWrite, EditAnywhere, Category = ReductionSettings)
 	bool bVisibilityAided;
 
-	UPROPERTY(EditAnywhere, Category = ReductionSettings)
+	UPROPERTY(BlueprintReadWrite, EditAnywhere, Category = ReductionSettings)
 	bool bCullOccluded;
 
 	/** Higher values generates fewer samples*/
-	UPROPERTY(EditAnywhere, Category = ReductionSettings)
+	UPROPERTY(BlueprintReadWrite, EditAnywhere, Category = ReductionSettings)
 	TEnumAsByte<EMeshFeatureImportance::Type> VisibilityAggressiveness;
 
 	/** Higher values minimize change to vertex color data. */
-	UPROPERTY(EditAnywhere, Category = ReductionSettings)
+	UPROPERTY(BlueprintReadWrite, EditAnywhere, Category = ReductionSettings)
 	TEnumAsByte<EMeshFeatureImportance::Type> VertexColorImportance;
 
 	/** Default settings. */
@@ -184,24 +184,24 @@ struct FMeshProxySettings
 {
 	GENERATED_USTRUCT_BODY()
 	/** Screen size of the resulting proxy mesh in pixels*/
-	UPROPERTY(EditAnywhere, Category = ProxySettings, meta = (ClampMin = "1", ClampMax = "1200", UIMin = "1", UIMax = "1200"))
+	UPROPERTY(BlueprintReadWrite, EditAnywhere, Category = ProxySettings, meta = (ClampMin = "1", ClampMax = "1200", UIMin = "1", UIMax = "1200"))
 	int32 ScreenSize;
 
 	/** If true, Spatial Sampling Distance will not be automatically computed based on geometry and you must set it directly */
-	UPROPERTY(EditAnywhere, AdvancedDisplay, Category = ProxySettings, meta = (InlineEditConditionToggle))
+	UPROPERTY(BlueprintReadWrite, EditAnywhere, AdvancedDisplay, Category = ProxySettings, meta = (InlineEditConditionToggle))
 	uint8 bOverrideVoxelSize : 1;
 
-
 	/** Override when converting multiple meshes for proxy LOD merging. Warning, large geometry with small sampling has very high memory costs*/
-	UPROPERTY(EditAnywhere, AdvancedDisplay, Category = ProxySettings, meta = (editcondition = "bOverrideVoxelSize", ClampMin = "0.1", DisplayName = "Overide Spatial Sampling Distance"))
+	UPROPERTY(BlueprintReadWrite, EditAnywhere, AdvancedDisplay, Category = ProxySettings, meta = (EditCondition = "bOverrideVoxelSize", ClampMin = "0.1", DisplayName = "Overide Spatial Sampling Distance"))
 	float VoxelSize;
-	
+
 	/** Material simplification */
-	UPROPERTY(EditAnywhere, Category = ProxySettings)
+	UPROPERTY(BlueprintReadWrite, EditAnywhere, Category = ProxySettings)
 	FMaterialProxySettings MaterialSettings;
 
 	UPROPERTY()
 	int32 TextureWidth_DEPRECATED;
+
 	UPROPERTY()
 	int32 TextureHeight_DEPRECATED;
 
@@ -218,84 +218,84 @@ struct FMeshProxySettings
 	bool bExportSpecularMap_DEPRECATED;
 
 	/** Determines whether or not the correct LOD models should be calculated given the source meshes and transition size */
-	UPROPERTY(EditAnywhere, Category = ProxySettings)
+	UPROPERTY(BlueprintReadWrite, EditAnywhere, Category = ProxySettings)
 	bool bCalculateCorrectLODModel;
 
 	/** Distance at which meshes should be merged together, this can close gaps like doors and windows in distant geometry */
-	UPROPERTY(EditAnywhere, Category = ProxySettings)
+	UPROPERTY(BlueprintReadWrite, EditAnywhere, Category = ProxySettings)
 	float MergeDistance;
 
 	/** Base color assigned to LOD geometry that can't be associated with the source geometry: e.g. doors and windows that have been closed by the Merge Distance */
-	UPROPERTY(EditAnywhere, Category = ProxySettings, meta = (DisplayName = "Unresolved Geometry Color"))
+	UPROPERTY(BlueprintReadWrite, EditAnywhere, Category = ProxySettings, meta = (DisplayName = "Unresolved Geometry Color"))
 	FColor UnresolvedGeometryColor;
 
 	/** Enable an override for material transfer distance */
-	UPROPERTY(EditAnywhere, Category = MaxRayCastDist, meta = (InlineEditConditionToggle))
+	UPROPERTY(BlueprintReadWrite, EditAnywhere, Category = MaxRayCastDist, meta = (InlineEditConditionToggle))
 	bool bOverrideTransferDistance;
 
 	/** Override search distance used when discovering texture values for simplified geometry. Useful when non-zero Merge Distance setting generates new geometry in concave corners.*/
-	UPROPERTY(EditAnywhere, Category = ProxySettings, meta = (EditCondition = "bOverrideTransferDistance", DisplayName = "Transfer Distance Override", ClampMin = 0))
+	UPROPERTY(BlueprintReadWrite, EditAnywhere, Category = ProxySettings, meta = (EditCondition = "bOverrideTransferDistance", DisplayName = "Transfer Distance Override", ClampMin = 0))
 	float MaxRayCastDist;
 
 	/** Enable the use of hard angle based vertex splitting */
-	UPROPERTY(EditAnywhere, Category = HardAngleThreshold, meta = (InlineEditConditionToggle))
+	UPROPERTY(BlueprintReadWrite, EditAnywhere, Category = HardAngleThreshold, meta = (InlineEditConditionToggle))
 	bool bUseHardAngleThreshold;
 
 	/** Angle at which a hard edge is introduced between faces */
-	UPROPERTY(EditAnywhere, Category = ProxySettings, meta = (EditCondition = "bUseHardAngleThreshold", DisplayName = "Hard Edge Angle", ClampMin = 0, ClampMax = 180))
+	UPROPERTY(BlueprintReadWrite, EditAnywhere, Category = ProxySettings, meta = (EditCondition = "bUseHardAngleThreshold", DisplayName = "Hard Edge Angle", ClampMin = 0, ClampMax = 180))
 	float HardAngleThreshold;
 
 	/** Controls the method used to calculate the normal for the simplified geometry */
-	UPROPERTY(EditAnywhere, Category = ProxySettings, meta = (DisplayName = "Normal Calculation Method"))
+	UPROPERTY(BlueprintReadWrite, EditAnywhere, Category = ProxySettings, meta = (DisplayName = "Normal Calculation Method"))
 	TEnumAsByte<EProxyNormalComputationMethod::Type> NormalCalculationMethod;
 
 	/** Lightmap resolution */
-	UPROPERTY(EditAnywhere, Category = ProxySettings, meta = (ClampMin = 32, ClampMax = 4096, EditCondition = "!bComputeLightMapResolution"))
+	UPROPERTY(BlueprintReadWrite, EditAnywhere, Category = ProxySettings, meta = (ClampMin = 32, ClampMax = 4096, EditCondition = "!bComputeLightMapResolution"))
 	int32 LightMapResolution;
 
 	/** If ticked will compute the lightmap resolution by summing the dimensions for each mesh included for merging */
-	UPROPERTY(EditAnywhere, Category = ProxySettings)
+	UPROPERTY(BlueprintReadWrite, EditAnywhere, Category = ProxySettings)
 	bool bComputeLightMapResolution;
 
 	/** Whether Simplygon should recalculate normals, otherwise the normals channel will be sampled from the original mesh */
-	UPROPERTY(EditAnywhere, Category = ProxySettings)
+	UPROPERTY(BlueprintReadWrite, EditAnywhere, Category = ProxySettings)
 	bool bRecalculateNormals;
 
 	UPROPERTY()
 	bool bBakeVertexData_DEPRECATED;
 
 	/** Whether or not to use available landscape geometry to cull away invisible triangles */
-	UPROPERTY(EditAnywhere, Category = LandscapeCulling)
+	UPROPERTY(BlueprintReadWrite, EditAnywhere, Category = LandscapeCulling)
 	bool bUseLandscapeCulling;
 
 	/** Level of detail of the landscape that should be used for the culling */
-	UPROPERTY(EditAnywhere, Category = LandscapeCulling, meta = (EditCondition="bUseLandscapeCulling"))
+	UPROPERTY(BlueprintReadWrite, EditAnywhere, Category = LandscapeCulling, meta = (EditCondition="bUseLandscapeCulling"))
 	TEnumAsByte<ELandscapeCullingPrecision::Type> LandscapeCullingPrecision;
 
 	/** Whether to allow adjacency buffers for tessellation in the merged mesh */
-	UPROPERTY(EditAnywhere, Category = ProxySettings)
+	UPROPERTY(BlueprintReadWrite, EditAnywhere, Category = ProxySettings)
 	bool bAllowAdjacency;
 
 	/** Whether to allow distance field to be computed for this mesh. Disable this to save memory if the merged mesh will only be rendered in the distance. */
-	UPROPERTY(EditAnywhere, Category = ProxySettings)
+	UPROPERTY(BlueprintReadWrite, EditAnywhere, Category = ProxySettings)
 	bool bAllowDistanceField;
 
 	/** Whether to attempt to re-use the source mesh's lightmap UVs when baking the material or always generate a new set. */
-	UPROPERTY(EditAnywhere, Category = ProxySettings)
+	UPROPERTY(BlueprintReadWrite, EditAnywhere, Category = ProxySettings)
 	bool bReuseMeshLightmapUVs;
 
 	/** Whether to generate collision for the merged mesh */
-	UPROPERTY(EditAnywhere, Category = ProxySettings)
+	UPROPERTY(BlueprintReadWrite, EditAnywhere, Category = ProxySettings)
 	bool bCreateCollision;
 
 	/** Whether to allow vertex colors saved in the merged mesh */
-	UPROPERTY(EditAnywhere, Category = ProxySettings)
+	UPROPERTY(BlueprintReadWrite, EditAnywhere, Category = ProxySettings)
 	bool bAllowVertexColors;
 
 	/** Whether to generate lightmap uvs for the merged mesh */
-	UPROPERTY(EditAnywhere, Category = ProxySettings)
+	UPROPERTY(BlueprintReadWrite, EditAnywhere, Category = ProxySettings)
 	bool bGenerateLightmapUVs;
-	
+
 	/** Default settings. */
 	FMeshProxySettings()
 		: ScreenSize(300)
@@ -327,7 +327,7 @@ struct FMeshProxySettings
 		, bCreateCollision(true)
 		, bAllowVertexColors(false)
 		, bGenerateLightmapUVs(false)
-	{ 
+	{
 		MaterialSettings.MaterialMergeType = EMaterialMergeType::MaterialMergeType_Simplygon;
 	}
 
@@ -363,7 +363,7 @@ UENUM()
 enum class EMeshLODSelectionType : uint8
 {
 	// Whether or not to export all of the LODs found in the source meshes
-	AllLODs = 0 UMETA(DisplayName = "Use all LOD levels"),	
+	AllLODs = 0 UMETA(DisplayName = "Use all LOD levels"),
 	// Whether or not to export all of the LODs found in the source meshes
 	SpecificLOD = 1 UMETA(DisplayName = "Use specific LOD level"),
 	// Whether or not to calculate the appropriate LOD model for the given screen size
@@ -400,7 +400,7 @@ struct FMeshMergingSettings
 	bool bGenerateLightMapUV;
 
 	/** Target lightmap resolution */
-	UPROPERTY(EditAnywhere, AdvancedDisplay, BlueprintReadWrite, Category = MeshSettings, meta=(ClampMax = 4096, editcondition = "!bComputedLightMapResolution"))
+	UPROPERTY(EditAnywhere, AdvancedDisplay, BlueprintReadWrite, Category = MeshSettings, meta=(ClampMax = 4096, EditCondition = "!bComputedLightMapResolution"))
 	int32 TargetLightMapResolution;
 
 	/** Whether or not the lightmap resolution should be computed by summing the lightmap resolutions for the input Mesh Components */
@@ -424,7 +424,7 @@ struct FMeshMergingSettings
 	bool bMergeMaterials;
 
 	/** Material simplification */
-	UPROPERTY(EditAnywhere, Category = MaterialSettings, BlueprintReadWrite, meta = (editcondition = "bMergeMaterials"))
+	UPROPERTY(EditAnywhere, Category = MaterialSettings, BlueprintReadWrite, meta = (EditCondition = "bMergeMaterials"))
 	FMaterialProxySettings MaterialSettings;
 
 	/** Whether or not vertex data such as vertex colours should be baked into the resulting mesh */
@@ -432,13 +432,13 @@ struct FMeshMergingSettings
 	bool bBakeVertexDataToMesh;
 
 	/** Whether or not vertex data such as vertex colours should be used when baking out materials */
-	UPROPERTY(EditAnywhere, Category = MaterialSettings, BlueprintReadWrite, meta = (editcondition = "bMergeMaterials"))
+	UPROPERTY(EditAnywhere, Category = MaterialSettings, BlueprintReadWrite, meta = (EditCondition = "bMergeMaterials"))
 	bool bUseVertexDataForBakingMaterial;
 
 	// Whether or not to calculate varying output texture sizes according to their importance in the final atlas texture
-	UPROPERTY(Category = MaterialSettings, EditAnywhere, BlueprintReadWrite, meta = (editcondition = "bMergeMaterials"))
+	UPROPERTY(Category = MaterialSettings, EditAnywhere, BlueprintReadWrite, meta = (EditCondition = "bMergeMaterials"))
 	bool bUseTextureBinning;
-			
+
 	/** Whether to attempt to re-use the source mesh's lightmap UVs when baking the material or always generate a new set. */
 	UPROPERTY(EditAnywhere, Category = MaterialSettings)
 	bool bReuseMeshLightmapUVs;
@@ -544,7 +544,7 @@ struct FSectionInfo
 	/** Material used by the section */
 	class UMaterialInterface* Material;
 	/** Name value for the section */
-	FName MaterialSlotName;	
+	FName MaterialSlotName;
 	/** List of properties enabled for the section (collision, cast shadow etc) */
 	TArray<FName> EnabledProperties;
 	/** Original index of Material in the source data */
@@ -574,7 +574,7 @@ enum class EMeshInstancingReplacementMethod
 };
 
 /** Mesh instance-replacement settings */
-USTRUCT()
+USTRUCT(Blueprintable)
 struct FMeshInstancingSettings
 {
 	GENERATED_BODY()
@@ -588,28 +588,28 @@ struct FMeshInstancingSettings
 	{}
 
 	/** The actor class to attach new instance static mesh components to */
-	UPROPERTY(EditAnywhere, NoClear, Category="Instancing")
+	UPROPERTY(BlueprintReadWrite, EditAnywhere, NoClear, Category="Instancing")
 	TSubclassOf<AActor> ActorClassToUse;
 
 	/** The number of static mesh instances needed before a mesh is replaced with an instanced version */
-	UPROPERTY(EditAnywhere, Category="Instancing", meta=(ClampMin=1))
+	UPROPERTY(BlueprintReadWrite, EditAnywhere, Category="Instancing", meta=(ClampMin=1))
 	int32 InstanceReplacementThreshold;
 
 	/** How to replace the original actors when instancing */
-	UPROPERTY(EditAnywhere, Category="Instancing")
+	UPROPERTY(BlueprintReadWrite, EditAnywhere, Category="Instancing")
 	EMeshInstancingReplacementMethod MeshReplacementMethod;
 
-	/** 
+	/**
 	 * Whether to skip the conversion to an instanced static mesh for meshes with vertex colors.
 	 * Instanced static meshes do not support vertex colors per-instance, so conversion will lose
 	 * this data.
 	 */
-	UPROPERTY(EditAnywhere, Category="Instancing")
+	UPROPERTY(BlueprintReadWrite, EditAnywhere, Category="Instancing")
 	bool bSkipMeshesWithVertexColors;
 
-	/** 
+	/**
 	 * Whether split up instanced static mesh components based on their intersection with HLOD volumes
 	 */
-	UPROPERTY(EditAnywhere, Category="Instancing", meta=(DisplayName="Use HLOD Volumes"))
+	UPROPERTY(BlueprintReadWrite, EditAnywhere, Category="Instancing", meta=(DisplayName="Use HLOD Volumes"))
 	bool bUseHLODVolumes;
 };
\ No newline at end of file
