Make composite curve tables update correctly when undoing and redoing actions.

diff --git a/Engine/Source/Runtime/Engine/Classes/Engine/CompositeCurveTable.h b/Engine/Source/Runtime/Engine/Classes/Engine/CompositeCurveTable.h
--- a/Engine/Source/Runtime/Engine/Classes/Engine/CompositeCurveTable.h
+++ b/Engine/Source/Runtime/Engine/Classes/Engine/CompositeCurveTable.h
@@ -26,6 +26,8 @@ class UCompositeCurveTable
 
 #if WITH_EDITOR
 	virtual void PostEditChangeProperty(FPropertyChangedEvent& PropertyChangedEvent) override;
+
+	virtual void PostEditUndo() override;
 #endif // WITH_EDITOR
 
 	ENGINE_API virtual void EmptyTable() override;
diff --git a/Engine/Source/Runtime/Engine/Private/CompositeCurveTable.cpp b/Engine/Source/Runtime/Engine/Private/CompositeCurveTable.cpp
--- a/Engine/Source/Runtime/Engine/Private/CompositeCurveTable.cpp
+++ b/Engine/Source/Runtime/Engine/Private/CompositeCurveTable.cpp
@@ -100,6 +100,13 @@ void UCompositeCurveTable::PostEditChangeProperty(FPropertyChangedEvent& Propert
 
 	Super::PostEditChangeProperty(PropertyChangedEvent);
 }
+
+void UCompositeCurveTable::PostEditUndo()
+{
+	Super::PostEditUndo();
+
+	OnParentTablesUpdated();
+}
 #endif // WITH_EDITOR
 
 void UCompositeCurveTable::OnParentTablesUpdated()
