Fix a bug where adding or removing variables from a struct would put composite tables using that struct into a bad state.

diff --git a/Engine/Source/Runtime/Engine/Classes/Engine/CompositeDataTable.h b/Engine/Source/Runtime/Engine/Classes/Engine/CompositeDataTable.h
--- a/Engine/Source/Runtime/Engine/Classes/Engine/CompositeDataTable.h
+++ b/Engine/Source/Runtime/Engine/Classes/Engine/CompositeDataTable.h
@@ -52,6 +52,9 @@ class UCompositeDataTable
 	ENGINE_API virtual void Serialize(FArchive& Ar) override;
 
 #if WITH_EDITOR
+	ENGINE_API virtual void CleanBeforeStructChange() override;
+	ENGINE_API virtual void RestoreAfterStructChange() override;
+
 	virtual void PostEditChangeProperty(FPropertyChangedEvent& PropertyChangedEvent) override;
 #endif // WITH_EDITOR
 
@@ -61,6 +64,10 @@ protected:
 	// Returns a pointer to the first table found that depends on itself if a loop exists. Returns nullptr if no loops are found.
 	const UCompositeDataTable* FindLoops(TArray<const UCompositeDataTable*> AlreadySeenTables) const;
 
+	// Empties the table
+	// if bClearParentTables is false then the row map will be cleared but the parent table array won't be changed
+	void EmptyCompositeTable(bool bClearParentTables);
+
 	void UpdateCachedRowMap();
 
 	void OnParentTablesUpdated();
@@ -68,6 +75,9 @@ protected:
 	// true if this asset is currently being loaded; false otherwise
 	uint8 bIsLoading : 1;
 
+	// if this is true then the parent table array will not be cleared when EmptyTable is called
+	uint8 bShouldNotClearParentTablesOnEmpty : 1;
+
 	// temporary copy used to detect changes so we can update delegates correctly on removal
 	UPROPERTY(transient)
 	TArray<UDataTable*> OldParentTables;
diff --git a/Engine/Source/Runtime/Engine/Classes/Engine/DataTable.h b/Engine/Source/Runtime/Engine/Classes/Engine/DataTable.h
--- a/Engine/Source/Runtime/Engine/Classes/Engine/DataTable.h
+++ b/Engine/Source/Runtime/Engine/Classes/Engine/DataTable.h
@@ -120,7 +120,7 @@ public:
 	UPROPERTY()
 	FName RowStructName;
 
-private:
+protected:
 	/** When RowStruct is being modified, row data is stored serialized with tags */
 	UPROPERTY(Transient)
 	TArray<uint8> RowsSerializedWithTags;
@@ -282,8 +282,8 @@ public:
 	ENGINE_API virtual void AddRow(FName RowName, const FTableRowBase& RowData);
 
 #if WITH_EDITOR
-	ENGINE_API void CleanBeforeStructChange();
-	ENGINE_API void RestoreAfterStructChange();
+	ENGINE_API virtual void CleanBeforeStructChange();
+	ENGINE_API virtual void RestoreAfterStructChange();
 
 	/** Output entire contents of table as a string */
 	ENGINE_API FString GetTableAsString(const EDataTableExportFlags InDTExportFlags = EDataTableExportFlags::None) const;
@@ -338,7 +338,7 @@ public:
 
 	//~ End UDataTable Interface
 
-private:
+protected:
 	void SaveStructData(FArchive& Ar);
 	void LoadStructData(FArchive& Ar);
 
diff --git a/Engine/Source/Runtime/Engine/Private/CompositeDataTable.cpp b/Engine/Source/Runtime/Engine/Private/CompositeDataTable.cpp
--- a/Engine/Source/Runtime/Engine/Private/CompositeDataTable.cpp
+++ b/Engine/Source/Runtime/Engine/Private/CompositeDataTable.cpp
@@ -21,6 +21,7 @@ UCompositeDataTable::UCompositeDataTable(const FObjectInitializer& ObjectInitial
 	: Super(ObjectInitializer)
 {
 	bIsLoading = false;
+	bShouldNotClearParentTablesOnEmpty = false;
 }
 
 void UCompositeDataTable::GetPreloadDependencies(TArray<UObject*>& OutDeps)
@@ -148,14 +149,19 @@ const UCompositeDataTable* UCompositeDataTable::FindLoops(TArray<const UComposit
 
 void UCompositeDataTable::EmptyTable()
 {
-	if (!bIsLoading)
+	EmptyCompositeTable(!bIsLoading && !bShouldNotClearParentTablesOnEmpty);
+}
+
+void UCompositeDataTable::EmptyCompositeTable(bool bClearParentTables)
+{
+	if (bClearParentTables)
 	{
 		ParentTables.Empty();
 	}
 #if WITH_EDITORONLY_DATA
 	RowSourceMap.Empty();
 #endif
-	Super::EmptyTable();
+	UDataTable::EmptyTable();
 }
 
 void UCompositeDataTable::RemoveRow(FName RowName)
@@ -179,6 +185,22 @@ void UCompositeDataTable::Serialize(FArchive& Ar)
 }
 
 #if WITH_EDITOR
+void UCompositeDataTable::CleanBeforeStructChange()
+{
+	bShouldNotClearParentTablesOnEmpty = true;
+	Super::CleanBeforeStructChange();
+	bShouldNotClearParentTablesOnEmpty = false;
+}
+
+void UCompositeDataTable::RestoreAfterStructChange()
+{
+	bShouldNotClearParentTablesOnEmpty = true;
+	Super::RestoreAfterStructChange();
+	bShouldNotClearParentTablesOnEmpty = false;
+
+	UpdateCachedRowMap();
+}
+
 void UCompositeDataTable::PostEditChangeProperty(FPropertyChangedEvent& PropertyChangedEvent)
 {
 	static FName Name_ParentTables = GET_MEMBER_NAME_CHECKED(UCompositeDataTable, ParentTables);
