Fix for just cleared flag issues with beam emitters.

diff --git a/Engine/Source/Runtime/Engine/Private/Particles/ParticleBeam2EmitterInstance.cpp b/Engine/Source/Runtime/Engine/Private/Particles/ParticleBeam2EmitterInstance.cpp
--- a/Engine/Source/Runtime/Engine/Private/Particles/ParticleBeam2EmitterInstance.cpp
+++ b/Engine/Source/Runtime/Engine/Private/Particles/ParticleBeam2EmitterInstance.cpp
@@ -659,6 +659,7 @@ void FParticleBeam2EmitterInstance::UpdateBoundingBox(float DeltaTime)
 		// Take scale into account as well
 		FVector Scale = Component->GetComponentTransform().GetScale3D();
 
+		bool bSkipDoubleSpawnUpdate = !SpriteTemplate->bUseLegacySpawningBehavior;
 		// Take each particle into account
 		for (int32 i=0; i<ActiveParticles; i++)
 		{
@@ -676,6 +677,9 @@ void FParticleBeam2EmitterInstance::UpdateBoundingBox(float DeltaTime)
 			FBeamParticleModifierPayloadData* SourceModifier = NULL;
 			FBeamParticleModifierPayloadData* TargetModifier = NULL;
 
+			bool bJustSpawned = (Particle->Flags & STATE_Particle_JustSpawned) != 0;
+			Particle->Flags &= ~STATE_Particle_JustSpawned;
+
 			BeamTypeData->GetDataPointers(this, (const uint8*)Particle, CurrentOffset, 
 				BeamData, InterpolatedPoints, NoiseRate, NoiseDelta, TargetNoisePoints, 
 				NextNoisePoints, TaperValues, NoiseDistanceScale,
@@ -683,8 +687,11 @@ void FParticleBeam2EmitterInstance::UpdateBoundingBox(float DeltaTime)
 
 			// Do linear integrator and update bounding box
 			Particle->OldLocation = Particle->Location;
-			Particle->Location	+= DeltaTime * Particle->Velocity;
-			Particle->Rotation	+= DeltaTime * Particle->RotationRate;
+
+			bool bSkipUpdate = bJustSpawned && bSkipDoubleSpawnUpdate;
+
+			Particle->Location	+= bSkipUpdate ? FVector::ZeroVector : DeltaTime * Particle->Velocity;
+			Particle->Rotation	+= bSkipUpdate ? 0.0f : DeltaTime * Particle->RotationRate;
 			Particle->OldLocation += PositionOffsetThisTick;
 			FVector Size = Particle->Size * Scale;
 			if (bUpdateBox)
@@ -1483,4 +1490,4 @@ UMaterialInterface* FParticleBeam2EmitterInstance::GetCurrentMaterial()
 	}
 	CurrentMaterial = RenderMaterial;
 	return RenderMaterial;
-}
+}
\ No newline at end of file
