Fixed: Skeletal Mesh with more than 75 bones crashes Mobile Preview
Limit ES3.1 to 75 bones like ES2. All ES3.1 feature level platforms use UB for Bones. Project has to set Compat.MAX_GPUSKIN_BONES=75 to support SkelMeshes with more than 75 bones for ES3.1 and ES2.

diff --git a/Engine/Source/Runtime/Engine/Private/GPUSkinVertexFactory.cpp b/Engine/Source/Runtime/Engine/Private/GPUSkinVertexFactory.cpp
--- a/Engine/Source/Runtime/Engine/Private/GPUSkinVertexFactory.cpp
+++ b/Engine/Source/Runtime/Engine/Private/GPUSkinVertexFactory.cpp
@@ -72,14 +72,12 @@ static TAutoConsoleVariable<int32> CVarVelocityTest(
 // These should match USE_BONES_SRV_BUFFER
 static inline bool SupportsBonesBufferSRV(EShaderPlatform Platform)
 {
-	//#todo-rco: Add Metal support
-	return IsFeatureLevelSupported(Platform, ERHIFeatureLevel::SM4) || IsVulkanPlatform(Platform);
+	return IsFeatureLevelSupported(Platform, ERHIFeatureLevel::SM4);
 }
 
 static inline bool SupportsBonesBufferSRV(ERHIFeatureLevel::Type InFeatureLevel)
 {
-	//#todo-rco: Add Metal support
-	return InFeatureLevel > ERHIFeatureLevel::ES3_1 || IsVulkanPlatform(GMaxRHIShaderPlatform);
+	return InFeatureLevel >= ERHIFeatureLevel::SM4;
 }
 // ---
 
