Fix for incorrect position being sent into viewport client from OnEndTouch() handler

diff --git a/Engine/Source/Runtime/Engine/Private/Slate/SceneViewport.cpp b/Engine/Source/Runtime/Engine/Private/Slate/SceneViewport.cpp
--- a/Engine/Source/Runtime/Engine/Private/Slate/SceneViewport.cpp
+++ b/Engine/Source/Runtime/Engine/Private/Slate/SceneViewport.cpp
@@ -779,13 +779,17 @@ FReply FSceneViewport::OnTouchEnded( const FGeometry& MyGeometry, const FPointer
 	// Start a new reply state
 	CurrentReplyState = FReply::Handled();
 
+	FIntPoint CurCursorPos;
+
 	UpdateCachedGeometry(MyGeometry);
 	if (--NumTouches > 0)
 	{
 		UpdateCachedCursorPos(MyGeometry, TouchEvent);
+		CurCursorPos = CachedCursorPos;
 	}
 	else
 	{
+		CurCursorPos = CachedCursorPos;
 		CachedCursorPos = FIntPoint(-1, -1);
 	}
 
@@ -795,7 +799,7 @@ FReply FSceneViewport::OnTouchEnded( const FGeometry& MyGeometry, const FPointer
 		// Switch to the viewport clients world before processing input
 		FScopedConditionalWorldSwitcher WorldSwitcher( ViewportClient );
 
-		if( !ViewportClient->InputTouch( this, TouchEvent.GetUserIndex(), TouchEvent.GetPointerIndex(), ETouchType::Ended, CachedCursorPos, 0.0f, FDateTime::Now(), TouchEvent.GetTouchpadIndex()) )
+		if( !ViewportClient->InputTouch( this, TouchEvent.GetUserIndex(), TouchEvent.GetPointerIndex(), ETouchType::Ended, CurCursorPos, 0.0f, FDateTime::Now(), TouchEvent.GetTouchpadIndex()) )
 		{
 			CurrentReplyState = FReply::Unhandled(); 
 		}
