[NO NON-PII COMMIT MESSAGE]

diff --git a/Engine/Source/Runtime/Engine/Private/AudioDevice.cpp b/Engine/Source/Runtime/Engine/Private/AudioDevice.cpp
--- a/Engine/Source/Runtime/Engine/Private/AudioDevice.cpp
+++ b/Engine/Source/Runtime/Engine/Private/AudioDevice.cpp
@@ -1788,7 +1788,7 @@ void FAudioDevice::RecursiveApplyAdjuster(const FSoundClassAdjuster& InAdjuster,
 	}
 	else
 	{
-		UE_LOG(LogAudio, Warning, TEXT("Sound class '%s' does not exist"), InSoundClass ? *InSoundClass->GetName() : TEXT("<null>"));
+		UE_LOG(LogAudio, Display, TEXT("RecursiveApplyAdjuster failed, likely because we are clearing the level."));
 	}
 }
 
