[NO NON-PII COMMIT MESSAGE]

diff --git a/Engine/Source/Runtime/Engine/Private/AudioDevice.cpp b/Engine/Source/Runtime/Engine/Private/AudioDevice.cpp
--- a/Engine/Source/Runtime/Engine/Private/AudioDevice.cpp
+++ b/Engine/Source/Runtime/Engine/Private/AudioDevice.cpp
@@ -4190,6 +4190,12 @@ bool FAudioDevice::SoundIsAudible(USoundBase* Sound, const UWorld* World, const
 	float DistanceScale = 1.0f;
 	if (bHasAttenuationSettings)
 	{
+		// If we are not using distance-based attenuation, this sound will be audible regardless of distance.
+		if (!AttenuationSettingsToApply->bAttenuate)
+		{
+			return true;
+		}
+
 		DistanceScale = AttenuationSettingsToApply->GetFocusDistanceScale(GetGlobalFocusSettings(), FocusFactor);
 	}
 
