Take destination stride into account when uploading texture data. This fixes corrupted textures on DX12.

diff --git a/Engine/Source/Runtime/Engine/Private/Slate/SlateTextures.cpp b/Engine/Source/Runtime/Engine/Private/Slate/SlateTextures.cpp
--- a/Engine/Source/Runtime/Engine/Private/Slate/SlateTextures.cpp
+++ b/Engine/Source/Runtime/Engine/Private/Slate/SlateTextures.cpp
@@ -160,14 +160,39 @@ void FSlateTexture2DRHIRef::ResizeTexture(uint32 InWidth, uint32 InHeight)
 	}
 }
 
+void FSlateTexture2DRHIRef::SetTextureData(const TArray<uint8>& Bytes)
+{
+	uint32 DstStride = 0;
+	uint8* DstData = (uint8*) RHILockTexture2D(GetTypedResource(), 0, RLM_WriteOnly, DstStride, false);
+
+	const uint32 NumBlocksX = (Width  + GPixelFormats[PixelFormat].BlockSizeX - 1) / GPixelFormats[PixelFormat].BlockSizeX;
+	const uint32 NumBlocksY = (Height + GPixelFormats[PixelFormat].BlockSizeY - 1) / GPixelFormats[PixelFormat].BlockSizeY;
+	const uint32 SrcStride = NumBlocksX * GPixelFormats[PixelFormat].BlockBytes;
+	ensure(SrcStride * NumBlocksY == Bytes.Num());
+
+	if (SrcStride == DstStride)
+	{
+		FMemory::Memcpy(DstData, Bytes.GetData(), Bytes.Num());
+	}
+	else
+	{
+		const uint8* SrcData = Bytes.GetData();
+		for (uint32 Row = 0; Row < NumBlocksY; ++Row)
+		{
+			FMemory::Memcpy(DstData, SrcData, SrcStride);
+			DstData += DstStride;
+			SrcData += SrcStride;
+		}
+	}
+
+	RHIUnlockTexture2D(GetTypedResource(), 0, false);
+}
+
 void FSlateTexture2DRHIRef::UpdateTexture(const TArray<uint8>& Bytes)
 {
 	if (IsInRenderingThread())
 	{
-		uint32 Stride = 0;
-		void* TextureBuffer = RHILockTexture2D(GetTypedResource(), 0, RLM_WriteOnly, Stride, false);
-		FMemory::Memcpy(TextureBuffer, Bytes.GetData(), Bytes.Num());
-		RHIUnlockTexture2D(GetTypedResource(), 0, false);
+		SetTextureData(Bytes);
 	}
 	else
 	{
@@ -175,10 +200,7 @@ void FSlateTexture2DRHIRef::UpdateTexture(const TArray<uint8>& Bytes)
 		FSlateTexture2DRHIRef*, TextureRHIRef, this,
 		const TArray<uint8>&, TextureData, Bytes,
 		{
-			uint32 Stride = 0;
-			void* TextureBuffer = RHILockTexture2D(TextureRHIRef->GetTypedResource(), 0, RLM_WriteOnly, Stride, false);
-			FMemory::Memcpy(TextureBuffer, TextureData.GetData(), TextureData.Num());
-			RHIUnlockTexture2D(TextureRHIRef->GetTypedResource(), 0, false);
+			TextureRHIRef->SetTextureData(TextureData);
 		});
 	}
 }
diff --git a/Engine/Source/Runtime/Engine/Public/Slate/SlateTextures.h b/Engine/Source/Runtime/Engine/Public/Slate/SlateTextures.h
--- a/Engine/Source/Runtime/Engine/Public/Slate/SlateTextures.h
+++ b/Engine/Source/Runtime/Engine/Public/Slate/SlateTextures.h
@@ -91,6 +91,8 @@ protected:
 	/** Height of this texture */
 	uint32 Height;
 private:
+	void SetTextureData(const TArray<uint8>& Bytes);
+
 	/** Texture creation flags for if this texture needs to be recreated dynamically */
 	uint32 TexCreateFlags;
 	/** Data used between ReleaseDynamicRHI and InitDynamicRHI.  May be null if the data is not used */
