Add UberGraphFunctionKey to try and detect mismatch between reflection data and uber graph frame instances, only present when VALIDATE_UBER_GRAPH_PERSISTENT_FRAME is true (currently outside of shipping and test builds)

diff --git a/Engine/Source/Editor/KismetCompiler/Private/KismetCompiler.cpp b/Engine/Source/Editor/KismetCompiler/Private/KismetCompiler.cpp
--- a/Engine/Source/Editor/KismetCompiler/Private/KismetCompiler.cpp
+++ b/Engine/Source/Editor/KismetCompiler/Private/KismetCompiler.cpp
@@ -2147,12 +2147,23 @@ void FKismetCompilerContext::PostcompileFunction(FKismetFunctionContext& Context
 	FinishCompilingFunction(Context);
 }
 
+#if VALIDATE_UBER_GRAPH_PERSISTENT_FRAME
+extern ENGINE_API int32 IncrementUberGraphSerialNumber();
+#endif//VALIDATE_UBER_GRAPH_PERSISTENT_FRAME
+
 /**
  * Handles final post-compilation setup, flags, creates cached values that would normally be set during deserialization, etc...
  */
 void FKismetCompilerContext::FinishCompilingFunction(FKismetFunctionContext& Context)
 {
 	SetCalculatedMetaDataAndFlags( Context.Function, Context.EntryPoint, Schema );
+	
+#if VALIDATE_UBER_GRAPH_PERSISTENT_FRAME
+	if( NewClass->UberGraphFunction == Context.Function )
+	{
+		NewClass->UberGraphFunctionKey = IncrementUberGraphSerialNumber();
+	}
+#endif//VALIDATE_UBER_GRAPH_PERSISTENT_FRAME
 }
 
 void FKismetCompilerContext::SetCalculatedMetaDataAndFlags(UFunction* Function, UK2Node_FunctionEntry* EntryNode, const UEdGraphSchema_K2* K2Schema)
diff --git a/Engine/Source/Runtime/Core/Public/Misc/Build.h b/Engine/Source/Runtime/Core/Public/Misc/Build.h
--- a/Engine/Source/Runtime/Core/Public/Misc/Build.h
+++ b/Engine/Source/Runtime/Core/Public/Misc/Build.h
@@ -285,6 +285,9 @@
 /** Enable UberGraphPersistentFrame feature. It can speed up BP compilation (re-instancing) in editor, but introduce an unnecessary overhead in runtime. */
 #define USE_UBER_GRAPH_PERSISTENT_FRAME 1
 
+/** Enable validation of the Uber Graph's persistent frame's layout, this is useful to detect uber graph frame related corruption */
+#define VALIDATE_UBER_GRAPH_PERSISTENT_FRAME (!(UE_BUILD_SHIPPING || UE_BUILD_TEST)) && USE_UBER_GRAPH_PERSISTENT_FRAME
+
 /** Enable fast calls for event thunks into an event graph that have no parameters  */
 #define UE_BLUEPRINT_EVENTGRAPH_FASTCALLS 1
 
diff --git a/Engine/Source/Runtime/CoreUObject/Private/UObject/ScriptCore.cpp b/Engine/Source/Runtime/CoreUObject/Private/UObject/ScriptCore.cpp
--- a/Engine/Source/Runtime/CoreUObject/Private/UObject/ScriptCore.cpp
+++ b/Engine/Source/Runtime/CoreUObject/Private/UObject/ScriptCore.cpp
@@ -59,6 +59,10 @@ static FAutoConsoleVariableRef CVarVerboseScriptStats(
 struct FPointerToUberGraphFrameCoreUObject
 {
 	uint8* RawPointer;
+
+#if VALIDATE_UBER_GRAPH_PERSISTENT_FRAME
+	uint32 UberGraphFunctionKey;
+#endif//VALIDATE_UBER_GRAPH_PERSISTENT_FRAME
 };
 #endif //USE_UBER_GRAPH_PERSISTENT_FRAME
 
diff --git a/Engine/Source/Runtime/Engine/Classes/Engine/BlueprintGeneratedClass.h b/Engine/Source/Runtime/Engine/Classes/Engine/BlueprintGeneratedClass.h
--- a/Engine/Source/Runtime/Engine/Classes/Engine/BlueprintGeneratedClass.h
+++ b/Engine/Source/Runtime/Engine/Classes/Engine/BlueprintGeneratedClass.h
@@ -89,9 +89,18 @@ public:
 	// WARNING: This struct layout definition repeated in ScriptCore.cpp as
 	// FPointerToUberGraphFrameCoreUObject to work around reflection generation issues:
 	uint8* RawPointer;
+#if VALIDATE_UBER_GRAPH_PERSISTENT_FRAME
+	uint32 UberGraphFunctionKey;
+#endif//VALIDATE_UBER_GRAPH_PERSISTENT_FRAME
+	// WARNING: This struct layout definition repeated in ScriptCore.cpp as
+	// FPointerToUberGraphFrameCoreUObject to work around reflection generation issues
+	//////////////////////////////////////////////////////////////////////////
 
 	FPointerToUberGraphFrame()
 		: RawPointer(nullptr)
+#if VALIDATE_UBER_GRAPH_PERSISTENT_FRAME
+		, UberGraphFunctionKey(0)
+#endif//VALIDATE_UBER_GRAPH_PERSISTENT_FRAME
 	{}
 
 	~FPointerToUberGraphFrame()
@@ -627,6 +636,10 @@ public:
 	UPROPERTY()
 	UFunction* UberGraphFunction;
 
+#if VALIDATE_UBER_GRAPH_PERSISTENT_FRAME
+	uint32 UberGraphFunctionKey;
+#endif//VALIDATE_UBER_GRAPH_PERSISTENT_FRAME
+
 #if WITH_EDITORONLY_DATA
 	// This is a list of event graph call function nodes that are simple (no argument) thunks into the event graph (typically used for animation delegates, etc...)
 	// It is a deprecated list only used for backwards compatibility prior to VER_UE4_SERIALIZE_BLUEPRINT_EVENTGRAPH_FASTCALLS_IN_UFUNCTION.
diff --git a/Engine/Source/Runtime/Engine/Private/BlueprintGeneratedClass.cpp b/Engine/Source/Runtime/Engine/Private/BlueprintGeneratedClass.cpp
--- a/Engine/Source/Runtime/Engine/Private/BlueprintGeneratedClass.cpp
+++ b/Engine/Source/Runtime/Engine/Private/BlueprintGeneratedClass.cpp
@@ -39,6 +39,9 @@ static FAutoConsoleVariableRef CVarBlueprintClusteringEnabled(
 
 UBlueprintGeneratedClass::UBlueprintGeneratedClass(const FObjectInitializer& ObjectInitializer)
 	: Super(ObjectInitializer)
+#if VALIDATE_UBER_GRAPH_PERSISTENT_FRAME
+	, UberGraphFunctionKey(0)
+#endif//VALIDATE_UBER_GRAPH_PERSISTENT_FRAME
 {
 	NumReplicatedProperties = 0;
 	bHasNativizedParent = false;
@@ -1240,6 +1243,9 @@ void UBlueprintGeneratedClass::CreatePersistentUberGraphFrame(UObject* Obj, bool
 					*GetPathNameSafe(UberGraphFunction), *GetPathNameSafe(Obj));
 			}
 			PointerToUberGraphFrame->RawPointer = FrameMemory;
+#if VALIDATE_UBER_GRAPH_PERSISTENT_FRAME
+			PointerToUberGraphFrame->UberGraphFunctionKey = UberGraphFunctionKey;
+#endif//VALIDATE_UBER_GRAPH_PERSISTENT_FRAME
 		}
 	}
 
@@ -1425,19 +1431,22 @@ void UBlueprintGeneratedClass::Link(FArchive& Ar, bool bRelinkExistingProperties
 	Super::Link(Ar, bRelinkExistingProperties);
 
 #if USE_UBER_GRAPH_PERSISTENT_FRAME
-	if (UsePersistentUberGraphFrame() && UberGraphFunction)
+	if(UsePersistentUberGraphFrame())
 	{
-		Ar.Preload(UberGraphFunction);
-
-		for (UStructProperty* Property : TFieldRange<UStructProperty>(this, EFieldIteratorFlags::ExcludeSuper))
+		if (UberGraphFunction)
 		{
-			if (Property->GetFName() == GetUberGraphFrameName())
+			Ar.Preload(UberGraphFunction);
+
+			for (UStructProperty* Property : TFieldRange<UStructProperty>(this, EFieldIteratorFlags::ExcludeSuper))
 			{
-				UberGraphFramePointerProperty = Property;
-				break;
+				if (Property->GetFName() == GetUberGraphFrameName())
+				{
+					UberGraphFramePointerProperty = Property;
+					break;
+				}
 			}
+			checkSlow(UberGraphFramePointerProperty);
 		}
-		checkSlow(UberGraphFramePointerProperty);
 	}
 #endif
 
@@ -1449,6 +1458,9 @@ void UBlueprintGeneratedClass::PurgeClass(bool bRecompilingOnLoad)
 	Super::PurgeClass(bRecompilingOnLoad);
 
 	UberGraphFunction = NULL;
+#if VALIDATE_UBER_GRAPH_PERSISTENT_FRAME
+	UberGraphFunctionKey = 0;
+#endif//VALIDATE_UBER_GRAPH_PERSISTENT_FRAME
 #if WITH_EDITORONLY_DATA
 	OverridenArchetypeForCDO = NULL;
 
@@ -1482,6 +1494,15 @@ void UBlueprintGeneratedClass::AddReferencedObjectsInUbergraphFrame(UObject* InT
 				checkSlow(PointerToUberGraphFrame)
 				if (PointerToUberGraphFrame->RawPointer)
 				{
+#if VALIDATE_UBER_GRAPH_PERSISTENT_FRAME
+					ensureMsgf(
+						PointerToUberGraphFrame->UberGraphFunctionKey == BPGC->UberGraphFunctionKey,
+						TEXT("Detected key mismatch in uber graph frame for instance %s of type %s, iteration will be unsafe"),
+						*InThis->GetPathName(),
+						*BPGC->GetPathName()
+					);
+#endif//VALIDATE_UBER_GRAPH_PERSISTENT_FRAME
+
 					checkSlow(BPGC->UberGraphFunction);
 					FVerySlowReferenceCollectorArchiveScope CollectorScope(
 						Collector.GetInternalPersistentFrameReferenceCollectorArchive(),
@@ -1522,8 +1543,24 @@ bool UBlueprintGeneratedClass::UsePersistentUberGraphFrame()
 #endif
 }
 
+#if VALIDATE_UBER_GRAPH_PERSISTENT_FRAME
+static TAtomic<int32> GUberGraphSerialNumber(0);
+
+ENGINE_API int32 IncrementUberGraphSerialNumber()
+{
+	return ++GUberGraphSerialNumber;
+}
+#endif//VALIDATE_UBER_GRAPH_PERSISTENT_FRAME
+
 void UBlueprintGeneratedClass::Serialize(FArchive& Ar)
 {
+#if VALIDATE_UBER_GRAPH_PERSISTENT_FRAME
+	if (Ar.IsLoading() && 0 == (Ar.GetPortFlags() & PPF_Duplicate))
+	{
+		UberGraphFunctionKey = IncrementUberGraphSerialNumber();
+	}
+#endif//VALIDATE_UBER_GRAPH_PERSISTENT_FRAME
+
 	Super::Serialize(Ar);
 
 	if (Ar.IsLoading() && 0 == (Ar.GetPortFlags() & PPF_Duplicate))
