Fix display name for VOIP Sample Rate in audio settings. #jira none #rb none

diff --git a/Engine/Source/Runtime/Engine/Classes/Sound/AudioSettings.h b/Engine/Source/Runtime/Engine/Classes/Sound/AudioSettings.h
--- a/Engine/Source/Runtime/Engine/Classes/Sound/AudioSettings.h
+++ b/Engine/Source/Runtime/Engine/Classes/Sound/AudioSettings.h
@@ -124,7 +124,7 @@ class ENGINE_API UAudioSettings : public UDeveloperSettings
 	FSoftObjectPath VoiPSoundClass;
 
 	/** Sample rate used for voice over IP. VOIP audio is resampled to the application's sample rate on the receiver side. */
-	UPROPERTY(config, EditAnywhere, Category = "Audio")
+	UPROPERTY(config, EditAnywhere, Category = "Audio", meta = (DisplayName = "VOIP Sample Rate"))
 	EVoiceSampleRate VoiPSampleRate;
 
 	/** The amount of time to buffer incoming voice audio for ahead of time. Increasing this value can help avoid jitter on slower network connections. */
