Clean up logs during audio cook. #jira none #rb none

diff --git a/Engine/Source/Runtime/Engine/Private/AudioDerivedData.cpp b/Engine/Source/Runtime/Engine/Private/AudioDerivedData.cpp
--- a/Engine/Source/Runtime/Engine/Private/AudioDerivedData.cpp
+++ b/Engine/Source/Runtime/Engine/Private/AudioDerivedData.cpp
@@ -864,7 +864,7 @@ static void CookSimpleWave(USoundWave* SoundWave, FName FormatName, const IAudio
 		// parse the wave data
 		if( !WaveInfo.ReadWaveHeader( RawWaveData, RawDataSize, 0 ) )
 		{
-			UE_LOG(LogTemp, Warning, TEXT( "Only mono or stereo 16 bit waves allowed: %s (%d bytes)" ), *SoundWave->GetFullName(), RawDataSize );
+			UE_LOG(LogAudioDerivedData, Warning, TEXT( "Only mono or stereo 16 bit waves allowed: %s (%d bytes)" ), *SoundWave->GetFullName(), RawDataSize );
 		}
 		else
 		{
@@ -896,10 +896,6 @@ static void CookSimpleWave(USoundWave* SoundWave, FName FormatName, const IAudio
 				WaveInfo.SampleDataSize = TotalDataSize;
 			}
 		}
-		else
-		{
-			UE_LOG(LogTemp, Log, TEXT("Couldn't find amount to resample to."));
-		}
 
 		FSoundQualityInfo QualityInfo = { 0 };
 		if (CompressionOverrides)
@@ -908,7 +904,7 @@ static void CookSimpleWave(USoundWave* SoundWave, FName FormatName, const IAudio
 			if (ModifiedCompressionQuality >= 1.0f)
 			{
 				QualityInfo.Quality = FMath::FloorToInt(ModifiedCompressionQuality);
-				UE_LOG(LogTemp, Display, TEXT("Compression Quality for %s will be modified from %d to %d."), *SoundWave->GetFullName(), SoundWave->CompressionQuality, QualityInfo.Quality);
+				UE_LOG(LogAudioDerivedData, Display, TEXT("Compression Quality for %s will be modified from %d to %d."), *SoundWave->GetFullName(), SoundWave->CompressionQuality, QualityInfo.Quality);
 			}
 			else
 			{
