Fixed: Crash launching project with Compat.MAX_GPUSKIN_BONES=75 and Preview Rendering set to Shader Model 5
Skeletal meshes that are not supported will not be rendered, instead of crashing. Log a warning with information about unsupported mesh.
Do not implicitly switch un-supported meshes to CPU skinning.

diff --git a/Engine/Source/Runtime/Engine/Private/Components/SkinnedMeshComponent.cpp b/Engine/Source/Runtime/Engine/Private/Components/SkinnedMeshComponent.cpp
--- a/Engine/Source/Runtime/Engine/Private/Components/SkinnedMeshComponent.cpp
+++ b/Engine/Source/Runtime/Engine/Private/Components/SkinnedMeshComponent.cpp
@@ -416,7 +416,8 @@ FPrimitiveSceneProxy* USkinnedMeshComponent::CreateSceneProxy()
 	{
 		// Only create a scene proxy if the bone count being used is supported, or if we don't have a skeleton (this is the case with destructibles)
 		int32 MaxBonesPerChunk = SkelMeshRenderData->GetMaxBonesPerSection();
-		if (MaxBonesPerChunk <= GetFeatureLevelMaxNumberOfBones(SceneFeatureLevel))
+		int32 MaxSupportedNumBones = MeshObject->IsCPUSkinned() ? MAX_int32 : GetFeatureLevelMaxNumberOfBones(SceneFeatureLevel);
+		if (MaxBonesPerChunk <= MaxSupportedNumBones)
 		{
 			Result = ::new FSkeletalMeshSceneProxy(this, SkelMeshRenderData);
 		}
@@ -511,20 +512,30 @@ void USkinnedMeshComponent::CreateRenderState_Concurrent()
 #endif
 
 			// Also check if skeletal mesh has too many bones/chunk for GPU skinning.
-			const bool bIsCPUSkinned = SkelMeshRenderData->RequiresCPUSkinning(SceneFeatureLevel) || ShouldCPUSkin();
 			if (bRenderStatic)
 			{
 				// GPU skin vertex buffer + LocalVertexFactory
 				MeshObject = ::new FSkeletalMeshObjectStatic(this, SkelMeshRenderData, SceneFeatureLevel); 
 			}
-			else if(bIsCPUSkinned)
+			else if(ShouldCPUSkin())
 			{
 				MeshObject = ::new FSkeletalMeshObjectCPUSkin(this, SkelMeshRenderData, SceneFeatureLevel);
 			}
-			else
+			// don't silently enable CPU skinning for unsupported meshes, just do not render them, so their absence can be noticed and fixed
+			else if (!SkelMeshRenderData->RequiresCPUSkinning(SceneFeatureLevel)) 
 			{
 				MeshObject = ::new FSkeletalMeshObjectGPUSkin(this, SkelMeshRenderData, SceneFeatureLevel);
 			}
+			else
+			{
+				int32 MaxBonesPerChunk = SkelMeshRenderData->GetMaxBonesPerSection();
+				int32 MaxSupportedGPUSkinBones = FMath::Min(GetFeatureLevelMaxNumberOfBones(SceneFeatureLevel), FGPUBaseSkinVertexFactory::GetMaxGPUSkinBones());
+				bool bHasExtraBoneInfluences = SkelMeshRenderData->HasExtraBoneInfluences();
+				FString FeatureLevelName; GetFeatureLevelName(SceneFeatureLevel, FeatureLevelName);
+
+				UE_LOG(LogSkinnedMeshComp, Warning, TEXT("SkeletalMesh %s, is not supported for current feature level (%s) and will not be rendered. NumBones %d (supported %d), HasExtraBoneInfluences: %s"), 
+					*GetNameSafe(SkeletalMesh), *FeatureLevelName, MaxBonesPerChunk, MaxSupportedGPUSkinBones, bHasExtraBoneInfluences ? TEXT("true"):TEXT("false"));
+			}
 
 			//Allow the editor a chance to manipulate it before its added to the scene
 			PostInitMeshObject(MeshObject);
diff --git a/Engine/Source/Runtime/Engine/Private/SkeletalMeshRenderData.cpp b/Engine/Source/Runtime/Engine/Private/SkeletalMeshRenderData.cpp
--- a/Engine/Source/Runtime/Engine/Private/SkeletalMeshRenderData.cpp
+++ b/Engine/Source/Runtime/Engine/Private/SkeletalMeshRenderData.cpp
@@ -186,7 +186,7 @@ bool FSkeletalMeshRenderData::HasExtraBoneInfluences() const
 
 bool FSkeletalMeshRenderData::RequiresCPUSkinning(ERHIFeatureLevel::Type FeatureLevel) const
 {
-	const int32 MaxGPUSkinBones = GetFeatureLevelMaxNumberOfBones(FeatureLevel);
+	const int32 MaxGPUSkinBones = FMath::Min(GetFeatureLevelMaxNumberOfBones(FeatureLevel), FGPUBaseSkinVertexFactory::GetMaxGPUSkinBones());
 	const int32 MaxBonesPerChunk = GetMaxBonesPerSection();
 	// Do CPU skinning if we need too many bones per chunk, or if we have too many influences per vertex on lower end
 	return (MaxBonesPerChunk > MaxGPUSkinBones) || (HasExtraBoneInfluences() && FeatureLevel < ERHIFeatureLevel::ES3_1);
diff --git a/Engine/Source/Runtime/Engine/Private/SkeletalRenderStatic.h b/Engine/Source/Runtime/Engine/Private/SkeletalRenderStatic.h
--- a/Engine/Source/Runtime/Engine/Private/SkeletalRenderStatic.h
+++ b/Engine/Source/Runtime/Engine/Private/SkeletalRenderStatic.h
@@ -30,7 +30,7 @@ public:
 	//virtual void UpdateRecomputeTangent(int32 MaterialIndex, int32 LODIndex, bool bRecomputeTangent) override {};
 	virtual void EnableOverlayRendering(bool bEnabled, const TArray<int32>* InBonesOfInterest, const TArray<UMorphTarget*>* InMorphTargetOfInterest) override {};
 	virtual void CacheVertices(int32 LODIndex, bool bForce) const override {};
-	virtual bool IsCPUSkinned() const override { return false; }
+	virtual bool IsCPUSkinned() const override { return true; }
 	virtual const FVertexFactory* GetSkinVertexFactory(const FSceneView* View, int32 LODIndex, int32 ChunkIdx) const override;
 	virtual TArray<FTransform>* GetComponentSpaceTransforms() const override;
 	virtual const TArray<FMatrix>& GetReferenceToLocalMatrices() const override;
