Reverted VK SRV changes from CL# 4478372. SkeletalMesh bones will always use SRVs on vulkan

diff --git a/Engine/Source/Runtime/Engine/Private/GPUSkinVertexFactory.cpp b/Engine/Source/Runtime/Engine/Private/GPUSkinVertexFactory.cpp
--- a/Engine/Source/Runtime/Engine/Private/GPUSkinVertexFactory.cpp
+++ b/Engine/Source/Runtime/Engine/Private/GPUSkinVertexFactory.cpp
@@ -72,17 +72,18 @@ static TAutoConsoleVariable<int32> CVarVelocityTest(
 // These should match USE_BONES_SRV_BUFFER
 static inline bool SupportsBonesBufferSRV(EShaderPlatform Platform)
 {
-	return IsFeatureLevelSupported(Platform, ERHIFeatureLevel::SM4);
+	//#todo-rco: Add Metal support
+	return IsFeatureLevelSupported(Platform, ERHIFeatureLevel::SM4) || IsVulkanPlatform(Platform);
 }
 
 static inline bool SupportsBonesBufferSRV(ERHIFeatureLevel::Type InFeatureLevel)
 {
-	return InFeatureLevel >= ERHIFeatureLevel::SM4;
+	//#todo-rco: Add Metal support
+	return InFeatureLevel > ERHIFeatureLevel::ES3_1 || IsVulkanPlatform(GMaxRHIShaderPlatform);
 }
 // ---
 
 
-
 /*-----------------------------------------------------------------------------
  FSharedPoolPolicyData
  -----------------------------------------------------------------------------*/
