HTML5 - fix audio stuttering

#rnx

diff --git a/Engine/Source/Runtime/Engine/Private/AudioDevice.cpp b/Engine/Source/Runtime/Engine/Private/AudioDevice.cpp
--- a/Engine/Source/Runtime/Engine/Private/AudioDevice.cpp
+++ b/Engine/Source/Runtime/Engine/Private/AudioDevice.cpp
@@ -4674,8 +4674,11 @@ void FAudioDevice::Flush(UWorld* WorldToFlush, bool bClearActivatedReverb)
 	// Immediately stop all pending active sounds
 	ProcessingPendingActiveSoundStops(WorldToFlush == nullptr || WorldToFlush->bIsTearingDown);
 
-	// Make sure any in-flight audio rendering commands get executed.
-	FlushAudioRenderingCommands();
+	if (FPlatformProcess::SupportsMultithreading()) // UE-64619 & UE-65245
+	{
+		// Make sure any in-flight audio rendering commands get executed.
+		FlushAudioRenderingCommands();
+	}
 
 	// Anytime we flush, make sure to clear all the listeners.  We'll get the right ones soon enough.
 	Listeners.Reset();
@@ -5323,4 +5326,4 @@ bool FAudioDevice::ShouldUseAttenuation(const UWorld* World) const
 	// - we are forcing the use of attenuation (e.g. for some editors)
 	const bool bIsInGameWorld = World ? World->IsGameWorld() : true;
 	return (bIsInGameWorld || bUseAttenuationForNonGameWorlds);
-}
\ No newline at end of file
+}
