Don't merge bunches of different reliability, since for example a reliable RPC can cause a bunch with properties to become reliable, introducing unnecessary latency for the properties.

diff --git a/Engine/Source/Runtime/Engine/Private/DataChannel.cpp b/Engine/Source/Runtime/Engine/Private/DataChannel.cpp
--- a/Engine/Source/Runtime/Engine/Private/DataChannel.cpp
+++ b/Engine/Source/Runtime/Engine/Private/DataChannel.cpp
@@ -897,6 +897,7 @@ FPacketIdRange UChannel::SendBunch( FOutBunch* Bunch, bool Merge )
 	if
 	(	Merge
 	&&	Connection->LastOut.ChIndex == Bunch->ChIndex
+	&&	Connection->LastOut.bReliable == Bunch->bReliable	// Don't merge bunches of different reliability, since for example a reliable RPC can cause a bunch with properties to become reliable, introducing unnecessary latency for the properties.
 	&&	Connection->AllowMerge
 	&&	Connection->LastEnd.GetNumBits()
 	&&	Connection->LastEnd.GetNumBits()==Connection->SendBuffer.GetNumBits()
@@ -906,7 +907,6 @@ FPacketIdRange UChannel::SendBunch( FOutBunch* Bunch, bool Merge )
 		check(!Connection->LastOut.IsError());
 		PreExistingBits = Connection->LastOut.GetNumBits();
 		Connection->LastOut.SerializeBits( Bunch->GetData(), Bunch->GetNumBits() );
-		Connection->LastOut.bReliable |= Bunch->bReliable;
 		Connection->LastOut.bOpen     |= Bunch->bOpen;
 		Connection->LastOut.bClose    |= Bunch->bClose;
 		OutBunch                       = Connection->LastOutBunch;
