Fall back to spawning a base APlayerState if the game mode's PlayerStateClass is null in AController::InitPlayerState.
Return an error message from AGameModeBase::InitNewPlayer if the new player's PlayerState is null instead of asserting.

diff --git a/Engine/Source/Runtime/Engine/Private/Controller.cpp b/Engine/Source/Runtime/Engine/Private/Controller.cpp
--- a/Engine/Source/Runtime/Engine/Private/Controller.cpp
+++ b/Engine/Source/Runtime/Engine/Private/Controller.cpp
@@ -511,7 +511,15 @@ void AController::InitPlayerState()
 			SpawnInfo.Instigator = Instigator;
 			SpawnInfo.SpawnCollisionHandlingOverride = ESpawnActorCollisionHandlingMethod::AlwaysSpawn;
 			SpawnInfo.ObjectFlags |= RF_Transient;	// We never want player states to save into a map
-			PlayerState = World->SpawnActor<APlayerState>(GameMode->PlayerStateClass, SpawnInfo );
+
+			TSubclassOf<APlayerState> PlayerStateClassToSpawn = GameMode->PlayerStateClass;
+			if (PlayerStateClassToSpawn.Get() == nullptr)
+			{
+				UE_LOG(LogPlayerController, Log, TEXT("AController::InitPlayerState: the PlayerStateClass of game mode %s is null, falling back to APlayerState."), *GameMode->GetName());
+				PlayerStateClassToSpawn = APlayerState::StaticClass();
+			}
+
+			PlayerState = World->SpawnActor<APlayerState>(PlayerStateClassToSpawn, SpawnInfo);
 	
 			// force a default player name if necessary
 			if (PlayerState && PlayerState->GetPlayerName().IsEmpty())
diff --git a/Engine/Source/Runtime/Engine/Private/GameModeBase.cpp b/Engine/Source/Runtime/Engine/Private/GameModeBase.cpp
--- a/Engine/Source/Runtime/Engine/Private/GameModeBase.cpp
+++ b/Engine/Source/Runtime/Engine/Private/GameModeBase.cpp
@@ -727,6 +727,12 @@ FString AGameModeBase::InitNewPlayer(APlayerController* NewPlayerController, con
 {
 	check(NewPlayerController);
 
+	// The player needs a PlayerState to register successfully
+	if (NewPlayerController->PlayerState == nullptr)
+	{
+		return FString(TEXT("PlayerState is null"));
+	}
+
 	FString ErrorMessage;
 
 	// Register the player with the session
