Add variable for small errors that result in being slightly inside the object.

diff --git a/Engine/Source/Runtime/Engine/Private/Components/PrimitiveComponent.cpp b/Engine/Source/Runtime/Engine/Private/Components/PrimitiveComponent.cpp
--- a/Engine/Source/Runtime/Engine/Private/Components/PrimitiveComponent.cpp
+++ b/Engine/Source/Runtime/Engine/Private/Components/PrimitiveComponent.cpp
@@ -69,6 +69,14 @@ static FAutoConsoleVariableRef CVarInitialOverlapTolerance(
 	TEXT("Dot product of movement direction and surface normal."),
 	ECVF_Default);
 
+static float HitDistanceToleranceCVar = 0.0f;
+static FAutoConsoleVariableRef CVarHitDistanceTolerance(
+	TEXT("p.HitDistanceTolerance"),
+	HitDistanceToleranceCVar,
+	TEXT("Tolerance for hit distance for overlap test in PrimitiveComponent movement.\n")
+	TEXT("Hits that are less than this distance are ignored."),
+	ECVF_Default);
+
 #if !(UE_BUILD_SHIPPING || UE_BUILD_TEST)
 int32 CVarShowInitialOverlaps = 0;
 FAutoConsoleVariableRef CVarRefShowInitialOverlaps(
@@ -1735,7 +1743,7 @@ static bool ShouldIgnoreHitResult(const UWorld* InWorld, FHitResult const& TestH
 	
 		// If we started penetrating, we may want to ignore it if we are moving out of penetration.
 		// This helps prevent getting stuck in walls.
-		if ( TestHit.bStartPenetrating && !(MoveFlags & MOVECOMP_NeverIgnoreBlockingOverlaps) )
+		if ( (TestHit.Distance < HitDistanceToleranceCVar || TestHit.bStartPenetrating) && !(MoveFlags & MOVECOMP_NeverIgnoreBlockingOverlaps) )
 		{
  			const float DotTolerance = InitialOverlapToleranceCVar;
 
