Fix for taking HDR screenshots when the buffer write is disabled

diff --git a/Engine/Source/Runtime/Engine/Private/HighResScreenshot.cpp b/Engine/Source/Runtime/Engine/Private/HighResScreenshot.cpp
--- a/Engine/Source/Runtime/Engine/Private/HighResScreenshot.cpp
+++ b/Engine/Source/Runtime/Engine/Private/HighResScreenshot.cpp
@@ -67,7 +67,9 @@ void FHighResScreenshotConfig::PopulateImageTaskParams(FImageWriteTask& InOutTas
 {
 	static const TConsoleVariableData<int32>* CVarDumpFramesAsHDR = IConsoleManager::Get().FindTConsoleVariableDataInt(TEXT("r.BufferVisualizationDumpFramesAsHDR"));
 
-	const bool bLocalCaptureHDR = bCaptureHDR || CVarDumpFramesAsHDR->GetValueOnAnyThread();
+	const bool bCaptureHDREnabledInUI = bCaptureHDR && bDumpBufferVisualizationTargets;
+
+	const bool bLocalCaptureHDR = bCaptureHDREnabledInUI || CVarDumpFramesAsHDR->GetValueOnAnyThread();
 
 	InOutTask.Format = bLocalCaptureHDR ? EImageFormat::EXR : EImageFormat::PNG;
 
