Remove overzealous ensure

diff --git a/Engine/Source/Runtime/Engine/Private/BlueprintGeneratedClass.cpp b/Engine/Source/Runtime/Engine/Private/BlueprintGeneratedClass.cpp
--- a/Engine/Source/Runtime/Engine/Private/BlueprintGeneratedClass.cpp
+++ b/Engine/Source/Runtime/Engine/Private/BlueprintGeneratedClass.cpp
@@ -812,8 +812,7 @@ UObject* UBlueprintGeneratedClass::FindArchetype(UClass* ArchetypeClass, const F
 			}
 		}
 	}
-
-	ensure(!Archetype || ArchetypeClass->IsChildOf(Archetype->GetClass()));
+	
 	return Archetype;
 }
 
