UE-65464 Typo in Support Software Occlusion Culling tooltip in Project Settings
#4.21

diff --git a/Engine/Source/Runtime/Engine/Classes/Engine/RendererSettings.h b/Engine/Source/Runtime/Engine/Classes/Engine/RendererSettings.h
--- a/Engine/Source/Runtime/Engine/Classes/Engine/RendererSettings.h
+++ b/Engine/Source/Runtime/Engine/Classes/Engine/RendererSettings.h
@@ -205,7 +205,7 @@ class ENGINE_API URendererSettings : public UDeveloperSettings
 
 	UPROPERTY(config, EditAnywhere, Category = Mobile, meta=(
 		ConsoleVariable="r.Mobile.AllowSoftwareOcclusion", DisplayName="Support Software Occlusion Culling",
-		ToolTip="Whether to support 'Software Occlusion Culling' on mobile platforms. This this package occluder information and enable Software Occlusion Culling.",
+		ToolTip="Whether to support 'Software Occlusion Culling' on mobile platforms. This will package occluder information and enable Software Occlusion Culling.",
 		ConfigRestartRequired=false))
 	uint32 bMobileAllowSoftwareOcclusionCulling:1;
 	
