Fix custom transition results being picked up as alternative roots on compilation

Applied the same changes as to other 'root nodes' to allow proper differentiation of the 'real' root.
Also fixes check() on load of any anim BP with custom transition blends.

diff --git a/Engine/Source/Editor/AnimGraph/Classes/AnimGraphNode_CustomTransitionResult.h b/Engine/Source/Editor/AnimGraph/Classes/AnimGraphNode_CustomTransitionResult.h
--- a/Engine/Source/Editor/AnimGraph/Classes/AnimGraphNode_CustomTransitionResult.h
+++ b/Engine/Source/Editor/AnimGraph/Classes/AnimGraphNode_CustomTransitionResult.h
@@ -4,11 +4,11 @@
 
 #include "CoreMinimal.h"
 #include "UObject/ObjectMacros.h"
-#include "AnimGraphNode_Root.h"
+#include "AnimGraphNode_StateResult.h"
 #include "AnimGraphNode_CustomTransitionResult.generated.h"
 
 UCLASS(MinimalAPI)
-class UAnimGraphNode_CustomTransitionResult : public UAnimGraphNode_Root
+class UAnimGraphNode_CustomTransitionResult : public UAnimGraphNode_StateResult
 {
 	GENERATED_UCLASS_BODY()
 
diff --git a/Engine/Source/Runtime/Core/Public/UObject/ReleaseObjectVersion.h b/Engine/Source/Runtime/Core/Public/UObject/ReleaseObjectVersion.h
--- a/Engine/Source/Runtime/Core/Public/UObject/ReleaseObjectVersion.h
+++ b/Engine/Source/Runtime/Core/Public/UObject/ReleaseObjectVersion.h
@@ -68,6 +68,9 @@ struct CORE_API FReleaseObjectVersion
 		// Change to not clamp RGB values > 1 on linear color curves
 		UnclampRGBColorCurves,
 
+		// Bugfix for FAnimObjectVersion::LinkTimeAnimBlueprintRootDiscovery.
+		LinkTimeAnimBlueprintRootDiscoveryBugFix,
+
 		// -----<new versions can be added above this line>-------------------------------------------------
 		VersionPlusOne,
 		LatestVersion = VersionPlusOne - 1
diff --git a/Engine/Source/Runtime/Engine/Private/Animation/AnimBlueprintGeneratedClass.cpp b/Engine/Source/Runtime/Engine/Private/Animation/AnimBlueprintGeneratedClass.cpp
--- a/Engine/Source/Runtime/Engine/Private/Animation/AnimBlueprintGeneratedClass.cpp
+++ b/Engine/Source/Runtime/Engine/Private/Animation/AnimBlueprintGeneratedClass.cpp
@@ -11,6 +11,7 @@
 #include "Animation/AnimNodeBase.h"
 #include "Animation/AnimInstance.h"
 #include "UObject/AnimObjectVersion.h"
+#include "UObject/ReleaseObjectVersion.h"
 
 /////////////////////////////////////////////////////
 // FStateMachineDebugData
@@ -238,6 +239,7 @@ UAnimBlueprintGeneratedClass::UAnimBlueprintGeneratedClass(const FObjectInitiali
 void UAnimBlueprintGeneratedClass::Serialize(FArchive& Ar)
 {
 	Ar.UsingCustomVersion(FAnimObjectVersion::GUID);
+	Ar.UsingCustomVersion(FReleaseObjectVersion::GUID);
 
 	Super::Serialize(Ar);
 }
@@ -265,7 +267,7 @@ void UAnimBlueprintGeneratedClass::Link(FArchive& Ar, bool bRelinkExistingProper
 					// If we are loading from a newer version, we must verify that there is only one root here.
 					// When linking an older version on load we may find multiple roots until the class is recompiled or saved
 #if DO_CHECK
-					if(Ar.IsLoading() && Ar.CustomVer(FAnimObjectVersion::GUID) >= FAnimObjectVersion::LinkTimeAnimBlueprintRootDiscovery)
+					if(Ar.IsLoading() && Ar.CustomVer(FReleaseObjectVersion::GUID) >= FReleaseObjectVersion::LinkTimeAnimBlueprintRootDiscoveryBugFix)
 					{
 						check(RootAnimNodeProperty == nullptr);
 					}
