Moving over:
"Fixed hang on HLOD proxy generation

diff --git a/Engine/Source/Runtime/Engine/Private/HLODProxy.cpp b/Engine/Source/Runtime/Engine/Private/HLODProxy.cpp
--- a/Engine/Source/Runtime/Engine/Private/HLODProxy.cpp
+++ b/Engine/Source/Runtime/Engine/Private/HLODProxy.cpp
@@ -71,7 +71,7 @@ const AActor* UHLODProxy::FindFirstActor(const ALODActor* LODActor)
 			const ALODActor* SubLODActor = Cast<ALODActor>(FirstActor);
 			if (SubLODActor)
 			{
-				SubLODActor->SubActors.IsValidIndex(0) ? SubLODActor->SubActors[0] : nullptr;
+				FirstActor = SubLODActor->SubActors.IsValidIndex(0) ? SubLODActor->SubActors[0] : nullptr; 
 			}
 			else
 			{
