UE4.21 - Change warning to display
#rnx

diff --git a/Engine/Source/Runtime/Engine/Private/Materials/MaterialShader.cpp b/Engine/Source/Runtime/Engine/Private/Materials/MaterialShader.cpp
--- a/Engine/Source/Runtime/Engine/Private/Materials/MaterialShader.cpp
+++ b/Engine/Source/Runtime/Engine/Private/Materials/MaterialShader.cpp
@@ -1440,7 +1440,7 @@ void FMaterialShaderMap::Compile(
 				*ShaderMapId.GetParameterSetLayerParametersKey()
 			);
   
-			UE_LOG(LogShaders, Warning, TEXT("	%s"), *DebugDescription);
+			UE_LOG(LogShaders, Display, TEXT("	%s"), *DebugDescription);
 #endif
 
 			uint32 NumShaders = 0;
