Fix crash when enabling per poly collision on Skeletal Meshes while "support compute skincache" is turned on.

diff --git a/Engine/Source/Runtime/Engine/Private/SkeletalMeshComponentPhysics.cpp b/Engine/Source/Runtime/Engine/Private/SkeletalMeshComponentPhysics.cpp
--- a/Engine/Source/Runtime/Engine/Private/SkeletalMeshComponentPhysics.cpp
+++ b/Engine/Source/Runtime/Engine/Private/SkeletalMeshComponentPhysics.cpp
@@ -1965,7 +1965,7 @@ void USkeletalMeshComponent::ComputeSkinnedPositions(USkeletalMeshComponent* Com
 					if (ActorData)
 					{
 						const uint32 SoftOffset = Section.GetVertexBufferIndex();
-						const uint32 NumSoftVerts = Section.GetNumVertices();
+						const uint32 NumSoftVerts = FMath::Min(Section.GetNumVertices(), ActorData->Positions.Num());
 
 						// if blend weight is 1.0, doesn't need to blend with a skinned position
 						if (Component->ClothBlendWeight < 1.0f)
