Fixed potential infinite stall when calling UTexture2D::WaitForStreaming().
This happened on RHI which enqueued a command in RHIAsyncReallocateTexture2D().

diff --git a/Engine/Source/Runtime/Engine/Private/Texture2D.cpp b/Engine/Source/Runtime/Engine/Private/Texture2D.cpp
--- a/Engine/Source/Runtime/Engine/Private/Texture2D.cpp
+++ b/Engine/Source/Runtime/Engine/Private/Texture2D.cpp
@@ -493,7 +493,8 @@ void UTexture2D::UpdateResource()
 	// Make sure there are no pending requests in flight.
 	while( UpdateStreamingStatus() == true )
 	{
-		// Give up timeslice.
+		// Force flush the RHI threads to execute all commands issued for texture streaming, and give up timeslice.
+		FlushRenderingCommands();
 		FPlatformProcess::Sleep(0);
 	}
 
@@ -553,7 +554,8 @@ void UTexture2D::WaitForStreaming()
 		// Make sure there are no pending requests in flight otherwise calling UpdateIndividualTexture could be prevented to defined a new requested mip.
 		while (	!IsReadyForStreaming() || UpdateStreamingStatus() ) 
 		{
-			// Give up timeslice.
+			// Force flush the RHI threads to execute all commands issued for texture streaming, and give up timeslice.
+			FlushRenderingCommands();
 			FPlatformProcess::Sleep(0);
 		}
 
@@ -564,7 +566,8 @@ void UTexture2D::WaitForStreaming()
 
 			while (	UpdateStreamingStatus() ) 
 			{
-				// Give up timeslice.
+				// Force flush the RHI threads to execute all commands issued for texture streaming, and give up timeslice.
+				FlushRenderingCommands();
 				FPlatformProcess::Sleep(0);
 			}
 		}
