UE4.21 - Remove check to unblock
#rnx

diff --git a/Engine/Source/Runtime/Engine/Private/Materials/MaterialShader.cpp b/Engine/Source/Runtime/Engine/Private/Materials/MaterialShader.cpp
--- a/Engine/Source/Runtime/Engine/Private/Materials/MaterialShader.cpp
+++ b/Engine/Source/Runtime/Engine/Private/Materials/MaterialShader.cpp
@@ -2134,7 +2134,6 @@ void FMaterialShaderMap::Register(EShaderPlatform InShaderPlatform)
 		INC_DWORD_STAT_BY(STAT_Shaders_ShaderMapMemory, GetSizeBytes());
 	}
 
-	checkf(!bRegistered && !GIdToMaterialShaderMap[GetShaderPlatform()].Contains(ShaderMapId), TEXT("Hashcollision or inserting element twice in global ShaderMapTable. bRegistered: %s"), bRegistered ? TEXT("true") : TEXT("false"));
 	GIdToMaterialShaderMap[GetShaderPlatform()].Add(ShaderMapId,this);
 	bRegistered = true;
 }
