Merging //UE4/Release-Staging-4.21 @ CL 4472944.

diff --git a/Engine/Plugins/Media/WebMMedia/Source/WebMMedia/Private/Player/WebMAudioDecoder.cpp b/Engine/Plugins/Media/WebMMedia/Source/WebMMedia/Private/Player/WebMAudioDecoder.cpp
--- a/Engine/Plugins/Media/WebMMedia/Source/WebMMedia/Private/Player/WebMAudioDecoder.cpp
+++ b/Engine/Plugins/Media/WebMMedia/Source/WebMMedia/Private/Player/WebMAudioDecoder.cpp
@@ -3,16 +3,25 @@
 #include "WebMAudioDecoder.h"
 #include "WebMMediaPrivate.h"
 #include "WebMMediaFrame.h"
+#include "WebMMediaAudioSample.h"
 #include "MediaSamples.h"
 
 THIRD_PARTY_INCLUDES_START
 #include <opus.h>
+#include <vorbis/codec.h>
 THIRD_PARTY_INCLUDES_END
 
+struct FWebMAudioDecoder::FVorbisDecoder
+{
+	vorbis_info Info;
+	vorbis_dsp_state DspState;
+	vorbis_block Block;
+};
+
 FWebMAudioDecoder::FWebMAudioDecoder(TSharedPtr<FMediaSamples, ESPMode::ThreadSafe> InSamples)
 	: Samples(InSamples)
 	, AudioSamplePool(new FWebMMediaAudioSamplePool)
-	, Decoder(nullptr)
+	, OpusDecoder(nullptr)
 {
 }
 
@@ -23,25 +32,153 @@ FWebMAudioDecoder::~FWebMAudioDecoder()
 		FTaskGraphInterface::Get().WaitUntilTaskCompletes(AudioDecodingTask);
 	}
 
-	if (Decoder)
+	if (OpusDecoder)
 	{
-		opus_decoder_destroy(Decoder);
+		opus_decoder_destroy(OpusDecoder);
 	}
 }
 
-void FWebMAudioDecoder::Initialize(int32 InSampleRate, int32 InChannels)
+bool FWebMAudioDecoder::Initialize(const char* InCodec, int32 InSampleRate, int32 InChannels, const uint8* CodecPrivateData, size_t CodecPrivateDataSize)
 {
+	if (OpusDecoder)
+	{
+		opus_decoder_destroy(OpusDecoder);
+		OpusDecoder = nullptr;
+	}
+
+	if (FCStringAnsi::Strcmp(InCodec, "A_OPUS") == 0)
+	{
+		Codec = ESupportedCodecs::Opus;
+	}
+	else if (FCStringAnsi::Strcmp(InCodec, "A_VORBIS") == 0)
+	{
+		Codec = ESupportedCodecs::Vorbis;
+	}
+	else
+	{
+		return false;
+	}
+
 	SampleRate = InSampleRate;
 	Channels = InChannels;
-	// Max supported frame is 120ms
-	FrameSize = 120 * SampleRate / 1000;
 
+	if (Codec == ESupportedCodecs::Opus)
+	{
+		InitializeOpus();
+	}
+	else
+	{
+		if (!InitializeVorbis(CodecPrivateData, CodecPrivateDataSize))
+		{
+			return false;
+		}
+	}
+
+	DecodeBuffer.SetNumUninitialized(FrameSize * Channels * sizeof(int16));
+
+	return true;
+}
+
+void FWebMAudioDecoder::InitializeOpus()
+{
 	int32 ErrorCode = 0;
-	Decoder = opus_decoder_create(SampleRate, Channels, &ErrorCode);
-	check(Decoder);
+	OpusDecoder = opus_decoder_create(SampleRate, Channels, &ErrorCode);
+	check(OpusDecoder);
 	check(ErrorCode == 0);
 
-	DecodeBuffer.SetNumUninitialized(FrameSize * Channels * sizeof(opus_int16));
+	// Max supported frame is 120ms
+	FrameSize = 120 * SampleRate / 1000;
+}
+
+bool FWebMAudioDecoder::InitializeVorbis(const uint8* CodecPrivateData, size_t CodecPrivateDataSize)
+{
+	if (CodecPrivateDataSize < 3 || !CodecPrivateData || CodecPrivateData[0] != 2)
+	{
+		UE_LOG(LogWebMMedia, Display, TEXT("Failed to initialize Vorbis - invalid data"));
+		return false;
+	}
+
+	size_t HeaderSize[3] = { 0 };
+	size_t Offset = 1;
+
+	// Calculate three headers sizes
+	for (int32 i = 0; i < 2; ++i)
+	{
+		for (;;)
+		{
+			if (Offset >= CodecPrivateDataSize)
+			{
+				UE_LOG(LogWebMMedia, Display, TEXT("Failed to initialize Vorbis - invalid offset"));
+				return false;
+			}
+			HeaderSize[i] += CodecPrivateData[Offset];
+			if (CodecPrivateData[Offset++] < 0xFF)
+			{
+				break;
+			}
+		}
+	}
+	HeaderSize[2] = CodecPrivateDataSize - (HeaderSize[0] + HeaderSize[1] + Offset);
+
+	if (HeaderSize[0] + HeaderSize[1] + HeaderSize[2] + Offset != CodecPrivateDataSize)
+	{
+		UE_LOG(LogWebMMedia, Display, TEXT("Failed to initialize Vorbis - parameters mismatch"));
+		return false;
+	}
+
+	ogg_packet Packet[3];
+	memset(Packet, 0, 3 * sizeof(ogg_packet));
+
+	Packet[0].packet = (unsigned char *) CodecPrivateData + Offset;
+	Packet[0].bytes = HeaderSize[0];
+	Packet[0].b_o_s = 1;
+
+	Packet[1].packet = (unsigned char *) CodecPrivateData + Offset + HeaderSize[0];
+	Packet[1].bytes = HeaderSize[1];
+
+	Packet[2].packet = (unsigned char *) CodecPrivateData + Offset + HeaderSize[0] + HeaderSize[1];
+	Packet[2].bytes = HeaderSize[2];
+
+	VorbisDecoder.Reset(new FVorbisDecoder());
+
+	vorbis_info_init(&VorbisDecoder->Info);
+
+	/* Upload three Vorbis headers into libvorbis */
+	vorbis_comment Vc;
+	vorbis_comment_init(&Vc);
+	for (int i = 0; i < 3; ++i)
+	{
+		if (vorbis_synthesis_headerin(&VorbisDecoder->Info, &Vc, &Packet[i]))
+		{
+			UE_LOG(LogWebMMedia, Display, TEXT("Failed to initialize Vorbis - invalid headers"));
+			vorbis_comment_clear(&Vc);
+			return false;
+		}
+	}
+
+	vorbis_comment_clear(&Vc);
+
+	if (vorbis_synthesis_init(&VorbisDecoder->DspState, &VorbisDecoder->Info))
+	{
+		UE_LOG(LogWebMMedia, Display, TEXT("Failed to initialize Vorbis (synthesis)"));
+		return false;
+	}
+
+	if (VorbisDecoder->Info.channels != Channels || VorbisDecoder->Info.rate != SampleRate)
+	{
+		UE_LOG(LogWebMMedia, Display, TEXT("Failed to initialize Vorbis - invalid parameters"));
+		return false;
+	}
+
+	if (vorbis_block_init(&VorbisDecoder->DspState, &VorbisDecoder->Block))
+	{
+		UE_LOG(LogWebMMedia, Display, TEXT("Failed to initialize Vorbis"));
+		return false;
+	}
+
+	FrameSize = 4096 / Channels;
+
+	return true;
 }
 
 void FWebMAudioDecoder::DecodeAudioFramesAsync(const TArray<TSharedPtr<FWebMFrame>>& AudioFrames)
@@ -63,16 +200,79 @@ void FWebMAudioDecoder::DoDecodeAudioFrames(const TArray<TSharedPtr<FWebMFrame>>
 {
 	for (const TSharedPtr<FWebMFrame>& AudioFrame : AudioFrames)
 	{
-		const TSharedRef<FWebMMediaAudioSample, ESPMode::ThreadSafe> AudioSample = AudioSamplePool->AcquireShared();
+		int32 NumOfSamplesDecoded;
+
+		if (Codec == ESupportedCodecs::Opus)
+		{
+			NumOfSamplesDecoded = DecodeOpus(AudioFrame);
+		}
+		else
+		{
+			NumOfSamplesDecoded = DecodeVorbis(AudioFrame);
+		}
+
+		if (NumOfSamplesDecoded > 0)
+		{
+			TSharedRef<FWebMMediaAudioSample, ESPMode::ThreadSafe> AudioSample = AudioSamplePool->AcquireShared();
+			AudioSample->Initialize(DecodeBuffer.GetData(), NumOfSamplesDecoded * Channels * 2, Channels, SampleRate, AudioFrame->Time, FTimespan::FromHours(1));
+			Samples->AddAudio(AudioSample);
+		}
+	}
+}
+
+int32 FWebMAudioDecoder::DecodeOpus(const TSharedPtr<FWebMFrame>& AudioFrame)
+{
+	int32 Result = opus_decode(OpusDecoder, AudioFrame->Data.GetData(), AudioFrame->Data.Num(), (opus_int16*) DecodeBuffer.GetData(), FrameSize, 0);
+	if (Result < 0)
+	{
+		UE_LOG(LogWebMMedia, Display, TEXT("Error decoding Opus audio frame"));
+	}
 
-		int32 NumOfSamplesDecoded = opus_decode(Decoder, AudioFrame->Data.GetData(), AudioFrame->Data.Num(), (opus_int16*) DecodeBuffer.GetData(), FrameSize, 0);
-		if (NumOfSamplesDecoded <= 0)
+	return Result;
+}
+
+int32 FWebMAudioDecoder::DecodeVorbis(const TSharedPtr<FWebMFrame>& AudioFrame)
+{
+	ogg_packet Packet;
+	memset(&Packet, 0, sizeof(ogg_packet));
+
+	Packet.packet = AudioFrame->Data.GetData();
+	Packet.bytes = AudioFrame->Data.Num();
+
+	if (vorbis_synthesis(&VorbisDecoder->Block, &Packet))
+	{
+		UE_LOG(LogWebMMedia, Display, TEXT("Error decoding Vorbis audio frame - vorbis_synthesis failed"));
+		return -1;
+	}
+
+	if (vorbis_synthesis_blockin(&VorbisDecoder->DspState, &VorbisDecoder->Block))
+	{
+		UE_LOG(LogWebMMedia, Display, TEXT("Error decoding Vorbis audio frame - vorbis_synthesis_blockin failed"));
+		return - 1;
+	}
+
+	int16* Buffer = (int16*) DecodeBuffer.GetData();
+	int32 Count = 0;
+	float** Pcm;
+	int32 SamplesCount = vorbis_synthesis_pcmout(&VorbisDecoder->DspState, &Pcm);
+	while (SamplesCount > 0)
+	{
+		int32 ToConvert = SamplesCount <= FrameSize ? SamplesCount : FrameSize;
+		for (int32 Channel = 0; Channel < Channels; ++Channel)
 		{
-			UE_LOG(LogWebMMedia, Display, TEXT("Error decoding audio frame"));
-			return;
+			float *FloatSamples = Pcm[Channel];
+			for (int32 i = 0, j = Channel; i < ToConvert; ++i, j += Channels)
+			{
+				int32 Sample = FMath::Clamp<int32>(FloatSamples[i] * 32767.0f, -32768, 32768);
+				Buffer[Count + j] = Sample;
+			}
 		}
+		vorbis_synthesis_read(&VorbisDecoder->DspState, ToConvert);
 
-		AudioSample->Initialize(DecodeBuffer.GetData(), NumOfSamplesDecoded * Channels * 2, Channels, SampleRate, AudioFrame->Time, FTimespan::FromHours(1));
-		Samples->AddAudio(AudioSample);
+		Count += ToConvert;
+
+		SamplesCount = vorbis_synthesis_pcmout(&VorbisDecoder->DspState, &Pcm);
 	}
+
+	return Count;
 }
diff --git a/Engine/Plugins/Media/WebMMedia/Source/WebMMedia/Private/Player/WebMAudioDecoder.h b/Engine/Plugins/Media/WebMMedia/Source/WebMMedia/Private/Player/WebMAudioDecoder.h
deleted file mode 100644
diff --git a/Engine/Plugins/Media/WebMMedia/Source/WebMMedia/Private/Player/WebMMediaAudioSample.h b/Engine/Plugins/Media/WebMMedia/Source/WebMMedia/Private/Player/WebMMediaAudioSample.h
deleted file mode 100644
diff --git a/Engine/Plugins/Media/WebMMedia/Source/WebMMedia/Private/Player/WebMMediaFrame.h b/Engine/Plugins/Media/WebMMedia/Source/WebMMedia/Private/Player/WebMMediaFrame.h
deleted file mode 100644
diff --git a/Engine/Plugins/Media/WebMMedia/Source/WebMMedia/Private/Player/WebMMediaPlayer.cpp b/Engine/Plugins/Media/WebMMedia/Source/WebMMedia/Private/Player/WebMMediaPlayer.cpp
--- a/Engine/Plugins/Media/WebMMedia/Source/WebMMedia/Private/Player/WebMMediaPlayer.cpp
+++ b/Engine/Plugins/Media/WebMMedia/Source/WebMMedia/Private/Player/WebMMediaPlayer.cpp
@@ -139,7 +139,7 @@ bool FWebMMediaPlayer::Open(const FString& Url, const IMediaOptions* /*Options*/
 
 	if (!MkvRead())
 	{
-		UE_LOG(LogWebMMedia, Error, TEXT("Error parsing MKV file: %s"), *FilePath);
+		UE_LOG(LogWebMMedia, Error, TEXT("Error parsing matroska file: %s"), *FilePath);
 		return false;
 	}
 
@@ -148,12 +148,19 @@ bool FWebMMediaPlayer::Open(const FString& Url, const IMediaOptions* /*Options*/
 
 	if (SelectedAudioTrack != INDEX_NONE)
 	{
-		AudioDecoder->Initialize(AudioTracks[SelectedAudioTrack]->GetSamplingRate(), AudioTracks[SelectedAudioTrack]->GetChannels());
+		const mkvparser::AudioTrack* AudioTrack = AudioTracks[SelectedAudioTrack];
+
+		size_t PrivateDataSize;
+		const uint8* PrivateData = AudioTrack->GetCodecPrivate(PrivateDataSize);
+
+		AudioDecoder->Initialize(AudioTrack->GetCodecId(), AudioTrack->GetSamplingRate(), AudioTrack->GetChannels(), PrivateData, PrivateDataSize);
 	}
 
 	if (SelectedVideoTrack != INDEX_NONE)
 	{
-		VideoDecoder->Initialize();
+		const mkvparser::VideoTrack* VideoTrack = VideoTracks[SelectedVideoTrack];
+
+		VideoDecoder->Initialize(VideoTrack->GetCodecId());
 	}
 
 	CurrentState = EMediaState::Stopped;
@@ -413,12 +420,19 @@ bool FWebMMediaPlayer::SelectTrack(EMediaTrackType TrackType, int32 TrackIndex)
 		{
 			if (SelectedAudioTrack != INDEX_NONE)
 			{
-				AudioDecoder->Initialize(AudioTracks[SelectedAudioTrack]->GetSamplingRate(), AudioTracks[SelectedAudioTrack]->GetChannels());
+				const mkvparser::AudioTrack* AudioTrack = AudioTracks[SelectedAudioTrack];
+
+				size_t PrivateDataSize;
+				const uint8* PrivateData = AudioTrack->GetCodecPrivate(PrivateDataSize);
+
+				AudioDecoder->Initialize(AudioTrack->GetCodecId(), AudioTrack->GetSamplingRate(), AudioTrack->GetChannels(), PrivateData, PrivateDataSize);
 			}
 
 			if (SelectedVideoTrack != INDEX_NONE)
 			{
-				VideoDecoder->Initialize();
+				const mkvparser::VideoTrack* VideoTrack = VideoTracks[SelectedVideoTrack];
+
+				VideoDecoder->Initialize(VideoTrack->GetCodecId());
 			}
 		}
 	}
@@ -526,12 +540,19 @@ bool FWebMMediaPlayer::Seek(const FTimespan& Time)
 
 	if (SelectedVideoTrack != INDEX_NONE)
 	{
-		VideoDecoder->Initialize();
+		const mkvparser::VideoTrack* VideoTrack = VideoTracks[SelectedVideoTrack];
+
+		VideoDecoder->Initialize(VideoTrack->GetCodecId());
 	}
 
 	if (SelectedAudioTrack != INDEX_NONE)
 	{
-		AudioDecoder->Initialize(AudioTracks[SelectedAudioTrack]->GetSamplingRate(), AudioTracks[SelectedAudioTrack]->GetChannels());
+		const mkvparser::AudioTrack* AudioTrack = AudioTracks[SelectedAudioTrack];
+
+		size_t PrivateDataSize;
+		const uint8* PrivateData = AudioTrack->GetCodecPrivate(PrivateDataSize);
+
+		AudioDecoder->Initialize(AudioTrack->GetCodecId(), AudioTrack->GetSamplingRate(), AudioTrack->GetChannels(), PrivateData, PrivateDataSize);
 	}
 
 	CurrentTime = Time;
@@ -606,7 +627,7 @@ bool FWebMMediaPlayer::MkvRead()
 
 		if (Track->GetType() == mkvparser::Track::kVideo)
 		{
-			if (FCStringAnsi::Strcmp(Track->GetCodecId(), "V_VP9") == 0)
+			if (FCStringAnsi::Strcmp(Track->GetCodecId(), "V_VP8") == 0 || FCStringAnsi::Strcmp(Track->GetCodecId(), "V_VP9") == 0)
 			{
 				VideoTracks.Add(static_cast<const mkvparser::VideoTrack*>(Track));
 			}
@@ -618,7 +639,7 @@ bool FWebMMediaPlayer::MkvRead()
 		}
 		else if (Track->GetType() == mkvparser::Track::kAudio)
 		{
-			if (FCStringAnsi::Strcmp(Track->GetCodecId(), "A_OPUS") == 0)
+			if (FCStringAnsi::Strcmp(Track->GetCodecId(), "A_OPUS") == 0 || FCStringAnsi::Strcmp(Track->GetCodecId(), "A_VORBIS") == 0)
 			{
 				AudioTracks.Add(static_cast<const mkvparser::AudioTrack*>(Track));
 			}
diff --git a/Engine/Plugins/Media/WebMMedia/Source/WebMMedia/Private/Player/WebMMediaTextureSample.cpp b/Engine/Plugins/Media/WebMMedia/Source/WebMMedia/Private/Player/WebMMediaTextureSample.cpp
--- a/Engine/Plugins/Media/WebMMedia/Source/WebMMedia/Private/Player/WebMMediaTextureSample.cpp
+++ b/Engine/Plugins/Media/WebMMedia/Source/WebMMedia/Private/Player/WebMMediaTextureSample.cpp
@@ -97,3 +97,8 @@ void FWebMMediaTextureSample::ShutdownPoolable()
 {
 	Time = FTimespan::Zero();
 }
+
+TRefCountPtr<FRHITexture2D> FWebMMediaTextureSample::GetTextureRef() const
+{
+	return Texture;
+}
diff --git a/Engine/Plugins/Media/WebMMedia/Source/WebMMedia/Private/Player/WebMMediaTextureSample.h b/Engine/Plugins/Media/WebMMedia/Source/WebMMedia/Private/Player/WebMMediaTextureSample.h
deleted file mode 100644
diff --git a/Engine/Plugins/Media/WebMMedia/Source/WebMMedia/Private/Player/WebMVideoDecoder.cpp b/Engine/Plugins/Media/WebMMedia/Source/WebMMedia/Private/Player/WebMVideoDecoder.cpp
--- a/Engine/Plugins/Media/WebMMedia/Source/WebMMedia/Private/Player/WebMVideoDecoder.cpp
+++ b/Engine/Plugins/Media/WebMMedia/Source/WebMMedia/Private/Player/WebMVideoDecoder.cpp
@@ -3,11 +3,12 @@
 #include "WebMVideoDecoder.h"
 #include "WebMMediaPrivate.h"
 #include "WebMMediaFrame.h"
+#include "WebMMediaTextureSample.h"
 #include "MediaShaders.h"
 #include "MediaSamples.h"
 #include "PipelineStateCache.h"
 #include "RHIStaticStates.h"
-# include "Containers/DynamicRHIResourceArray.h"
+#include "Containers/DynamicRHIResourceArray.h"
 
 namespace
 {
@@ -61,31 +62,39 @@ FWebMVideoDecoder::FWebMVideoDecoder(TSharedPtr<FMediaSamples, ESPMode::ThreadSa
 	: Samples(InSamples)
 	, VideoSamplePool(new FWebMMediaTextureSamplePool)
 	, bTexturesCreated(false)
+	, bIsInitialized(false)
 {
-	const int32 NumOfThreads = 1;
-	const vpx_codec_dec_cfg_t CodecConfig = { NumOfThreads, 0, 0 };
-	verify(vpx_codec_dec_init(&Context, vpx_codec_vp9_dx(), &CodecConfig, /*VPX_CODEC_USE_FRAME_THREADING*/ 0) == 0);
 }
 
 FWebMVideoDecoder::~FWebMVideoDecoder()
 {
-	if (VideoDecodingTask && !VideoDecodingTask->IsComplete())
-	{
-		FTaskGraphInterface::Get().WaitUntilTaskCompletes(VideoDecodingTask);
-	}
-
-	// Make sure all compute shader decoding is done
-	FlushRenderingCommands();
-
-	vpx_codec_destroy(&Context);
+	Close();
 }
 
-void FWebMVideoDecoder::Initialize()
+bool FWebMVideoDecoder::Initialize(const char* CodecName)
 {
-	// Flush decoder
-	vpx_codec_decode(&Context, nullptr, 0, nullptr, 0);
+	Close();
 
-	bTexturesCreated = false;
+	const int32 NumOfThreads = 1;
+	const vpx_codec_dec_cfg_t CodecConfig = { NumOfThreads, 0, 0 };
+
+	if (FCStringAnsi::Strcmp(CodecName, "V_VP8") == 0)
+	{
+		verify(vpx_codec_dec_init(&Context, vpx_codec_vp8_dx(), &CodecConfig, /*VPX_CODEC_USE_FRAME_THREADING*/ 0) == 0);
+	}
+	else if (FCStringAnsi::Strcmp(CodecName, "V_VP9") == 0)
+	{
+		verify(vpx_codec_dec_init(&Context, vpx_codec_vp9_dx(), &CodecConfig, /*VPX_CODEC_USE_FRAME_THREADING*/ 0) == 0);
+	}
+	else
+	{
+		UE_LOG(LogWebMMedia, Display, TEXT("Unsupported video codec: %s"), CodecName);
+		return false;
+	}
+
+	bIsInitialized = true;
+
+	return true;
 }
 
 void FWebMVideoDecoder::DecodeVideoFramesAsync(const TArray<TSharedPtr<FWebMFrame>>& VideoFrames)
@@ -155,6 +164,25 @@ void FWebMVideoDecoder::CreateTextures(const vpx_image_t* Image)
 	DecodedV = RHICreateTexture2D(Image->stride[2], Image->d_h / 2, PF_G8, 1, 1, TexCreate_Dynamic, CreateInfo);
 }
 
+void FWebMVideoDecoder::Close()
+{
+	if (VideoDecodingTask && !VideoDecodingTask->IsComplete())
+	{
+		FTaskGraphInterface::Get().WaitUntilTaskCompletes(VideoDecodingTask);
+	}
+
+	// Make sure all compute shader decoding is done
+	FlushRenderingCommands();
+
+	if (bIsInitialized)
+	{
+		vpx_codec_destroy(&Context);
+		bIsInitialized = false;
+	}
+
+	bTexturesCreated = false;
+}
+
 void FWebMVideoDecoder::ConvertYUVToRGBAndSubmit(const FConvertParams& Params)
 {
 	TSharedPtr<FWebMMediaTextureSample, ESPMode::ThreadSafe> VideoSample = Params.VideoSample;
diff --git a/Engine/Plugins/Media/WebMMedia/Source/WebMMedia/Private/Player/WebMVideoDecoder.h b/Engine/Plugins/Media/WebMMedia/Source/WebMMedia/Private/Player/WebMVideoDecoder.h
deleted file mode 100644
diff --git a/Engine/Plugins/Media/WebMMedia/Source/WebMMedia/Public/WebMAudioDecoder.h b/Engine/Plugins/Media/WebMMedia/Source/WebMMedia/Public/WebMAudioDecoder.h
new file mode 100644
--- /dev/null
+++ b/Engine/Plugins/Media/WebMMedia/Source/WebMMedia/Public/WebMAudioDecoder.h
@@ -0,0 +1,49 @@
+// Copyright 1998-2018 Epic Games, Inc. All Rights Reserved.
+
+#pragma once
+
+#include "Async/TaskGraphInterfaces.h"
+#include "Templates/SharedPointer.h"
+#include "Templates/UniquePtr.h"
+
+class FMediaSamples;
+class FWebMMediaAudioSample;
+class FWebMMediaAudioSamplePool;
+struct FWebMFrame;
+struct OpusDecoder;
+
+class WEBMMEDIA_API FWebMAudioDecoder
+{
+public:
+	FWebMAudioDecoder(TSharedPtr<FMediaSamples, ESPMode::ThreadSafe> InSamples);
+	~FWebMAudioDecoder();
+
+public:
+	bool Initialize(const char* InCodec, int32 InSampleRate, int32 InChannels, const uint8* CodecPrivateData, size_t CodecPrivateDataSize);
+	void DecodeAudioFramesAsync(const TArray<TSharedPtr<FWebMFrame>>& AudioFrames);
+
+private:
+	enum ESupportedCodecs
+	{
+		Opus,
+		Vorbis,
+	};
+
+	struct FVorbisDecoder;
+	TSharedPtr<FMediaSamples, ESPMode::ThreadSafe> Samples;
+	TUniquePtr<FWebMMediaAudioSamplePool> AudioSamplePool;
+	TUniquePtr<FVorbisDecoder> VorbisDecoder;
+	FGraphEventRef AudioDecodingTask;
+	TArray<uint8> DecodeBuffer;
+	OpusDecoder* OpusDecoder;
+	ESupportedCodecs Codec;
+	int32 FrameSize;
+	int32 SampleRate;
+	int32 Channels;
+
+	void InitializeOpus();
+	bool InitializeVorbis(const uint8* CodecPrivateData, size_t CodecPrivateDataSize);
+	void DoDecodeAudioFrames(const TArray<TSharedPtr<FWebMFrame>>& AudioFrames);
+	int32 DecodeOpus(const TSharedPtr<FWebMFrame>& AudioFrame);
+	int32 DecodeVorbis(const TSharedPtr<FWebMFrame>& AudioFrame);
+};
diff --git a/Engine/Plugins/Media/WebMMedia/Source/WebMMedia/Public/WebMMediaAudioSample.h b/Engine/Plugins/Media/WebMMedia/Source/WebMMedia/Public/WebMMediaAudioSample.h
new file mode 100644
--- /dev/null
+++ b/Engine/Plugins/Media/WebMMedia/Source/WebMMedia/Public/WebMMediaAudioSample.h
@@ -0,0 +1,122 @@
+// Copyright 1998-2018 Epic Games, Inc. All Rights Reserved.
+
+#pragma once
+
+#include "CoreTypes.h"
+#include "Containers/Array.h"
+#include "IMediaAudioSample.h"
+#include "MediaObjectPool.h"
+#include "MediaSampleQueue.h"
+#include "Math/IntPoint.h"
+#include "Misc/Timespan.h"
+
+
+/**
+ * Implements a media audio sample for WebMMedia.
+ */
+class FWebMMediaAudioSample
+	: public IMediaAudioSample
+	, public IMediaPoolable
+{
+public:
+
+	/** Default constructor. */
+	FWebMMediaAudioSample()
+		: Channels(0)
+		, Duration(FTimespan::Zero())
+		, SampleRate(0)
+		, Time(FTimespan::Zero())
+	{ }
+
+	/** Virtual destructor. */
+	virtual ~FWebMMediaAudioSample() { }
+
+public:
+
+	/**
+	 * Initialize the sample.
+	 *
+	 * @param InBuffer The sample's data buffer.
+	 * @param InSize The size of the sample buffer (in bytes).
+	 * @param InTime The sample time (relative to presentation clock).
+	 * @param InDuration The duration for which the sample is valid.
+	 */
+	void Initialize(
+		const uint8* InBuffer,
+		uint32 InSize,
+		uint32 InChannels,
+		uint32 InSampleRate,
+		FTimespan InTime,
+		FTimespan InDuration)
+	{
+		check(InBuffer && InSize > 0);
+
+		Buffer.Reset(InSize);
+		Buffer.Append(InBuffer, InSize);
+
+		Channels = InChannels;
+		Duration = InDuration;
+		SampleRate = InSampleRate;
+		Time = InTime;
+	}
+
+public:
+
+	//~ IMediaAudioSample interface
+
+	virtual const void* GetBuffer() override
+	{
+		return Buffer.GetData();
+	}
+
+	virtual uint32 GetChannels() const override
+	{
+		return Channels;
+	}
+
+	virtual FTimespan GetDuration() const override
+	{
+		return Duration;
+	}
+
+	virtual EMediaAudioSampleFormat GetFormat() const override
+	{
+		return EMediaAudioSampleFormat::Int16;
+	}
+
+	virtual uint32 GetFrames() const override
+	{
+		return Buffer.Num() / (Channels * sizeof(int16));
+	}
+
+	virtual uint32 GetSampleRate() const override
+	{
+		return SampleRate;
+	}
+
+	virtual FTimespan GetTime() const override
+	{
+		return Time;
+	}
+
+private:
+
+	/** The sample's data buffer. */
+	TArray<uint8> Buffer;
+
+	/** Number of audio channels. */
+	uint32 Channels;
+
+	/** The duration for which the sample is valid. */
+	FTimespan Duration;
+
+	/** Audio sample rate (in samples per second). */
+	uint32 SampleRate;
+
+	/** Presentation time for which the sample was generated. */
+	FTimespan Time;
+};
+
+
+/** Implements a pool for WebM audio sample objects. */
+class FWebMMediaAudioSamplePool : public TMediaObjectPool<FWebMMediaAudioSample> { };
diff --git a/Engine/Plugins/Media/WebMMedia/Source/WebMMedia/Public/WebMMediaFrame.h b/Engine/Plugins/Media/WebMMedia/Source/WebMMedia/Public/WebMMediaFrame.h
new file mode 100644
--- /dev/null
+++ b/Engine/Plugins/Media/WebMMedia/Source/WebMMedia/Public/WebMMediaFrame.h
@@ -0,0 +1,11 @@
+// Copyright 1998-2018 Epic Games, Inc. All Rights Reserved.
+
+#pragma once
+
+#include "CoreMinimal.h"
+
+struct FWebMFrame
+{
+	TArray<uint8> Data;
+	FTimespan Time;
+};
diff --git a/Engine/Plugins/Media/WebMMedia/Source/WebMMedia/Public/WebMMediaTextureSample.h b/Engine/Plugins/Media/WebMMedia/Source/WebMMedia/Public/WebMMediaTextureSample.h
new file mode 100644
--- /dev/null
+++ b/Engine/Plugins/Media/WebMMedia/Source/WebMMedia/Public/WebMMediaTextureSample.h
@@ -0,0 +1,51 @@
+// Copyright 1998-2018 Epic Games, Inc. All Rights Reserved.
+
+#pragma once
+
+#include "CoreMinimal.h"
+#include "IMediaTextureSample.h"
+#include "MediaObjectPool.h"
+#include "RHI.h"
+#include "RHIUtilities.h"
+
+class FRHITexture2D;
+
+class WEBMMEDIA_API FWebMMediaTextureSample
+	: public IMediaTextureSample
+	, public IMediaPoolable
+{
+public:
+	FWebMMediaTextureSample();
+
+public:
+	void Initialize(FIntPoint InDisplaySize, FIntPoint InTotalSize, FTimespan InTime);
+	void CreateTexture();
+
+public:
+	//~ IMediaTextureSample interface
+	virtual const void* GetBuffer() override;
+	virtual FIntPoint GetDim() const override;
+	virtual FTimespan GetDuration() const override;
+	virtual EMediaTextureSampleFormat GetFormat() const override;
+	virtual FIntPoint GetOutputDim() const override;
+	virtual uint32 GetStride() const override;
+	virtual FRHITexture* GetTexture() const override;
+	virtual FTimespan GetTime() const override;
+	virtual bool IsCacheable() const override;
+	virtual bool IsOutputSrgb() const override;
+
+public:
+	//~ IMediaPoolable interface
+	virtual void ShutdownPoolable() override;
+
+public:
+	TRefCountPtr<FRHITexture2D> GetTextureRef() const;
+
+private:
+	TRefCountPtr<FRHITexture2D> Texture;
+	FTimespan Time;
+	FIntPoint TotalSize;
+	FIntPoint DisplaySize;
+};
+
+class FWebMMediaTextureSamplePool : public TMediaObjectPool<FWebMMediaTextureSample> { };
diff --git a/Engine/Plugins/Media/WebMMedia/Source/WebMMedia/Public/WebMVideoDecoder.h b/Engine/Plugins/Media/WebMMedia/Source/WebMMedia/Public/WebMVideoDecoder.h
new file mode 100644
--- /dev/null
+++ b/Engine/Plugins/Media/WebMMedia/Source/WebMMedia/Public/WebMVideoDecoder.h
@@ -0,0 +1,49 @@
+// Copyright 1998-2018 Epic Games, Inc. All Rights Reserved.
+
+#pragma once
+
+#include "Templates/SharedPointer.h"
+#include "MediaShaders.h"
+
+THIRD_PARTY_INCLUDES_START
+#include <vpx/vpx_decoder.h>
+#include <vpx/vp8dx.h>
+THIRD_PARTY_INCLUDES_END
+
+class FMediaSamples;
+class FWebMMediaTextureSample;
+class FWebMMediaTextureSamplePool;
+struct FWebMFrame;
+
+class WEBMMEDIA_API FWebMVideoDecoder
+{
+public:
+	FWebMVideoDecoder(TSharedPtr<FMediaSamples, ESPMode::ThreadSafe> InSamples);
+	~FWebMVideoDecoder();
+
+public:
+	bool Initialize(const char* CodecName);
+	void DecodeVideoFramesAsync(const TArray<TSharedPtr<FWebMFrame>>& VideoFrames);
+
+private:
+	struct FConvertParams
+	{
+		TSharedPtr<FWebMMediaTextureSample, ESPMode::ThreadSafe> VideoSample;
+		const vpx_image_t* Image;
+	};
+
+	vpx_codec_ctx_t Context;
+	TSharedPtr<FMediaSamples, ESPMode::ThreadSafe> Samples;
+	TUniquePtr<FWebMMediaTextureSamplePool> VideoSamplePool;
+	TRefCountPtr<FRHITexture2D> DecodedY;
+	TRefCountPtr<FRHITexture2D> DecodedU;
+	TRefCountPtr<FRHITexture2D> DecodedV;
+	FGraphEventRef VideoDecodingTask;
+	bool bTexturesCreated;
+	bool bIsInitialized;
+
+	void ConvertYUVToRGBAndSubmit(const FConvertParams& Params);
+	void DoDecodeVideoFrames(const TArray<TSharedPtr<FWebMFrame>>& VideoFrames);
+	void CreateTextures(const vpx_image_t* Image);
+	void Close();
+};
diff --git a/Engine/Plugins/Media/WebMMedia/Source/WebMMedia/WebMMedia.Build.cs b/Engine/Plugins/Media/WebMMedia/Source/WebMMedia/WebMMedia.Build.cs
--- a/Engine/Plugins/Media/WebMMedia/Source/WebMMedia/WebMMedia.Build.cs
+++ b/Engine/Plugins/Media/WebMMedia/Source/WebMMedia/WebMMedia.Build.cs
@@ -7,38 +7,23 @@ namespace UnrealBuildTool.Rules
 	{
 		public WebMMedia(ReadOnlyTargetRules Target) : base(Target)
 		{
-			DynamicallyLoadedModuleNames.AddRange(
-				new string[] {
-					"Media",
-				});
-
 			PrivateDependencyModuleNames.AddRange(
 				new string[] {
 					"WebMMediaFactory",
 					"Core",
 					"Engine",
+					"Media",
 					"MediaUtils",
 					"RenderCore",
 					"RHI",
 					"ShaderCore",
 					"UtilityShaders",
 					"libOpus",
-					//"UEOgg",
-					//"Vorbis",
+					"UEOgg",
+					"Vorbis",
 					"LibVpx",
 					"LibWebM",
 				});
-
-			PrivateIncludePathModuleNames.AddRange(
-				new string[] {
-					"Media",
-				});
-
-			PrivateIncludePaths.AddRange(
-				new string[] {
-					"WebMMedia/Private",
-					"WebMMedia/Private/Player",
-				});
 		}
 	}
 }
diff --git a/Engine/Plugins/Media/WebMMedia/Source/WebMMediaEditor/Private/WebMFileMediaSourceFactory.cpp b/Engine/Plugins/Media/WebMMedia/Source/WebMMediaEditor/Private/WebMFileMediaSourceFactory.cpp
--- a/Engine/Plugins/Media/WebMMedia/Source/WebMMediaEditor/Private/WebMFileMediaSourceFactory.cpp
+++ b/Engine/Plugins/Media/WebMMedia/Source/WebMMediaEditor/Private/WebMFileMediaSourceFactory.cpp
@@ -13,7 +13,7 @@ UWebMPlatFileMediaSourceFactory::UWebMPlatFileMediaSourceFactory(const FObjectIn
 	: Super(ObjectInitializer)
 {
 	// add formats we support
-	Formats.Add(TEXT("mkv;Matroska Multimedia File"));
+	Formats.Add(TEXT("webm;WebM Multimedia File"));
 
 	SupportedClass = UFileMediaSource::StaticClass();
 
diff --git a/Engine/Plugins/Media/WebMMedia/Source/WebMMediaFactory/Private/WebMMediaFactoryModule.cpp b/Engine/Plugins/Media/WebMMedia/Source/WebMMediaFactory/Private/WebMMediaFactoryModule.cpp
--- a/Engine/Plugins/Media/WebMMedia/Source/WebMMediaFactory/Private/WebMMediaFactoryModule.cpp
+++ b/Engine/Plugins/Media/WebMMedia/Source/WebMMediaFactory/Private/WebMMediaFactoryModule.cpp
@@ -111,7 +111,7 @@ public:
 	virtual void StartupModule() override
 	{
 		// supported file extensions
-		SupportedFileExtensions.Add(TEXT("mkv"));
+		SupportedFileExtensions.Add(TEXT("webm"));
 
 		// if the module got built and run succesfully on this platform, it means it's supported
 		SupportedPlatforms.Add(FPlatformMisc::GetUBTPlatform());
diff --git a/Engine/Plugins/Runtime/GameplayAbilities/Source/GameplayAbilities/Public/AttributeSet.h b/Engine/Plugins/Runtime/GameplayAbilities/Source/GameplayAbilities/Public/AttributeSet.h
--- a/Engine/Plugins/Runtime/GameplayAbilities/Source/GameplayAbilities/Public/AttributeSet.h
+++ b/Engine/Plugins/Runtime/GameplayAbilities/Source/GameplayAbilities/Public/AttributeSet.h
@@ -342,7 +342,7 @@ struct TStructOpsTypeTraits<FScalableFloat>
 {
 	enum
 	{
-		WithSerializeFromMismatchedTag = true,
+		WithStructuredSerializeFromMismatchedTag = true,
 	};
 };
 
diff --git a/Engine/Plugins/Runtime/GoogleARCore/Source/GoogleARCoreRendering/Private/GoogleARCorePassthroughCameraRenderer.cpp b/Engine/Plugins/Runtime/GoogleARCore/Source/GoogleARCoreRendering/Private/GoogleARCorePassthroughCameraRenderer.cpp
--- a/Engine/Plugins/Runtime/GoogleARCore/Source/GoogleARCoreRendering/Private/GoogleARCorePassthroughCameraRenderer.cpp
+++ b/Engine/Plugins/Runtime/GoogleARCore/Source/GoogleARCoreRendering/Private/GoogleARCorePassthroughCameraRenderer.cpp
@@ -100,7 +100,7 @@ void FGoogleARCorePassthroughCameraRenderer::UpdateOverlayUVCoordinate_RenderThr
 {
 	check(InOverlayUVs.Num() == 8);
 
-	bool bNeedToFlipCameraImage = !RHINeedsToSwitchVerticalAxis(GShaderPlatformForFeatureLevel[GMaxRHIFeatureLevel]) || IsMobileHDR();
+	bool bNeedToFlipCameraImage = RHINeedsToSwitchVerticalAxis(GShaderPlatformForFeatureLevel[GMaxRHIFeatureLevel]) && !IsMobileHDR();
 
 	if (bNeedToFlipCameraImage)
 	{
diff --git a/Engine/Shaders/Private/MaterialTemplate.ush b/Engine/Shaders/Private/MaterialTemplate.ush
--- a/Engine/Shaders/Private/MaterialTemplate.ush
+++ b/Engine/Shaders/Private/MaterialTemplate.ush
@@ -1362,7 +1362,7 @@ MaterialFloat4 MobileSceneTextureLookup(inout FMaterialPixelParameters Parameter
 	//PPI_SceneDepth
 	if (SceneTextureId == 1)
 	{
-		MaterialFloat Depth = ConvertFromDeviceZ(Texture2DSample(MobileSceneTextures.SceneDepthTexture, MobileSceneTextures.SceneDepthTextureSampler, UV).r);
+		MaterialFloat Depth = Texture2DSample(MobileSceneTextures.SceneColorTexture, MobileSceneTextures.SceneColorTextureSampler, UV).a;
 		return MaterialFloat4(Depth.rrr, 0.f);
 	}
 	//PPI_CustomDepth
diff --git a/Engine/Shaders/Private/MobileBasePassPixelShader.usf b/Engine/Shaders/Private/MobileBasePassPixelShader.usf
--- a/Engine/Shaders/Private/MobileBasePassPixelShader.usf
+++ b/Engine/Shaders/Private/MobileBasePassPixelShader.usf
@@ -640,7 +640,7 @@ void Main(
 			if (GetHDR32bppEncodeMode() == HDR_ENCODE_NONE)
 			{
 				// FP16 HDR stores Z in framebuffer alpha
-				OutColor.a = BasePassInterpolants.PixelPosition.w; 
+				OutColor.a = min(BasePassInterpolants.PixelPosition.w, 65500); 
 			}
 		#endif
 	#endif
diff --git a/Engine/Source/Developer/Android/AndroidTargetPlatform/Private/AndroidTargetPlatform.cpp b/Engine/Source/Developer/Android/AndroidTargetPlatform/Private/AndroidTargetPlatform.cpp
--- a/Engine/Source/Developer/Android/AndroidTargetPlatform/Private/AndroidTargetPlatform.cpp
+++ b/Engine/Source/Developer/Android/AndroidTargetPlatform/Private/AndroidTargetPlatform.cpp
@@ -250,14 +250,8 @@ bool FAndroidTargetPlatform::SupportsVulkan() const
 
 bool FAndroidTargetPlatform::SupportsSoftwareOcclusion() const
 {
-	// default to not supporting
-	bool bSupportsSoftwareOcclusion = false;
-#if WITH_ENGINE
-	int32 IntValue = 0;
-	GConfig->GetInt(TEXT("ConsoleVariables"), TEXT("r.Mobile.AllowSoftwareOcclusion"), IntValue, GEngineIni);
-	bSupportsSoftwareOcclusion = (IntValue != 0);
-#endif
-	return bSupportsSoftwareOcclusion;
+	static auto* CVarMobileAllowSoftwareOcclusion = IConsoleManager::Get().FindTConsoleVariableDataInt(TEXT("r.Mobile.AllowSoftwareOcclusion"));
+	return CVarMobileAllowSoftwareOcclusion->GetValueOnAnyThread() != 0;
 }
 
 /* ITargetPlatform overrides
diff --git a/Engine/Source/Developer/IOS/IOSTargetPlatform/Private/IOSTargetPlatform.cpp b/Engine/Source/Developer/IOS/IOSTargetPlatform/Private/IOSTargetPlatform.cpp
--- a/Engine/Source/Developer/IOS/IOSTargetPlatform/Private/IOSTargetPlatform.cpp
+++ b/Engine/Source/Developer/IOS/IOSTargetPlatform/Private/IOSTargetPlatform.cpp
@@ -484,10 +484,8 @@ static bool CookASTC()
 
 static bool SupportsSoftwareOcclusion()
 {
-	// default to not support
-	int32 IntValue = 0;
-	GConfig->GetInt(TEXT("ConsoleVariables"), TEXT("r.Mobile.AllowSoftwareOcclusion"), IntValue, GEngineIni);
-	return (IntValue != 0);
+	static auto* CVarMobileAllowSoftwareOcclusion = IConsoleManager::Get().FindTConsoleVariableDataInt(TEXT("r.Mobile.AllowSoftwareOcclusion"));
+	return CVarMobileAllowSoftwareOcclusion->GetValueOnAnyThread() != 0;
 }
 
 bool FIOSTargetPlatform::CanSupportXGEShaderCompile() const
diff --git a/Engine/Source/Developer/TargetDeviceServices/Private/Services/TargetDeviceService.cpp b/Engine/Source/Developer/TargetDeviceServices/Private/Services/TargetDeviceService.cpp
--- a/Engine/Source/Developer/TargetDeviceServices/Private/Services/TargetDeviceService.cpp
+++ b/Engine/Source/Developer/TargetDeviceServices/Private/Services/TargetDeviceService.cpp
@@ -245,7 +245,9 @@ bool FTargetDeviceService::Start()
 		ClaimHost = FPlatformProcess::ComputerName();
 		ClaimUser = FPlatformProcess::UserName(false);
 
-		MessageEndpoint->Publish(new FTargetDeviceClaimed(DeviceName, ClaimHost, ClaimUser));
+		// message is going to be deleted by FMemory::Free() (see FMessageContext destructor), so allocate it with Malloc
+		void* Memory = FMemory::Malloc(sizeof(FTargetDeviceClaimed));
+		MessageEndpoint->Publish(new(Memory) FTargetDeviceClaimed(DeviceName, ClaimHost, ClaimUser));
 
 		Running = true;
 	}
@@ -258,8 +260,9 @@ void FTargetDeviceService::Stop()
 {
 	if (Running)
 	{
-	
-		MessageEndpoint->Publish(new FTargetDeviceUnclaimed(DeviceName, FPlatformProcess::ComputerName(), FPlatformProcess::UserName(false)));
+		// message is going to be deleted by FMemory::Free() (see FMessageContext destructor), so allocate it with Malloc
+		void* Memory = FMemory::Malloc(sizeof(FTargetDeviceUnclaimed));
+		MessageEndpoint->Publish(new(Memory) FTargetDeviceUnclaimed(DeviceName, FPlatformProcess::ComputerName(), FPlatformProcess::UserName(false)));
 		FPlatformProcess::SleepNoStats(0.01);
 
 		// Only stop the device if we care about device claiming
@@ -356,7 +359,9 @@ void FTargetDeviceService::HandleClaimedMessage(const FTargetDeviceClaimed& Mess
 	{
 		if (Context->GetSender() != MessageEndpoint->GetAddress())
 		{
-			MessageEndpoint->Send(new FTargetDeviceClaimDenied(DeviceName, FPlatformProcess::ComputerName(), FPlatformProcess::UserName(false)), Context->GetSender());
+			// message is going to be deleted by FMemory::Free() (see FMessageContext destructor), so allocate it with Malloc
+			void* Memory = FMemory::Malloc(sizeof(FTargetDeviceClaimDenied));
+			MessageEndpoint->Send(new(Memory) FTargetDeviceClaimDenied(DeviceName, FPlatformProcess::ComputerName(), FPlatformProcess::UserName(false)), Context->GetSender());
 		}
 	}
 	else
@@ -425,7 +430,9 @@ void FTargetDeviceService::HandleDeployCommitMessage(const FTargetDeviceServiceD
 		bool Succeeded = TargetDevice->Deploy(SourceFolder, OutAppId);
 
 		IFileManager::Get().DeleteDirectory(*SourceFolder, false, true);
-		MessageEndpoint->Send(new FTargetDeviceServiceDeployFinished(Message.Variant, OutAppId, Succeeded, Message.TransactionId), Context->GetSender());
+		// message is going to be deleted by FMemory::Free() (see FMessageContext destructor), so allocate it with Malloc
+		void* Memory = FMemory::Malloc(sizeof(FTargetDeviceServiceDeployFinished));
+		MessageEndpoint->Send(new(Memory) FTargetDeviceServiceDeployFinished(Message.Variant, OutAppId, Succeeded, Message.TransactionId), Context->GetSender());
 	}
 }
 
@@ -460,7 +467,9 @@ void FTargetDeviceService::HandleLaunchAppMessage(const FTargetDeviceServiceLaun
 
 		if (MessageEndpoint.IsValid())
 		{
-			MessageEndpoint->Send(new FTargetDeviceServiceLaunchFinished(Message.AppID, ProcessId, Succeeded), Context->GetSender());
+			// message is going to be deleted by FMemory::Free() (see FMessageContext destructor), so allocate it with Malloc
+			void* Memory = FMemory::Malloc(sizeof(FTargetDeviceServiceLaunchFinished));
+			MessageEndpoint->Send(new(Memory) FTargetDeviceServiceLaunchFinished(Message.AppID, ProcessId, Succeeded), Context->GetSender());
 		}
 	}
 }
@@ -485,7 +494,9 @@ void FTargetDeviceService::HandlePingMessage(const FTargetDeviceServicePing& InM
 		const FString& PlatformName = DefaultDevice->GetTargetPlatform().PlatformName();
 		const PlatformInfo::FPlatformInfo* VanillaInfo = PlatformInfo::FindVanillaPlatformInfo(FName(*PlatformName));
 
-		FTargetDeviceServicePong* Message = new FTargetDeviceServicePong();
+		// message is going to be deleted by FMemory::Free() (see FMessageContext destructor), so allocate it with Malloc
+		void* Memory = FMemory::Malloc(sizeof(FTargetDeviceServicePong));
+		FTargetDeviceServicePong* Message = new(Memory) FTargetDeviceServicePong();
 
 		Message->Name = DefaultDevice->GetName();
 		Message->Type = TargetDeviceTypes::ToString(DefaultDevice->GetDeviceType());
@@ -595,6 +606,8 @@ void FTargetDeviceService::HandleRunExecutableMessage(const FTargetDeviceService
 		uint32 OutProcessId;
 		bool Succeeded = TargetDevice->Run(Message.ExecutablePath, Message.Params, &OutProcessId);
 
-		MessageEndpoint->Send(new FTargetDeviceServiceRunFinished(Message.Variant, Message.ExecutablePath, OutProcessId, Succeeded), Context->GetSender());
+		// message is going to be deleted by FMemory::Free() (see FMessageContext destructor), so allocate it with Malloc
+		void* Memory = FMemory::Malloc(sizeof(FTargetDeviceServiceRunFinished));
+		MessageEndpoint->Send(new(Memory) FTargetDeviceServiceRunFinished(Message.Variant, Message.ExecutablePath, OutProcessId, Succeeded), Context->GetSender());
 	}
 }
diff --git a/Engine/Source/Editor/EditorWidgets/Private/SDropTarget.cpp b/Engine/Source/Editor/EditorWidgets/Private/SDropTarget.cpp
--- a/Engine/Source/Editor/EditorWidgets/Private/SDropTarget.cpp
+++ b/Engine/Source/Editor/EditorWidgets/Private/SDropTarget.cpp
@@ -153,17 +153,14 @@ int32 SDropTarget::OnPaint(const FPaintArgs& Args, const FGeometry& AllottedGeom
 		{
 			FLinearColor DashColor = bAllowDrop ? ValidColor : InvalidColor;
 
-			const FSlateBrush* HorizontalBrush = FEditorStyle::GetBrush("WideDash.Horizontal");
-			const FSlateBrush* VerticalBrush = FEditorStyle::GetBrush("WideDash.Vertical");
-
 			int32 DashLayer = LayerId + 1;
 
 			// Top
 			FSlateDrawElement::MakeBox(
 				OutDrawElements,
 				DashLayer,
-				AllottedGeometry.ToPaintGeometry(FVector2D(0, 0), FVector2D(AllottedGeometry.GetLocalSize().X, HorizontalBrush->ImageSize.Y)),
-				HorizontalBrush,
+				AllottedGeometry.ToPaintGeometry(FVector2D(0, 0), FVector2D(AllottedGeometry.GetLocalSize().X, HorizontalImage->ImageSize.Y)),
+				HorizontalImage,
 				ESlateDrawEffect::None,
 				DashColor);
 
@@ -171,8 +168,8 @@ int32 SDropTarget::OnPaint(const FPaintArgs& Args, const FGeometry& AllottedGeom
 			FSlateDrawElement::MakeBox(
 				OutDrawElements,
 				DashLayer,
-				AllottedGeometry.ToPaintGeometry(FVector2D(0, AllottedGeometry.GetLocalSize().Y - HorizontalBrush->ImageSize.Y), FVector2D(AllottedGeometry.Size.X, HorizontalBrush->ImageSize.Y)),
-				HorizontalBrush,
+				AllottedGeometry.ToPaintGeometry(FVector2D(0, AllottedGeometry.GetLocalSize().Y - HorizontalImage->ImageSize.Y), FVector2D(AllottedGeometry.Size.X, HorizontalImage->ImageSize.Y)),
+				HorizontalImage,
 				ESlateDrawEffect::None,
 				DashColor);
 
@@ -180,8 +177,8 @@ int32 SDropTarget::OnPaint(const FPaintArgs& Args, const FGeometry& AllottedGeom
 			FSlateDrawElement::MakeBox(
 				OutDrawElements,
 				DashLayer,
-				AllottedGeometry.ToPaintGeometry(FVector2D(0, 0), FVector2D(VerticalBrush->ImageSize.X, AllottedGeometry.GetLocalSize().Y)),
-				VerticalBrush,
+				AllottedGeometry.ToPaintGeometry(FVector2D(0, 0), FVector2D(VerticalImage->ImageSize.X, AllottedGeometry.GetLocalSize().Y)),
+				VerticalImage,
 				ESlateDrawEffect::None,
 				DashColor);
 
@@ -189,8 +186,8 @@ int32 SDropTarget::OnPaint(const FPaintArgs& Args, const FGeometry& AllottedGeom
 			FSlateDrawElement::MakeBox(
 				OutDrawElements,
 				DashLayer,
-				AllottedGeometry.ToPaintGeometry(FVector2D(AllottedGeometry.GetLocalSize().X - VerticalBrush->ImageSize.X, 0), FVector2D(VerticalBrush->ImageSize.X, AllottedGeometry.GetLocalSize().Y)),
-				VerticalBrush,
+				AllottedGeometry.ToPaintGeometry(FVector2D(AllottedGeometry.GetLocalSize().X - VerticalImage->ImageSize.X, 0), FVector2D(VerticalImage->ImageSize.X, AllottedGeometry.GetLocalSize().Y)),
+				VerticalImage,
 				ESlateDrawEffect::None,
 				DashColor);
 
diff --git a/Engine/Source/Editor/KismetCompiler/Private/KismetCompiler.cpp b/Engine/Source/Editor/KismetCompiler/Private/KismetCompiler.cpp
--- a/Engine/Source/Editor/KismetCompiler/Private/KismetCompiler.cpp
+++ b/Engine/Source/Editor/KismetCompiler/Private/KismetCompiler.cpp
@@ -4181,7 +4181,7 @@ void FKismetCompilerContext::CompileFunctions(EInternalCompilerFlags InternalFla
 				}
 
 				// >>> Backwards Compatibility: Propagate data from the skel CDO to the gen CDO if we haven't already done so for this blueprint
-				if( !bIsSkeletonOnly && !Blueprint->IsGeneratedClassAuthoritative() )
+				if( !bIsSkeletonOnly && !Blueprint->IsGeneratedClassAuthoritative() && CompileOptions.CompileType != EKismetCompileType::Cpp )
 				{
 					UEditorEngine::FCopyPropertiesForUnrelatedObjectsParams CopyDetails;
 					CopyDetails.bAggressiveDefaultSubobjectReplacement = false;
diff --git a/Engine/Source/Editor/MaterialEditor/Private/MaterialEditor.cpp b/Engine/Source/Editor/MaterialEditor/Private/MaterialEditor.cpp
--- a/Engine/Source/Editor/MaterialEditor/Private/MaterialEditor.cpp
+++ b/Engine/Source/Editor/MaterialEditor/Private/MaterialEditor.cpp
@@ -3234,7 +3234,7 @@ TSharedRef<SDockTab> FMaterialEditor::SpawnTab_GraphCanvas(const FSpawnTabArgs&
 
 TSharedRef<SDockTab> FMaterialEditor::SpawnTab_MaterialProperties(const FSpawnTabArgs& Args)
 {
-	SpawnedDetailsTab = SNew(SDockTab)
+	TSharedPtr<SDockTab> DetailsTab = SNew(SDockTab)
 		.Icon( FEditorStyle::GetBrush("LevelEditor.Tabs.Details") )
 		.Label( LOCTEXT("MaterialDetailsTitle", "Details") )
 		[
@@ -3246,8 +3246,8 @@ TSharedRef<SDockTab> FMaterialEditor::SpawnTab_MaterialProperties(const FSpawnTa
 		// Since we're initialising, make sure nothing is selected
 		GraphEditor->ClearSelectionSet();
 	}
-
-	return SpawnedDetailsTab.ToSharedRef();
+	SpawnedDetailsTab = DetailsTab;
+	return DetailsTab.ToSharedRef();
 }
 
 TSharedRef<SDockTab> FMaterialEditor::SpawnTab_Palette(const FSpawnTabArgs& Args)
@@ -4750,9 +4750,9 @@ void FMaterialEditor::NotifyExternalMaterialChange()
 
 void FMaterialEditor::FocusDetailsPanel()
 {
-	if (SpawnedDetailsTab && !SpawnedDetailsTab->IsForeground())
+	if (SpawnedDetailsTab.IsValid() && !SpawnedDetailsTab.Pin()->IsForeground())
 	{
-		SpawnedDetailsTab->DrawAttention();
+		SpawnedDetailsTab.Pin()->DrawAttention();
 	}
 }
 
diff --git a/Engine/Source/Editor/MaterialEditor/Private/MaterialEditor.h b/Engine/Source/Editor/MaterialEditor/Private/MaterialEditor.h
--- a/Engine/Source/Editor/MaterialEditor/Private/MaterialEditor.h
+++ b/Engine/Source/Editor/MaterialEditor/Private/MaterialEditor.h
@@ -797,7 +797,7 @@ private:
 	TSharedPtr<class FMaterialStats> MaterialStatsManager;
 
 	/** Tab that holds the details panel */
-	TSharedPtr<SDockTab> SpawnedDetailsTab;
+	TWeakPtr<SDockTab> SpawnedDetailsTab;
 
 	/** Stores the quality level used to preview the material graph */
 	EMaterialQualityLevel::Type NodeQualityLevel;
diff --git a/Engine/Source/Editor/PropertyEditor/Private/UserInterface/PropertyEditor/SPropertyEditorAsset.cpp b/Engine/Source/Editor/PropertyEditor/Private/UserInterface/PropertyEditor/SPropertyEditorAsset.cpp
--- a/Engine/Source/Editor/PropertyEditor/Private/UserInterface/PropertyEditor/SPropertyEditorAsset.cpp
+++ b/Engine/Source/Editor/PropertyEditor/Private/UserInterface/PropertyEditor/SPropertyEditorAsset.cpp
@@ -215,7 +215,8 @@ void SPropertyEditorAsset::Construct( const FArguments& InArgs, const TSharedPtr
 	{
 		NewAssetFactories = InArgs._NewAssetFactories.GetValue();
 	}
-	else if (AllowedClassFilters.Num() > 1 && !AllowedClassFilters.Contains(UObject::StaticClass()))
+	// If there are more allowed classes than just UObject 
+	else if (AllowedClassFilters.Num() > 1 || !AllowedClassFilters.Contains(UObject::StaticClass()))
 	{
 		NewAssetFactories = PropertyCustomizationHelpers::GetNewAssetFactoriesForClasses(AllowedClassFilters, DisallowedClassFilters);
 	}
diff --git a/Engine/Source/Editor/UMGEditor/Private/Designer/SDesignerView.cpp b/Engine/Source/Editor/UMGEditor/Private/Designer/SDesignerView.cpp
--- a/Engine/Source/Editor/UMGEditor/Private/Designer/SDesignerView.cpp
+++ b/Engine/Source/Editor/UMGEditor/Private/Designer/SDesignerView.cpp
@@ -1558,7 +1558,10 @@ void SDesignerView::ClearDropPreviews()
 
 		// Since the widget has been removed from the widget tree, move it into the transient package. Otherwise,
 		// it will remain outered to the widget tree and end up as a property in the BP class layout as a result.
-		DropPreview.Widget->Rename(nullptr, GetTransientPackage());
+		if (DropPreview.Widget->GetOutermost() != GetTransientPackage())
+		{
+			DropPreview.Widget->Rename(nullptr, GetTransientPackage());
+		}
 	}
 	DropPreviews.Empty();
 }
@@ -2787,7 +2790,10 @@ void SDesignerView::ProcessDropAndAddWidget(const FGeometry& MyGeometry, const F
 		// Move the remaining widgets into the transient package. Otherwise, they will remain outered to the WidgetTree and end up as properties in the BP class layout as a result.
 		for (UWidget* Widget : DragDropPreviewWidgets)
 		{
-			Widget->Rename(nullptr, GetTransientPackage());
+			if (Widget->GetOutermost() != GetTransientPackage())
+			{
+				Widget->Rename(nullptr, GetTransientPackage());
+			}
 		}
 
 		// If we had preview widgets, we know that we can not be performing a selected widget drag/drop operation. Bail.
@@ -2841,24 +2847,24 @@ void SDesignerView::ProcessDropAndAddWidget(const FGeometry& MyGeometry, const F
 			BlueprintEditor.Pin()->SetHoveredWidget(TargetReference);
 
 			// If the widget being hovered over is a panel, attempt to place it into that panel.
-			if (Target && Target->IsA(UPanelWidget::StaticClass()))
+			if (Target && Target->IsA(UPanelWidget::StaticClass()) && !DraggedWidget.bStayingInParent)
 			{
 				bWidgetMoved = true;
 				UPanelWidget* NewParent = Cast<UPanelWidget>(Target);
 
 				UWidget* Widget = bIsPreview ? DraggedWidget.Preview : DraggedWidget.Template;
-
+				UPanelWidget* ParentWidget = bIsPreview ? Cast<UPanelWidget>(DraggedWidget.ParentWidget.GetPreview()) : Cast<UPanelWidget>(DraggedWidget.ParentWidget.GetTemplate());
 				if (ensure(Widget))
 				{
-					bool bIsChangingParent = Widget->GetParent() != NewParent;
+					bool bIsChangingParent = ParentWidget != NewParent;
 					UBlueprint* OriginalBP = nullptr;
 
-					check(Widget->GetParent() != nullptr || bIsChangingParent);
-
-					// If this isn't a preview operation we need to modify a few things to properly undo the operation.
-					if (!bIsPreview)
+					if (bIsChangingParent)
 					{
-						if (bIsChangingParent)
+						check(ParentWidget != nullptr);
+
+						// If this isn't a preview operation we need to modify a few things to properly undo the operation.
+						if (!bIsPreview)
 						{
 							NewParent->SetFlags(RF_Transactional);
 							NewParent->Modify();
@@ -2868,26 +2874,17 @@ void SDesignerView::ProcessDropAndAddWidget(const FGeometry& MyGeometry, const F
 
 							// If the Widget is changing parents, there's a chance it might be moving to a different WidgetTree as well.
 							UWidgetTree* OriginalWidgetTree = Cast<UWidgetTree>(Widget->GetOuter());
-							
+
 							if (UWidgetTree::TryMoveWidgetToNewTree(Widget, BP->WidgetTree))
 							{
 								// The Widget likely originated from a different blueprint, so get what blueprint it was originally a part of.
 								OriginalBP = OriginalWidgetTree ? OriginalWidgetTree->GetTypedOuter<UBlueprint>() : nullptr;
 							}
-						}
 
-						Widget->Modify();
-					}
-
-					if (Widget->GetParent() && bIsChangingParent)
-					{
-						if (!bIsPreview)
-						{
-							Widget->GetParent()->Modify();
+							Widget->Modify();
+							ParentWidget->Modify();
 						}
 
-						Widget->GetParent()->RemoveChild(Widget);
-
 						// The Widget originated from a different blueprint, so mark it as modified.
 						if (OriginalBP && OriginalBP != BP)
 						{
@@ -2927,17 +2924,11 @@ void SDesignerView::ProcessDropAndAddWidget(const FGeometry& MyGeometry, const F
 
 					FVector2D LocalPosition = WidgetUnderCursorGeometry.AbsoluteToLocal(ScreenSpacePosition);
 
-					UPanelSlot* Slot;
+					UPanelSlot* Slot = nullptr;
 					if (bIsChangingParent)
 					{
 						Slot = NewParent->AddChild(Widget);
 					}
-					else
-					{
-						check(Widget->GetParent()->GetChildIndex(Widget) != INDEX_NONE);
-
-						Slot = Widget->GetParent()->GetSlots()[Widget->GetParent()->GetChildIndex(Widget)];
-					}
 
 					if (Slot != nullptr)
 					{
diff --git a/Engine/Source/Editor/UMGEditor/Public/WidgetBlueprint.h b/Engine/Source/Editor/UMGEditor/Public/WidgetBlueprint.h
--- a/Engine/Source/Editor/UMGEditor/Public/WidgetBlueprint.h
+++ b/Engine/Source/Editor/UMGEditor/Public/WidgetBlueprint.h
@@ -176,7 +176,7 @@ struct TStructOpsTypeTraits<FWidgetAnimation_DEPRECATED> : public TStructOpsType
 {
 	enum
 	{
-		WithSerializeFromMismatchedTag = true,
+		WithStructuredSerializeFromMismatchedTag = true,
 	};
 };
 
diff --git a/Engine/Source/Programs/UnrealBuildTool/Configuration/UEBuildBinary.cs b/Engine/Source/Programs/UnrealBuildTool/Configuration/UEBuildBinary.cs
--- a/Engine/Source/Programs/UnrealBuildTool/Configuration/UEBuildBinary.cs
+++ b/Engine/Source/Programs/UnrealBuildTool/Configuration/UEBuildBinary.cs
@@ -782,8 +782,12 @@ namespace UnrealBuildTool
 					}
 					else
 					{
+						// Get the intermediate directory
+						DirectoryReference ResourceIntermediateDirectory = ((UEBuildModuleCPP)Modules.First()).IntermediateDirectory;
+
 						// Create a compile environment for resource files
 						CppCompileEnvironment ResourceCompileEnvironment = new CppCompileEnvironment(BinaryCompileEnvironment);
+						WindowsPlatform.SetupResourceCompileEnvironment(ResourceCompileEnvironment, ResourceIntermediateDirectory, Target);
 
 						// @todo: This should be in some Windows code somewhere...
 						// Set the original file name macro; used in PCLaunch.rc to set the binary metadata fields.
@@ -796,7 +800,7 @@ namespace UnrealBuildTool
 
 						// Otherwise compile the default resource file per-binary, so that it gets the correct ORIGINAL_FILE_NAME macro.
 						FileItem DefaultResourceFile = FileItem.GetItemByFileReference(FileReference.Combine(UnrealBuildTool.EngineSourceDirectory, "Runtime", "Launch", "Resources", "Windows", "PCLaunch.rc"));
-						CPPOutput DefaultResourceOutput = ToolChain.CompileRCFiles(BinaryCompileEnvironment, new List<FileItem> { DefaultResourceFile }, ((UEBuildModuleCPP)Modules.First()).IntermediateDirectory, ActionGraph);
+						CPPOutput DefaultResourceOutput = ToolChain.CompileRCFiles(ResourceCompileEnvironment, new List<FileItem> { DefaultResourceFile }, ResourceIntermediateDirectory, ActionGraph);
 						BinaryLinkEnvironment.InputFiles.AddRange(DefaultResourceOutput.ObjectFiles);
 					}
 				}
diff --git a/Engine/Source/Programs/UnrealBuildTool/Configuration/UEBuildModuleCPP.cs b/Engine/Source/Programs/UnrealBuildTool/Configuration/UEBuildModuleCPP.cs
--- a/Engine/Source/Programs/UnrealBuildTool/Configuration/UEBuildModuleCPP.cs
+++ b/Engine/Source/Programs/UnrealBuildTool/Configuration/UEBuildModuleCPP.cs
@@ -565,20 +565,8 @@ namespace UnrealBuildTool
 			// must be specified relative to the resource file itself or Engine/Source.
 			if(SourceFilesToBuild.RCFiles.Count > 0)
 			{
-				// Figure the icon to use. We can only use a custom icon when compiling to a project-specific intemediate directory (and not for the shared editor executable, for example).
-				FileReference IconFile;
-				if(Target.ProjectFile != null && IntermediateDirectory.IsUnderDirectory(Target.ProjectFile.Directory))
-				{
-					IconFile = WindowsPlatform.GetApplicationIcon(Target.ProjectFile);
-				}
-				else
-				{
-					IconFile = WindowsPlatform.GetApplicationIcon(null);
-				}
-
-				// Setup the compile environment, setting the icon to use via a macro. This is used in PCLaunch.rc2.
 				CppCompileEnvironment ResourceCompileEnvironment = new CppCompileEnvironment(BinaryCompileEnvironment);
-				ResourceCompileEnvironment.Definitions.Add(String.Format("BUILD_ICON_FILE_NAME=\"\\\"{0}\\\"\"", IconFile.FullName.Replace("\\", "\\\\")));
+				WindowsPlatform.SetupResourceCompileEnvironment(ResourceCompileEnvironment, IntermediateDirectory, Target);
 				LinkInputFiles.AddRange(ToolChain.CompileRCFiles(ResourceCompileEnvironment, SourceFilesToBuild.RCFiles, IntermediateDirectory, ActionGraph).ObjectFiles);
 			}
 
diff --git a/Engine/Source/Programs/UnrealBuildTool/Configuration/UEBuildTarget.cs b/Engine/Source/Programs/UnrealBuildTool/Configuration/UEBuildTarget.cs
--- a/Engine/Source/Programs/UnrealBuildTool/Configuration/UEBuildTarget.cs
+++ b/Engine/Source/Programs/UnrealBuildTool/Configuration/UEBuildTarget.cs
@@ -2189,8 +2189,10 @@ namespace UnrealBuildTool
 
 						FileItem DefaultResourceFile = FileItem.GetExistingItemByFileReference(FileReference.Combine(UnrealBuildTool.EngineSourceDirectory, "Runtime", "Launch", "Resources", "Windows", "PCLaunch.rc"));
 						DefaultResourceFile.CachedIncludePaths = DefaultResourceCompileEnvironment.IncludePaths;
-						CPPOutput DefaultResourceOutput = TargetToolChain.CompileRCFiles(DefaultResourceCompileEnvironment, new List<FileItem> { DefaultResourceFile }, EngineIntermediateDirectory, ActionGraph);
 
+						WindowsPlatform.SetupResourceCompileEnvironment(DefaultResourceCompileEnvironment, EngineIntermediateDirectory, Rules);
+
+						CPPOutput DefaultResourceOutput = TargetToolChain.CompileRCFiles(DefaultResourceCompileEnvironment, new List<FileItem> { DefaultResourceFile }, EngineIntermediateDirectory, ActionGraph);
 						GlobalLinkEnvironment.DefaultResourceFiles.AddRange(DefaultResourceOutput.ObjectFiles);
 					}
 				}
@@ -4051,6 +4053,7 @@ namespace UnrealBuildTool
 				DirectoryReference GeneratedCodeDirectory = null;
 				if (RulesObject.Type != ModuleRules.ModuleType.External)
 				{
+					// Get the base directory
 					if (RulesObject.Plugin != null)
 					{
 						GeneratedCodeDirectory = RulesObject.Plugin.Directory;
@@ -4067,7 +4070,18 @@ namespace UnrealBuildTool
 					{
 						GeneratedCodeDirectory = ProjectDirectory;
 					}
-					GeneratedCodeDirectory = DirectoryReference.Combine(GeneratedCodeDirectory, PlatformIntermediateFolder, AppName, "Inc", ModuleName);
+
+					// Get the subfolder containing generated code
+					GeneratedCodeDirectory = DirectoryReference.Combine(GeneratedCodeDirectory, PlatformIntermediateFolder, AppName, "Inc");
+
+					// Append the binaries subfolder, if present. We rely on this to ensure that build products can be filtered correctly.
+					if(RulesObject.BinariesSubFolder != null)
+					{
+						GeneratedCodeDirectory = DirectoryReference.Combine(GeneratedCodeDirectory, RulesObject.BinariesSubFolder);
+					}
+
+					// Finally, append the module name.
+					GeneratedCodeDirectory = DirectoryReference.Combine(GeneratedCodeDirectory, ModuleName);
 				}
 
 				// For legacy modules, add a bunch of default include paths.
diff --git a/Engine/Source/Programs/UnrealBuildTool/Platform/Windows/UEBuildWindows.cs b/Engine/Source/Programs/UnrealBuildTool/Platform/Windows/UEBuildWindows.cs
--- a/Engine/Source/Programs/UnrealBuildTool/Platform/Windows/UEBuildWindows.cs
+++ b/Engine/Source/Programs/UnrealBuildTool/Platform/Windows/UEBuildWindows.cs
@@ -1349,6 +1349,29 @@ namespace UnrealBuildTool
 			return FileReference.Combine(UnrealBuildTool.EngineDirectory, "Source", "Runtime", "Launch", "Resources", "Windows", "UE4.ico");
 		}
 
+		/// <summary>
+		/// Configures the resource compile environment for the given target
+		/// </summary>
+		/// <param name="ResourceCompileEnvironment">The compile environment</param>
+		/// <param name="IntermediateDirectory">The output directory for compiled files</param>
+		/// <param name="Target">The target being built</param>
+		public static void SetupResourceCompileEnvironment(CppCompileEnvironment ResourceCompileEnvironment, DirectoryReference IntermediateDirectory, ReadOnlyTargetRules Target)
+		{
+			// Figure the icon to use. We can only use a custom icon when compiling to a project-specific intemediate directory (and not for the shared editor executable, for example).
+			FileReference IconFile;
+			if(Target.ProjectFile != null && IntermediateDirectory.IsUnderDirectory(Target.ProjectFile.Directory))
+			{
+				IconFile = WindowsPlatform.GetApplicationIcon(Target.ProjectFile);
+			}
+			else
+			{
+				IconFile = WindowsPlatform.GetApplicationIcon(null);
+			}
+
+			// Setup the compile environment, setting the icon to use via a macro. This is used in PCLaunch.rc2.
+			ResourceCompileEnvironment.Definitions.Add(String.Format("BUILD_ICON_FILE_NAME=\"\\\"{0}\\\"\"", IconFile.FullName.Replace("\\", "\\\\")));
+		}
+
 		/// <summary>
 		/// Modify the rules for a newly created module, in a target that's being built for this platform.
 		/// This is not required - but allows for hiding details of a particular platform.
diff --git a/Engine/Source/Runtime/AnimGraphRuntime/Public/AnimNodes/AnimNode_StateResult.h b/Engine/Source/Runtime/AnimGraphRuntime/Public/AnimNodes/AnimNode_StateResult.h
--- a/Engine/Source/Runtime/AnimGraphRuntime/Public/AnimNodes/AnimNode_StateResult.h
+++ b/Engine/Source/Runtime/AnimGraphRuntime/Public/AnimNodes/AnimNode_StateResult.h
@@ -24,6 +24,6 @@ struct TStructOpsTypeTraits<FAnimNode_StateResult>
 {
 	enum
 	{
-		WithSerializeFromMismatchedTag = true,
+		WithStructuredSerializeFromMismatchedTag = true,
 	};
 };
diff --git a/Engine/Source/Runtime/CoreUObject/Private/UObject/Class.cpp b/Engine/Source/Runtime/CoreUObject/Private/UObject/Class.cpp
--- a/Engine/Source/Runtime/CoreUObject/Private/UObject/Class.cpp
+++ b/Engine/Source/Runtime/CoreUObject/Private/UObject/Class.cpp
@@ -2048,7 +2048,7 @@ void UScriptStruct::PrepareCppStructOps()
 		UE_LOG(LogClass, Verbose, TEXT("Native struct %s has native ImportTextItem."),*GetName());
 		StructFlags = EStructFlags(StructFlags | STRUCT_ImportTextItemNative);
 	}
-	if (CppStructOps->HasSerializeFromMismatchedTag())
+	if (CppStructOps->HasSerializeFromMismatchedTag() || CppStructOps->HasStructuredSerializeFromMismatchedTag())
 	{
 		UE_LOG(LogClass, Verbose, TEXT("Native struct %s has native SerializeFromMismatchedTag."),*GetName());
 		StructFlags = EStructFlags(StructFlags | STRUCT_SerializeFromMismatchedTag);
diff --git a/Engine/Source/Runtime/CoreUObject/Private/UObject/Property.cpp b/Engine/Source/Runtime/CoreUObject/Private/UObject/Property.cpp
--- a/Engine/Source/Runtime/CoreUObject/Private/UObject/Property.cpp
+++ b/Engine/Source/Runtime/CoreUObject/Private/UObject/Property.cpp
@@ -302,7 +302,7 @@ struct TStructOpsTypeTraits<FSoftObjectPath> : public TStructOpsTypeTraitsBase2<
 		WithIdenticalViaEquality = true,
 		WithExportTextItem = true,
 		WithImportTextItem = true,
-		WithSerializeFromMismatchedTag = true,
+		WithStructuredSerializeFromMismatchedTag = true,
 	};
 };
 IMPLEMENT_STRUCT(SoftObjectPath);
@@ -318,7 +318,7 @@ struct TStructOpsTypeTraits<FSoftClassPath> : public TStructOpsTypeTraitsBase2<F
 		WithIdenticalViaEquality = true,
 		WithExportTextItem = true,
 		WithImportTextItem = true,
-		WithSerializeFromMismatchedTag = true,
+		WithStructuredSerializeFromMismatchedTag = true,
 	};
 };
 IMPLEMENT_STRUCT(SoftClassPath);
@@ -333,7 +333,7 @@ struct TStructOpsTypeTraits<FPrimaryAssetType> : public TStructOpsTypeTraitsBase
 		WithIdenticalViaEquality = true,
 		WithExportTextItem = true,
 		WithImportTextItem = true,
-		WithSerializeFromMismatchedTag = true,
+		WithStructuredSerializeFromMismatchedTag = true,
 	};
 };
 IMPLEMENT_STRUCT(PrimaryAssetType);
@@ -348,7 +348,7 @@ struct TStructOpsTypeTraits<FPrimaryAssetId> : public TStructOpsTypeTraitsBase2<
 		WithIdenticalViaEquality = true,
 		WithExportTextItem = true,
 		WithImportTextItem = true,
-		WithSerializeFromMismatchedTag = true,
+		WithStructuredSerializeFromMismatchedTag = true,
 	};
 };
 IMPLEMENT_STRUCT(PrimaryAssetId);
diff --git a/Engine/Source/Runtime/CoreUObject/Private/UObject/PropertyStruct.cpp b/Engine/Source/Runtime/CoreUObject/Private/UObject/PropertyStruct.cpp
--- a/Engine/Source/Runtime/CoreUObject/Private/UObject/PropertyStruct.cpp
+++ b/Engine/Source/Runtime/CoreUObject/Private/UObject/PropertyStruct.cpp
@@ -345,16 +345,25 @@ EConvertFromTypeResult UStructProperty::ConvertFromType(const FPropertyTag& Tag,
 		if ((Struct->StructFlags & STRUCT_SerializeFromMismatchedTag) && (Tag.Type != NAME_StructProperty || (Tag.StructName != Struct->GetFName())))
 		{
 			UScriptStruct::ICppStructOps* CppStructOps = Struct->GetCppStructOps();
-			check(CppStructOps && CppStructOps->HasSerializeFromMismatchedTag()); // else should not have STRUCT_SerializeFromMismatchedTag
+			check(CppStructOps && (CppStructOps->HasSerializeFromMismatchedTag() || CppStructOps->HasStructuredSerializeFromMismatchedTag())); // else should not have STRUCT_SerializeFromMismatchedTag
 			void* DestAddress = ContainerPtrToValuePtr<void>(Data, Tag.ArrayIndex);
-			if (CppStructOps->SerializeFromMismatchedTag(Tag, Slot, DestAddress))
+			if (CppStructOps->HasStructuredSerializeFromMismatchedTag() && CppStructOps->StructuredSerializeFromMismatchedTag(Tag, Slot, DestAddress))
 			{
 				return EConvertFromTypeResult::Converted;
 			}
-			else
+			else 
 			{
-				UE_LOG(LogClass, Warning, TEXT("SerializeFromMismatchedTag failed: Type mismatch in %s of %s - Previous (%s) Current(StructProperty) for package:  %s"), *Tag.Name.ToString(), *GetName(), *Tag.Type.ToString(), *UnderlyingArchive.GetArchiveName());
-				return EConvertFromTypeResult::CannotConvert;
+				FArchiveUObjectFromStructuredArchive Ar(Slot);
+				if (CppStructOps->HasSerializeFromMismatchedTag() && CppStructOps->SerializeFromMismatchedTag(Tag, Ar, DestAddress))
+				{
+
+					return EConvertFromTypeResult::Converted;
+				}
+				else
+				{
+					UE_LOG(LogClass, Warning, TEXT("SerializeFromMismatchedTag failed: Type mismatch in %s of %s - Previous (%s) Current(StructProperty) for package:  %s"), *Tag.Name.ToString(), *GetName(), *Tag.Type.ToString(), *UnderlyingArchive.GetArchiveName());
+					return EConvertFromTypeResult::CannotConvert;
+				}
 			}
 		}
 
diff --git a/Engine/Source/Runtime/CoreUObject/Public/UObject/Class.h b/Engine/Source/Runtime/CoreUObject/Public/UObject/Class.h
--- a/Engine/Source/Runtime/CoreUObject/Public/UObject/Class.h
+++ b/Engine/Source/Runtime/CoreUObject/Public/UObject/Class.h
@@ -538,33 +538,12 @@ struct TStructOpsTypeTraitsBase2
 		WithNetSerializer              = false,                         // struct has a NetSerialize function for serializing its state to an FArchive used for network replication.
 		WithNetDeltaSerializer         = false,                         // struct has a NetDeltaSerialize function for serializing differences in state from a previous NetSerialize operation.
 		WithSerializeFromMismatchedTag = false,                         // struct has a SerializeFromMismatchedTag function for converting from other property tags.
-		WithPostScriptConstruct        = false,							// struct has a PostScriptConstruct function which is called after it is constructed in blueprints
+		WithStructuredSerializeFromMismatchedTag = false,               // struct has an FStructuredArchive-based SerializeFromMismatchedTag function for converting from other property tags.
+		WithPostScriptConstruct        = false,				// struct has a PostScriptConstruct function which is called after it is constructed in blueprints
 		WithNetSharedSerialization     = false,                         // struct has a NetSerialize function that does not require the package map to serialize its state.
 	};
 };
 
-/** type traits to cover the custom aspects of a script struct **/
-struct DEPRECATED(4.16, "TStructOpsTypeTraitsBase has been deprecated, use TStructOpsTypeTraitsBase2<T> instead.") TStructOpsTypeTraitsBase
-{
-	enum
-	{
-		WithZeroConstructor            = false, // struct can be constructed as a valid object by filling its memory footprint with zeroes.
-		WithNoInitConstructor          = false, // struct has a constructor which takes an EForceInit parameter which will force the constructor to perform initialization, where the default constructor performs 'uninitialization'.
-		WithNoDestructor               = false, // struct will not have its destructor called when it is destroyed.
-		WithCopy                       = false, // struct can be copied via its copy assignment operator.
-		WithIdenticalViaEquality       = false, // struct can be compared via its operator==.  This should be mutually exclusive with WithIdentical.
-		WithIdentical                  = false, // struct can be compared via an Identical(const T* Other, uint32 PortFlags) function.  This should be mutually exclusive with WithIdenticalViaEquality.
-		WithExportTextItem             = false, // struct has an ExportTextItem function used to serialize its state into a string.
-		WithImportTextItem             = false, // struct has an ImportTextItem function used to deserialize a string into an object of that class.
-		WithAddStructReferencedObjects = false, // struct has an AddStructReferencedObjects function which allows it to add references to the garbage collector.
-		WithSerializer                 = false, // struct has a Serialize function for serializing its state to an FArchive.
-		WithPostSerialize              = false, // struct has a PostSerialize function which is called after it is serialized
-		WithNetSerializer              = false, // struct has a NetSerialize function for serializing its state to an FArchive used for network replication.
-		WithNetDeltaSerializer         = false, // struct has a NetDeltaSerialize function for serializing differences in state from a previous NetSerialize operation.
-		WithSerializeFromMismatchedTag = false, // struct has a SerializeFromMismatchedTag function for converting from other property tags.
-	};
-};
-
 template<class CPPSTRUCT>
 struct TStructOpsTypeTraits : public TStructOpsTypeTraitsBase2<CPPSTRUCT>
 {
@@ -783,13 +762,25 @@ FORCEINLINE typename TEnableIf<TStructOpsTypeTraits<CPPSTRUCT>::WithImportTextIt
  * Selection of SerializeFromMismatchedTag call.
  */
 template<class CPPSTRUCT>
-FORCEINLINE typename TEnableIf<!TStructOpsTypeTraits<CPPSTRUCT>::WithSerializeFromMismatchedTag, bool>::Type SerializeFromMismatchedTagOrNot(FPropertyTag const& Tag, FStructuredArchive::FSlot Slot, CPPSTRUCT *Data)
+FORCEINLINE typename TEnableIf<!TStructOpsTypeTraits<CPPSTRUCT>::WithSerializeFromMismatchedTag, bool>::Type SerializeFromMismatchedTagOrNot(FPropertyTag const& Tag, FArchive& Ar, CPPSTRUCT *Data)
 {
 	return false;
 }
 
 template<class CPPSTRUCT>
-FORCEINLINE typename TEnableIf<TStructOpsTypeTraits<CPPSTRUCT>::WithSerializeFromMismatchedTag, bool>::Type SerializeFromMismatchedTagOrNot(FPropertyTag const& Tag, FStructuredArchive::FSlot Slot, CPPSTRUCT *Data)
+FORCEINLINE typename TEnableIf<TStructOpsTypeTraits<CPPSTRUCT>::WithSerializeFromMismatchedTag, bool>::Type SerializeFromMismatchedTagOrNot(FPropertyTag const& Tag, FArchive& Ar, CPPSTRUCT *Data)
+{
+	return Data->SerializeFromMismatchedTag(Tag, Ar);
+}
+
+template<class CPPSTRUCT>
+FORCEINLINE typename TEnableIf<!TStructOpsTypeTraits<CPPSTRUCT>::WithStructuredSerializeFromMismatchedTag, bool>::Type StructuredSerializeFromMismatchedTagOrNot(FPropertyTag const& Tag, FStructuredArchive::FSlot Slot, CPPSTRUCT *Data)
+{
+	return false;
+}
+
+template<class CPPSTRUCT>
+FORCEINLINE typename TEnableIf<TStructOpsTypeTraits<CPPSTRUCT>::WithStructuredSerializeFromMismatchedTag, bool>::Type StructuredSerializeFromMismatchedTagOrNot(FPropertyTag const& Tag, FStructuredArchive::FSlot Slot, CPPSTRUCT *Data)
 {
 	return Data->SerializeFromMismatchedTag(Tag, Slot);
 }
@@ -937,11 +928,14 @@ public:
 
 		/** return true if this class wants to serialize from some other tag (usually for conversion purposes) **/
 		virtual bool HasSerializeFromMismatchedTag() = 0;
+		virtual bool HasStructuredSerializeFromMismatchedTag() = 0;
+
 		/** 
 		 * Serialize this structure, from some other tag
 		 * @return true if this succeeded, false will trigger a warning and not serialize at all
 		 */
-		virtual bool SerializeFromMismatchedTag(struct FPropertyTag const& Tag, FStructuredArchive::FSlot Slot, void *Data) = 0;
+		virtual bool SerializeFromMismatchedTag(struct FPropertyTag const& Tag, FArchive& Ar, void *Data) = 0;
+		virtual bool StructuredSerializeFromMismatchedTag(struct FPropertyTag const& Tag, FStructuredArchive::FSlot Slot, void *Data) = 0;
 
 		/** return true if this struct has a GetTypeHash */
 		virtual bool HasGetTypeHash() = 0;
@@ -1105,11 +1099,23 @@ public:
 		{
 			return TTraits::WithSerializeFromMismatchedTag;
 		}
-		virtual bool SerializeFromMismatchedTag(struct FPropertyTag const& Tag, FStructuredArchive::FSlot Slot, void *Data) override
+		virtual bool SerializeFromMismatchedTag(struct FPropertyTag const& Tag, FArchive& Ar, void *Data) override
 		{
 			check(TTraits::WithSerializeFromMismatchedTag); // don't call this if we have indicated it is not allowed
-			return SerializeFromMismatchedTagOrNot(Tag, Slot, (CPPSTRUCT*)Data);
+			return SerializeFromMismatchedTagOrNot(Tag, Ar, (CPPSTRUCT*)Data);
 		}
+		virtual bool HasStructuredSerializeFromMismatchedTag() override
+		{
+			return TTraits::WithStructuredSerializeFromMismatchedTag;
+		}
+		virtual bool StructuredSerializeFromMismatchedTag(struct FPropertyTag const& Tag, FStructuredArchive::FSlot Slot, void *Data) override
+		{
+			check(TTraits::WithStructuredSerializeFromMismatchedTag); // don't call this if we have indicated it is not allowed
+			return StructuredSerializeFromMismatchedTagOrNot(Tag, Slot, (CPPSTRUCT*)Data);
+		}
+
+		static_assert(!(TTraits::WithSerializeFromMismatchedTag && TTraits::WithStructuredSerializeFromMismatchedTag), "Structs cannot have both WithSerializeFromMismatchedTag and WithStructuredSerializeFromMismatchedTag set");
+
 		virtual bool HasGetTypeHash() override
 		{
 			return THasGetTypeHash<CPPSTRUCT>::Value;
diff --git a/Engine/Source/Runtime/Engine/Classes/Engine/MeshMerging.h b/Engine/Source/Runtime/Engine/Classes/Engine/MeshMerging.h
--- a/Engine/Source/Runtime/Engine/Classes/Engine/MeshMerging.h
+++ b/Engine/Source/Runtime/Engine/Classes/Engine/MeshMerging.h
@@ -24,67 +24,67 @@ namespace EMeshFeatureImportance
 }
 
 /** Settings used to reduce a mesh. */
-USTRUCT()
+USTRUCT(Blueprintable)
 struct FMeshReductionSettings
 {
 	GENERATED_USTRUCT_BODY()
 
 	/** Percentage of triangles to keep. 1.0 = no reduction, 0.0 = no triangles. */
-	UPROPERTY(EditAnywhere, Category = ReductionSettings)
+	UPROPERTY(BlueprintReadWrite, EditAnywhere, Category = ReductionSettings)
 	float PercentTriangles;
 
 	/** The maximum distance in object space by which the reduced mesh may deviate from the original mesh. */
-	UPROPERTY(EditAnywhere, Category = ReductionSettings)
+	UPROPERTY(BlueprintReadWrite, EditAnywhere, Category = ReductionSettings)
 	float MaxDeviation;
 
 	/** The amount of error in pixels allowed for this LOD. */
-	UPROPERTY(EditAnywhere, Category = ReductionSettings)
+	UPROPERTY(BlueprintReadWrite, EditAnywhere, Category = ReductionSettings)
 	float PixelError;
 
 	/** Threshold in object space at which vertices are welded together. */
-	UPROPERTY(EditAnywhere, Category = ReductionSettings)
+	UPROPERTY(BlueprintReadWrite, EditAnywhere, Category = ReductionSettings)
 	float WeldingThreshold;
 
 	/** Angle at which a hard edge is introduced between faces. */
-	UPROPERTY(EditAnywhere, Category = ReductionSettings)
+	UPROPERTY(BlueprintReadWrite, EditAnywhere, Category = ReductionSettings)
 	float HardAngleThreshold;
 
 	/** Higher values minimize change to border edges. */
-	UPROPERTY(EditAnywhere, Category = ReductionSettings)
+	UPROPERTY(BlueprintReadWrite, EditAnywhere, Category = ReductionSettings)
 	TEnumAsByte<EMeshFeatureImportance::Type> SilhouetteImportance;
 
 	/** Higher values reduce texture stretching. */
-	UPROPERTY(EditAnywhere, Category = ReductionSettings)
+	UPROPERTY(BlueprintReadWrite, EditAnywhere, Category = ReductionSettings)
 	TEnumAsByte<EMeshFeatureImportance::Type> TextureImportance;
 
 	/** Higher values try to preserve normals better. */
-	UPROPERTY(EditAnywhere, Category = ReductionSettings)
+	UPROPERTY(BlueprintReadWrite, EditAnywhere, Category = ReductionSettings)
 	TEnumAsByte<EMeshFeatureImportance::Type> ShadingImportance;
 
-	UPROPERTY(EditAnywhere, Category = ReductionSettings)
+	UPROPERTY(BlueprintReadWrite, EditAnywhere, Category = ReductionSettings)
 	bool bRecalculateNormals;
 
-	UPROPERTY(EditAnywhere, Category = ReductionSettings)
+	UPROPERTY(BlueprintReadWrite, EditAnywhere, Category = ReductionSettings)
 	int32 BaseLODModel;
 
-	UPROPERTY(EditAnywhere, Category = ReductionSettings)
+	UPROPERTY(BlueprintReadWrite, EditAnywhere, Category = ReductionSettings)
 	bool bGenerateUniqueLightmapUVs;
 
-	UPROPERTY(EditAnywhere, Category = ReductionSettings)
+	UPROPERTY(BlueprintReadWrite, EditAnywhere, Category = ReductionSettings)
 	bool bKeepSymmetry;
 
-	UPROPERTY(EditAnywhere, Category = ReductionSettings)
+	UPROPERTY(BlueprintReadWrite, EditAnywhere, Category = ReductionSettings)
 	bool bVisibilityAided;
 
-	UPROPERTY(EditAnywhere, Category = ReductionSettings)
+	UPROPERTY(BlueprintReadWrite, EditAnywhere, Category = ReductionSettings)
 	bool bCullOccluded;
 
 	/** Higher values generates fewer samples*/
-	UPROPERTY(EditAnywhere, Category = ReductionSettings)
+	UPROPERTY(BlueprintReadWrite, EditAnywhere, Category = ReductionSettings)
 	TEnumAsByte<EMeshFeatureImportance::Type> VisibilityAggressiveness;
 
 	/** Higher values minimize change to vertex color data. */
-	UPROPERTY(EditAnywhere, Category = ReductionSettings)
+	UPROPERTY(BlueprintReadWrite, EditAnywhere, Category = ReductionSettings)
 	TEnumAsByte<EMeshFeatureImportance::Type> VertexColorImportance;
 
 	/** Default settings. */
@@ -184,24 +184,24 @@ struct FMeshProxySettings
 {
 	GENERATED_USTRUCT_BODY()
 	/** Screen size of the resulting proxy mesh in pixels*/
-	UPROPERTY(EditAnywhere, Category = ProxySettings, meta = (ClampMin = "1", ClampMax = "1200", UIMin = "1", UIMax = "1200"))
+	UPROPERTY(BlueprintReadWrite, EditAnywhere, Category = ProxySettings, meta = (ClampMin = "1", ClampMax = "1200", UIMin = "1", UIMax = "1200"))
 	int32 ScreenSize;
 
 	/** If true, Spatial Sampling Distance will not be automatically computed based on geometry and you must set it directly */
-	UPROPERTY(EditAnywhere, AdvancedDisplay, Category = ProxySettings, meta = (InlineEditConditionToggle))
+	UPROPERTY(BlueprintReadWrite, EditAnywhere, AdvancedDisplay, Category = ProxySettings, meta = (InlineEditConditionToggle))
 	uint8 bOverrideVoxelSize : 1;
 
-
 	/** Override when converting multiple meshes for proxy LOD merging. Warning, large geometry with small sampling has very high memory costs*/
-	UPROPERTY(EditAnywhere, AdvancedDisplay, Category = ProxySettings, meta = (editcondition = "bOverrideVoxelSize", ClampMin = "0.1", DisplayName = "Overide Spatial Sampling Distance"))
+	UPROPERTY(BlueprintReadWrite, EditAnywhere, AdvancedDisplay, Category = ProxySettings, meta = (EditCondition = "bOverrideVoxelSize", ClampMin = "0.1", DisplayName = "Overide Spatial Sampling Distance"))
 	float VoxelSize;
-	
+
 	/** Material simplification */
-	UPROPERTY(EditAnywhere, Category = ProxySettings)
+	UPROPERTY(BlueprintReadWrite, EditAnywhere, Category = ProxySettings)
 	FMaterialProxySettings MaterialSettings;
 
 	UPROPERTY()
 	int32 TextureWidth_DEPRECATED;
+
 	UPROPERTY()
 	int32 TextureHeight_DEPRECATED;
 
@@ -218,84 +218,84 @@ struct FMeshProxySettings
 	bool bExportSpecularMap_DEPRECATED;
 
 	/** Determines whether or not the correct LOD models should be calculated given the source meshes and transition size */
-	UPROPERTY(EditAnywhere, Category = ProxySettings)
+	UPROPERTY(BlueprintReadWrite, EditAnywhere, Category = ProxySettings)
 	bool bCalculateCorrectLODModel;
 
 	/** Distance at which meshes should be merged together, this can close gaps like doors and windows in distant geometry */
-	UPROPERTY(EditAnywhere, Category = ProxySettings)
+	UPROPERTY(BlueprintReadWrite, EditAnywhere, Category = ProxySettings)
 	float MergeDistance;
 
 	/** Base color assigned to LOD geometry that can't be associated with the source geometry: e.g. doors and windows that have been closed by the Merge Distance */
-	UPROPERTY(EditAnywhere, Category = ProxySettings, meta = (DisplayName = "Unresolved Geometry Color"))
+	UPROPERTY(BlueprintReadWrite, EditAnywhere, Category = ProxySettings, meta = (DisplayName = "Unresolved Geometry Color"))
 	FColor UnresolvedGeometryColor;
 
 	/** Enable an override for material transfer distance */
-	UPROPERTY(EditAnywhere, Category = MaxRayCastDist, meta = (InlineEditConditionToggle))
+	UPROPERTY(BlueprintReadWrite, EditAnywhere, Category = MaxRayCastDist, meta = (InlineEditConditionToggle))
 	bool bOverrideTransferDistance;
 
 	/** Override search distance used when discovering texture values for simplified geometry. Useful when non-zero Merge Distance setting generates new geometry in concave corners.*/
-	UPROPERTY(EditAnywhere, Category = ProxySettings, meta = (EditCondition = "bOverrideTransferDistance", DisplayName = "Transfer Distance Override", ClampMin = 0))
+	UPROPERTY(BlueprintReadWrite, EditAnywhere, Category = ProxySettings, meta = (EditCondition = "bOverrideTransferDistance", DisplayName = "Transfer Distance Override", ClampMin = 0))
 	float MaxRayCastDist;
 
 	/** Enable the use of hard angle based vertex splitting */
-	UPROPERTY(EditAnywhere, Category = HardAngleThreshold, meta = (InlineEditConditionToggle))
+	UPROPERTY(BlueprintReadWrite, EditAnywhere, Category = HardAngleThreshold, meta = (InlineEditConditionToggle))
 	bool bUseHardAngleThreshold;
 
 	/** Angle at which a hard edge is introduced between faces */
-	UPROPERTY(EditAnywhere, Category = ProxySettings, meta = (EditCondition = "bUseHardAngleThreshold", DisplayName = "Hard Edge Angle", ClampMin = 0, ClampMax = 180))
+	UPROPERTY(BlueprintReadWrite, EditAnywhere, Category = ProxySettings, meta = (EditCondition = "bUseHardAngleThreshold", DisplayName = "Hard Edge Angle", ClampMin = 0, ClampMax = 180))
 	float HardAngleThreshold;
 
 	/** Controls the method used to calculate the normal for the simplified geometry */
-	UPROPERTY(EditAnywhere, Category = ProxySettings, meta = (DisplayName = "Normal Calculation Method"))
+	UPROPERTY(BlueprintReadWrite, EditAnywhere, Category = ProxySettings, meta = (DisplayName = "Normal Calculation Method"))
 	TEnumAsByte<EProxyNormalComputationMethod::Type> NormalCalculationMethod;
 
 	/** Lightmap resolution */
-	UPROPERTY(EditAnywhere, Category = ProxySettings, meta = (ClampMin = 32, ClampMax = 4096, EditCondition = "!bComputeLightMapResolution"))
+	UPROPERTY(BlueprintReadWrite, EditAnywhere, Category = ProxySettings, meta = (ClampMin = 32, ClampMax = 4096, EditCondition = "!bComputeLightMapResolution"))
 	int32 LightMapResolution;
 
 	/** If ticked will compute the lightmap resolution by summing the dimensions for each mesh included for merging */
-	UPROPERTY(EditAnywhere, Category = ProxySettings)
+	UPROPERTY(BlueprintReadWrite, EditAnywhere, Category = ProxySettings)
 	bool bComputeLightMapResolution;
 
 	/** Whether Simplygon should recalculate normals, otherwise the normals channel will be sampled from the original mesh */
-	UPROPERTY(EditAnywhere, Category = ProxySettings)
+	UPROPERTY(BlueprintReadWrite, EditAnywhere, Category = ProxySettings)
 	bool bRecalculateNormals;
 
 	UPROPERTY()
 	bool bBakeVertexData_DEPRECATED;
 
 	/** Whether or not to use available landscape geometry to cull away invisible triangles */
-	UPROPERTY(EditAnywhere, Category = LandscapeCulling)
+	UPROPERTY(BlueprintReadWrite, EditAnywhere, Category = LandscapeCulling)
 	bool bUseLandscapeCulling;
 
 	/** Level of detail of the landscape that should be used for the culling */
-	UPROPERTY(EditAnywhere, Category = LandscapeCulling, meta = (EditCondition="bUseLandscapeCulling"))
+	UPROPERTY(BlueprintReadWrite, EditAnywhere, Category = LandscapeCulling, meta = (EditCondition="bUseLandscapeCulling"))
 	TEnumAsByte<ELandscapeCullingPrecision::Type> LandscapeCullingPrecision;
 
 	/** Whether to allow adjacency buffers for tessellation in the merged mesh */
-	UPROPERTY(EditAnywhere, Category = ProxySettings)
+	UPROPERTY(BlueprintReadWrite, EditAnywhere, Category = ProxySettings)
 	bool bAllowAdjacency;
 
 	/** Whether to allow distance field to be computed for this mesh. Disable this to save memory if the merged mesh will only be rendered in the distance. */
-	UPROPERTY(EditAnywhere, Category = ProxySettings)
+	UPROPERTY(BlueprintReadWrite, EditAnywhere, Category = ProxySettings)
 	bool bAllowDistanceField;
 
 	/** Whether to attempt to re-use the source mesh's lightmap UVs when baking the material or always generate a new set. */
-	UPROPERTY(EditAnywhere, Category = ProxySettings)
+	UPROPERTY(BlueprintReadWrite, EditAnywhere, Category = ProxySettings)
 	bool bReuseMeshLightmapUVs;
 
 	/** Whether to generate collision for the merged mesh */
-	UPROPERTY(EditAnywhere, Category = ProxySettings)
+	UPROPERTY(BlueprintReadWrite, EditAnywhere, Category = ProxySettings)
 	bool bCreateCollision;
 
 	/** Whether to allow vertex colors saved in the merged mesh */
-	UPROPERTY(EditAnywhere, Category = ProxySettings)
+	UPROPERTY(BlueprintReadWrite, EditAnywhere, Category = ProxySettings)
 	bool bAllowVertexColors;
 
 	/** Whether to generate lightmap uvs for the merged mesh */
-	UPROPERTY(EditAnywhere, Category = ProxySettings)
+	UPROPERTY(BlueprintReadWrite, EditAnywhere, Category = ProxySettings)
 	bool bGenerateLightmapUVs;
-	
+
 	/** Default settings. */
 	FMeshProxySettings()
 		: ScreenSize(300)
@@ -327,7 +327,7 @@ struct FMeshProxySettings
 		, bCreateCollision(true)
 		, bAllowVertexColors(false)
 		, bGenerateLightmapUVs(false)
-	{ 
+	{
 		MaterialSettings.MaterialMergeType = EMaterialMergeType::MaterialMergeType_Simplygon;
 	}
 
@@ -363,7 +363,7 @@ UENUM()
 enum class EMeshLODSelectionType : uint8
 {
 	// Whether or not to export all of the LODs found in the source meshes
-	AllLODs = 0 UMETA(DisplayName = "Use all LOD levels"),	
+	AllLODs = 0 UMETA(DisplayName = "Use all LOD levels"),
 	// Whether or not to export all of the LODs found in the source meshes
 	SpecificLOD = 1 UMETA(DisplayName = "Use specific LOD level"),
 	// Whether or not to calculate the appropriate LOD model for the given screen size
@@ -400,7 +400,7 @@ struct FMeshMergingSettings
 	bool bGenerateLightMapUV;
 
 	/** Target lightmap resolution */
-	UPROPERTY(EditAnywhere, AdvancedDisplay, BlueprintReadWrite, Category = MeshSettings, meta=(ClampMax = 4096, editcondition = "!bComputedLightMapResolution"))
+	UPROPERTY(EditAnywhere, AdvancedDisplay, BlueprintReadWrite, Category = MeshSettings, meta=(ClampMax = 4096, EditCondition = "!bComputedLightMapResolution"))
 	int32 TargetLightMapResolution;
 
 	/** Whether or not the lightmap resolution should be computed by summing the lightmap resolutions for the input Mesh Components */
@@ -424,7 +424,7 @@ struct FMeshMergingSettings
 	bool bMergeMaterials;
 
 	/** Material simplification */
-	UPROPERTY(EditAnywhere, Category = MaterialSettings, BlueprintReadWrite, meta = (editcondition = "bMergeMaterials"))
+	UPROPERTY(EditAnywhere, Category = MaterialSettings, BlueprintReadWrite, meta = (EditCondition = "bMergeMaterials"))
 	FMaterialProxySettings MaterialSettings;
 
 	/** Whether or not vertex data such as vertex colours should be baked into the resulting mesh */
@@ -432,13 +432,13 @@ struct FMeshMergingSettings
 	bool bBakeVertexDataToMesh;
 
 	/** Whether or not vertex data such as vertex colours should be used when baking out materials */
-	UPROPERTY(EditAnywhere, Category = MaterialSettings, BlueprintReadWrite, meta = (editcondition = "bMergeMaterials"))
+	UPROPERTY(EditAnywhere, Category = MaterialSettings, BlueprintReadWrite, meta = (EditCondition = "bMergeMaterials"))
 	bool bUseVertexDataForBakingMaterial;
 
 	// Whether or not to calculate varying output texture sizes according to their importance in the final atlas texture
-	UPROPERTY(Category = MaterialSettings, EditAnywhere, BlueprintReadWrite, meta = (editcondition = "bMergeMaterials"))
+	UPROPERTY(Category = MaterialSettings, EditAnywhere, BlueprintReadWrite, meta = (EditCondition = "bMergeMaterials"))
 	bool bUseTextureBinning;
-			
+
 	/** Whether to attempt to re-use the source mesh's lightmap UVs when baking the material or always generate a new set. */
 	UPROPERTY(EditAnywhere, Category = MaterialSettings)
 	bool bReuseMeshLightmapUVs;
@@ -544,7 +544,7 @@ struct FSectionInfo
 	/** Material used by the section */
 	class UMaterialInterface* Material;
 	/** Name value for the section */
-	FName MaterialSlotName;	
+	FName MaterialSlotName;
 	/** List of properties enabled for the section (collision, cast shadow etc) */
 	TArray<FName> EnabledProperties;
 	/** Original index of Material in the source data */
@@ -574,7 +574,7 @@ enum class EMeshInstancingReplacementMethod
 };
 
 /** Mesh instance-replacement settings */
-USTRUCT()
+USTRUCT(Blueprintable)
 struct FMeshInstancingSettings
 {
 	GENERATED_BODY()
@@ -588,28 +588,28 @@ struct FMeshInstancingSettings
 	{}
 
 	/** The actor class to attach new instance static mesh components to */
-	UPROPERTY(EditAnywhere, NoClear, Category="Instancing")
+	UPROPERTY(BlueprintReadWrite, EditAnywhere, NoClear, Category="Instancing")
 	TSubclassOf<AActor> ActorClassToUse;
 
 	/** The number of static mesh instances needed before a mesh is replaced with an instanced version */
-	UPROPERTY(EditAnywhere, Category="Instancing", meta=(ClampMin=1))
+	UPROPERTY(BlueprintReadWrite, EditAnywhere, Category="Instancing", meta=(ClampMin=1))
 	int32 InstanceReplacementThreshold;
 
 	/** How to replace the original actors when instancing */
-	UPROPERTY(EditAnywhere, Category="Instancing")
+	UPROPERTY(BlueprintReadWrite, EditAnywhere, Category="Instancing")
 	EMeshInstancingReplacementMethod MeshReplacementMethod;
 
-	/** 
+	/**
 	 * Whether to skip the conversion to an instanced static mesh for meshes with vertex colors.
 	 * Instanced static meshes do not support vertex colors per-instance, so conversion will lose
 	 * this data.
 	 */
-	UPROPERTY(EditAnywhere, Category="Instancing")
+	UPROPERTY(BlueprintReadWrite, EditAnywhere, Category="Instancing")
 	bool bSkipMeshesWithVertexColors;
 
-	/** 
+	/**
 	 * Whether split up instanced static mesh components based on their intersection with HLOD volumes
 	 */
-	UPROPERTY(EditAnywhere, Category="Instancing", meta=(DisplayName="Use HLOD Volumes"))
+	UPROPERTY(BlueprintReadWrite, EditAnywhere, Category="Instancing", meta=(DisplayName="Use HLOD Volumes"))
 	bool bUseHLODVolumes;
 };
\ No newline at end of file
diff --git a/Engine/Source/Runtime/Engine/Classes/Engine/RendererSettings.h b/Engine/Source/Runtime/Engine/Classes/Engine/RendererSettings.h
--- a/Engine/Source/Runtime/Engine/Classes/Engine/RendererSettings.h
+++ b/Engine/Source/Runtime/Engine/Classes/Engine/RendererSettings.h
@@ -202,6 +202,12 @@ class ENGINE_API URendererSettings : public UDeveloperSettings
 		ToolTip="Whether to support 'Dithered LOD Transition' material option on mobile platforms. Enabling this may degrade performance as rendering will not benefit from Early-Z optimization.",
 		ConfigRestartRequired=true))
 	uint32 bMobileAllowDitheredLODTransition:1;
+
+	UPROPERTY(config, EditAnywhere, Category = Mobile, meta=(
+		ConsoleVariable="r.Mobile.AllowSoftwareOcclusion", DisplayName="Support Software Occlusion Culling",
+		ToolTip="Whether to support 'Software Occlusion Culling' on mobile platforms. This this package occluder information and enable Software Occlusion Culling.",
+		ConfigRestartRequired=false))
+	uint32 bMobileAllowSoftwareOcclusionCulling:1;
 	
 	UPROPERTY(config, EditAnywhere, Category = Materials, meta = (
 		ConsoleVariable = "r.DiscardUnusedQuality", DisplayName = "Game Discards Unused Material Quality Levels",
diff --git a/Engine/Source/Runtime/Engine/Classes/GameFramework/Actor.h b/Engine/Source/Runtime/Engine/Classes/GameFramework/Actor.h
--- a/Engine/Source/Runtime/Engine/Classes/GameFramework/Actor.h
+++ b/Engine/Source/Runtime/Engine/Classes/GameFramework/Actor.h
@@ -604,7 +604,7 @@ public:
 	uint8 bIsEditorPreviewActor:1;
 
 	/** Whether this actor is hidden by the layer browser. */
-	UPROPERTY()
+	UPROPERTY(Transient)
 	uint8 bHiddenEdLayer:1;
 
 	/** Whether this actor is hidden by the level browser. */
diff --git a/Engine/Source/Runtime/Engine/Classes/Materials/MaterialExpressionMaterialFunctionCall.h b/Engine/Source/Runtime/Engine/Classes/Materials/MaterialExpressionMaterialFunctionCall.h
--- a/Engine/Source/Runtime/Engine/Classes/Materials/MaterialExpressionMaterialFunctionCall.h
+++ b/Engine/Source/Runtime/Engine/Classes/Materials/MaterialExpressionMaterialFunctionCall.h
@@ -80,7 +80,7 @@ class UMaterialExpressionMaterialFunctionCall : public UMaterialExpression
 	GENERATED_UCLASS_BODY()
 
 	/** The function to call. */
-	UPROPERTY(EditAnywhere, Category=MaterialExpressionMaterialFunctionCall)
+	UPROPERTY(EditAnywhere, Category=MaterialExpressionMaterialFunctionCall, meta=(DisallowedClasses="MaterialFunctionInstance"))
 	class UMaterialFunctionInterface* MaterialFunction;
 
 #if WITH_EDITORONLY_DATA
diff --git a/Engine/Source/Runtime/Engine/Private/Materials/HLSLMaterialTranslator.h b/Engine/Source/Runtime/Engine/Private/Materials/HLSLMaterialTranslator.h
--- a/Engine/Source/Runtime/Engine/Private/Materials/HLSLMaterialTranslator.h
+++ b/Engine/Source/Runtime/Engine/Private/Materials/HLSLMaterialTranslator.h
@@ -3787,6 +3787,14 @@ protected:
 			// On mobile in post process material, there is no need to do ViewportUV->BufferUV conversion because ViewSize == BufferSize.
 			if (Material->GetMaterialDomain() == MD_PostProcess)
 			{
+				int32 BlendableLocation = Material->GetBlendableLocation();
+				if (SceneTextureId == PPI_SceneDepth && !(BlendableLocation == BL_BeforeTranslucency || BlendableLocation == BL_BeforeTonemapping))
+				{
+					// SceneDepth lookups are not available when using MSAA, but we can access depth stored in SceneColor.A channel
+					// SceneColor.A channel holds depth till BeforeTonemapping location, then it's gets overwritten
+					return Errorf(TEXT("SceneDepth lookups are only available when BlendableLocation is BeforeTranslucency or BeforeTonemapping"));
+				}
+				
 				if (ViewportUV == INDEX_NONE)
 				{
 					UV = TextureCoordinate(0, false, false);
diff --git a/Engine/Source/Runtime/Engine/Public/PerPlatformProperties.h b/Engine/Source/Runtime/Engine/Public/PerPlatformProperties.h
--- a/Engine/Source/Runtime/Engine/Public/PerPlatformProperties.h
+++ b/Engine/Source/Runtime/Engine/Public/PerPlatformProperties.h
@@ -75,20 +75,6 @@ struct ENGINE_API TPerPlatformProperty
 		return true;
 	}
 
-	/* Load old properties that have been converted to FPerPlatformX */
-	bool SerializeFromMismatchedTag(const FPropertyTag& Tag, FStructuredArchive::FSlot Slot)
-	{
-		if (Tag.Type == _BasePropertyName)
-		{
-			_StructType* This = StaticCast<_StructType*>(this);
-			_ValueType OldValue;
-			Slot << OldValue;
-			*This = _StructType(OldValue);
-			return true;
-		}
-		return false;
-	}
-
 	/* Serialization */
 	bool Serialize(FStructuredArchive::FSlot Slot)
 	{
diff --git a/Engine/Source/Runtime/Engine/Public/ShowFlags.h b/Engine/Source/Runtime/Engine/Public/ShowFlags.h
--- a/Engine/Source/Runtime/Engine/Public/ShowFlags.h
+++ b/Engine/Source/Runtime/Engine/Public/ShowFlags.h
@@ -368,7 +368,7 @@ private:
 		SetCameraSafeFrames(false);
 		SetVisualizeOutOfBoundsPixels(false);
 		SetHighResScreenshotMask(false);
-		SetHMDDistortion(true); // only effective if a HMD is in use
+		SetHMDDistortion(false); // only effective if a HMD is in use
 		SetStereoRendering(false);
 		SetDistanceCulledPrimitives(InitMode == ESFIM_Editor);
 		SetVisualizeLightCulling(false);
diff --git a/Engine/Source/Runtime/GameplayTags/Classes/GameplayTagContainer.h b/Engine/Source/Runtime/GameplayTags/Classes/GameplayTagContainer.h
--- a/Engine/Source/Runtime/GameplayTags/Classes/GameplayTagContainer.h
+++ b/Engine/Source/Runtime/GameplayTags/Classes/GameplayTagContainer.h
@@ -256,7 +256,7 @@ struct TStructOpsTypeTraits< FGameplayTag > : public TStructOpsTypeTraitsBase2<
 		WithNetSerializer = true,
 		WithNetSharedSerialization = true,
 		WithPostSerialize = true,
-		WithSerializeFromMismatchedTag = true,
+		WithStructuredSerializeFromMismatchedTag = true,
 		WithImportTextItem = true,
 	};
 };
diff --git a/Engine/Source/Runtime/InputCore/Classes/InputCoreTypes.h b/Engine/Source/Runtime/InputCore/Classes/InputCoreTypes.h
--- a/Engine/Source/Runtime/InputCore/Classes/InputCoreTypes.h
+++ b/Engine/Source/Runtime/InputCore/Classes/InputCoreTypes.h
@@ -105,7 +105,7 @@ struct TStructOpsTypeTraits<FKey> : public TStructOpsTypeTraitsBase2<FKey>
 {
 	enum
 	{
-		WithSerializeFromMismatchedTag = true,
+		WithStructuredSerializeFromMismatchedTag = true,
 		WithExportTextItem = true,
 		WithImportTextItem = true,
 		WithPostSerialize = true,
diff --git a/Engine/Source/Runtime/MovieScene/Public/Channels/MovieSceneBoolChannel.h b/Engine/Source/Runtime/MovieScene/Public/Channels/MovieSceneBoolChannel.h
--- a/Engine/Source/Runtime/MovieScene/Public/Channels/MovieSceneBoolChannel.h
+++ b/Engine/Source/Runtime/MovieScene/Public/Channels/MovieSceneBoolChannel.h
@@ -136,7 +136,7 @@ private:
 template<>
 struct TStructOpsTypeTraits<FMovieSceneBoolChannel> : public TStructOpsTypeTraitsBase2<FMovieSceneBoolChannel>
 {
-	enum { WithSerializeFromMismatchedTag = true };
+	enum { WithStructuredSerializeFromMismatchedTag = true };
 };
 
 template<>
diff --git a/Engine/Source/Runtime/MovieScene/Public/Channels/MovieSceneByteChannel.h b/Engine/Source/Runtime/MovieScene/Public/Channels/MovieSceneByteChannel.h
--- a/Engine/Source/Runtime/MovieScene/Public/Channels/MovieSceneByteChannel.h
+++ b/Engine/Source/Runtime/MovieScene/Public/Channels/MovieSceneByteChannel.h
@@ -155,7 +155,7 @@ private:
 template<>
 struct TStructOpsTypeTraits<FMovieSceneByteChannel> : public TStructOpsTypeTraitsBase2<FMovieSceneByteChannel>
 {
-	enum { WithSerializeFromMismatchedTag = true };
+	enum { WithStructuredSerializeFromMismatchedTag = true };
 };
 
 template<>
diff --git a/Engine/Source/Runtime/MovieScene/Public/Channels/MovieSceneFloatChannel.h b/Engine/Source/Runtime/MovieScene/Public/Channels/MovieSceneFloatChannel.h
--- a/Engine/Source/Runtime/MovieScene/Public/Channels/MovieSceneFloatChannel.h
+++ b/Engine/Source/Runtime/MovieScene/Public/Channels/MovieSceneFloatChannel.h
@@ -287,7 +287,7 @@ public:
 template<>
 struct TStructOpsTypeTraits<FMovieSceneFloatChannel> : public TStructOpsTypeTraitsBase2<FMovieSceneFloatChannel>
 {
-	enum { WithSerializeFromMismatchedTag = true };
+	enum { WithStructuredSerializeFromMismatchedTag = true };
 };
 
 template<>
diff --git a/Engine/Source/Runtime/MovieScene/Public/Channels/MovieSceneIntegerChannel.h b/Engine/Source/Runtime/MovieScene/Public/Channels/MovieSceneIntegerChannel.h
--- a/Engine/Source/Runtime/MovieScene/Public/Channels/MovieSceneIntegerChannel.h
+++ b/Engine/Source/Runtime/MovieScene/Public/Channels/MovieSceneIntegerChannel.h
@@ -136,7 +136,7 @@ private:
 template<>
 struct TStructOpsTypeTraits<FMovieSceneIntegerChannel> : public TStructOpsTypeTraitsBase2<FMovieSceneIntegerChannel>
 {
-	enum { WithSerializeFromMismatchedTag = true };
+	enum { WithStructuredSerializeFromMismatchedTag = true };
 };
 
 template<>
diff --git a/Engine/Source/Runtime/MovieScene/Public/MovieSceneFrameMigration.h b/Engine/Source/Runtime/MovieScene/Public/MovieSceneFrameMigration.h
--- a/Engine/Source/Runtime/MovieScene/Public/MovieSceneFrameMigration.h
+++ b/Engine/Source/Runtime/MovieScene/Public/MovieSceneFrameMigration.h
@@ -96,7 +96,7 @@ private:
 template<>
 struct TStructOpsTypeTraits<FMovieSceneFrameRange> : public TStructOpsTypeTraitsBase2<FMovieSceneFrameRange>
 {
-	enum { WithSerializeFromMismatchedTag = true, WithSerializer = true, WithIdenticalViaEquality = true,
+	enum { WithStructuredSerializeFromMismatchedTag = true, WithSerializer = true, WithIdenticalViaEquality = true,
 		   WithExportTextItem = true, WithImportTextItem = true};
 };
 
diff --git a/Engine/Source/Runtime/MovieScene/Public/MovieSceneSequencePlayer.h b/Engine/Source/Runtime/MovieScene/Public/MovieSceneSequencePlayer.h
--- a/Engine/Source/Runtime/MovieScene/Public/MovieSceneSequencePlayer.h
+++ b/Engine/Source/Runtime/MovieScene/Public/MovieSceneSequencePlayer.h
@@ -111,7 +111,7 @@ struct FMovieSceneSequencePlaybackSettings
 
 template<> struct TStructOpsTypeTraits<FMovieSceneSequencePlaybackSettings> : public TStructOpsTypeTraitsBase2<FMovieSceneSequencePlaybackSettings>
 {
-	enum { WithCopy = true, WithSerializeFromMismatchedTag = true };
+	enum { WithCopy = true, WithStructuredSerializeFromMismatchedTag = true };
 };
 
 /**
diff --git a/Engine/Source/Runtime/MovieSceneTracks/Public/Sections/MovieSceneParticleSection.h b/Engine/Source/Runtime/MovieSceneTracks/Public/Sections/MovieSceneParticleSection.h
--- a/Engine/Source/Runtime/MovieSceneTracks/Public/Sections/MovieSceneParticleSection.h
+++ b/Engine/Source/Runtime/MovieSceneTracks/Public/Sections/MovieSceneParticleSection.h
@@ -31,7 +31,7 @@ struct FMovieSceneParticleChannel : public FMovieSceneByteChannel
 template<>
 struct TStructOpsTypeTraits<FMovieSceneParticleChannel> : public TStructOpsTypeTraitsBase2<FMovieSceneParticleChannel>
 {
-	enum { WithSerializeFromMismatchedTag = true };
+	enum { WithStructuredSerializeFromMismatchedTag = true };
 };
 
 
diff --git a/Engine/Source/Runtime/MovieSceneTracks/Public/Sections/MovieSceneStringChannel.h b/Engine/Source/Runtime/MovieSceneTracks/Public/Sections/MovieSceneStringChannel.h
--- a/Engine/Source/Runtime/MovieSceneTracks/Public/Sections/MovieSceneStringChannel.h
+++ b/Engine/Source/Runtime/MovieSceneTracks/Public/Sections/MovieSceneStringChannel.h
@@ -122,7 +122,7 @@ private:
 template<>
 struct TStructOpsTypeTraits<FMovieSceneStringChannel> : public TStructOpsTypeTraitsBase2<FMovieSceneStringChannel>
 {
-	enum { WithSerializeFromMismatchedTag = true };
+	enum { WithStructuredSerializeFromMismatchedTag = true };
 };
 
 
diff --git a/Engine/Source/Runtime/Portal/Rpc/Private/PortalRpcLocator.cpp b/Engine/Source/Runtime/Portal/Rpc/Private/PortalRpcLocator.cpp
--- a/Engine/Source/Runtime/Portal/Rpc/Private/PortalRpcLocator.cpp
+++ b/Engine/Source/Runtime/Portal/Rpc/Private/PortalRpcLocator.cpp
@@ -68,7 +68,9 @@ private:
 		}
 
 		// @todo sarge: implement actual product GUID
-		MessageEndpoint->Publish(new FPortalRpcLocateServer(FGuid(), EngineVersion, MacAddress, UserId), EMessageScope::Network);
+		// message is going to be deleted by FMemory::Free() (see FMessageContext destructor), so allocate it with Malloc
+		void* Memory = FMemory::Malloc(sizeof(FPortalRpcLocateServer));
+		MessageEndpoint->Publish(new(Memory) FPortalRpcLocateServer(FGuid(), EngineVersion, MacAddress, UserId), EMessageScope::Network);
 
 		return true;
 	}
diff --git a/Engine/Source/Runtime/Renderer/Private/CompositionLighting/PostProcessDeferredDecals.cpp b/Engine/Source/Runtime/Renderer/Private/CompositionLighting/PostProcessDeferredDecals.cpp
--- a/Engine/Source/Runtime/Renderer/Private/CompositionLighting/PostProcessDeferredDecals.cpp
+++ b/Engine/Source/Runtime/Renderer/Private/CompositionLighting/PostProcessDeferredDecals.cpp
@@ -934,7 +934,10 @@ void FRCPassPostProcessDeferredDecals::Process(FRenderingCompositePassContext& C
 					RHICmdList.DrawIndexedPrimitive(GetUnitCubeIndexBuffer(), PT_TriangleList, 0, 0, 8, 0, ARRAY_COUNT(GCubeIndices) / 3, 1);
 					RenderTargetManager.bGufferADirty |= (RenderTargetManager.TargetsToResolve[FDecalRenderTargetManager::GBufferAIndex] != nullptr);
 				}
+			}
 
+			if (CurrentStage == DRS_Emissive)
+			{
 				// we don't modify stencil but if out input was having stencil for us (after base pass - we need to clear)
 				// Clear stencil to 0, which is the assumed default by other passes
 				DrawClearQuad(RHICmdList, false, FLinearColor(), false, 0, true, 0, SceneContext.GetSceneDepthSurface()->GetSizeXY(), FIntRect());
diff --git a/Engine/Source/Runtime/Renderer/Private/MobileShadingRenderer.cpp b/Engine/Source/Runtime/Renderer/Private/MobileShadingRenderer.cpp
--- a/Engine/Source/Runtime/Renderer/Private/MobileShadingRenderer.cpp
+++ b/Engine/Source/Runtime/Renderer/Private/MobileShadingRenderer.cpp
@@ -78,7 +78,6 @@ FMobileSceneRenderer::FMobileSceneRenderer(const FSceneViewFamily* InViewFamily,
 	:	FSceneRenderer(InViewFamily, HitProxyConsumer)
 {
 	bModulatedShadowsInUse = false;
-	bPostProcessUsesDepthTexture = false;
 }
 
 TUniformBufferRef<FMobileDirectionalLightShaderParameters>& GetNullMobileDirectionalLightShaderParameters()
@@ -131,9 +130,6 @@ void FMobileSceneRenderer::InitViews(FRHICommandListImmediate& RHICmdList)
 	// Now that the indirect lighting cache is updated, we can update the primitive precomputed lighting buffers.
 	UpdatePrimitivePrecomputedLightingBuffers();
 
-	UpdatePostProcessUsageFlags();
-
-	
 	PostInitViewCustomData();
 	
 	OnStartFrame(RHICmdList);
@@ -408,7 +404,7 @@ void FMobileSceneRenderer::Render(FRHICommandListImmediate& RHICmdList)
 		const bool bForceDepthResolve = CVarMobileForceDepthResolve.GetValueOnRenderThread() == 1;
 		const bool bSeparateTranslucencyActive = IsMobileSeparateTranslucencyActive(View);
 
-		bKeepDepthContent = bForceDepthResolve || bPostProcessUsesDepthTexture || bSeparateTranslucencyActive ||
+		bKeepDepthContent = bForceDepthResolve || bSeparateTranslucencyActive ||
 			(View.bIsSceneCapture && (ViewFamily.SceneCaptureSource == ESceneCaptureSource::SCS_SceneColorHDR || ViewFamily.SceneCaptureSource == ESceneCaptureSource::SCS_SceneColorSceneDepth));
 	}
 
@@ -808,32 +804,3 @@ void FMobileSceneRenderer::CopyMobileMultiViewSceneColor(FRHICommandListImmediat
 			EDRF_UseTriangleOptimization);
 	}
 }
-
-void FMobileSceneRenderer::UpdatePostProcessUsageFlags()
-{
-	bPostProcessUsesDepthTexture = false;
-	// Find out whether post-process materials require SceneDepth lookups, otherwise renderer can discard depth buffer before starting post-processing pass
-	for (int32 ViewIndex = 0; ViewIndex < Views.Num(); ViewIndex++)
-	{
-		const FBlendableManager& BlendableManager = Views[ViewIndex].FinalPostProcessSettings.BlendableManager;
-		FBlendableEntry* BlendableIt = nullptr;
-
-		while (FPostProcessMaterialNode* DataPtr = BlendableManager.IterateBlendables<FPostProcessMaterialNode>(BlendableIt))
-		{
-			if (DataPtr->IsValid())
-			{
-				FMaterialRenderProxy* Proxy = DataPtr->GetMaterialInterface()->GetRenderProxy(false);
-				check(Proxy);
-
-				const FMaterial* Material = Proxy->GetMaterial(Views[ViewIndex].GetFeatureLevel());
-				check(Material);
-
-				if (Material->MaterialUsesSceneDepthLookup_RenderThread())
-				{
-					bPostProcessUsesDepthTexture = true;
-					break;
-				}
-			}
-		}
-	}
-}
diff --git a/Engine/Source/Runtime/Renderer/Private/SceneOcclusion.cpp b/Engine/Source/Runtime/Renderer/Private/SceneOcclusion.cpp
--- a/Engine/Source/Runtime/Renderer/Private/SceneOcclusion.cpp
+++ b/Engine/Source/Runtime/Renderer/Private/SceneOcclusion.cpp
@@ -52,10 +52,9 @@ FAutoConsoleVariableRef CVarOcclusionCullCascadedShadowMaps(
 	ECVF_Scalability | ECVF_RenderThreadSafe
 	);
 
-int32 GMobileAllowSoftwareOcclusion = 0;
-FAutoConsoleVariableRef CVarMobileAllowSoftwareOcclusion(
+static TAutoConsoleVariable<int32> CVarMobileAllowSoftwareOcclusion(
 	TEXT("r.Mobile.AllowSoftwareOcclusion"),
-	GMobileAllowSoftwareOcclusion,
+	0,
 	TEXT("Whether to allow rasterizing scene on CPU for primitive occlusion.\n"),
 	ECVF_RenderThreadSafe
 	);
@@ -263,7 +262,8 @@ bool FSceneViewState::IsShadowOccluded(FRHICommandListImmediate& RHICmdList, FSc
 
 void FSceneViewState::ConditionallyAllocateSceneSoftwareOcclusion(ERHIFeatureLevel::Type InFeatureLevel)
 {
-	bool bShouldBeEnabled = InFeatureLevel <= ERHIFeatureLevel::ES3_1 && GMobileAllowSoftwareOcclusion != 0;
+	bool bMobileAllowSoftwareOcclusion = CVarMobileAllowSoftwareOcclusion.GetValueOnAnyThread() != 0;
+	bool bShouldBeEnabled = InFeatureLevel <= ERHIFeatureLevel::ES3_1 && bMobileAllowSoftwareOcclusion;
 
 	if (bShouldBeEnabled && !SceneSoftwareOcclusion)
 	{
diff --git a/Engine/Source/Runtime/Renderer/Private/SceneRendering.h b/Engine/Source/Runtime/Renderer/Private/SceneRendering.h
--- a/Engine/Source/Runtime/Renderer/Private/SceneRendering.h
+++ b/Engine/Source/Runtime/Renderer/Private/SceneRendering.h
@@ -1746,9 +1746,6 @@ protected:
 	/** Copy scene color from the mobile multi-view render targat array to side by side stereo scene color */
 	void CopyMobileMultiViewSceneColor(FRHICommandListImmediate& RHICmdList);
 
-	/** Gather information about post-processing pass, which can be used by render for optimizations. Called by InitViews */
-	void UpdatePostProcessUsageFlags();
-
 	/** Render inverse opacity for the dynamic meshes. */
 	bool RenderInverseOpacityDynamic(FRHICommandListImmediate& RHICmdList, const FViewInfo& View, const FDrawingPolicyRenderState& DrawRenderState);
 
@@ -1756,9 +1753,7 @@ protected:
 	void PostInitViewCustomData();
 	
 private:
-
 	bool bModulatedShadowsInUse;
-	bool bPostProcessUsesDepthTexture;
 };
 
 // The noise textures need to be set in Slate too.
diff --git a/Engine/Source/Runtime/Slate/Private/Framework/Application/SlateApplication.cpp b/Engine/Source/Runtime/Slate/Private/Framework/Application/SlateApplication.cpp
--- a/Engine/Source/Runtime/Slate/Private/Framework/Application/SlateApplication.cpp
+++ b/Engine/Source/Runtime/Slate/Private/Framework/Application/SlateApplication.cpp
@@ -6063,6 +6063,13 @@ bool FSlateApplication::ProcessMouseButtonUpEvent( FPointerEvent& MouseEvent )
 {
 	SCOPE_CYCLE_COUNTER(STAT_ProcessMouseButtonUp);
 
+	// If in responsive mode throttle, leave it on mouse up.  Release this before dispatching the event to prevent being stuck in this mode
+	// until the next click if a modal dialog is opened.
+	if (MouseButtonDownResponsivnessThrottle.IsValid())
+	{
+		FSlateThrottleManager::Get().LeaveResponsiveMode(MouseButtonDownResponsivnessThrottle);
+	}
+
 	QueueSynthesizedMouseMove();
 	SetLastUserInteractionTime(this->GetCurrentTime());
 	LastUserInteractionTimeForThrottling = LastUserInteractionTime;
@@ -6078,12 +6085,6 @@ bool FSlateApplication::ProcessMouseButtonUpEvent( FPointerEvent& MouseEvent )
 	FWidgetPath EmptyPath;
 	const bool bHandled = RoutePointerUpEvent( EmptyPath, MouseEvent ).IsEventHandled();
 
-	// If in responsive mode throttle, leave it on mouse up.
-	if( MouseButtonDownResponsivnessThrottle.IsValid() )
-	{
-		FSlateThrottleManager::Get().LeaveResponsiveMode( MouseButtonDownResponsivnessThrottle );
-	}
-
 	if ( PressedMouseButtons.Num() == 0 )
 	{
 		// Release Capture
diff --git a/Engine/Source/Runtime/SlateCore/Public/Sound/SlateSound.h b/Engine/Source/Runtime/SlateCore/Public/Sound/SlateSound.h
--- a/Engine/Source/Runtime/SlateCore/Public/Sound/SlateSound.h
+++ b/Engine/Source/Runtime/SlateCore/Public/Sound/SlateSound.h
@@ -84,6 +84,6 @@ struct TStructOpsTypeTraits<FSlateSound>
 {
 	enum 
 	{
-		WithSerializeFromMismatchedTag = true,
+		WithStructuredSerializeFromMismatchedTag = true,
 	};
 };
diff --git a/Engine/Source/Runtime/SlateCore/Public/Styling/SlateColor.h b/Engine/Source/Runtime/SlateCore/Public/Styling/SlateColor.h
--- a/Engine/Source/Runtime/SlateCore/Public/Styling/SlateColor.h
+++ b/Engine/Source/Runtime/SlateCore/Public/Styling/SlateColor.h
@@ -210,6 +210,6 @@ struct TStructOpsTypeTraits<FSlateColor>
 {
 	enum 
 	{
-		WithSerializeFromMismatchedTag = true,
+		WithStructuredSerializeFromMismatchedTag = true,
 	};
 };
diff --git a/Templates/FP_FirstPerson/manifest.json b/Templates/FP_FirstPerson/manifest.json
--- a/Templates/FP_FirstPerson/manifest.json
+++ b/Templates/FP_FirstPerson/manifest.json
@@ -10,6 +10,10 @@
 			"Language": "ko",
 			"Text": "일인칭"
 		},
+		{
+			"Language": "zh-CN",
+			"Text": "第一人称游戏"
+		},
 		{
 			"Language": "ja",
 			"Text": "ファーストパーソン"
@@ -23,11 +27,15 @@
 		},
 		{
 			"Language": "ko",
-			"Text": "일인칭 팩은 일인칭 시점에서 보이는 (양팔로 표현되는) 플레이어 캐릭터를 특징으로 합니다. 캐릭터 이동은 키보드, 콘트롤러, 터치 디바이스의 가상 조이스틱으로 가능합니다. 추가적으로 플레이어는 마우스, 콘트롤러, 터치 디바이스의 가상 조이스틱을 사용하여 둘러보기도 가능합니다. 캐릭터에는 총이 장착되어 있어, 마우스, 콘트롤러, 터치 디바이스의 가상 조이스틱을 사용하면 레벨과 상호작용하는 '즉시 적중' 무기가 발사됩니다.\n예제 맵도 포함되어 있습니다."
+			"Text": "일인칭 시점에서 보이는 (양팔로 표현되는) 플레이어 캐릭터를 특징으로 합니다. 캐릭터 이동은 키보드, 컨트롤러, 터치 디바이스의 가상 조이스틱으로 가능합니다. 추가적으로 플레이어는 마우스, 컨트롤러, 터치 디바이스의 가상 조이스틱을 사용하여 주변을 둘러볼 수도 있습니다. 캐릭터에는 총이 장착되어 있어 마우스, 컨트롤러, 터치 디바이스의 가상 조이스틱을 사용하면 간단한 발사체가 발사되어, 아레나 벽에 튕기면서 레벨에 있는 몇몇 물리 오브젝트에 영향을 줍니다."
+		},
+		{
+			"Language": "zh-CN",
+			"Text": "第一人称模板拥有一个第一人称的玩家角色（通过一双手臂呈现）。您可以通过键盘、手柄或触摸式设备上的虚拟摇杆控制角色在关卡中四处移动。此外，您还可以使用鼠标、手柄或触摸式设备上的虚拟摇杆转动角色视角，观察四周。角色还装备了枪支，可使用鼠标、手柄或触摸式设备上的虚拟摇杆发射一颗球形子弹，并影响关卡中的一些物理对象，同时从竞技场的墙壁上弹开。"
 		},
 		{
 			"Language": "ja",
-			"Text": "このテンプレートはプレイヤーキャラクター自身の視点で操作するゲームに適しています。キャラクターは武器を持っている腕の部分だけが見えます。このキャラクターはキーボード・コントローラ・バーチャルジョイスティック・タッチパッド等を使ってレベル中を自由に歩き回ることができます。さらに、マウス・コントローラー・バーチャルジョイスティック・タッチパネル等で、視点の向きを変えることができます。装備された武器はマウス・コントローラ・バーチャルジョイスティック・タッチパネルの操作で打つことができます。\nサンプルのマップも用意されています。"
+			"Text": "ファーストパーソンには、一人称視点で表示されるプレイヤーキャラクタ（腕として表示されています）がいます。キーボード、コントローラー、タッチデバイスでのバーチャルジョイスティックでこのキャラクタを移動させることができます。追加して、プレイヤーはマウス、コントローラー、タッチデバイスでのバーチャルジョイスティックで周りを見回すことができます。さらにキャラクタは銃を装備して、マウス、コントローラー、タッチデバイスでのバーチャルジョイスティックを使ってシンプルな球体の弾を発射できます。この弾は物理オブジェクトに影響し、アリーナの壁で反射します。"
 		}
 	],
 	"AssetTypes":
@@ -35,14 +43,6 @@
 		{
 			"Language": "en",
 			"Text": "Animation Sequence, Animation Montage, Sound Wave, Skeletal Mesh, Animation Blueprint, Skeleton, Static Mesh, Texture."
-		},
-		{
-			"Language": "ko",
-			"Text": "애니메이션 시퀀스, 애니메이션 몽타주, 사운드 웨이브, 스켈레탈 메시, 애니메이션 블루프린트, 스켈레톤, 스태틱 메시, 텍스처"
-		},
-		{
-			"Language": "ja",
-			"Text": "アニメーションシーケンス・アニメーションモンタージュ・サウンドウェーブ・スケルタルメッシュ・アニメーションBP・スケルトン・スケルタルメッシュ・テクスチャ"
 		}
 	],
 	"SearchTags":
@@ -50,12 +50,8 @@
 		{
 			"Language": "en",
 			"Text": "First Person,Shoot,Fire,Aim"
-		},
-		{
-			"Language": "ko",
-			"Text": "일인칭,슈팅,발사,조준"
 		}
-	],
+	],	
 	"ClassTypes": "Character, GameMode, HUD, Actor, Character, GameMode, HUD, Actor",
 	"Category": "CodeFeature",
 	"FocusAsset": "/Game/FirstPersonCPP/FirstPersonOverview.FirstPersonOverview",
diff --git a/Templates/FP_FirstPersonBP/manifest.json b/Templates/FP_FirstPersonBP/manifest.json
--- a/Templates/FP_FirstPersonBP/manifest.json
+++ b/Templates/FP_FirstPersonBP/manifest.json
@@ -10,6 +10,10 @@
 			"Language": "ko",
 			"Text": "일인칭"
 		},
+		{
+			"Language": "zh-CN",
+			"Text": "第一人称游戏"
+		},
 		{
 			"Language": "ja",
 			"Text": "ファーストパーソン"
@@ -23,11 +27,15 @@
 		},
 		{
 			"Language": "ko",
-			"Text": "일인칭 팩은 일인칭 시점에서 보이는 (양팔로 표현되는) 플레이어 캐릭터를 특징으로 합니다. 캐릭터 이동은 키보드, 콘트롤러, 터치 디바이스의 가상 조이스틱으로 가능합니다. 추가적으로 플레이어는 마우스, 콘트롤러, 터치 디바이스의 가상 조이스틱을 사용하여 둘러보기도 가능합니다. 캐릭터에는 총이 장착되어 있어, 마우스, 콘트롤러, 터치 디바이스의 가상 조이스틱을 사용하면 레벨과 상호작용하는 '즉시 적중' 무기가 발사됩니다.\n예제 맵도 포함되어 있습니다."
+			"Text": "일인칭 시점에서 보이는 (양팔로 표현되는) 플레이어 캐릭터를 특징으로 합니다. 캐릭터 이동은 키보드, 컨트롤러, 터치 디바이스의 가상 조이스틱으로 가능합니다. 추가적으로 플레이어는 마우스, 컨트롤러, 터치 디바이스의 가상 조이스틱을 사용하여 주변을 둘러볼 수도 있습니다. 캐릭터에는 총이 장착되어 있어 마우스, 컨트롤러, 터치 디바이스의 가상 조이스틱을 사용하면 간단한 발사체가 발사되어, 아레나 벽에 튕기면서 레벨에 있는 몇몇 물리 오브젝트에 영향을 줍니다."
+		},
+		{
+			"Language": "zh-CN",
+			"Text": "第一人称模板拥有一个第一人称的玩家角色（通过一双手臂呈现）。您可以通过键盘、手柄或触摸式设备上的虚拟摇杆控制角色在关卡中四处移动。此外，您还可以使用鼠标、手柄或触摸式设备上的虚拟摇杆转动角色视角，观察四周。角色还装备了枪支，可使用鼠标、手柄或触摸式设备上的虚拟摇杆发射一颗球形子弹，并影响关卡中的一些物理对象，同时从竞技场的墙壁上弹开。"
 		},
 		{
 			"Language": "ja",
-			"Text": "このテンプレートはプレイヤーキャラクター自身の視点で操作するゲームに適しています。キャラクターは武器を持っている腕の部分だけが見えます。このキャラクターはキーボード・コントローラ・バーチャルジョイスティック・タッチパッド等を使ってレベル中を自由に歩き回ることができます。さらに、マウス・コントローラー・バーチャルジョイスティック・タッチパネル等で、視点の向きを変えることができます。装備された武器はマウス・コントローラ・バーチャルジョイスティック・タッチパネルの操作で打つことができます。\nサンプルのマップも用意されています。"
+			"Text": "ファーストパーソンには、一人称視点で表示されるプレイヤーキャラクタ（腕として表示されています）がいます。キーボード、コントローラー、タッチデバイスでのバーチャルジョイスティックでこのキャラクタを移動させることができます。追加して、プレイヤーはマウス、コントローラー、タッチデバイスでのバーチャルジョイスティックで周りを見回すことができます。さらにキャラクタは銃を装備して、マウス、コントローラー、タッチデバイスでのバーチャルジョイスティックを使ってシンプルな球体の弾を発射できます。この弾は物理オブジェクトに影響し、アリーナの壁で反射します。"
 		}
 	],
 	"AssetTypes":
@@ -35,14 +43,6 @@
 		{
 			"Language": "en",
 			"Text": "Animation Sequence, Animation Montage, Sound Wave, Skeletal Mesh, Animation Blueprint, Skeleton, Static Mesh, Texture."
-		},
-		{
-			"Language": "ko",
-			"Text": "애니메이션 시퀀스, 애니메이션 몽타주, 사운드 웨이브, 스켈레탈 메시, 애니메이션 블루프린트, 스켈레톤, 스태틱 메시, 텍스처"
-		},
-		{
-			"Language": "ja",
-			"Text": "アニメーションシーケンス・アニメーションモンタージュ・サウンドウェーブ・スケルタルメッシュ・アニメーションBP・スケルトン・スケルタルメッシュ・テクスチャ"
 		}
 	],
 	"SearchTags":
@@ -50,10 +50,6 @@
 		{
 			"Language": "en",
 			"Text": "First Person,Shoot,Fire,Aim"
-		},
-		{
-			"Language": "ko",
-			"Text": "일인칭,슈팅,발사,조준"
 		}
 	],	
 	"ClassTypes": "Character, GameMode, HUD, Actor, Character, GameMode, HUD, Actor",
diff --git a/Templates/TP_2DSideScroller/Config/TemplateDefs.ini b/Templates/TP_2DSideScroller/Config/TemplateDefs.ini
--- a/Templates/TP_2DSideScroller/Config/TemplateDefs.ini
+++ b/Templates/TP_2DSideScroller/Config/TemplateDefs.ini
@@ -3,7 +3,13 @@ LocalizedDisplayNames=(Language="en",Text="2D Side Scroller")
 LocalizedDescriptions=(Language="en",Text="This is a paper 2D version of the Side Scroller template. It features a 2D sprite character that uses texture based animation where the camera is positioned at the avatar's side. Character movement, which can be controlled with either keyboard, controller or virtual joystick on a touch device, is limited to two dimensions: forward, backward and jumping. A number of ledges are featured in the level and the character can jump between them.")
 
 LocalizedDisplayNames=(Language="ko",Text="2D 횡스크롤")
-LocalizedDescriptions=(Language="ko",Text="횡스크롤 템플릿의 페이퍼 2D 버전입니다. 카메라가 아바타의 측면에 위치하여 텍스처 기반 애니메이션을 사용하는 2D 스프라이트 캐릭터를 특징으로 합니다. 캐릭터 이동은 키보드, 콘트롤러, 터치 디바이스의 가상 조이스틱 중 하나로 가능하며, 이차원으로 제한되어 앞, 뒤, 점프만 가능합니다. 레벨에는 캐릭터가 점프해 다닐 수 있는 턱이 여러개 있습니다.")
+LocalizedDescriptions=(Language="ko",Text="횡스크롤 템플릿의 페이퍼 2D 버전입니다. 카메라가 아바타의 측면에 위치하여 텍스처 기반 애니메이션을 사용하는 2D 스프라이트 캐릭터를 특징으로 합니다. 캐릭터 이동은 키보드, 컨트롤러, 터치 디바이스의 가상 조이스틱 중 하나로 가능하며, 이차원으로 제한되어 앞, 뒤, 점프만 가능합니다. 레벨에는 캐릭터가 점프해 다닐 수 있는 턱이 여러개 있습니다.")
+
+LocalizedDisplayNames=(Language="ja",Text="2Dサイドスクロール")
+LocalizedDescriptions=(Language="ja",Text="サイドスクロールテンプレートの Paper 2D 版です。テクスチャベースのアニメーションを使用する 2D スプライトを使用し、カメラはアバターの側面に存在します。キャラクタの移動はキーボード、コントローラー、タッチデバイスでのバーチャルジョイスティックで操作可能で、二次元に制限されています。前進、後退とジャンプだけです。レベルには足場があり、それらの間をジャンプすることができます。")
+
+LocalizedDisplayNames=(Language="zh-CN",Text="2D横版过关游戏")
+LocalizedDescriptions=(Language="zh-CN",Text="这是一个2D横版过关游戏的模板。它拥有一个2D的Sprite角色，使用基于贴图的动画，摄像机则位于角色的一侧。您可以通过键盘、手柄或触摸式设备上的虚拟摇杆来控制角色移动，但仅限于两个方向：前进、后退及跳跃。关卡中有许多岩架，角色可以在它们之间跳跃。 ")
 
 ClassTypes=PaperCharacter, GameMode, CameraComponent, SpringArmComponent, PaperFlipbook
 AssetTypes=Texture,Sprite, Paper Flipbook, Material
diff --git a/Templates/TP_2DSideScroller/manifest.json b/Templates/TP_2DSideScroller/manifest.json
--- a/Templates/TP_2DSideScroller/manifest.json
+++ b/Templates/TP_2DSideScroller/manifest.json
@@ -13,6 +13,10 @@
 		{
 			"Language": "ja",
 			"Text": "2D横スクロール"
+		},
+		{
+			"Language": "zh-CN",
+			"Text": "2D横版过关游戏"
 		}
 	],
 	"Description":
@@ -23,11 +27,15 @@
 		},
 		{
 			"Language": "ko",
-			"Text": "횡스크롤 팩의 페이퍼 2D 버전입니다. 카메라가 아바타의 측면에 위치하여 텍스처 기반 애니메이션을 사용하는 2D 스프라이트 캐릭터를 특징으로 합니다. 캐릭터 이동은 키보드, 콘트롤러, 터치 디바이스의 가상 조이스틱 중 하나로 가능하며, 이차원으로 제한되어 앞, 뒤, 점프만 가능합니다. 레벨에는 캐릭터가 점프해 다닐 수 있는 턱이 여러개 있습니다."
+			"Text": "횡스크롤 팩의 페이퍼 2D 버전입니다. 카메라가 아바타의 측면에 위치하여 텍스처 기반 애니메이션을 사용하는 2D 스프라이트 캐릭터를 특징으로 합니다. 캐릭터 이동은 키보드, 컨트롤러, 터치 디바이스의 가상 조이스틱 중 하나로 가능하며, 이차원으로 제한되어 앞, 뒤, 점프만 가능합니다. 레벨에는 캐릭터가 점프해 다닐 수 있는 턱이 여러개 있습니다."
 		},
 		{
 			"Language": "ja",
-			"Text": "このテンプレートはペーパー２Dを使った横スクロール系のゲームに適しています。キャラクターは２Dのスプライトとテクスチャ切り替えのアニメーションで設計され、カメラは常にキャラクターを追いかける形になっています。キャラクターはキーボード・コントローラ・タッチパネル上のバーチャルジョイスティックで操作することができます。基本操作は前進・後退・ジャンプの三つです。サンプルのレベルに配置された障害物はキャラクターが飛び越えていける間隔で配置されています。"
+			"Text": "サイドスクロールの Paper 2D 版です。テクスチャベースのアニメーションを使用する 2D スプライトを使用し、カメラはアバターの側面に存在します。キャラクタの移動はキーボード、コントローラー、タッチデバイスでのバーチャルジョイスティックで操作可能で、二次元に制限されています。前進、後退とジャンプだけです。レベルには足場があり、それらの間をジャンプすることができます。"
+		},
+		{
+			"Language": "zh-CN",
+			"Text": "这是一个2D横版过关游戏的模板。它拥有一个2D的Sprite角色，使用基于贴图的动画，摄像机则位于角色的一侧。您可以通过键盘、手柄或触摸式设备上的虚拟摇杆来控制角色移动，但仅限于两个方向：前进、后退及跳跃。关卡中有许多岩架，角色可以在它们之间跳跃。"
 		}
 	],
 	"AssetTypes":
@@ -35,14 +43,6 @@
 		{
 			"Language": "en",
 			"Text": "Texture,Sprite, Sprite Flipbook, Material"
-		},
-		{
-			"Language": "ko",
-			"Text": "텍스처, 스프라이트, 스프라이트 플립북, 머티리얼"
-		},
-		{
-			"Language": "ja",
-			"Text": "テクスチャ・スプライト・スプライトブリップブック・マテリアル"
 		}
 	],	
 	"SearchTags":
@@ -50,10 +50,6 @@
 		{
 			"Language": "en",
 			"Text": "Side,Scroller,Platform,2D,Paper"
-		},
-		{
-			"Language": "ko",
-			"Text": "횡,스크롤,플랫폼,2D,페이퍼"
 		}
 	],
 	"ClassTypes": "PaperCharacter, GameMode, CameraComponent, SpringArmComponent, PaperFlipbook",
diff --git a/Templates/TP_2DSideScrollerBP/Config/TemplateDefs.ini b/Templates/TP_2DSideScrollerBP/Config/TemplateDefs.ini
--- a/Templates/TP_2DSideScrollerBP/Config/TemplateDefs.ini
+++ b/Templates/TP_2DSideScrollerBP/Config/TemplateDefs.ini
@@ -3,7 +3,13 @@ LocalizedDisplayNames=(Language="en",Text="2D Side Scroller")
 LocalizedDescriptions=(Language="en",Text="This is a paper 2D version of the Side Scroller template. It features a 2D sprite character that uses texture based animation where the camera is positioned at the avatar's side. Character movement, which can be controlled with either keyboard, controller or virtual joystick on a touch device, is limited to two dimensions: forward, backward and jumping. A number of ledges are featured in the level and the character can jump between them.")
 
 LocalizedDisplayNames=(Language="ko",Text="2D 횡스크롤")
-LocalizedDescriptions=(Language="ko",Text="횡스크롤 템플릿의 페이퍼 2D 버전입니다. 카메라가 아바타의 측면에 위치하여 텍스처 기반 애니메이션을 사용하는 2D 스프라이트 캐릭터를 특징으로 합니다. 캐릭터 이동은 키보드, 콘트롤러, 터치 디바이스의 가상 조이스틱 중 하나로 가능하며, 이차원으로 제한되어 앞, 뒤, 점프만 가능합니다. 레벨에는 캐릭터가 점프해 다닐 수 있는 턱이 여러개 있습니다.")
+LocalizedDescriptions=(Language="ko",Text="횡스크롤 템플릿의 페이퍼 2D 버전입니다. 카메라가 아바타의 측면에 위치하여 텍스처 기반 애니메이션을 사용하는 2D 스프라이트 캐릭터를 특징으로 합니다. 캐릭터 이동은 키보드, 컨트롤러, 터치 디바이스의 가상 조이스틱 중 하나로 가능하며, 이차원으로 제한되어 앞, 뒤, 점프만 가능합니다. 레벨에는 캐릭터가 점프해 다닐 수 있는 턱이 여러개 있습니다.")
+
+LocalizedDisplayNames=(Language="ja",Text="2Dサイドスクロール")
+LocalizedDescriptions=(Language="ja",Text="サイドスクロールテンプレートの Paper 2D 版です。テクスチャベースのアニメーションを使用する 2D スプライトを使用し、カメラはアバターの側面に存在します。キャラクタの移動はキーボード、コントローラー、タッチデバイスでのバーチャルジョイスティックで操作可能で、二次元に制限されています。前進、後退とジャンプだけです。レベルには足場があり、それらの間をジャンプすることができます。")
+
+LocalizedDisplayNames=(Language="zh-CN",Text="2D横版过关游戏")
+LocalizedDescriptions=(Language="zh-CN",Text="这是一个2D横版过关游戏的模板。它拥有一个2D的Sprite角色，使用基于贴图的动画，摄像机则位于角色的一侧。您可以通过键盘、手柄或触摸式设备上的虚拟摇杆来控制角色移动，但仅限于两个方向：前进、后退及跳跃。关卡中有许多岩架，角色可以在它们之间跳跃。 ")
 
 ClassTypes=APaperCharacter, AGameMode, UCameraComponent, USpringArmComponent, UPaperFlipbook
 AssetTypes=Texture,Sprite, Sprite Flipbook, Material
diff --git a/Templates/TP_2DSideScrollerBP/manifest.json b/Templates/TP_2DSideScrollerBP/manifest.json
--- a/Templates/TP_2DSideScrollerBP/manifest.json
+++ b/Templates/TP_2DSideScrollerBP/manifest.json
@@ -13,6 +13,10 @@
 		{
 			"Language": "ja",
 			"Text": "2D横スクロール"
+		},
+		{
+			"Language": "zh-CN",
+			"Text": "2D横版过关游戏"
 		}
 	],
 	"Description":
@@ -23,11 +27,15 @@
 		},
 		{
 			"Language": "ko",
-			"Text": "횡스크롤 팩의 페이퍼 2D 버전입니다. 카메라가 아바타의 측면에 위치하여 텍스처 기반 애니메이션을 사용하는 2D 스프라이트 캐릭터를 특징으로 합니다. 캐릭터 이동은 키보드, 콘트롤러, 터치 디바이스의 가상 조이스틱 중 하나로 가능하며, 이차원으로 제한되어 앞, 뒤, 점프만 가능합니다. 레벨에는 캐릭터가 점프해 다닐 수 있는 턱이 여러개 있습니다."
+			"Text": "횡스크롤 팩의 페이퍼 2D 버전입니다. 카메라가 아바타의 측면에 위치하여 텍스처 기반 애니메이션을 사용하는 2D 스프라이트 캐릭터를 특징으로 합니다. 캐릭터 이동은 키보드, 컨트롤러, 터치 디바이스의 가상 조이스틱 중 하나로 가능하며, 이차원으로 제한되어 앞, 뒤, 점프만 가능합니다. 레벨에는 캐릭터가 점프해 다닐 수 있는 턱이 여러개 있습니다."
 		},
 		{
 			"Language": "ja",
-			"Text": "このテンプレートはペーパー２Dを使った横スクロール系のゲームに適しています。キャラクターは２Dのスプライトとテクスチャ切り替えのアニメーションで設計され、カメラは常にキャラクターを追いかける形になっています。キャラクターはキーボード・コントローラ・タッチパネル上のバーチャルジョイスティックで操作することができます。基本操作は前進・後退・ジャンプの三つです。サンプルのレベルに配置された障害物はキャラクターが飛び越えていける間隔で配置されています。"
+			"Text": "サイドスクロールの Paper 2D 版です。テクスチャベースのアニメーションを使用する 2D スプライトを使用し、カメラはアバターの側面に存在します。キャラクタの移動はキーボード、コントローラー、タッチデバイスでのバーチャルジョイスティックで操作可能で、二次元に制限されています。前進、後退とジャンプだけです。レベルには足場があり、それらの間をジャンプすることができます。"
+		},
+		{
+			"Language": "zh-CN",
+			"Text": "这是一个2D横版过关游戏的模板。它拥有一个2D的Sprite角色，使用基于贴图的动画，摄像机则位于角色的一侧。您可以通过键盘、手柄或触摸式设备上的虚拟摇杆来控制角色移动，但仅限于两个方向：前进、后退及跳跃。关卡中有许多岩架，角色可以在它们之间跳跃。"
 		}
 	],
 	"AssetTypes":
@@ -35,26 +43,13 @@
 		{
 			"Language": "en",
 			"Text": "Texture,Sprite, Sprite Flipbook, Material"
-		},
-		{
-			"Language": "ko",
-			"Text": "텍스처, 스프라이트, 스프라이트 플립북, 머티리얼"
-		},
-		{
-			"Language": "ja",
-			"Text": "テクスチャ・スプライト・スプライトフリップブック・マテリアル"
 		}
-
 	],	
 	"SearchTags":
 	[
 		{
 			"Language": "en",
 			"Text": "Side,Scroller,Platform,2D,Paper"
-		},
-		{
-			"Language": "ko",
-			"Text": "횡,스크롤,플랫폼,2D,페이퍼"
 		}
 	],
 	"ClassTypes": "PaperCharacter, GameMode, CameraComponent, SpringArmComponent, PaperFlipbook",
diff --git a/Templates/TP_FirstPerson/Config/TemplateDefs.ini b/Templates/TP_FirstPerson/Config/TemplateDefs.ini
--- a/Templates/TP_FirstPerson/Config/TemplateDefs.ini
+++ b/Templates/TP_FirstPerson/Config/TemplateDefs.ini
@@ -3,7 +3,13 @@ LocalizedDisplayNames=(Language="en",Text="First Person")
 LocalizedDescriptions=(Language="en",Text="The first person template features a player character (represented by a pair of arms) which is viewed from first person perspective. The character can be moved around the level using keyboard, controller or virtual joystick on a touch device. Additionally the player can look around using mouse, controller or virtual joystick on a touch device. The character is also equipped with a gun that, using mouse,  controller or virtual joystick on a touch device will fire a simple sphere projectile that will affect some physics objects in the level, whilst bounce off the arena walls.")
 
 LocalizedDisplayNames=(Language="ko",Text="일인칭")
-LocalizedDescriptions=(Language="ko",Text="일인칭 템플릿은 일인칭 시점에서 보이는 (양팔로 표현되는) 플레이어 캐릭터를 특징으로 합니다. 캐릭터 이동은 키보드, 콘트롤러, 터치 디바이스의 가상 조이스틱으로 가능합니다. 추가적으로 플레이어는 마우스, 콘트롤러, 터치 디바이스의 가상 조이스틱을 사용하여 둘러보기도 가능합니다. 캐릭터에는 총이 장착되어 있어, 마우스, 콘트롤러, 터치 디바이스의 가상 조이스틱을 사용하면 간단한 구체 프로젝타일이 발사되어, 아레나 벽에 튕기면서 레벨에 있는 몇몇 물리 오브젝트에 영향을 끼칩니다.")
+LocalizedDescriptions=(Language="ko",Text="일인칭 템플릿은 일인칭 시점에서 보이는 (양팔로 표현되는) 플레이어 캐릭터를 특징으로 합니다. 캐릭터 이동은 키보드, 컨트롤러, 터치 디바이스의 가상 조이스틱으로 가능합니다. 추가적으로 플레이어는 마우스, 컨트롤러, 터치 디바이스의 가상 조이스틱을 사용하여 주변을 둘러볼 수도 있습니다. 캐릭터에는 총이 장착되어 있어 마우스, 컨트롤러, 터치 디바이스의 가상 조이스틱을 사용하면 간단한 발사체가 발사되어, 아레나 벽에 튕기면서 레벨에 있는 몇몇 물리 오브젝트에 영향을 줍니다.")
+
+LocalizedDisplayNames=(Language="ja",Text="ファーストパーソン")
+LocalizedDescriptions=(Language="ja",Text="ファーストパーソンテンプレートには、一人称視点で表示されるプレイヤーキャラクタ（腕として表示されています）がいます。キーボード、コントローラー、タッチデバイスでのバーチャルジョイスティックでこのキャラクタを移動させることができます。追加して、プレイヤーはマウス、コントローラー、タッチデバイスでのバーチャルジョイスティックで周りを見回すことができます。さらにキャラクタは銃を装備して、マウス、コントローラー、タッチデバイスでのバーチャルジョイスティックを使ってシンプルな球体の弾を発射できます。この弾は物理オブジェクトに影響し、アリーナの壁で反射します。")
+
+LocalizedDisplayNames=(Language="zh-CN",Text="第一人称游戏")
+LocalizedDescriptions=(Language="zh-CN",Text="第一人称模板拥有一个第一人称的玩家角色（通过一双手臂呈现）。您可以通过键盘、手柄或触摸式设备上的虚拟摇杆控制角色在关卡中四处移动。此外，您还可以使用鼠标、手柄或触摸式设备上的虚拟摇杆转动角色视角，观察四周。角色还装备了枪支，可使用鼠标、手柄或触摸式设备上的虚拟摇杆发射一颗球形子弹，并影响关卡中的一些物理对象，同时从竞技场的墙壁上弹开。")
 
 ClassTypes=Character, GameMode, HUD, Actor
 AssetTypes=Animation Sequence, Animation Montage, Sound Wave, Skeletal Mesh, Animation Blueprint, Skeleton, Static Mesh, Texture.
diff --git a/Templates/TP_FirstPersonBP/Config/TemplateDefs.ini b/Templates/TP_FirstPersonBP/Config/TemplateDefs.ini
--- a/Templates/TP_FirstPersonBP/Config/TemplateDefs.ini
+++ b/Templates/TP_FirstPersonBP/Config/TemplateDefs.ini
@@ -3,7 +3,13 @@ LocalizedDisplayNames=(Language="en",Text="First Person")
 LocalizedDescriptions=(Language="en",Text="The first person template features a player character (represented by a pair of arms) which is viewed from first person perspective. The character can be moved around the level using keyboard, controller or virtual joystick on a touch device. Additionally the player can look around using mouse, controller or virtual joystick on a touch device. The character is also equipped with a gun that, using mouse,  controller or virtual joystick on a touch device will fire a simple sphere projectile that will affect some physics objects in the level, whilst bounce off the arena walls.")
 
 LocalizedDisplayNames=(Language="ko",Text="일인칭")
-LocalizedDescriptions=(Language="ko",Text="일인칭 템플릿은 일인칭 시점에서 보이는 (양팔로 표현되는) 플레이어 캐릭터를 특징으로 합니다. 캐릭터 이동은 키보드, 콘트롤러, 터치 디바이스의 가상 조이스틱으로 가능합니다. 추가적으로 플레이어는 마우스, 콘트롤러, 터치 디바이스의 가상 조이스틱을 사용하여 둘러보기도 가능합니다. 캐릭터에는 총이 장착되어 있어, 마우스, 콘트롤러, 터치 디바이스의 가상 조이스틱을 사용하면 간단한 구체 프로젝타일이 발사되어, 아레나 벽에 튕기면서 레벨에 있는 몇몇 물리 오브젝트에 영향을 끼칩니다.")
+LocalizedDescriptions=(Language="ko",Text="일인칭 템플릿은 일인칭 시점에서 보이는 (양팔로 표현되는) 플레이어 캐릭터를 특징으로 합니다. 캐릭터 이동은 키보드, 컨트롤러, 터치 디바이스의 가상 조이스틱으로 가능합니다. 추가적으로 플레이어는 마우스, 컨트롤러, 터치 디바이스의 가상 조이스틱을 사용하여 주변을 둘러볼 수도 있습니다. 캐릭터에는 총이 장착되어 있어 마우스, 컨트롤러, 터치 디바이스의 가상 조이스틱을 사용하면 간단한 발사체가 발사되어, 아레나 벽에 튕기면서 레벨에 있는 몇몇 물리 오브젝트에 영향을 줍니다.")
+
+LocalizedDisplayNames=(Language="ja",Text="ファーストパーソン")
+LocalizedDescriptions=(Language="ja",Text="ファーストパーソンテンプレートには、一人称視点で表示されるプレイヤーキャラクタ（腕として表示されています）がいます。キーボード、コントローラー、タッチデバイスでのバーチャルジョイスティックでこのキャラクタを移動させることができます。追加して、プレイヤーはマウス、コントローラー、タッチデバイスでのバーチャルジョイスティックで周りを見回すことができます。さらにキャラクタは銃を装備して、マウス、コントローラー、タッチデバイスでのバーチャルジョイスティックを使ってシンプルな球体の弾を発射できます。この弾は物理オブジェクトに影響し、アリーナの壁で反射します。")
+
+LocalizedDisplayNames=(Language="zh-CN",Text="第一人称游戏")
+LocalizedDescriptions=(Language="zh-CN",Text="第一人称模板拥有一个第一人称的玩家角色（通过一双手臂呈现）。您可以通过键盘、手柄或触摸式设备上的虚拟摇杆控制角色在关卡中四处移动。此外，您还可以使用鼠标、手柄或触摸式设备上的虚拟摇杆转动角色视角，观察四周。角色还装备了枪支，可使用鼠标、手柄或触摸式设备上的虚拟摇杆发射一颗球形子弹，并影响关卡中的一些物理对象，同时从竞技场的墙壁上弹开。")
 
 ClassTypes=Character, GameMode, HUD, Actor
 AssetTypes=Animation Sequence, Animation Montage, Sound Wave, Skeletal Mesh, Animation Blueprint, Skeleton, Static Mesh, Texture.
diff --git a/Templates/TP_Flying/Config/TemplateDefs.ini b/Templates/TP_Flying/Config/TemplateDefs.ini
--- a/Templates/TP_Flying/Config/TemplateDefs.ini
+++ b/Templates/TP_Flying/Config/TemplateDefs.ini
@@ -3,7 +3,13 @@ LocalizedDisplayNames=(Language="en",Text="Flying")
 LocalizedDescriptions=(Language="en",Text="The flying template features a spaceship with the camera positioned behind the ship represented by a static mesh. The keyboard, controller or virtual joystick controls the pitch and yaw of the ship as well as providing a means to allow speed control of the ships forward movement. The level is populated with several clusters of boxes that block the movement of the ship.")
 
 LocalizedDisplayNames=(Language="ko",Text="비행")
-LocalizedDescriptions=(Language="ko",Text="비행 템플릿은 뒤에 카메라가 위치한 스태틱 메시로 표현되는 비행선을 특징으로 합니다. 키보드, 콘트롤러, 가상 조이스틱으로 비행선의 핏치와 요 뿐만 아니라, 비행선의 전방 이동 속도를 제어할 수도 있습니다. 레벨에는 비행선의 이동을 막는 박스 덩어리가 여럿 채워져 있습니다.")
+LocalizedDescriptions=(Language="ko",Text="비행 템플릿은 뒤에 카메라가 위치한 스태틱 메시로 표현되는 비행선을 특징으로 합니다. 키보드, 컨트롤러, 가상 조이스틱으로 비행선의 피치와 요 뿐만 아니라, 비행선의 전방 이동 속도를 제어할 수도 있습니다. 레벨에는 비행선의 이동을 막는 박스 덩어리가 여럿 있습니다.")
+
+LocalizedDisplayNames=(Language="ja",Text="フライング")
+LocalizedDescriptions=(Language="ja",Text="フライングテンプレートには、スタティックメッシュで表現された宇宙船の背後にカメラがあります。キーボード、コントローラー、バーチャルジョイスティックで宇宙船のピッチとヨーのコントロールと、前進時の速度の調整ができます。レベルには積まれた箱があり、宇宙船の移動を邪魔しています。")
+
+LocalizedDisplayNames=(Language="zh-CN",Text="飞行类游戏")
+LocalizedDescriptions=(Language="zh-CN",Text="飞行类游戏模板里拥有一艘宇宙飞船，摄像机位于飞船后方，由静态网格体呈现。键盘、手柄或触摸式设备上的虚拟摇杆可控制飞船的纵倾和航向，并提供了控制飞船前进速度的方法。关卡中还有一些立方体集群，可阻挡飞船的移动。")
 
 ClassTypes=GameMode, Pawn, StaticMeshComponent, SpringArmComponent, CameraComponent, InputComponent, PrimitiveComponent
 AssetTypes=Static Mesh
diff --git a/Templates/TP_Flying/manifest.json b/Templates/TP_Flying/manifest.json
--- a/Templates/TP_Flying/manifest.json
+++ b/Templates/TP_Flying/manifest.json
@@ -10,9 +10,13 @@
 			"Language": "ko",
 			"Text": "비행"
 		},
+		{
+			"Language": "zh-CN",
+			"Text": "飞行类游戏"
+		},
 		{
 			"Language": "ja",
-			"Text": "飛行"
+			"Text": "フライング"
 		}
 	],
 	"Description":
@@ -23,11 +27,15 @@
 		},
 		{
 			"Language": "ko",
-			"Text": "비행 팩은 뒤에 카메라가 위치한 스태틱 메시로 표현되는 비행선을 특징으로 합니다. 키보드, 콘트롤러, 가상 조이스틱으로 비행선의 핏치와 요 뿐만 아니라, 비행선의 전방 이동 속도를 제어할 수도 있습니다. 레벨에는 비행선의 이동을 막는 박스 덩어리가 여럿 채워져 있습니다."
+			"Text": "비행 템플릿은 뒤에 카메라가 위치한 스태틱 메시로 표현되는 비행선을 특징으로 합니다. 키보드, 컨트롤러, 가상 조이스틱으로 비행선의 피치와 요 뿐만 아니라, 비행선의 전방 이동 속도를 제어할 수도 있습니다. 레벨에는 비행선의 이동을 막는 박스 덩어리가 여럿 있습니다."
+		},
+		{
+			"Language": "zh-CN",
+			"Text": "飞行类游戏模板里拥有一艘宇宙飞船，摄像机位于飞船后方，由静态网格体呈现。键盘、手柄或触摸式设备上的虚拟摇杆可控制飞船的纵倾和航向，并提供了控制飞船前进速度的方法。关卡中还有一些立方体集群，可阻挡飞船的移动。"
 		},
 		{
 			"Language": "ja",
-			"Text": "The flying pack features a spaceship with the camera positioned behind the ship represented by a static mesh. The keyboard, controller or virtual joystick controls the pitch and yaw of the ship as well as providing a means to allow speed control of the ships forward movement. The level is populated with several clusters of boxes that block the movement of the ship."
+			"Text": "フライングテンプレートには、スタティックメッシュで表現された宇宙船の背後にカメラがあります。キーボード、コントローラー、バーチャルジョイスティックで宇宙船のピッチとヨーのコントロールと、前進時の速度の調整ができます。レベルには積まれた箱があり、宇宙船の移動を邪魔しています。"
 		}
 
 	],
@@ -36,14 +44,6 @@
 		{
 			"Language": "en",
 			"Text": "Static Mesh"
-		},
-		{
-			"Language": "ko",
-			"Text": "스태틱 메시"
-		},
-		{
-			"Language": "ja",
-			"Text": "スタティックメッシュ"
 		}
 
 	],	
@@ -52,10 +52,6 @@
 		{
 			"Language": "en",
 			"Text": "Flying,Hover,UFO"
-		},
-		{
-			"Language": "ko",
-			"Text": "비행,호버,UFO"
 		}
 	],	
 	"ClassTypes": "GameMode, Pawn, StaticMeshComponent, SpringArmComponent, CameraComponent, InputComponent, PrimitiveComponent",	
diff --git a/Templates/TP_FlyingBP/Config/TemplateDefs.ini b/Templates/TP_FlyingBP/Config/TemplateDefs.ini
--- a/Templates/TP_FlyingBP/Config/TemplateDefs.ini
+++ b/Templates/TP_FlyingBP/Config/TemplateDefs.ini
@@ -3,7 +3,13 @@ LocalizedDisplayNames=(Language="en",Text="Flying")
 LocalizedDescriptions=(Language="en",Text="The flying template features a spaceship with the camera positioned behind the ship represented by a static mesh. The keyboard, controller or virtual joystick controls the pitch and yaw of the ship as well as providing a means to allow speed control of the ships forward movement. The level is populated with several clusters of boxes that block the movement of the ship.")
 
 LocalizedDisplayNames=(Language="ko",Text="비행")
-LocalizedDescriptions=(Language="ko",Text="비행 템플릿은 뒤에 카메라가 위치한 스태틱 메시로 표현되는 비행선을 특징으로 합니다. 키보드, 콘트롤러, 가상 조이스틱으로 비행선의 핏치와 요 뿐만 아니라, 비행선의 전방 이동 속도를 제어할 수도 있습니다. 레벨에는 비행선의 이동을 막는 박스 덩어리가 여럿 채워져 있습니다.")
+LocalizedDescriptions=(Language="ko",Text="비행 템플릿은 뒤에 카메라가 위치한 스태틱 메시로 표현되는 비행선을 특징으로 합니다. 키보드, 컨트롤러, 가상 조이스틱으로 비행선의 피치와 요 뿐만 아니라, 비행선의 전방 이동 속도를 제어할 수도 있습니다. 레벨에는 비행선의 이동을 막는 박스 덩어리가 여럿 있습니다.")
+
+LocalizedDisplayNames=(Language="ja",Text="フライング")
+LocalizedDescriptions=(Language="ja",Text="フライングテンプレートには、スタティックメッシュで表現された宇宙船の背後にカメラがあります。キーボード、コントローラー、バーチャルジョイスティックで宇宙船のピッチとヨーのコントロールと、前進時の速度の調整ができます。レベルには積まれた箱があり、宇宙船の移動を邪魔しています。")
+
+LocalizedDisplayNames=(Language="zh-CN",Text="飞行类游戏")
+LocalizedDescriptions=(Language="zh-CN",Text="飞行类游戏模板里拥有一艘宇宙飞船，摄像机位于飞船后方，由静态网格体呈现。键盘、手柄或触摸式设备上的虚拟摇杆可控制飞船的纵倾和航向，并提供了控制飞船前进速度的方法。关卡中还有一些立方体集群，可阻挡飞船的移动。")
 
 ClassTypes=GameMode, Pawn, StaticMeshComponent, SpringArmComponent, CameraComponent, InputComponent, PrimitiveComponent
 AssetTypes=Static Mesh
diff --git a/Templates/TP_FlyingBP/manifest.json b/Templates/TP_FlyingBP/manifest.json
--- a/Templates/TP_FlyingBP/manifest.json
+++ b/Templates/TP_FlyingBP/manifest.json
@@ -10,9 +10,13 @@
 			"Language": "ko",
 			"Text": "비행"
 		},
+		{
+			"Language": "zh-CN",
+			"Text": "飞行类游戏"
+		},
 		{
 			"Language": "ja",
-			"Text": "飛行"
+			"Text": "フライング"
 		}
 	],
 	"Description":
@@ -23,37 +27,31 @@
 		},
 		{
 			"Language": "ko",
-			"Text": "비행 팩은 뒤에 카메라가 위치한 스태틱 메시로 표현되는 비행선을 특징으로 합니다. 키보드, 콘트롤러, 가상 조이스틱으로 비행선의 핏치와 요 뿐만 아니라, 비행선의 전방 이동 속도를 제어할 수도 있습니다. 레벨에는 비행선의 이동을 막는 박스 덩어리가 여럿 채워져 있습니다."
+			"Text": "비행 템플릿은 뒤에 카메라가 위치한 스태틱 메시로 표현되는 비행선을 특징으로 합니다. 키보드, 컨트롤러, 가상 조이스틱으로 비행선의 피치와 요 뿐만 아니라, 비행선의 전방 이동 속도를 제어할 수도 있습니다. 레벨에는 비행선의 이동을 막는 박스 덩어리가 여럿 있습니다."
+		},
+		{
+			"Language": "zh-CN",
+			"Text": "飞行类游戏模板里拥有一艘宇宙飞船，摄像机位于飞船后方，由静态网格体呈现。键盘、手柄或触摸式设备上的虚拟摇杆可控制飞船的纵倾和航向，并提供了控制飞船前进速度的方法。关卡中还有一些立方体集群，可阻挡飞船的移动。"
 		},
 		{
 			"Language": "ja",
-			"Text": "The flying pack features a spaceship with the camera positioned behind the ship represented by a static mesh. The keyboard, controller or virtual joystick controls the pitch and yaw of the ship as well as providing a means to allow speed control of the ships forward movement. The level is populated with several clusters of boxes that block the movement of the ship."
+			"Text": "フライングテンプレートには、スタティックメッシュで表現された宇宙船の背後にカメラがあります。キーボード、コントローラー、バーチャルジョイスティックで宇宙船のピッチとヨーのコントロールと、前進時の速度の調整ができます。レベルには積まれた箱があり、宇宙船の移動を邪魔しています。"
 		}
+
 	],
 	"AssetTypes":
 	[
 		{
 			"Language": "en",
 			"Text": "Static Mesh"
-		},
-		{
-			"Language": "ko",
-			"Text": "스태틱 메시"
-		},
-		{
-			"Language": "ja",
-			"Text": "スタティックメッシュ"
 		}
+
 	],	
 	"SearchTags":
 	[
 		{
 			"Language": "en",
 			"Text": "Flying,Hover,UFO"
-		},
-		{
-			"Language": "ko",
-			"Text": "비행,호버,UFO"
 		}
 	],	
 	"ClassTypes": "GameMode, Pawn, StaticMeshComponent, SpringArmComponent, CameraComponent, InputComponent, PrimitiveComponent",
diff --git a/Templates/TP_HandheldARBP/Config/TemplateDefs.ini b/Templates/TP_HandheldARBP/Config/TemplateDefs.ini
--- a/Templates/TP_HandheldARBP/Config/TemplateDefs.ini
+++ b/Templates/TP_HandheldARBP/Config/TemplateDefs.ini
@@ -3,7 +3,13 @@ LocalizedDisplayNames=(Language="en",Text="Handheld AR")
 LocalizedDescriptions=(Language="en",Text="A starting point for building Augmented Reality applications for Android and iOS devices. Includes runtime logic for toggling AR mode on and off, debug information regarding plane detection, example code for hit detection and handling light estimation.")
 
 LocalizedDisplayNames=(Language="ko",Text="휴대형 증강 현실")
-LocalizedDescriptions=(Language="ko",Text="안드로이드 및 iOS 디바이스용 증강 현실 어플리케이션을 만들기 위한 템플릿입니다. AR 모드를 켜고 끄는 런타임 로직, 면 인식 관련 디버깅 정보, 히트 디텍션과 빛 추정 처리 관련 예제 코드가 포함되어 있습니다.")
+LocalizedDescriptions=(Language="ko",Text="Android 및 iOS 디바이스용 증강 현실 앱을 만들기 위한 템플릿입니다. AR 모드를 켜고 끄는 런타임 로직, 면 인식 관련 디버깅 정보, 히트 디텍션과 빛 추정 처리 관련 예제 코드가 포함되어 있습니다.")
+
+LocalizedDisplayNames=(Language="ja",Text="携帯 AR")
+LocalizedDescriptions=(Language="ja",Text="Android と iOS 向けの拡張現実アプリケーション作成の土台です。AR モードのオンオフを行う実行時のロジック、平面検出に関するデバッグ情報、ヒット検出のサンプルコード、ライト推定などを含んでいます。")
+
+LocalizedDisplayNames=(Language="zh-CN",Text="手持式AR应用")
+LocalizedDescriptions=(Language="zh-CN",Text="为安卓及iOS设备创建增强现实应用的良好起点，包含开启和关闭AR模式的运行逻辑，关于平面检测的调试信息，命中检测及处理预计光照的示例代码。")
 
 ClassTypes=Pawn, GameMode, Widget, Actor
 AssetTypes=Static Mesh
diff --git a/Templates/TP_HandheldARBP/manifest.json b/Templates/TP_HandheldARBP/manifest.json
--- a/Templates/TP_HandheldARBP/manifest.json
+++ b/Templates/TP_HandheldARBP/manifest.json
@@ -5,6 +5,18 @@
 		{
 			"Language": "en",
 			"Text": "Handheld AR"
+		},
+		{
+			"Language": "ko",
+			"Text": "휴대형 증강 현실"
+		},
+		{
+			"Language": "ja",
+			"Text": "携帯 AR"
+		},
+		{
+			"Language": "zh-CN",
+			"Text": "手持式AR应用"
 		}
 	],
 	"Description":
@@ -12,6 +24,18 @@
     {
       "Language": "en",
       "Text": "A starting point for building Augmented Reality applications for Android and iOS devices. Includes runtime logic for toggling AR mode on and off, debug information regarding plane detection, example code for hit detection and handling light estimation."
+    },
+    {
+      "Language": "ko",
+      "Text": "Android 및 iOS 디바이스용 증강 현실 앱을 만들기 위한 템플릿입니다. AR 모드를 켜고 끄는 런타임 로직, 면 인식 관련 디버깅 정보, 히트 디텍션과 빛 추정 처리 관련 예제 코드가 포함되어 있습니다."
+    },
+    {
+      "Language": "ja",
+      "Text": "Android と iOS 向けの拡張現実アプリケーション作成の土台です。AR モードのオンオフを行う実行時のロジック、平面検出に関するデバッグ情報、ヒット検出のサンプルコード、ライト推定などを含んでいます。"
+    },
+    {
+      "Language": "zh-CN",
+      "Text": "为安卓及iOS设备创建增强现实应用的良好起点，包含开启和关闭AR模式的运行逻辑，关于平面检测的调试信息，命中检测及处理预计光照的示例代码。"
     }
 	],
 	"AssetTypes":
diff --git a/Templates/TP_Puzzle/Config/TemplateDefs.ini b/Templates/TP_Puzzle/Config/TemplateDefs.ini
--- a/Templates/TP_Puzzle/Config/TemplateDefs.ini
+++ b/Templates/TP_Puzzle/Config/TemplateDefs.ini
@@ -3,7 +3,13 @@ LocalizedDisplayNames=(Language="en",Text="Puzzle")
 LocalizedDescriptions=(Language="en",Text="The puzzle template features a simple grid of blocks. The camera is postioned overhead of the blocks. The user can interact with these blocks by clicking using either the mouse or touch screen. When a block is clicked its color changes and a score, which shown using a text component within the level is increased.")
 
 LocalizedDisplayNames=(Language="ko",Text="퍼즐")
-LocalizedDescriptions=(Language="ko",Text="퍼즐 템플릿은 단순한 블록 모눈을 특징으로 합니다. 카메라는 블록 위에 위치해 있습니다. 사용자는 이 블록을 마우스나 터치스크린으로 클릭하여 상호작용할 수 있습니다. 블록을 클릭하면 색이 바뀌며, 레벨에 텍스트 컴포넌트를 사용하여 표시되는 점수가 올라갑니다.")
+LocalizedDescriptions=(Language="ko",Text="퍼즐 템플릿은 단순한 블록 모눈을 특징으로 합니다. 카메라는 블록 위에 위치해 있습니다. 사용자는 이 블록을 마우스나 터치스크린으로 클릭하여 상호작용할 수 있습니다. 블록을 클릭하면 색이 바뀌면서 레벨에 표시된 점수가 올라가며, 여기에는 텍스트 컴포넌트가 사용되었습니다.")
+
+LocalizedDisplayNames=(Language="ja",Text="パズル")
+LocalizedDescriptions=(Language="ja",Text="パズルテンプレートにはシンプルなブロックのグリッドが含まれています。カメラはブロックを上から見下ろしています。ユーザーはマウスやタッチスクリーンでブロックをクリックすることができます。ブロックをクリックすると色が変化し得点が追加されます。得点はテキストコンポーネントで表示されています。")
+
+LocalizedDisplayNames=(Language="zh-CN",Text="拼图类游戏")
+LocalizedDescriptions=(Language="zh-CN",Text="拼图类游戏模板拥有一张简单的方块网格。摄像机位于方块的正上方。用户可以通过鼠标或触摸屏点击来与这些方块交互。被点击的方块会改变颜色，同时关卡内使用文本组件显示的分数会增加。")
 
 ClassTypes=Actor, PlayerController, SceneComponent, TextRenderComponent, StaticMeshComponent, MaterialInstance, PrimitiveComponent
 AssetTypes=Static Mesh 
diff --git a/Templates/TP_Puzzle/manifest.json b/Templates/TP_Puzzle/manifest.json
--- a/Templates/TP_Puzzle/manifest.json
+++ b/Templates/TP_Puzzle/manifest.json
@@ -10,6 +10,10 @@
 			"Language": "ko",
 			"Text": "퍼즐"
 		},
+		{
+			"Language": "zh-CN",
+			"Text": "拼图类游戏"
+		},
 		{
 			"Language": "ja",
 			"Text": "パズル"
@@ -23,11 +27,15 @@
 		},
 		{
 			"Language": "ko",
-			"Text": "퍼즐 팩은 단순한 블록 모눈을 특징으로 합니다. 카메라는 블록 위에 위치해 있습니다. 사용자는 이 블록을 마우스나 터치스크린으로 클릭하여 상호작용할 수 있습니다. 블록을 클릭하면 색이 바뀌며, 레벨에 텍스트 컴포넌트를 사용하여 표시되는 점수가 올라갑니다."
+			"Text": "퍼즐 템플릿은 단순한 블록 모눈을 특징으로 합니다. 카메라는 블록 위에 위치해 있습니다. 사용자는 이 블록을 마우스나 터치스크린으로 클릭하여 상호작용할 수 있습니다. 블록을 클릭하면 색이 바뀌면서 레벨에 표시된 점수가 올라가며, 여기에는 텍스트 컴포넌트가 사용되었습니다."
+		},
+		{
+			"Language": "zh-CN",
+			"Text": "拼图类游戏模板拥有一张简单的方块网格。摄像机位于方块的正上方。用户可以通过鼠标或触摸屏点击来与这些方块交互。被点击的方块会改变颜色，同时关卡内使用文本组件显示的分数会增加。"
 		},
 		{
 			"Language": "ja",
-			"Text": "The puzzle pack features a simple grid of blocks. The camera is positioned overhead of the blocks. The user can interact with these blocks by clicking using either the mouse or touch screen. When a block is clicked its color changes and a score, which shown using a text component within the level is increased."
+			"Text": "パズルテンプレートにはシンプルなブロックのグリッドが含まれています。カメラはブロックを上から見下ろしています。ユーザーはマウスやタッチスクリーンでブロックをクリックすることができます。ブロックをクリックすると色が変化し得点が追加されます。得点はテキストコンポーネントで表示されています。"
 		}
 	],
 	"AssetTypes":
@@ -35,14 +43,6 @@
 		{
 			"Language": "en",
 			"Text": "Static Mesh"
-		},
-		{
-			"Language": "ko",
-			"Text": "스태틱 메시"
-		},
-		{
-			"Language": "ja",
-			"Text": "スタティックメッシュ"
 		}
 	],	
 	"SearchTags":
@@ -50,10 +50,6 @@
 		{
 			"Language": "en",
 			"Text": "Puzzle,Blocks,Match"
-		},
-		{
-			"Language": "ko",
-			"Text": "퍼즐,블록,짝찾기"
 		}
 	],	
 	"ClassTypes": "Actor, PlayerController, SceneComponent, TextRenderComponent, StaticMeshComponent, MaterialInstance, PrimitiveComponent",
diff --git a/Templates/TP_PuzzleBP/Config/TemplateDefs.ini b/Templates/TP_PuzzleBP/Config/TemplateDefs.ini
--- a/Templates/TP_PuzzleBP/Config/TemplateDefs.ini
+++ b/Templates/TP_PuzzleBP/Config/TemplateDefs.ini
@@ -3,7 +3,13 @@ LocalizedDisplayNames=(Language="en",Text="Puzzle")
 LocalizedDescriptions=(Language="en",Text="The puzzle template features a simple grid of blocks. The camera is postioned overhead of the blocks. The user can interact with these blocks by clicking using either the mouse or touch screen. When a block is clicked its color changes and a score, which shown using a text component within the level is increased.")
 
 LocalizedDisplayNames=(Language="ko",Text="퍼즐")
-LocalizedDescriptions=(Language="ko",Text="퍼즐 템플릿은 단순한 블록 모눈을 특징으로 합니다. 카메라는 블록 위에 위치해 있습니다. 사용자는 이 블록을 마우스나 터치스크린으로 클릭하여 상호작용할 수 있습니다. 블록을 클릭하면 색이 바뀌며, 레벨에 텍스트 컴포넌트를 사용하여 표시되는 점수가 올라갑니다.")
+LocalizedDescriptions=(Language="ko",Text="퍼즐 템플릿은 단순한 블록 모눈을 특징으로 합니다. 카메라는 블록 위에 위치해 있습니다. 사용자는 이 블록을 마우스나 터치스크린으로 클릭하여 상호작용할 수 있습니다. 블록을 클릭하면 색이 바뀌면서 레벨에 표시된 점수가 올라가며, 여기에는 텍스트 컴포넌트가 사용되었습니다.")
+
+LocalizedDisplayNames=(Language="ja",Text="パズル")
+LocalizedDescriptions=(Language="ja",Text="パズルテンプレートにはシンプルなブロックのグリッドが含まれています。カメラはブロックを上から見下ろしています。ユーザーはマウスやタッチスクリーンでブロックをクリックすることができます。ブロックをクリックすると色が変化し得点が追加されます。得点はテキストコンポーネントで表示されています。")
+
+LocalizedDisplayNames=(Language="zh-CN",Text="拼图类游戏")
+LocalizedDescriptions=(Language="zh-CN",Text="拼图类游戏模板拥有一张简单的方块网格。摄像机位于方块的正上方。用户可以通过鼠标或触摸屏点击来与这些方块交互。被点击的方块会改变颜色，同时关卡内使用文本组件显示的分数会增加。")
 
 ClassTypes=Actor, PlayerController, SceneComponent, TextRenderComponent, StaticMeshComponent, MaterialInstance, PrimitiveComponent
 AssetTypes=Static Mesh
diff --git a/Templates/TP_PuzzleBP/manifest.json b/Templates/TP_PuzzleBP/manifest.json
--- a/Templates/TP_PuzzleBP/manifest.json
+++ b/Templates/TP_PuzzleBP/manifest.json
@@ -10,6 +10,10 @@
 			"Language": "ko",
 			"Text": "퍼즐"
 		},
+		{
+			"Language": "zh-CN",
+			"Text": "拼图类游戏"
+		},
 		{
 			"Language": "ja",
 			"Text": "パズル"
@@ -23,11 +27,15 @@
 		},
 		{
 			"Language": "ko",
-			"Text": "퍼즐 팩은 단순한 블록 모눈을 특징으로 합니다. 카메라는 블록 위에 위치해 있습니다. 사용자는 이 블록을 마우스나 터치스크린으로 클릭하여 상호작용할 수 있습니다. 블록을 클릭하면 색이 바뀌며, 레벨에 텍스트 컴포넌트를 사용하여 표시되는 점수가 올라갑니다."
+			"Text": "퍼즐 템플릿은 단순한 블록 모눈을 특징으로 합니다. 카메라는 블록 위에 위치해 있습니다. 사용자는 이 블록을 마우스나 터치스크린으로 클릭하여 상호작용할 수 있습니다. 블록을 클릭하면 색이 바뀌면서 레벨에 표시된 점수가 올라가며, 여기에는 텍스트 컴포넌트가 사용되었습니다."
+		},
+		{
+			"Language": "zh-CN",
+			"Text": "拼图类游戏模板拥有一张简单的方块网格。摄像机位于方块的正上方。用户可以通过鼠标或触摸屏点击来与这些方块交互。被点击的方块会改变颜色，同时关卡内使用文本组件显示的分数会增加。"
 		},
 		{
 			"Language": "ja",
-			"Text": "The puzzle pack features a simple grid of blocks. The camera is positioned overhead of the blocks. The user can interact with these blocks by clicking using either the mouse or touch screen. When a block is clicked its color changes and a score, which shown using a text component within the level is increased."
+			"Text": "パズルテンプレートにはシンプルなブロックのグリッドが含まれています。カメラはブロックを上から見下ろしています。ユーザーはマウスやタッチスクリーンでブロックをクリックすることができます。ブロックをクリックすると色が変化し得点が追加されます。得点はテキストコンポーネントで表示されています。"
 		}
 	],
 	"AssetTypes":
@@ -35,14 +43,6 @@
 		{
 			"Language": "en",
 			"Text": "Static Mesh"
-		},
-		{
-			"Language": "ko",
-			"Text": "스태틱 메시"
-		},
-		{
-			"Language": "ja",
-			"Text": "スタティックメッシュ"
 		}
 	],	
 	"SearchTags":
@@ -50,10 +50,6 @@
 		{
 			"Language": "en",
 			"Text": "Puzzle,Blocks,Match"
-		},
-		{
-			"Language": "ko",
-			"Text": "퍼즐,블록,짝찾기"
 		}
 	],	
 	"ClassTypes": "Actor, PlayerController, SceneComponent, TextRenderComponent, StaticMeshComponent, MaterialInstance, PrimitiveComponent",
diff --git a/Templates/TP_Rolling/Config/TemplateDefs.ini b/Templates/TP_Rolling/Config/TemplateDefs.ini
--- a/Templates/TP_Rolling/Config/TemplateDefs.ini
+++ b/Templates/TP_Rolling/Config/TemplateDefs.ini
@@ -3,7 +3,13 @@ LocalizedDisplayNames=(Language="en",Text="Rolling")
 LocalizedDescriptions=(Language="en",Text="The rolling template features a physics-based rolling ball which can be rolled in any direction using either a keyboard, controller or virtual joystick on touch devices. The camera is positioned above and behind the ball and follows the movement of the ball. Within the level are some ramps, several physics cubes in the level that can be knocked over with the ball as well as some cubes which cannot.")
 
 LocalizedDisplayNames=(Language="ko",Text="굴리기")
-LocalizedDescriptions=(Language="ko",Text="굴리기 템플릿은 키보드, 콘트롤러, 터치 디바이스의 가상 조이스틱을 사용하여 어느 방향이든 굴릴 수 있는 물리 기반 공을 특징으로 합니다. 카메라는 공의 상단 뒤쪽에 위치해 있어 공의 이동을 따라갑니다. 레벨 안에는 경사로와, 공으로 넘어뜨릴 수도 넘어뜨리지 못할수도 있는 물리 큐브가 여럿 있습니다.")
+LocalizedDescriptions=(Language="ko",Text="굴리기 템플릿은 키보드, 컨트롤러, 터치 디바이스의 가상 조이스틱을 사용하여 어느 방향이든 굴릴 수 있는 물리 기반 공을 특징으로 합니다. 카메라는 공의 상단 뒤쪽에 위치해 있어 공의 이동을 따라갑니다. 레벨 안에는 경사로와, 공으로 넘어뜨릴 수도 넘어뜨리지 못할수도 있는 물리 큐브가 여럿 있습니다.")
+
+LocalizedDisplayNames=(Language="ja",Text="ローリング")
+LocalizedDescriptions=(Language="ja",Text="ローリングテンプレートには物理シミュレーションを行っているボールが含まれ、キーボード、コントローラー、タッチデバイスでのバーチャルジョイスティックで好きな方向に転がすことができます。カメラはボールの上方の後方に位置し、ボールの移動に追従します。レベル内には坂や、ボールで倒すことのできる物理キューブと、物理ではないキューブが含まれています。")
+
+LocalizedDisplayNames=(Language="zh-CN",Text="滚球类游戏")
+LocalizedDescriptions=(Language="zh-CN",Text="滚球类游戏模板中有一个基于物理的滚球，您可以通过键盘、手柄或触摸式设备的虚拟摇杆让其向任何方向滚动。摄像机位于球体的后上方，并会跟随球体移动。在关卡内还有一些斜坡，以及一些可以用球体撞开的物理立方体和无法撞开的立方体。")
 
 ClassTypes=GameMode, Pawn, StaticMeshComponent, SpringArmComponent, CameraComponent, PrimitiveComponent
 AssetTypes=StaticMesh, Material, Physical Material
diff --git a/Templates/TP_Rolling/manifest.json b/Templates/TP_Rolling/manifest.json
--- a/Templates/TP_Rolling/manifest.json
+++ b/Templates/TP_Rolling/manifest.json
@@ -10,9 +10,13 @@
 			"Language": "ko",
 			"Text": "굴리기"
 		},
+		{
+			"Language": "zh-CN",
+			"Text": "滚球类游戏"
+		},
 		{
 			"Language": "ja",
-			"Text": "転がり"
+			"Text": "ローリング"
 		}
 	],
 	"Description":
@@ -23,11 +27,15 @@
 		},
 		{
 			"Language": "ko",
-			"Text": "굴리기 팩은 키보드, 콘트롤러, 터치 디바이스의 가상 조이스틱을 사용하여 어느 방향이든 굴릴 수 있는 물리 기반 공을 특징으로 합니다. 카메라는 공의 상단 뒤쪽에 위치해 있어 공의 이동을 따라갑니다. 레벨 안에는 경사로와, 공으로 넘어뜨릴 수도 넘어뜨리지 못할수도 있는 물리 큐브가 여럿 있습니다."
+			"Text": "굴리기 템플릿은 키보드, 컨트롤러, 터치 디바이스의 가상 조이스틱을 사용하여 어느 방향이든 굴릴 수 있는 물리 기반 공을 특징으로 합니다. 카메라는 공의 상단 뒤쪽에 위치해 있어 공의 이동을 따라갑니다. 레벨 안에는 경사로와, 공으로 넘어뜨릴 수도 넘어뜨리지 못할수도 있는 물리 큐브가 여럿 있습니다."
+		},
+		{
+			"Language": "zh-CN",
+			"Text": "滚球类游戏模板中有一个基于物理的滚球，您可以通过键盘、手柄或触摸式设备的虚拟摇杆让其向任何方向滚动。摄像机位于球体的后上方，并会跟随球体移动。在关卡内还有一些斜坡，以及一些可以用球体撞开的物理立方体和无法撞开的立方体。"
 		},
 		{
 			"Language": "ja",
-			"Text": "転がりテンプレートでは、物理ベースで転がるボールを操作するゲームを作ることができます。ボールは、キーボード・コントローラー・バーチャルジョイスティック・タッチパネルで操作することができます。カメラはボールの上空少し後ろからボールの動きを捉えています。テンプレートに含まれるレベルには幾つかの坂になった障害物が設置されています。レベル内に配置されている幾つかの箱はボールで押して動かすことができますが、できないものも含まれています。"
+			"Text": "ローリングテンプレートには物理シミュレーションを行っているボールが含まれ、キーボード、コントローラー、タッチデバイスでのバーチャルジョイスティックで好きな方向に転がすことができます。カメラはボールの上方の後方に位置し、ボールの移動に追従します。レベル内には坂や、ボールで倒すことのできる物理キューブと、物理ではないキューブが含まれています。"
 		}
 	],
 	"AssetTypes":
@@ -35,14 +43,6 @@
 		{
 			"Language": "en",
 			"Text": "StaticMesh, Material, Physical Material"
-		},
-		{
-			"Language": "ko",
-			"Text": "스태틱 메시, 머티리얼, 피지컬 머티리얼"
-		},
-		{
-			"Language": "ja",
-			"Text": "スタティックメッシュ・マテリアル・フィジカルマテリアル"
 		}
 	],
 	"SearchTags":
@@ -50,10 +50,6 @@
 		{
 			"Language": "en",
 			"Text": "Rolling,Marble"
-		},
-		{
-			"Language": "ko",
-			"Text": "굴리기,마블"
 		}
 	],	
 	"AdditionalFeaturePacks":
diff --git a/Templates/TP_RollingBP/Config/TemplateDefs.ini b/Templates/TP_RollingBP/Config/TemplateDefs.ini
--- a/Templates/TP_RollingBP/Config/TemplateDefs.ini
+++ b/Templates/TP_RollingBP/Config/TemplateDefs.ini
@@ -3,7 +3,13 @@ LocalizedDisplayNames=(Language="en",Text="Rolling")
 LocalizedDescriptions=(Language="en",Text="The rolling template features a physics-based rolling ball which can be rolled in any direction using either a keyboard, controller or virtual joystick on touch devices. The camera is positioned above and behind the ball and follows the movement of the ball. Within the level are some ramps, several physics cubes in the level that can be knocked over with the ball as well as some cubes which cannot.")
 
 LocalizedDisplayNames=(Language="ko",Text="굴리기")
-LocalizedDescriptions=(Language="ko",Text="굴리기 템플릿은 키보드, 콘트롤러, 터치 디바이스의 가상 조이스틱을 사용하여 어느 방향이든 굴릴 수 있는 물리 기반 공을 특징으로 합니다. 카메라는 공의 상단 뒤쪽에 위치해 있어 공의 이동을 따라갑니다. 레벨 안에는 경사로와, 공으로 넘어뜨릴 수도 넘어뜨리지 못할수도 있는 물리 큐브가 여럿 있습니다.")
+LocalizedDescriptions=(Language="ko",Text="굴리기 템플릿은 키보드, 컨트롤러, 터치 디바이스의 가상 조이스틱을 사용하여 어느 방향이든 굴릴 수 있는 물리 기반 공을 특징으로 합니다. 카메라는 공의 상단 뒤쪽에 위치해 있어 공의 이동을 따라갑니다. 레벨 안에는 경사로와, 공으로 넘어뜨릴 수도 넘어뜨리지 못할수도 있는 물리 큐브가 여럿 있습니다.")
+
+LocalizedDisplayNames=(Language="ja",Text="ローリング")
+LocalizedDescriptions=(Language="ja",Text="ローリングテンプレートには物理シミュレーションを行っているボールが含まれ、キーボード、コントローラー、タッチデバイスでのバーチャルジョイスティックで好きな方向に転がすことができます。カメラはボールの上方の後方に位置し、ボールの移動に追従します。レベル内には坂や、ボールで倒すことのできる物理キューブと、物理ではないキューブが含まれています。")
+
+LocalizedDisplayNames=(Language="zh-CN",Text="滚球类游戏")
+LocalizedDescriptions=(Language="zh-CN",Text="滚球类游戏模板中有一个基于物理的滚球，您可以通过键盘、手柄或触摸式设备的虚拟摇杆让其向任何方向滚动。摄像机位于球体的后上方，并会跟随球体移动。在关卡内还有一些斜坡，以及一些可以用球体撞开的物理立方体和无法撞开的立方体。")
 
 ClassTypes=GameMode, Pawn, StaticMeshComponent, SpringArmComponent, CameraComponent, PrimitiveComponent
 AssetTypes=StaticMesh, Material, Physical Material
diff --git a/Templates/TP_RollingBP/manifest.json b/Templates/TP_RollingBP/manifest.json
--- a/Templates/TP_RollingBP/manifest.json
+++ b/Templates/TP_RollingBP/manifest.json
@@ -10,9 +10,13 @@
 			"Language": "ko",
 			"Text": "굴리기"
 		},
+		{
+			"Language": "zh-CN",
+			"Text": "滚球类游戏"
+		},
 		{
 			"Language": "ja",
-			"Text": "転がり"
+			"Text": "ローリング"
 		}
 	],
 	"Description":
@@ -23,11 +27,15 @@
 		},
 		{
 			"Language": "ko",
-			"Text": "굴리기 팩은 키보드, 콘트롤러, 터치 디바이스의 가상 조이스틱을 사용하여 어느 방향이든 굴릴 수 있는 물리 기반 공을 특징으로 합니다. 카메라는 공의 상단 뒤쪽에 위치해 있어 공의 이동을 따라갑니다. 레벨 안에는 경사로와, 공으로 넘어뜨릴 수도 넘어뜨리지 못할수도 있는 물리 큐브가 여럿 있습니다."
+			"Text": "굴리기 템플릿은 키보드, 컨트롤러, 터치 디바이스의 가상 조이스틱을 사용하여 어느 방향이든 굴릴 수 있는 물리 기반 공을 특징으로 합니다. 카메라는 공의 상단 뒤쪽에 위치해 있어 공의 이동을 따라갑니다. 레벨 안에는 경사로와, 공으로 넘어뜨릴 수도 넘어뜨리지 못할수도 있는 물리 큐브가 여럿 있습니다."
+		},
+		{
+			"Language": "zh-CN",
+			"Text": "滚球类游戏模板中有一个基于物理的滚球，您可以通过键盘、手柄或触摸式设备的虚拟摇杆让其向任何方向滚动。摄像机位于球体的后上方，并会跟随球体移动。在关卡内还有一些斜坡，以及一些可以用球体撞开的物理立方体和无法撞开的立方体。"
 		},
 		{
 			"Language": "ja",
-			"Text": "転がりテンプレートでは、物理ベースで転がるボールを操作するゲームを作ることができます。ボールは、キーボード・コントローラー・バーチャルジョイスティック・タッチパネルで操作することができます。カメラはボールの上空少し後ろからボールの動きを捉えています。テンプレートに含まれるレベルには幾つかの坂になった障害物が設置されています。レベル内に配置されている幾つかの箱はボールで押して動かすことができますが、できないものも含まれています。"
+			"Text": "ローリングテンプレートには物理シミュレーションを行っているボールが含まれ、キーボード、コントローラー、タッチデバイスでのバーチャルジョイスティックで好きな方向に転がすことができます。カメラはボールの上方の後方に位置し、ボールの移動に追従します。レベル内には坂や、ボールで倒すことのできる物理キューブと、物理ではないキューブが含まれています。"
 		}
 	],
 	"AdditionalFeaturePacks":
diff --git a/Templates/TP_SideScroller/Config/TemplateDefs.ini b/Templates/TP_SideScroller/Config/TemplateDefs.ini
--- a/Templates/TP_SideScroller/Config/TemplateDefs.ini
+++ b/Templates/TP_SideScroller/Config/TemplateDefs.ini
@@ -3,7 +3,13 @@ LocalizedDisplayNames=(Language="en",Text="Side Scroller")
 LocalizedDescriptions=(Language="en",Text="This side scroller template features an animated skeletal mesh character with a camera positioned at the character's side. Movement, which can be controlled with either keyboard, controller or virtual joystick on a touch device is limited to two dimensions: forward, backward and jumping. A number of ledges are featured in the level and the character can jump between them.")
 
 LocalizedDisplayNames=(Language="ko",Text="횡스크롤")
-LocalizedDescriptions=(Language="ko",Text="횡스크롤 템플릿에는 측면에 카메라가 위치한 애니메이션 적용 스켈레탈 메시 캐릭터를 특징으로 합니다. 키보드, 콘트롤러, 터치 디바이스의 가상 조이스틱으로 제어할 수 있는 이동은 두 축으로 제한되어 있어, 앞, 뒤, 점프만 가능합니다. 레벨에는 캐릭터가 점프해 다닐 수 있는 턱이 다수 있습니다.")
+LocalizedDescriptions=(Language="ko",Text="횡스크롤 템플릿에는 카메라가 측면에 위치한 애니메이션 적용 스켈레탈 메시 캐릭터를 특징으로 합니다. 키보드, 컨트롤러, 터치 디바이스의 가상 조이스틱으로 제어할 수 있는 이동은 두 축으로 제한되어 있어, 앞, 뒤, 점프만 가능합니다. 레벨에는 캐릭터가 점프해 다닐 수 있는 턱이 다수 있습니다.")
+
+LocalizedDisplayNames=(Language="ja",Text="サイドスクロール")
+LocalizedDescriptions=(Language="ja",Text="サイドスクロールテンプレートにはアニメーションを使用するスケルタルメッシュのキャラクタが存在し、カメラはキャラクタの側面に存在します。移動はキーボード、コントローラー、タッチデバイスでのバーチャルジョイスティックで操作可能で、二次元に制限されています。前進、後退とジャンプだけです。レベルには足場があり、それらの間をジャンプすることができます。")
+
+LocalizedDisplayNames=(Language="zh-CN",Text="横版过关游戏")
+LocalizedDescriptions=(Language="zh-CN",Text="该横版过关类游戏模板拥有一个可活动的骨架网格体角色，摄像机位于角色的一侧。您可以通过键盘、手柄或触摸式设备上的虚拟摇杆来控制角色移动，但仅限于两个方向：前进、后退及跳跃。关卡中有许多岩架，角色可以在它们之间跳跃。")
 
 ClassTypes=CameraComponent, SpringArmComponent, PaperFlipbook, PaperCharacter, GameMode
 AssetTypes=Animation Sequence, Animation Blueprint, Skeletal Mesh, Blend Space 1D, Skeleton
diff --git a/Templates/TP_SideScroller/manifest.json b/Templates/TP_SideScroller/manifest.json
--- a/Templates/TP_SideScroller/manifest.json
+++ b/Templates/TP_SideScroller/manifest.json
@@ -10,6 +10,10 @@
 			"Language": "ko",
 			"Text": "횡스크롤"
 		},
+		{
+			"Language": "zh-CN",
+			"Text": "横版过关游戏"
+		},
 		{
 			"Language": "ja",
 			"Text": "横スクロール"
@@ -23,11 +27,15 @@
 		},
 		{
 			"Language": "ko",
-			"Text": "횡스크롤 팩에는 측면에 카메라가 위치한 애니메이션 적용 스켈레탈 메시 캐릭터를 특징으로 합니다. 키보드, 콘트롤러, 터치 디바이스의 가상 조이스틱으로 제어할 수 있는 이동은 두 축으로 제한되어 있어, 앞, 뒤, 점프만 가능합니다. 레벨에는 캐릭터가 점프해 다닐 수 있는 턱이 다수 있습니다."
+			"Text": "횡스크롤 템플릿에는 카메라가 측면에 위치한 애니메이션 적용 스켈레탈 메시 캐릭터를 특징으로 합니다. 키보드, 컨트롤러, 터치 디바이스의 가상 조이스틱으로 제어할 수 있는 이동은 두 축으로 제한되어 있어, 앞, 뒤, 점프만 가능합니다. 레벨에는 캐릭터가 점프해 다닐 수 있는 턱이 다수 있습니다."
+		},
+		{
+			"Language": "zh-CN",
+			"Text": "该横版过关类游戏模板拥有一个可活动的骨架网格体角色，摄像机位于角色的一侧。您可以通过键盘、手柄或触摸式设备上的虚拟摇杆来控制角色移动，但仅限于两个方向：前进、后退及跳跃。关卡中有许多岩架，角色可以在它们之间跳跃。"
 		},
 		{
 			"Language": "ja",
-			"Text": "The side scroller pack features an animated skeletal mesh character with a camera positioned at the character's side. Movement, which can be controlled with either keyboard, controller or virtual joystick on a touch device is limited to two dimensions: forward, backward and jumping. A number of ledges are featured in the level and the character can jump between them."
+			"Text": "サイドスクロールテンプレートにはアニメーションを使用するスケルタルメッシュのキャラクタが存在し、カメラはキャラクタの側面に存在します。移動はキーボード、コントローラー、タッチデバイスでのバーチャルジョイスティックで操作可能で、二次元に制限されています。前進、後退とジャンプだけです。レベルには足場があり、それらの間をジャンプすることができます。"
 		}
 
 	],
@@ -36,25 +44,13 @@
 		{
 			"Language": "en",
 			"Text": "Animation Sequence, Animation Blueprint, Skeletal Mesh, Blend Space 1D, Skeleton"
-		},
-		{
-			"Language": "ko",
-			"Text": "애니메이션 시퀀스, 애니메이션 블루프린트, 스켈레탈 메시, 블렌드 스페이스 1D, 스켈레톤"
-		},
-		{
-			"Language": "ja",
-			"Text": "アニメーションシーケンス・アニメーションブループリント・スケルタルメッシュ・ブレンドスペース１D・スケルトン"
 		}
-	],	
+	],
 	"SearchTags":
 	[
 		{
 			"Language": "en",
 			"Text": "Side,Scroller,Platform"
-		},
-		{
-			"Language": "ko",
-			"Text": "횡,스크롤,플랫폼"
 		}
 	],	
 	"ClassTypes": "CameraComponent, SpringArmComponent, PaperFlipbook, PaperCharacter, GameMode",
diff --git a/Templates/TP_SideScrollerBP/Config/TemplateDefs.ini b/Templates/TP_SideScrollerBP/Config/TemplateDefs.ini
--- a/Templates/TP_SideScrollerBP/Config/TemplateDefs.ini
+++ b/Templates/TP_SideScrollerBP/Config/TemplateDefs.ini
@@ -3,7 +3,13 @@ LocalizedDisplayNames=(Language="en",Text="Side Scroller")
 LocalizedDescriptions=(Language="en",Text="This side scroller template features an animated skeletal mesh character with a camera positioned at the character's side. Movement, which can be controlled with either keyboard, controller or virtual joystick on a touch device is limited to two dimensions: forward, backward and jumping. A number of ledges are featured in the level and the character can jump between them.")
 
 LocalizedDisplayNames=(Language="ko",Text="횡스크롤")
-LocalizedDescriptions=(Language="ko",Text="횡스크롤 템플릿에는 측면에 카메라가 위치한 애니메이션 적용 스켈레탈 메시 캐릭터를 특징으로 합니다. 키보드, 콘트롤러, 터치 디바이스의 가상 조이스틱으로 제어할 수 있는 이동은 두 축으로 제한되어 있어, 앞, 뒤, 점프만 가능합니다. 레벨에는 캐릭터가 점프해 다닐 수 있는 턱이 다수 있습니다.")
+LocalizedDescriptions=(Language="ko",Text="횡스크롤 템플릿에는 카메라가 측면에 위치한 애니메이션 적용 스켈레탈 메시 캐릭터를 특징으로 합니다. 키보드, 컨트롤러, 터치 디바이스의 가상 조이스틱으로 제어할 수 있는 이동은 두 축으로 제한되어 있어, 앞, 뒤, 점프만 가능합니다. 레벨에는 캐릭터가 점프해 다닐 수 있는 턱이 다수 있습니다.")
+
+LocalizedDisplayNames=(Language="ja",Text="サイドスクロール")
+LocalizedDescriptions=(Language="ja",Text="サイドスクロールテンプレートにはアニメーションを使用するスケルタルメッシュのキャラクタが存在し、カメラはキャラクタの側面に存在します。移動はキーボード、コントローラー、タッチデバイスでのバーチャルジョイスティックで操作可能で、二次元に制限されています。前進、後退とジャンプだけです。レベルには足場があり、それらの間をジャンプすることができます。")
+
+LocalizedDisplayNames=(Language="zh-CN",Text="横版过关游戏")
+LocalizedDescriptions=(Language="zh-CN",Text="该横版过关类游戏模板拥有一个可活动的骨架网格体角色，摄像机位于角色的一侧。您可以通过键盘、手柄或触摸式设备上的虚拟摇杆来控制角色移动，但仅限于两个方向：前进、后退及跳跃。关卡中有许多岩架，角色可以在它们之间跳跃。")
 
 ClassTypes=CameraComponent, SpringArmComponent, PaperFlipbook, PaperCharacter, GameMode
 AssetTypes=Animation Sequence, Animation Blueprint, Skeletal Mesh, Blend Space 1D, Skeleton
diff --git a/Templates/TP_SideScrollerBP/manifest.json b/Templates/TP_SideScrollerBP/manifest.json
--- a/Templates/TP_SideScrollerBP/manifest.json
+++ b/Templates/TP_SideScrollerBP/manifest.json
@@ -10,6 +10,10 @@
 			"Language": "ko",
 			"Text": "횡스크롤"
 		},
+		{
+			"Language": "zh-CN",
+			"Text": "横版过关游戏"
+		},
 		{
 			"Language": "ja",
 			"Text": "横スクロール"
@@ -23,11 +27,15 @@
 		},
 		{
 			"Language": "ko",
-			"Text": "횡스크롤 팩에는 측면에 카메라가 위치한 애니메이션 적용 스켈레탈 메시 캐릭터를 특징으로 합니다. 키보드, 콘트롤러, 터치 디바이스의 가상 조이스틱으로 제어할 수 있는 이동은 두 축으로 제한되어 있어, 앞, 뒤, 점프만 가능합니다. 레벨에는 캐릭터가 점프해 다닐 수 있는 턱이 다수 있습니다."
+			"Text": "횡스크롤 템플릿에는 카메라가 측면에 위치한 애니메이션 적용 스켈레탈 메시 캐릭터를 특징으로 합니다. 키보드, 컨트롤러, 터치 디바이스의 가상 조이스틱으로 제어할 수 있는 이동은 두 축으로 제한되어 있어, 앞, 뒤, 점프만 가능합니다. 레벨에는 캐릭터가 점프해 다닐 수 있는 턱이 다수 있습니다."
+		},
+		{
+			"Language": "zh-CN",
+			"Text": "该横版过关类游戏模板拥有一个可活动的骨架网格体角色，摄像机位于角色的一侧。您可以通过键盘、手柄或触摸式设备上的虚拟摇杆来控制角色移动，但仅限于两个方向：前进、后退及跳跃。关卡中有许多岩架，角色可以在它们之间跳跃。"
 		},
 		{
 			"Language": "ja",
-			"Text": "The side scroller pack features an animated skeletal mesh character with a camera positioned at the character's side. Movement, which can be controlled with either keyboard, controller or virtual joystick on a touch device is limited to two dimensions: forward, backward and jumping. A number of ledges are featured in the level and the character can jump between them."
+			"Text": "サイドスクロールテンプレートにはアニメーションを使用するスケルタルメッシュのキャラクタが存在し、カメラはキャラクタの側面に存在します。移動はキーボード、コントローラー、タッチデバイスでのバーチャルジョイスティックで操作可能で、二次元に制限されています。前進、後退とジャンプだけです。レベルには足場があり、それらの間をジャンプすることができます。"
 		}
 
 	],
@@ -36,14 +44,6 @@
 		{
 			"Language": "en",
 			"Text": "Animation Sequence, Animation Blueprint, Skeletal Mesh, Blend Space 1D, Skeleton"
-		},
-		{
-			"Language": "ko",
-			"Text": "애니메이션 시퀀스, 애니메이션 블루프린트, 스켈레탈 메시, 블렌드 스페이스 1D, 스켈레톤"
-		},
-		{
-			"Language": "ja",
-			"Text": "アニメーションシーケンス・アニメーションブループリント・スケルタルメッシュ・ブレンドスペース１D・スケルトン"
 		}
 	],
 	"SearchTags":
@@ -51,10 +51,6 @@
 		{
 			"Language": "en",
 			"Text": "Side,Scroller,Platform"
-		},
-		{
-			"Language": "ko",
-			"Text": "횡,스크롤,플랫폼"
 		}
 	],	
 	"ClassTypes": "CameraComponent, SpringArmComponent, PaperFlipbook, PaperCharacter, GameMode",
diff --git a/Templates/TP_ThirdPerson/Config/TemplateDefs.ini b/Templates/TP_ThirdPerson/Config/TemplateDefs.ini
--- a/Templates/TP_ThirdPerson/Config/TemplateDefs.ini
+++ b/Templates/TP_ThirdPerson/Config/TemplateDefs.ini
@@ -5,6 +5,12 @@ LocalizedDescriptions=(Language="en",Text="The third person template features a
 LocalizedDisplayNames=(Language="ko",Text="삼인칭")
 LocalizedDescriptions=(Language="ko",Text="삼인칭 템플릿은 카메라가 머리위 약간 뒤쪽에 위치한 플레이가능 캐릭터를 특징으로 합니다. 마우스, 키보드, 터치 디바이스의 가상 조이스틱으로 캐릭터가 이동함에 따라 카메라는 캐릭터를 따라다니며 어깨위 시점으로 봅니다. 주인공 캐릭터를 강조하는 시점으로, 액션이나 어드벤처 게임에 많이 쓰입니다. 레벨에는 플레이어가 점프할 수 있는 오브젝트가 여럿 있습니다.")
 
+LocalizedDisplayNames=(Language="ja",Text="サードパーソン")
+LocalizedDescriptions=(Language="ja",Text="サードパーソンテンプレートにはプレイアブルキャラクタが存在し、カメラがキャラクタの背後の少し上方にあります。キーボード、コントローラー、タッチデバイスでのバーチャルジョイスティックでキャラクタが移動すると、カメラは肩越しのビューで追従します。この視点はメインキャラクタを強調するのもでアクションアドベンチャーゲームで人気があるものです。レベルにはプレイヤーがジャンプして乗ることができるオブジェクトが置かれています。")
+
+LocalizedDisplayNames=(Language="zh-CN",Text="第三人称游戏")
+LocalizedDescriptions=(Language="zh-CN",Text="第三人称游戏模板拥有一个可游玩的角色，摄像机位于角色身后偏上的位置。您可以通过鼠标、键盘、手柄或触摸式设备上控制角色移动，摄像机会跟随角色，通过过肩式视角观察世界。此视角突出主角，在动作及冒险类游戏中非常流行。关卡中还包含一些对象，可以让角色跳到上面。")
+
 ClassTypes=GameMode, Character, SpringArmComponent, CameraComponent, InputComponent
 AssetTypes=Animation Sequence, Animation Blueprint, Skeleton, Blend Space 1D, Skeletal Mesh
 
diff --git a/Templates/TP_ThirdPerson/manifest.json b/Templates/TP_ThirdPerson/manifest.json
--- a/Templates/TP_ThirdPerson/manifest.json
+++ b/Templates/TP_ThirdPerson/manifest.json
@@ -10,9 +10,13 @@
 			"Language": "ko",
 			"Text": "삼인칭"
 		},
+		{
+			"Language": "zh-CN",
+			"Text": "第三人称游戏"
+		},
 		{
 			"Language": "ja",
-			"Text": "アクション"
+			"Text": "サードパーソン"
 		}
 	],
 	"Description":
@@ -23,11 +27,15 @@
 		},
 		{
 			"Language": "ko",
-			"Text": "삼인칭 팩은 카메라가 머리위 약간 뒤쪽에 위치한 플레이가능 캐릭터를 특징으로 합니다. 마우스, 키보드, 터치 디바이스의 가상 조이스틱으로 캐릭터가 이동함에 따라 카메라는 캐릭터를 따라다니며 어깨위 시점으로 봅니다. 주인공 캐릭터를 강조하는 시점으로, 액션이나 어드벤처 게임에 많이 쓰입니다. 레벨에는 플레이어가 점프할 수 있는 오브젝트가 여럿 있습니다."
+			"Text": "삼인칭 템플릿은 카메라가 머리위 약간 뒤쪽에 위치한 플레이가능 캐릭터를 특징으로 합니다. 마우스, 키보드, 터치 디바이스의 가상 조이스틱으로 캐릭터가 이동함에 따라 카메라는 캐릭터를 따라다니며 어깨위 시점으로 봅니다. 주인공 캐릭터를 강조하는 시점으로, 액션이나 어드벤처 게임에 많이 쓰입니다. 레벨에는 플레이어가 점프할 수 있는 오브젝트가 여럿 있습니다."
+		},
+		{
+			"Language": "zh-CN",
+			"Text": "第三人称游戏模板拥有一个可游玩的角色，摄像机位于角色身后偏上的位置。您可以通过鼠标、键盘、手柄或触摸式设备上控制角色移动，摄像机会跟随角色，通过过肩式视角观察世界。此视角突出主角，在动作及冒险类游戏中非常流行。关卡中还包含一些对象，可以让角色跳到上面。"
 		},
 		{
 			"Language": "ja",
-			"Text": "アクションテンプレートではキャラクターを三人称視点で見下ろして操作します。キャラクターの操作にはマウス・キーボード・コントローラー・バーチャルジョイスティック・タッチパネルを使うことができます。カメラはキャラクターの上空後方から常に追従します。アクションゲームやアドベンチャーゲームに適しています。 レベルに配置されたブロックにはジャンプして登ることができます。"
+			"Text": "サードパーソンテンプレートにはプレイアブルキャラクタが存在し、カメラがキャラクタの背後の少し上方にあります。キーボード、コントローラー、タッチデバイスでのバーチャルジョイスティックでキャラクタが移動すると、カメラは肩越しのビューで追従します。この視点はメインキャラクタを強調するのもでアクションアドベンチャーゲームで人気があるものです。レベルにはプレイヤーがジャンプして乗ることができるオブジェクトが置かれています。"
 		}
 	],
 	"AssetTypes":
@@ -35,14 +43,6 @@
 		{
 			"Language": "en",
 			"Text": "Animation Sequence, Animation Blueprint, Skeleton, Blend Space 1D, Skeletal Mesh"
-		},
-		{
-			"Language": "ko",
-			"Text": "애니메이션 시퀀스, 애니메이션 블루프린트, 스켈레톤, 블렌드 스페이스 1D, 스켈레탈 메시"
-		},
-		{
-			"Language": "ja",
-			"Text": "アニメーションシーケンス・アニメーションブループリント・スケルトン・ブレンドスペース１D・スケルタルメッシュ"
 		}
 	],	
 	"SearchTags":
@@ -50,10 +50,6 @@
 		{
 			"Language": "en",
 			"Text": "Action, Adventure, Third"
-		},
-		{
-			"Language": "ko",
-			"Text": "액션, 어드벤처, 삼인칭"
 		}
 	],
 	"ClassTypes": "GameMode, Character, SpringArmComponent, CameraComponent, InputComponent",
diff --git a/Templates/TP_ThirdPersonBP/Config/TemplateDefs.ini b/Templates/TP_ThirdPersonBP/Config/TemplateDefs.ini
--- a/Templates/TP_ThirdPersonBP/Config/TemplateDefs.ini
+++ b/Templates/TP_ThirdPersonBP/Config/TemplateDefs.ini
@@ -5,6 +5,12 @@ LocalizedDescriptions=(Language="en",Text="The third person template features a
 LocalizedDisplayNames=(Language="ko",Text="삼인칭")
 LocalizedDescriptions=(Language="ko",Text="삼인칭 템플릿은 카메라가 머리위 약간 뒤쪽에 위치한 플레이가능 캐릭터를 특징으로 합니다. 마우스, 키보드, 터치 디바이스의 가상 조이스틱으로 캐릭터가 이동함에 따라 카메라는 캐릭터를 따라다니며 어깨위 시점으로 봅니다. 주인공 캐릭터를 강조하는 시점으로, 액션이나 어드벤처 게임에 많이 쓰입니다. 레벨에는 플레이어가 점프할 수 있는 오브젝트가 여럿 있습니다.")
 
+LocalizedDisplayNames=(Language="ja",Text="サードパーソン")
+LocalizedDescriptions=(Language="ja",Text="サードパーソンテンプレートにはプレイアブルキャラクタが存在し、カメラがキャラクタの背後の少し上方にあります。キーボード、コントローラー、タッチデバイスでのバーチャルジョイスティックでキャラクタが移動すると、カメラは肩越しのビューで追従します。この視点はメインキャラクタを強調するのもでアクションアドベンチャーゲームで人気があるものです。レベルにはプレイヤーがジャンプして乗ることができるオブジェクトが置かれています。")
+
+LocalizedDisplayNames=(Language="zh-CN",Text="第三人称游戏")
+LocalizedDescriptions=(Language="zh-CN",Text="第三人称游戏模板拥有一个可游玩的角色，摄像机位于角色身后偏上的位置。您可以通过鼠标、键盘、手柄或触摸式设备上控制角色移动，摄像机会跟随角色，通过过肩式视角观察世界。此视角突出主角，在动作及冒险类游戏中非常流行。关卡中还包含一些对象，可以让角色跳到上面。")
+
 ClassTypes=GameMode, Character, SpringArmComponent, CameraComponent, InputComponent
 AssetTypes=Animation Sequence, Animation Blueprint, Skeleton, Blend Space 1D, Skeletal Mesh
 
diff --git a/Templates/TP_ThirdPersonBP/manifest.json b/Templates/TP_ThirdPersonBP/manifest.json
--- a/Templates/TP_ThirdPersonBP/manifest.json
+++ b/Templates/TP_ThirdPersonBP/manifest.json
@@ -10,9 +10,13 @@
 			"Language": "ko",
 			"Text": "삼인칭"
 		},
+		{
+			"Language": "zh-CN",
+			"Text": "第三人称游戏"
+		},
 		{
 			"Language": "ja",
-			"Text": "三人称視点"
+			"Text": "サードパーソン"
 		}
 	],
 	"Description":
@@ -23,11 +27,15 @@
 		},
 		{
 			"Language": "ko",
-			"Text": "삼인칭 팩은 카메라가 머리위 약간 뒤쪽에 위치한 플레이가능 캐릭터를 특징으로 합니다. 마우스, 키보드, 터치 디바이스의 가상 조이스틱으로 캐릭터가 이동함에 따라 카메라는 캐릭터를 따라다니며 어깨위 시점으로 봅니다. 주인공 캐릭터를 강조하는 시점으로, 액션이나 어드벤처 게임에 많이 쓰입니다. 레벨에는 플레이어가 점프할 수 있는 오브젝트가 여럿 있습니다."
+			"Text": "삼인칭 템플릿은 카메라가 머리위 약간 뒤쪽에 위치한 플레이가능 캐릭터를 특징으로 합니다. 마우스, 키보드, 터치 디바이스의 가상 조이스틱으로 캐릭터가 이동함에 따라 카메라는 캐릭터를 따라다니며 어깨위 시점으로 봅니다. 주인공 캐릭터를 강조하는 시점으로, 액션이나 어드벤처 게임에 많이 쓰입니다. 레벨에는 플레이어가 점프할 수 있는 오브젝트가 여럿 있습니다."
+		},
+		{
+			"Language": "zh-CN",
+			"Text": "第三人称游戏模板拥有一个可游玩的角色，摄像机位于角色身后偏上的位置。您可以通过鼠标、键盘、手柄或触摸式设备上控制角色移动，摄像机会跟随角色，通过过肩式视角观察世界。此视角突出主角，在动作及冒险类游戏中非常流行。关卡中还包含一些对象，可以让角色跳到上面。"
 		},
 		{
 			"Language": "ja",
-			"Text": "三人称視点テンプレートではキャラクターを三人称視点で見下ろして操作します。キャラクターの操作にはマウス・キーボード・コントローラー・バーチャルジョイスティック・タッチパネルを使うことができます。カメラはキャラクターの上空後方から常に追従します。アクションゲームやアドベンチャーゲームに適しています。 レベルに配置されたブロックにはジャンプして登ることができます。"
+			"Text": "サードパーソンテンプレートにはプレイアブルキャラクタが存在し、カメラがキャラクタの背後の少し上方にあります。キーボード、コントローラー、タッチデバイスでのバーチャルジョイスティックでキャラクタが移動すると、カメラは肩越しのビューで追従します。この視点はメインキャラクタを強調するのもでアクションアドベンチャーゲームで人気があるものです。レベルにはプレイヤーがジャンプして乗ることができるオブジェクトが置かれています。"
 		}
 	],
 	"AssetTypes":
@@ -35,25 +43,13 @@
 		{
 			"Language": "en",
 			"Text": "Animation Sequence, Animation Blueprint, Skeleton, Blend Space 1D, Skeletal Mesh"
-		},
-		{
-			"Language": "ko",
-			"Text": "애니메이션 시퀀스, 애니메이션 블루프린트, 스켈레톤, 블렌드 스페이스 1D, 스켈레탈 메시"
-		},
-		{
-			"Language": "ja",
-			"Text": "アニメーションシーケンス・アニメーションブループリント・スケルトン・ブレンドスペース１D・スケルタルメッシュ"
 		}
-	],
+	],	
 	"SearchTags":
 	[
 		{
 			"Language": "en",
 			"Text": "Action, Adventure, Third"
-		},
-		{
-			"Language": "ko",
-			"Text": "액션, 어드벤처, 삼인칭"
 		}
 	],
 	"ClassTypes": "GameMode, Character, SpringArmComponent, CameraComponent, InputComponent",
diff --git a/Templates/TP_TopDown/Config/TemplateDefs.ini b/Templates/TP_TopDown/Config/TemplateDefs.ini
--- a/Templates/TP_TopDown/Config/TemplateDefs.ini
+++ b/Templates/TP_TopDown/Config/TemplateDefs.ini
@@ -3,7 +3,13 @@ LocalizedDisplayNames=(Language="en",Text="Top Down")
 LocalizedDescriptions=(Language="en",Text="The top down template features a playable character where the camera is positioned behind and above the avatar at greater distance than in third person perspective. The player is controlled by using the mouse or touchscreen to click on the required destination and uses the navigation system to facilitate the character movement. This perspective is often used in action role playing games.")
 
 LocalizedDisplayNames=(Language="ko",Text="내려보기")
-LocalizedDescriptions=(Language="ko",Text="내려보기 템플릿은 뒤쪽으로 삼인칭 시점보다 먼 거리에 있는 아바타 위에 카메라가 위치한 플레이가능 캐릭터를 특징으로 합니다. 플레이어 제어는 마우스나 터치스크린으로 원하는 목적지를 클릭하는 식으로 이루어지며, 원활한 캐릭터 이동을 위해 내비게이션 시스템을 사용합니다. 이 시점은 액션 롤플레잉 게임에 자주 사용됩니다.")
+LocalizedDescriptions=(Language="ko",Text="내려보기 템플릿은 플레이어 캐릭터 위 삼인칭 시점보다 멀리 카메라가 위치합니다. 플레이어 제어는 마우스나 터치스크린으로 원하는 목적지를 클릭하는 식으로 이루어지며, 원활한 캐릭터 이동을 위해 내비게이션 시스템을 사용합니다. 이 시점은 액션 롤플레잉 게임에 자주 쓰입니다.")
+
+LocalizedDisplayNames=(Language="ja",Text="トップダウン")
+LocalizedDescriptions=(Language="ja",Text="トップダウンテンプレートにはプレイアブルキャラクタが存在し、カメラはキャラクタの背後、サードパーソンテンプレートより遠い上方にあります。マウスやタッチスクリーンで行き先をクリックすることでキャラクタを操作することができます。キャラクタはナビゲーションシステムを利用して移動します。この視点はアクションロールプレイングゲームでよく使用されています。")
+
+LocalizedDisplayNames=(Language="zh-CN",Text="俯视角游戏")
+LocalizedDescriptions=(Language="zh-CN",Text="俯视角游戏模板拥有一个可游玩角色，摄像机位于角色的后上方，但位置比第三人称视角要高得多。玩家可以使用鼠标或触摸屏点击目的地，或使用导航系统来辅助角色移动。此视角常用于动作角色扮演类游戏。")
 
 ClassTypes=GameMode, Character, CameraComponent, SpringArmComponent, PlayerController
 AssetTypes=Animation Sequence, Skeletal Mesh, Animation Blueprint,Skeleton, Blend Space 1D
diff --git a/Templates/TP_TopDown/manifest.json b/Templates/TP_TopDown/manifest.json
--- a/Templates/TP_TopDown/manifest.json
+++ b/Templates/TP_TopDown/manifest.json
@@ -10,6 +10,10 @@
 			"Language": "ko",
 			"Text": "내려보기"
 		},
+		{
+			"Language": "zh-CN",
+			"Text": "俯视角游戏"
+		},
 		{
 			"Language": "ja",
 			"Text": "トップダウン"
@@ -23,11 +27,15 @@
 		},
 		{
 			"Language": "ko",
-			"Text": "내려보기 팩은 뒤쪽으로 삼인칭 시점보다 먼 거리에 있는 아바타 위에 카메라가 위치한 플레이가능 캐릭터를 특징으로 합니다. 플레이어 제어는 마우스나 터치스크린으로 원하는 목적지를 클릭하는 식으로 이루어지며, 원활한 캐릭터 이동을 위해 내비게이션 시스템을 사용합니다. 이 시점은 액션 롤플레잉 게임에 자주 사용됩니다."
+			"Text": "내려보기 템플릿은 플레이어 캐릭터 위 삼인칭 시점보다 멀리 카메라가 위치합니다. 플레이어 제어는 마우스나 터치스크린으로 원하는 목적지를 클릭하는 식으로 이루어지며, 원활한 캐릭터 이동을 위해 내비게이션 시스템을 사용합니다. 이 시점은 액션 롤플레잉 게임에 자주 쓰입니다."
+		},
+		{
+			"Language": "zh-CN",
+			"Text": "俯视角游戏模板拥有一个可游玩角色，摄像机位于角色的后上方，但位置比第三人称视角要高得多。玩家可以使用鼠标或触摸屏点击目的地，或使用导航系统来辅助角色移动。此视角常用于动作角色扮演类游戏。"
 		},
 		{
 			"Language": "ja",
-			"Text": "The top down pack features a playable character where the camera is positioned behind and above the avatar at greater distance than in third person perspective. The player is controlled by using the mouse or touchscreen to click on the required destination and uses the navigation system to facilitate the character movement. This perspective is often used in action role playing games."
+			"Text": "トップダウンテンプレートにはプレイアブルキャラクタが存在し、カメラはキャラクタの背後、サードパーソンテンプレートより遠い上方にあります。マウスやタッチスクリーンで行き先をクリックすることでキャラクタを操作することができます。キャラクタはナビゲーションシステムを利用して移動します。この視点はアクションロールプレイングゲームでよく使用されています。"
 		}
 	],
 	"AssetTypes":
@@ -35,14 +43,6 @@
 		{
 			"Language": "en",
 			"Text": "Animation Sequence, Skeletal Mesh, Animation Blueprint,Skeleton, Blend Space 1D"
-		},
-		{
-			"Language": "ko",
-			"Text": "애니메이션 시퀀스, 스켈레탈 메시, 애니메이션 블루프린트, 스켈레톤, 블렌드 스페이스 1D"
-		},
-		{
-			"Language": "ja",
-			"Text": "アニメーションシーケンス・スケルタルメッシュ・アニメーションブループリント・スケルトン・ブレンドシェイプ１D"
 		}
 	],	
 	"SearchTags":
@@ -50,10 +50,6 @@
 		{
 			"Language": "en",
 			"Text": "Role, RPG,Role Playing,Top Down"
-		},
-		{
-			"Language": "ko",
-			"Text": "롤, RPG, 롤플레잉, 내려보기, 탑다운"
 		}
 	],
 	"ClassTypes": "GameMode, Character, CameraComponent, SpringArmComponent, PlayerController",
diff --git a/Templates/TP_TopDownBP/Config/TemplateDefs.ini b/Templates/TP_TopDownBP/Config/TemplateDefs.ini
--- a/Templates/TP_TopDownBP/Config/TemplateDefs.ini
+++ b/Templates/TP_TopDownBP/Config/TemplateDefs.ini
@@ -3,8 +3,13 @@ LocalizedDisplayNames=(Language="en",Text="Top Down")
 LocalizedDescriptions=(Language="en",Text="The top down template features a playable character where the camera is positioned behind and above the avatar at greater distance than in third person perspective. The player is controlled by using the mouse or touchscreen to click on the required destination and uses the navigation system to facilitate the character movement. This perspective is often used in action role playing games.")
 
 LocalizedDisplayNames=(Language="ko",Text="내려보기")
-LocalizedDescriptions=(Language="ko",Text="내려보기 템플릿은 뒤쪽으로 삼인칭 시점보다 먼 거리에 있는 아바타 위에 카메라가 위치한 플레이가능 캐릭터를 특징으로 합니다. 플레이어 제어는 마우스나 터치스크린으로 원하는 목적지를 클릭하는 식으로 이루어지며, 원활한 캐릭터 이동을 위해 내비게이션 시스템을 사용합니다. 이 시점은 액션 롤플레잉 게임에 자주 사용됩니다.")
+LocalizedDescriptions=(Language="ko",Text="내려보기 템플릿은 플레이어 캐릭터 위 삼인칭 시점보다 멀리 카메라가 위치합니다. 플레이어 제어는 마우스나 터치스크린으로 원하는 목적지를 클릭하는 식으로 이루어지며, 원활한 캐릭터 이동을 위해 내비게이션 시스템을 사용합니다. 이 시점은 액션 롤플레잉 게임에 자주 쓰입니다.")
 
+LocalizedDisplayNames=(Language="ja",Text="トップダウン")
+LocalizedDescriptions=(Language="ja",Text="トップダウンテンプレートにはプレイアブルキャラクタが存在し、カメラはキャラクタの背後、サードパーソンテンプレートより遠い上方にあります。マウスやタッチスクリーンで行き先をクリックすることでキャラクタを操作することができます。キャラクタはナビゲーションシステムを利用して移動します。この視点はアクションロールプレイングゲームでよく使用されています。")
+
+LocalizedDisplayNames=(Language="zh-CN",Text="俯视角游戏")
+LocalizedDescriptions=(Language="zh-CN",Text="俯视角游戏模板拥有一个可游玩角色，摄像机位于角色的后上方，但位置比第三人称视角要高得多。玩家可以使用鼠标或触摸屏点击目的地，或使用导航系统来辅助角色移动。此视角常用于动作角色扮演类游戏。")
 
 ClassTypes=GameMode, Character, CameraComponent, SpringArmComponent, PlayerController
 AssetTypes=Animation Sequence, Skeletal Mesh, Animation Blueprint,Skeleton, Blend Space 1D
diff --git a/Templates/TP_TopDownBP/manifest.json b/Templates/TP_TopDownBP/manifest.json
--- a/Templates/TP_TopDownBP/manifest.json
+++ b/Templates/TP_TopDownBP/manifest.json
@@ -10,6 +10,10 @@
 			"Language": "ko",
 			"Text": "내려보기"
 		},
+		{
+			"Language": "zh-CN",
+			"Text": "俯视角游戏"
+		},
 		{
 			"Language": "ja",
 			"Text": "トップダウン"
@@ -23,11 +27,15 @@
 		},
 		{
 			"Language": "ko",
-			"Text": "내려보기 팩은 뒤쪽으로 삼인칭 시점보다 먼 거리에 있는 아바타 위에 카메라가 위치한 플레이가능 캐릭터를 특징으로 합니다. 플레이어 제어는 마우스나 터치스크린으로 원하는 목적지를 클릭하는 식으로 이루어지며, 원활한 캐릭터 이동을 위해 내비게이션 시스템을 사용합니다. 이 시점은 액션 롤플레잉 게임에 자주 사용됩니다."
+			"Text": "내려보기 템플릿은 플레이어 캐릭터 위 삼인칭 시점보다 멀리 카메라가 위치합니다. 플레이어 제어는 마우스나 터치스크린으로 원하는 목적지를 클릭하는 식으로 이루어지며, 원활한 캐릭터 이동을 위해 내비게이션 시스템을 사용합니다. 이 시점은 액션 롤플레잉 게임에 자주 쓰입니다."
+		},
+		{
+			"Language": "zh-CN",
+			"Text": "俯视角游戏模板拥有一个可游玩角色，摄像机位于角色的后上方，但位置比第三人称视角要高得多。玩家可以使用鼠标或触摸屏点击目的地，或使用导航系统来辅助角色移动。此视角常用于动作角色扮演类游戏。"
 		},
 		{
 			"Language": "ja",
-			"Text": "The top down pack features a playable character where the camera is positioned behind and above the avatar at greater distance than in third person perspective. The player is controlled by using the mouse or touchscreen to click on the required destination and uses the navigation system to facilitate the character movement. This perspective is often used in action role playing games."
+			"Text": "トップダウンテンプレートにはプレイアブルキャラクタが存在し、カメラはキャラクタの背後、サードパーソンテンプレートより遠い上方にあります。マウスやタッチスクリーンで行き先をクリックすることでキャラクタを操作することができます。キャラクタはナビゲーションシステムを利用して移動します。この視点はアクションロールプレイングゲームでよく使用されています。"
 		}
 	],
 	"AssetTypes":
@@ -35,14 +43,6 @@
 		{
 			"Language": "en",
 			"Text": "Animation Sequence, Skeletal Mesh, Animation Blueprint,Skeleton, Blend Space 1D"
-		},
-		{
-			"Language": "ko",
-			"Text": "애니메이션 시퀀스, 스켈레탈 메시, 애니메이션 블루프린트, 스켈레톤, 블렌드 스페이스 1D"
-		},
-		{
-			"Language": "ja",
-			"Text": "アニメーションシーケンス・スケルタルメッシュ・アニメーションブループリント・スケルトン・ブレンドシェイプ１D"
 		}
 	],	
 	"SearchTags":
@@ -50,10 +50,6 @@
 		{
 			"Language": "en",
 			"Text": "Role, RPG,Role Playing,Top Down"
-		},
-		{
-			"Language": "ko",
-			"Text": "롤, RPG, 롤플레잉, 내려보기, 탑다운"
 		}
 	],
 	"ClassTypes": "GameMode, Character, CameraComponent, SpringArmComponent, PlayerController",
diff --git a/Templates/TP_TwinStick/Config/TemplateDefs.ini b/Templates/TP_TwinStick/Config/TemplateDefs.ini
--- a/Templates/TP_TwinStick/Config/TemplateDefs.ini
+++ b/Templates/TP_TwinStick/Config/TemplateDefs.ini
@@ -3,7 +3,13 @@ LocalizedDisplayNames=(Language="en",Text="Twin Stick Shooter")
 LocalizedDescriptions=(Language="en",Text="The twin stick shooter template features a playable character where the camera is positioned almost directly above the character. Control of the character is with either keyboard, controller or virtual joystick on a touch device. Movement and firing, which are both on the X/Y axis are controlled independantly allowing for firing in a different direction than which the character is facing. The level contains several cubes that are affected by the projectiles and several that are not.")
 
 LocalizedDisplayNames=(Language="ko",Text="트윈 스틱 슈팅")
-LocalizedDescriptions=(Language="ko",Text="트윈 스틱 슈팅 템플릿은 카메라가 거의 캐릭터에 붙어있는 플레이가능 캐릭터를 특징으로 합니다. 캐릭터 콘트롤은 키보드, 콘트롤러, 터치 디바이스의 가상 조이스틱 중 하나로 가능합니다. 이동과 발사 둘 다 X/Y 축에 있으며 독립적으로 제어되어 캐릭터가 향하는 방향과 다른 방향으로 발사할 수 있습니다. 레벨에는 프로젝타일에 영향받는 큐브 여럿, 받지 않는 큐브도 여럿 들어있습니다.")
+LocalizedDescriptions=(Language="ko",Text="트윈 스틱 슈팅 템플릿은 카메라가 플레이어 캐릭터에 거의 붙어있습니다. 캐릭터 제어는 키보드, 컨트롤러, 터치 디바이스의 가상 조이스틱 중 하나로 가능합니다. 이동과 발사가 X/Y 축으로 각각 독립 제어되어 캐릭터가 향하는 방향과 다른 방향으로 발사할 수 있습니다. 레벨에는 발사체에 영향받는 큐브도 받지 않는 큐브도 여럿 있습니다.")
+
+LocalizedDisplayNames=(Language="ja",Text="ツインスティックシューター")
+LocalizedDescriptions=(Language="ja",Text="ツインスティックシューターテンプレートにはプレイアブルキャラクタが存在し、カメラはほぼキャラクタの真上にあります。キーボード、コントローラー、タッチデバイスでのバーチャルジョイスティックでこのキャラクタを操作できます。移動と弾の発射は、どちらも X と Y 軸で独立に操作します。キャラクタ向きと別の方向に弾を発射することもできます。弾の影響を受けるキューブと、影響を受けないキューブがレベルに含まれています。")
+
+LocalizedDisplayNames=(Language="zh-CN",Text="双摇杆射击类游戏")
+LocalizedDescriptions=(Language="zh-CN",Text="双摇杆设计类游戏模板拥有一个可游玩角色，摄像机几乎位于角色的正后方。您可以通过键盘、手柄或触摸式设备的虚拟摇杆控制角色移动。角色可以在X/Y轴上移动和射击，且为独立控制，所以角色可以向自身面向之外的方向射击。关卡内还有一些可以被子弹影响或不能被子弹影响的立方体。")
 
 ClassTypes=GameMode, Pawn, Actor, StaticMeshComponent, ProjectileMovementComponent, StaticMeshComponent, PrimitiveComponent, CameraComponent, SpringArmComponent, InputComponent
 AssetTypes=Sound Wave, Static Mesh
diff --git a/Templates/TP_TwinStick/manifest.json b/Templates/TP_TwinStick/manifest.json
--- a/Templates/TP_TwinStick/manifest.json
+++ b/Templates/TP_TwinStick/manifest.json
@@ -10,9 +10,13 @@
 			"Language": "ko",
 			"Text": "트윈 스틱 슈팅"
 		},
+		{
+			"Language": "zh-CN",
+			"Text": "双摇杆射击类游戏"
+		},
 		{
 			"Language": "ja",
-			"Text": "スクロールシューター"
+			"Text": "ツインスティックシューター"
 		}
 	],
 	"Description":
@@ -23,11 +27,15 @@
 		},
 		{
 			"Language": "ko",
-			"Text": "트윈 스틱 슈팅 팩은 카메라가 거의 캐릭터에 붙어있는 플레이가능 캐릭터를 특징으로 합니다. 캐릭터 콘트롤은 키보드, 콘트롤러, 터치 디바이스의 가상 조이스틱 중 하나로 가능합니다. 이동과 발사 둘 다 X/Y 축에 있으며 독립적으로 제어되어 캐릭터가 향하는 방향과 다른 방향으로 발사할 수 있습니다. 레벨에는 프로젝타일에 영향받는 큐브 여럿, 받지 않는 큐브도 여럿 들어있습니다."
+			"Text": "트윈 스틱 슈팅 템플릿은 카메라가 플레이어 캐릭터에 거의 붙어있습니다. 캐릭터 제어는 키보드, 컨트롤러, 터치 디바이스의 가상 조이스틱 중 하나로 가능합니다. 이동과 발사가 X/Y 축으로 각각 독립 제어되어 캐릭터가 향하는 방향과 다른 방향으로 발사할 수 있습니다. 레벨에는 발사체에 영향받는 큐브도 받지 않는 큐브도 여럿 있습니다."
+		},
+		{
+			"Language": "zh-CN",
+			"Text": "双摇杆设计类游戏模板拥有一个可游玩角色，摄像机几乎位于角色的正后方。您可以通过键盘、手柄或触摸式设备的虚拟摇杆控制角色移动。角色可以在X/Y轴上移动和射击，且为独立控制，所以角色可以向自身面向之外的方向射击。关卡内还有一些可以被子弹影响或不能被子弹影响的立方体。"
 		},
 		{
 			"Language": "ja",
-			"Text": "The twin stick shooter pack features a playable character where the camera is positioned almost directly above the character. Control of the character is with either keyboard, controller or virtual joystick on a touch device. Movement and firing, which are both on the X/Y axis are controlled independantly allowing for firing in a different direction than which the character is facing. The level contains several cubes that are affected by the projectiles and several that are not."
+			"Text": "ツインスティックシューターテンプレートにはプレイアブルキャラクタが存在し、カメラはほぼキャラクタの真上にあります。キーボード、コントローラー、タッチデバイスでのバーチャルジョイスティックでこのキャラクタを操作できます。移動と弾の発射は、どちらも X と Y 軸で独立に操作します。キャラクタ向きと別の方向に弾を発射することもできます。弾の影響を受けるキューブと、影響を受けないキューブがレベルに含まれています。"
 		}
 	],
 	"AssetTypes":
@@ -35,14 +43,6 @@
 		{
 			"Language": "en",
 			"Text": "Sound Wave, Static Mesh"
-		},
-		{
-			"Language": "ko",
-			"Text": "사운드 웨이브, 스태틱 메시"
-		},
-		{
-			"Language": "ja",
-			"Text": "サウンドウェーブ・スタティックメッシュ"
 		}
 	],
 	"SearchTags":
@@ -50,10 +50,6 @@
 		{
 			"Language": "en",
 			"Text": "Twin,Stick,Fire,Shooter"
-		},
-		{
-			"Language": "ko",
-			"Text": "트윈,스틱,슈팅"
 		}
 	],
 	"ClassTypes": "GameMode, Pawn, Actor, StaticMeshComponent, ProjectileMovementComponent, StaticMeshComponent, PrimitiveComponent, CameraComponent, SpringArmComponent, InputComponent",
diff --git a/Templates/TP_TwinStickBP/Config/TemplateDefs.ini b/Templates/TP_TwinStickBP/Config/TemplateDefs.ini
--- a/Templates/TP_TwinStickBP/Config/TemplateDefs.ini
+++ b/Templates/TP_TwinStickBP/Config/TemplateDefs.ini
@@ -3,7 +3,13 @@ LocalizedDisplayNames=(Language="en",Text="Twin Stick Shooter")
 LocalizedDescriptions=(Language="en",Text="The twin stick shooter template features a playable character where the camera is positioned almost directly above the character. Control of the character is with either keyboard, controller or virtual joystick on a touch device. Movement and firing, which are both on the X/Y axis are controlled independantly allowing for firing in a different direction than which the character is facing. The level contains several cubes that are affected by the projectiles and several that are not.")
 
 LocalizedDisplayNames=(Language="ko",Text="트윈 스틱 슈팅")
-LocalizedDescriptions=(Language="ko",Text="트윈 스틱 슈팅 템플릿은 카메라가 거의 캐릭터에 붙어있는 플레이가능 캐릭터를 특징으로 합니다. 캐릭터 콘트롤은 키보드, 콘트롤러, 터치 디바이스의 가상 조이스틱 중 하나로 가능합니다. 이동과 발사 둘 다 X/Y 축에 있으며 독립적으로 제어되어 캐릭터가 향하는 방향과 다른 방향으로 발사할 수 있습니다. 레벨에는 프로젝타일에 영향받는 큐브 여럿, 받지 않는 큐브도 여럿 들어있습니다.")
+LocalizedDescriptions=(Language="ko",Text="트윈 스틱 슈팅 템플릿은 카메라가 플레이어 캐릭터에 거의 붙어있습니다. 캐릭터 제어는 키보드, 컨트롤러, 터치 디바이스의 가상 조이스틱 중 하나로 가능합니다. 이동과 발사가 X/Y 축으로 각각 독립 제어되어 캐릭터가 향하는 방향과 다른 방향으로 발사할 수 있습니다. 레벨에는 발사체에 영향받는 큐브도 받지 않는 큐브도 여럿 있습니다.")
+
+LocalizedDisplayNames=(Language="ja",Text="ツインスティックシューター")
+LocalizedDescriptions=(Language="ja",Text="ツインスティックシューターテンプレートにはプレイアブルキャラクタが存在し、カメラはほぼキャラクタの真上にあります。キーボード、コントローラー、タッチデバイスでのバーチャルジョイスティックでこのキャラクタを操作できます。移動と弾の発射は、どちらも X と Y 軸で独立に操作します。キャラクタ向きと別の方向に弾を発射することもできます。弾の影響を受けるキューブと、影響を受けないキューブがレベルに含まれています。")
+
+LocalizedDisplayNames=(Language="zh-CN",Text="双摇杆射击类游戏")
+LocalizedDescriptions=(Language="zh-CN",Text="双摇杆设计类游戏模板拥有一个可游玩角色，摄像机几乎位于角色的正后方。您可以通过键盘、手柄或触摸式设备的虚拟摇杆控制角色移动。角色可以在X/Y轴上移动和射击，且为独立控制，所以角色可以向自身面向之外的方向射击。关卡内还有一些可以被子弹影响或不能被子弹影响的立方体。")
 
 ClassTypes=GameMode, Pawn, Actor, StaticMeshComponent, ProjectileMovementComponent, StaticMeshComponent, PrimitiveComponent, CameraComponent, SpringArmComponent, InputComponent
 AssetTypes=Sound Wave, Static Mesh
diff --git a/Templates/TP_TwinStickBP/manifest.json b/Templates/TP_TwinStickBP/manifest.json
--- a/Templates/TP_TwinStickBP/manifest.json
+++ b/Templates/TP_TwinStickBP/manifest.json
@@ -10,9 +10,13 @@
 			"Language": "ko",
 			"Text": "트윈 스틱 슈팅"
 		},
+		{
+			"Language": "zh-CN",
+			"Text": "双摇杆射击类游戏"
+		},
 		{
 			"Language": "ja",
-			"Text": "スクロールシューター"
+			"Text": "ツインスティックシューター"
 		}
 	],
 	"Description":
@@ -23,11 +27,15 @@
 		},
 		{
 			"Language": "ko",
-			"Text": "트윈 스틱 슈팅 팩은 카메라가 거의 캐릭터에 붙어있는 플레이가능 캐릭터를 특징으로 합니다. 캐릭터 콘트롤은 키보드, 콘트롤러, 터치 디바이스의 가상 조이스틱 중 하나로 가능합니다. 이동과 발사 둘 다 X/Y 축에 있으며 독립적으로 제어되어 캐릭터가 향하는 방향과 다른 방향으로 발사할 수 있습니다. 레벨에는 프로젝타일에 영향받는 큐브 여럿, 받지 않는 큐브도 여럿 들어있습니다."
+			"Text": "트윈 스틱 슈팅 템플릿은 카메라가 플레이어 캐릭터에 거의 붙어있습니다. 캐릭터 제어는 키보드, 컨트롤러, 터치 디바이스의 가상 조이스틱 중 하나로 가능합니다. 이동과 발사가 X/Y 축으로 각각 독립 제어되어 캐릭터가 향하는 방향과 다른 방향으로 발사할 수 있습니다. 레벨에는 발사체에 영향받는 큐브도 받지 않는 큐브도 여럿 있습니다."
+		},
+		{
+			"Language": "zh-CN",
+			"Text": "双摇杆设计类游戏模板拥有一个可游玩角色，摄像机几乎位于角色的正后方。您可以通过键盘、手柄或触摸式设备的虚拟摇杆控制角色移动。角色可以在X/Y轴上移动和射击，且为独立控制，所以角色可以向自身面向之外的方向射击。关卡内还有一些可以被子弹影响或不能被子弹影响的立方体。"
 		},
 		{
 			"Language": "ja",
-			"Text": "The twin stick shooter pack features a playable character where the camera is positioned almost directly above the character. Control of the character is with either keyboard, controller or virtual joystick on a touch device. Movement and firing, which are both on the X/Y axis are controlled independantly allowing for firing in a different direction than which the character is facing. The level contains several cubes that are affected by the projectiles and several that are not."
+			"Text": "ツインスティックシューターテンプレートにはプレイアブルキャラクタが存在し、カメラはほぼキャラクタの真上にあります。キーボード、コントローラー、タッチデバイスでのバーチャルジョイスティックでこのキャラクタを操作できます。移動と弾の発射は、どちらも X と Y 軸で独立に操作します。キャラクタ向きと別の方向に弾を発射することもできます。弾の影響を受けるキューブと、影響を受けないキューブがレベルに含まれています。"
 		}
 	],
 	"AssetTypes":
@@ -35,14 +43,6 @@
 		{
 			"Language": "en",
 			"Text": "Sound Wave, Static Mesh"
-		},
-		{
-			"Language": "ko",
-			"Text": "사운드 웨이브, 스태틱 메시"
-		},
-		{
-			"Language": "ja",
-			"Text": "サウンドウェーブ・スタティックメッシュ"
 		}
 	],
 	"SearchTags":
@@ -50,10 +50,6 @@
 		{
 			"Language": "en",
 			"Text": "Twin,Stick,Fire,Shooter"
-		},
-		{
-			"Language": "ko",
-			"Text": "트윈,스틱,슈팅"
 		}
 	],
 	"ClassTypes": "GameMode, Pawn, Actor, StaticMeshComponent, ProjectileMovementComponent, StaticMeshComponent, PrimitiveComponent, CameraComponent, SpringArmComponent, InputComponent",
diff --git a/Templates/TP_Vehicle/Config/TemplateDefs.ini b/Templates/TP_Vehicle/Config/TemplateDefs.ini
--- a/Templates/TP_Vehicle/Config/TemplateDefs.ini
+++ b/Templates/TP_Vehicle/Config/TemplateDefs.ini
@@ -3,7 +3,13 @@ LocalizedDisplayNames=(Language="en",Text="Vehicle")
 LocalizedDescriptions=(Language="en",Text="The vehicle template features a simple physics-driven vehicle. Vehicle movement can be controlled with either keyboard, controller or virtual joystick on a touch device. The template features 2 camera views - On is positioned above and behind the vehicle, whilst the other is positioned inside. When using the inside view the mouse or touchscreen or a HMD device. There is also a HUD both in and out of the car that shows the current gear and speed. The level features a smooth ramp, a slightly irregluar one, some objects that can be knocked over and a post which cannot.")
 
 LocalizedDisplayNames=(Language="ko",Text="비히클")
-LocalizedDescriptions=(Language="ko",Text="비히클 템플릿은 단순한 물리 기반 비히클을 특징으로 합니다. 비히클 이동은 키보드, 콘트롤러, 터치 디바이스의 가상 조이스틱 등 어느 것으로도 제어 가능합니다. 템플릿에는 두 가지 카메라 뷰가 제공되는데, 하나는 비히클 상단 뒤쪽에, 다른 하나는 안쪽에 위치해 있습니다. 안쪽에 있을 때는 마우스나 터치스크린이나 HMD 디바이스를 사용하여 둘러볼 수 있습니다. 차량 안에도 밖에도 현재 기어와 속도를 나타내는 HUD 가 있습니다. 레벨에는 약간 불규칙한 부드러운 경사로, 넘어뜨릴 수 있는 오브젝트와 넘어뜨릴 수 없는 기둥이 있습니다.")
+LocalizedDescriptions=(Language="ko",Text="비히클 템플릿은 단순한 물리 기반 비히클을 특징으로 합니다. 비히클 이동은 키보드, 컨트롤러, 터치 디바이스의 가상 조이스틱 등 어느 것으로도 제어 가능합니다. 템플릿에는 두 가지 카메라 뷰가 제공되는데, 하나는 비히클 상단 뒤쪽에, 다른 하나는 안쪽에 위치해 있습니다. 안쪽에 있을 때는 마우스나 터치스크린이나 HMD 디바이스를 사용하여 주변을 둘러볼 수 있습니다. 차량 안에도 밖에도 현재 기어와 속도를 나타내는 HUD 가 있습니다. 레벨에는 약간 불규칙한 부드러운 경사로, 넘어뜨릴 수 있는 오브젝트와 넘어뜨릴 수 없는 기둥이 있습니다.")
+
+LocalizedDisplayNames=(Language="ja",Text="ビークル")
+LocalizedDescriptions=(Language="ja",Text="ビークルテンプレートにはシンプルな物理駆動のビークルが含まれています。キーボード、コントローラー、タッチデバイスでのバーチャルジョイスティックでビークルを操作できます。このテンプレートには 2 方向からのカメラビューが用意されています。後方真上からのビューとビークル内部からのビューです。内部からのカメラを使用している時はマウス、タッチスクリーン、HMD デバイスで周りを見ることができます。また現在のギアと速度を示す HUD も用意されています。レベルにはスムーズな坂、多少不規則な坂、倒すことのできるオブジェクトとそうではない柱があります。")
+
+LocalizedDisplayNames=(Language="zh-CN",Text="载具类游戏")
+LocalizedDescriptions=(Language="zh-CN",Text="载具类游戏模板拥有一辆简单的物理驱动载具。您可通过键盘、手柄或触摸式设备上的虚拟摇杆来控制载具移动。该模板有2种摄像机视角——一种位于载具的后上方，而另一种位于载具内部。在使用内部视角时，您可通过鼠标、触摸屏或HMD设备观察四周。。载具内外还各有一个HUD，以显示当前装备和速度。关卡内含有一个光滑的斜坡和一个略不规则的斜坡，一些可以被撞飞的对象以及无法被撞开的柱子。")
 
 ClassTypes=GameMode, WheeledVehicle, HUD, Font, SpringArmComponent, CameraComponent, TextRenderComponent, InputComponent, VehicleWheel
 AssetTypes=Static Mesh, Animation Blueprint, Physics Asset, Physics Material, Skelton, Skeletal Mesh, Material, Material instance, Texture
diff --git a/Templates/TP_Vehicle/manifest.json b/Templates/TP_Vehicle/manifest.json
--- a/Templates/TP_Vehicle/manifest.json
+++ b/Templates/TP_Vehicle/manifest.json
@@ -10,6 +10,10 @@
 			"Language": "ko",
 			"Text": "비히클"
 		},
+		{
+			"Language": "zh-CN",
+			"Text": "载具类游戏"
+		},
 		{
 			"Language": "ja",
 			"Text": "レース"
@@ -23,11 +27,15 @@
 		},
 		{
 			"Language": "ko",
-			"Text": "비히클 팩은 단순한 물리 기반 비히클을 특징으로 합니다. 비히클 이동은 키보드, 콘트롤러, 터치 디바이스의 가상 조이스틱 등 어느 것으로도 제어 가능합니다. 템플릿에는 두 가지 카메라 뷰가 제공되는데, 하나는 비히클 상단 뒤쪽에, 다른 하나는 안쪽에 위치해 있습니다. 안쪽에 있을 때는 마우스나 터치스크린이나 HMD 디바이스를 사용하여 둘러볼 수 있습니다. 차량 안에도 밖에도 현재 기어와 속도를 나타내는 HUD 가 있습니다. 레벨에는 약간 불규칙한 부드러운 경사로, 넘어뜨릴 수 있는 오브젝트와 넘어뜨릴 수 없는 기둥이 있습니다."
+			"Text": "비히클 템플릿은 단순한 물리 기반 비히클을 특징으로 합니다. 비히클 이동은 키보드, 컨트롤러, 터치 디바이스의 가상 조이스틱 등 어느 것으로도 제어 가능합니다. 템플릿에는 두 가지 카메라 뷰가 제공되는데, 하나는 비히클 상단 뒤쪽에, 다른 하나는 안쪽에 위치해 있습니다. 안쪽에 있을 때는 마우스나 터치스크린이나 HMD 디바이스를 사용하여 주변을 둘러볼 수 있습니다. 차량 안에도 밖에도 현재 기어와 속도를 나타내는 HUD 가 있습니다. 레벨에는 약간 불규칙한 부드러운 경사로, 넘어뜨릴 수 있는 오브젝트와 넘어뜨릴 수 없는 기둥이 있습니다."
+		},
+		{
+			"Language": "zh-CN",
+			"Text": "载具类游戏模板拥有一辆简单的物理驱动载具。您可通过键盘、手柄或触摸式设备上的虚拟摇杆来控制载具移动。该模板有2种摄像机视角——一种位于载具的后上方，而另一种位于载具内部。在使用内部视角时，您可通过鼠标、触摸屏或HMD设备观察四周。。载具内外还各有一个HUD，以显示当前装备和速度。关卡内含有一个光滑的斜坡和一个略不规则的斜坡，一些可以被撞飞的对象以及无法被撞开的柱子。"
 		},
 		{
 			"Language": "ja",
-			"Text": "The vehicle pack features a simple physics-driven vehicle. Vehicle movement can be controlled with either keyboard, controller or virtual joystick on a touch device. The template features 2 camera views - On is positioned above and behind the vehicle, whilst the other is positioned inside. When using the inside view the mouse or touchscreen or a HMD device. There is also a HUD both in and out of the car that shows the current gear and speed. The level features a smooth ramp, a slightly irregluar one, some objects that can be knocked over and a post which cannot."
+			"Text": "ビークルテンプレートにはシンプルな物理駆動のビークルが含まれています。キーボード、コントローラー、タッチデバイスでのバーチャルジョイスティックでビークルを操作できます。このテンプレートには 2 方向からのカメラビューが用意されています。後方真上からのビューとビークル内部からのビューです。内部からのカメラを使用している時はマウス、タッチスクリーン、HMD デバイスで周りを見ることができます。また現在のギアと速度を示す HUD も用意されています。レベルにはスムーズな坂、多少不規則な坂、倒すことのできるオブジェクトとそうではない柱があります。"
 		}
 	],
 	"AssetTypes":
@@ -35,14 +43,6 @@
 		{
 			"Language": "en",
 			"Text": "Static Mesh, Animation Blueprint, Physics Asset, Physics Material, Skelton, Skeletal Mesh, Material, Material instance, Texture"
-		},
-		{
-			"Language": "ko",
-			"Text": "스태틱 메시, 애니메이션 블루프린트, 피직스 애셋, 피직스 머티리얼, 스켈레톤, 스켈레탈 메시, 머티리얼, 머티리얼 인스턴스, 텍스처"
-		},
-		{
-			"Language": "ja",
-			"Text": "スタティックメッシュ・アニメーションブループリント・フィジックス・フィジカルマテリアル・スケルトン・スケルタルメッシュ・マテリアル・マテリアルインスタンス・テクスチャ"
 		}
 	],
 	"SearchTags":
@@ -50,12 +50,8 @@
 		{
 			"Language": "en",
 			"Text": "Tyre,Wheel,Racing,Vehicle,Car"
-		},
-		{
-			"Language": "ko",
-			"Text": "타이어,휠,레이싱,비히클,차,경주"
 		}
-	],	
+	],
 	"ClassTypes": "GameMode, WheeledVehicle, HUD, Font, SpringArmComponent, CameraComponent, TextRenderComponent, InputComponent, VehicleWheel",
 	"Category": "CodeFeature",
 	"FocusAsset": "/Game/VehicleCPP/VehicleOverview.VehicleOverview",
diff --git a/Templates/TP_VehicleAdv/Config/TemplateDefs.ini b/Templates/TP_VehicleAdv/Config/TemplateDefs.ini
--- a/Templates/TP_VehicleAdv/Config/TemplateDefs.ini
+++ b/Templates/TP_VehicleAdv/Config/TemplateDefs.ini
@@ -3,7 +3,13 @@ LocalizedDisplayNames=(Language="en",Text="VehicleAdvanced")
 LocalizedDescriptions=(Language="en",Text="The vehicle advanced template features a more complex vehicle than the regular vehicle template. Vehicle movement can be controlled with either keyboard, controller or virtual joystick on a touch device. The template features 2 camera views - One is positioned above and behind the vehicle, whilst the other is positioned inside. When using the inside view you can look around using the mouse or touchscreen or a HMD device. The vehicle itself features working double wishbone suspension. There is also an engine sound in this template. There is also a HUD both in and out of the car that shows the current gear and speed. The level contains a landscape with rough terrain as well as a simple track that features bumps, ramps and loop the loops.")
 
 LocalizedDisplayNames=(Language="ko",Text="비히클 고급")
-LocalizedDescriptions=(Language="ko",Text="비히클 고급 템플릿은 일반 비히클 템플릿보다 복잡한 비히클이 특징입니다. 비히클 이동은 키보드, 콘트롤러, 터치 디바이스의 가상 조이스틱 등 어느 것으로도 제어 가능합니다. 템플릿에는 두 가지 카메라 뷰가 제공되는데, 하나는 비히클 상단 뒤쪽에, 다른 하나는 안쪽에 위치해 있습니다. 안쪽에 있을 때는 마우스나 터치스크린이나 HMD 디바이스를 사용하여 둘러볼 수 있습니다. 비히클 자체에는 더블 위시본 서스펜션은 물론 엔진 사운드도 구현되어 있습니다. 차량 안에도 밖에도 현재 기어와 속도를 나타내는 HUD 가 있습니다. 레벨에는 거친 지형의 랜드스케이프는 물론, 방지턱, 경사로, 루프로 된 간단한 트랙도 포함되어 있습니다.")
+LocalizedDescriptions=(Language="ko",Text="비히클 고급 템플릿은 일반 비히클 템플릿보다 복잡한 비히클이 특징입니다. 비히클 이동은 키보드, 컨트롤러, 터치 디바이스의 가상 조이스틱 등 어느 것으로도 제어 가능합니다. 템플릿에는 두 가지 카메라 뷰가 제공되는데, 하나는 비히클 상단 뒤쪽에, 다른 하나는 안쪽에 위치해 있습니다. 안쪽에 있을 때는 마우스나 터치스크린이나 HMD 디바이스를 사용하여 둘러볼 수 있습니다. 비히클 자체에는 더블 위시본 서스펜션은 물론 엔진 사운드도 구현되어 있습니다. 차량 안에도 밖에도 현재 기어와 속도를 나타내는 HUD 가 있습니다. 레벨에는 거친 지형의 랜드스케이프는 물론, 방지턱, 경사로, 간단한 루프 트랙도 포함되어 있습니다.")
+
+LocalizedDisplayNames=(Language="ja",Text="ビークル応用")
+LocalizedDescriptions=(Language="ja",Text="ビークル応用テンプレートには通常のビークルテンプレートより複雑なビークルが含まれています。キーボード、コントローラー、タッチデバイスでのバーチャルジョイスティックでビークルを操作できます。このテンプレートには 2 方向からのカメラビューが用意されています。後方真上からのビューとビークル内部からのビューです。内部からのカメラを使用している時はマウス、タッチスクリーン、HMD デバイスで周りを見ることができます。このビークルにはダブルウィッシュボーン式サスペンションが含まれています。またエンジン音もテンプレートにあります。また現在のギアと速度を示す HUD も用意されています。レベルには荒地を含む地形と、隆起、坂、ループを含む簡単なレーストラックがあります。")
+
+LocalizedDisplayNames=(Language="zh-CN",Text="高级载具类游戏")
+LocalizedDescriptions=(Language="zh-CN",Text="高级载具类游戏模板含有一辆比普通载具类游戏模板更复杂的载具。您可通过键盘、手柄或触摸式设备上的虚拟摇杆来控制载具移动。该模板有2种摄像机视角——一种位于载具的后上方，而另一种位于载具内部。在使用内部视角时，您可通过鼠标、触摸屏或HMD设备观察四周。载具本身具有双横臂式悬架。该模板还包含引擎音效。载具内外还各有一个HUD，以显示当前装备和速度。关卡内含有一片粗糙地形，以及一条包含隆起、斜坡以及环道的简单道路。")
 
 ClassTypes=GameMode, WheeledVehicle, HUD, SpringArmComponent, Font, CameraComponent, TextRenderComponent, InputComponent, PhysicalMaterial, AudioComponent
 AssetTypes=Static Mesh, Landscape Layer, Material, Instance, Physical Material, Sound Attenuation, Sound cue, Sound Wave, Texture, Skeletal Mesh, Skeleton, Physics Asset, Animation blueprint, Data Asset, Vehicle Wheel
diff --git a/Templates/TP_VehicleAdv/manifest.json b/Templates/TP_VehicleAdv/manifest.json
--- a/Templates/TP_VehicleAdv/manifest.json
+++ b/Templates/TP_VehicleAdv/manifest.json
@@ -6,6 +6,14 @@
 			"Language": "en",
 			"Text": "Vehicle Advanced"
 		},
+		{
+			"Language": "ja",
+			"Text": "ビークル応用"
+		},
+		{
+			"Language": "zh-CN",
+			"Text": "高级载具类游戏"
+		},
 		{
 			"Language": "ko",
 			"Text": "비히클 고급"
@@ -17,9 +25,17 @@
 			"Language": "en",
 			"Text": "The vehicle advanced pack features a more complex vehicle than the regular vehicle template. Vehicle movement can be controlled with either keyboard, controller or virtual joystick on a touch device. The template features 2 camera views - One is positioned above and behind the vehicle, whilst the other is positioned inside. When using the inside view you can look around using the mouse or touchscreen or a HMD device. The vehicle itself features working double wishbone suspension. There is also an engine sound in this template. There is also a HUD both in and out of the car that shows the current gear and speed. The level contains a landscape with rough terrain as well as a simple track that features bumps, ramps and loop the loops."
 		},
+		{
+			"Language": "ja",
+			"Text": "ビークル応用テンプレートには通常のビークルテンプレートより複雑なビークルが含まれています。キーボード、コントローラー、タッチデバイスでのバーチャルジョイスティックでビークルを操作できます。このテンプレートには 2 方向からのカメラビューが用意されています。後方真上からのビューとビークル内部からのビューです。内部からのカメラを使用している時はマウス、タッチスクリーン、HMD デバイスで周りを見ることができます。このビークルにはダブルウィッシュボーン式サスペンションが含まれています。またエンジン音もテンプレートにあります。また現在のギアと速度を示す HUD も用意されています。レベルには荒地を含む地形と、隆起、坂、ループを含む簡単なレーストラックがあります。"
+		},
+		{
+			"Language": "zh-CN",
+			"Text": "高级载具类游戏模板含有一辆比普通载具类游戏模板更复杂的载具。您可通过键盘、手柄或触摸式设备上的虚拟摇杆来控制载具移动。该模板有2种摄像机视角——一种位于载具的后上方，而另一种位于载具内部。在使用内部视角时，您可通过鼠标、触摸屏或HMD设备观察四周。载具本身具有双横臂式悬架。该模板还包含引擎音效。载具内外还各有一个HUD，以显示当前装备和速度。关卡内含有一片粗糙地形，以及一条包含隆起、斜坡以及环道的简单道路。"
+		},
 		{
 			"Language": "ko",
-			"Text": "비히클 고급 팩은 일반 비히클 템플릿보다 복잡한 비히클이 특징입니다. 비히클 이동은 키보드, 콘트롤러, 터치 디바이스의 가상 조이스틱 등 어느 것으로도 제어 가능합니다. 템플릿에는 두 가지 카메라 뷰가 제공되는데, 하나는 비히클 상단 뒤쪽에, 다른 하나는 안쪽에 위치해 있습니다. 안쪽에 있을 때는 마우스나 터치스크린이나 HMD 디바이스를 사용하여 둘러볼 수 있습니다. 비히클 자체에는 더블 위시본 서스펜션은 물론 엔진 사운드도 구현되어 있습니다. 차량 안에도 밖에도 현재 기어와 속도를 나타내는 HUD 가 있습니다. 레벨에는 거친 지형의 랜드스케이프는 물론, 방지턱, 경사로, 루프로 된 간단한 트랙도 포함되어 있습니다."
+			"Text": "비히클 고급 템플릿은 일반 비히클 템플릿보다 복잡한 비히클이 특징입니다. 비히클 이동은 키보드, 컨트롤러, 터치 디바이스의 가상 조이스틱 등 어느 것으로도 제어 가능합니다. 템플릿에는 두 가지 카메라 뷰가 제공되는데, 하나는 비히클 상단 뒤쪽에, 다른 하나는 안쪽에 위치해 있습니다. 안쪽에 있을 때는 마우스나 터치스크린이나 HMD 디바이스를 사용하여 둘러볼 수 있습니다. 비히클 자체에는 더블 위시본 서스펜션은 물론 엔진 사운드도 구현되어 있습니다. 차량 안에도 밖에도 현재 기어와 속도를 나타내는 HUD 가 있습니다. 레벨에는 거친 지형의 랜드스케이프는 물론, 방지턱, 경사로, 간단한 루프 트랙도 포함되어 있습니다."
 		}
 	],
 	"AssetTypes":
@@ -27,10 +43,6 @@
 		{
 			"Language": "en",
 			"Text": "Static Mesh, Landscape Layer, Material, Instance, Physical Material, Sound Attenuation, Sound cue, Sound Wave, Texture, Skeletal Mesh, Skeleton, Physics Asset, Animation blueprint, Data Asset, Vehicle Wheel"
-		},
-		{
-			"Language": "ko",
-			"Text": "스태틱 메시, 랜드스케이프 레이어, 머티리얼, 인스턴스, 피지컬 머티리얼, 사운드 감쇠, 사운드 큐, 사운드 웨이브, 텍스처, 스켈레탈 메시, 스켈레톤, 피직스 애셋, 애니메이션 블루프린트, 데이터 애셋, 비히클 휠"
 		}
 	],	
 	"SearchTags":
@@ -38,10 +50,6 @@
 		{
 			"Language": "en",
 			"Text": "Tyre,Suspension,Wheel,Racing,Vehicle,Car"
-		},
-		{
-			"Language": "ko",
-			"Text": "타이어,서스펜션,휠,레이싱,비히클,차,경주"
 		}
 	],
 	"ClassTypes": "GameMode, WheeledVehicle, HUD, SpringArmComponent, Font, CameraComponent, TextRenderComponent, InputComponent, PhysicalMaterial, AudioComponent",
diff --git a/Templates/TP_VehicleAdvBP/Config/TemplateDefs.ini b/Templates/TP_VehicleAdvBP/Config/TemplateDefs.ini
--- a/Templates/TP_VehicleAdvBP/Config/TemplateDefs.ini
+++ b/Templates/TP_VehicleAdvBP/Config/TemplateDefs.ini
@@ -3,7 +3,13 @@ LocalizedDisplayNames=(Language="en",Text="VehicleAdvanced")
 LocalizedDescriptions=(Language="en",Text="The vehicle advanced template features a more complex vehicle than the regular vehicle template. Vehicle movement can be controlled with either keyboard, controller or virtual joystick on a touch device. The template features 2 camera views - One is positioned above and behind the vehicle, whilst the other is positioned inside. When using the inside view you can look around using the mouse or touchscreen or a HMD device. The vehicle itself features working double wishbone suspension. There is also an engine sound in this template. There is also a HUD both in and out of the car that shows the current gear and speed. The level contains a landscape with rough terrain as well as a simple track that features bumps, ramps and loop the loops.")
 
 LocalizedDisplayNames=(Language="ko",Text="비히클 고급")
-LocalizedDescriptions=(Language="ko",Text="비히클 고급 템플릿은 일반 비히클 템플릿보다 복잡한 비히클이 특징입니다. 비히클 이동은 키보드, 콘트롤러, 터치 디바이스의 가상 조이스틱 등 어느 것으로도 제어 가능합니다. 템플릿에는 두 가지 카메라 뷰가 제공되는데, 하나는 비히클 상단 뒤쪽에, 다른 하나는 안쪽에 위치해 있습니다. 안쪽에 있을 때는 마우스나 터치스크린이나 HMD 디바이스를 사용하여 둘러볼 수 있습니다. 비히클 자체에는 더블 위시본 서스펜션은 물론 엔진 사운드도 구현되어 있습니다. 차량 안에도 밖에도 현재 기어와 속도를 나타내는 HUD 가 있습니다. 레벨에는 거친 지형의 랜드스케이프는 물론, 방지턱, 경사로, 루프로 된 간단한 트랙도 포함되어 있습니다.")
+LocalizedDescriptions=(Language="ko",Text="비히클 고급 템플릿은 일반 비히클 템플릿보다 복잡한 비히클이 특징입니다. 비히클 이동은 키보드, 컨트롤러, 터치 디바이스의 가상 조이스틱 등 어느 것으로도 제어 가능합니다. 템플릿에는 두 가지 카메라 뷰가 제공되는데, 하나는 비히클 상단 뒤쪽에, 다른 하나는 안쪽에 위치해 있습니다. 안쪽에 있을 때는 마우스나 터치스크린이나 HMD 디바이스를 사용하여 둘러볼 수 있습니다. 비히클 자체에는 더블 위시본 서스펜션은 물론 엔진 사운드도 구현되어 있습니다. 차량 안에도 밖에도 현재 기어와 속도를 나타내는 HUD 가 있습니다. 레벨에는 거친 지형의 랜드스케이프는 물론, 방지턱, 경사로, 간단한 루프 트랙도 포함되어 있습니다.")
+
+LocalizedDisplayNames=(Language="ja",Text="ビークル応用")
+LocalizedDescriptions=(Language="ja",Text="ビークル応用テンプレートには通常のビークルテンプレートより複雑なビークルが含まれています。キーボード、コントローラー、タッチデバイスでのバーチャルジョイスティックでビークルを操作できます。このテンプレートには 2 方向からのカメラビューが用意されています。後方真上からのビューとビークル内部からのビューです。内部からのカメラを使用している時はマウス、タッチスクリーン、HMD デバイスで周りを見ることができます。このビークルにはダブルウィッシュボーン式サスペンションが含まれています。またエンジン音もテンプレートにあります。また現在のギアと速度を示す HUD も用意されています。レベルには荒地を含む地形と、隆起、坂、ループを含む簡単なレーストラックがあります。")
+
+LocalizedDisplayNames=(Language="zh-CN",Text="高级载具类游戏")
+LocalizedDescriptions=(Language="zh-CN",Text="高级载具类游戏模板含有一辆比普通载具类游戏模板更复杂的载具。您可通过键盘、手柄或触摸式设备上的虚拟摇杆来控制载具移动。该模板有2种摄像机视角——一种位于载具的后上方，而另一种位于载具内部。在使用内部视角时，您可通过鼠标、触摸屏或HMD设备观察四周。载具本身具有双横臂式悬架。该模板还包含引擎音效。载具内外还各有一个HUD，以显示当前装备和速度。关卡内含有一片粗糙地形，以及一条包含隆起、斜坡以及环道的简单道路。")
 
 ClassTypes=GameMode, WheeledVehicle, HUD, SpringArmComponent, Font, CameraComponent, TextRenderComponent, InputComponent, PhysicalMaterial, AudioComponent
 AssetTypes=Static Mesh, Landscape Layer, Material, Instance, Physical Material, Sound Attenuation, Sound cue, Sound Wave, Texture, Skeletal Mesh, Skeleton, Physics Asset, Animation blueprint, Data Asset, Vehicle Wheel
diff --git a/Templates/TP_VehicleAdvBP/manifest.json b/Templates/TP_VehicleAdvBP/manifest.json
--- a/Templates/TP_VehicleAdvBP/manifest.json
+++ b/Templates/TP_VehicleAdvBP/manifest.json
@@ -6,6 +6,14 @@
 			"Language": "en",
 			"Text": "Vehicle Advanced"
 		},
+		{
+			"Language": "ja",
+			"Text": "ビークル応用"
+		},
+		{
+			"Language": "zh-CN",
+			"Text": "高级载具类游戏"
+		},
 		{
 			"Language": "ko",
 			"Text": "비히클 고급"
@@ -17,9 +25,17 @@
 			"Language": "en",
 			"Text": "The vehicle advanced pack features a more complex vehicle than the regular vehicle template. Vehicle movement can be controlled with either keyboard, controller or virtual joystick on a touch device. The template features 2 camera views - One is positioned above and behind the vehicle, whilst the other is positioned inside. When using the inside view you can look around using the mouse or touchscreen or a HMD device. The vehicle itself features working double wishbone suspension. There is also an engine sound in this template. There is also a HUD both in and out of the car that shows the current gear and speed. The level contains a landscape with rough terrain as well as a simple track that features bumps, ramps and loop the loops."
 		},
+		{
+			"Language": "ja",
+			"Text": "ビークル応用テンプレートには通常のビークルテンプレートより複雑なビークルが含まれています。キーボード、コントローラー、タッチデバイスでのバーチャルジョイスティックでビークルを操作できます。このテンプレートには 2 方向からのカメラビューが用意されています。後方真上からのビューとビークル内部からのビューです。内部からのカメラを使用している時はマウス、タッチスクリーン、HMD デバイスで周りを見ることができます。このビークルにはダブルウィッシュボーン式サスペンションが含まれています。またエンジン音もテンプレートにあります。また現在のギアと速度を示す HUD も用意されています。レベルには荒地を含む地形と、隆起、坂、ループを含む簡単なレーストラックがあります。"
+		},
+		{
+			"Language": "zh-CN",
+			"Text": "高级载具类游戏模板含有一辆比普通载具类游戏模板更复杂的载具。您可通过键盘、手柄或触摸式设备上的虚拟摇杆来控制载具移动。该模板有2种摄像机视角——一种位于载具的后上方，而另一种位于载具内部。在使用内部视角时，您可通过鼠标、触摸屏或HMD设备观察四周。载具本身具有双横臂式悬架。该模板还包含引擎音效。载具内外还各有一个HUD，以显示当前装备和速度。关卡内含有一片粗糙地形，以及一条包含隆起、斜坡以及环道的简单道路。"
+		},
 		{
 			"Language": "ko",
-			"Text": "비히클 고급 팩은 일반 비히클 템플릿보다 복잡한 비히클이 특징입니다. 비히클 이동은 키보드, 콘트롤러, 터치 디바이스의 가상 조이스틱 등 어느 것으로도 제어 가능합니다. 템플릿에는 두 가지 카메라 뷰가 제공되는데, 하나는 비히클 상단 뒤쪽에, 다른 하나는 안쪽에 위치해 있습니다. 안쪽에 있을 때는 마우스나 터치스크린이나 HMD 디바이스를 사용하여 둘러볼 수 있습니다. 비히클 자체에는 더블 위시본 서스펜션은 물론 엔진 사운드도 구현되어 있습니다. 차량 안에도 밖에도 현재 기어와 속도를 나타내는 HUD 가 있습니다. 레벨에는 거친 지형의 랜드스케이프는 물론, 방지턱, 경사로, 루프로 된 간단한 트랙도 포함되어 있습니다."
+			"Text": "비히클 고급 템플릿은 일반 비히클 템플릿보다 복잡한 비히클이 특징입니다. 비히클 이동은 키보드, 컨트롤러, 터치 디바이스의 가상 조이스틱 등 어느 것으로도 제어 가능합니다. 템플릿에는 두 가지 카메라 뷰가 제공되는데, 하나는 비히클 상단 뒤쪽에, 다른 하나는 안쪽에 위치해 있습니다. 안쪽에 있을 때는 마우스나 터치스크린이나 HMD 디바이스를 사용하여 둘러볼 수 있습니다. 비히클 자체에는 더블 위시본 서스펜션은 물론 엔진 사운드도 구현되어 있습니다. 차량 안에도 밖에도 현재 기어와 속도를 나타내는 HUD 가 있습니다. 레벨에는 거친 지형의 랜드스케이프는 물론, 방지턱, 경사로, 간단한 루프 트랙도 포함되어 있습니다."
 		}
 	],
 	"AssetTypes":
@@ -27,10 +43,6 @@
 		{
 			"Language": "en",
 			"Text": "Static Mesh, Landscape Layer, Material, Instance, Physical Material, Sound Attenuation, Sound cue, Sound Wave, Texture, Skeletal Mesh, Skeleton, Physics Asset, Animation blueprint, Data Asset, Vehicle Wheel"
-		},
-		{
-			"Language": "ko",
-			"Text": "스태틱 메시, 랜드스케이프 레이어, 머티리얼, 인스턴스, 피지컬 머티리얼, 사운드 감쇠, 사운드 큐, 사운드 웨이브, 텍스처, 스켈레탈 메시, 스켈레톤, 피직스 애셋, 애니메이션 블루프린트, 데이터 애셋, 비히클 휠"
 		}
 	],	
 	"SearchTags":
@@ -38,10 +50,6 @@
 		{
 			"Language": "en",
 			"Text": "Tyre,Suspension,Wheel,Racing,Vehicle,Car"
-		},
-		{
-			"Language": "ko",
-			"Text": "타이어,서스펜션,휠,레이싱,비히클,차,경주"
 		}
 	],
 	"ClassTypes": "GameMode, WheeledVehicle, HUD, SpringArmComponent, Font, CameraComponent, TextRenderComponent, InputComponent, PhysicalMaterial, AudioComponent",	
diff --git a/Templates/TP_VehicleBP/Config/TemplateDefs.ini b/Templates/TP_VehicleBP/Config/TemplateDefs.ini
--- a/Templates/TP_VehicleBP/Config/TemplateDefs.ini
+++ b/Templates/TP_VehicleBP/Config/TemplateDefs.ini
@@ -3,7 +3,13 @@ LocalizedDisplayNames=(Language="en",Text="Vehicle")
 LocalizedDescriptions=(Language="en",Text="The vehicle template features a simple physics-driven vehicle. Vehicle movement can be controlled with either keyboard, controller or virtual joystick on a touch device. The template features 2 camera views - One is positioned above and behind the vehicle, whilst the other is positioned inside. When using the inside view the mouse or touchscreen or a HMD device. There is also a HUD both in and out of the car that shows the current gear and speed. The level features a smooth ramp, a slightly irregluar one, some objects that can be knocked over and a post which cannot.")
 
 LocalizedDisplayNames=(Language="ko",Text="비히클")
-LocalizedDescriptions=(Language="ko",Text="비히클 템플릿은 단순한 물리 기반 비히클을 특징으로 합니다. 비히클 이동은 키보드, 콘트롤러, 터치 디바이스의 가상 조이스틱 등 어느 것으로도 제어 가능합니다. 템플릿에는 두 가지 카메라 뷰가 제공되는데, 하나는 비히클 상단 뒤쪽에, 다른 하나는 안쪽에 위치해 있습니다. 안쪽에 있을 때는 마우스나 터치스크린이나 HMD 디바이스를 사용하여 둘러볼 수 있습니다. 차량 안에도 밖에도 현재 기어와 속도를 나타내는 HUD 가 있습니다. 레벨에는 약간 불규칙한 부드러운 경사로, 넘어뜨릴 수 있는 오브젝트와 넘어뜨릴 수 없는 기둥이 있습니다.")
+LocalizedDescriptions=(Language="ko",Text="비히클 템플릿은 단순한 물리 기반 비히클을 특징으로 합니다. 비히클 이동은 키보드, 컨트롤러, 터치 디바이스의 가상 조이스틱 등 어느 것으로도 제어 가능합니다. 템플릿에는 두 가지 카메라 뷰가 제공되는데, 하나는 비히클 상단 뒤쪽에, 다른 하나는 안쪽에 위치해 있습니다. 안쪽에 있을 때는 마우스나 터치스크린이나 HMD 디바이스를 사용하여 주변을 둘러볼 수 있습니다. 차량 안에도 밖에도 현재 기어와 속도를 나타내는 HUD 가 있습니다. 레벨에는 약간 불규칙한 부드러운 경사로, 넘어뜨릴 수 있는 오브젝트와 넘어뜨릴 수 없는 기둥이 있습니다.")
+
+LocalizedDisplayNames=(Language="ja",Text="ビークル")
+LocalizedDescriptions=(Language="ja",Text="ビークルテンプレートにはシンプルな物理駆動のビークルが含まれています。キーボード、コントローラー、タッチデバイスでのバーチャルジョイスティックでビークルを操作できます。このテンプレートには 2 方向からのカメラビューが用意されています。後方真上からのビューとビークル内部からのビューです。内部からのカメラを使用している時はマウス、タッチスクリーン、HMD デバイスで周りを見ることができます。また現在のギアと速度を示す HUD も用意されています。レベルにはスムーズな坂、多少不規則な坂、倒すことのできるオブジェクトとそうではない柱があります。")
+
+LocalizedDisplayNames=(Language="zh-CN",Text="载具类游戏")
+LocalizedDescriptions=(Language="zh-CN",Text="载具类游戏模板拥有一辆简单的物理驱动载具。您可通过键盘、手柄或触摸式设备上的虚拟摇杆来控制载具移动。该模板有2种摄像机视角——一种位于载具的后上方，而另一种位于载具内部。在使用内部视角时，您可通过鼠标、触摸屏或HMD设备观察四周。。载具内外还各有一个HUD，以显示当前装备和速度。关卡内含有一个光滑的斜坡和一个略不规则的斜坡，一些可以被撞飞的对象以及无法被撞开的柱子。")
 
 ClassTypes=GameMode, WheeledVehicle, HUD, Font, SpringArmComponent, CameraComponent, TextRenderComponent, InputComponent, VehicleWheel
 AssetTypes=Static Mesh, Animation Blueprint, Physics Asset, Physics Material, Skelton, Skeletal Mesh, Material, Material instance, Texture
diff --git a/Templates/TP_VehicleBP/manifest.json b/Templates/TP_VehicleBP/manifest.json
--- a/Templates/TP_VehicleBP/manifest.json
+++ b/Templates/TP_VehicleBP/manifest.json
@@ -10,6 +10,10 @@
 			"Language": "ko",
 			"Text": "비히클"
 		},
+		{
+			"Language": "zh-CN",
+			"Text": "载具类游戏"
+		},
 		{
 			"Language": "ja",
 			"Text": "レース"
@@ -23,11 +27,15 @@
 		},
 		{
 			"Language": "ko",
-			"Text": "비히클 팩은 단순한 물리 기반 비히클을 특징으로 합니다. 비히클 이동은 키보드, 콘트롤러, 터치 디바이스의 가상 조이스틱 등 어느 것으로도 제어 가능합니다. 템플릿에는 두 가지 카메라 뷰가 제공되는데, 하나는 비히클 상단 뒤쪽에, 다른 하나는 안쪽에 위치해 있습니다. 안쪽에 있을 때는 마우스나 터치스크린이나 HMD 디바이스를 사용하여 둘러볼 수 있습니다. 차량 안에도 밖에도 현재 기어와 속도를 나타내는 HUD 가 있습니다. 레벨에는 약간 불규칙한 부드러운 경사로, 넘어뜨릴 수 있는 오브젝트와 넘어뜨릴 수 없는 기둥이 있습니다."
+			"Text": "비히클 템플릿은 단순한 물리 기반 비히클을 특징으로 합니다. 비히클 이동은 키보드, 컨트롤러, 터치 디바이스의 가상 조이스틱 등 어느 것으로도 제어 가능합니다. 템플릿에는 두 가지 카메라 뷰가 제공되는데, 하나는 비히클 상단 뒤쪽에, 다른 하나는 안쪽에 위치해 있습니다. 안쪽에 있을 때는 마우스나 터치스크린이나 HMD 디바이스를 사용하여 주변을 둘러볼 수 있습니다. 차량 안에도 밖에도 현재 기어와 속도를 나타내는 HUD 가 있습니다. 레벨에는 약간 불규칙한 부드러운 경사로, 넘어뜨릴 수 있는 오브젝트와 넘어뜨릴 수 없는 기둥이 있습니다."
+		},
+		{
+			"Language": "zh-CN",
+			"Text": "载具类游戏模板拥有一辆简单的物理驱动载具。您可通过键盘、手柄或触摸式设备上的虚拟摇杆来控制载具移动。该模板有2种摄像机视角——一种位于载具的后上方，而另一种位于载具内部。在使用内部视角时，您可通过鼠标、触摸屏或HMD设备观察四周。。载具内外还各有一个HUD，以显示当前装备和速度。关卡内含有一个光滑的斜坡和一个略不规则的斜坡，一些可以被撞飞的对象以及无法被撞开的柱子。"
 		},
 		{
 			"Language": "ja",
-			"Text": "The vehicle pack features a simple physics-driven vehicle. Vehicle movement can be controlled with either keyboard, controller or virtual joystick on a touch device. The template features 2 camera views - On is positioned above and behind the vehicle, whilst the other is positioned inside. When using the inside view the mouse or touchscreen or a HMD device. There is also a HUD both in and out of the car that shows the current gear and speed. The level features a smooth ramp, a slightly irregluar one, some objects that can be knocked over and a post which cannot."
+			"Text": "ビークルテンプレートにはシンプルな物理駆動のビークルが含まれています。キーボード、コントローラー、タッチデバイスでのバーチャルジョイスティックでビークルを操作できます。このテンプレートには 2 方向からのカメラビューが用意されています。後方真上からのビューとビークル内部からのビューです。内部からのカメラを使用している時はマウス、タッチスクリーン、HMD デバイスで周りを見ることができます。また現在のギアと速度を示す HUD も用意されています。レベルにはスムーズな坂、多少不規則な坂、倒すことのできるオブジェクトとそうではない柱があります。"
 		}
 	],
 	"AssetTypes":
@@ -35,14 +43,6 @@
 		{
 			"Language": "en",
 			"Text": "Static Mesh, Animation Blueprint, Physics Asset, Physics Material, Skelton, Skeletal Mesh, Material, Material instance, Texture"
-		},
-		{
-			"Language": "ko",
-			"Text": "스태틱 메시, 애니메이션 블루프린트, 피직스 애셋, 피직스 머티리얼, 스켈레톤, 스켈레탈 메시, 머티리얼, 머티리얼 인스턴스, 텍스처"
-		},
-		{
-			"Language": "ja",
-			"Text": "スタティックメッシュ・アニメーションブループリント・フィジックス・フィジカルマテリアル・スケルトン・スケルタルメッシュ・マテリアル・マテリアルインスタンス・テクスチャ"
 		}
 	],
 	"SearchTags":
@@ -50,10 +50,6 @@
 		{
 			"Language": "en",
 			"Text": "Tyre,Wheel,Racing,Vehicle,Car"
-		},
-		{
-			"Language": "ko",
-			"Text": "타이어,휠,레이싱,비히클,차,경주"
 		}
 	],
 	"ClassTypes": "GameMode, WheeledVehicle, HUD, Font, SpringArmComponent, CameraComponent, TextRenderComponent, InputComponent, VehicleWheel",
diff --git a/Templates/TP_VirtualRealityBP/Config/TemplateDefs.ini b/Templates/TP_VirtualRealityBP/Config/TemplateDefs.ini
--- a/Templates/TP_VirtualRealityBP/Config/TemplateDefs.ini
+++ b/Templates/TP_VirtualRealityBP/Config/TemplateDefs.ini
@@ -3,7 +3,13 @@ LocalizedDisplayNames=(Language="en",Text="Virtual Reality")
 LocalizedDescriptions=(Language="en",Text="Blueprint Virtual Reality Template for Desktop and Playstation VR. The template features two modes of teleport locomotion with separate pawns and maps for each. The MotionControllerMap lets players teleport around using motion controllers. The HMDLocomotionMap lets players teleport using either a Gamepad or point-and-teleport using just the HeadMountedDisplay.")
 
 LocalizedDisplayNames=(Language="ko",Text="가상 현실")
-LocalizedDescriptions=(Language="ko",Text="데스크탑 및 Playstation VR 용 블루프린트 가상 현실 템플릿입니다. 각 플랫폼 별 전용 폰과 맵은 물론 두 가지 텔레포트 이동 방식이 포함되어 있습니다. 플레이어는 MotionControllerMap 에서 모션 컨트롤러를 사용하여 텔레포트할 수 있으며, HMDLocomotionMap 에서 게임패드 아니면 HeadMountedDisplay 만으로 이동할 곳을 지정하는 방식으로 텔레포트할 수 있습니다.")
+LocalizedDescriptions=(Language="ko",Text="데스크톱 및 Playstation VR 용 블루프린트 가상 현실 템플릿입니다. 각 플랫폼 별 전용 폰과 맵은 물론 두 가지 텔레포트 이동 방식이 포함되어 있습니다. 플레이어는 MotionControllerMap 에서 모션 컨트롤러를 사용하여 텔레포트할 수 있으며, HMDLocomotionMap 에서 게임패드 아니면 HeadMountedDisplay 만으로 이동할 곳을 지정하는 방식으로 텔레포트할 수 있습니다.")
+
+LocalizedDisplayNames=(Language="ja",Text="バーチャルリアリティ")
+LocalizedDescriptions=(Language="ja",Text="デスクトップ VR と Playstation VR 向け、ブループリントのバーチャルリアリティテンプレートです。テンプレートには二種類のテレポートロコモーションモードがあります。それぞれに ポーンとマップが用意されています。MotionControllerMap ではモーションコントローラーを使ったテレポートが可能です。HMDLocomotionMap ではゲームパッドを使ったテレポート、及び、ヘッドマウントディスプレイを使ったポイントアンドテレポートが可能です。")
+
+LocalizedDisplayNames=(Language="zh-CN",Text="虚拟现实应用")
+LocalizedDescriptions=(Language="zh-CN",Text="蓝图虚拟现实应用模板面向台式PC以及Playstation VR设备。模板包含2中传送移动模式，分别拥有独立的Pawn和地图。MotionControllerMap可以让玩家使用体感手柄四处传送，而HMDLocomotionMap可以让玩家使用游戏手柄或头显设备进行“指向-传送”式的传送。")
 
 ClassTypes=Character, GameMode, HUD, Actor
 AssetTypes=Animation Sequence, Animation Montage, Sound Wave, Skeletal Mesh, Animation Blueprint, Skeleton, Static Mesh, Texture.
diff --git a/Templates/TP_VirtualRealityBP/manifest.json b/Templates/TP_VirtualRealityBP/manifest.json
--- a/Templates/TP_VirtualRealityBP/manifest.json
+++ b/Templates/TP_VirtualRealityBP/manifest.json
@@ -4,12 +4,36 @@
     {
       "Language": "en",
       "Text": "Virtual Reality"
+    },
+    {
+      "Language": "ko",
+      "Text": "가상 현실"
+    },
+    {
+      "Language": "ja",
+      "Text": "バーチャルリアリティ"
+    },
+    {
+      "Language": "zh-CN",
+      "Text": "虚拟现实应用"
     }
   ],
   "Description": [
     {
       "Language": "en",
       "Text": "Blueprint Virtual Reality feature pack for Desktop and Playstation VR. The pack features two modes of teleport locomotion with separate pawns and maps for each. The MotionControllerMap lets players teleport around using motion controllers. The HMDLocomotionMap lets players teleport using either a Gamepad or point-and-teleport using just the HeadMountedDisplay."
+    },
+    {
+      "Language": "ko",
+      "Text": "데스크톱 및 Playstation VR 용 블루프린트 가상 현실 템플릿입니다. 각 플랫폼 별 전용 폰과 맵은 물론 두 가지 텔레포트 이동 방식이 포함되어 있습니다. 플레이어는 MotionControllerMap 에서 모션 컨트롤러를 사용하여 텔레포트할 수 있으며, HMDLocomotionMap 에서 게임패드 아니면 HeadMountedDisplay 만으로 이동할 곳을 지정하는 방식으로 텔레포트할 수 있습니다."
+    },
+    {
+      "Language": "ja",
+      "Text": "デスクトップ VR と Playstation VR 向け、ブループリントのバーチャルリアリティテンプレートです。テンプレートには二種類のテレポートロコモーションモードがあります。それぞれに ポーンとマップが用意されています。MotionControllerMap ではモーションコントローラーを使ったテレポートが可能です。HMDLocomotionMap ではゲームパッドを使ったテレポート、及び、ヘッドマウントディスプレイを使ったポイントアンドテレポートが可能です。"
+    },
+    {
+      "Language": "zh-CN",
+      "Text": "蓝图虚拟现实应用模板面向台式PC以及Playstation VR设备。模板包含2中传送移动模式，分别拥有独立的Pawn和地图。MotionControllerMap可以让玩家使用体感手柄四处传送，而HMDLocomotionMap可以让玩家使用游戏手柄或头显设备进行“指向-传送”式的传送。"
     }
   ],
   "AssetTypes": [
diff --git a/Templates/TP_nDisplay/Config/TemplateDefs.ini b/Templates/TP_nDisplay/Config/TemplateDefs.ini
--- a/Templates/TP_nDisplay/Config/TemplateDefs.ini
+++ b/Templates/TP_nDisplay/Config/TemplateDefs.ini
@@ -2,6 +2,15 @@
 LocalizedDisplayNames=(Language="en",Text="nDisplay")
 LocalizedDescriptions=(Language="en",Text="This project template provides multi-display functionality using a cluster of PCs. Use as a starting point for stereoscropic/monoscopic Caves and PowerWall setups.")
 
+LocalizedDisplayNames=(Language="ko",Text="nDisplay")
+LocalizedDescriptions=(Language="ko",Text="PC 클러스터에 연결된 멀티 디스플레이 기능을 제공하는 프로젝트 템플릿입니다. 스테레오/모노스코픽 동굴이나 파워 월 구성을 위한 시작점으로 사용할 수 있습니다.")
+
+LocalizedDisplayNames=(Language="ja",Text="nDisplay")
+LocalizedDescriptions=(Language="ja",Text="PCクラスタを使ったマルチディスプレイ機能のテンプレートです。立体視/平面視の Cave や PowerWall のセットアップでの土台として使用してください。")
+
+LocalizedDisplayNames=(Language="zh-CN",Text="nDisplay")
+LocalizedDescriptions=(Language="zh-CN",Text="此项目模板提供了使用PC集群的多显示设备功能。它可作为立体/平面Caves以及PowerWall设置的基础。")
+
 ClassTypes=GameMode, Actor
 AssetTypes=Static Mesh, Texture
 
diff --git a/Templates/TP_nDisplayBP/Config/TemplateDefs.ini b/Templates/TP_nDisplayBP/Config/TemplateDefs.ini
--- a/Templates/TP_nDisplayBP/Config/TemplateDefs.ini
+++ b/Templates/TP_nDisplayBP/Config/TemplateDefs.ini
@@ -2,6 +2,15 @@
 LocalizedDisplayNames=(Language="en",Text="nDisplay")
 LocalizedDescriptions=(Language="en",Text="This project template provides multi-display functionality using a cluster of PCs. Use as a starting point for stereoscropic/monoscopic Caves and PowerWall setups.")
 
+LocalizedDisplayNames=(Language="ko",Text="nDisplay")
+LocalizedDescriptions=(Language="ko",Text="PC 클러스터에 연결된 멀티 디스플레이 기능을 제공하는 프로젝트 템플릿입니다. 스테레오/모노스코픽 동굴이나 파워 월 구성을 위한 시작점으로 사용할 수 있습니다.")
+
+LocalizedDisplayNames=(Language="ja",Text="nDisplay")
+LocalizedDescriptions=(Language="ja",Text="PCクラスタを使ったマルチディスプレイ機能のテンプレートです。立体視/平面視の Cave や PowerWall のセットアップでの土台として使用してください。")
+
+LocalizedDisplayNames=(Language="zh-CN",Text="nDisplay")
+LocalizedDescriptions=(Language="zh-CN",Text="此项目模板提供了使用PC集群的多显示设备功能。它可作为立体/平面Caves以及PowerWall设置的基础。")
+
 ClassTypes=GameMode, Actor
 AssetTypes=Static Mesh, Texture
 
