Fix from Oculus.
Without this change passing a compressed sound wav for haptic use would cause a crash.

diff --git a/Engine/Source/Runtime/Engine/Private/GameFramework/HapticFeedbackEffect.cpp b/Engine/Source/Runtime/Engine/Private/GameFramework/HapticFeedbackEffect.cpp
--- a/Engine/Source/Runtime/Engine/Private/GameFramework/HapticFeedbackEffect.cpp
+++ b/Engine/Source/Runtime/Engine/Private/GameFramework/HapticFeedbackEffect.cpp
@@ -173,7 +173,7 @@ void UHapticFeedbackEffect_SoundWave::PrepareSoundWaveBuffer()
 	{
 		return;
 	}
-	AD->Precache(SoundWave, true, false);
+	AD->Precache(SoundWave, true, false, true);
 	SoundWave->InitAudioResource(AD->GetRuntimeFormat(SoundWave));
 	uint8* PCMData = SoundWave->RawPCMData;
 	int32 RawPCMDataSize = SoundWave->RawPCMDataSize;
