Made sure material serialization is always using the same set of custom versions to prevent crashes when loading assets cooked with cook on the fly

diff --git a/Engine/Source/Runtime/Engine/Private/Materials/Material.cpp b/Engine/Source/Runtime/Engine/Private/Materials/Material.cpp
--- a/Engine/Source/Runtime/Engine/Private/Materials/Material.cpp
+++ b/Engine/Source/Runtime/Engine/Private/Materials/Material.cpp
@@ -55,6 +55,9 @@
 #include "Interfaces/ITargetPlatformManagerModule.h"
 #include "Interfaces/ITargetPlatform.h"
 #include "Materials/MaterialExpressionComment.h"
+#include "UObject/EditorObjectVersion.h"
+#include "UObject/ReleaseObjectVersion.h"
+
 #if WITH_EDITOR
 #include "Logging/TokenizedMessage.h"
 #include "Logging/MessageLog.h"
@@ -627,6 +630,9 @@ void SerializeInlineShaderMaps(
 	LLM_SCOPE(ELLMTag::MaterialShaderMaps);
 	SCOPED_LOADTIMER(SerializeInlineShaderMaps);
 
+	Ar.UsingCustomVersion(FEditorObjectVersion::GUID);
+	Ar.UsingCustomVersion(FReleaseObjectVersion::GUID);
+
 	if (Ar.IsSaving())
 	{
 		int32 NumResourcesToSave = 0;
@@ -3926,7 +3932,7 @@ void UMaterial::PostLoad()
 	{
 		SCOPE_SECONDS_COUNTER(MaterialLoadTime);
 // Daniel: Disable compiling shaders for cooked platforms as the cooker will manually call the BeginCacheForCookedPlatformData function and load balence
-/*#if WITH_EDITOR
+#if 0 && WITH_EDITOR
 		// enable caching in postload for derived data cache commandlet and cook by the book
 		ITargetPlatformManagerModule* TPM = GetTargetPlatformManager();
 		if (TPM && (TPM->RestrictFormatsToRuntimeOnly() == false))
@@ -3938,7 +3944,7 @@ void UMaterial::PostLoad()
 				BeginCacheForCookedPlatformData(Platforms[FormatIndex]);
 			}
 		}
-#endif*/
+#endif
 		//Don't compile shaders in post load for dev overhead materials.
 		if (FApp::CanEverRender() && !bIsMaterialEditorStatsMaterial)
 		{
