Add option to prim comp to replicate physics to autonomous proxies, default on for original behavior.

diff --git a/Engine/Source/Runtime/Engine/Classes/Components/PrimitiveComponent.h b/Engine/Source/Runtime/Engine/Classes/Components/PrimitiveComponent.h
--- a/Engine/Source/Runtime/Engine/Classes/Components/PrimitiveComponent.h
+++ b/Engine/Source/Runtime/Engine/Classes/Components/PrimitiveComponent.h
@@ -479,6 +479,11 @@ public:
 	UPROPERTY(EditAnywhere, BlueprintReadWrite, Category = Physics)
 	uint8 bApplyImpulseOnDamage : 1;
 
+	/** True if physics should be replicated to autonomous proxies. This should be true for
+		server-authoritative simulations, and false for client authoritative simulations. */
+	UPROPERTY(EditAnywhere, BlueprintReadWrite, Category = Physics)
+	uint8 bReplicatePhysicsToAutonomousProxy : 1;
+
 	// General flags.
 	
 	/** If this is True, this component must always be loaded on clients, even if Hidden and CollisionEnabled is NoCollision. */
diff --git a/Engine/Source/Runtime/Engine/Private/Components/PrimitiveComponent.cpp b/Engine/Source/Runtime/Engine/Private/Components/PrimitiveComponent.cpp
--- a/Engine/Source/Runtime/Engine/Private/Components/PrimitiveComponent.cpp
+++ b/Engine/Source/Runtime/Engine/Private/Components/PrimitiveComponent.cpp
@@ -281,6 +281,7 @@ UPrimitiveComponent::UPrimitiveComponent(const FObjectInitializer& ObjectInitial
 	LastCheckedAllCollideableDescendantsTime = 0.f;
 	
 	bApplyImpulseOnDamage = true;
+	bReplicatePhysicsToAutonomousProxy = true;
 
 	bReceiveMobileCSMShadows = true;
 #if WITH_EDITORONLY_DATA
diff --git a/Engine/Source/Runtime/Engine/Private/PhysicsEngine/PhysicsReplication.cpp b/Engine/Source/Runtime/Engine/Private/PhysicsEngine/PhysicsReplication.cpp
--- a/Engine/Source/Runtime/Engine/Private/PhysicsEngine/PhysicsReplication.cpp
+++ b/Engine/Source/Runtime/Engine/Private/PhysicsEngine/PhysicsReplication.cpp
@@ -379,32 +379,36 @@ void FPhysicsReplication::OnTick(float DeltaSeconds, TMap<TWeakObjectPtr<UPrimit
 			{
 				FReplicatedPhysicsTarget& PhysicsTarget = Itr.Value();
 				FRigidBodyState& UpdatedState = PhysicsTarget.TargetState;
-
-				AActor* OwningActor = PrimComp->GetOwner();
-				if (OwningActor && OwningActor->Role == ROLE_SimulatedProxy)	//TODO: can we avoid the replication all together?
+				bool bUpdated = false;
+				if (AActor* OwningActor = PrimComp->GetOwner())
 				{
-					// Get the ping of the guy who owns this thing. If nobody is,
-					// then it's server authoritative.
-					const float OwnerPing = GetOwnerPing(OwningActor, PhysicsTarget);
-
-					// Get the total ping - this approximates the time since the update was
-					// actually generated on the machine that is doing the authoritative sim.
-					// NOTE: We divide by 2 to approximate 1-way ping from 2-way ping.
-					const float PingSecondsOneWay = (LocalPing + OwnerPing) * 0.5f * 0.001f;
-
-					if (UpdatedState.Flags & ERigidBodyFlags::NeedsUpdate)
+					const bool bIsSimulated = OwningActor->Role == ROLE_SimulatedProxy;
+					const bool bIsReplicatedAutonomous = OwningActor->Role == ROLE_AutonomousProxy && PrimComp->bReplicatePhysicsToAutonomousProxy;
+					if (bIsSimulated || bIsReplicatedAutonomous)
 					{
-						const bool bRestoredState = ApplyRigidBodyState(DeltaSeconds, BI, PhysicsTarget, PhysicErrorCorrection, PingSecondsOneWay);
+						// Get the ping of the guy who owns this thing. If nobody is,
+						// then it's server authoritative.
+						const float OwnerPing = GetOwnerPing(OwningActor, PhysicsTarget);
 
-						// Need to update the component to match new position.
-						if (PhysicsReplicationCVars::SkipSkeletalRepOptimization == 0 || Cast<USkeletalMeshComponent>(PrimComp) == nullptr)	//simulated skeletal mesh does its own polling of physics results so we don't need to call this as it'll happen at the end of the physics sim
-						{
-							PrimComp->SyncComponentToRBPhysics();
-						}
+						// Get the total ping - this approximates the time since the update was
+						// actually generated on the machine that is doing the authoritative sim.
+						// NOTE: We divide by 2 to approximate 1-way ping from 2-way ping.
+						const float PingSecondsOneWay = (LocalPing + OwnerPing) * 0.5f * 0.001f;
 
-						if (bRestoredState)
+						if (UpdatedState.Flags & ERigidBodyFlags::NeedsUpdate)
 						{
-							bRemoveItr = true;
+							const bool bRestoredState = ApplyRigidBodyState(DeltaSeconds, BI, PhysicsTarget, PhysicErrorCorrection, PingSecondsOneWay);
+
+							// Need to update the component to match new position.
+							if (PhysicsReplicationCVars::SkipSkeletalRepOptimization == 0 || Cast<USkeletalMeshComponent>(PrimComp) == nullptr)	//simulated skeletal mesh does its own polling of physics results so we don't need to call this as it'll happen at the end of the physics sim
+							{
+								PrimComp->SyncComponentToRBPhysics();
+							}
+
+							if (bRestoredState)
+							{
+								bRemoveItr = true;
+							}
 						}
 					}
 				}
