REL - Workaround for an issue where ShaderPlatformToPlatformName does not have a TargetPlatform for Vulkan or OGL caused by 4150481 and 4159946

diff --git a/Engine/Shaders/Public/ShaderVersion.ush b/Engine/Shaders/Public/ShaderVersion.ush
--- a/Engine/Shaders/Public/ShaderVersion.ush
+++ b/Engine/Shaders/Public/ShaderVersion.ush
@@ -2,4 +2,4 @@
 // Each time the console command is executed it generates a new GUID. As this file is included
 // in Platform.ush (which should be included in any shader) it allows to invalidate the shader DDC.
 // 
-// GUID = 7416D80B4444417FE033F3884948C9B5
\ No newline at end of file
+// GUID = 1607B783BC304B918A05A1FD37C2C435
\ No newline at end of file
diff --git a/Engine/Source/Runtime/Engine/Private/ShaderCompiler/ShaderCompiler.cpp b/Engine/Source/Runtime/Engine/Private/ShaderCompiler/ShaderCompiler.cpp
--- a/Engine/Source/Runtime/Engine/Private/ShaderCompiler/ShaderCompiler.cpp
+++ b/Engine/Source/Runtime/Engine/Private/ShaderCompiler/ShaderCompiler.cpp
@@ -3168,7 +3168,15 @@ void GlobalBeginCompileShader(
 	bool bForwardShading = false;
 	{
 		ITargetPlatform* TargetPlatform = GetTargetPlatformManager()->FindTargetPlatform(ShaderPlatformToPlatformName(EShaderPlatform(Target.Platform)).ToString());
-		bForwardShading = TargetPlatform && TargetPlatform->UsesForwardShading();
+		if (TargetPlatform)
+		{
+			bForwardShading = TargetPlatform->UsesForwardShading();
+		}
+		else
+		{
+			static IConsoleVariable* CVarForwardShading = IConsoleManager::Get().FindConsoleVariable(TEXT("r.ForwardShading"));
+			bForwardShading = CVarForwardShading ? (CVarForwardShading->GetInt() != 0) : false;
+		}
 		Input.Environment.SetDefine(TEXT("FORWARD_SHADING"), bForwardShading);
 	}
 
