PR #5211: Add GetOuterObject function (Contributed by [REDACTED_PERSON])

diff --git a/Engine/Source/Runtime/Engine/Classes/Kismet/KismetSystemLibrary.h b/Engine/Source/Runtime/Engine/Classes/Kismet/KismetSystemLibrary.h
--- a/Engine/Source/Runtime/Engine/Classes/Kismet/KismetSystemLibrary.h
+++ b/Engine/Source/Runtime/Engine/Classes/Kismet/KismetSystemLibrary.h
@@ -113,6 +113,10 @@ class ENGINE_API UKismetSystemLibrary : public UBlueprintFunctionLibrary
 	UFUNCTION(BlueprintPure, Category = "Utilities", meta = (DisplayName = "Get Display Name"))
 	static FString GetClassDisplayName(UClass* Class);
 
+	// Returns the outer object of an object.
+	UFUNCTION(BlueprintPure, Category = "Utilities")
+	static UObject* GetOuterObject(const UObject* Object);
+
 	// Engine build number, for displaying to end users.
 	UFUNCTION(BlueprintPure, Category="Development", meta=(BlueprintThreadSafe))
 	static FString GetEngineVersion();
diff --git a/Engine/Source/Runtime/Engine/Private/KismetSystemLibrary.cpp b/Engine/Source/Runtime/Engine/Private/KismetSystemLibrary.cpp
--- a/Engine/Source/Runtime/Engine/Private/KismetSystemLibrary.cpp
+++ b/Engine/Source/Runtime/Engine/Private/KismetSystemLibrary.cpp
@@ -83,6 +83,11 @@ FString UKismetSystemLibrary::GetClassDisplayName(UClass* Class)
 	return Class ? Class->GetName() : FString();
 }
 
+UObject* UKismetSystemLibrary::GetOuterObject(const UObject* Object)
+{
+	return Object ? Object->GetOuter() : nullptr;
+}
+
 FString UKismetSystemLibrary::GetEngineVersion()
 {
 	return FEngineVersion::Current().ToString();
