#ue4 - Fix discrepancy in SpringArmComponent location lag when substepping versus when not. In practice this was pretty minor. Also removed pointless multiply of a value and then its inverse.

#udn https://udn.unrealengine.com/questions/464297/view.html

diff --git a/Engine/Source/Runtime/Engine/Private/GameFramework/SpringArmComponent.cpp b/Engine/Source/Runtime/Engine/Private/GameFramework/SpringArmComponent.cpp
--- a/Engine/Source/Runtime/Engine/Private/GameFramework/SpringArmComponent.cpp
+++ b/Engine/Source/Runtime/Engine/Private/GameFramework/SpringArmComponent.cpp
@@ -84,20 +84,18 @@ void USpringArmComponent::UpdateDesiredArmLocation(bool bDoTrace, bool bDoLocati
 {
 	FRotator DesiredRot = GetTargetRotation();
 
-	const float InverseCameraLagMaxTimeStep = (1.f / CameraLagMaxTimeStep);
-
 	// Apply 'lag' to rotation if desired
 	if(bDoRotationLag)
 	{
 		if (bUseCameraLagSubstepping && DeltaTime > CameraLagMaxTimeStep && CameraRotationLagSpeed > 0.f)
 		{
-			const FRotator ArmRotStep = (DesiredRot - PreviousDesiredRot).GetNormalized() * (CameraLagMaxTimeStep / DeltaTime);
+			const FRotator ArmRotStep = (DesiredRot - PreviousDesiredRot).GetNormalized() * (1.f / DeltaTime);
 			FRotator LerpTarget = PreviousDesiredRot;
 			float RemainingTime = DeltaTime;
 			while (RemainingTime > KINDA_SMALL_NUMBER)
 			{
 				const float LerpAmount = FMath::Min(CameraLagMaxTimeStep, RemainingTime);
-				LerpTarget += ArmRotStep * (LerpAmount * InverseCameraLagMaxTimeStep);
+				LerpTarget += ArmRotStep * LerpAmount;
 				RemainingTime -= LerpAmount;
 
 				DesiredRot = FRotator(FMath::QInterpTo(FQuat(PreviousDesiredRot), FQuat(LerpTarget), LerpAmount, CameraRotationLagSpeed));
@@ -119,13 +117,14 @@ void USpringArmComponent::UpdateDesiredArmLocation(bool bDoTrace, bool bDoLocati
 	{
 		if (bUseCameraLagSubstepping && DeltaTime > CameraLagMaxTimeStep && CameraLagSpeed > 0.f)
 		{
-			const FVector ArmMovementStep = (ArmOrigin - PreviousArmOrigin) * (CameraLagMaxTimeStep / DeltaTime);			
-			FVector LerpTarget = PreviousArmOrigin;
+			const FVector ArmMovementStep = (DesiredLoc - PreviousDesiredLoc) * (1.f / DeltaTime);
+			FVector LerpTarget = PreviousDesiredLoc;
+
 			float RemainingTime = DeltaTime;
 			while (RemainingTime > KINDA_SMALL_NUMBER)
 			{
 				const float LerpAmount = FMath::Min(CameraLagMaxTimeStep, RemainingTime);
-				LerpTarget += ArmMovementStep * (LerpAmount * InverseCameraLagMaxTimeStep);
+				LerpTarget += ArmMovementStep * LerpAmount;
 				RemainingTime -= LerpAmount;
 
 				DesiredLoc = FMath::VInterpTo(PreviousDesiredLoc, LerpTarget, LerpAmount, CameraLagSpeed);
