[NO NON-PII COMMIT MESSAGE]

diff --git a/Engine/Plugins/Runtime/Synthesis/Source/Synthesis/Classes/SynthComponents/SynthComponentGranulator.h b/Engine/Plugins/Runtime/Synthesis/Source/Synthesis/Classes/SynthComponents/SynthComponentGranulator.h
--- a/Engine/Plugins/Runtime/Synthesis/Source/Synthesis/Classes/SynthComponents/SynthComponentGranulator.h
+++ b/Engine/Plugins/Runtime/Synthesis/Source/Synthesis/Classes/SynthComponents/SynthComponentGranulator.h
@@ -123,7 +123,8 @@ public:
 
 protected:
 
-	TQueue<USoundWave*> PendingStoppingSoundWaves;
+	UPROPERTY(Transient)
+	USoundWave* GranulatedSoundWave;
 
 	Audio::FGranularSynth GranularSynth;
 	Audio::FSoundWavePCMLoader SoundWaveLoader;
diff --git a/Engine/Plugins/Runtime/Synthesis/Source/Synthesis/Classes/SynthComponents/SynthComponentWaveTable.h b/Engine/Plugins/Runtime/Synthesis/Source/Synthesis/Classes/SynthComponents/SynthComponentWaveTable.h
--- a/Engine/Plugins/Runtime/Synthesis/Source/Synthesis/Classes/SynthComponents/SynthComponentWaveTable.h
+++ b/Engine/Plugins/Runtime/Synthesis/Source/Synthesis/Classes/SynthComponents/SynthComponentWaveTable.h
@@ -86,10 +86,14 @@ public:
 	bool IsLoaded() const;
 
 protected:
+	void LoadSoundWaveInternal();
+
 	Audio::FSampleBufferReader SampleBufferReader;
 	Audio::TSampleBuffer<int16> SampleBuffer;
 	Audio::FSoundWavePCMLoader SoundWaveLoader;
 
 	float SampleDurationSec;
 	float SamplePlaybackProgressSec;
+
+	bool bIsLoaded;
 };
diff --git a/Engine/Plugins/Runtime/Synthesis/Source/Synthesis/Private/SynthComponentGranulator.cpp b/Engine/Plugins/Runtime/Synthesis/Source/Synthesis/Private/SynthComponentGranulator.cpp
--- a/Engine/Plugins/Runtime/Synthesis/Source/Synthesis/Private/SynthComponentGranulator.cpp
+++ b/Engine/Plugins/Runtime/Synthesis/Source/Synthesis/Private/SynthComponentGranulator.cpp
@@ -19,7 +19,6 @@ UGranularSynth::~UGranularSynth()
 bool UGranularSynth::Init(int32& SampleRate)
 {
 	NumChannels = 2;
-	SoundWaveLoader.Init(GetAudioDevice());
 	return true;
 }
 
@@ -102,26 +101,34 @@ void UGranularSynth::NoteOff(const float Note, const bool bKill)
 
 void UGranularSynth::TickComponent(float DeltaTime, enum ELevelTick TickType, FActorComponentTickFunction *ThisTickFunction)
 {
-	if (SoundWaveLoader.Update())
-	{
-		Audio::TSampleBuffer<int16> SampleBuffer;
-		SoundWaveLoader.GetSampleBuffer(SampleBuffer);
-
-		SynthCommand([this, SampleBuffer]()
-		{
-			GranularSynth.LoadSampleBuffer(SampleBuffer);
-
-			// Clear the pending sound waves queue since we've now loaded a new buffer of data
-			SoundWaveLoader.Reset();
-		});
-	}
+	SoundWaveLoader.Update();
 }
 
 void UGranularSynth::SetSoundWave(USoundWave* InSoundWave)
 {
-	if (InSoundWave)
+	if (InSoundWave != GranulatedSoundWave)
 	{
-		SoundWaveLoader.LoadSoundWave(InSoundWave);
+		GranulatedSoundWave = InSoundWave;
+		bIsLoaded = false;
+
+		if (InSoundWave)
+		{
+			TFunction<void(const USoundWave * SoundWave, const Audio::FSampleBuffer & SampleBuffer)> OnLoaded
+				= [this](const USoundWave * SoundWave, const Audio::FSampleBuffer & SampleBuffer)
+			{
+				if (SoundWave == GranulatedSoundWave)
+				{
+					SynthCommand([this, SampleBuffer]()
+					{
+						GranularSynth.LoadSampleBuffer(SampleBuffer);
+					});
+
+					bIsLoaded = true;
+				}
+			};
+
+			SoundWaveLoader.LoadSoundWave(InSoundWave, MoveTemp(OnLoaded));
+		}
 	}
 }
 
@@ -218,5 +225,5 @@ float UGranularSynth::GetCurrentPlayheadTime() const
 
 bool UGranularSynth::IsLoaded() const
 {
-	return SoundWaveLoader.IsSoundWaveLoaded();
+	return bIsLoaded;
 }
\ No newline at end of file
diff --git a/Engine/Plugins/Runtime/Synthesis/Source/Synthesis/Private/SynthComponentWaveTable.cpp b/Engine/Plugins/Runtime/Synthesis/Source/Synthesis/Private/SynthComponentWaveTable.cpp
--- a/Engine/Plugins/Runtime/Synthesis/Source/Synthesis/Private/SynthComponentWaveTable.cpp
+++ b/Engine/Plugins/Runtime/Synthesis/Source/Synthesis/Private/SynthComponentWaveTable.cpp
@@ -11,6 +11,7 @@ USynthSamplePlayer::USynthSamplePlayer(const FObjectInitializer& ObjInitializer)
 	, SoundWave(nullptr)
 	, SampleDurationSec(0.0f)
 	, SamplePlaybackProgressSec(0.0F)
+	, bIsLoaded(false)
 {
 	PrimaryComponentTick.bCanEverTick = true;
 }
@@ -24,12 +25,11 @@ bool USynthSamplePlayer::Init(int32& SampleRate)
 	NumChannels = 2;
 
 	SampleBufferReader.Init(SampleRate);
-	SoundWaveLoader.Init(GetAudioDevice());
 
-	if (SoundWave != nullptr)
-	{
-		SoundWaveLoader.LoadSoundWave(SoundWave);
-	}
+ 	if (SoundWave != nullptr)
+ 	{
+		LoadSoundWaveInternal();
+ 	}
 
 	return true;
 }
@@ -90,7 +90,7 @@ float USynthSamplePlayer::GetSampleDuration() const
 
 bool USynthSamplePlayer::IsLoaded() const
 {
-	return SoundWaveLoader.IsSoundWaveLoaded();
+	return bIsLoaded;
 }
 
 float USynthSamplePlayer::GetCurrentPlaybackProgressTime() const
@@ -107,14 +107,41 @@ float USynthSamplePlayer::GetCurrentPlaybackProgressPercent() const
 	return 0.0f;
 }
 
-void USynthSamplePlayer::SetSoundWave(USoundWave* InSoundWave)
+void USynthSamplePlayer::LoadSoundWaveInternal()
 {
-	SoundWaveLoader.LoadSoundWave(InSoundWave);
+	bIsLoaded = false;
 
-	SynthCommand([this]()
+	if (SoundWave)
 	{
-		SampleBufferReader.ClearBuffer();
-	});
+		TFunction<void(const USoundWave * OutSoundWave, const Audio::FSampleBuffer & OutSampleBuffer)> OnLoaded
+			= [this](const USoundWave * OutSoundWave, const Audio::FSampleBuffer & OutSampleBuffer)
+		{
+			if (OutSoundWave == SoundWave)
+			{
+				OnSampleLoaded.Broadcast();
+
+				SynthCommand([this, OutSampleBuffer]()
+				{
+					SampleBuffer = OutSampleBuffer;
+					SampleBufferReader.ClearBuffer();
+				});
+
+				bIsLoaded = true;
+			}
+		};
+
+		SoundWaveLoader.LoadSoundWave(SoundWave, MoveTemp(OnLoaded));
+	}
+}
+
+
+void USynthSamplePlayer::SetSoundWave(USoundWave* InSoundWave)
+{
+	if (SoundWave != InSoundWave)
+	{
+		SoundWave = InSoundWave;
+		LoadSoundWaveInternal();
+	}
 }
 
 void USynthSamplePlayer::OnRegister()
@@ -132,22 +159,7 @@ void USynthSamplePlayer::OnUnregister()
 
 void USynthSamplePlayer::TickComponent(float DeltaTime, enum ELevelTick TickType, FActorComponentTickFunction *ThisTickFunction)
 {
-	if (SoundWaveLoader.Update())
-	{
-		OnSampleLoaded.Broadcast();
-
-		Audio::FSampleBuffer NewSampleBuffer;
-		SoundWaveLoader.GetSampleBuffer(NewSampleBuffer);
-
-		SynthCommand([this, NewSampleBuffer]()
-		{
-			SampleBuffer = NewSampleBuffer;
-			SampleBufferReader.ClearBuffer();
-
-			// Clear the pending sound waves queue since we've now loaded a new buffer of data
-			SoundWaveLoader.Reset();
-		});
-	}
+	SoundWaveLoader.Update();
 
 	OnSamplePlaybackProgress.Broadcast(GetCurrentPlaybackProgressTime(), GetCurrentPlaybackProgressPercent());
 }
diff --git a/Engine/Source/Runtime/Engine/Classes/Sound/SampleBuffer.h b/Engine/Source/Runtime/Engine/Classes/Sound/SampleBuffer.h
--- a/Engine/Source/Runtime/Engine/Classes/Sound/SampleBuffer.h
+++ b/Engine/Source/Runtime/Engine/Classes/Sound/SampleBuffer.h
@@ -8,6 +8,7 @@
 #include "Sound/SoundEffectBase.h"
 #include "Async/AsyncWork.h"
 #include "Templates/UnrealTypeTraits.h"
+#include "UObject/GCObject.h"
 
 class USoundWave;
 class FAudioDevice;
@@ -281,44 +282,40 @@ namespace Audio
 	/* and call Update on every tick until it returns true, at which point  */
 	/* you may call GetSampleBuffer to get the decoded audio.               */
 	/************************************************************************/
-	class ENGINE_API FSoundWavePCMLoader
+	class ENGINE_API FSoundWavePCMLoader : public FGCObject
 	{
 	public:
 		FSoundWavePCMLoader();
 
-		// Intialize loader with audio device
-		void Init(FAudioDevice* InAudioDevice);
-
 		// Loads a USoundWave, call on game thread.
-		void LoadSoundWave(USoundWave* InSoundWave);
+		void LoadSoundWave(USoundWave* InSoundWave, TFunction<void(const USoundWave* SoundWave, const Audio::FSampleBuffer& OutSampleBuffer)> OnLoaded);
 
-		// Update the loading state. Returns true once the sound wave is loaded/decoded.
-		// Call on game thread.
-		bool Update();
+		// Update the loading state. 
+		void Update();
 
-		// Returns the sample buffer data once the sound wave is loaded/decoded. Call on game thread thread.
-		void GetSampleBuffer(TSampleBuffer<>& OutSampleBuffer);
+		//~ GCObject Interface
+		void AddReferencedObjects(FReferenceCollector& Collector) override;
+		//~ GCObject Interface
 
-		// Empties pending sound wave load references. Call on audio rendering thread.
-		void Reset();
+	private:
 
-		// Queries whether the current sound wave has finished loading/decoding
-		bool IsSoundWaveLoaded() const { return bIsLoaded; }
+		struct FLoadingSoundWaveInfo
+		{
+			// The sound wave which is loading PCM data
+			USoundWave* SoundWave;
+
+			// The lambda function to call when t he sound wave finishes loading
+			TFunction<void(const USoundWave* SoundWave, const Audio::FSampleBuffer& LoadedSampleBuffer)> OnLoaded;
+
+			// Whether the sound wave load/decode is in-flight
+			bool bIsLoading;
+
+			// Whether or not the sound wave has already been loaded
+			bool bIsLoaded;
+		};
 
-	private:
-		
-		// Ptr to the audio device to use to do the decoding
-		FAudioDevice* AudioDevice;
 		// Reference to current loading sound wave
-		USoundWave* SoundWave;	
-		// Struct to meta-data of decoded PCM buffer and ptr to PCM data
-		TSampleBuffer<> SampleBuffer;
-		// Queue of sound wave ptrs to hold references to them until fully released in audio render thread
-		TQueue<USoundWave*> PendingStoppingSoundWaves;
-		// Whether the sound wave load/decode is in-flight
-		bool bIsLoading;
-		// Whether or not the sound wave has already been loaded
-		bool bIsLoaded;
+		TArray<FLoadingSoundWaveInfo> LoadingSoundWaves;
 	};
 
 	// Enum used to express the current state of a FSoundWavePCMWriter's current operation.
diff --git a/Engine/Source/Runtime/Engine/Private/SampleBuffer.cpp b/Engine/Source/Runtime/Engine/Private/SampleBuffer.cpp
--- a/Engine/Source/Runtime/Engine/Private/SampleBuffer.cpp
+++ b/Engine/Source/Runtime/Engine/Private/SampleBuffer.cpp
@@ -12,95 +12,81 @@
 namespace Audio
 {
 	FSoundWavePCMLoader::FSoundWavePCMLoader()
-		: AudioDevice(nullptr)
-		, SoundWave(nullptr)
-		, bIsLoading(false)
-		, bIsLoaded(false)
 	{
 	}
 
-	void FSoundWavePCMLoader::Init(FAudioDevice* InAudioDevice)
+	void FSoundWavePCMLoader::LoadSoundWave(USoundWave* InSoundWave, TFunction<void(const USoundWave* SoundWave, const Audio::FSampleBuffer& OutSampleBuffer)> OnLoaded)
 	{
-		AudioDevice = InAudioDevice;
-	}
+		FAudioDevice* AudioDevice = FAudioDevice::GetMainAudioDevice();
 
-	void FSoundWavePCMLoader::LoadSoundWave(USoundWave* InSoundWave)
-	{
 		if (!AudioDevice || !InSoundWave)
 		{
 			return;
 		}
 
-		// Queue existing sound wave reference so it can be
-		// cleared when the audio thread gets newly loaded audio data. 
-		// Don't want to kill the sound wave PCM data while it's playing on audio thread.
-		if (SoundWave)
-		{
-			PendingStoppingSoundWaves.Enqueue(SoundWave);
-		}
+		FLoadingSoundWaveInfo LoadingSoundWaveInfo;
+		LoadingSoundWaveInfo.SoundWave = InSoundWave;
+		LoadingSoundWaveInfo.OnLoaded = MoveTemp(OnLoaded);
+		LoadingSoundWaveInfo.bIsLoading = true;
 
-		SoundWave = InSoundWave;
-		if (!SoundWave->RawPCMData || SoundWave->AudioDecompressor)
+		if (LoadingSoundWaveInfo.SoundWave->GetPrecacheState() != ESoundWavePrecacheState::Done)
 		{
-			bIsLoaded = false;
-			bIsLoading = true;
+			LoadingSoundWaveInfo.bIsLoaded = false;
 
-			if (!SoundWave->RawPCMData)
-			{
-				// Kick off a decompression/precache of the sound wave
-				AudioDevice->Precache(SoundWave, false, true, true);
-			}
+			// Kick off a decompression/precache of the sound wave
+			AudioDevice->Precache(InSoundWave, false, true, true);
 		}
 		else
 		{
-			bIsLoading = true;
-			bIsLoaded = true;
+			LoadingSoundWaveInfo.bIsLoaded = true;
 		}
+
+		LoadingSoundWaves.Add(LoadingSoundWaveInfo);
 	}
 
-	bool FSoundWavePCMLoader::Update()
+	void FSoundWavePCMLoader::Update()
 	{
-		if (bIsLoading)
+		for (int32 i = LoadingSoundWaves.Num() - 1; i >= 0; --i)
 		{
-			check(SoundWave);
-	
-			if (bIsLoaded || (SoundWave->AudioDecompressor && SoundWave->AudioDecompressor->IsDone()))
+			FLoadingSoundWaveInfo& LoadingSoundWaveInfo = LoadingSoundWaves[i];
+
+			if (LoadingSoundWaveInfo.bIsLoading || LoadingSoundWaveInfo.bIsLoaded)
 			{
-				if (SoundWave->AudioDecompressor)
+				check(LoadingSoundWaveInfo.SoundWave);
+
+				if (LoadingSoundWaveInfo.bIsLoaded || LoadingSoundWaveInfo.SoundWave->GetPrecacheState() == ESoundWavePrecacheState::Done)
 				{
-					check(!bIsLoaded);
-					delete SoundWave->AudioDecompressor;
-					SoundWave->AudioDecompressor = nullptr;
-					SoundWave->SetPrecacheState(ESoundWavePrecacheState::Done);
-				}
+					LoadingSoundWaveInfo.bIsLoading = false;
+					LoadingSoundWaveInfo.bIsLoaded = true;
 
-				bIsLoading = false;
-				bIsLoaded = true;
+					TSampleBuffer<> SampleBuffer;
 
-				SampleBuffer.RawPCMData.Reset(SoundWave->RawPCMDataSize);
-				SampleBuffer.RawPCMData.AddUninitialized(SoundWave->RawPCMDataSize);
-				FMemory::Memcpy(SampleBuffer.RawPCMData.GetData(), SoundWave->RawPCMData, SoundWave->RawPCMDataSize);
-				SampleBuffer.NumSamples = SoundWave->RawPCMDataSize / sizeof(int16);
-				SampleBuffer.NumChannels = SoundWave->NumChannels;
-				SampleBuffer.NumFrames = SampleBuffer.NumSamples / SoundWave->NumChannels;
-				SampleBuffer.SampleRate = SoundWave->GetSampleRateForCurrentPlatform();
-				SampleBuffer.SampleDuration = (float)SampleBuffer.NumFrames / SampleBuffer.SampleRate;
+					USoundWave* SoundWave = LoadingSoundWaveInfo.SoundWave;
 
-				return true;
-			}
-		}
+					SampleBuffer.RawPCMData.Reset(SoundWave->RawPCMDataSize);
+					SampleBuffer.RawPCMData.AddUninitialized(SoundWave->RawPCMDataSize);
+					FMemory::Memcpy(SampleBuffer.RawPCMData.GetData(), SoundWave->RawPCMData, SoundWave->RawPCMDataSize);
+					SampleBuffer.NumSamples = SoundWave->RawPCMDataSize / sizeof(int16);
+					SampleBuffer.NumChannels = SoundWave->NumChannels;
+					SampleBuffer.NumFrames = SampleBuffer.NumSamples / SoundWave->NumChannels;
+					SampleBuffer.SampleRate = SoundWave->GetSampleRateForCurrentPlatform();
+					SampleBuffer.SampleDuration = (float)SampleBuffer.NumFrames / SampleBuffer.SampleRate;
 
-		return false;
-	}
+					LoadingSoundWaveInfo.OnLoaded(SoundWave, SampleBuffer);
 
-	void FSoundWavePCMLoader::GetSampleBuffer(TSampleBuffer<>& OutSampleBuffer)
-	{
-		OutSampleBuffer = SampleBuffer;
+					LoadingSoundWaves.RemoveAtSwap(i, 1, false);
+				}
+			}
+		}
 	}
 
-	void FSoundWavePCMLoader::Reset()
+	void FSoundWavePCMLoader::AddReferencedObjects(FReferenceCollector& Collector)
 	{
-		PendingStoppingSoundWaves.Empty();
+		for (FLoadingSoundWaveInfo& LoadingSoundWave : LoadingSoundWaves)
+		{
+			check(LoadingSoundWave.SoundWave);
+			Collector.AddReferencedObject(LoadingSoundWave.SoundWave);
+		}
 	}
 
 	FSoundWavePCMWriter::FSoundWavePCMWriter(int32 InChunkSize)
