Temporarily disable XGE shader compile on build farms to determine if it's causing random failures

diff --git a/Engine/Source/Runtime/Engine/Private/ShaderCompiler/ShaderCompilerXGE.cpp b/Engine/Source/Runtime/Engine/Private/ShaderCompiler/ShaderCompilerXGE.cpp
--- a/Engine/Source/Runtime/Engine/Private/ShaderCompiler/ShaderCompilerXGE.cpp
+++ b/Engine/Source/Runtime/Engine/Private/ShaderCompiler/ShaderCompilerXGE.cpp
@@ -81,6 +81,12 @@ namespace XGEShaderCompilerVariables
 				XGEShaderCompilerVariables::Enabled = 0;
 			}
 
+			//something going wrong with XGE shader compile on the build farm at the moment. Temporarily disable.
+			if (GIsBuildMachine)
+			{
+				XGEShaderCompilerVariables::Enabled = 0;
+			}
+
 			bInitialized = true;
 		}
 	}
