(4.21.1) Don't crash when initializing the GameMode if the GameStateClass has become null, instead fall back to AGameState
#rnx

diff --git a/Engine/Source/Runtime/Engine/Private/GameMode.cpp b/Engine/Source/Runtime/Engine/Private/GameMode.cpp
--- a/Engine/Source/Runtime/Engine/Private/GameMode.cpp
+++ b/Engine/Source/Runtime/Engine/Private/GameMode.cpp
@@ -72,7 +72,12 @@ void AGameMode::InitGame(const FString& MapName, const FString& Options, FString
 	Super::InitGame(MapName, Options, ErrorMessage);
 	SetMatchState(MatchState::EnteringMap);
 
-	if (!GameStateClass->IsChildOf<AGameState>())
+	if (GameStateClass == nullptr)
+	{
+		UE_LOG(LogGameMode, Error, TEXT("GameStateClass is not set, falling back to AGameState.");
+		GameStateClass = AGameState::StaticClass();
+	}
+	else if (!GameStateClass->IsChildOf<AGameState>())
 	{
 		UE_LOG(LogGameMode, Error, TEXT("Mixing AGameStateBase with AGameMode is not compatible. Change AGameStateBase subclass (%s) to derive from AGameState, or make both derive from Base"), *GameStateClass->GetName());
 	}
