Fix missing )
#rnx

diff --git a/Engine/Source/Runtime/Engine/Private/GameMode.cpp b/Engine/Source/Runtime/Engine/Private/GameMode.cpp
--- a/Engine/Source/Runtime/Engine/Private/GameMode.cpp
+++ b/Engine/Source/Runtime/Engine/Private/GameMode.cpp
@@ -74,7 +74,7 @@ void AGameMode::InitGame(const FString& MapName, const FString& Options, FString
 
 	if (GameStateClass == nullptr)
 	{
-		UE_LOG(LogGameMode, Error, TEXT("GameStateClass is not set, falling back to AGameState.");
+		UE_LOG(LogGameMode, Error, TEXT("GameStateClass is not set, falling back to AGameState."));
 		GameStateClass = AGameState::StaticClass();
 	}
 	else if (!GameStateClass->IsChildOf<AGameState>())
