Don't populated Graphs with nullptr entries

diff --git a/Engine/Source/Runtime/Engine/Private/Blueprint.cpp b/Engine/Source/Runtime/Engine/Private/Blueprint.cpp
--- a/Engine/Source/Runtime/Engine/Private/Blueprint.cpp
+++ b/Engine/Source/Runtime/Engine/Private/Blueprint.cpp
@@ -1599,28 +1599,40 @@ void UBlueprint::GetAllGraphs(TArray<UEdGraph*>& Graphs) const
 	for (int32 i = 0; i < FunctionGraphs.Num(); ++i)
 	{
 		UEdGraph* Graph = FunctionGraphs[i];
-		Graphs.Add(Graph);
-		Graph->GetAllChildrenGraphs(Graphs);
+		if(Graph)
+		{
+			Graphs.Add(Graph);
+			Graph->GetAllChildrenGraphs(Graphs);
+		}
 	}
 	for (int32 i = 0; i < MacroGraphs.Num(); ++i)
 	{
 		UEdGraph* Graph = MacroGraphs[i];
-		Graphs.Add(Graph);
-		Graph->GetAllChildrenGraphs(Graphs);
+		if(Graph)
+		{
+			Graphs.Add(Graph);
+			Graph->GetAllChildrenGraphs(Graphs);
+		}
 	}
 
 	for (int32 i = 0; i < UbergraphPages.Num(); ++i)
 	{
 		UEdGraph* Graph = UbergraphPages[i];
-		Graphs.Add(Graph);
-		Graph->GetAllChildrenGraphs(Graphs);
+		if(Graph)
+		{
+			Graphs.Add(Graph);
+			Graph->GetAllChildrenGraphs(Graphs);
+		}
 	}
 
 	for (int32 i = 0; i < DelegateSignatureGraphs.Num(); ++i)
 	{
 		UEdGraph* Graph = DelegateSignatureGraphs[i];
-		Graphs.Add(Graph);
-		Graph->GetAllChildrenGraphs(Graphs);
+		if(Graph)
+		{
+			Graphs.Add(Graph);
+			Graph->GetAllChildrenGraphs(Graphs);
+		}
 	}
 
 	for (int32 BPIdx=0; BPIdx<ImplementedInterfaces.Num(); BPIdx++)
@@ -1629,8 +1641,11 @@ void UBlueprint::GetAllGraphs(TArray<UEdGraph*>& Graphs) const
 		for (int32 GraphIdx = 0; GraphIdx < InterfaceDesc.Graphs.Num(); GraphIdx++)
 		{
 			UEdGraph* Graph = InterfaceDesc.Graphs[GraphIdx];
-			Graphs.Add(Graph);
-			Graph->GetAllChildrenGraphs(Graphs);
+			if(Graph)
+			{
+				Graphs.Add(Graph);
+				Graph->GetAllChildrenGraphs(Graphs);
+			}
 		}
 	}
 #endif // WITH_EDITORONLY_DATA
