Fix deprecation of deprecated macro
#rnx

diff --git a/Engine/Source/Runtime/Engine/Classes/Engine/BlueprintCore.h b/Engine/Source/Runtime/Engine/Classes/Engine/BlueprintCore.h
--- a/Engine/Source/Runtime/Engine/Classes/Engine/BlueprintCore.h
+++ b/Engine/Source/Runtime/Engine/Classes/Engine/BlueprintCore.h
@@ -34,10 +34,10 @@ private:
 	FGuid BlueprintGuid;
 
 public:
-	DEPRECATED(4.22, "The minimum UE4 object version implies the Blueprint generated class is always authoritative. It is no longer necessary to explicitly set it.")
+	UE_DEPRECATED(4.22, "The minimum UE4 object version implies the Blueprint generated class is always authoritative. It is no longer necessary to explicitly set it.")
 	void SetLegacyGeneratedClassIsAuthoritative() {}
 
-	DEPRECATED(4.22, "The minimum UE4 object version implies the Blueprint generated class is always authoritative. It is no longer necessary to explicitly check it.")
+	UE_DEPRECATED(4.22, "The minimum UE4 object version implies the Blueprint generated class is always authoritative. It is no longer necessary to explicitly check it.")
 	bool IsGeneratedClassAuthoritative()
 	{
 		return true;
