PR #5311: Extend Splitscreen types with 2 new more types (Contributed by [REDACTED_PERSON])

diff --git a/Engine/Source/Runtime/Engine/Classes/Engine/ViewportSplitScreen.h b/Engine/Source/Runtime/Engine/Classes/Engine/ViewportSplitScreen.h
--- a/Engine/Source/Runtime/Engine/Classes/Engine/ViewportSplitScreen.h
+++ b/Engine/Source/Runtime/Engine/Classes/Engine/ViewportSplitScreen.h
@@ -31,10 +31,14 @@ namespace ESplitScreenType
 		ThreePlayer_FavorBottom,
 		//3 Player vertical split
 		ThreePlayer_Vertical,
+		//3 Player horizontal split
+		ThreePlayer_Horizontal,
 		// 4 Player grid split
 		FourPlayer_Grid,
 		// 4 Player vertical split
 		FourPlayer_Vertical,
+		// 4 Player horizontal split
+		FourPlayer_Horizontal,
 
 		SplitTypeCount
 	};
diff --git a/Engine/Source/Runtime/Engine/Private/GameViewportClient.cpp b/Engine/Source/Runtime/Engine/Private/GameViewportClient.cpp
--- a/Engine/Source/Runtime/Engine/Private/GameViewportClient.cpp
+++ b/Engine/Source/Runtime/Engine/Private/GameViewportClient.cpp
@@ -187,6 +187,10 @@ UGameViewportClient::UGameViewportClient(const FObjectInitializer& ObjectInitial
 	SplitscreenInfo[ESplitScreenType::ThreePlayer_Vertical].PlayerData.Add(FPerPlayerSplitscreenData(0.333f, 1.0f, 0.333f, 0.0f));
 	SplitscreenInfo[ESplitScreenType::ThreePlayer_Vertical].PlayerData.Add(FPerPlayerSplitscreenData(0.333f, 1.0f, 0.666f, 0.0f));
 
+	SplitscreenInfo[ESplitScreenType::ThreePlayer_Horizontal].PlayerData.Add(FPerPlayerSplitscreenData(1.0f, 0.333f, 0.0f, 0.0f));
+	SplitscreenInfo[ESplitScreenType::ThreePlayer_Horizontal].PlayerData.Add(FPerPlayerSplitscreenData(1.0f, 0.333f, 0.0f, 0.333f));
+	SplitscreenInfo[ESplitScreenType::ThreePlayer_Horizontal].PlayerData.Add(FPerPlayerSplitscreenData(1.0f, 0.333f, 0.0f, 0.666f));
+
 	SplitscreenInfo[ESplitScreenType::FourPlayer_Grid].PlayerData.Add(FPerPlayerSplitscreenData(0.5f, 0.5f, 0.0f, 0.0f));
 	SplitscreenInfo[ESplitScreenType::FourPlayer_Grid].PlayerData.Add(FPerPlayerSplitscreenData(0.5f, 0.5f, 0.5f, 0.0f));
 	SplitscreenInfo[ESplitScreenType::FourPlayer_Grid].PlayerData.Add(FPerPlayerSplitscreenData(0.5f, 0.5f, 0.0f, 0.5f));
@@ -197,6 +201,11 @@ UGameViewportClient::UGameViewportClient(const FObjectInitializer& ObjectInitial
 	SplitscreenInfo[ESplitScreenType::FourPlayer_Vertical].PlayerData.Add(FPerPlayerSplitscreenData(0.25f, 1.0f, 0.5f, 0.0f));
 	SplitscreenInfo[ESplitScreenType::FourPlayer_Vertical].PlayerData.Add(FPerPlayerSplitscreenData(0.25f, 1.0f, 0.75f, 0.0f));
 
+	SplitscreenInfo[ESplitScreenType::FourPlayer_Horizontal].PlayerData.Add(FPerPlayerSplitscreenData(1.f, 0.25f, 0.0f, 0.0f));
+	SplitscreenInfo[ESplitScreenType::FourPlayer_Horizontal].PlayerData.Add(FPerPlayerSplitscreenData(1.f, 0.25f, 0.0f, 0.25f));
+	SplitscreenInfo[ESplitScreenType::FourPlayer_Horizontal].PlayerData.Add(FPerPlayerSplitscreenData(1.f, 0.25f, 0.0f, 0.5f));
+	SplitscreenInfo[ESplitScreenType::FourPlayer_Horizontal].PlayerData.Add(FPerPlayerSplitscreenData(1.f, 0.25f, 0.0f, 0.75f));
+
 	MaxSplitscreenPlayers = 4;
 	bSuppressTransitionMessage = true;
 
@@ -2019,6 +2028,10 @@ void UGameViewportClient::UpdateActiveSplitscreenType()
 				SplitType = ESplitScreenType::ThreePlayer_Vertical;
 				break;
 
+			case EThreePlayerSplitScreenType::Horizontal:
+				SplitType = ESplitScreenType::ThreePlayer_Horizontal;
+				break;
+
 			default:
 				check(0);
 			}
@@ -2036,6 +2049,10 @@ void UGameViewportClient::UpdateActiveSplitscreenType()
 				SplitType = ESplitScreenType::FourPlayer_Vertical;
 				break;
 
+			case EFourPlayerSplitScreenType::Horizontal:
+				SplitType = ESplitScreenType::FourPlayer_Horizontal;
+				break;
+
 			default:
 				check(0);
 			}
@@ -2107,6 +2124,8 @@ bool UGameViewportClient::HasTopSafeZone( int32 LocalPlayerIndex )
 
 	case ESplitScreenType::TwoPlayer_Horizontal:
 	case ESplitScreenType::ThreePlayer_FavorTop:
+	case ESplitScreenType::ThreePlayer_Horizontal:
+	case ESplitScreenType::FourPlayer_Horizontal:
 		return (LocalPlayerIndex == 0);
 
 	case ESplitScreenType::ThreePlayer_FavorBottom:
@@ -2133,7 +2152,11 @@ bool UGameViewportClient::HasBottomSafeZone( int32 LocalPlayerIndex )
 
 	case ESplitScreenType::ThreePlayer_FavorBottom:
 	case ESplitScreenType::FourPlayer_Grid:
+	case ESplitScreenType::ThreePlayer_Horizontal:
 		return (LocalPlayerIndex > 1);
+
+	case ESplitScreenType::FourPlayer_Horizontal:
+		return (LocalPlayerIndex > 2);
 	}
 
 	return false;
@@ -2145,6 +2168,8 @@ bool UGameViewportClient::HasLeftSafeZone( int32 LocalPlayerIndex )
 	{
 	case ESplitScreenType::None:
 	case ESplitScreenType::TwoPlayer_Horizontal:
+	case ESplitScreenType::ThreePlayer_Horizontal:
+	case ESplitScreenType::FourPlayer_Horizontal:
 		return true;
 
 	case ESplitScreenType::TwoPlayer_Vertical:
@@ -2169,6 +2194,8 @@ bool UGameViewportClient::HasRightSafeZone( int32 LocalPlayerIndex )
 	{
 	case ESplitScreenType::None:
 	case ESplitScreenType::TwoPlayer_Horizontal:
+	case ESplitScreenType::ThreePlayer_Horizontal:
+	case ESplitScreenType::FourPlayer_Horizontal:
 		return true;
 
 	case ESplitScreenType::TwoPlayer_Vertical:
@@ -2234,6 +2261,10 @@ void UGameViewportClient::GetPixelSizeOfScreen( float& Width, float& Height, UCa
 		Width = Canvas->ClipX * 3;
 		Height = Canvas->ClipY;
 		return;
+	case ESplitScreenType::ThreePlayer_Horizontal:
+		Width = Canvas->ClipX;
+		Height = Canvas->ClipY * 3;
+		return;
 	case ESplitScreenType::FourPlayer_Grid:
 		Width = Canvas->ClipX * 2;
 		Height = Canvas->ClipY * 2;
@@ -2241,6 +2272,11 @@ void UGameViewportClient::GetPixelSizeOfScreen( float& Width, float& Height, UCa
 	case ESplitScreenType::FourPlayer_Vertical:
 		Width = Canvas->ClipX * 4;
 		Height = Canvas->ClipY;
+		return;
+	case ESplitScreenType::FourPlayer_Horizontal:
+		Width = Canvas->ClipX;
+		Height = Canvas->ClipY * 4;
+		return;
 	}
 }
 
diff --git a/Engine/Source/Runtime/EngineSettings/Classes/GameMapsSettings.h b/Engine/Source/Runtime/EngineSettings/Classes/GameMapsSettings.h
--- a/Engine/Source/Runtime/EngineSettings/Classes/GameMapsSettings.h
+++ b/Engine/Source/Runtime/EngineSettings/Classes/GameMapsSettings.h
@@ -28,7 +28,8 @@ namespace EThreePlayerSplitScreenType
 	{
 		FavorTop,
 		FavorBottom,
-		Vertical
+		Vertical,
+		Horizontal
 	};
 }
 
@@ -36,7 +37,8 @@ UENUM()
 enum class EFourPlayerSplitScreenType : uint8
 {
 	Grid,
-	Vertical
+	Vertical,
+	Horizontal
 };
 
 /** Helper structure, used to associate GameModes with shortcut names. */
