Change FStreamableHandle::CancelHandle to correctly stop the callback from happening if the handle has already completed loading but is in the delayed callback queue
Fix it so setting s.StreamableDelegateDelayFrames to 0 actually causes the callbacks to happen immediately, this is not recommended
This fixes potential issues with delayed callbacks going off after assets were GCd due to the handle being cancelled

diff --git a/Engine/Source/Runtime/Engine/Classes/Engine/StreamableManager.h b/Engine/Source/Runtime/Engine/Classes/Engine/StreamableManager.h
--- a/Engine/Source/Runtime/Engine/Classes/Engine/StreamableManager.h
+++ b/Engine/Source/Runtime/Engine/Classes/Engine/StreamableManager.h
@@ -16,13 +16,17 @@ DECLARE_DELEGATE_OneParam(FStreamableUpdateDelegate, TSharedRef<struct FStreamab
 /** A handle to a synchronous or async load. As long as the handle is Active, loaded assets will stay in memory */
 struct ENGINE_API FStreamableHandle : public TSharedFromThis<FStreamableHandle, ESPMode::Fast>
 {
-	/** If this request has finished loading, meaning all available assets were loaded and delegate was called. If assets failed to load they will still be missing */
+	/** 
+	 * If this request has finished loading, meaning all available assets were loaded
+	 * Any assets that failed to load will still be null
+	 * This can be true before the completion callback has happened as it may be in the delayed callback queue
+	 */
 	bool HasLoadCompleted() const
 	{
 		return bLoadCompleted;
 	}
 
-	/** If this request was cancelled. Assets may still have been loaded, but delegate will not be called */
+	/** If this request was cancelled. Assets may still have been loaded, but completion delegate was not called */
 	bool WasCanceled() const
 	{
 		return bCanceled;
@@ -76,8 +80,8 @@ struct ENGINE_API FStreamableHandle : public TSharedFromThis<FStreamableHandle,
 
 	/**
 	 * Cancel a request, callable from within the manager or externally
-	 * This will immediately release the handle even if it is still in progress
-	 * This will normally stop the completion delegate from getting called, unless that delegate is already in the callback delay queue
+	 * This will immediately release the handle even if it is still in progress, and call the cancel callback if bound
+	 * This stops the completion callback from happening, even if it is in the delayed callback queue
 	 */
 	void CancelHandle();
 
@@ -120,8 +124,14 @@ struct ENGINE_API FStreamableHandle : public TSharedFromThis<FStreamableHandle,
 	/** Get the StreamableManager for this handle */
 	struct FStreamableManager* GetOwningManager() const;
 
-	/** Calls a StreamableDelegate, this may defer execution to the next frame depending on CVars */
-	static void ExecuteDelegate(const FStreamableDelegate& Delegate, TSharedPtr<FStreamableHandle> AssociatedHandle = nullptr);
+	/** 
+	 * Calls a StreamableDelegate, this will add to the delayed callback queue depending on s.StreamableDelegateDelayFrames
+	 *
+	 * @param Delegate			Primary delegate to execute
+	 * @param AssociatedHandle	Streamable handle associated with this delegate, may be null
+	 * @param CancelDelegate	If handle gets cancelled before primary delegate executes, this delegate will be called instead
+	 */
+	static void ExecuteDelegate(const FStreamableDelegate& Delegate, TSharedPtr<FStreamableHandle> AssociatedHandle = nullptr, const FStreamableDelegate& CancelDelegate = FStreamableDelegate());
 
 	/** Destructor */
 	~FStreamableHandle();
diff --git a/Engine/Source/Runtime/Engine/Private/StreamableManager.cpp b/Engine/Source/Runtime/Engine/Private/StreamableManager.cpp
--- a/Engine/Source/Runtime/Engine/Private/StreamableManager.cpp
+++ b/Engine/Source/Runtime/Engine/Private/StreamableManager.cpp
@@ -25,11 +25,44 @@ class FStreamableDelegateDelayHelper : public FTickableGameObject
 public:
 
 	/** Adds a delegate to deferred list */
-	void AddDelegate(const FStreamableDelegate& Delegate, TSharedPtr<FStreamableHandle> AssociatedHandle)
+	void AddDelegate(const FStreamableDelegate& Delegate, const FStreamableDelegate& CancelDelegate, TSharedPtr<FStreamableHandle> AssociatedHandle)
 	{
 		DataLock.Lock();
 
-		PendingDelegates.Emplace(Delegate, AssociatedHandle);
+		PendingDelegates.Emplace(Delegate, CancelDelegate, AssociatedHandle);
+
+		DataLock.Unlock();
+	}
+
+	/** Cancels delegate for handle, this will either delete the delegate or replace with the cancel delegate */
+	void CancelDelegatesForHandle(TSharedPtr<FStreamableHandle> AssociatedHandle)
+	{
+		if (!AssociatedHandle.IsValid())
+		{
+			return;
+		}
+
+		DataLock.Lock();
+
+		for (int32 DelegateIndex = 0; DelegateIndex < PendingDelegates.Num(); DelegateIndex++)
+		{
+			FPendingDelegate& PendingDelegate = PendingDelegates[DelegateIndex];
+			if (PendingDelegate.RelatedHandle == AssociatedHandle)
+			{
+				if (PendingDelegate.CancelDelegate.IsBound())
+				{
+					// Replace with cancel delegate
+					PendingDelegate.Delegate = PendingDelegate.CancelDelegate;
+					PendingDelegate.CancelDelegate.Unbind();
+				}
+				else
+				{
+					// Remove entirely
+					PendingDelegates.RemoveAt(DelegateIndex);
+					DelegateIndex--;
+				}
+			}
+		}
 
 		DataLock.Unlock();
 	}
@@ -61,7 +94,7 @@ public:
 			if (--PendingDelegates[DelegateIndex].DelayFrames <= 0)
 			{
 				// Add to call array and remove from tracking one
-				DelegatesToCall.Emplace(PendingDelegates[DelegateIndex].Delegate, PendingDelegates[DelegateIndex].RelatedHandle);
+				DelegatesToCall.Emplace(PendingDelegates[DelegateIndex].Delegate, FStreamableDelegate(), PendingDelegates[DelegateIndex].RelatedHandle);
 				PendingDelegates.RemoveAt(DelegateIndex--);
 			}
 		}
@@ -101,17 +134,21 @@ private:
 
 	struct FPendingDelegate
 	{
-		/** Delegate to call on next frame */
+		/** Delegate to call when frames are up */
 		FStreamableDelegate Delegate;
 
+		/** Delegate to call if this gets cancelled early */
+		FStreamableDelegate CancelDelegate;
+
 		/** Handle related to delegates, needs to keep these around to avoid things GCing before the user callback goes off. This may be null */
 		TSharedPtr<FStreamableHandle> RelatedHandle;
 
 		/** Frames left to delay */
 		int32 DelayFrames;
 
-		FPendingDelegate(const FStreamableDelegate& InDelegate, TSharedPtr<FStreamableHandle> InHandle)
+		FPendingDelegate(const FStreamableDelegate& InDelegate, const FStreamableDelegate& InCancelDelegate, TSharedPtr<FStreamableHandle> InHandle)
 			: Delegate(InDelegate)
+			, CancelDelegate(InCancelDelegate)
 			, RelatedHandle(InHandle)
 			, DelayFrames(GStreamableDelegateDelayFrames)
 		{}
@@ -347,22 +384,32 @@ void FStreamableHandle::CancelHandle()
 {
 	check(IsInGameThread());
 
-	if (bReleased || bCanceled || !OwningManager)
+	if (bCanceled || !OwningManager)
 	{
 		// Too late to cancel it
 		return;
 	}
 
-	if (bLoadCompleted)
+	TSharedRef<FStreamableHandle> SharedThis = AsShared();
+	if (bLoadCompleted || bReleased)
 	{
-		ReleaseHandle();
+		// Cancel if it's in the pending queue
+		if (StreamableDelegateDelayHelper)
+		{
+			StreamableDelegateDelayHelper->CancelDelegatesForHandle(SharedThis);
+		}
+
+		if (!bReleased)
+		{
+			ReleaseHandle();
+		}
+
+		bCanceled = true;
 		return;
 	}
 
 	bCanceled = true;
 
-	TSharedRef<FStreamableHandle> SharedThis = AsShared();
-
 	ExecuteDelegate(CancelDelegate, SharedThis);
 	UnbindDelegates();
 
@@ -483,7 +530,7 @@ void FStreamableHandle::CompleteLoad()
 	{
 		bLoadCompleted = true;
 
-		ExecuteDelegate(CompleteDelegate, AsShared());
+		ExecuteDelegate(CompleteDelegate, AsShared(), CancelDelegate);
 		UnbindDelegates();
 
 		if (ParentHandles.Num() > 0)
@@ -599,16 +646,25 @@ void FStreamableHandle::AsyncLoadCallbackWrapper(const FName& PackageName, UPack
 	}
 }
 
-void FStreamableHandle::ExecuteDelegate(const FStreamableDelegate& Delegate, TSharedPtr<FStreamableHandle> AssociatedHandle)
+void FStreamableHandle::ExecuteDelegate(const FStreamableDelegate& Delegate, TSharedPtr<FStreamableHandle> AssociatedHandle, const FStreamableDelegate& CancelDelegate)
 {
 	if (Delegate.IsBound())
 	{
-		if (!StreamableDelegateDelayHelper)
+		if (GStreamableDelegateDelayFrames == 0)
 		{
-			StreamableDelegateDelayHelper = new FStreamableDelegateDelayHelper;
+			// Execute it immediately
+			Delegate.Execute();
 		}
+		else
+		{
+			// Add to execution queue for next tick
+			if (!StreamableDelegateDelayHelper)
+			{
+				StreamableDelegateDelayHelper = new FStreamableDelegateDelayHelper;
+			}
 
-		StreamableDelegateDelayHelper->AddDelegate(Delegate, AssociatedHandle);
+			StreamableDelegateDelayHelper->AddDelegate(Delegate, CancelDelegate, AssociatedHandle);
+		}
 	}
 }
 
