(4.21.2) When requesting a streaming level validate that it is not going to result in a duplicate level being "loaded" invalidating the 1:1 relationship between ULevelStreaming and the level it manages. If it is violating the relationship, treat the streaming level as failed to load.
#rnx

diff --git a/Engine/Source/Runtime/Engine/Private/LevelStreaming.cpp b/Engine/Source/Runtime/Engine/Private/LevelStreaming.cpp
--- a/Engine/Source/Runtime/Engine/Private/LevelStreaming.cpp
+++ b/Engine/Source/Runtime/Engine/Private/LevelStreaming.cpp
@@ -851,6 +851,38 @@ bool ULevelStreaming::RequestLevel(UWorld* PersistentWorld, bool bAllowLevelLoad
 		return false;
 	}
 
+	auto ValidateUniqueLevel = [this, PersistentWorld]()
+	{
+		for (ULevelStreaming* OtherLevel : PersistentWorld->GetStreamingLevels())
+		{
+			if (OtherLevel == nullptr || OtherLevel == this)
+			{
+				continue;
+			}
+
+			const ECurrentState OtherState = OtherLevel->GetCurrentState();
+			if (OtherState == ECurrentState::FailedToLoad || OtherState == ECurrentState::Removed || (OtherState == ECurrentState::Unloaded && (OtherLevel->TargetState == ETargetState::Unloaded || OtherLevel->TargetState == ETargetState::UnloadedAndRemoved)))
+			{
+				// If the other level isn't loaded or in the process of being loaded we don't need to consider it
+				continue;
+			}
+
+			if (OtherLevel->WorldAsset == WorldAsset)
+			{
+				UE_LOG(LogLevelStreaming, Warning, TEXT("Streaming Level '%s' uses same destination for level ('%s') as '%s'. Level cannot be loaded again and this StreamingLevel will be flagged as failed to load."), *GetPathName(), *WorldAsset.GetLongPackageName(), *OtherLevel->GetPathName());
+				return false;
+			}
+		}
+
+		return true;
+	};
+
+	if (!ValidateUniqueLevel())
+	{
+		CurrentState = ECurrentState::FailedToLoad;
+		return false;
+	}
+
 	EPackageFlags PackageFlags = PKG_ContainsMap;
 	int32 PIEInstanceID = INDEX_NONE;
 
