PR #5326: Add UKismetSystemLibrary::GetConsoleVariableBoolValue() (Contributed by [REDACTED_PERSON])

diff --git a/Engine/Source/Runtime/Engine/Classes/Kismet/KismetSystemLibrary.h b/Engine/Source/Runtime/Engine/Classes/Kismet/KismetSystemLibrary.h
--- a/Engine/Source/Runtime/Engine/Classes/Kismet/KismetSystemLibrary.h
+++ b/Engine/Source/Runtime/Engine/Classes/Kismet/KismetSystemLibrary.h
@@ -384,8 +384,8 @@ class ENGINE_API UKismetSystemLibrary : public UBlueprintFunctionLibrary
 	 * @param	VariableName	Name of the console variable to find.
 	 * @return	The value if found, 0 otherwise.
 	 */
-	UFUNCTION(BlueprintCallable, Category="Development",meta=(WorldContext="WorldContextObject"))
-	static float GetConsoleVariableFloatValue(UObject* WorldContextObject, const FString& VariableName);
+	UFUNCTION(BlueprintCallable, Category="Development")
+	static float GetConsoleVariableFloatValue(const FString& VariableName);
 
 	/**
 	 * Attempts to retrieve the value of the specified integer console variable, if it exists.
@@ -393,8 +393,17 @@ class ENGINE_API UKismetSystemLibrary : public UBlueprintFunctionLibrary
 	 * @param	VariableName	Name of the console variable to find.
 	 * @return	The value if found, 0 otherwise.
 	 */
-	UFUNCTION(BlueprintCallable, Category="Development",meta=(WorldContext="WorldContextObject"))
-	static int32 GetConsoleVariableIntValue(UObject* WorldContextObject, const FString& VariableName);
+	UFUNCTION(BlueprintCallable, Category="Development")
+	static int32 GetConsoleVariableIntValue(const FString& VariableName);
+
+	/**
+	 * Evaluates, if it exists, whether the specified integer console variable has a non-zero value (true) or not (false).
+	 *
+	 * @param	VariableName	Name of the console variable to find.
+	 * @return	True if found and has a non-zero value, false otherwise.
+	 */
+	UFUNCTION(BlueprintCallable, Category="Development")
+	static bool GetConsoleVariableBoolValue(const FString& VariableName);
 
 	/** 
 	 *	Exit the current game 
diff --git a/Engine/Source/Runtime/Engine/Private/KismetSystemLibrary.cpp b/Engine/Source/Runtime/Engine/Private/KismetSystemLibrary.cpp
--- a/Engine/Source/Runtime/Engine/Private/KismetSystemLibrary.cpp
+++ b/Engine/Source/Runtime/Engine/Private/KismetSystemLibrary.cpp
@@ -308,7 +308,7 @@ void UKismetSystemLibrary::ExecuteConsoleCommand(UObject* WorldContextObject, co
 	}
 }
 
-float UKismetSystemLibrary::GetConsoleVariableFloatValue(UObject* WorldContextObject, const FString& VariableName)
+float UKismetSystemLibrary::GetConsoleVariableFloatValue(const FString& VariableName)
 {
 	float Value = 0.0f;
 
@@ -325,7 +325,7 @@ float UKismetSystemLibrary::GetConsoleVariableFloatValue(UObject* WorldContextOb
 	return Value;
 }
 
-int32 UKismetSystemLibrary::GetConsoleVariableIntValue(UObject* WorldContextObject, const FString& VariableName)
+int32 UKismetSystemLibrary::GetConsoleVariableIntValue(const FString& VariableName)
 {
 	int32 Value = 0;
 
@@ -342,7 +342,10 @@ int32 UKismetSystemLibrary::GetConsoleVariableIntValue(UObject* WorldContextObje
 	return Value;
 }
 
-
+bool UKismetSystemLibrary::GetConsoleVariableBoolValue(const FString& VariableName)
+{
+	return (GetConsoleVariableIntValue(VariableName) != 0);
+}
 
 
 void UKismetSystemLibrary::QuitGame(UObject* WorldContextObject, class APlayerController* SpecificPlayer, TEnumAsByte<EQuitPreference::Type> QuitPreference, bool bIgnorePlatformRestrictions)
