Fix a crash that can occur on editor load of a child Blueprint class after having previously swapped component names in the parent Blueprint class.

#rn fix blueprints

diff --git a/Engine/Source/Runtime/Engine/Private/InheritableComponentHandler.cpp b/Engine/Source/Runtime/Engine/Private/InheritableComponentHandler.cpp
--- a/Engine/Source/Runtime/Engine/Private/InheritableComponentHandler.cpp
+++ b/Engine/Source/Runtime/Engine/Private/InheritableComponentHandler.cpp
@@ -467,29 +467,27 @@ const FComponentOverrideRecord* UInheritableComponentHandler::FindRecord(const F
 
 void UInheritableComponentHandler::FixComponentTemplateName(UActorComponent* ComponentTemplate, const FString& NewName)
 {
-	// Override template names were not previously kept in sync w/ past node rename operations. Thus, we need to check for
-	// and correct other (stale) template names. Otherwise, these could collide with the one we're trying to correct here.
-	for (int32 Index = 0; Index < Records.Num(); ++Index)
+	// Look for a collision with the template we're trying to rename here. It's possible that names were swapped on the
+	// original component template objects that were inherited from the associated Blueprint's parent class, for example.
+	FComponentOverrideRecord* MatchingRecord = Records.FindByPredicate([ComponentTemplate, NewName](FComponentOverrideRecord& Record)
 	{
-		FComponentOverrideRecord& Record = Records[Index];
 		if (Record.ComponentTemplate && Record.ComponentTemplate != ComponentTemplate && Record.ComponentTemplate->GetName() == NewName)
 		{
-			if (UActorComponent* OriginalTemplate = Record.ComponentKey.GetOriginalTemplate())
-			{
-				if (OriginalTemplate->GetName() != Record.ComponentTemplate->GetName())
-				{
-					// Recursively fix up this record's component template name first to also match its original template, which will then free up the name.
-					FixComponentTemplateName(Record.ComponentTemplate, OriginalTemplate->GetName());
-				}
-			}
-
-			// There should only be at most one collision, so we'll stop looking now.
-			break;
+			const UActorComponent* OriginalTemplate = Record.ComponentKey.GetOriginalTemplate();
+			return ensureMsgf(OriginalTemplate && OriginalTemplate->GetName() != Record.ComponentTemplate->GetName(),
+				TEXT("Found a collision with an existing override record, but its associated template object is either invalid or already matches its inherited template's name (%s). This is unexpected."), *NewName);
 		}
-	}
 
-	// Precondition: There are no other objects in the same scope with this name.
-	check(!FindObjectWithOuter(ComponentTemplate->GetOuter(), nullptr, FName(*NewName)));
+		return false;
+	});
+
+	// If we found a collision, temporarily rename the associated template object to something unique so that it no longer
+	// collides with the one we're trying to correct here. This will be fixed up when we later encounter this record during
+	// PostLoad() validation and see that it still doesn't match its original template name.
+	if (MatchingRecord)
+	{
+		MatchingRecord->ComponentTemplate->Rename(nullptr, nullptr, REN_DontCreateRedirectors | REN_ForceNoResetLoaders);
+	}
 
 	// Now that we're sure there are no collisions with other records, we can safely rename this one to its new name.
 	ComponentTemplate->Rename(*NewName, nullptr, REN_DontCreateRedirectors | REN_ForceNoResetLoaders);
