Sequencer: Fix animations continuing to play - make sure montage play state is set to sequencer play state.

diff --git a/Engine/Source/Runtime/Engine/Classes/Animation/AnimMontage.h b/Engine/Source/Runtime/Engine/Classes/Animation/AnimMontage.h
--- a/Engine/Source/Runtime/Engine/Classes/Animation/AnimMontage.h
+++ b/Engine/Source/Runtime/Engine/Classes/Animation/AnimMontage.h
@@ -522,8 +522,8 @@ public:
 	ENGINE_API static UAnimMontage* SetMatineeAnimPositionInner(FName SlotName, USkeletalMeshComponent* SkeletalMeshComponent, UAnimSequenceBase* InAnimSequence, float InPosition, bool bLooping);
 	ENGINE_API static UAnimMontage* PreviewMatineeSetAnimPositionInner(FName SlotName, USkeletalMeshComponent* SkeletalMeshComponent, UAnimSequenceBase* InAnimSequence, float InPosition, bool bLooping, bool bFireNotifies, float DeltaTime);
 	/** static functions that are used by sequencer montage support*/
-	ENGINE_API static UAnimMontage* SetSequencerMontagePosition(FName SlotName, USkeletalMeshComponent* SkeletalMeshComponent, int32& InOutInstanceId, UAnimSequenceBase* InAnimSequence, float InPosition, float Weight, bool bLooping);
-	ENGINE_API static UAnimMontage* PreviewSequencerMontagePosition(FName SlotName, USkeletalMeshComponent* SkeletalMeshComponent, int32& InOutInstanceId, UAnimSequenceBase* InAnimSequence, float InPosition, float Weight, bool bLooping, bool bFireNotifies);
+	ENGINE_API static UAnimMontage* SetSequencerMontagePosition(FName SlotName, USkeletalMeshComponent* SkeletalMeshComponent, int32& InOutInstanceId, UAnimSequenceBase* InAnimSequence, float InPosition, float Weight, bool bLooping, bool bPlaying);
+	ENGINE_API static UAnimMontage* PreviewSequencerMontagePosition(FName SlotName, USkeletalMeshComponent* SkeletalMeshComponent, int32& InOutInstanceId, UAnimSequenceBase* InAnimSequence, float InPosition, float Weight, bool bLooping, bool bFireNotifies, bool bPlaying);
 private:
 	static UAnimMontage* InitializeMatineeControl(FName SlotName, USkeletalMeshComponent* SkeletalMeshComponent, UAnimSequenceBase* InAnimSequence, bool bLooping);
 };
diff --git a/Engine/Source/Runtime/Engine/Private/Animation/AnimMontage.cpp b/Engine/Source/Runtime/Engine/Private/Animation/AnimMontage.cpp
--- a/Engine/Source/Runtime/Engine/Private/Animation/AnimMontage.cpp
+++ b/Engine/Source/Runtime/Engine/Private/Animation/AnimMontage.cpp
@@ -2671,7 +2671,7 @@ UAnimMontage* FAnimMontageInstance::PreviewMatineeSetAnimPositionInner(FName Slo
 	return PlayingMontage;
 }
 
-UAnimMontage* FAnimMontageInstance::SetSequencerMontagePosition(FName SlotName, USkeletalMeshComponent* SkeletalMeshComponent, int32& InOutInstanceId, UAnimSequenceBase* InAnimSequence, float InPosition, float Weight, bool bLooping)
+UAnimMontage* FAnimMontageInstance::SetSequencerMontagePosition(FName SlotName, USkeletalMeshComponent* SkeletalMeshComponent, int32& InOutInstanceId, UAnimSequenceBase* InAnimSequence, float InPosition, float Weight, bool bLooping, bool bPlaying)
 {
 	UAnimInstance* AnimInst = SkeletalMeshComponent->GetAnimInstance();
 	if (AnimInst)
@@ -2699,6 +2699,7 @@ UAnimMontage* FAnimMontageInstance::SetSequencerMontagePosition(FName SlotName,
 			MontageInstanceToUpdate->Blend.SetDesiredValue(Weight);
 			MontageInstanceToUpdate->Blend.SetAlpha(Weight);
 			MontageInstanceToUpdate->SetPosition(InPosition);
+			MontageInstanceToUpdate->bPlaying = bPlaying;
 			return PlayingMontage;
 		}
 	}
@@ -2710,7 +2711,7 @@ UAnimMontage* FAnimMontageInstance::SetSequencerMontagePosition(FName SlotName,
 	return nullptr;
 }
 
-UAnimMontage* FAnimMontageInstance::PreviewSequencerMontagePosition(FName SlotName, USkeletalMeshComponent* SkeletalMeshComponent, int32& InOutInstanceId, UAnimSequenceBase* InAnimSequence, float InPosition, float Weight, bool bLooping, bool bFireNotifies)
+UAnimMontage* FAnimMontageInstance::PreviewSequencerMontagePosition(FName SlotName, USkeletalMeshComponent* SkeletalMeshComponent, int32& InOutInstanceId, UAnimSequenceBase* InAnimSequence, float InPosition, float Weight, bool bLooping, bool bFireNotifies, bool bPlaying)
 {
 	UAnimInstance* AnimInst = SkeletalMeshComponent->GetAnimInstance();
 	if (AnimInst)
@@ -2718,7 +2719,7 @@ UAnimMontage* FAnimMontageInstance::PreviewSequencerMontagePosition(FName SlotNa
 		FAnimMontageInstance* MontageInstanceToUpdate = AnimInst->GetMontageInstanceForID(InOutInstanceId);
 		float PreviousPosition = (MontageInstanceToUpdate) ? MontageInstanceToUpdate->GetPosition() : InPosition;
 
-		UAnimMontage* PlayingMontage = SetSequencerMontagePosition(SlotName, SkeletalMeshComponent, InOutInstanceId, InAnimSequence, InPosition, Weight, bLooping);
+		UAnimMontage* PlayingMontage = SetSequencerMontagePosition(SlotName, SkeletalMeshComponent, InOutInstanceId, InAnimSequence, InPosition, Weight, bLooping, bPlaying);
 		if (PlayingMontage)
 		{
 			// we have to get it again in case if this is new
diff --git a/Engine/Source/Runtime/MovieSceneTracks/Private/Evaluation/MovieSceneSkeletalAnimationTemplate.cpp b/Engine/Source/Runtime/MovieSceneTracks/Private/Evaluation/MovieSceneSkeletalAnimationTemplate.cpp
--- a/Engine/Source/Runtime/MovieSceneTracks/Private/Evaluation/MovieSceneSkeletalAnimationTemplate.cpp
+++ b/Engine/Source/Runtime/MovieSceneTracks/Private/Evaluation/MovieSceneSkeletalAnimationTemplate.cpp
@@ -287,7 +287,7 @@ namespace MovieScene
 				FMontagePlayerPerSectionData* SectionData = MontageData.Find(Section);
 
 				int32 InstanceId = (SectionData) ? SectionData->MontageInstanceId : INDEX_NONE;
-				TWeakObjectPtr<UAnimMontage> Montage = FAnimMontageInstance::SetSequencerMontagePosition(SlotName, SkeletalMeshComponent, InstanceId, InAnimSequence, InPosition, Weight, bLooping);
+				TWeakObjectPtr<UAnimMontage> Montage = FAnimMontageInstance::SetSequencerMontagePosition(SlotName, SkeletalMeshComponent, InstanceId, InAnimSequence, InPosition, Weight, bLooping, bPlaying);
 
 				if (Montage.IsValid())
 				{
@@ -327,7 +327,7 @@ namespace MovieScene
 				FMontagePlayerPerSectionData* SectionData = MontageData.Find(Section);
 
 				int32 InstanceId = (SectionData)? SectionData->MontageInstanceId : INDEX_NONE;
-				TWeakObjectPtr<UAnimMontage> Montage = FAnimMontageInstance::PreviewSequencerMontagePosition(SlotName, SkeletalMeshComponent, InstanceId, InAnimSequence, InPosition, Weight, bLooping, bFireNotifies);
+				TWeakObjectPtr<UAnimMontage> Montage = FAnimMontageInstance::PreviewSequencerMontagePosition(SlotName, SkeletalMeshComponent, InstanceId, InAnimSequence, InPosition, Weight, bLooping, bFireNotifies, bPlaying);
 
 				if (Montage.IsValid())
 				{
