Merging //UE4/Dev-Main to Dev-Sequencer (//UE4/Dev-Sequencer)

diff --git a/Engine/Plugins/Media/MediaPlayerEditor/Source/MediaPlayerEditor/Public/Widgets/SMediaPermutationsSelector.h b/Engine/Plugins/Media/MediaPlayerEditor/Source/MediaPlayerEditor/Public/Widgets/SMediaPermutationsSelector.h
--- a/Engine/Plugins/Media/MediaPlayerEditor/Source/MediaPlayerEditor/Public/Widgets/SMediaPermutationsSelector.h
+++ b/Engine/Plugins/Media/MediaPlayerEditor/Source/MediaPlayerEditor/Public/Widgets/SMediaPermutationsSelector.h
@@ -122,7 +122,7 @@ public:
 		: _ColumnHeight(200)
 	{}
 		/** Array of columns */
-		SLATE_SUPPORTS_SLOT_WITH_ARGS(typename FColumn)
+		SLATE_SUPPORTS_SLOT_WITH_ARGS(FColumn)
 		/** Array of data items that we are displaying */
 		SLATE_ARGUMENT(TArray<ItemType>, PermutationsSource)
 		/** Default selected item in ItemsSource */
diff --git a/Engine/Plugins/Media/MediaPlayerEditor/Source/MediaPlayerEditor/Public/Widgets/SMediaPermutationsSelector.inl b/Engine/Plugins/Media/MediaPlayerEditor/Source/MediaPlayerEditor/Public/Widgets/SMediaPermutationsSelector.inl
--- a/Engine/Plugins/Media/MediaPlayerEditor/Source/MediaPlayerEditor/Public/Widgets/SMediaPermutationsSelector.inl
+++ b/Engine/Plugins/Media/MediaPlayerEditor/Source/MediaPlayerEditor/Public/Widgets/SMediaPermutationsSelector.inl
@@ -8,7 +8,7 @@
 template <typename ItemType, typename ItemBuilder>
 typename SMediaPermutationsSelector<ItemType, ItemBuilder>::FColumn::FArguments SMediaPermutationsSelector<ItemType, ItemBuilder>::Column(const FName& InColumnName)
 {
-	FColumn::FArguments NewArgs;
+	typename FColumn::FArguments NewArgs;
 	NewArgs._ColumnName = InColumnName;
 	return NewArgs;
 }
diff --git a/Engine/Source/Runtime/Core/Public/GenericPlatform/OSAllocationPool.h b/Engine/Source/Runtime/Core/Public/GenericPlatform/OSAllocationPool.h
--- a/Engine/Source/Runtime/Core/Public/GenericPlatform/OSAllocationPool.h
+++ b/Engine/Source/Runtime/Core/Public/GenericPlatform/OSAllocationPool.h
@@ -307,6 +307,6 @@ public:
 
 	void PrintDebugInfo()
 	{
-		printf("BlockSize: %llu NumAllocated/TotalBlocks = %lu/%lu\n", (uint64)BlockSize, (NumBlocks - NumFreeBlocks), NumBlocks);
+		printf("BlockSize: %llu NumAllocated/TotalBlocks = %llu/%llu\n", (uint64)BlockSize, (uint64)(NumBlocks - NumFreeBlocks), (uint64)NumBlocks);
 	}
 };
diff --git a/Engine/Source/Runtime/CoreUObject/Private/UObject/EnumProperty.cpp b/Engine/Source/Runtime/CoreUObject/Private/UObject/EnumProperty.cpp
--- a/Engine/Source/Runtime/CoreUObject/Private/UObject/EnumProperty.cpp
+++ b/Engine/Source/Runtime/CoreUObject/Private/UObject/EnumProperty.cpp
@@ -78,6 +78,7 @@ void UEnumProperty::SerializeItem(FStructuredArchive::FSlot Slot, void* Value, v
 
 	if (Enum && UnderlyingArchive.UseToResolveEnumerators())
 	{
+		Slot.EnterStream();
 		int64 IntValue = UnderlyingProp->GetSignedIntPropertyValue(Value);
 		int64 ResolvedIndex = Enum->ResolveEnumerator(UnderlyingArchive, IntValue);
 		UnderlyingProp->SetIntPropertyValue(Value, ResolvedIndex);
diff --git a/Engine/Source/Runtime/CoreUObject/Private/UObject/PropertyByte.cpp b/Engine/Source/Runtime/CoreUObject/Private/UObject/PropertyByte.cpp
--- a/Engine/Source/Runtime/CoreUObject/Private/UObject/PropertyByte.cpp
+++ b/Engine/Source/Runtime/CoreUObject/Private/UObject/PropertyByte.cpp
@@ -25,6 +25,7 @@ void UByteProperty::SerializeItem( FStructuredArchive::FSlot Slot, void* Value,
 
 	if(Enum && UnderlyingArchive.UseToResolveEnumerators())
 	{
+		Slot.EnterStream();
 		 const int64 ResolvedIndex = Enum->ResolveEnumerator(UnderlyingArchive, *(uint8*)Value);
 		 *(uint8*)Value = static_cast<uint8>(ResolvedIndex);
 		 return;
diff --git a/Engine/Source/Runtime/Engine/Private/Components/ReflectionCaptureComponent.cpp b/Engine/Source/Runtime/Engine/Private/Components/ReflectionCaptureComponent.cpp
--- a/Engine/Source/Runtime/Engine/Private/Components/ReflectionCaptureComponent.cpp
+++ b/Engine/Source/Runtime/Engine/Private/Components/ReflectionCaptureComponent.cpp
@@ -1039,6 +1039,9 @@ void UReflectionCaptureComponent::UpdateReflectionCaptureContents(UWorld* WorldT
 		// Note: this will also prevent uploads of cubemaps from DDC, which is unintentional
 		&& (GShaderCompilingManager == NULL || !GShaderCompilingManager->IsCompiling()))
 	{
+		//guarantee that all render proxies are up to date before kicking off this render
+		WorldToUpdate->SendAllEndOfFrameUpdates();
+
 		TArray<UReflectionCaptureComponent*> WorldCombinedCaptures;
 
 		for (int32 CaptureIndex = ReflectionCapturesToUpdate.Num() - 1; CaptureIndex >= 0; CaptureIndex--)
diff --git a/Engine/Source/Runtime/SlateCore/Public/Widgets/DeclarativeSyntaxSupport.h b/Engine/Source/Runtime/SlateCore/Public/Widgets/DeclarativeSyntaxSupport.h
--- a/Engine/Source/Runtime/SlateCore/Public/Widgets/DeclarativeSyntaxSupport.h
+++ b/Engine/Source/Runtime/SlateCore/Public/Widgets/DeclarativeSyntaxSupport.h
@@ -382,7 +382,7 @@ template<typename WidgetType> struct TSlateBaseNamedArgs;
  */
 #define SLATE_SUPPORTS_SLOT_WITH_ARGS( SlotType ) \
 	TArray< SlotType* > Slots; \
-	WidgetArgsType& operator + (const SlotType::FArguments& ArgumentsForNewSlot) \
+	WidgetArgsType& operator + (const typename SlotType::FArguments& ArgumentsForNewSlot) \
 		{ \
 			Slots.Add( new SlotType( ArgumentsForNewSlot ) ); \
 			return *this; \
