Added null check to resolve static analysis warning (in practice all UBlueprint objects are part of a UPackage)

diff --git a/Engine/Source/Runtime/Engine/Private/Blueprint.cpp b/Engine/Source/Runtime/Engine/Private/Blueprint.cpp
--- a/Engine/Source/Runtime/Engine/Private/Blueprint.cpp
+++ b/Engine/Source/Runtime/Engine/Private/Blueprint.cpp
@@ -840,7 +840,7 @@ bool UBlueprint::SupportsNativization(FText* OutReason) const
 			Asset = Asset->GetOuter();
 		}
 
-		const IBlueprintAssetHandler* Handler = FBlueprintAssetHandler::Get().FindHandler(Asset->GetClass());
+		const IBlueprintAssetHandler* Handler = Asset ? FBlueprintAssetHandler::Get().FindHandler(Asset->GetClass()) : nullptr;
 		if (Handler && !Handler->SupportsNativization(Asset, this, OutReason))
 		{
 			return false;
