Fix crash in replay checkpoint recording if the streamer destroys the checkpoint archive during the flush

diff --git a/Engine/Source/Runtime/Engine/Private/DemoNetDriver.cpp b/Engine/Source/Runtime/Engine/Private/DemoNetDriver.cpp
--- a/Engine/Source/Runtime/Engine/Private/DemoNetDriver.cpp
+++ b/Engine/Source/Runtime/Engine/Private/DemoNetDriver.cpp
@@ -1745,8 +1745,8 @@ void UDemoNetDriver::SaveCheckpoint()
 	// Do the first checkpoint tick now if we're not amortizing
 	if (GetCheckpointSaveMaxMSPerFrame() <= 0.0f)
 	{
-	TickCheckpoint();
-}
+		TickCheckpoint();
+	}
 }
 
 class FRepActorsCheckpointParams
@@ -1785,203 +1785,206 @@ void UDemoNetDriver::TickCheckpoint()
 		return;
 	}
 
-	FScopedForceUnicodeInArchive ScopedUnicodeSerialization(*CheckpointArchive);
-
 	FRepActorsCheckpointParams Params
 	{
 		FPlatformTime::Seconds(),
 		(double)GetCheckpointSaveMaxMSPerFrame() / 1000
 	};
 
-	UDemoNetConnection* ClientConnection = CastChecked<UDemoNetConnection>(ClientConnections[0]);
+	bool bExecuteNextState = true;
+	double CurrentTime = Params.StartCheckpointTime;
 
-	CheckpointSaveContext.TotalCheckpointSaveFrames++;
+	{
+		FScopedForceUnicodeInArchive ScopedUnicodeSerialization(*CheckpointArchive);
 
-	FlushNetChecked(*ClientConnection);
+		UDemoNetConnection* ClientConnection = CastChecked<UDemoNetConnection>(ClientConnections[0]);
 
-	UPackageMapClient* PackageMapClient = ( ( UPackageMapClient* )ClientConnection->PackageMap );
+		CheckpointSaveContext.TotalCheckpointSaveFrames++;
 
-	// Save package map ack status in case we export stuff during the checkpoint (so we can restore the connection back to what it was before we saved the checkpoint)
-	PackageMapClient->OverridePackageMapExportAckStatus( &CheckpointSaveContext.CheckpointAckState );
+		FlushNetChecked(*ClientConnection);
 
-	bool bExecuteNextState = true;
-	double CurrentTime = Params.StartCheckpointTime;
-	while (bExecuteNextState && (CheckpointSaveContext.CheckpointSaveState != ECheckpointSaveState_Finalize) && !(Params.CheckpointMaxUploadTimePerFrame > 0 && CurrentTime - Params.StartCheckpointTime > Params.CheckpointMaxUploadTimePerFrame))
-	{
-		switch (CheckpointSaveContext.CheckpointSaveState)
+		UPackageMapClient* PackageMapClient = ( ( UPackageMapClient* )ClientConnection->PackageMap );
+
+		// Save package map ack status in case we export stuff during the checkpoint (so we can restore the connection back to what it was before we saved the checkpoint)
+		PackageMapClient->OverridePackageMapExportAckStatus( &CheckpointSaveContext.CheckpointAckState );
+
+		while (bExecuteNextState && (CheckpointSaveContext.CheckpointSaveState != ECheckpointSaveState_Finalize) && !(Params.CheckpointMaxUploadTimePerFrame > 0 && CurrentTime - Params.StartCheckpointTime > Params.CheckpointMaxUploadTimePerFrame))
 		{
-			case ECheckpointSaveState_ProcessCheckpointActors:
+			switch (CheckpointSaveContext.CheckpointSaveState)
 			{
-				SCOPED_NAMED_EVENT(UDemoNetDriver_ProcessCheckpointActors, FColor::Green);
+				case ECheckpointSaveState_ProcessCheckpointActors:
+				{
+					SCOPED_NAMED_EVENT(UDemoNetDriver_ProcessCheckpointActors, FColor::Green);
 
-				// Save the replicated server time so we can restore it after the checkpoint has been serialized.
-				// This preserves the existing behavior and prevents clients from receiving updated server time
-				// more often than the normal update rate.
-				AGameStateBase* const GameState = World != nullptr ? World->GetGameState() : nullptr;
+					// Save the replicated server time so we can restore it after the checkpoint has been serialized.
+					// This preserves the existing behavior and prevents clients from receiving updated server time
+					// more often than the normal update rate.
+					AGameStateBase* const GameState = World != nullptr ? World->GetGameState() : nullptr;
 
-				const float SavedReplicatedServerTimeSeconds = GameState ? GameState->ReplicatedWorldTimeSeconds : -1.0f;
+					const float SavedReplicatedServerTimeSeconds = GameState ? GameState->ReplicatedWorldTimeSeconds : -1.0f;
 
-				// Normally AGameStateBase::ReplicatedWorldTimeSeconds is only updated periodically,
-				// but we want to make sure it's accurate for the checkpoint.
-				if ( GameState )
-				{
-					GameState->UpdateServerTimeSeconds();
-				}
+					// Normally AGameStateBase::ReplicatedWorldTimeSeconds is only updated periodically,
+					// but we want to make sure it's accurate for the checkpoint.
+					if ( GameState )
+					{
+						GameState->UpdateServerTimeSeconds();
+					}
 
-				{
-					// Re-use the existing connection to record all properties that have changed since channels were first opened
-					// Set bResendAllDataSinceOpen to true to signify that we want to do this
-					TGuardValue<bool> ResendAllData(ClientConnection->bResendAllDataSinceOpen, true);
+					{
+						// Re-use the existing connection to record all properties that have changed since channels were first opened
+						// Set bResendAllDataSinceOpen to true to signify that we want to do this
+						TGuardValue<bool> ResendAllData(ClientConnection->bResendAllDataSinceOpen, true);
 
-					// Can't use conditionally create here, because NumActorsToProcess will be empty when HasLevelStreamingFixes is false.
-					TUniquePtr<FScopedPacketManager> PacketManager;
+						// Can't use conditionally create here, because NumActorsToProcess will be empty when HasLevelStreamingFixes is false.
+						TUniquePtr<FScopedPacketManager> PacketManager;
 
-					int32 ProcessedLevelIndex = -1;
-					const int32 UseScopedPacketManager = HasLevelStreamingFixes() ? 1 : 0;
+						int32 ProcessedLevelIndex = -1;
+						const int32 UseScopedPacketManager = HasLevelStreamingFixes() ? 1 : 0;
 
-					bool bContinue = true;
-					int32 NumActorsToReplicate = CheckpointSaveContext.PendingCheckpointActors.Num();
+						bool bContinue = true;
+						int32 NumActorsToReplicate = CheckpointSaveContext.PendingCheckpointActors.Num();
 
-					do
-					{
-						const PendingCheckPointActor Current = CheckpointSaveContext.PendingCheckpointActors.Pop();
-						AActor* Actor = Current.Actor.Get();
+						do
+						{
+							const PendingCheckPointActor Current = CheckpointSaveContext.PendingCheckpointActors.Pop();
+							AActor* Actor = Current.Actor.Get();
+
+							if (UseScopedPacketManager & (Current.LevelIndex != ProcessedLevelIndex))
+							{
+								PacketManager.Reset(new FScopedPacketManager(*ClientConnection, Current.LevelIndex + 1));
+								ProcessedLevelIndex = Current.LevelIndex;
+							}
 
-						if (UseScopedPacketManager & (Current.LevelIndex != ProcessedLevelIndex))
+							bContinue = ReplicateCheckpointActor(Actor, ClientConnection, Params);				
+						}
+						while (--NumActorsToReplicate && bContinue);
+
+						if ( GameState )
 						{
-							PacketManager.Reset(new FScopedPacketManager(*ClientConnection, Current.LevelIndex + 1));
-							ProcessedLevelIndex = Current.LevelIndex;
+							// Restore the game state's replicated world time
+							GameState->ReplicatedWorldTimeSeconds = SavedReplicatedServerTimeSeconds;
 						}
 
-						bContinue = ReplicateCheckpointActor(Actor, ClientConnection, Params);				
-					}
-					while (--NumActorsToReplicate && bContinue);
+						FlushNetChecked(*ClientConnection);
 
-					if ( GameState )
-					{
-						// Restore the game state's replicated world time
-						GameState->ReplicatedWorldTimeSeconds = SavedReplicatedServerTimeSeconds;
+						PackageMapClient->OverridePackageMapExportAckStatus( nullptr );
 					}
 
-					FlushNetChecked(*ClientConnection);
-
-					PackageMapClient->OverridePackageMapExportAckStatus( nullptr );
-				}
-
-				// We are done processing for this frame so  store the TotalCheckpointSave time here to be true to the old behavior which did not account for the	actual saving time of the check point
-				CheckpointSaveContext.TotalCheckpointReplicationTimeSeconds += ( FPlatformTime::Seconds() - Params.StartCheckpointTime );
+					// We are done processing for this frame so  store the TotalCheckpointSave time here to be true to the old behavior which did not account for the	actual saving time of the check point
+					CheckpointSaveContext.TotalCheckpointReplicationTimeSeconds += ( FPlatformTime::Seconds() - Params.StartCheckpointTime );
 
-				// if we have replicated all checkpointactors, move on to the next state
-				if ( CheckpointSaveContext.PendingCheckpointActors.Num() == 0 )
-				{
-					CheckpointSaveContext.CheckpointSaveState = ECheckpointSaveState_SerializeDeletedStartupActors;
+					// if we have replicated all checkpointactors, move on to the next state
+					if ( CheckpointSaveContext.PendingCheckpointActors.Num() == 0 )
+					{
+						CheckpointSaveContext.CheckpointSaveState = ECheckpointSaveState_SerializeDeletedStartupActors;
+					}
 				}
-			}
-			break;
+				break;
 
-			case ECheckpointSaveState_SerializeDeletedStartupActors:
-			{
-				// Postpone execution of this state if we have used to much of our alloted time, this value can be tweaked based on profiling
-				const double RequiredRatioFor_SerializeDeletedStartupActors = 0.6;
-				if ((bExecuteNextState = ShouldExecuteState(Params, CurrentTime, RequiredRatioFor_SerializeDeletedStartupActors)) == true)
+				case ECheckpointSaveState_SerializeDeletedStartupActors:
 				{
-					SCOPED_NAMED_EVENT(UDemoNetDriver_SerializeDeletedStartupActors, FColor::Green);
+					// Postpone execution of this state if we have used to much of our alloted time, this value can be tweaked based on profiling
+					const double RequiredRatioFor_SerializeDeletedStartupActors = 0.6;
+					if ((bExecuteNextState = ShouldExecuteState(Params, CurrentTime, RequiredRatioFor_SerializeDeletedStartupActors)) == true)
+					{
+						SCOPED_NAMED_EVENT(UDemoNetDriver_SerializeDeletedStartupActors, FColor::Green);
 
-					//
-					// We're done saving this checkpoint, now we need to write out all data for it.
-					//
+						//
+						// We're done saving this checkpoint, now we need to write out all data for it.
+						//
 
-					CheckpointSaveContext.bWriteCheckpointOffset = HasLevelStreamingFixes();
-					if (HasLevelStreamingFixes())
-					{
-						CheckpointSaveContext.CheckpointOffset = CheckpointArchive->Tell();
-						// We will rewrite this offset when we are done saving the checkpoint
-						*CheckpointArchive << CheckpointSaveContext.CheckpointOffset;
-					}
+						CheckpointSaveContext.bWriteCheckpointOffset = HasLevelStreamingFixes();
+						if (HasLevelStreamingFixes())
+						{
+							CheckpointSaveContext.CheckpointOffset = CheckpointArchive->Tell();
+							// We will rewrite this offset when we are done saving the checkpoint
+							*CheckpointArchive << CheckpointSaveContext.CheckpointOffset;
+						}
 
-					*CheckpointArchive << CurrentLevelIndex;
+						*CheckpointArchive << CurrentLevelIndex;
 
-					// Save deleted startup actors	
-					*CheckpointArchive << DeletedNetStartupActors;
+						// Save deleted startup actors	
+						*CheckpointArchive << DeletedNetStartupActors;
 
-					CheckpointSaveContext.CheckpointSaveState = ECheckpointSaveState_SerializeGuidCache;
+						CheckpointSaveContext.CheckpointSaveState = ECheckpointSaveState_SerializeGuidCache;
+					}
 				}
-			}
-			break;
+				break;
 
-			case ECheckpointSaveState_SerializeGuidCache:
-			{
-				// Postpone execution of this state if we have used to much of our alloted time, this value can be tweaked based on profiling
-				const double RequiredRatioFor_SerializeGuidCache = 0.8;
-				if ((bExecuteNextState = ShouldExecuteState(Params, CurrentTime, RequiredRatioFor_SerializeGuidCache)) == true)
+				case ECheckpointSaveState_SerializeGuidCache:
 				{
-					SCOPED_NAMED_EVENT(UDemoNetDriver_SerializeGuidCache, FColor::Green);
+					// Postpone execution of this state if we have used to much of our alloted time, this value can be tweaked based on profiling
+					const double RequiredRatioFor_SerializeGuidCache = 0.8;
+					if ((bExecuteNextState = ShouldExecuteState(Params, CurrentTime, RequiredRatioFor_SerializeGuidCache)) == true)
+					{
+						SCOPED_NAMED_EVENT(UDemoNetDriver_SerializeGuidCache, FColor::Green);
 
-					// Save the current guid cache
-					SerializeGuidCache( GuidCache, CheckpointArchive );
+						// Save the current guid cache
+						SerializeGuidCache( GuidCache, CheckpointArchive );
 
-					CheckpointSaveContext.CheckpointSaveState = ECheckpointSaveState_SerializeNetFieldExportGroupMap;
+						CheckpointSaveContext.CheckpointSaveState = ECheckpointSaveState_SerializeNetFieldExportGroupMap;
+					}
 				}
-			}
-			break;
+				break;
 
-			case ECheckpointSaveState_SerializeNetFieldExportGroupMap:
-			{
-				// Postpone execution of this state if we have used to much of our alloted time, this value can be tweaked based on profiling
-				const double RequiredRatioFor_SerializeNetFieldExportGroupMap = 0.6;
-				if ((bExecuteNextState = ShouldExecuteState(Params, CurrentTime, RequiredRatioFor_SerializeNetFieldExportGroupMap)) == true)
+				case ECheckpointSaveState_SerializeNetFieldExportGroupMap:
 				{
-					SCOPED_NAMED_EVENT(UDemoNetDriver_SerializeNetFieldExportGroupMap, FColor::Green);
+					// Postpone execution of this state if we have used to much of our alloted time, this value can be tweaked based on profiling
+					const double RequiredRatioFor_SerializeNetFieldExportGroupMap = 0.6;
+					if ((bExecuteNextState = ShouldExecuteState(Params, CurrentTime, RequiredRatioFor_SerializeNetFieldExportGroupMap)) == true)
+					{
+						SCOPED_NAMED_EVENT(UDemoNetDriver_SerializeNetFieldExportGroupMap, FColor::Green);
 
-					// Save the compatible rep layout map
-					PackageMapClient->SerializeNetFieldExportGroupMap( *CheckpointArchive );
+						// Save the compatible rep layout map
+						PackageMapClient->SerializeNetFieldExportGroupMap( *CheckpointArchive );
 
-					CheckpointSaveContext.CheckpointSaveState = ECheckpointSaveState_SerializeDemoFrameFromQueuedDemoPackets;
+						CheckpointSaveContext.CheckpointSaveState = ECheckpointSaveState_SerializeDemoFrameFromQueuedDemoPackets;
+					}
 				}
-			}
-			break;
+				break;
 
-			case ECheckpointSaveState_SerializeDemoFrameFromQueuedDemoPackets:
-			{
-				// Postpone execution of this state if we have used to much of our alloted time, this value can be tweaked based on profiling
-				const double RequiredRatioFor_SerializeDemoFrameFromQueuedDemoPackets = 0.8;
-				if ((bExecuteNextState = ShouldExecuteState(Params, CurrentTime, RequiredRatioFor_SerializeDemoFrameFromQueuedDemoPackets)) == true)
+				case ECheckpointSaveState_SerializeDemoFrameFromQueuedDemoPackets:
 				{
-					SCOPED_NAMED_EVENT(UDemoNetDriver_SerializeDemoFrameFromQueuedDemoPackets, FColor::Green);
-
-					// Write offset
-					if (CheckpointSaveContext.bWriteCheckpointOffset)
+					// Postpone execution of this state if we have used to much of our alloted time, this value can be tweaked based on profiling
+					const double RequiredRatioFor_SerializeDemoFrameFromQueuedDemoPackets = 0.8;
+					if ((bExecuteNextState = ShouldExecuteState(Params, CurrentTime, RequiredRatioFor_SerializeDemoFrameFromQueuedDemoPackets)) == true)
 					{
-						const FArchivePos CurrentPosition = CheckpointArchive->Tell();
-						FArchivePos Offset = CurrentPosition - (CheckpointSaveContext.CheckpointOffset + sizeof(FArchivePos));
-						CheckpointArchive->Seek(CheckpointSaveContext.CheckpointOffset);
-						*CheckpointArchive << Offset;
-						CheckpointArchive->Seek(CurrentPosition);	
-					}
+						SCOPED_NAMED_EVENT(UDemoNetDriver_SerializeDemoFrameFromQueuedDemoPackets, FColor::Green);
 
-					// Get the size of the guid data saved
-					CheckpointSaveContext.GuidCacheSize = CheckpointArchive->TotalSize();
+						// Write offset
+						if (CheckpointSaveContext.bWriteCheckpointOffset)
+						{
+							const FArchivePos CurrentPosition = CheckpointArchive->Tell();
+							FArchivePos Offset = CurrentPosition - (CheckpointSaveContext.CheckpointOffset + sizeof(FArchivePos));
+							CheckpointArchive->Seek(CheckpointSaveContext.CheckpointOffset);
+							*CheckpointArchive << Offset;
+							CheckpointArchive->Seek(CurrentPosition);	
+						}
 
-					// This will cause the entire name list to be written out again.
-					// Note, WriteDemoFrameFromQueuedDemoPackets will set this to 0 so we guard the value.
-					// This is because when checkpoint amortization is enabled, it's possible for new levels to stream
-					// in while recording a checkpoint, and we want to make sure those get written out to the normal
-					// streaming archive next frame.
-					TGuardValue<uint32> NumLevelsAddedThisFrameGuard(NumLevelsAddedThisFrame, AllLevelStatuses.Num());
+						// Get the size of the guid data saved
+						CheckpointSaveContext.GuidCacheSize = CheckpointArchive->TotalSize();
 
-					// Write out all of the queued up packets generated while saving the checkpoint
-					WriteDemoFrameFromQueuedDemoPackets( *CheckpointArchive, ClientConnection->QueuedCheckpointPackets, static_cast<float>(LastCheckpointTime) );
+						// This will cause the entire name list to be written out again.
+						// Note, WriteDemoFrameFromQueuedDemoPackets will set this to 0 so we guard the value.
+						// This is because when checkpoint amortization is enabled, it's possible for new levels to stream
+						// in while recording a checkpoint, and we want to make sure those get written out to the normal
+						// streaming archive next frame.
+						TGuardValue<uint32> NumLevelsAddedThisFrameGuard(NumLevelsAddedThisFrame, AllLevelStatuses.Num());
 
-					CheckpointSaveContext.CheckpointSaveState = ECheckpointSaveState_Finalize;
+						// Write out all of the queued up packets generated while saving the checkpoint
+						WriteDemoFrameFromQueuedDemoPackets( *CheckpointArchive, ClientConnection->QueuedCheckpointPackets, static_cast<float>(LastCheckpointTime) );
+
+						CheckpointSaveContext.CheckpointSaveState = ECheckpointSaveState_Finalize;
+					}
 				}
+				break;
+
+				default:
+					break;
 			}
-			break;
 
-			default:
-				break;
+			CurrentTime = FPlatformTime::Seconds();
 		}
-
-		CurrentTime = FPlatformTime::Seconds();
 	}
 
 	// accumulate time spent over all checkpoint ticks
