Fix missing bitfield specifier
#rnx

diff --git a/Engine/Source/Runtime/Engine/Classes/Engine/MeshMerging.h b/Engine/Source/Runtime/Engine/Classes/Engine/MeshMerging.h
--- a/Engine/Source/Runtime/Engine/Classes/Engine/MeshMerging.h
+++ b/Engine/Source/Runtime/Engine/Classes/Engine/MeshMerging.h
@@ -504,7 +504,7 @@ struct FMeshMergingSettings
 	UPROPERTY()
 	uint8 bImportVertexColors_DEPRECATED:1;
 	UPROPERTY()
-	uint8 bCalculateCorrectLODModel_DEPRECATED;
+	uint8 bCalculateCorrectLODModel_DEPRECATED:1;
 	/** Whether to export normal maps for material merging */
 	UPROPERTY()
 	uint8 bExportNormalMap_DEPRECATED:1;
