Removed unnecessary check when removing level texture streaming data in FLevelTextureManager::Remove()

diff --git a/Engine/Source/Runtime/Engine/Private/Streaming/LevelTextureManager.cpp b/Engine/Source/Runtime/Engine/Private/Streaming/LevelTextureManager.cpp
--- a/Engine/Source/Runtime/Engine/Private/Streaming/LevelTextureManager.cpp
+++ b/Engine/Source/Runtime/Engine/Private/Streaming/LevelTextureManager.cpp
@@ -33,7 +33,8 @@ void FLevelTextureManager::Remove(FRemovedTextureArray* RemovedTextures)
 		if (Component)
 		{
 			check(Component->IsValidLowLevelFast()); // Check that this component was not already destroyed.
-			check(Component->bAttachedToStreamingManagerAsStatic);  // Check that is correctly tracked
+			// Don't check as there can be duplicates in PendingComponents
+			// check(Component->bAttachedToStreamingManagerAsStatic);  
 
 			// A component can only be referenced in one level, so if it was here, we can clear the flag
 			Component->bAttachedToStreamingManagerAsStatic = false;
