Stop mysterious cooker asserts caused by a missing &= when processing whether batches of shader compiles were successful.

diff --git a/Engine/Source/Runtime/Engine/Private/ShaderCompiler/ShaderCompiler.cpp b/Engine/Source/Runtime/Engine/Private/ShaderCompiler/ShaderCompiler.cpp
--- a/Engine/Source/Runtime/Engine/Private/ShaderCompiler/ShaderCompiler.cpp
+++ b/Engine/Source/Runtime/Engine/Private/ShaderCompiler/ShaderCompiler.cpp
@@ -1983,7 +1983,7 @@ void FShaderCompilingManager::ProcessCompiledShaderMaps(
 					for (int32 Index = 0; Index < PipelineJob->StageJobs.Num(); ++Index)
 					{
 						bool bCheckSucceeded = CheckSingleJob(PipelineJob->StageJobs[Index]->GetSingleShaderJob(), MaterialsArray, Errors);
-						bSuccess = PipelineJob->StageJobs[Index]->bSucceeded && bCheckSucceeded;
+						bSuccess &= PipelineJob->StageJobs[Index]->bSucceeded && bCheckSucceeded;
 					}
 				}
 			}
