DR - added additional debug information for out of order initzialized renderdata

diff --git a/Engine/Source/Runtime/Engine/Private/StaticMesh.cpp b/Engine/Source/Runtime/Engine/Private/StaticMesh.cpp
--- a/Engine/Source/Runtime/Engine/Private/StaticMesh.cpp
+++ b/Engine/Source/Runtime/Engine/Private/StaticMesh.cpp
@@ -843,6 +843,7 @@ void FStaticMeshRenderData::InitResources(ERHIFeatureLevel::Type InFeatureLevel,
 			LODVertexFactories[LODIndex].InitResources(LODResources[LODIndex], Owner);
 		}
 	}
+	bIsInitialized = true;
 }
 
 void FStaticMeshRenderData::ReleaseResources()
diff --git a/Engine/Source/Runtime/Engine/Private/StaticMeshRender.cpp b/Engine/Source/Runtime/Engine/Private/StaticMeshRender.cpp
--- a/Engine/Source/Runtime/Engine/Private/StaticMeshRender.cpp
+++ b/Engine/Source/Runtime/Engine/Private/StaticMeshRender.cpp
@@ -132,6 +132,7 @@ FStaticMeshSceneProxy::FStaticMeshSceneProxy(UStaticMeshComponent* InComponent,
 #endif
 {
 	check(RenderData);
+	checkf(RenderData->IsInitialized(), TEXT("Uninitialized Renderdata for Mesh: %s"), *InComponent->GetStaticMesh()->GetFName().ToString());
 
 	const auto FeatureLevel = GetScene().GetFeatureLevel();
 
diff --git a/Engine/Source/Runtime/Engine/Public/StaticMeshResources.h b/Engine/Source/Runtime/Engine/Public/StaticMeshResources.h
--- a/Engine/Source/Runtime/Engine/Public/StaticMeshResources.h
+++ b/Engine/Source/Runtime/Engine/Public/StaticMeshResources.h
@@ -369,6 +369,11 @@ public:
 	/** Bounds of the renderable mesh. */
 	FBoxSphereBounds Bounds;
 
+	bool IsInitialized()
+	{
+		return bIsInitialized;
+	}
+
 	/** True if LODs share static lighting data. */
 	bool bLODsShareStaticLighting;
 
@@ -427,6 +432,7 @@ private:
 	/** Allow the editor to explicitly update section information. */
 	friend class FLevelOfDetailSettingsLayout;
 #endif // #if WITH_EDITORONLY_DATA
+	bool bIsInitialized = false;
 };
 
 /**
