TC_HDR_Compressed editor texture visualization fixes by [REDACTED_PERSON].

diff --git a/Engine/Source/Editor/TextureEditor/Private/Widgets/STextureEditorViewport.cpp b/Engine/Source/Editor/TextureEditor/Private/Widgets/STextureEditorViewport.cpp
--- a/Engine/Source/Editor/TextureEditor/Private/Widgets/STextureEditorViewport.cpp
+++ b/Engine/Source/Editor/TextureEditor/Private/Widgets/STextureEditorViewport.cpp
@@ -355,7 +355,7 @@ EVisibility STextureEditorViewport::HandleExposureBiasWidgetVisibility( ) const
 {
 	UTexture* Texture = ToolkitPtr.Pin()->GetTexture();
 
-	if ((Texture != NULL) && (Texture->CompressionSettings == TC_HDR))
+	if ((Texture != NULL) && (Texture->CompressionSettings == TC_HDR || Texture->CompressionSettings == TC_HDR_Compressed))
 	{
 		return EVisibility::Visible;
 	}
diff --git a/Engine/Source/Runtime/Engine/Private/Texture2D.cpp b/Engine/Source/Runtime/Engine/Private/Texture2D.cpp
--- a/Engine/Source/Runtime/Engine/Private/Texture2D.cpp
+++ b/Engine/Source/Runtime/Engine/Private/Texture2D.cpp
@@ -1169,7 +1169,7 @@ FTexture2DResource::FTexture2DResource( UTexture2D* InOwner, int32 InitialMipCou
 	MipFadeSetting = (Owner->LODGroup == TEXTUREGROUP_Lightmap || Owner->LODGroup == TEXTUREGROUP_Shadowmap) ? MipFade_Slow : MipFade_Normal;
 
 	// HDR images are stored in linear but still require gamma correction to display correctly.
-	bIgnoreGammaConversions = !Owner->SRGB && Owner->CompressionSettings != TC_HDR;
+	bIgnoreGammaConversions = !Owner->SRGB && Owner->CompressionSettings != TC_HDR && Owner->CompressionSettings != TC_HDR_Compressed;
 	bSRGB = InOwner->SRGB;
 
 	check(InitialMipCount>0);
