DR - Linker fix
#rnx

diff --git a/Engine/Source/Runtime/Engine/Private/Rendering/ShaderPipelineStateCache.cpp b/Engine/Source/Runtime/Engine/Private/Rendering/ShaderPipelineStateCache.cpp
new file mode 100644
--- /dev/null
+++ b/Engine/Source/Runtime/Engine/Private/Rendering/ShaderPipelineStateCache.cpp
@@ -0,0 +1,40 @@
+// Copyright 1998-2018 Epic Games, Inc. All Rights Reserved.
+
+#include "Rendering/ShaderPipelineStateCache.h"
+#include "ShaderPipelineCache.h"
+
+void FShaderPipelineStateCache::Initialize(EShaderPlatform Platform)
+{
+	FShaderPipelineCache::Initialize(Platform);
+}
+
+void FShaderPipelineStateCache::Shutdown()
+{
+	FShaderPipelineCache::Shutdown();
+}
+
+FShaderPipelineStateCache::FShaderPipelineStateCache(EShaderPlatform Platform)
+	: FShaderPipelineCache(Platform)
+	, FTickableObjectRenderThread(true, false) // (RegisterNow, HighFrequency)
+{
+}
+
+bool FShaderPipelineStateCache::IsTickable() const
+{
+	return FShaderPipelineCache::IsTickable();
+}
+
+void FShaderPipelineStateCache::Tick(float DeltaTime)
+{
+	return FShaderPipelineCache::Tick(DeltaTime);
+}
+
+bool FShaderPipelineStateCache::NeedsRenderingResumedForRenderingThreadTick() const
+{
+	return FShaderPipelineCache::NeedsRenderingResumedForRenderingThreadTick();
+}
+
+TStatId FShaderPipelineStateCache::GetStatId() const
+{
+	return FShaderPipelineCache::GetStatId();
+}
diff --git a/Engine/Source/Runtime/Engine/Public/Rendering/ShaderPipelineStateCache.h b/Engine/Source/Runtime/Engine/Public/Rendering/ShaderPipelineStateCache.h
--- a/Engine/Source/Runtime/Engine/Public/Rendering/ShaderPipelineStateCache.h
+++ b/Engine/Source/Runtime/Engine/Public/Rendering/ShaderPipelineStateCache.h
@@ -17,46 +17,22 @@ public:
 	 * Initializes the shader pipeline cache for the desired platform, called by the engine.
 	 * The shader platform is used only to load the initial pipeline cache and can be changed by closing & reopening the cache if necessary.
 	 */
-	static void Initialize(EShaderPlatform Platform)
-	{
-		FShaderPipelineCache::Initialize(Platform);
-	}
-	
+	static void Initialize(EShaderPlatform Platform);
+
 	/** Terminates the shader pipeline cache, called by the engine. */
-	static void Shutdown()
-	{
-		FShaderPipelineCache::Shutdown();
-	}
+	static void Shutdown();
+
+	FShaderPipelineStateCache(EShaderPlatform Platform);
 
-	FShaderPipelineStateCache(EShaderPlatform Platform)
-	: FShaderPipelineCache(Platform)
-	, FTickableObjectRenderThread(true, false) // (RegisterNow, HighFrequency)
-	{
-		
-	}
-	
 	virtual ~FShaderPipelineStateCache()
 	{
-		
 	}
 
-	virtual bool IsTickable() const override final
-	{
-		return FShaderPipelineCache::IsTickable();
-	}
-	
-	virtual void Tick( float DeltaTime ) override final
-	{
-		return FShaderPipelineCache::Tick(DeltaTime);
-	}
-	
-	virtual bool NeedsRenderingResumedForRenderingThreadTick() const override final
-	{
-		return FShaderPipelineCache::NeedsRenderingResumedForRenderingThreadTick();
-	}
-	
-	virtual TStatId GetStatId() const override final
-	{
-		return FShaderPipelineCache::GetStatId();
-	}
+	virtual bool IsTickable() const override final;
+
+	virtual void Tick( float DeltaTime ) override final;
+
+	virtual bool NeedsRenderingResumedForRenderingThreadTick() const override final;
+
+	virtual TStatId GetStatId() const override final;
 };
