Merging //UE4/Dev-Rendering-Graph@4492585 to Dev-Rendering (//UE4/Dev-Rendering)

This implements the framework to write the high level rendering code into passes organized in direct acyclic graph. It is also unifying the uniform buffer, shader parameters, and pass parameters to same single API: structures with run time meta data. This allow high level user code be extremely seamless, user code debugging, and render graph ease of implementation and debugging.

Issue of collaborative work of Arne Schnober, Brian Karis, Daniel Wright, Marcus Wassmer and Guillaume Abadie.

Names of the graph managed resources are not final.

diff --git a/Engine/Config/ConsoleVariables.ini b/Engine/Config/ConsoleVariables.ini
--- a/Engine/Config/ConsoleVariables.ini
+++ b/Engine/Config/ConsoleVariables.ini
@@ -33,7 +33,7 @@
 
 
 ; Uncomment to get detailed logs on shader compiles and the opportunity to retry on errors
-; r.ShaderDevelopmentMode=1
+;r.ShaderDevelopmentMode=1
 ; Uncomment to dump shaders in the Saved folder
 ; Warning: leaving this on for a while will fill your hard drive with many small files and folders
 ;r.DumpShaderDebugInfo=1
@@ -43,6 +43,8 @@
 ;r.DumpShaderDebugWorkerCommandLine=1
 ; Uncomment to enable parallel rendering at startup
 ;r.RHICmdBypass=0
+; Uncomment to get render graph executing passes as they get created to easily debug crashes caused by pass wiring logic.
+;r.Graph.ImmediateMode=1
 
 ; r.XGEShaderCompile is now enabled by default in source. Uncomment to disable XGE shader compilation.
 ;r.XGEShaderCompile = 0
diff --git a/Engine/Plugins/2D/Paper2D/Source/Paper2D/Paper2D.Build.cs b/Engine/Plugins/2D/Paper2D/Source/Paper2D/Paper2D.Build.cs
--- a/Engine/Plugins/2D/Paper2D/Source/Paper2D/Paper2D.Build.cs
+++ b/Engine/Plugins/2D/Paper2D/Source/Paper2D/Paper2D.Build.cs
@@ -13,7 +13,6 @@ public class Paper2D : ModuleRules
 				"Core",
 				"CoreUObject",
 				"Engine",
-				"ShaderCore",
 				"RenderCore",
 				"RHI",
         		        "SlateCore",
diff --git a/Engine/Plugins/Blendables/LightPropagationVolume/Source/LightPropagationVolumeRuntime/LightPropagationVolumeRuntime.Build.cs b/Engine/Plugins/Blendables/LightPropagationVolume/Source/LightPropagationVolumeRuntime/LightPropagationVolumeRuntime.Build.cs
--- a/Engine/Plugins/Blendables/LightPropagationVolume/Source/LightPropagationVolumeRuntime/LightPropagationVolumeRuntime.Build.cs
+++ b/Engine/Plugins/Blendables/LightPropagationVolume/Source/LightPropagationVolumeRuntime/LightPropagationVolumeRuntime.Build.cs
@@ -16,7 +16,6 @@ public class LightPropagationVolumeRuntime : ModuleRules
 				"Renderer"
 //				"RHI",
 //				"RenderCore",
-//				"ShaderCore",
 			}
 		);
 	}
diff --git a/Engine/Plugins/Compositing/LensDistortion/Source/LensDistortion/LensDistortion.Build.cs b/Engine/Plugins/Compositing/LensDistortion/Source/LensDistortion/LensDistortion.Build.cs
--- a/Engine/Plugins/Compositing/LensDistortion/Source/LensDistortion/LensDistortion.Build.cs
+++ b/Engine/Plugins/Compositing/LensDistortion/Source/LensDistortion/LensDistortion.Build.cs
@@ -29,7 +29,6 @@ namespace UnrealBuildTool.Rules
                     "RHI",
                     "Engine",
                     "RenderCore",
-                    "ShaderCore",
 					// ... add other public dependencies that you statically link with here ...
 				}
 				);
diff --git a/Engine/Plugins/Enterprise/DatasmithContent/Source/DatasmithContent/DatasmithContent.Build.cs b/Engine/Plugins/Enterprise/DatasmithContent/Source/DatasmithContent/DatasmithContent.Build.cs
--- a/Engine/Plugins/Enterprise/DatasmithContent/Source/DatasmithContent/DatasmithContent.Build.cs
+++ b/Engine/Plugins/Enterprise/DatasmithContent/Source/DatasmithContent/DatasmithContent.Build.cs
@@ -14,7 +14,7 @@ namespace UnrealBuildTool.Rules
 					"Core",
 					"CoreUObject",
 					"Engine",
-                    "ShaderCore",
+                    "RenderCore",
                 }
 			);
 
diff --git a/Engine/Plugins/FX/HoudiniNiagara/Source/HoudiniNiagara/HoudiniNiagara.Build.cs b/Engine/Plugins/FX/HoudiniNiagara/Source/HoudiniNiagara/HoudiniNiagara.Build.cs
--- a/Engine/Plugins/FX/HoudiniNiagara/Source/HoudiniNiagara/HoudiniNiagara.Build.cs
+++ b/Engine/Plugins/FX/HoudiniNiagara/Source/HoudiniNiagara/HoudiniNiagara.Build.cs
@@ -44,7 +44,7 @@ public class HoudiniNiagara : ModuleRules
 				"NiagaraShader",
 				"RenderCore",
 				"UtilityShaders",
-				"ShaderCore"
+				"RenderCore"
 			}
 		);
 			
@@ -61,7 +61,7 @@ public class HoudiniNiagara : ModuleRules
 				"RHI",
 				"UtilityShaders",
 				"NiagaraVertexFactories",
-				"ShaderCore"
+				"RenderCore"
 			}
 		);
 			
diff --git a/Engine/Plugins/FX/HoudiniNiagara/Source/HoudiniNiagaraEditor/HoudiniNiagaraEditor.Build.cs b/Engine/Plugins/FX/HoudiniNiagara/Source/HoudiniNiagaraEditor/HoudiniNiagaraEditor.Build.cs
--- a/Engine/Plugins/FX/HoudiniNiagara/Source/HoudiniNiagaraEditor/HoudiniNiagaraEditor.Build.cs
+++ b/Engine/Plugins/FX/HoudiniNiagara/Source/HoudiniNiagaraEditor/HoudiniNiagaraEditor.Build.cs
@@ -49,7 +49,7 @@ public class HoudiniNiagaraEditor : ModuleRules
                 "RHI",
                 "UtilityShaders",
                 "NiagaraVertexFactories",
-                "ShaderCore",
+                "RenderCore",
                 "HoudiniNiagara",
                 "UnrealEd",
 				"EditorStyle",
@@ -66,9 +66,7 @@ public class HoudiniNiagaraEditor : ModuleRules
 				"SlateCore",
                 "Niagara",
                 "NiagaraShader",
-                "RenderCore",
                 "UtilityShaders",
-                "ShaderCore",
                 "HoudiniNiagara"
 			}
         );
diff --git a/Engine/Plugins/FX/Niagara/Source/Niagara/Niagara.Build.cs b/Engine/Plugins/FX/Niagara/Source/Niagara/Niagara.Build.cs
--- a/Engine/Plugins/FX/Niagara/Source/Niagara/Niagara.Build.cs
+++ b/Engine/Plugins/FX/Niagara/Source/Niagara/Niagara.Build.cs
@@ -12,9 +12,7 @@ public class Niagara : ModuleRules
                 "NiagaraShader",
                 "Core",
                 "Engine",
-                "RenderCore",
                 "UtilityShaders",
-                "ShaderCore",
                 "TimeManagement",
                 "Renderer",
             }
@@ -32,7 +30,7 @@ public class Niagara : ModuleRules
                 "RHI",
                 "UtilityShaders",
                 "NiagaraVertexFactories",
-                "ShaderCore"
+                "RenderCore"
             }
         );
 
diff --git a/Engine/Plugins/FX/Niagara/Source/Niagara/Private/NiagaraDataInterfaceCollisionQuery.cpp b/Engine/Plugins/FX/Niagara/Source/Niagara/Private/NiagaraDataInterfaceCollisionQuery.cpp
--- a/Engine/Plugins/FX/Niagara/Source/Niagara/Private/NiagaraDataInterfaceCollisionQuery.cpp
+++ b/Engine/Plugins/FX/Niagara/Source/Niagara/Private/NiagaraDataInterfaceCollisionQuery.cpp
@@ -437,7 +437,7 @@ struct FNiagaraDataInterfaceParametersCS_CollisionQuery : public FNiagaraDataInt
 {
 	virtual void Bind(const FNiagaraDataInterfaceParamRef& ParamRef, const class FShaderParameterMap& ParameterMap) override
 	{
-		PassUniformBuffer.Bind(ParameterMap, FSceneTexturesUniformParameters::StaticStruct.GetShaderVariableName());
+		PassUniformBuffer.Bind(ParameterMap, FSceneTexturesUniformParameters::StaticStructMetadata.GetShaderVariableName());
 		check(PassUniformBuffer.IsBound());
 	}
 
diff --git a/Engine/Plugins/FX/Niagara/Source/Niagara/Private/NiagaraScriptExecutionParameterStore.cpp b/Engine/Plugins/FX/Niagara/Source/Niagara/Private/NiagaraScriptExecutionParameterStore.cpp
--- a/Engine/Plugins/FX/Niagara/Source/Niagara/Private/NiagaraScriptExecutionParameterStore.cpp
+++ b/Engine/Plugins/FX/Niagara/Source/Niagara/Private/NiagaraScriptExecutionParameterStore.cpp
@@ -25,67 +25,67 @@ FNiagaraScriptExecutionParameterStore& FNiagaraScriptExecutionParameterStore::op
 
 uint32 OffsetAlign(uint32 SrcOffset, uint32 Size)
 {
-	uint32 OffsetRemaining = UNIFORM_BUFFER_STRUCT_ALIGNMENT - (SrcOffset % UNIFORM_BUFFER_STRUCT_ALIGNMENT);
+	uint32 OffsetRemaining = SHADER_PARAMETER_STRUCT_ALIGNMENT - (SrcOffset % SHADER_PARAMETER_STRUCT_ALIGNMENT);
 	if (Size <= OffsetRemaining)
 	{
 		return SrcOffset;
 	}
 	else
 	{
-		return Align(SrcOffset, UNIFORM_BUFFER_STRUCT_ALIGNMENT);
+		return Align(SrcOffset, SHADER_PARAMETER_STRUCT_ALIGNMENT);
 	}
 }
 
 void FNiagaraScriptExecutionParameterStore::GenerateLayoutInfoInternal(TArray<FNiagaraScriptExecutionPaddingInfo>& Members, uint32& NextMemberOffset, const UStruct* InSrcStruct, uint32 InSrcOffset)
 {
-	uint32 VectorPaddedSize = (TUniformBufferTypeInfo<FVector4, false>::NumRows * TUniformBufferTypeInfo<FVector4, false>::NumColumns) * sizeof(float);
+	uint32 VectorPaddedSize = (TShaderParameterTypeInfo<FVector4>::NumRows * TShaderParameterTypeInfo<FVector4>::NumColumns) * sizeof(float);
 
 	// Now insert an appropriate data member into the mix...
 	if (InSrcStruct == FNiagaraTypeDefinition::GetBoolStruct() || InSrcStruct == FNiagaraTypeDefinition::GetIntStruct())
 	{
-		uint32 IntSize = (TUniformBufferTypeInfo<uint32, false>::NumRows * TUniformBufferTypeInfo<uint32, false>::NumColumns) * sizeof(uint32);
-		Members.Emplace(InSrcOffset, Align(NextMemberOffset, TUniformBufferTypeInfo<uint32, false>::Alignment), IntSize, IntSize);
+		uint32 IntSize = (TShaderParameterTypeInfo<uint32>::NumRows * TShaderParameterTypeInfo<uint32>::NumColumns) * sizeof(uint32);
+		Members.Emplace(InSrcOffset, Align(NextMemberOffset, TShaderParameterTypeInfo<uint32>::Alignment), IntSize, IntSize);
 		InSrcOffset += sizeof(uint32);
 		NextMemberOffset = Members[Members.Num() - 1].DestOffset + Members[Members.Num() - 1].DestSize;
 	}
 	else if (InSrcStruct == FNiagaraTypeDefinition::GetFloatStruct())
 	{
-		uint32 FloatSize = (TUniformBufferTypeInfo<float, false>::NumRows * TUniformBufferTypeInfo<float, false>::NumColumns) * sizeof(float);
-		Members.Emplace(InSrcOffset, Align(NextMemberOffset, TUniformBufferTypeInfo<float, false>::Alignment), FloatSize, FloatSize);
+		uint32 FloatSize = (TShaderParameterTypeInfo<float>::NumRows * TShaderParameterTypeInfo<float>::NumColumns) * sizeof(float);
+		Members.Emplace(InSrcOffset, Align(NextMemberOffset, TShaderParameterTypeInfo<float>::Alignment), FloatSize, FloatSize);
 		InSrcOffset += sizeof(float);
 		NextMemberOffset = Members[Members.Num() - 1].DestOffset + Members[Members.Num() - 1].DestSize;
 	}
 	else if (InSrcStruct == FNiagaraTypeDefinition::GetVec2Struct())
 	{
-		uint32 StructFinalSize = (TUniformBufferTypeInfo<FVector2D, false>::NumRows * TUniformBufferTypeInfo<FVector2D, false>::NumColumns) * sizeof(float);
+		uint32 StructFinalSize = (TShaderParameterTypeInfo<FVector2D>::NumRows * TShaderParameterTypeInfo<FVector2D>::NumColumns) * sizeof(float);
 		Members.Emplace(InSrcOffset, OffsetAlign(NextMemberOffset, VectorPaddedSize), StructFinalSize, VectorPaddedSize);
 		InSrcOffset += sizeof(FVector2D);
 		NextMemberOffset = Members[Members.Num() - 1].DestOffset + Members[Members.Num() - 1].DestSize;
 	}
 	else if (InSrcStruct == FNiagaraTypeDefinition::GetVec3Struct())
 	{
-		uint32 StructFinalSize = (TUniformBufferTypeInfo<FVector, false>::NumRows * TUniformBufferTypeInfo<FVector, false>::NumColumns) * sizeof(float);
+		uint32 StructFinalSize = (TShaderParameterTypeInfo<FVector>::NumRows * TShaderParameterTypeInfo<FVector>::NumColumns) * sizeof(float);
 		Members.Emplace(InSrcOffset, OffsetAlign(NextMemberOffset, VectorPaddedSize), StructFinalSize, VectorPaddedSize);
 		InSrcOffset += sizeof(FVector);
 		NextMemberOffset = Members[Members.Num() - 1].DestOffset + Members[Members.Num() - 1].DestSize;
 	}
 	else if (InSrcStruct == FNiagaraTypeDefinition::GetVec4Struct() || InSrcStruct == FNiagaraTypeDefinition::GetColorStruct() || InSrcStruct == FNiagaraTypeDefinition::GetQuatStruct())
 	{
-		uint32 StructFinalSize = (TUniformBufferTypeInfo<FVector4, false>::NumRows * TUniformBufferTypeInfo<FVector4, false>::NumColumns) * sizeof(float);
-		Members.Emplace(InSrcOffset, Align(NextMemberOffset, TUniformBufferTypeInfo<FVector4, false>::Alignment), StructFinalSize, StructFinalSize);
+		uint32 StructFinalSize = (TShaderParameterTypeInfo<FVector4>::NumRows * TShaderParameterTypeInfo<FVector4>::NumColumns) * sizeof(float);
+		Members.Emplace(InSrcOffset, Align(NextMemberOffset, TShaderParameterTypeInfo<FVector4>::Alignment), StructFinalSize, StructFinalSize);
 		InSrcOffset += sizeof(FVector4);
 		NextMemberOffset = Members[Members.Num() - 1].DestOffset + Members[Members.Num() - 1].DestSize;
 	}
 	else if (InSrcStruct == FNiagaraTypeDefinition::GetMatrix4Struct())
 	{
-		uint32 StructFinalSize = (TUniformBufferTypeInfo<FMatrix, false>::NumRows * TUniformBufferTypeInfo<FMatrix, false>::NumColumns) * sizeof(float);
-		Members.Emplace(InSrcOffset, Align(NextMemberOffset, TUniformBufferTypeInfo<FMatrix, false>::Alignment), StructFinalSize, StructFinalSize);
+		uint32 StructFinalSize = (TShaderParameterTypeInfo<FMatrix>::NumRows * TShaderParameterTypeInfo<FMatrix>::NumColumns) * sizeof(float);
+		Members.Emplace(InSrcOffset, Align(NextMemberOffset, TShaderParameterTypeInfo<FMatrix>::Alignment), StructFinalSize, StructFinalSize);
 		InSrcOffset += sizeof(FMatrix);
 		NextMemberOffset = Members[Members.Num() - 1].DestOffset + Members[Members.Num() - 1].DestSize;
 	}
 	else
 	{
-		NextMemberOffset = Align(NextMemberOffset, UNIFORM_BUFFER_STRUCT_ALIGNMENT); // New structs should be aligned to the head..
+		NextMemberOffset = Align(NextMemberOffset, SHADER_PARAMETER_STRUCT_ALIGNMENT); // New structs should be aligned to the head..
 
 		for (TFieldIterator<UProperty> PropertyIt(InSrcStruct, EFieldIteratorFlags::IncludeSuper); PropertyIt; ++PropertyIt)
 		{
@@ -138,7 +138,7 @@ void FNiagaraScriptExecutionParameterStore::AddPaddedParamSize(const FNiagaraTyp
 	if (PaddingInfo.Num() != 0)
 	{
 		NextMemberOffset = PaddingInfo[PaddingInfo.Num() - 1].DestOffset + PaddingInfo[PaddingInfo.Num() - 1].DestSize;
-		PaddedParameterSize = Align(NextMemberOffset, UNIFORM_BUFFER_STRUCT_ALIGNMENT);
+		PaddedParameterSize = Align(NextMemberOffset, SHADER_PARAMETER_STRUCT_ALIGNMENT);
 	}
 	else
 	{
diff --git a/Engine/Plugins/FX/Niagara/Source/NiagaraShader/NiagaraShader.Build.cs b/Engine/Plugins/FX/Niagara/Source/NiagaraShader/NiagaraShader.Build.cs
--- a/Engine/Plugins/FX/Niagara/Source/NiagaraShader/NiagaraShader.Build.cs
+++ b/Engine/Plugins/FX/Niagara/Source/NiagaraShader/NiagaraShader.Build.cs
@@ -10,8 +10,6 @@ public class NiagaraShader : ModuleRules
             new string[] {
                 "Core",
                 "Engine",
-				"RenderCore",
-                "ShaderCore",
                 "CoreUObject",
                 "NiagaraCore"
             }
@@ -21,7 +19,7 @@ public class NiagaraShader : ModuleRules
         PublicDependencyModuleNames.AddRange(
             new string[] {
                 "Core",
-                "ShaderCore",
+                "RenderCore",
                 "VectorVM",
                 "RHI",
                 "NiagaraCore"
diff --git a/Engine/Plugins/FX/Niagara/Source/NiagaraShader/Private/NiagaraShader.cpp b/Engine/Plugins/FX/Niagara/Source/NiagaraShader/Private/NiagaraShader.cpp
--- a/Engine/Plugins/FX/Niagara/Source/NiagaraShader/Private/NiagaraShader.cpp
+++ b/Engine/Plugins/FX/Niagara/Source/NiagaraShader/Private/NiagaraShader.cpp
@@ -274,7 +274,7 @@ void FNiagaraShaderType::CacheUniformBufferIncludes(TMap<const TCHAR*, FCachedUn
 	{
 		FCachedUniformBufferDeclaration& BufferDeclaration = It.Value();
 
-		for (TLinkedList<FUniformBufferStruct*>::TIterator StructIt(FUniformBufferStruct::GetStructList()); StructIt; StructIt.Next())
+		for (TLinkedList<FShaderParametersMetadata*>::TIterator StructIt(FShaderParametersMetadata::GetStructList()); StructIt; StructIt.Next())
 		{
 			if (It.Key() == StructIt->GetShaderVariableName())
 			{
@@ -309,7 +309,7 @@ void FNiagaraShaderType::AddReferencedUniformBufferIncludes(FShaderCompilerEnvir
 			*FString::Printf(TEXT("/Engine/Generated/UniformBuffers/%s.ush"), It.Key()), *Declaration
 		);
 
-		for (TLinkedList<FUniformBufferStruct*>::TIterator StructIt(FUniformBufferStruct::GetStructList()); StructIt; StructIt.Next())
+		for (TLinkedList<FShaderParametersMetadata*>::TIterator StructIt(FShaderParametersMetadata::GetStructList()); StructIt; StructIt.Next())
 		{
 			if (It.Key() == StructIt->GetShaderVariableName())
 			{
diff --git a/Engine/Plugins/FX/Niagara/Source/NiagaraVertexFactories/Private/NiagaraMeshVertexFactory.cpp b/Engine/Plugins/FX/Niagara/Source/NiagaraVertexFactories/Private/NiagaraMeshVertexFactory.cpp
--- a/Engine/Plugins/FX/Niagara/Source/NiagaraVertexFactories/Private/NiagaraMeshVertexFactory.cpp
+++ b/Engine/Plugins/FX/Niagara/Source/NiagaraVertexFactories/Private/NiagaraMeshVertexFactory.cpp
@@ -9,7 +9,7 @@
 #include "ParticleResources.h"
 #include "ShaderParameterUtils.h"
 
-IMPLEMENT_UNIFORM_BUFFER_STRUCT(FNiagaraMeshUniformParameters,TEXT("NiagaraMeshVF"));
+IMPLEMENT_GLOBAL_SHADER_PARAMETER_STRUCT(FNiagaraMeshUniformParameters,"NiagaraMeshVF");
 
 class FNiagaraMeshVertexFactoryShaderParameters : public FVertexFactoryShaderParameters
 {
diff --git a/Engine/Plugins/FX/Niagara/Source/NiagaraVertexFactories/Private/NiagaraRibbonVertexFactory.cpp b/Engine/Plugins/FX/Niagara/Source/NiagaraVertexFactories/Private/NiagaraRibbonVertexFactory.cpp
--- a/Engine/Plugins/FX/Niagara/Source/NiagaraVertexFactories/Private/NiagaraRibbonVertexFactory.cpp
+++ b/Engine/Plugins/FX/Niagara/Source/NiagaraVertexFactories/Private/NiagaraRibbonVertexFactory.cpp
@@ -9,7 +9,7 @@
 #include "ParticleResources.h"
 #include "ShaderParameterUtils.h"
 
-IMPLEMENT_UNIFORM_BUFFER_STRUCT(FNiagaraRibbonUniformParameters,TEXT("NiagaraRibbonVF"));
+IMPLEMENT_GLOBAL_SHADER_PARAMETER_STRUCT(FNiagaraRibbonUniformParameters,"NiagaraRibbonVF");
 
 
 class FNiagaraRibbonVertexFactoryShaderParameters : public FVertexFactoryShaderParameters
diff --git a/Engine/Plugins/FX/Niagara/Source/NiagaraVertexFactories/Private/NiagaraSpriteVertexFactory.cpp b/Engine/Plugins/FX/Niagara/Source/NiagaraVertexFactories/Private/NiagaraSpriteVertexFactory.cpp
--- a/Engine/Plugins/FX/Niagara/Source/NiagaraVertexFactories/Private/NiagaraSpriteVertexFactory.cpp
+++ b/Engine/Plugins/FX/Niagara/Source/NiagaraVertexFactories/Private/NiagaraSpriteVertexFactory.cpp
@@ -9,7 +9,7 @@
 #include "ParticleResources.h"
 #include "ShaderParameterUtils.h"
 
-IMPLEMENT_UNIFORM_BUFFER_STRUCT(FNiagaraSpriteUniformParameters,TEXT("NiagaraSpriteVF"));
+IMPLEMENT_GLOBAL_SHADER_PARAMETER_STRUCT(FNiagaraSpriteUniformParameters,"NiagaraSpriteVF");
 
 TGlobalResource<FNullDynamicParameterVertexBuffer> GNullNiagaraDynamicParameterVertexBuffer;
 
diff --git a/Engine/Plugins/FX/Niagara/Source/NiagaraVertexFactories/Public/NiagaraMeshVertexFactory.h b/Engine/Plugins/FX/Niagara/Source/NiagaraVertexFactories/Public/NiagaraMeshVertexFactory.h
--- a/Engine/Plugins/FX/Niagara/Source/NiagaraVertexFactories/Public/NiagaraMeshVertexFactory.h
+++ b/Engine/Plugins/FX/Niagara/Source/NiagaraVertexFactories/Public/NiagaraMeshVertexFactory.h
@@ -28,28 +28,28 @@ struct FShaderCompilerEnvironment;
 /**
 * Uniform buffer for mesh particle vertex factories.
 */
-BEGIN_UNIFORM_BUFFER_STRUCT(FNiagaraMeshUniformParameters, NIAGARAVERTEXFACTORIES_API)
-	UNIFORM_MEMBER_EX( FMatrix, LocalToWorld, EShaderPrecisionModifier::Half)
-	UNIFORM_MEMBER_EX( FMatrix, LocalToWorldInverseTransposed, EShaderPrecisionModifier::Half)
-	UNIFORM_MEMBER(FVector4, SubImageSize)
-	UNIFORM_MEMBER(uint32, TexCoordWeightA)
-	UNIFORM_MEMBER(uint32, TexCoordWeightB)
-	UNIFORM_MEMBER(uint32, PrevTransformAvailable)
-	UNIFORM_MEMBER(float, DeltaSeconds)
-	UNIFORM_MEMBER(int, PositionDataOffset)
-	UNIFORM_MEMBER(int, VelocityDataOffset)
-	UNIFORM_MEMBER(int, ColorDataOffset)
-	UNIFORM_MEMBER(int, TransformDataOffset)
-	UNIFORM_MEMBER(int, ScaleDataOffset)
-	UNIFORM_MEMBER(int, SizeDataOffset)
-	UNIFORM_MEMBER(int, MaterialParamDataOffset)
-	UNIFORM_MEMBER(int, MaterialParam1DataOffset)
-	UNIFORM_MEMBER(int, MaterialParam2DataOffset)
-	UNIFORM_MEMBER(int, MaterialParam3DataOffset)
-	UNIFORM_MEMBER(int, NormalizedAgeDataOffset)
-	UNIFORM_MEMBER(int, MaterialRandomDataOffset)
-	UNIFORM_MEMBER(FVector4, DefaultPos)
-END_UNIFORM_BUFFER_STRUCT(FNiagaraMeshUniformParameters)
+BEGIN_GLOBAL_SHADER_PARAMETER_STRUCT(FNiagaraMeshUniformParameters, NIAGARAVERTEXFACTORIES_API)
+	SHADER_PARAMETER_EX( FMatrix, LocalToWorld, EShaderPrecisionModifier::Half)
+	SHADER_PARAMETER_EX( FMatrix, LocalToWorldInverseTransposed, EShaderPrecisionModifier::Half)
+	SHADER_PARAMETER(FVector4, SubImageSize)
+	SHADER_PARAMETER(uint32, TexCoordWeightA)
+	SHADER_PARAMETER(uint32, TexCoordWeightB)
+	SHADER_PARAMETER(uint32, PrevTransformAvailable)
+	SHADER_PARAMETER(float, DeltaSeconds)
+	SHADER_PARAMETER(int, PositionDataOffset)
+	SHADER_PARAMETER(int, VelocityDataOffset)
+	SHADER_PARAMETER(int, ColorDataOffset)
+	SHADER_PARAMETER(int, TransformDataOffset)
+	SHADER_PARAMETER(int, ScaleDataOffset)
+	SHADER_PARAMETER(int, SizeDataOffset)
+	SHADER_PARAMETER(int, MaterialParamDataOffset)
+	SHADER_PARAMETER(int, MaterialParam1DataOffset)
+	SHADER_PARAMETER(int, MaterialParam2DataOffset)
+	SHADER_PARAMETER(int, MaterialParam3DataOffset)
+	SHADER_PARAMETER(int, NormalizedAgeDataOffset)
+	SHADER_PARAMETER(int, MaterialRandomDataOffset)
+	SHADER_PARAMETER(FVector4, DefaultPos)
+END_GLOBAL_SHADER_PARAMETER_STRUCT()
 
 typedef TUniformBufferRef<FNiagaraMeshUniformParameters> FNiagaraMeshUniformBufferRef;
 
diff --git a/Engine/Plugins/FX/Niagara/Source/NiagaraVertexFactories/Public/NiagaraRibbonVertexFactory.h b/Engine/Plugins/FX/Niagara/Source/NiagaraVertexFactories/Public/NiagaraRibbonVertexFactory.h
--- a/Engine/Plugins/FX/Niagara/Source/NiagaraVertexFactories/Public/NiagaraRibbonVertexFactory.h
+++ b/Engine/Plugins/FX/Niagara/Source/NiagaraVertexFactories/Public/NiagaraRibbonVertexFactory.h
@@ -35,32 +35,32 @@ struct FNiagaraRibbonVertexDynamicParameter
 /**
 * Uniform buffer for particle beam/trail vertex factories.
 */
-BEGIN_UNIFORM_BUFFER_STRUCT(FNiagaraRibbonUniformParameters, NIAGARAVERTEXFACTORIES_API)
-	UNIFORM_MEMBER(FVector4, CameraRight)
-	UNIFORM_MEMBER(FVector4, CameraUp)
-	UNIFORM_MEMBER(FVector4, ScreenAlignment)
-	UNIFORM_MEMBER(int, PositionDataOffset)
-	UNIFORM_MEMBER(int, VelocityDataOffset)
-	UNIFORM_MEMBER(int, WidthDataOffset)
-	UNIFORM_MEMBER(int, TwistDataOffset)
-	UNIFORM_MEMBER(int, ColorDataOffset)
-	UNIFORM_MEMBER(int, FacingDataOffset)
-	UNIFORM_MEMBER(int, NormalizedAgeDataOffset)
-	UNIFORM_MEMBER(int, MaterialRandomDataOffset)
-	UNIFORM_MEMBER(int, MaterialParamDataOffset)
-	UNIFORM_MEMBER(int, MaterialParam1DataOffset)
-	UNIFORM_MEMBER(int, MaterialParam2DataOffset)
-	UNIFORM_MEMBER(int, MaterialParam3DataOffset)
-	UNIFORM_MEMBER(int, TotalNumInstances)
-	UNIFORM_MEMBER(uint32, UseCustomFacing)
-	UNIFORM_MEMBER(uint32, InvertDrawOrder)
-	UNIFORM_MEMBER(float, UV0TilingDistance)
-	UNIFORM_MEMBER(float, UV1TilingDistance)
-	UNIFORM_MEMBER(FVector4, PackedVData)
-	UNIFORM_MEMBER_EX(FMatrix, LocalToWorld, EShaderPrecisionModifier::Half)
-	UNIFORM_MEMBER_EX(FMatrix, LocalToWorldInverseTransposed, EShaderPrecisionModifier::Half)
-	UNIFORM_MEMBER_EX(float, DeltaSeconds, EShaderPrecisionModifier::Half)
-	END_UNIFORM_BUFFER_STRUCT(FNiagaraRibbonUniformParameters)
+BEGIN_GLOBAL_SHADER_PARAMETER_STRUCT(FNiagaraRibbonUniformParameters, NIAGARAVERTEXFACTORIES_API)
+	SHADER_PARAMETER(FVector4, CameraRight)
+	SHADER_PARAMETER(FVector4, CameraUp)
+	SHADER_PARAMETER(FVector4, ScreenAlignment)
+	SHADER_PARAMETER(int, PositionDataOffset)
+	SHADER_PARAMETER(int, VelocityDataOffset)
+	SHADER_PARAMETER(int, WidthDataOffset)
+	SHADER_PARAMETER(int, TwistDataOffset)
+	SHADER_PARAMETER(int, ColorDataOffset)
+	SHADER_PARAMETER(int, FacingDataOffset)
+	SHADER_PARAMETER(int, NormalizedAgeDataOffset)
+	SHADER_PARAMETER(int, MaterialRandomDataOffset)
+	SHADER_PARAMETER(int, MaterialParamDataOffset)
+	SHADER_PARAMETER(int, MaterialParam1DataOffset)
+	SHADER_PARAMETER(int, MaterialParam2DataOffset)
+	SHADER_PARAMETER(int, MaterialParam3DataOffset)
+	SHADER_PARAMETER(int, TotalNumInstances)
+	SHADER_PARAMETER(uint32, UseCustomFacing)
+	SHADER_PARAMETER(uint32, InvertDrawOrder)
+	SHADER_PARAMETER(float, UV0TilingDistance)
+	SHADER_PARAMETER(float, UV1TilingDistance)
+	SHADER_PARAMETER(FVector4, PackedVData)
+	SHADER_PARAMETER_EX(FMatrix, LocalToWorld, EShaderPrecisionModifier::Half)
+	SHADER_PARAMETER_EX(FMatrix, LocalToWorldInverseTransposed, EShaderPrecisionModifier::Half)
+	SHADER_PARAMETER_EX(float, DeltaSeconds, EShaderPrecisionModifier::Half)
+END_GLOBAL_SHADER_PARAMETER_STRUCT()
 typedef TUniformBufferRef<FNiagaraRibbonUniformParameters> FNiagaraRibbonUniformBufferRef;
 
 /**
diff --git a/Engine/Plugins/FX/Niagara/Source/NiagaraVertexFactories/Public/NiagaraSpriteVertexFactory.h b/Engine/Plugins/FX/Niagara/Source/NiagaraVertexFactories/Public/NiagaraSpriteVertexFactory.h
--- a/Engine/Plugins/FX/Niagara/Source/NiagaraVertexFactories/Public/NiagaraSpriteVertexFactory.h
+++ b/Engine/Plugins/FX/Niagara/Source/NiagaraVertexFactories/Public/NiagaraSpriteVertexFactory.h
@@ -19,41 +19,41 @@ class FMaterial;
 /**
  * Uniform buffer for particle sprite vertex factories.
  */
-BEGIN_UNIFORM_BUFFER_STRUCT( FNiagaraSpriteUniformParameters, NIAGARAVERTEXFACTORIES_API)
-	UNIFORM_MEMBER_EX( FMatrix, LocalToWorld, EShaderPrecisionModifier::Half)
-	UNIFORM_MEMBER_EX( FMatrix, LocalToWorldInverseTransposed, EShaderPrecisionModifier::Half)
-	UNIFORM_MEMBER_EX( FVector, CustomFacingVectorMask, EShaderPrecisionModifier::Half)
-	UNIFORM_MEMBER_EX( FVector4, TangentSelector, EShaderPrecisionModifier::Half )
-	UNIFORM_MEMBER_EX( FVector4, NormalsSphereCenter, EShaderPrecisionModifier::Half )
-	UNIFORM_MEMBER_EX( FVector4, NormalsCylinderUnitDirection, EShaderPrecisionModifier::Half )
-	UNIFORM_MEMBER_EX( FVector4, SubImageSize, EShaderPrecisionModifier::Half )
-	UNIFORM_MEMBER_EX( FVector, CameraFacingBlend, EShaderPrecisionModifier::Half )
-	UNIFORM_MEMBER_EX( float, RemoveHMDRoll, EShaderPrecisionModifier::Half )
-	UNIFORM_MEMBER( FVector4, MacroUVParameters )
-	UNIFORM_MEMBER_EX( float, RotationScale, EShaderPrecisionModifier::Half )
-	UNIFORM_MEMBER_EX( float, RotationBias, EShaderPrecisionModifier::Half )
-	UNIFORM_MEMBER_EX( float, NormalsType, EShaderPrecisionModifier::Half )
-	UNIFORM_MEMBER_EX( float, DeltaSeconds, EShaderPrecisionModifier::Half )
-	UNIFORM_MEMBER_EX( FVector2D, PivotOffset, EShaderPrecisionModifier::Half )
-	UNIFORM_MEMBER(int, PositionDataOffset)
-	UNIFORM_MEMBER(int, VelocityDataOffset)
-	UNIFORM_MEMBER(int, RotationDataOffset)
-	UNIFORM_MEMBER(int, SizeDataOffset)
-	UNIFORM_MEMBER(int, SubimageDataOffset)
-	UNIFORM_MEMBER(int, ColorDataOffset)
-	UNIFORM_MEMBER(int, MaterialParamDataOffset)
-	UNIFORM_MEMBER(int, MaterialParam1DataOffset)
-	UNIFORM_MEMBER(int, MaterialParam2DataOffset)
-	UNIFORM_MEMBER(int, MaterialParam3DataOffset)
-	UNIFORM_MEMBER(int, FacingDataOffset)
-	UNIFORM_MEMBER(int, AlignmentDataOffset)
-	UNIFORM_MEMBER(int, SubImageBlendMode)
-	UNIFORM_MEMBER(int, CameraOffsetDataOffset)
-	UNIFORM_MEMBER(int, UVScaleDataOffset)
-	UNIFORM_MEMBER(int, NormalizedAgeDataOffset)
-	UNIFORM_MEMBER(int, MaterialRandomDataOffset)
-	UNIFORM_MEMBER(FVector4, DefaultPos)
-	END_UNIFORM_BUFFER_STRUCT(FNiagaraSpriteUniformParameters)
+BEGIN_GLOBAL_SHADER_PARAMETER_STRUCT( FNiagaraSpriteUniformParameters, NIAGARAVERTEXFACTORIES_API)
+	SHADER_PARAMETER_EX( FMatrix, LocalToWorld, EShaderPrecisionModifier::Half)
+	SHADER_PARAMETER_EX( FMatrix, LocalToWorldInverseTransposed, EShaderPrecisionModifier::Half)
+	SHADER_PARAMETER_EX( FVector, CustomFacingVectorMask, EShaderPrecisionModifier::Half)
+	SHADER_PARAMETER_EX( FVector4, TangentSelector, EShaderPrecisionModifier::Half )
+	SHADER_PARAMETER_EX( FVector4, NormalsSphereCenter, EShaderPrecisionModifier::Half )
+	SHADER_PARAMETER_EX( FVector4, NormalsCylinderUnitDirection, EShaderPrecisionModifier::Half )
+	SHADER_PARAMETER_EX( FVector4, SubImageSize, EShaderPrecisionModifier::Half )
+	SHADER_PARAMETER_EX( FVector, CameraFacingBlend, EShaderPrecisionModifier::Half )
+	SHADER_PARAMETER_EX( float, RemoveHMDRoll, EShaderPrecisionModifier::Half )
+	SHADER_PARAMETER( FVector4, MacroUVParameters )
+	SHADER_PARAMETER_EX( float, RotationScale, EShaderPrecisionModifier::Half )
+	SHADER_PARAMETER_EX( float, RotationBias, EShaderPrecisionModifier::Half )
+	SHADER_PARAMETER_EX( float, NormalsType, EShaderPrecisionModifier::Half )
+	SHADER_PARAMETER_EX( float, DeltaSeconds, EShaderPrecisionModifier::Half )
+	SHADER_PARAMETER_EX( FVector2D, PivotOffset, EShaderPrecisionModifier::Half )
+	SHADER_PARAMETER(int, PositionDataOffset)
+	SHADER_PARAMETER(int, VelocityDataOffset)
+	SHADER_PARAMETER(int, RotationDataOffset)
+	SHADER_PARAMETER(int, SizeDataOffset)
+	SHADER_PARAMETER(int, SubimageDataOffset)
+	SHADER_PARAMETER(int, ColorDataOffset)
+	SHADER_PARAMETER(int, MaterialParamDataOffset)
+	SHADER_PARAMETER(int, MaterialParam1DataOffset)
+	SHADER_PARAMETER(int, MaterialParam2DataOffset)
+	SHADER_PARAMETER(int, MaterialParam3DataOffset)
+	SHADER_PARAMETER(int, FacingDataOffset)
+	SHADER_PARAMETER(int, AlignmentDataOffset)
+	SHADER_PARAMETER(int, SubImageBlendMode)
+	SHADER_PARAMETER(int, CameraOffsetDataOffset)
+	SHADER_PARAMETER(int, UVScaleDataOffset)
+	SHADER_PARAMETER(int, NormalizedAgeDataOffset)
+	SHADER_PARAMETER(int, MaterialRandomDataOffset)
+	SHADER_PARAMETER(FVector4, DefaultPos)
+END_GLOBAL_SHADER_PARAMETER_STRUCT()
 
 typedef TUniformBufferRef<FNiagaraSpriteUniformParameters> FNiagaraSpriteUniformBufferRef;
 
diff --git a/Engine/Plugins/Runtime/ArchVisCharacter/Source/ArchVisCharacter/ArchVisCharacter.Build.cs b/Engine/Plugins/Runtime/ArchVisCharacter/Source/ArchVisCharacter/ArchVisCharacter.Build.cs
--- a/Engine/Plugins/Runtime/ArchVisCharacter/Source/ArchVisCharacter/ArchVisCharacter.Build.cs
+++ b/Engine/Plugins/Runtime/ArchVisCharacter/Source/ArchVisCharacter/ArchVisCharacter.Build.cs
@@ -15,7 +15,6 @@ namespace UnrealBuildTool.Rules
 					"CoreUObject",
                     "Engine",
 //                     "RenderCore",
-//                     "ShaderCore",
 //                     "RHI"
 				}
 				);
diff --git a/Engine/Plugins/Runtime/Oculus/OculusVR/Source/OculusHMD/Private/OculusHMD_StressTester.cpp b/Engine/Plugins/Runtime/Oculus/OculusVR/Source/OculusHMD/Private/OculusHMD_StressTester.cpp
--- a/Engine/Plugins/Runtime/Oculus/OculusVR/Source/OculusHMD/Private/OculusHMD_StressTester.cpp
+++ b/Engine/Plugins/Runtime/Oculus/OculusVR/Source/OculusHMD/Private/OculusHMD_StressTester.cpp
@@ -21,21 +21,21 @@ DECLARE_CYCLE_STAT(TEXT("GPUStressRendering"), STAT_GPUStressRendering, STATGROU
 //-------------------------------------------------------------------------------------------------
 
 //This buffer should contain variables that never, or rarely change
-BEGIN_UNIFORM_BUFFER_STRUCT(FOculusPixelShaderConstantParameters,)
-//UNIFORM_MEMBER(FVector4, Name)
-END_UNIFORM_BUFFER_STRUCT(FOculusPixelShaderConstantParameters)
+BEGIN_GLOBAL_SHADER_PARAMETER_STRUCT(FOculusPixelShaderConstantParameters, )
+//SHADER_PARAMETER(FVector4, Name)
+END_GLOBAL_SHADER_PARAMETER_STRUCT()
 
-IMPLEMENT_UNIFORM_BUFFER_STRUCT(FOculusPixelShaderConstantParameters, TEXT("PSConstants"))
+IMPLEMENT_GLOBAL_SHADER_PARAMETER_STRUCT(FOculusPixelShaderConstantParameters, "PSConstants");
 
 typedef TUniformBufferRef<FOculusPixelShaderConstantParameters> FOculusPixelShaderConstantParametersRef;
 
 
 //This buffer is for variables that change very often (each frame for example)
-BEGIN_UNIFORM_BUFFER_STRUCT(FOculusPixelShaderVariableParameters,)
-UNIFORM_MEMBER(int, IterationsMultiplier)
-END_UNIFORM_BUFFER_STRUCT(FOculusPixelShaderVariableParameters)
+BEGIN_GLOBAL_SHADER_PARAMETER_STRUCT(FOculusPixelShaderVariableParameters, )
+SHADER_PARAMETER(int, IterationsMultiplier)
+END_GLOBAL_SHADER_PARAMETER_STRUCT()
 
-IMPLEMENT_UNIFORM_BUFFER_STRUCT(FOculusPixelShaderVariableParameters, TEXT("PSVariables"))
+IMPLEMENT_GLOBAL_SHADER_PARAMETER_STRUCT(FOculusPixelShaderVariableParameters, "PSVariables");
 
 typedef TUniformBufferRef<FOculusPixelShaderVariableParameters> FOculusPixelShaderVariableParametersRef;
 
diff --git a/Engine/Plugins/Runtime/PreLoadScreenMoviePlayer/Source/PreLoadScreenMoviePlayer/PreLoadScreenMoviePlayer.Build.cs b/Engine/Plugins/Runtime/PreLoadScreenMoviePlayer/Source/PreLoadScreenMoviePlayer/PreLoadScreenMoviePlayer.Build.cs
--- a/Engine/Plugins/Runtime/PreLoadScreenMoviePlayer/Source/PreLoadScreenMoviePlayer/PreLoadScreenMoviePlayer.Build.cs
+++ b/Engine/Plugins/Runtime/PreLoadScreenMoviePlayer/Source/PreLoadScreenMoviePlayer/PreLoadScreenMoviePlayer.Build.cs
@@ -12,7 +12,7 @@ public class PreLoadScreenMoviePlayer : ModuleRules
 			new string[] {
                 "Core",
                 "RHI",
-                "ShaderCore",
+                "RenderCore",
                 "MoviePlayer",
                 "Slate",
                 "SlateCore",
diff --git a/Engine/Shaders/Public/ShaderVersion.ush b/Engine/Shaders/Public/ShaderVersion.ush
--- a/Engine/Shaders/Public/ShaderVersion.ush
+++ b/Engine/Shaders/Public/ShaderVersion.ush
@@ -2,4 +2,4 @@
 // Each time the console command is executed it generates a new GUID. As this file is included
 // in Platform.ush (which should be included in any shader) it allows to invalidate the shader DDC.
 // 
-// GUID = F7C9D61CA78A4854B1D6EE5D3F987298
+// GUID = 8301458A46E1B310E41BA7B2A499BFC1
\ No newline at end of file
diff --git a/Engine/Source/Developer/ShaderCompilerCommon/Private/ShaderCompilerCommon.cpp b/Engine/Source/Developer/ShaderCompilerCommon/Private/ShaderCompilerCommon.cpp
--- a/Engine/Source/Developer/ShaderCompilerCommon/Private/ShaderCompilerCommon.cpp
+++ b/Engine/Source/Developer/ShaderCompilerCommon/Private/ShaderCompilerCommon.cpp
@@ -131,15 +131,17 @@ bool BuildResourceTableMapping(
 			switch( Entry.Type )
 			{
 			case UBMT_TEXTURE:
+			case UBMT_GRAPH_TRACKED_TEXTURE:
 				OutSRT.TextureMap.Add(ResourceMap);
 				break;
 			case UBMT_SAMPLER:
 				OutSRT.SamplerMap.Add(ResourceMap);
 				break;
 			case UBMT_SRV:
+			case UBMT_GRAPH_TRACKED_SRV:
 				OutSRT.ShaderResourceViewMap.Add(ResourceMap);
 				break;
-			case UBMT_UAV:
+			case UBMT_GRAPH_TRACKED_UAV:
 				OutSRT.UnorderedAccessViewMap.Add(ResourceMap);
 				break;
 			default:
diff --git a/Engine/Source/Editor/UnrealEd/Private/UnrealEdPrivatePCH.h b/Engine/Source/Editor/UnrealEd/Private/UnrealEdPrivatePCH.h
--- a/Engine/Source/Editor/UnrealEd/Private/UnrealEdPrivatePCH.h
+++ b/Engine/Source/Editor/UnrealEd/Private/UnrealEdPrivatePCH.h
@@ -445,16 +445,14 @@
 #include "UniformBuffer.h"
 #include "PackedNormal.h"
 #include "RenderUtils.h"
-
-// From Messaging:
-#include "IMessageContext.h"
-
-// From ShaderCore:
 #include "ShaderParameters.h"
 #include "ShaderCore.h"
 #include "Shader.h"
 #include "VertexFactory.h"
 
+// From Messaging:
+#include "IMessageContext.h"
+
 // From AssetRegistry:
 #include "AssetData.h"
 #include "IAssetRegistry.h"
diff --git a/Engine/Source/Editor/UnrealEd/Public/UnrealEdSharedPCH.h b/Engine/Source/Editor/UnrealEd/Public/UnrealEdSharedPCH.h
--- a/Engine/Source/Editor/UnrealEd/Public/UnrealEdSharedPCH.h
+++ b/Engine/Source/Editor/UnrealEd/Public/UnrealEdSharedPCH.h
@@ -442,16 +442,14 @@
 #include "UniformBuffer.h"
 #include "PackedNormal.h"
 #include "RenderUtils.h"
-
-// From Messaging:
-#include "IMessageContext.h"
-
-// From ShaderCore:
 #include "ShaderParameters.h"
 #include "Shader.h"
 #include "ShaderCore.h"
 #include "VertexFactory.h"
 
+// From Messaging:
+#include "IMessageContext.h"
+
 // From AssetRegistry:
 #include "AssetData.h"
 #include "ARFilter.h"
diff --git a/Engine/Source/Editor/ViewportSnapping/ViewportSnapping.Build.cs b/Engine/Source/Editor/ViewportSnapping/ViewportSnapping.Build.cs
--- a/Engine/Source/Editor/ViewportSnapping/ViewportSnapping.Build.cs
+++ b/Engine/Source/Editor/ViewportSnapping/ViewportSnapping.Build.cs
@@ -14,7 +14,6 @@ public class ViewportSnapping : ModuleRules
 				"CoreUObject",
 				"Engine",
 // 				"RenderCore",
-// 				"ShaderCore",
 // 				"RHI",
 				"Slate",
 				"UnrealEd"
diff --git a/Engine/Source/Runtime/Core/Public/Windows/WindowsPlatformCompilerSetup.h b/Engine/Source/Runtime/Core/Public/Windows/WindowsPlatformCompilerSetup.h
--- a/Engine/Source/Runtime/Core/Public/Windows/WindowsPlatformCompilerSetup.h
+++ b/Engine/Source/Runtime/Core/Public/Windows/WindowsPlatformCompilerSetup.h
@@ -23,7 +23,7 @@ static_assert(_MSC_VER < 1910 || _MSC_VER >= 1913, "Visual Studio 2017 version 1
 #pragma warning (error:      4001 4002 4003 4004      4006 4007 4008 4009 4010 4011 4012 4013 4014 4015 4016 4017 4018 4019 4020 4021 4022 4023 4024 4025 4026 4027 4028 4029 4030 4031 4032 4033 4034 4035 4036 4037 4038 4039 4040 4041 4042 4043 4044 4045 4046 4047 4048 4049 4050 4051 4052 4053 4054 4055 4056 4057 4058 4059 4060           4063 4064 4065 4066 4067 4068 4069 4070 4071 4072 4073 4074 4075 4076 4077 4078 4079 4080 4081 4082 4083 4084 4085 4086 4087 4088 4089 4090 4091 4092 4093 4094 4095 4096 4097 4098 4099)
 #pragma warning (error: 4100 4101 4102 4103 4104 4105 4106 4107 4108 4109 4110 4111 4112 4113 4114 4115 4116 4117 4118 4119 4120 4121 4122 4123 4124 4125 4126 4127 4128 4129 4130 4131 4132 4133 4134 4135 4136 4137 4138 4139 4140 4141 4142 4143 4144 4145 4146 4147 4148 4149 4150 4151 4152 4153 4154 4155 4156 4157 4158 4159 4160 4161 4162 4163 4164      4166 4167 4168 4169 4170 4171 4172 4173 4174 4175 4176 4177 4178 4179 4180 4181 4182 4183 4184 4185 4186 4187 4188 4189 4190 4191 4192 4193 4194 4195 4196 4197 4198 4199)
 #pragma warning (error: 4200 4201 4202 4203 4204 4205 4206 4207 4208 4209 4210 4211 4212 4213 4214 4215 4216 4217 4218 4219 4220 4221 4222 4223 4224 4225 4226 4227 4228 4229 4230 4231 4232 4233 4234 4235 4236 4237 4238 4239 4240 4241      4243 4244 4245 4246 4247 4248 4249 4250 4251 4252 4253      4255 4256 4257 4258 4259 4260 4261 4262 4263 4264           4267 4268 4269 4270 4271 4272 4273 4274 4275 4276 4277 4278 4279 4280 4281 4282 4283 4284 4285 4286 4287 4288 4289 4290 4291 4292 4293 4294 4295      4297 4298 4299)
-#pragma warning (error: 4300 4301 4302 4303 4304      4306 4307 4308 4309 4310           4313 4314 4315 4316 4317 4318 4319 4320 4321 4322 4323 4324 4325 4326 4327 4328 4329 4330 4331 4332 4333 4334 4335 4336 4337 4338      4340 4341      4343 4344      4346 4347 4348 4349           4352 4353 4354 4355 4356 4357 4358 4359 4360 4361 4362 4363 4364      4366 4367 4368 4369           4372      4374 4375 4376 4377 4378 4379 4380 4381 4382 4383 4384 4385 4386 4387      4389 4390 4391 4392 4393 4394 4395 4396 4397 4398 4399)
+#pragma warning (error: 4300 4301 4302 4303 4304      4306      4308 4309 4310           4313 4314 4315 4316 4317 4318 4319 4320 4321 4322 4323 4324 4325 4326 4327 4328 4329 4330 4331 4332 4333 4334 4335 4336 4337 4338      4340 4341      4343 4344      4346 4347 4348 4349           4352 4353 4354 4355 4356 4357 4358 4359 4360 4361 4362 4363 4364      4366 4367 4368 4369           4372      4374 4375 4376 4377 4378 4379 4380 4381 4382 4383 4384 4385 4386 4387      4389 4390 4391 4392 4393 4394 4395 4396 4397 4398 4399)
 #pragma warning (error: 4400 4401 4402 4403 4404 4405 4406 4407 4408 4409 4410 4411      4413 4414 4415 4416 4417 4418 4419 4420 4421 4422 4423 4424 4425 4426 4427 4428 4429 4430 4431 4432 4433 4434      4436      4438 4439 4440 4441 4442 4443      4445 4446 4447 4448 4449 4450 4451 4452 4453 4454 4455                     4460 4461 4462      4464 4465 4466 4467 4468 4469 4470           4473 4474 4475 4476 4477 4478 4479 4480      4482 4483 4484 4485 4486 4487 4488 4489 4490 4491 4492 4493 4494 4495 4496 4497 4498 4499)
 #pragma warning (error: 4500 4501 4502 4503 4504 4505 4506 4507 4508 4509      4511 4512 4513 4514 4515 4516 4517 4518 4519 4520 4521 4522 4523 4524 4525 4526 4527 4528 4529 4530 4531 4532 4533 4534 4535 4536 4537 4538 4539 4540 4541 4542 4543 4544 4545 4546                4550 4551 4552 4553 4554      4556 4557 4558 4559 4560 4561 4562 4563 4564 4565 4566 4567 4568 4569 4570      4572 4573      4575 4576 4577 4578 4579 4580 4581 4582 4583 4584 4585 4586 4587 4588 4589 4590 4591      4593 4594 4595 4596 4597 4598     )
 #pragma warning (error: 4600 4601 4602 4603 4604      4606 4607      4609 4610 4611 4612 4613 4614 4615 4616 4617 4618      4620 4621 4622 4623 4624 4625 4626 4627 4628 4629 4630 4631 4632 4633 4634 4635 4636 4637 4638 4639 4640 4641 4642      4644 4645 4646 4647 4648 4649 4650      4652 4653 4654 4655 4656 4657 4658 4659 4660 4661 4662 4663 4664 4665 4666 4667      4669 4670 4671 4672 4673 4674 4675 4676 4677 4678 4679 4680 4681 4682 4683 4684 4685 4686 4687 4688 4689 4690 4691      4693 4694 4695 4696 4697 4698 4699)
@@ -44,7 +44,8 @@ static_assert(_MSC_VER < 1910 || _MSC_VER >= 1913, "Visual Studio 2017 version 1
 // 4265 - 'class' : class has virtual functions, but destructor is not virtual														http://msdn.microsoft.com/en-us/library/wzxffy8c.aspx
 // 4266 - 'function' : no override available for virtual member function from base 'type'; function is hidden						http://msdn.microsoft.com/en-us/library/4b76ty10.aspx
 // 4305 - 'identifier' : truncation from 'type1' to 'type2'																			http://msdn.microsoft.com/en-us/library/0as1ke3f.aspx
-// 4296 - 'operator' : expression is always false																					http://msdn.microsoft.com/en-us/library/wz2y40yt.aspx 
+// 4296 - 'operator' : expression is always false																					http://msdn.microsoft.com/en-us/library/wz2y40yt.aspx
+// 4307 - '' : integral constant overflow																							http://msdn.microsoft.com/en-us/library/4kze989h.aspx 
 // 4311 - 'variable' : pointer truncation from 'type' to 'type'																		http://msdn.microsoft.com/en-us/library/4t91x2k5.aspx
 // 4312 - 'operation' : conversion from 'type1' to 'type2' of greater size															http://msdn.microsoft.com/en-us/library/h97f4b9y.aspx
 // 4339 - 'type' : use of undefined type detected in CLR meta-data - use of this type may lead to a runtime exception				http://msdn.microsoft.com/en-us/library/3fxw8y6x.aspx
@@ -110,7 +111,8 @@ static_assert(_MSC_VER < 1910 || _MSC_VER >= 1913, "Visual Studio 2017 version 1
 #pragma warning(disable : 4244) // 'expression' : conversion from 'type' to 'type', possible loss of data						
 #pragma warning(disable : 4245) // 'initializing': conversion from 'type' to 'type', signed/unsigned mismatch 
 #pragma warning(disable : 4251) // 'type' needs to have dll-interface to be used by clients of 'type'
-#pragma warning(disable : 4291) // typedef-name '' used as synonym for class-name ''                    
+#pragma warning(disable : 4291) // typedef-name '' used as synonym for class-name ''    
+#pragma warning(disable : 4307) // '': integral constant overflow   
 //#pragma warning(disable : 4305) // 'argument' : truncation from 'double' to 'float' --- fp:precise builds only!!!!!!!
 #pragma warning(disable : 4324) // structure was padded due to __declspec(align())						
 #pragma warning(disable : 4355) // this used in base initializer list                                   
diff --git a/Engine/Source/Runtime/EmptyRHI/EmptyRHI.Build.cs b/Engine/Source/Runtime/EmptyRHI/EmptyRHI.Build.cs
--- a/Engine/Source/Runtime/EmptyRHI/EmptyRHI.Build.cs
+++ b/Engine/Source/Runtime/EmptyRHI/EmptyRHI.Build.cs
@@ -13,8 +13,7 @@ public class EmptyRHI : ModuleRules
 				"Core",
 				"Engine",
 				"RHI",
-				"RenderCore",
-				"ShaderCore"
+				"RenderCore"
 			}
 			);
 
diff --git a/Engine/Source/Runtime/Engine/Classes/Materials/MaterialParameterCollection.h b/Engine/Source/Runtime/Engine/Classes/Materials/MaterialParameterCollection.h
--- a/Engine/Source/Runtime/Engine/Classes/Materials/MaterialParameterCollection.h
+++ b/Engine/Source/Runtime/Engine/Classes/Materials/MaterialParameterCollection.h
@@ -118,7 +118,7 @@ class UMaterialParameterCollection : public UObject
 	ENGINE_API const FCollectionVectorParameter* GetVectorParameterByName(FName ParameterName) const;
 
 	/** Accessor for the uniform buffer layout description. */
-	const FUniformBufferStruct& GetUniformBufferStruct() const
+	const FShaderParametersMetadata& GetUniformBufferStruct() const
 	{
 		check(UniformBufferStruct);
 		return *UniformBufferStruct;
@@ -129,7 +129,7 @@ private:
 	/** Default resource used when no instance is available. */
 	class FMaterialParameterCollectionInstanceResource* DefaultResource;
 
-	TUniquePtr<FUniformBufferStruct> UniformBufferStruct;
+	TUniquePtr<FShaderParametersMetadata> UniformBufferStruct;
 
 	void CreateBufferStruct();
 
diff --git a/Engine/Source/Runtime/Engine/Private/Components/MaterialBillboardComponent.cpp b/Engine/Source/Runtime/Engine/Private/Components/MaterialBillboardComponent.cpp
--- a/Engine/Source/Runtime/Engine/Private/Components/MaterialBillboardComponent.cpp
+++ b/Engine/Source/Runtime/Engine/Private/Components/MaterialBillboardComponent.cpp
@@ -223,7 +223,7 @@ public:
 							BatchElement.FirstIndex = 0;
 							BatchElement.MinVertexIndex = 0;
 							BatchElement.MaxVertexIndex = 3;
-							BatchElement.PrimitiveUniformBuffer = GetUniformBuffer();
+							BatchElement.PrimitiveUniformBuffer = GetUniformBuffer().GetUniformBufferRef();
 							BatchElement.NumPrimitives = 2;
 							BatchElement.BaseVertexIndex = WriteOffset;
 
diff --git a/Engine/Source/Runtime/Engine/Private/EnginePrivatePCH.h b/Engine/Source/Runtime/Engine/Private/EnginePrivatePCH.h
--- a/Engine/Source/Runtime/Engine/Private/EnginePrivatePCH.h
+++ b/Engine/Source/Runtime/Engine/Private/EnginePrivatePCH.h
@@ -365,8 +365,6 @@
 #include "UniformBuffer.h"
 #include "PackedNormal.h"
 #include "RenderUtils.h"
-
-// From ShaderCore:
 #include "ShaderParameters.h"
 #include "ShaderCore.h"
 #include "Shader.h"
diff --git a/Engine/Source/Runtime/Engine/Private/GPUSkinVertexFactory.cpp b/Engine/Source/Runtime/Engine/Private/GPUSkinVertexFactory.cpp
--- a/Engine/Source/Runtime/Engine/Private/GPUSkinVertexFactory.cpp
+++ b/Engine/Source/Runtime/Engine/Private/GPUSkinVertexFactory.cpp
@@ -27,9 +27,9 @@ static TAutoConsoleVariable<int32> CVarGPUSkinLimit2BoneInfluences(
 	TEXT("Whether to use 2 bones influence instead of default 4 for GPU skinning. Cannot be changed at runtime."),
 	ECVF_ReadOnly);
 
-IMPLEMENT_UNIFORM_BUFFER_STRUCT(FAPEXClothUniformShaderParameters,TEXT("APEXClothParam"));
+IMPLEMENT_GLOBAL_SHADER_PARAMETER_STRUCT(FAPEXClothUniformShaderParameters,"APEXClothParam");
 
-IMPLEMENT_UNIFORM_BUFFER_STRUCT(FBoneMatricesUniformShaderParameters,TEXT("Bones"));
+IMPLEMENT_GLOBAL_SHADER_PARAMETER_STRUCT(FBoneMatricesUniformShaderParameters,"Bones");
 
 static FBoneMatricesUniformShaderParameters GBoneUniformStruct;
 
@@ -317,7 +317,7 @@ bool FGPUBaseSkinVertexFactory::FShaderDataType::UpdateBoneData(FRHICommandListI
 	}
 	else
 	{
-		UniformBuffer = RHICreateUniformBuffer(&GBoneUniformStruct, FBoneMatricesUniformShaderParameters::StaticStruct.GetLayout(), UniformBuffer_MultiFrame);
+		UniformBuffer = RHICreateUniformBuffer(&GBoneUniformStruct, FBoneMatricesUniformShaderParameters::StaticStructMetadata.GetLayout(), UniformBuffer_MultiFrame);
 	}
 	return false;
 }
diff --git a/Engine/Source/Runtime/Engine/Private/GPUSort.cpp b/Engine/Source/Runtime/Engine/Private/GPUSort.cpp
--- a/Engine/Source/Runtime/Engine/Private/GPUSort.cpp
+++ b/Engine/Source/Runtime/Engine/Private/GPUSort.cpp
@@ -49,15 +49,15 @@ void SetRadixSortShaderCompilerEnvironment( FShaderCompilerEnvironment& OutEnvir
 	Uniform buffer for passing in radix sort parameters.
 ------------------------------------------------------------------------------*/
 
-BEGIN_UNIFORM_BUFFER_STRUCT( FRadixSortParameters, )
-	UNIFORM_MEMBER( uint32, RadixShift )
-	UNIFORM_MEMBER( uint32, TilesPerGroup )
-	UNIFORM_MEMBER( uint32, ExtraTileCount )
-	UNIFORM_MEMBER( uint32, ExtraKeyCount )
-	UNIFORM_MEMBER( uint32, GroupCount )
-END_UNIFORM_BUFFER_STRUCT( FRadixSortParameters )
-
-IMPLEMENT_UNIFORM_BUFFER_STRUCT(FRadixSortParameters,TEXT("RadixSortUB"));
+BEGIN_GLOBAL_SHADER_PARAMETER_STRUCT( FRadixSortParameters, )
+	SHADER_PARAMETER( uint32, RadixShift )
+	SHADER_PARAMETER( uint32, TilesPerGroup )
+	SHADER_PARAMETER( uint32, ExtraTileCount )
+	SHADER_PARAMETER( uint32, ExtraKeyCount )
+	SHADER_PARAMETER( uint32, GroupCount )
+END_GLOBAL_SHADER_PARAMETER_STRUCT()
+
+IMPLEMENT_GLOBAL_SHADER_PARAMETER_STRUCT(FRadixSortParameters,"RadixSortUB");
 
 typedef TUniformBufferRef<FRadixSortParameters> FRadixSortUniformBufferRef;
 
diff --git a/Engine/Source/Runtime/Engine/Private/LocalVertexFactory.cpp b/Engine/Source/Runtime/Engine/Private/LocalVertexFactory.cpp
--- a/Engine/Source/Runtime/Engine/Private/LocalVertexFactory.cpp
+++ b/Engine/Source/Runtime/Engine/Private/LocalVertexFactory.cpp
@@ -41,7 +41,7 @@ void FLocalVertexFactoryShaderParameters::Serialize(FArchive& Ar)
 	Ar << LODParameter;
 }
 
-IMPLEMENT_UNIFORM_BUFFER_STRUCT(FLocalVertexFactoryUniformShaderParameters, TEXT("LocalVF"));
+IMPLEMENT_GLOBAL_SHADER_PARAMETER_STRUCT(FLocalVertexFactoryUniformShaderParameters, "LocalVF");
 
 TUniformBufferRef<FLocalVertexFactoryUniformShaderParameters> CreateLocalVFUniformBuffer(const FLocalVertexFactory* LocalVertexFactory, FColorVertexBuffer* OverrideColorVertexBuffer)
 {
diff --git a/Engine/Source/Runtime/Engine/Private/Materials/MaterialUniformExpressions.cpp b/Engine/Source/Runtime/Engine/Private/Materials/MaterialUniformExpressions.cpp
--- a/Engine/Source/Runtime/Engine/Private/Materials/MaterialUniformExpressions.cpp
+++ b/Engine/Source/Runtime/Engine/Private/Materials/MaterialUniformExpressions.cpp
@@ -227,30 +227,25 @@ void FUniformExpressionSet::SetParameterCollections(const TArray<UMaterialParame
 	}
 }
 
-static FShaderUniformBufferParameter* ConstructMaterialUniformBufferParameter()
-{
-	return nullptr;
-}
-
 
 static FName MaterialLayoutName(TEXT("Material"));
 
 void FUniformExpressionSet::CreateBufferStruct()
 {
 	// Make sure FUniformExpressionSet::CreateDebugLayout() is in sync
-	TArray<FUniformBufferStruct::FMember> Members;
+	TArray<FShaderParametersMetadata::FMember> Members;
 	uint32 NextMemberOffset = 0;
 
 	if (UniformVectorExpressions.Num())
 	{
-		new(Members) FUniformBufferStruct::FMember(TEXT("VectorExpressions"),TEXT(""),NextMemberOffset,UBMT_FLOAT32,EShaderPrecisionModifier::Half,1,4,UniformVectorExpressions.Num(),NULL);
+		new(Members) FShaderParametersMetadata::FMember(TEXT("VectorExpressions"),TEXT(""),NextMemberOffset,UBMT_FLOAT32,EShaderPrecisionModifier::Half,1,4,UniformVectorExpressions.Num(),NULL);
 		const uint32 VectorArraySize = UniformVectorExpressions.Num() * sizeof(FVector4);
 		NextMemberOffset += VectorArraySize;
 	}
 
 	if (UniformScalarExpressions.Num())
 	{
-		new(Members) FUniformBufferStruct::FMember(TEXT("ScalarExpressions"),TEXT(""),NextMemberOffset,UBMT_FLOAT32,EShaderPrecisionModifier::Half,1,4,(UniformScalarExpressions.Num() + 3) / 4,NULL);
+		new(Members) FShaderParametersMetadata::FMember(TEXT("ScalarExpressions"),TEXT(""),NextMemberOffset,UBMT_FLOAT32,EShaderPrecisionModifier::Half,1,4,(UniformScalarExpressions.Num() + 3) / 4,NULL);
 		const uint32 ScalarArraySize = (UniformScalarExpressions.Num() + 3) / 4 * sizeof(FVector4);
 		NextMemberOffset += ScalarArraySize;
 	}
@@ -287,58 +282,56 @@ void FUniformExpressionSet::CreateBufferStruct()
 	for (int32 i = 0; i < Uniform2DTextureExpressions.Num(); ++i)
 	{
 		check((NextMemberOffset & 0x7) == 0);
-		new(Members) FUniformBufferStruct::FMember(*Texture2DNames[i],TEXT("Texture2D"),NextMemberOffset,UBMT_TEXTURE,EShaderPrecisionModifier::Float,1,1,1,NULL);
+		new(Members) FShaderParametersMetadata::FMember(*Texture2DNames[i],TEXT("Texture2D"),NextMemberOffset,UBMT_TEXTURE,EShaderPrecisionModifier::Float,1,1,1,NULL);
 		NextMemberOffset += 8;
-		new(Members) FUniformBufferStruct::FMember(*Texture2DSamplerNames[i],TEXT("SamplerState"),NextMemberOffset,UBMT_SAMPLER,EShaderPrecisionModifier::Float,1,1,1,NULL);
+		new(Members) FShaderParametersMetadata::FMember(*Texture2DSamplerNames[i],TEXT("SamplerState"),NextMemberOffset,UBMT_SAMPLER,EShaderPrecisionModifier::Float,1,1,1,NULL);
 		NextMemberOffset += 8;
 	}
 
 	for (int32 i = 0; i < UniformCubeTextureExpressions.Num(); ++i)
 	{
 		check((NextMemberOffset & 0x7) == 0);
-		new(Members) FUniformBufferStruct::FMember(*TextureCubeNames[i],TEXT("TextureCube"),NextMemberOffset,UBMT_TEXTURE,EShaderPrecisionModifier::Float,1,1,1,NULL);
+		new(Members) FShaderParametersMetadata::FMember(*TextureCubeNames[i],TEXT("TextureCube"),NextMemberOffset,UBMT_TEXTURE,EShaderPrecisionModifier::Float,1,1,1,NULL);
 		NextMemberOffset += 8;
-		new(Members) FUniformBufferStruct::FMember(*TextureCubeSamplerNames[i],TEXT("SamplerState"),NextMemberOffset,UBMT_SAMPLER,EShaderPrecisionModifier::Float,1,1,1,NULL);
+		new(Members) FShaderParametersMetadata::FMember(*TextureCubeSamplerNames[i],TEXT("SamplerState"),NextMemberOffset,UBMT_SAMPLER,EShaderPrecisionModifier::Float,1,1,1,NULL);
 		NextMemberOffset += 8;
 	}
 
 	for (int32 i = 0; i < UniformVolumeTextureExpressions.Num(); ++i)
 	{
 		check((NextMemberOffset & 0x7) == 0);
-		new(Members) FUniformBufferStruct::FMember(*VolumeTextureNames[i],TEXT("Texture3D"),NextMemberOffset,UBMT_TEXTURE,EShaderPrecisionModifier::Float,1,1,1,NULL);
+		new(Members) FShaderParametersMetadata::FMember(*VolumeTextureNames[i],TEXT("Texture3D"),NextMemberOffset,UBMT_TEXTURE,EShaderPrecisionModifier::Float,1,1,1,NULL);
 		NextMemberOffset += 8;
-		new(Members) FUniformBufferStruct::FMember(*VolumeTextureSamplerNames[i],TEXT("SamplerState"),NextMemberOffset,UBMT_SAMPLER,EShaderPrecisionModifier::Float,1,1,1,NULL);
+		new(Members) FShaderParametersMetadata::FMember(*VolumeTextureSamplerNames[i],TEXT("SamplerState"),NextMemberOffset,UBMT_SAMPLER,EShaderPrecisionModifier::Float,1,1,1,NULL);
 		NextMemberOffset += 8;
 	}
 
 	for (int32 i = 0; i < UniformExternalTextureExpressions.Num(); ++i)
 	{
 		check((NextMemberOffset & 0x7) == 0);
-		new(Members) FUniformBufferStruct::FMember(*ExternalTextureNames[i], TEXT("TextureExternal"), NextMemberOffset, UBMT_TEXTURE, EShaderPrecisionModifier::Float, 1, 1, 1, NULL);
+		new(Members) FShaderParametersMetadata::FMember(*ExternalTextureNames[i], TEXT("TextureExternal"), NextMemberOffset, UBMT_TEXTURE, EShaderPrecisionModifier::Float, 1, 1, 1, NULL);
 		NextMemberOffset += 8;
-		new(Members) FUniformBufferStruct::FMember(*MediaTextureSamplerNames[i], TEXT("SamplerState"), NextMemberOffset, UBMT_SAMPLER, EShaderPrecisionModifier::Float, 1, 1, 1, NULL);
+		new(Members) FShaderParametersMetadata::FMember(*MediaTextureSamplerNames[i], TEXT("SamplerState"), NextMemberOffset, UBMT_SAMPLER, EShaderPrecisionModifier::Float, 1, 1, 1, NULL);
 		NextMemberOffset += 8;
 	}
 
-	new(Members) FUniformBufferStruct::FMember(TEXT("Wrap_WorldGroupSettings"),TEXT("SamplerState"),NextMemberOffset,UBMT_SAMPLER,EShaderPrecisionModifier::Float,1,1,1,NULL);
+	new(Members) FShaderParametersMetadata::FMember(TEXT("Wrap_WorldGroupSettings"),TEXT("SamplerState"),NextMemberOffset,UBMT_SAMPLER,EShaderPrecisionModifier::Float,1,1,1,NULL);
 	NextMemberOffset += 8;
 
-	new(Members) FUniformBufferStruct::FMember(TEXT("Clamp_WorldGroupSettings"),TEXT("SamplerState"),NextMemberOffset,UBMT_SAMPLER,EShaderPrecisionModifier::Float,1,1,1,NULL);
+	new(Members) FShaderParametersMetadata::FMember(TEXT("Clamp_WorldGroupSettings"),TEXT("SamplerState"),NextMemberOffset,UBMT_SAMPLER,EShaderPrecisionModifier::Float,1,1,1,NULL);
 	NextMemberOffset += 8;
 
-	const uint32 StructSize = Align(NextMemberOffset,UNIFORM_BUFFER_STRUCT_ALIGNMENT);
+	const uint32 StructSize = Align(NextMemberOffset, SHADER_PARAMETER_STRUCT_ALIGNMENT);
 	UniformBufferStruct.Emplace(
+		FShaderParametersMetadata::EUseCase::DataDrivenShaderParameterStruct,
 		MaterialLayoutName,
 		TEXT("MaterialUniforms"),
 		TEXT("Material"),
-		ConstructMaterialUniformBufferParameter,
 		StructSize,
-		Members,
-		false
-		);
+		Members);
 }
 
-const FUniformBufferStruct& FUniformExpressionSet::GetUniformBufferStruct() const
+const FShaderParametersMetadata& FUniformExpressionSet::GetUniformBufferStruct() const
 {
 	return UniformBufferStruct.GetValue();
 }
@@ -353,7 +346,7 @@ FUniformBufferRHIRef FUniformExpressionSet::CreateUniformBuffer(const FMaterialR
 	if (UniformBufferStruct->GetSize() > 0)
 	{
 		FMemMark Mark(FMemStack::Get());
-		void* const TempBuffer = FMemStack::Get().PushBytes(UniformBufferStruct->GetSize(),UNIFORM_BUFFER_STRUCT_ALIGNMENT);
+		void* const TempBuffer = FMemStack::Get().PushBytes(UniformBufferStruct->GetSize(), SHADER_PARAMETER_STRUCT_ALIGNMENT);
 		checkf(TempBuffer, TEXT("Failed to allocate uniform buffer struct of %i bytes."), UniformBufferStruct->GetSize());
 
 		FLinearColor* TempVectorBuffer = (FLinearColor*)TempBuffer;
diff --git a/Engine/Source/Runtime/Engine/Private/Materials/ParameterCollection.cpp b/Engine/Source/Runtime/Engine/Private/Materials/ParameterCollection.cpp
--- a/Engine/Source/Runtime/Engine/Private/Materials/ParameterCollection.cpp
+++ b/Engine/Source/Runtime/Engine/Private/Materials/ParameterCollection.cpp
@@ -413,31 +413,28 @@ const FCollectionVectorParameter* UMaterialParameterCollection::GetVectorParamet
 	return nullptr;
 }
 
-FShaderUniformBufferParameter* ConstructCollectionUniformBufferParameter() { return nullptr; }
-
 void UMaterialParameterCollection::CreateBufferStruct()
 {	
-	TArray<FUniformBufferStruct::FMember> Members;
+	TArray<FShaderParametersMetadata::FMember> Members;
 	uint32 NextMemberOffset = 0;
 
 	const uint32 NumVectors = FMath::DivideAndRoundUp(ScalarParameters.Num(), 4) + VectorParameters.Num();
-	new(Members) FUniformBufferStruct::FMember(TEXT("Vectors"),TEXT(""),NextMemberOffset,UBMT_FLOAT32,EShaderPrecisionModifier::Half,1,4,NumVectors, nullptr);
+	new(Members) FShaderParametersMetadata::FMember(TEXT("Vectors"),TEXT(""),NextMemberOffset,UBMT_FLOAT32,EShaderPrecisionModifier::Half,1,4,NumVectors, nullptr);
 	const uint32 VectorArraySize = NumVectors * sizeof(FVector4);
 	NextMemberOffset += VectorArraySize;
 	static FName LayoutName(TEXT("MaterialCollection"));
-	const uint32 StructSize = Align(NextMemberOffset,UNIFORM_BUFFER_STRUCT_ALIGNMENT);
+	const uint32 StructSize = Align(NextMemberOffset, SHADER_PARAMETER_STRUCT_ALIGNMENT);
 
 	// If Collections ever get non-numeric resources (eg Textures), OutEnvironment.ResourceTableMap has a map by name
 	// and the N ParameterCollection Uniform Buffers ALL are named "MaterialCollection" with different hashes!
 	// (and the hlsl cbuffers are named MaterialCollection0, etc, so the names don't match the layout)
-	UniformBufferStruct = MakeUnique<FUniformBufferStruct>(
+	UniformBufferStruct = MakeUnique<FShaderParametersMetadata>(
+		FShaderParametersMetadata::EUseCase::DataDrivenShaderParameterStruct,
 		LayoutName,
 		TEXT("MaterialCollection"),
 		TEXT("MaterialCollection"),
-		ConstructCollectionUniformBufferParameter,
 		StructSize,
-		Members,
-		false
+		Members
 		);
 }
 
diff --git a/Engine/Source/Runtime/Engine/Private/MeshParticleVertexFactory.cpp b/Engine/Source/Runtime/Engine/Private/MeshParticleVertexFactory.cpp
--- a/Engine/Source/Runtime/Engine/Private/MeshParticleVertexFactory.cpp
+++ b/Engine/Source/Runtime/Engine/Private/MeshParticleVertexFactory.cpp
@@ -338,4 +338,4 @@ FVertexFactoryShaderParameters* FMeshParticleVertexFactory::ConstructShaderParam
 
 IMPLEMENT_VERTEX_FACTORY_TYPE(FMeshParticleVertexFactory,"/Engine/Private/MeshParticleVertexFactory.ush",true,false,true,false,false);
 IMPLEMENT_VERTEX_FACTORY_TYPE(FMeshParticleVertexFactoryEmulatedInstancing,"/Engine/Private/MeshParticleVertexFactory.ush",true,false,true,false,false);
-IMPLEMENT_UNIFORM_BUFFER_STRUCT(FMeshParticleUniformParameters,TEXT("MeshParticleVF"));
+IMPLEMENT_GLOBAL_SHADER_PARAMETER_STRUCT(FMeshParticleUniformParameters,"MeshParticleVF");
diff --git a/Engine/Source/Runtime/Engine/Private/Particles/FXSystem.cpp b/Engine/Source/Runtime/Engine/Private/Particles/FXSystem.cpp
--- a/Engine/Source/Runtime/Engine/Private/Particles/FXSystem.cpp
+++ b/Engine/Source/Runtime/Engine/Private/Particles/FXSystem.cpp
@@ -373,7 +373,7 @@ void FFXSystem::PreRender(FRHICommandListImmediate& RHICmdList, const FGlobalDis
 void FFXSystem::PostRenderOpaque(
 	FRHICommandListImmediate& RHICmdList, 
 	const FUniformBufferRHIParamRef ViewUniformBuffer, 
-	const FUniformBufferStruct* SceneTexturesUniformBufferStruct,
+	const FShaderParametersMetadata* SceneTexturesUniformBufferStruct,
 	FUniformBufferRHIParamRef SceneTexturesUniformBuffer)
 {
 	if (RHISupportsGPUParticles() && IsParticleCollisionModeSupported(GetShaderPlatform(), PCM_DepthBuffer))
diff --git a/Engine/Source/Runtime/Engine/Private/Particles/FXSystemPrivate.h b/Engine/Source/Runtime/Engine/Private/Particles/FXSystemPrivate.h
--- a/Engine/Source/Runtime/Engine/Private/Particles/FXSystemPrivate.h
+++ b/Engine/Source/Runtime/Engine/Private/Particles/FXSystemPrivate.h
@@ -152,7 +152,7 @@ public:
 	virtual void PostRenderOpaque(
 		FRHICommandListImmediate& RHICmdList, 
 		const FUniformBufferRHIParamRef ViewUniformBuffer, 
-		const FUniformBufferStruct* SceneTexturesUniformBufferStruct,
+		const FShaderParametersMetadata* SceneTexturesUniformBufferStruct,
 		FUniformBufferRHIParamRef SceneTexturesUniformBuffer) override;
 	// End FFXSystemInterface.
 
@@ -254,7 +254,7 @@ private:
 		EParticleSimulatePhase::Type Phase,
 		const FUniformBufferRHIParamRef ViewUniformBuffer,
 		const FGlobalDistanceFieldParameterData* GlobalDistanceFieldParameterData,
-		const FUniformBufferStruct* SceneTexturesUniformBufferStruct,
+		const FShaderParametersMetadata* SceneTexturesUniformBufferStruct,
 		FUniformBufferRHIParamRef SceneTexturesUniformBuffer
 		);
 
diff --git a/Engine/Source/Runtime/Engine/Private/Particles/ParticleBeamTrailVertexFactory.cpp b/Engine/Source/Runtime/Engine/Private/Particles/ParticleBeamTrailVertexFactory.cpp
--- a/Engine/Source/Runtime/Engine/Private/Particles/ParticleBeamTrailVertexFactory.cpp
+++ b/Engine/Source/Runtime/Engine/Private/Particles/ParticleBeamTrailVertexFactory.cpp
@@ -8,7 +8,7 @@
 #include "ParticleHelper.h"
 #include "ShaderParameterUtils.h"
 
-IMPLEMENT_UNIFORM_BUFFER_STRUCT(FParticleBeamTrailUniformParameters,TEXT("BeamTrailVF"));
+IMPLEMENT_GLOBAL_SHADER_PARAMETER_STRUCT(FParticleBeamTrailUniformParameters,"BeamTrailVF");
 
 /**
  * Shader parameters for the beam/trail vertex factory.
diff --git a/Engine/Source/Runtime/Engine/Private/Particles/ParticleCurveTexture.cpp b/Engine/Source/Runtime/Engine/Private/Particles/ParticleCurveTexture.cpp
--- a/Engine/Source/Runtime/Engine/Private/Particles/ParticleCurveTexture.cpp
+++ b/Engine/Source/Runtime/Engine/Private/Particles/ParticleCurveTexture.cpp
@@ -38,12 +38,12 @@ Shaders used for uploading curves to the GPU.
 /**
 * Uniform buffer to hold parameters for particle curve injection.
 */
-BEGIN_UNIFORM_BUFFER_STRUCT(FParticleCurveInjectionParameters, )
-UNIFORM_MEMBER(FVector2D, PixelScale)
-UNIFORM_MEMBER(FVector2D, CurveOffset)
-END_UNIFORM_BUFFER_STRUCT(FParticleCurveInjectionParameters)
+BEGIN_GLOBAL_SHADER_PARAMETER_STRUCT(FParticleCurveInjectionParameters, )
+SHADER_PARAMETER(FVector2D, PixelScale)
+SHADER_PARAMETER(FVector2D, CurveOffset)
+END_GLOBAL_SHADER_PARAMETER_STRUCT()
 
-IMPLEMENT_UNIFORM_BUFFER_STRUCT(FParticleCurveInjectionParameters, TEXT("ParticleCurveInjection"));
+IMPLEMENT_GLOBAL_SHADER_PARAMETER_STRUCT(FParticleCurveInjectionParameters, "ParticleCurveInjection");
 
 typedef TUniformBufferRef<FParticleCurveInjectionParameters> FParticleCurveInjectionBufferRef;
 
diff --git a/Engine/Source/Runtime/Engine/Private/Particles/ParticleGpuSimulation.cpp b/Engine/Source/Runtime/Engine/Private/Particles/ParticleGpuSimulation.cpp
--- a/Engine/Source/Runtime/Engine/Private/Particles/ParticleGpuSimulation.cpp
+++ b/Engine/Source/Runtime/Engine/Private/Particles/ParticleGpuSimulation.cpp
@@ -526,41 +526,41 @@ private:
 /**
  * Uniform buffer for GPU particle sprite emitters.
  */
-BEGIN_UNIFORM_BUFFER_STRUCT(FGPUSpriteEmitterUniformParameters,)
-	UNIFORM_MEMBER(FVector4, ColorCurve)
-	UNIFORM_MEMBER(FVector4, ColorScale)
-	UNIFORM_MEMBER(FVector4, ColorBias)
-	UNIFORM_MEMBER(FVector4, MiscCurve)
-	UNIFORM_MEMBER(FVector4, MiscScale)
-	UNIFORM_MEMBER(FVector4, MiscBias)
-	UNIFORM_MEMBER(FVector4, SizeBySpeed)
-	UNIFORM_MEMBER(FVector4, SubImageSize)
-	UNIFORM_MEMBER(FVector4, TangentSelector)
-	UNIFORM_MEMBER(FVector, CameraFacingBlend)
-	UNIFORM_MEMBER(float, RemoveHMDRoll)
-	UNIFORM_MEMBER(float, RotationRateScale)
-	UNIFORM_MEMBER(float, RotationBias)
-	UNIFORM_MEMBER(float, CameraMotionBlurAmount)
-	UNIFORM_MEMBER(FVector2D, PivotOffset)
-END_UNIFORM_BUFFER_STRUCT(FGPUSpriteEmitterUniformParameters)
-
-IMPLEMENT_UNIFORM_BUFFER_STRUCT(FGPUSpriteEmitterUniformParameters,TEXT("EmitterUniforms"));
+BEGIN_GLOBAL_SHADER_PARAMETER_STRUCT(FGPUSpriteEmitterUniformParameters,)
+	SHADER_PARAMETER(FVector4, ColorCurve)
+	SHADER_PARAMETER(FVector4, ColorScale)
+	SHADER_PARAMETER(FVector4, ColorBias)
+	SHADER_PARAMETER(FVector4, MiscCurve)
+	SHADER_PARAMETER(FVector4, MiscScale)
+	SHADER_PARAMETER(FVector4, MiscBias)
+	SHADER_PARAMETER(FVector4, SizeBySpeed)
+	SHADER_PARAMETER(FVector4, SubImageSize)
+	SHADER_PARAMETER(FVector4, TangentSelector)
+	SHADER_PARAMETER(FVector, CameraFacingBlend)
+	SHADER_PARAMETER(float, RemoveHMDRoll)
+	SHADER_PARAMETER(float, RotationRateScale)
+	SHADER_PARAMETER(float, RotationBias)
+	SHADER_PARAMETER(float, CameraMotionBlurAmount)
+	SHADER_PARAMETER(FVector2D, PivotOffset)
+END_GLOBAL_SHADER_PARAMETER_STRUCT()
+
+IMPLEMENT_GLOBAL_SHADER_PARAMETER_STRUCT(FGPUSpriteEmitterUniformParameters, "EmitterUniforms");
 
 typedef TUniformBufferRef<FGPUSpriteEmitterUniformParameters> FGPUSpriteEmitterUniformBufferRef;
 
 /**
  * Uniform buffer to hold dynamic parameters for GPU particle sprite emitters.
  */
-BEGIN_UNIFORM_BUFFER_STRUCT( FGPUSpriteEmitterDynamicUniformParameters, )
-	UNIFORM_MEMBER( FVector2D, LocalToWorldScale )
-	UNIFORM_MEMBER( float, EmitterInstRandom)
-	UNIFORM_MEMBER( FVector4, AxisLockRight )
-	UNIFORM_MEMBER( FVector4, AxisLockUp )
-	UNIFORM_MEMBER( FVector4, DynamicColor)
-	UNIFORM_MEMBER( FVector4, MacroUVParameters )
-END_UNIFORM_BUFFER_STRUCT( FGPUSpriteEmitterDynamicUniformParameters )
+BEGIN_GLOBAL_SHADER_PARAMETER_STRUCT( FGPUSpriteEmitterDynamicUniformParameters, )
+	SHADER_PARAMETER( FVector2D, LocalToWorldScale )
+	SHADER_PARAMETER( float, EmitterInstRandom)
+	SHADER_PARAMETER( FVector4, AxisLockRight )
+	SHADER_PARAMETER( FVector4, AxisLockUp )
+	SHADER_PARAMETER( FVector4, DynamicColor)
+	SHADER_PARAMETER( FVector4, MacroUVParameters )
+END_GLOBAL_SHADER_PARAMETER_STRUCT()
 
-IMPLEMENT_UNIFORM_BUFFER_STRUCT(FGPUSpriteEmitterDynamicUniformParameters,TEXT("EmitterDynamicUniforms"));
+IMPLEMENT_GLOBAL_SHADER_PARAMETER_STRUCT(FGPUSpriteEmitterDynamicUniformParameters, "EmitterDynamicUniforms");
 
 typedef TUniformBufferRef<FGPUSpriteEmitterDynamicUniformParameters> FGPUSpriteEmitterDynamicUniformBufferRef;
 
@@ -822,31 +822,31 @@ IMPLEMENT_VERTEX_FACTORY_TYPE(FGPUSpriteVertexFactory,"/Engine/Private/ParticleG
 /**
  * Uniform buffer to hold parameters for particle simulation.
  */
-BEGIN_UNIFORM_BUFFER_STRUCT(FParticleSimulationParameters,)
-	UNIFORM_MEMBER(FVector4, AttributeCurve)
-	UNIFORM_MEMBER(FVector4, AttributeCurveScale)
-	UNIFORM_MEMBER(FVector4, AttributeCurveBias)
-	UNIFORM_MEMBER(FVector4, AttributeScale)
-	UNIFORM_MEMBER(FVector4, AttributeBias)
-	UNIFORM_MEMBER(FVector4, MiscCurve)
-	UNIFORM_MEMBER(FVector4, MiscScale)
-	UNIFORM_MEMBER(FVector4, MiscBias)
-	UNIFORM_MEMBER(FVector, Acceleration)
-	UNIFORM_MEMBER(FVector, OrbitOffsetBase)
-	UNIFORM_MEMBER(FVector, OrbitOffsetRange)
-	UNIFORM_MEMBER(FVector, OrbitFrequencyBase)
-	UNIFORM_MEMBER(FVector, OrbitFrequencyRange)
-	UNIFORM_MEMBER(FVector, OrbitPhaseBase)
-	UNIFORM_MEMBER(FVector, OrbitPhaseRange)
-	UNIFORM_MEMBER(float, CollisionRadiusScale)
-	UNIFORM_MEMBER(float, CollisionRadiusBias)
-	UNIFORM_MEMBER(float, CollisionTimeBias)
-	UNIFORM_MEMBER(float, CollisionRandomSpread)
-	UNIFORM_MEMBER(float, CollisionRandomDistribution)
-	UNIFORM_MEMBER(float, OneMinusFriction)
-END_UNIFORM_BUFFER_STRUCT(FParticleSimulationParameters)
-
-IMPLEMENT_UNIFORM_BUFFER_STRUCT(FParticleSimulationParameters,TEXT("Simulation"));
+BEGIN_GLOBAL_SHADER_PARAMETER_STRUCT(FParticleSimulationParameters,)
+	SHADER_PARAMETER(FVector4, AttributeCurve)
+	SHADER_PARAMETER(FVector4, AttributeCurveScale)
+	SHADER_PARAMETER(FVector4, AttributeCurveBias)
+	SHADER_PARAMETER(FVector4, AttributeScale)
+	SHADER_PARAMETER(FVector4, AttributeBias)
+	SHADER_PARAMETER(FVector4, MiscCurve)
+	SHADER_PARAMETER(FVector4, MiscScale)
+	SHADER_PARAMETER(FVector4, MiscBias)
+	SHADER_PARAMETER(FVector, Acceleration)
+	SHADER_PARAMETER(FVector, OrbitOffsetBase)
+	SHADER_PARAMETER(FVector, OrbitOffsetRange)
+	SHADER_PARAMETER(FVector, OrbitFrequencyBase)
+	SHADER_PARAMETER(FVector, OrbitFrequencyRange)
+	SHADER_PARAMETER(FVector, OrbitPhaseBase)
+	SHADER_PARAMETER(FVector, OrbitPhaseRange)
+	SHADER_PARAMETER(float, CollisionRadiusScale)
+	SHADER_PARAMETER(float, CollisionRadiusBias)
+	SHADER_PARAMETER(float, CollisionTimeBias)
+	SHADER_PARAMETER(float, CollisionRandomSpread)
+	SHADER_PARAMETER(float, CollisionRandomDistribution)
+	SHADER_PARAMETER(float, OneMinusFriction)
+END_GLOBAL_SHADER_PARAMETER_STRUCT()
+
+IMPLEMENT_GLOBAL_SHADER_PARAMETER_STRUCT(FParticleSimulationParameters, "Simulation");
 
 typedef TUniformBufferRef<FParticleSimulationParameters> FParticleSimulationBufferRef;
 
@@ -949,16 +949,16 @@ FArchive& operator<<(FArchive& Ar, FParticlePerFrameSimulationShaderParameters&
  * Uniform buffer to hold parameters for vector fields sampled during particle
  * simulation.
  */
-BEGIN_UNIFORM_BUFFER_STRUCT( FVectorFieldUniformParameters,)
-	UNIFORM_MEMBER( int32, Count )
-	UNIFORM_MEMBER_ARRAY( FMatrix, WorldToVolume, [MAX_VECTOR_FIELDS] )
-	UNIFORM_MEMBER_ARRAY( FMatrix, VolumeToWorld, [MAX_VECTOR_FIELDS] )
-	UNIFORM_MEMBER_ARRAY( FVector4, IntensityAndTightness, [MAX_VECTOR_FIELDS] )
-	UNIFORM_MEMBER_ARRAY( FVector4, VolumeSize, [MAX_VECTOR_FIELDS] )
-	UNIFORM_MEMBER_ARRAY( FVector4, TilingAxes, [MAX_VECTOR_FIELDS] )
-END_UNIFORM_BUFFER_STRUCT( FVectorFieldUniformParameters )
+BEGIN_GLOBAL_SHADER_PARAMETER_STRUCT( FVectorFieldUniformParameters,)
+	SHADER_PARAMETER( int32, Count )
+	SHADER_PARAMETER_ARRAY( FMatrix, WorldToVolume, [MAX_VECTOR_FIELDS] )
+	SHADER_PARAMETER_ARRAY( FMatrix, VolumeToWorld, [MAX_VECTOR_FIELDS] )
+	SHADER_PARAMETER_ARRAY( FVector4, IntensityAndTightness, [MAX_VECTOR_FIELDS] )
+	SHADER_PARAMETER_ARRAY( FVector4, VolumeSize, [MAX_VECTOR_FIELDS] )
+	SHADER_PARAMETER_ARRAY( FVector4, TilingAxes, [MAX_VECTOR_FIELDS] )
+END_GLOBAL_SHADER_PARAMETER_STRUCT()
 
-IMPLEMENT_UNIFORM_BUFFER_STRUCT(FVectorFieldUniformParameters,TEXT("VectorFields"));
+IMPLEMENT_GLOBAL_SHADER_PARAMETER_STRUCT(FVectorFieldUniformParameters, "VectorFields");
 
 typedef TUniformBufferRef<FVectorFieldUniformParameters> FVectorFieldUniformBufferRef;
 
@@ -1126,7 +1126,7 @@ public:
 		const FUniformBufferRHIParamRef ViewUniformBuffer,
 		ERHIFeatureLevel::Type FeatureLevel,
 		const FGlobalDistanceFieldParameterData* GlobalDistanceFieldParameterData,
-		const FUniformBufferStruct* SceneTexturesUniformBufferStruct,
+		const FShaderParametersMetadata* SceneTexturesUniformBufferStruct,
 		FUniformBufferRHIParamRef SceneTexturesUniformBuffer
 		)
 	{
@@ -1500,7 +1500,7 @@ void ExecuteSimulationCommands(
 	FParticleSimulationResources* ParticleSimulationResources,
 	const FUniformBufferRHIParamRef ViewUniformBuffer,
 	const FGlobalDistanceFieldParameterData* GlobalDistanceFieldParameterData,
-	const FUniformBufferStruct* SceneTexturesUniformBufferStruct,
+	const FShaderParametersMetadata* SceneTexturesUniformBufferStruct,
 	FUniformBufferRHIParamRef SceneTexturesUniformBuffer,
 	bool bUseFixDT)
 {
@@ -1569,7 +1569,7 @@ void ExecuteSimulationCommands(
 	FParticleSimulationResources* ParticleSimulationResources,
 	const FUniformBufferRHIParamRef ViewUniformBuffer,
 	const FGlobalDistanceFieldParameterData* GlobalDistanceFieldParameterData,
-	const FUniformBufferStruct* SceneTexturesUniformBufferStruct,
+	const FShaderParametersMetadata* SceneTexturesUniformBufferStruct,
 	FUniformBufferRHIParamRef SceneTexturesUniformBuffer,
 	EParticleSimulatePhase::Type Phase,
 	bool bUseFixDT)
@@ -1680,11 +1680,11 @@ void ClearTiles(FRHICommandList& RHICmdList, FGraphicsPipelineStateInitializer&
 /**
  * Uniform buffer to hold parameters for particle simulation.
  */
-BEGIN_UNIFORM_BUFFER_STRUCT( FParticleInjectionParameters, )
-	UNIFORM_MEMBER( FVector2D, PixelScale )
-END_UNIFORM_BUFFER_STRUCT( FParticleInjectionParameters )
+BEGIN_GLOBAL_SHADER_PARAMETER_STRUCT( FParticleInjectionParameters, )
+	SHADER_PARAMETER( FVector2D, PixelScale )
+END_GLOBAL_SHADER_PARAMETER_STRUCT()
 
-IMPLEMENT_UNIFORM_BUFFER_STRUCT(FParticleInjectionParameters,TEXT("ParticleInjection"));
+IMPLEMENT_GLOBAL_SHADER_PARAMETER_STRUCT(FParticleInjectionParameters, "ParticleInjection");
 
 typedef TUniformBufferRef<FParticleInjectionParameters> FParticleInjectionBufferRef;
 
@@ -1925,11 +1925,11 @@ void InjectNewParticles(FRHICommandList& RHICmdList, FGraphicsPipelineStateIniti
 /**
  * Uniform buffer to hold parameters for visualizing particle simulation.
  */
-BEGIN_UNIFORM_BUFFER_STRUCT( FParticleSimVisualizeParameters, )
-	UNIFORM_MEMBER( FVector4, ScaleBias )
-END_UNIFORM_BUFFER_STRUCT( FParticleSimVisualizeParameters )
+BEGIN_GLOBAL_SHADER_PARAMETER_STRUCT( FParticleSimVisualizeParameters, )
+	SHADER_PARAMETER( FVector4, ScaleBias )
+END_GLOBAL_SHADER_PARAMETER_STRUCT()
 
-IMPLEMENT_UNIFORM_BUFFER_STRUCT(FParticleSimVisualizeParameters,TEXT("PSV"));
+IMPLEMENT_GLOBAL_SHADER_PARAMETER_STRUCT(FParticleSimVisualizeParameters,"PSV");
 
 typedef TUniformBufferRef<FParticleSimVisualizeParameters> FParticleSimVisualizeBufferRef;
 
@@ -2186,13 +2186,13 @@ static void BuildParticleVertexBuffer( FVertexBufferRHIParamRef VertexBufferRHI,
 /**
  * Uniform buffer parameters for generating particle bounds.
  */
-BEGIN_UNIFORM_BUFFER_STRUCT( FParticleBoundsParameters, )
-	UNIFORM_MEMBER( uint32, ChunksPerGroup )
-	UNIFORM_MEMBER( uint32, ExtraChunkCount )
-	UNIFORM_MEMBER( uint32, ParticleCount )
-END_UNIFORM_BUFFER_STRUCT( FParticleBoundsParameters )
+BEGIN_GLOBAL_SHADER_PARAMETER_STRUCT( FParticleBoundsParameters, )
+	SHADER_PARAMETER( uint32, ChunksPerGroup )
+	SHADER_PARAMETER( uint32, ExtraChunkCount )
+	SHADER_PARAMETER( uint32, ParticleCount )
+END_GLOBAL_SHADER_PARAMETER_STRUCT()
 
-IMPLEMENT_UNIFORM_BUFFER_STRUCT(FParticleBoundsParameters,TEXT("ParticleBounds"));
+IMPLEMENT_GLOBAL_SHADER_PARAMETER_STRUCT(FParticleBoundsParameters,"ParticleBounds");
 
 typedef TUniformBufferRef<FParticleBoundsParameters> FParticleBoundsUniformBufferRef;
 
@@ -4686,7 +4686,7 @@ void FFXSystem::SimulateGPUParticles(
 	EParticleSimulatePhase::Type Phase,
 	const FUniformBufferRHIParamRef ViewUniformBuffer,
 	const FGlobalDistanceFieldParameterData* GlobalDistanceFieldParameterData,
-	const FUniformBufferStruct* SceneTexturesUniformBufferStruct,
+	const FShaderParametersMetadata* SceneTexturesUniformBufferStruct,
 	FUniformBufferRHIParamRef SceneTexturesUniformBuffer
 	)
 {
diff --git a/Engine/Source/Runtime/Engine/Private/Particles/ParticleSortingGPU.cpp b/Engine/Source/Runtime/Engine/Private/Particles/ParticleSortingGPU.cpp
--- a/Engine/Source/Runtime/Engine/Private/Particles/ParticleSortingGPU.cpp
+++ b/Engine/Source/Runtime/Engine/Private/Particles/ParticleSortingGPU.cpp
@@ -25,16 +25,16 @@
 /**
  * Uniform buffer parameters for generating particle sort keys.
  */
-BEGIN_UNIFORM_BUFFER_STRUCT( FParticleKeyGenParameters, )
-	UNIFORM_MEMBER( FVector4, ViewOrigin )
-	UNIFORM_MEMBER( uint32, ChunksPerGroup )
-	UNIFORM_MEMBER( uint32, ExtraChunkCount )
-	UNIFORM_MEMBER( uint32, OutputOffset )
-	UNIFORM_MEMBER( uint32, EmitterKey )
-	UNIFORM_MEMBER( uint32, KeyCount )
-END_UNIFORM_BUFFER_STRUCT( FParticleKeyGenParameters )
-
-IMPLEMENT_UNIFORM_BUFFER_STRUCT(FParticleKeyGenParameters,TEXT("ParticleKeyGen"));
+BEGIN_GLOBAL_SHADER_PARAMETER_STRUCT( FParticleKeyGenParameters, )
+	SHADER_PARAMETER( FVector4, ViewOrigin )
+	SHADER_PARAMETER( uint32, ChunksPerGroup )
+	SHADER_PARAMETER( uint32, ExtraChunkCount )
+	SHADER_PARAMETER( uint32, OutputOffset )
+	SHADER_PARAMETER( uint32, EmitterKey )
+	SHADER_PARAMETER( uint32, KeyCount )
+END_GLOBAL_SHADER_PARAMETER_STRUCT()
+
+IMPLEMENT_GLOBAL_SHADER_PARAMETER_STRUCT(FParticleKeyGenParameters,"ParticleKeyGen");
 
 typedef TUniformBufferRef<FParticleKeyGenParameters> FParticleKeyGenUniformBufferRef;
 
diff --git a/Engine/Source/Runtime/Engine/Private/Particles/ParticleVertexFactory.cpp b/Engine/Source/Runtime/Engine/Private/Particles/ParticleVertexFactory.cpp
--- a/Engine/Source/Runtime/Engine/Private/Particles/ParticleVertexFactory.cpp
+++ b/Engine/Source/Runtime/Engine/Private/Particles/ParticleVertexFactory.cpp
@@ -9,7 +9,7 @@
 #include "ParticleResources.h"
 #include "ShaderParameterUtils.h"
 
-IMPLEMENT_UNIFORM_BUFFER_STRUCT(FParticleSpriteUniformParameters,TEXT("SpriteVF"));
+IMPLEMENT_GLOBAL_SHADER_PARAMETER_STRUCT(FParticleSpriteUniformParameters, "SpriteVF");
 
 TGlobalResource<FNullDynamicParameterVertexBuffer> GNullDynamicParameterVertexBuffer;
 
diff --git a/Engine/Source/Runtime/Engine/Private/SceneView.cpp b/Engine/Source/Runtime/Engine/Private/SceneView.cpp
--- a/Engine/Source/Runtime/Engine/Private/SceneView.cpp
+++ b/Engine/Source/Runtime/Engine/Private/SceneView.cpp
@@ -31,10 +31,10 @@ DEFINE_LOG_CATEGORY(LogBufferVisualization);
 DECLARE_CYCLE_STAT(TEXT("StartFinalPostprocessSettings"), STAT_StartFinalPostprocessSettings, STATGROUP_Engine);
 DECLARE_CYCLE_STAT(TEXT("OverridePostProcessSettings"), STAT_OverridePostProcessSettings, STATGROUP_Engine);
 
-IMPLEMENT_UNIFORM_BUFFER_STRUCT(FPrimitiveUniformShaderParameters,TEXT("Primitive"));
-IMPLEMENT_UNIFORM_BUFFER_STRUCT(FViewUniformShaderParameters,TEXT("View"));
-IMPLEMENT_UNIFORM_BUFFER_STRUCT(FInstancedViewUniformShaderParameters, TEXT("InstancedView"));
-IMPLEMENT_UNIFORM_BUFFER_STRUCT(FMobileDirectionalLightShaderParameters, TEXT("MobileDirectionalLight"));
+IMPLEMENT_GLOBAL_SHADER_PARAMETER_STRUCT(FPrimitiveUniformShaderParameters, "Primitive");
+IMPLEMENT_GLOBAL_SHADER_PARAMETER_STRUCT(FViewUniformShaderParameters, "View");
+IMPLEMENT_GLOBAL_SHADER_PARAMETER_STRUCT(FInstancedViewUniformShaderParameters, "InstancedView");
+IMPLEMENT_GLOBAL_SHADER_PARAMETER_STRUCT(FMobileDirectionalLightShaderParameters, "MobileDirectionalLight");
 
 
 static TAutoConsoleVariable<float> CVarSSRMaxRoughness(
diff --git a/Engine/Source/Runtime/Engine/Private/ShaderCompiler/ShaderCompiler.cpp b/Engine/Source/Runtime/Engine/Private/ShaderCompiler/ShaderCompiler.cpp
--- a/Engine/Source/Runtime/Engine/Private/ShaderCompiler/ShaderCompiler.cpp
+++ b/Engine/Source/Runtime/Engine/Private/ShaderCompiler/ShaderCompiler.cpp
@@ -51,7 +51,7 @@
 
 DEFINE_LOG_CATEGORY(LogShaderCompilers);
 
-SHADERCORE_API bool UsePreExposure(EShaderPlatform Platform);
+RENDERCORE_API bool UsePreExposure(EShaderPlatform Platform);
 
 #if ENABLE_COOK_STATS
 namespace GlobalShaderCookStats
@@ -2598,7 +2598,7 @@ bool FShaderCompilingManager::IsShaderCompilerWorkerRunning(FProcHandle & Worker
 }
 
 /* Generates a uniform buffer struct member hlsl declaration using the member's metadata. */
-static void GenerateUniformBufferStructMember(FString& Result, const FUniformBufferStruct::FMember& Member)
+static void GenerateUniformBufferStructMember(FString& Result, const FShaderParametersMetadata::FMember& Member)
 {
 	// Generate the base type name.
 	FString BaseTypeName;
@@ -2650,8 +2650,8 @@ static void GenerateUniformBufferStructMember(FString& Result, const FUniformBuf
 ENGINE_API void GenerateInstancedStereoCode(FString& Result)
 {
 	// Find the InstancedView uniform buffer struct
-	const FUniformBufferStruct* InstancedView = nullptr;
-	for (TLinkedList<FUniformBufferStruct*>::TIterator StructIt(FUniformBufferStruct::GetStructList()); StructIt; StructIt.Next())
+	const FShaderParametersMetadata* InstancedView = nullptr;
+	for (TLinkedList<FShaderParametersMetadata*>::TIterator StructIt(FShaderParametersMetadata::GetStructList()); StructIt; StructIt.Next())
 	{
 		if (StructIt->GetShaderVariableName() == FString(TEXT("InstancedView")))
 		{
@@ -2660,14 +2660,14 @@ ENGINE_API void GenerateInstancedStereoCode(FString& Result)
 		}
 	}
 	checkSlow(InstancedView != nullptr);
-	const TArray<FUniformBufferStruct::FMember>& StructMembers = InstancedView->GetMembers();
+	const TArray<FShaderParametersMetadata::FMember>& StructMembers = InstancedView->GetMembers();
 
 	// ViewState definition
 	Result =  "struct ViewState\r\n";
 	Result += "{\r\n";
 	for (int32 MemberIndex = 0; MemberIndex < StructMembers.Num(); ++MemberIndex)
 	{
-		const FUniformBufferStruct::FMember& Member = StructMembers[MemberIndex];
+		const FShaderParametersMetadata::FMember& Member = StructMembers[MemberIndex];
 		FString MemberDecl;
 		GenerateUniformBufferStructMember(MemberDecl, StructMembers[MemberIndex]);
 		Result += FString::Printf(TEXT("\t%s;\r\n"), *MemberDecl);
@@ -2680,7 +2680,7 @@ ENGINE_API void GenerateInstancedStereoCode(FString& Result)
 	Result += "\tViewState Result;\r\n";
 	for (int32 MemberIndex = 0; MemberIndex < StructMembers.Num(); ++MemberIndex)
 	{
-		const FUniformBufferStruct::FMember& Member = StructMembers[MemberIndex];
+		const FShaderParametersMetadata::FMember& Member = StructMembers[MemberIndex];
 		Result += FString::Printf(TEXT("\tResult.%s = View.%s;\r\n"), Member.GetName(), Member.GetName());
 	}
 	Result += "\treturn Result;\r\n";
@@ -2692,7 +2692,7 @@ ENGINE_API void GenerateInstancedStereoCode(FString& Result)
 	Result += "\tViewState Result;\r\n";
 	for (int32 MemberIndex = 0; MemberIndex < StructMembers.Num(); ++MemberIndex)
 	{
-		const FUniformBufferStruct::FMember& Member = StructMembers[MemberIndex];
+		const FShaderParametersMetadata::FMember& Member = StructMembers[MemberIndex];
 		Result += FString::Printf(TEXT("\tResult.%s = InstancedView.%s;\r\n"), Member.GetName(), Member.GetName());
 	}
 	Result += "\treturn Result;\r\n";
@@ -2705,7 +2705,7 @@ ENGINE_API void GenerateInstancedStereoCode(FString& Result)
 	Result += "\tViewState Result;\r\n";
 	for (int32 MemberIndex = 0; MemberIndex < StructMembers.Num(); ++MemberIndex)
 	{
-		const FUniformBufferStruct::FMember& Member = StructMembers[MemberIndex];
+		const FShaderParametersMetadata::FMember& Member = StructMembers[MemberIndex];
 		Result += FString::Printf(TEXT("\tResult.%s = (ViewIndex == 0) ? View.%s : InstancedView.%s;\r\n"), Member.GetName(), Member.GetName(), Member.GetName());
 	}
 	Result += "\treturn Result;\r\n";
@@ -3349,12 +3349,6 @@ public:
 		// reload the global shaders
 		CompileGlobalShaderMap(true);
 
-		//invalidate global bound shader states so they will be created with the new shaders the next time they are set (in SetGlobalBoundShaderState)
-		for(TLinkedList<FGlobalBoundShaderStateResource*>::TIterator It(FGlobalBoundShaderStateResource::GetGlobalBoundShaderStateList());It;It.Next())
-		{
-			BeginUpdateResourceRHI(*It);
-		}
-
 		// load all the mesh material shaders if any were sent back
 		if (MeshMaterialMaps.Num() > 0)
 		{
@@ -3408,12 +3402,6 @@ void RecompileGlobalShaders()
 		});
 
 		GShaderCompilingManager->ProcessAsyncResults(false, true);
-
-		//invalidate global bound shader states so they will be created with the new shaders the next time they are set (in SetGlobalBoundShaderState)
-		for (TLinkedList<FGlobalBoundShaderStateResource*>::TIterator It(FGlobalBoundShaderStateResource::GetGlobalBoundShaderStateList()); It; It.Next())
-		{
-			BeginUpdateResourceRHI(*It);
-		}
 	}
 }
 
@@ -3627,13 +3615,24 @@ FShader* FGlobalShaderTypeCompiler::FinishCompileShader(FGlobalShaderType* Shade
 			ShaderPipelineType = nullptr;
 		}
 
+
+		// Create the global shader map hash
+		FSHAHash GlobalShaderMapHash;
+		{
+			FSHA1 HashState;
+			const TCHAR* GlobalShaderString = TEXT("GlobalShaderMap");
+			HashState.UpdateWithString(GlobalShaderString, FCString::Strlen(GlobalShaderString));
+			HashState.Final();
+			HashState.GetHash(&GlobalShaderMapHash.Hash[0]);
+		}
+
 		// Find a shader with the same key in memory
-		Shader = CurrentJob.ShaderType->FindShaderById(FShaderId(GGlobalShaderMapHash, ShaderPipelineType, nullptr, CurrentJob.ShaderType, CurrentJob.PermutationId, CurrentJob.Input.Target));
+		Shader = CurrentJob.ShaderType->FindShaderById(FShaderId(GlobalShaderMapHash, ShaderPipelineType, nullptr, CurrentJob.ShaderType, CurrentJob.PermutationId, CurrentJob.Input.Target));
 
 		// There was no shader with the same key so create a new one with the compile output, which will bind shader parameters
 		if (!Shader)
 		{
-			Shader = (*(ShaderType->ConstructCompiledRef))(FGlobalShaderType::CompiledShaderInitializerType(ShaderType, CurrentJob.PermutationId, CurrentJob.Output, Resource, GGlobalShaderMapHash, ShaderPipelineType, nullptr));
+			Shader = (*(ShaderType->ConstructCompiledRef))(FGlobalShaderType::CompiledShaderInitializerType(ShaderType, CurrentJob.PermutationId, CurrentJob.Output, Resource, GlobalShaderMapHash, ShaderPipelineType, nullptr));
 			CurrentJob.Output.ParameterMap.VerifyBindingsAreComplete(ShaderType->GetName(), CurrentJob.Output.Target, CurrentJob.VFType);
 		}
 	}
@@ -4374,12 +4373,6 @@ void BeginRecompileGlobalShaders(const TArray<FShaderType*>& OutdatedShaderTypes
 				}
 			}
 
-			//invalidate global bound shader states so they will be created with the new shaders the next time they are set (in SetGlobalBoundShaderState)
-			for (TLinkedList<FGlobalBoundShaderStateResource*>::TIterator It(FGlobalBoundShaderStateResource::GetGlobalBoundShaderStateList()); It; It.Next())
-			{
-				BeginUpdateResourceRHI(*It);
-			}
-
 			VerifyGlobalShaders(ShaderPlatform, false);
 		}
 	}
diff --git a/Engine/Source/Runtime/Engine/Private/SpeedTreeWind.cpp b/Engine/Source/Runtime/Engine/Private/SpeedTreeWind.cpp
--- a/Engine/Source/Runtime/Engine/Private/SpeedTreeWind.cpp
+++ b/Engine/Source/Runtime/Engine/Private/SpeedTreeWind.cpp
@@ -22,7 +22,7 @@
 
 #include "SpeedTreeWind.h"
 
-IMPLEMENT_UNIFORM_BUFFER_STRUCT(FSpeedTreeUniformParameters,TEXT("SpeedTreeData"));
+IMPLEMENT_GLOBAL_SHADER_PARAMETER_STRUCT(FSpeedTreeUniformParameters, "SpeedTreeData");
 
 
 ///////////////////////////////////////////////////////////////////////  
diff --git a/Engine/Source/Runtime/Engine/Private/VectorField.cpp b/Engine/Source/Runtime/Engine/Private/VectorField.cpp
--- a/Engine/Source/Runtime/Engine/Private/VectorField.cpp
+++ b/Engine/Source/Runtime/Engine/Private/VectorField.cpp
@@ -641,17 +641,17 @@ void UVectorFieldComponent::PostEditChangeProperty(FPropertyChangedEvent& Proper
 	Shader for constructing animated vector fields.
 ------------------------------------------------------------------------------*/
 
-BEGIN_UNIFORM_BUFFER_STRUCT( FCompositeAnimatedVectorFieldUniformParameters, )
-	UNIFORM_MEMBER( FVector4, FrameA )
-	UNIFORM_MEMBER( FVector4, FrameB )
-	UNIFORM_MEMBER( FVector, VoxelSize )
-	UNIFORM_MEMBER( float, FrameLerp )
-	UNIFORM_MEMBER( float, NoiseScale )
-	UNIFORM_MEMBER( float, NoiseMax )
-	UNIFORM_MEMBER( uint32, Op )
-END_UNIFORM_BUFFER_STRUCT( FCompositeAnimatedVectorFieldUniformParameters )
-
-IMPLEMENT_UNIFORM_BUFFER_STRUCT(FCompositeAnimatedVectorFieldUniformParameters,TEXT("CVF"));
+BEGIN_GLOBAL_SHADER_PARAMETER_STRUCT( FCompositeAnimatedVectorFieldUniformParameters, )
+	SHADER_PARAMETER( FVector4, FrameA )
+	SHADER_PARAMETER( FVector4, FrameB )
+	SHADER_PARAMETER( FVector, VoxelSize )
+	SHADER_PARAMETER( float, FrameLerp )
+	SHADER_PARAMETER( float, NoiseScale )
+	SHADER_PARAMETER( float, NoiseMax )
+	SHADER_PARAMETER( uint32, Op )
+END_GLOBAL_SHADER_PARAMETER_STRUCT()
+
+IMPLEMENT_GLOBAL_SHADER_PARAMETER_STRUCT(FCompositeAnimatedVectorFieldUniformParameters,"CVF");
 
 typedef TUniformBufferRef<FCompositeAnimatedVectorFieldUniformParameters> FCompositeAnimatedVectorFieldUniformBufferRef;
 
diff --git a/Engine/Source/Runtime/Engine/Private/VectorFieldVisualization.cpp b/Engine/Source/Runtime/Engine/Private/VectorFieldVisualization.cpp
--- a/Engine/Source/Runtime/Engine/Private/VectorFieldVisualization.cpp
+++ b/Engine/Source/Runtime/Engine/Private/VectorFieldVisualization.cpp
@@ -14,7 +14,7 @@
 #include "ShaderParameterUtils.h"
 #include "FXSystem.h"
 
-IMPLEMENT_UNIFORM_BUFFER_STRUCT(FVectorFieldVisualizationParameters,TEXT("VectorFieldVis"));
+IMPLEMENT_GLOBAL_SHADER_PARAMETER_STRUCT(FVectorFieldVisualizationParameters,"VectorFieldVis");
 
 /*------------------------------------------------------------------------------
 	Vertex factory for visualizing vector fields.
diff --git a/Engine/Source/Runtime/Engine/Private/VectorFieldVisualization.h b/Engine/Source/Runtime/Engine/Private/VectorFieldVisualization.h
--- a/Engine/Source/Runtime/Engine/Private/VectorFieldVisualization.h
+++ b/Engine/Source/Runtime/Engine/Private/VectorFieldVisualization.h
@@ -36,12 +36,12 @@ class FSceneView;
 /**
  * Uniform buffer to hold parameters for vector field visualization.
  */
-BEGIN_UNIFORM_BUFFER_STRUCT( FVectorFieldVisualizationParameters, )
-	UNIFORM_MEMBER( FMatrix, VolumeToWorld )
-	UNIFORM_MEMBER( FMatrix, VolumeToWorldNoScale )
-	UNIFORM_MEMBER( FVector, VoxelSize )
-	UNIFORM_MEMBER( float, Scale )
-END_UNIFORM_BUFFER_STRUCT( FVectorFieldVisualizationParameters )
+BEGIN_GLOBAL_SHADER_PARAMETER_STRUCT( FVectorFieldVisualizationParameters, )
+	SHADER_PARAMETER( FMatrix, VolumeToWorld )
+	SHADER_PARAMETER( FMatrix, VolumeToWorldNoScale )
+	SHADER_PARAMETER( FVector, VoxelSize )
+	SHADER_PARAMETER( float, Scale )
+END_GLOBAL_SHADER_PARAMETER_STRUCT()
 typedef TUniformBufferRef<FVectorFieldVisualizationParameters> FVectorFieldVisualizationBufferRef;
 
 /**
diff --git a/Engine/Source/Runtime/Engine/Public/EngineSharedPCH.h b/Engine/Source/Runtime/Engine/Public/EngineSharedPCH.h
--- a/Engine/Source/Runtime/Engine/Public/EngineSharedPCH.h
+++ b/Engine/Source/Runtime/Engine/Public/EngineSharedPCH.h
@@ -448,8 +448,6 @@
 #include "UniformBuffer.h"
 #include "PackedNormal.h"
 #include "RenderUtils.h"
-
-// From ShaderCore:
 #include "ShaderParameters.h"
 #include "ShaderCore.h"
 #include "Shader.h"
diff --git a/Engine/Source/Runtime/Engine/Public/FXSystem.h b/Engine/Source/Runtime/Engine/Public/FXSystem.h
--- a/Engine/Source/Runtime/Engine/Public/FXSystem.h
+++ b/Engine/Source/Runtime/Engine/Public/FXSystem.h
@@ -202,7 +202,7 @@ public:
 	virtual void PostRenderOpaque(
 		FRHICommandListImmediate& RHICmdList, 
 		const FUniformBufferRHIParamRef ViewUniformBuffer, 
-		const class FUniformBufferStruct* SceneTexturesUniformBufferStruct,
+		const class FShaderParametersMetadata* SceneTexturesUniformBufferStruct,
 		FUniformBufferRHIParamRef SceneTexturesUniformBuffer) = 0;
 
 	/**
diff --git a/Engine/Source/Runtime/Engine/Public/GPUSkinVertexFactory.h b/Engine/Source/Runtime/Engine/Public/GPUSkinVertexFactory.h
--- a/Engine/Source/Runtime/Engine/Public/GPUSkinVertexFactory.h
+++ b/Engine/Source/Runtime/Engine/Public/GPUSkinVertexFactory.h
@@ -68,30 +68,32 @@ MS_ALIGN(16) struct FSkinMatrix3x4
 } GCC_ALIGN(16);
 
 template<>
-class TUniformBufferTypeInfo<FSkinMatrix3x4, false>
+struct TShaderParameterTypeInfo<FSkinMatrix3x4>
 {
-public:
-	enum { BaseType = UBMT_FLOAT32 };
-	enum { NumRows = 3 };
-	enum { NumColumns = 4 };
-	enum { NumElements = 0 };
-	enum { Alignment = 16 };
-	enum { IsResource = 0 };
-	static const FUniformBufferStruct* GetStruct() { return NULL; }
+	static constexpr EUniformBufferBaseType BaseType = UBMT_FLOAT32;
+	static constexpr int32 NumRows = 3;
+	static constexpr int32 NumColumns = 4;
+	static constexpr int32 NumElements = 0;
+	static constexpr int32 Alignment = 16;
+	static constexpr bool bIsStoredInConstantBuffer = true;
+
+	using TAlignedType = TAlignedTypedef<FSkinMatrix3x4, Alignment>::Type;
+
+	static const FShaderParametersMetadata* GetStructMetadata() { return NULL; }
 };
 
 // Uniform buffer for APEX cloth
-BEGIN_UNIFORM_BUFFER_STRUCT(FAPEXClothUniformShaderParameters,)
-END_UNIFORM_BUFFER_STRUCT(FAPEXClothUniformShaderParameters)
+BEGIN_GLOBAL_SHADER_PARAMETER_STRUCT(FAPEXClothUniformShaderParameters,)
+END_GLOBAL_SHADER_PARAMETER_STRUCT()
 
 enum
 {
 	MAX_GPU_BONE_MATRICES_UNIFORMBUFFER = 75,
 };
 
-BEGIN_UNIFORM_BUFFER_STRUCT(FBoneMatricesUniformShaderParameters,)
-	UNIFORM_MEMBER_ARRAY(FSkinMatrix3x4, BoneMatrices, [MAX_GPU_BONE_MATRICES_UNIFORMBUFFER])
-END_UNIFORM_BUFFER_STRUCT(FBoneMatricesUniformShaderParameters)
+BEGIN_GLOBAL_SHADER_PARAMETER_STRUCT(FBoneMatricesUniformShaderParameters,)
+	SHADER_PARAMETER_ARRAY(FSkinMatrix3x4, BoneMatrices, [MAX_GPU_BONE_MATRICES_UNIFORMBUFFER])
+END_GLOBAL_SHADER_PARAMETER_STRUCT()
 
 #define SET_BONE_DATA(B, X) B.SetMatrixTranspose(X)
 
diff --git a/Engine/Source/Runtime/Engine/Public/LocalVertexFactory.h b/Engine/Source/Runtime/Engine/Public/LocalVertexFactory.h
--- a/Engine/Source/Runtime/Engine/Public/LocalVertexFactory.h
+++ b/Engine/Source/Runtime/Engine/Public/LocalVertexFactory.h
@@ -15,12 +15,12 @@ struct FMeshBatchElement;
 	LocalVertexFactory.h: Local vertex factory definitions.
 =============================================================================*/
 
-BEGIN_UNIFORM_BUFFER_STRUCT(FLocalVertexFactoryUniformShaderParameters,ENGINE_API)
-	UNIFORM_MEMBER(FIntVector,VertexFetch_Parameters)
-	UNIFORM_MEMBER_SRV(Buffer<float2>, VertexFetch_TexCoordBuffer)
-	UNIFORM_MEMBER_SRV(Buffer<float4>, VertexFetch_PackedTangentsBuffer)
-	UNIFORM_MEMBER_SRV(Buffer<float4>, VertexFetch_ColorComponentsBuffer)
-END_UNIFORM_BUFFER_STRUCT(FLocalVertexFactoryUniformShaderParameters)
+BEGIN_GLOBAL_SHADER_PARAMETER_STRUCT(FLocalVertexFactoryUniformShaderParameters,ENGINE_API)
+	SHADER_PARAMETER(FIntVector,VertexFetch_Parameters)
+	SHADER_PARAMETER_SRV(Buffer<float2>, VertexFetch_TexCoordBuffer)
+	SHADER_PARAMETER_SRV(Buffer<float4>, VertexFetch_PackedTangentsBuffer)
+	SHADER_PARAMETER_SRV(Buffer<float4>, VertexFetch_ColorComponentsBuffer)
+END_GLOBAL_SHADER_PARAMETER_STRUCT()
 
 extern TUniformBufferRef<FLocalVertexFactoryUniformShaderParameters> CreateLocalVFUniformBuffer(const class FLocalVertexFactory* VertexFactory, class FColorVertexBuffer* OverrideColorVertexBuffer);
 
diff --git a/Engine/Source/Runtime/Engine/Public/MaterialShared.h b/Engine/Source/Runtime/Engine/Public/MaterialShared.h
--- a/Engine/Source/Runtime/Engine/Public/MaterialShared.h
+++ b/Engine/Source/Runtime/Engine/Public/MaterialShared.h
@@ -1,4 +1,4 @@
-﻿// Copyright 1998-2018 Epic Games, Inc. All Rights Reserved.
+// Copyright 1998-2018 Epic Games, Inc. All Rights Reserved.
 
 /*=============================================================================
 	MaterialShared.h: Shared material definitions.
@@ -356,7 +356,7 @@ public:
 
 	void SetParameterCollections(const TArray<class UMaterialParameterCollection*>& Collections);
 	void CreateBufferStruct();
-	ENGINE_API const FUniformBufferStruct& GetUniformBufferStruct() const;
+	ENGINE_API const FShaderParametersMetadata& GetUniformBufferStruct() const;
 
 	ENGINE_API FUniformBufferRHIRef CreateUniformBuffer(const FMaterialRenderContext& MaterialRenderContext, FRHICommandList* CommandListIfLocalMode, struct FLocalUniformBuffer* OutLocalUniformBuffer) const;
 
@@ -369,7 +369,7 @@ public:
 			+ UniformVolumeTextureExpressions.GetAllocatedSize()
 			+ UniformExternalTextureExpressions.GetAllocatedSize()
 			+ ParameterCollections.GetAllocatedSize()
-			+ (UniformBufferStruct ? (sizeof(FUniformBufferStruct) + UniformBufferStruct->GetMembers().GetAllocatedSize()) : 0);
+			+ (UniformBufferStruct ? (sizeof(FShaderParametersMetadata) + UniformBufferStruct->GetMembers().GetAllocatedSize()) : 0);
 	}
 
 protected:
@@ -385,7 +385,7 @@ protected:
 	TArray<FGuid> ParameterCollections;
 
 	/** The structure of a uniform buffer containing values for these uniform expressions. */
-	TOptional<FUniformBufferStruct> UniformBufferStruct;
+	TOptional<FShaderParametersMetadata> UniformBufferStruct;
 
 	friend class FMaterial;
 	friend class FHLSLMaterialTranslator;
diff --git a/Engine/Source/Runtime/Engine/Public/MeshParticleVertexFactory.h b/Engine/Source/Runtime/Engine/Public/MeshParticleVertexFactory.h
--- a/Engine/Source/Runtime/Engine/Public/MeshParticleVertexFactory.h
+++ b/Engine/Source/Runtime/Engine/Public/MeshParticleVertexFactory.h
@@ -23,13 +23,13 @@ struct FShaderCompilerEnvironment;
 /**
  * Uniform buffer for mesh particle vertex factories.
  */
-BEGIN_UNIFORM_BUFFER_STRUCT( FMeshParticleUniformParameters, ENGINE_API)
-	UNIFORM_MEMBER( FVector4, SubImageSize )
-	UNIFORM_MEMBER( uint32, TexCoordWeightA )
-	UNIFORM_MEMBER( uint32, TexCoordWeightB )
-	UNIFORM_MEMBER( uint32, PrevTransformAvailable )
-	UNIFORM_MEMBER( float, DeltaSeconds)
-	END_UNIFORM_BUFFER_STRUCT(FMeshParticleUniformParameters)
+BEGIN_GLOBAL_SHADER_PARAMETER_STRUCT( FMeshParticleUniformParameters, ENGINE_API)
+	SHADER_PARAMETER( FVector4, SubImageSize )
+	SHADER_PARAMETER( uint32, TexCoordWeightA )
+	SHADER_PARAMETER( uint32, TexCoordWeightB )
+	SHADER_PARAMETER( uint32, PrevTransformAvailable )
+	SHADER_PARAMETER( float, DeltaSeconds)
+END_GLOBAL_SHADER_PARAMETER_STRUCT()
 typedef TUniformBufferRef<FMeshParticleUniformParameters> FMeshParticleUniformBufferRef;
 
 class FMeshParticleInstanceVertices;
diff --git a/Engine/Source/Runtime/Engine/Public/ParticleBeamTrailVertexFactory.h b/Engine/Source/Runtime/Engine/Public/ParticleBeamTrailVertexFactory.h
--- a/Engine/Source/Runtime/Engine/Public/ParticleBeamTrailVertexFactory.h
+++ b/Engine/Source/Runtime/Engine/Public/ParticleBeamTrailVertexFactory.h
@@ -17,11 +17,11 @@ class FMaterial;
 /**
  * Uniform buffer for particle beam/trail vertex factories.
  */
-BEGIN_UNIFORM_BUFFER_STRUCT( FParticleBeamTrailUniformParameters, ENGINE_API)
-	UNIFORM_MEMBER( FVector4, CameraRight )
-	UNIFORM_MEMBER( FVector4, CameraUp )
-	UNIFORM_MEMBER( FVector4, ScreenAlignment )
-END_UNIFORM_BUFFER_STRUCT( FParticleBeamTrailUniformParameters )
+BEGIN_GLOBAL_SHADER_PARAMETER_STRUCT( FParticleBeamTrailUniformParameters, ENGINE_API)
+	SHADER_PARAMETER( FVector4, CameraRight )
+	SHADER_PARAMETER( FVector4, CameraUp )
+	SHADER_PARAMETER( FVector4, ScreenAlignment )
+END_GLOBAL_SHADER_PARAMETER_STRUCT()
 typedef TUniformBufferRef<FParticleBeamTrailUniformParameters> FParticleBeamTrailUniformBufferRef;
 
 /**
