GPU Particle system updated to use renderpasses.

diff --git a/Engine/Source/Runtime/Engine/Private/Particles/ParticleCurveTexture.cpp b/Engine/Source/Runtime/Engine/Private/Particles/ParticleCurveTexture.cpp
--- a/Engine/Source/Runtime/Engine/Private/Particles/ParticleCurveTexture.cpp
+++ b/Engine/Source/Runtime/Engine/Private/Particles/ParticleCurveTexture.cpp
@@ -188,106 +188,105 @@ static void InjectCurves(
 
 	RHICmdList.BeginUpdateMultiFrameResource(CurveTextureTargetRHI);
 
+	ERenderTargetLoadAction LoadAction = ERenderTargetLoadAction::ELoad;
+
 	if (bFirstCall)
 	{
-		TransitionSetRenderTargetsHelper(RHICmdList, CurveTextureTargetRHI, FTextureRHIParamRef(), FExclusiveDepthStencil::DepthWrite_StencilWrite);
-
-		FRHIRenderTargetView View = FRHIRenderTargetView(CurveTextureTargetRHI, ERenderTargetLoadAction::EClear);
-		FRHISetRenderTargetsInfo Info(1, &View, FRHIDepthRenderTargetView());
-		RHICmdList.SetRenderTargetsAndClear(Info);
+		LoadAction = ERenderTargetLoadAction::EClear;
 		bFirstCall = false;
 	}
-	else
-	{
-		SetRenderTarget(RHICmdList, CurveTextureTargetRHI, FTextureRHIParamRef(), true);
-	}
 
-	FGraphicsPipelineStateInitializer GraphicsPSOInit;
-	RHICmdList.ApplyCachedRenderTargets(GraphicsPSOInit);
+	FRHIRenderPassInfo RPInfo(CurveTextureTargetRHI, MakeRenderTargetActions(LoadAction, ERenderTargetStoreAction::EStore), CurveTextureRHI);
+	TransitionRenderPassTargets(RHICmdList, RPInfo);
+	RHICmdList.BeginRenderPass(RPInfo, TEXT("InjectCurves"));
+	{
+		FGraphicsPipelineStateInitializer GraphicsPSOInit;
+		RHICmdList.ApplyCachedRenderTargets(GraphicsPSOInit);
 
-	RHICmdList.SetScissorRect(false, 0, 0, 0, 0);
-	RHICmdList.SetViewport(0, 0, 0.0f, GParticleCurveTextureSizeX, GParticleCurveTextureSizeY, 1.0f);
-	GraphicsPSOInit.DepthStencilState = TStaticDepthStencilState<false, CF_Always>::GetRHI();
-	GraphicsPSOInit.RasterizerState = TStaticRasterizerState<FM_Solid, CM_None>::GetRHI();
-	GraphicsPSOInit.BlendState = TStaticBlendState<>::GetRHI();
+		RHICmdList.SetScissorRect(false, 0, 0, 0, 0);
+		RHICmdList.SetViewport(0, 0, 0.0f, GParticleCurveTextureSizeX, GParticleCurveTextureSizeY, 1.0f);
+		GraphicsPSOInit.DepthStencilState = TStaticDepthStencilState<false, CF_Always>::GetRHI();
+		GraphicsPSOInit.RasterizerState = TStaticRasterizerState<FM_Solid, CM_None>::GetRHI();
+		GraphicsPSOInit.BlendState = TStaticBlendState<>::GetRHI();
 
-	// Grab shaders.
-	TShaderMapRef<FParticleCurveInjectionVS> VertexShader(GetGlobalShaderMap(GMaxRHIFeatureLevel));
-	TShaderMapRef<FParticleCurveInjectionPS> PixelShader(GetGlobalShaderMap(GMaxRHIFeatureLevel));
+		// Grab shaders.
+		TShaderMapRef<FParticleCurveInjectionVS> VertexShader(GetGlobalShaderMap(GMaxRHIFeatureLevel));
+		TShaderMapRef<FParticleCurveInjectionPS> PixelShader(GetGlobalShaderMap(GMaxRHIFeatureLevel));
 
-	GraphicsPSOInit.BoundShaderState.VertexDeclarationRHI = GParticleCurveInjectionVertexDeclaration.VertexDeclarationRHI;
-	GraphicsPSOInit.BoundShaderState.VertexShaderRHI = GETSAFERHISHADER_VERTEX(*VertexShader);
-	GraphicsPSOInit.BoundShaderState.PixelShaderRHI = GETSAFERHISHADER_PIXEL(*PixelShader);
-	GraphicsPSOInit.PrimitiveType = PT_TriangleList;
+		GraphicsPSOInit.BoundShaderState.VertexDeclarationRHI = GParticleCurveInjectionVertexDeclaration.VertexDeclarationRHI;
+		GraphicsPSOInit.BoundShaderState.VertexShaderRHI = GETSAFERHISHADER_VERTEX(*VertexShader);
+		GraphicsPSOInit.BoundShaderState.PixelShaderRHI = GETSAFERHISHADER_PIXEL(*PixelShader);
+		GraphicsPSOInit.PrimitiveType = PT_TriangleList;
 
-	SetGraphicsPipelineState(RHICmdList, GraphicsPSOInit);
+		SetGraphicsPipelineState(RHICmdList, GraphicsPSOInit);
 
-	int32 PendingCurveCount = InPendingCurves.Num();
+		int32 PendingCurveCount = InPendingCurves.Num();
 
-	// get sum of size of all curve textures
-	//
-	int32 TotalSamples = 0;
-	for (int32 CurveIndex = 0; CurveIndex < PendingCurveCount; ++CurveIndex)
-	{
-		FCurveSamples& Curve = InPendingCurves[CurveIndex];
-		check(Curve.Samples);
-		TotalSamples += Curve.TexelAllocation.Size;
-	}
+		// get sum of size of all curve textures
+		//
+		int32 TotalSamples = 0;
+		for (int32 CurveIndex = 0; CurveIndex < PendingCurveCount; ++CurveIndex)
+		{
+			FCurveSamples& Curve = InPendingCurves[CurveIndex];
+			check(Curve.Samples);
+			TotalSamples += Curve.TexelAllocation.Size;
+		}
 
 
-	// get a buffer for all curve textures at once, and copy curve data over
-	//
-	FColor* RESTRICT DestSamples = (FColor*)RHILockVertexBuffer(ScratchVertexBufferRHI, 0, TotalSamples * sizeof(FColor), RLM_WriteOnly);
-	int32 CurrOffset = 0;
+		// get a buffer for all curve textures at once, and copy curve data over
+		//
+		FColor* RESTRICT DestSamples = (FColor*)RHILockVertexBuffer(ScratchVertexBufferRHI, 0, TotalSamples * sizeof(FColor), RLM_WriteOnly);
+		int32 CurrOffset = 0;
 
-	for (int32 CurveIndex = 0; CurveIndex < PendingCurveCount; ++CurveIndex)
-	{
-		FCurveSamples& Curve = InPendingCurves[CurveIndex];
+		for (int32 CurveIndex = 0; CurveIndex < PendingCurveCount; ++CurveIndex)
+		{
+			FCurveSamples& Curve = InPendingCurves[CurveIndex];
 
-		// Copy curve samples in to the scratch vertex buffer.
-		const int32 SampleCount = Curve.TexelAllocation.Size;
-		FMemory::Memcpy(DestSamples + CurrOffset, Curve.Samples, SampleCount * sizeof(FColor));
-		FMemory::Free(Curve.Samples);
-		Curve.Samples = NULL;
+			// Copy curve samples in to the scratch vertex buffer.
+			const int32 SampleCount = Curve.TexelAllocation.Size;
+			FMemory::Memcpy(DestSamples + CurrOffset, Curve.Samples, SampleCount * sizeof(FColor));
+			FMemory::Free(Curve.Samples);
+			Curve.Samples = NULL;
 
-		CurrOffset += SampleCount;
-	}
+			CurrOffset += SampleCount;
+		}
 
-	RHICmdList.UnlockVertexBuffer(ScratchVertexBufferRHI);
+		RHICmdList.UnlockVertexBuffer(ScratchVertexBufferRHI);
 
 
-	// now draw the curves into the curve texture target, reading from the single buffer we just filled
-	//
-	int32 CurrSrcOffset = 0;
-	for (int32 CurveIndex = 0; CurveIndex < PendingCurveCount; ++CurveIndex)
-	{
-		FCurveSamples& Curve = InPendingCurves[CurveIndex];
-		const int32 SampleCount = Curve.TexelAllocation.Size;
-
-		// Compute the offset in to the texture.
-		FVector2D CurveOffset((float)Curve.TexelAllocation.X / GParticleCurveTextureSizeX,
-			(float)Curve.TexelAllocation.Y / GParticleCurveTextureSizeY);
-
-		// Stream 0: New particles.
-		RHICmdList.SetStreamSource(0, ScratchVertexBufferRHI, CurrSrcOffset*sizeof(FColor));
-		// Stream 0: TexCoord.
-		RHICmdList.SetStreamSource(1, GParticleTexCoordVertexBuffer.VertexBufferRHI, 0);
-
-		VertexShader->SetParameters(RHICmdList, CurveOffset);
-
-		// Inject particles.
-		RHICmdList.DrawIndexedPrimitive(
-			GParticleIndexBuffer.IndexBufferRHI,
-			/*BaseVertexIndex=*/ 0,
-			/*MinIndex=*/ 0,
-			/*NumVertices=*/ 4,
-			/*StartIndex=*/ 0,
-			/*NumPrimitives=*/ 2,
-			/*NumInstances=*/ SampleCount
-		);
-		CurrSrcOffset += SampleCount;
+		// now draw the curves into the curve texture target, reading from the single buffer we just filled
+		//
+		int32 CurrSrcOffset = 0;
+		for (int32 CurveIndex = 0; CurveIndex < PendingCurveCount; ++CurveIndex)
+		{
+			FCurveSamples& Curve = InPendingCurves[CurveIndex];
+			const int32 SampleCount = Curve.TexelAllocation.Size;
+
+			// Compute the offset in to the texture.
+			FVector2D CurveOffset((float)Curve.TexelAllocation.X / GParticleCurveTextureSizeX,
+				(float)Curve.TexelAllocation.Y / GParticleCurveTextureSizeY);
+
+			// Stream 0: New particles.
+			RHICmdList.SetStreamSource(0, ScratchVertexBufferRHI, CurrSrcOffset * sizeof(FColor));
+			// Stream 0: TexCoord.
+			RHICmdList.SetStreamSource(1, GParticleTexCoordVertexBuffer.VertexBufferRHI, 0);
+
+			VertexShader->SetParameters(RHICmdList, CurveOffset);
+
+			// Inject particles.
+			RHICmdList.DrawIndexedPrimitive(
+				GParticleIndexBuffer.IndexBufferRHI,
+				/*BaseVertexIndex=*/ 0,
+				/*MinIndex=*/ 0,
+				/*NumVertices=*/ 4,
+				/*StartIndex=*/ 0,
+				/*NumPrimitives=*/ 2,
+				/*NumInstances=*/ SampleCount
+			);
+			CurrSrcOffset += SampleCount;
+		}
 	}
-	RHICmdList.CopyToResolveTarget(CurveTextureTargetRHI, CurveTextureRHI, FResolveParams());
+	RHICmdList.EndRenderPass();
 	RHICmdList.EndUpdateMultiFrameResource(CurveTextureTargetRHI);
 }
 
diff --git a/Engine/Source/Runtime/Engine/Private/Particles/ParticleGpuSimulation.cpp b/Engine/Source/Runtime/Engine/Private/Particles/ParticleGpuSimulation.cpp
--- a/Engine/Source/Runtime/Engine/Private/Particles/ParticleGpuSimulation.cpp
+++ b/Engine/Source/Runtime/Engine/Private/Particles/ParticleGpuSimulation.cpp
@@ -2087,54 +2087,59 @@ static void VisualizeGPUSimulation(
 
 	// Setup render states.
 	FIntPoint TargetSize = RenderTarget->GetSizeXY();
-	SetRenderTarget(RHICmdList, RenderTarget->GetRenderTargetTexture(), FTextureRHIParamRef());
-	RHICmdList.ApplyCachedRenderTargets(GraphicsPSOInit);
 
-	RHICmdList.SetViewport(0, 0, 0.0f, TargetSize.X, TargetSize.Y, 1.0f);
-	GraphicsPSOInit.DepthStencilState = TStaticDepthStencilState<false, CF_Always>::GetRHI();
-	GraphicsPSOInit.RasterizerState = TStaticRasterizerState<FM_Solid, CM_None>::GetRHI();
-	GraphicsPSOInit.BlendState = TStaticBlendState<>::GetRHI();
+	FRHIRenderPassInfo RPInfo(RenderTarget->GetRenderTargetTexture(), ERenderTargetActions::Load_Store);
+	RHICmdList.BeginRenderPass(RPInfo, TEXT("VisualizeGPUSimulation"));
+	{
+		RHICmdList.ApplyCachedRenderTargets(GraphicsPSOInit);
 
-	// Grab shaders.
-	TShaderMapRef<FParticleSimVisualizeVS> VertexShader(GetGlobalShaderMap(FeatureLevel));
-	TShaderMapRef<FParticleSimVisualizePS> PixelShader(GetGlobalShaderMap(FeatureLevel));
+		RHICmdList.SetViewport(0, 0, 0.0f, TargetSize.X, TargetSize.Y, 1.0f);
+		GraphicsPSOInit.DepthStencilState = TStaticDepthStencilState<false, CF_Always>::GetRHI();
+		GraphicsPSOInit.RasterizerState = TStaticRasterizerState<FM_Solid, CM_None>::GetRHI();
+		GraphicsPSOInit.BlendState = TStaticBlendState<>::GetRHI();
 
-	GraphicsPSOInit.BoundShaderState.VertexDeclarationRHI = GParticleSimVisualizeVertexDeclaration.VertexDeclarationRHI;
-	GraphicsPSOInit.BoundShaderState.VertexShaderRHI = GETSAFERHISHADER_VERTEX(*VertexShader);
-	GraphicsPSOInit.BoundShaderState.PixelShaderRHI = GETSAFERHISHADER_PIXEL(*PixelShader);
-	GraphicsPSOInit.PrimitiveType = PT_TriangleList;
+		// Grab shaders.
+		TShaderMapRef<FParticleSimVisualizeVS> VertexShader(GetGlobalShaderMap(FeatureLevel));
+		TShaderMapRef<FParticleSimVisualizePS> PixelShader(GetGlobalShaderMap(FeatureLevel));
 
-	SetGraphicsPipelineState(RHICmdList, GraphicsPSOInit);
+		GraphicsPSOInit.BoundShaderState.VertexDeclarationRHI = GParticleSimVisualizeVertexDeclaration.VertexDeclarationRHI;
+		GraphicsPSOInit.BoundShaderState.VertexShaderRHI = GETSAFERHISHADER_VERTEX(*VertexShader);
+		GraphicsPSOInit.BoundShaderState.PixelShaderRHI = GETSAFERHISHADER_PIXEL(*PixelShader);
+		GraphicsPSOInit.PrimitiveType = PT_TriangleList;
 
-	// Parameters for the visualization.
-	FParticleSimVisualizeParameters Parameters;
-	Parameters.ScaleBias.X = 2.0f * DisplaySizeX / (float)TargetSize.X;
-	Parameters.ScaleBias.Y = 2.0f * DisplaySizeY / (float)TargetSize.Y;
-	Parameters.ScaleBias.Z = 2.0f * DisplayOffsetX / (float)TargetSize.X - 1.0f;
-	Parameters.ScaleBias.W = 2.0f * DisplayOffsetY / (float)TargetSize.Y - 1.0f;
-	FParticleSimVisualizeBufferRef UniformBuffer = FParticleSimVisualizeBufferRef::CreateUniformBufferImmediate( Parameters, UniformBuffer_SingleDraw );
-	VertexShader->SetParameters(RHICmdList, UniformBuffer);
-	PixelShader->SetParameters(RHICmdList, VisualizationMode, StateTextures.PositionTextureRHI, CurveTextureRHI);
-
-	const int32 VertexStride = sizeof(FVector2D);
-	
-	// Bind vertex stream.
-	RHICmdList.SetStreamSource(
-		0,
-		GParticleTexCoordVertexBuffer.VertexBufferRHI,
-		/*VertexOffset=*/ 0
+		SetGraphicsPipelineState(RHICmdList, GraphicsPSOInit);
+
+		// Parameters for the visualization.
+		FParticleSimVisualizeParameters Parameters;
+		Parameters.ScaleBias.X = 2.0f * DisplaySizeX / (float)TargetSize.X;
+		Parameters.ScaleBias.Y = 2.0f * DisplaySizeY / (float)TargetSize.Y;
+		Parameters.ScaleBias.Z = 2.0f * DisplayOffsetX / (float)TargetSize.X - 1.0f;
+		Parameters.ScaleBias.W = 2.0f * DisplayOffsetY / (float)TargetSize.Y - 1.0f;
+		FParticleSimVisualizeBufferRef UniformBuffer = FParticleSimVisualizeBufferRef::CreateUniformBufferImmediate(Parameters, UniformBuffer_SingleDraw);
+		VertexShader->SetParameters(RHICmdList, UniformBuffer);
+		PixelShader->SetParameters(RHICmdList, VisualizationMode, StateTextures.PositionTextureRHI, CurveTextureRHI);
+
+		const int32 VertexStride = sizeof(FVector2D);
+
+		// Bind vertex stream.
+		RHICmdList.SetStreamSource(
+			0,
+			GParticleTexCoordVertexBuffer.VertexBufferRHI,
+			/*VertexOffset=*/ 0
 		);
 
-	// Draw.
-	RHICmdList.DrawIndexedPrimitive(
-		GParticleIndexBuffer.IndexBufferRHI,
-		/*BaseVertexIndex=*/ 0,
-		/*MinIndex=*/ 0,
-		/*NumVertices=*/ 4,
-		/*StartIndex=*/ 0,
-		/*NumPrimitives=*/ 2,
-		/*NumInstances=*/ 1
+		// Draw.
+		RHICmdList.DrawIndexedPrimitive(
+			GParticleIndexBuffer.IndexBufferRHI,
+			/*BaseVertexIndex=*/ 0,
+			/*MinIndex=*/ 0,
+			/*NumVertices=*/ 4,
+			/*StartIndex=*/ 0,
+			/*NumPrimitives=*/ 2,
+			/*NumInstances=*/ 1
 		);
+	}
+	RHICmdList.EndRenderPass();
 }
 
 /**
@@ -4712,9 +4717,12 @@ void FFXSystem::SimulateGPUParticles(
 			
 			RHICmdList.BeginUpdateMultiFrameResource(CurrentStateRenderTargets[0]);
 			RHICmdList.BeginUpdateMultiFrameResource(CurrentStateRenderTargets[1]);
-			
-			SetRenderTarget(RHICmdList, CurrentStateTextures.PositionTextureTargetRHI, FTextureRHIRef(), ESimpleRenderTargetMode::EClearColorAndDepth);
-			SetRenderTarget(RHICmdList, CurrentStateTextures.VelocityTextureTargetRHI, FTextureRHIRef(), ESimpleRenderTargetMode::EClearColorAndDepth);
+
+			{
+				FRHIRenderPassInfo RenderPassInfo(2, CurrentStateRenderTargets, ERenderTargetActions::Clear_Store);
+				RHICmdList.BeginRenderPass(RenderPassInfo, TEXT("GPUParticlesClearStateTextures"));
+				RHICmdList.EndRenderPass();
+			}
 			
 			CurrentStateTextures.bTexturesCleared = true;
 			
@@ -4726,11 +4734,15 @@ void FFXSystem::SimulateGPUParticles(
 		{
 			RHICmdList.BeginUpdateMultiFrameResource(PreviousStateRenderTargets[0]);
 			RHICmdList.BeginUpdateMultiFrameResource(PreviousStateRenderTargets[1]);
-			
-			SetRenderTarget(RHICmdList, PrevStateTextures.PositionTextureTargetRHI, FTextureRHIRef(), ESimpleRenderTargetMode::EClearColorAndDepth);
-			RHICmdList.CopyToResolveTarget(PrevStateTextures.PositionTextureTargetRHI, PrevStateTextures.PositionTextureTargetRHI, FResolveParams());
-			SetRenderTarget(RHICmdList, PrevStateTextures.VelocityTextureTargetRHI, FTextureRHIRef(), ESimpleRenderTargetMode::EClearColorAndDepth);
-			RHICmdList.CopyToResolveTarget(PrevStateTextures.VelocityTextureTargetRHI, PrevStateTextures.VelocityTextureTargetRHI, FResolveParams());
+
+			{
+				FRHIRenderPassInfo RPInfo(2, PreviousStateRenderTargets, ERenderTargetActions::Clear_Store);
+				RPInfo.ColorRenderTargets[0].ResolveTarget = PrevStateTextures.PositionTextureTargetRHI;
+				RPInfo.ColorRenderTargets[1].ResolveTarget = PrevStateTextures.VelocityTextureTargetRHI;
+
+				RHICmdList.BeginRenderPass(RPInfo, TEXT("GPUParticlesClearPreviousStateTextures"));
+				RHICmdList.EndRenderPass();
+			}
 			
 			PrevStateTextures.bTexturesCleared = true;
 			
@@ -4884,44 +4896,48 @@ void FFXSystem::SimulateGPUParticles(
 	
 	if ( SimulationCommands.Num() || TilesToClear.Num())
 	{
-		SCOPED_DRAW_EVENT(RHICmdList, GPUParticles_SimulateAndClear);
-
-		SetRenderTargets(RHICmdList, 2, CurrentStateRenderTargets, FTextureRHIParamRef(), 0, NULL);
-		FGraphicsPipelineStateInitializer GraphicsPSOInit;
-		RHICmdList.ApplyCachedRenderTargets(GraphicsPSOInit);
-		
-		RHICmdList.SetViewport(0, 0, 0.0f, GParticleSimulationTextureSizeX, GParticleSimulationTextureSizeY, 1.0f);
-		GraphicsPSOInit.DepthStencilState = TStaticDepthStencilState<false, CF_Always>::GetRHI();
-		GraphicsPSOInit.RasterizerState = TStaticRasterizerState<FM_Solid, CM_None>::GetRHI();
-		GraphicsPSOInit.BlendState = TStaticBlendState<>::GetRHI();
-		
-		// Simulate particles in all active tiles.
-		if ( SimulationCommands.Num() )
+		FRHIRenderPassInfo RPInfo(2, CurrentStateRenderTargets, ERenderTargetActions::Load_Store);
+		RHICmdList.BeginRenderPass(RPInfo, TEXT("GPUParticles_SimulateAndClear"));
 		{
-			SCOPED_DRAW_EVENT(RHICmdList, ParticleSimulationCommands);
-			
-			ExecuteSimulationCommands(
-						RHICmdList,
-						GraphicsPSOInit,
-						FeatureLevel,
-						SimulationCommands,
-						ParticleSimulationResources,
-						ViewUniformBuffer,
-						GlobalDistanceFieldParameterData,
-						SceneTexturesUniformBufferStruct,
-						SceneTexturesUniformBuffer,
-						Phase,
-						FixDeltaSeconds > 0
-						);
-		}
+			SCOPED_DRAW_EVENT(RHICmdList, GPUParticles_SimulateAndClear);
 
-		// Clear any newly allocated tiles.
-		if (TilesToClear.Num())
-		{
-			SCOPED_DRAW_EVENT(RHICmdList, ParticleTilesClear);
-			
-			ClearTiles(RHICmdList, GraphicsPSOInit, FeatureLevel, TilesToClear);
+			FGraphicsPipelineStateInitializer GraphicsPSOInit;
+			RHICmdList.ApplyCachedRenderTargets(GraphicsPSOInit);
+
+			RHICmdList.SetViewport(0, 0, 0.0f, GParticleSimulationTextureSizeX, GParticleSimulationTextureSizeY, 1.0f);
+			GraphicsPSOInit.DepthStencilState = TStaticDepthStencilState<false, CF_Always>::GetRHI();
+			GraphicsPSOInit.RasterizerState = TStaticRasterizerState<FM_Solid, CM_None>::GetRHI();
+			GraphicsPSOInit.BlendState = TStaticBlendState<>::GetRHI();
+
+			// Simulate particles in all active tiles.
+			if (SimulationCommands.Num())
+			{
+				SCOPED_DRAW_EVENT(RHICmdList, ParticleSimulationCommands);
+
+				ExecuteSimulationCommands(
+					RHICmdList,
+					GraphicsPSOInit,
+					FeatureLevel,
+					SimulationCommands,
+					ParticleSimulationResources,
+					ViewUniformBuffer,
+					GlobalDistanceFieldParameterData,
+					SceneTexturesUniformBufferStruct,
+					SceneTexturesUniformBuffer,
+					Phase,
+					FixDeltaSeconds > 0
+				);
+			}
+
+			// Clear any newly allocated tiles.
+			if (TilesToClear.Num())
+			{
+				SCOPED_DRAW_EVENT(RHICmdList, ParticleTilesClear);
+
+				ClearTiles(RHICmdList, GraphicsPSOInit, FeatureLevel, TilesToClear);
+			}
 		}
+		RHICmdList.EndRenderPass();
 	}
 
 	// Inject any new particles that have spawned into the simulation.
@@ -4942,29 +4958,26 @@ void FFXSystem::SimulateGPUParticles(
 		RHICmdList.BeginUpdateMultiFrameResource(ParticleSimulationResources->RenderAttributesTexture.TextureTargetRHI);
 		RHICmdList.BeginUpdateMultiFrameResource(ParticleSimulationResources->SimulationAttributesTexture.TextureTargetRHI);
 
-		SetRenderTargets(RHICmdList, 4, InjectRenderTargets, FTextureRHIParamRef(), 0, NULL, true);
-		FGraphicsPipelineStateInitializer GraphicsPSOInit;
-		RHICmdList.ApplyCachedRenderTargets(GraphicsPSOInit);
-		
-		RHICmdList.SetViewport(0, 0, 0.0f, GParticleSimulationTextureSizeX, GParticleSimulationTextureSizeY, 1.0f);
-		GraphicsPSOInit.DepthStencilState = TStaticDepthStencilState<false, CF_Always>::GetRHI();
-		GraphicsPSOInit.RasterizerState = TStaticRasterizerState<FM_Solid, CM_None>::GetRHI();
-		GraphicsPSOInit.BlendState = TStaticBlendState<>::GetRHI();
+		FRHIRenderPassInfo RPInfo(4, InjectRenderTargets, ERenderTargetActions::Load_Store);
+		RPInfo.ColorRenderTargets[2].ResolveTarget = ParticleSimulationResources->RenderAttributesTexture.TextureRHI;
+		RPInfo.ColorRenderTargets[3].ResolveTarget = ParticleSimulationResources->SimulationAttributesTexture.TextureRHI;
+		{
+			TransitionRenderPassTargets(RHICmdList, RPInfo);
+			RHICmdList.BeginRenderPass(RPInfo, TEXT("ParticleInjection"));
 
-		// Inject particles.
-		InjectNewParticles<false>(RHICmdList, GraphicsPSOInit, FeatureLevel, NewParticles);
+			FGraphicsPipelineStateInitializer GraphicsPSOInit;
+			RHICmdList.ApplyCachedRenderTargets(GraphicsPSOInit);
 
-		// Resolve attributes textures. State textures are resolved later.
-		RHICmdList.CopyToResolveTarget(
-			ParticleSimulationResources->RenderAttributesTexture.TextureTargetRHI,
-			ParticleSimulationResources->RenderAttributesTexture.TextureRHI,
-			FResolveParams()
-			);
-		RHICmdList.CopyToResolveTarget(
-			ParticleSimulationResources->SimulationAttributesTexture.TextureTargetRHI,
-			ParticleSimulationResources->SimulationAttributesTexture.TextureRHI,
-			FResolveParams()
-			);
+			RHICmdList.SetViewport(0, 0, 0.0f, GParticleSimulationTextureSizeX, GParticleSimulationTextureSizeY, 1.0f);
+			GraphicsPSOInit.DepthStencilState = TStaticDepthStencilState<false, CF_Always>::GetRHI();
+			GraphicsPSOInit.RasterizerState = TStaticRasterizerState<FM_Solid, CM_None>::GetRHI();
+			GraphicsPSOInit.BlendState = TStaticBlendState<>::GetRHI();
+
+			// Inject particles.
+			InjectNewParticles<false>(RHICmdList, GraphicsPSOInit, this->FeatureLevel, NewParticles);
+
+			RHICmdList.EndRenderPass();
+		}
 
 		if (GNumAlternateFrameRenderingGroups > 1 && CVarGPUParticleAFRReinject.GetValueOnRenderThread() == 1)
 		{			
@@ -4988,29 +5001,34 @@ void FFXSystem::SimulateGPUParticles(
 		//will make this second step simulate an interpolated extra 7ms.  This second interpolated step is what we render on THIS frame, but it is NOT fed into the next frame's simulation.  Thus we do not need to transfer it between GPUs in AFR mode.
 		FParticleStateTextures& VisualizeStateTextures = ParticleSimulationResources->GetPreviousStateTextures();
 		FTextureRHIParamRef VisualizeStateRHIs[2] = { VisualizeStateTextures.PositionTextureTargetRHI, VisualizeStateTextures.VelocityTextureTargetRHI };
-		RHICmdList.TransitionResources(EResourceTransitionAccess::EWritable, VisualizeStateRHIs, 2);
-		
-		SetRenderTargets(RHICmdList, 2, VisualizeStateRHIs, FTextureRHIParamRef(), 0, NULL);
-		FGraphicsPipelineStateInitializer GraphicsPSOInit;
-		RHICmdList.ApplyCachedRenderTargets(GraphicsPSOInit);
-		
-		GraphicsPSOInit.DepthStencilState = TStaticDepthStencilState<false, CF_Always>::GetRHI();
-		GraphicsPSOInit.RasterizerState = TStaticRasterizerState<FM_Solid, CM_None>::GetRHI();
-		GraphicsPSOInit.BlendState = TStaticBlendState<>::GetRHI();
-		
-		ExecuteSimulationCommands(
-					RHICmdList,
-					GraphicsPSOInit,
-					FeatureLevel,
-					SimulationCommands,
-					ParticleSimulationResources,
-					ViewUniformBuffer,
-					GlobalDistanceFieldParameterData,
-					SceneTexturesUniformBufferStruct,
-					SceneTexturesUniformBuffer,
-					Phase,
-					false
-					);
+
+		FRHIRenderPassInfo RPInfo(2, VisualizeStateRHIs, ERenderTargetActions::Load_Store);
+		{
+			TransitionRenderPassTargets(RHICmdList, RPInfo);
+			RHICmdList.BeginRenderPass(RPInfo, TEXT("ExecuteSimulationCommands"));
+			FGraphicsPipelineStateInitializer GraphicsPSOInit;
+			RHICmdList.ApplyCachedRenderTargets(GraphicsPSOInit);
+
+			GraphicsPSOInit.DepthStencilState = TStaticDepthStencilState<false, CF_Always>::GetRHI();
+			GraphicsPSOInit.RasterizerState = TStaticRasterizerState<FM_Solid, CM_None>::GetRHI();
+			GraphicsPSOInit.BlendState = TStaticBlendState<>::GetRHI();
+
+			ExecuteSimulationCommands(
+				RHICmdList,
+				GraphicsPSOInit,
+				FeatureLevel,
+				SimulationCommands,
+				ParticleSimulationResources,
+				ViewUniformBuffer,
+				GlobalDistanceFieldParameterData,
+				SceneTexturesUniformBufferStruct,
+				SceneTexturesUniformBuffer,
+				Phase,
+				false
+			);
+
+			RHICmdList.EndRenderPass();
+		}
 		RHICmdList.TransitionResources(EResourceTransitionAccess::EReadable, VisualizeStateRHIs, 2);		
 	}
 
@@ -5018,9 +5036,6 @@ void FFXSystem::SimulateGPUParticles(
 	TilesToClear.Reset();
 	NewParticles.Reset();
 
-	// Clear render targets so we can safely read from them.
-	SetRenderTarget(RHICmdList, FTextureRHIParamRef(), FTextureRHIParamRef());
-
 	// Stats.
 	if (Phase == GetLastParticleSimulationPhase(GetShaderPlatform()))
 	{
@@ -5042,29 +5057,26 @@ void FFXSystem::UpdateMultiGPUResources(FRHICommandListImmediate& RHICmdList)
 			ParticleSimulationResources->RenderAttributesTexture.TextureTargetRHI,
 			ParticleSimulationResources->SimulationAttributesTexture.TextureTargetRHI
 		};
-		SetRenderTargets(RHICmdList, 2, InjectRenderTargets, FTextureRHIParamRef(), 0, NULL, true);
-		RHICmdList.SetViewport(0, 0, 0.0f, GParticleSimulationTextureSizeX, GParticleSimulationTextureSizeY, 1.0f);
-		FGraphicsPipelineStateInitializer GraphicsPSOInit;
-		RHICmdList.ApplyCachedRenderTargets(GraphicsPSOInit);
-		GraphicsPSOInit.DepthStencilState = TStaticDepthStencilState<false, CF_Always>::GetRHI();
-		GraphicsPSOInit.RasterizerState = TStaticRasterizerState<FM_Solid, CM_None>::GetRHI();
-		GraphicsPSOInit.BlendState = TStaticBlendState<>::GetRHI();
 
-		// Inject particles.
-		InjectNewParticles<true>(RHICmdList, GraphicsPSOInit, FeatureLevel, LastFrameNewParticles);
+		FRHIRenderPassInfo RPInfo(2, InjectRenderTargets, ERenderTargetActions::Load_Store);
+		RPInfo.ColorRenderTargets[0].ResolveTarget = ParticleSimulationResources->RenderAttributesTexture.TextureRHI;
+		RPInfo.ColorRenderTargets[1].ResolveTarget = ParticleSimulationResources->SimulationAttributesTexture.TextureRHI;
+		{
+			TransitionRenderPassTargets(RHICmdList, RPInfo);
+			RHICmdList.BeginRenderPass(RPInfo, TEXT("UpdateMultiGPUResources"));
 
-		// Resolve attributes textures. State textures are resolved later.
-		RHICmdList.CopyToResolveTarget(
-			ParticleSimulationResources->RenderAttributesTexture.TextureTargetRHI,
-			ParticleSimulationResources->RenderAttributesTexture.TextureRHI,
-			FResolveParams()
-		);
-		RHICmdList.CopyToResolveTarget(
-			ParticleSimulationResources->SimulationAttributesTexture.TextureTargetRHI,
-			ParticleSimulationResources->SimulationAttributesTexture.TextureRHI,
-			FResolveParams()
-		);
+			RHICmdList.SetViewport(0, 0, 0.0f, GParticleSimulationTextureSizeX, GParticleSimulationTextureSizeY, 1.0f);
+			FGraphicsPipelineStateInitializer GraphicsPSOInit;
+			RHICmdList.ApplyCachedRenderTargets(GraphicsPSOInit);
+			GraphicsPSOInit.DepthStencilState = TStaticDepthStencilState<false, CF_Always>::GetRHI();
+			GraphicsPSOInit.RasterizerState = TStaticRasterizerState<FM_Solid, CM_None>::GetRHI();
+			GraphicsPSOInit.BlendState = TStaticBlendState<>::GetRHI();
+
+			// Inject particles.
+			InjectNewParticles<true>(RHICmdList, GraphicsPSOInit, this->FeatureLevel, this->LastFrameNewParticles);
 
+			RHICmdList.EndRenderPass();
+		}
 	}
 
 	// Clear out particles from last frame
