Renderpass updates.
Kismet rendering to render passes.

diff --git a/Engine/Source/Runtime/Engine/Private/KismetRenderingLibrary.cpp b/Engine/Source/Runtime/Engine/Private/KismetRenderingLibrary.cpp
--- a/Engine/Source/Runtime/Engine/Private/KismetRenderingLibrary.cpp
+++ b/Engine/Source/Runtime/Engine/Private/KismetRenderingLibrary.cpp
@@ -51,8 +51,11 @@ void UKismetRenderingLibrary::ClearRenderTarget2D(UObject* WorldContextObject, U
 		ENQUEUE_RENDER_COMMAND(ClearRTCommand)(
 			[RenderTargetResource, ClearColor](FRHICommandList& RHICmdList)
 			{
-				SetRenderTarget(RHICmdList, RenderTargetResource->GetRenderTargetTexture(), FTextureRHIRef(), true);
+				FRHIRenderPassInfo RPInfo(RenderTargetResource->GetRenderTargetTexture(), ERenderTargetActions::DontLoad_Store);
+				TransitionRenderPassTargets(RHICmdList, RPInfo);
+				RHICmdList.BeginRenderPass(RPInfo, TEXT("ClearRT"));
 				DrawClearQuad(RHICmdList, ClearColor);
+				RHICmdList.EndRenderPass();
 			});
 	}
 }
