Fix persistent debug lines not working

diff --git a/Engine/Source/Runtime/Engine/Private/DrawDebugHelpers.cpp b/Engine/Source/Runtime/Engine/Private/DrawDebugHelpers.cpp
--- a/Engine/Source/Runtime/Engine/Private/DrawDebugHelpers.cpp
+++ b/Engine/Source/Runtime/Engine/Private/DrawDebugHelpers.cpp
@@ -29,6 +29,11 @@ ULineBatchComponent* GetDebugLineBatcher( const UWorld* InWorld, bool bPersisten
 	return (InWorld ? (bDepthIsForeground ? InWorld->ForegroundLineBatcher : (( bPersistentLines || (LifeTime > 0.f) ) ? InWorld->PersistentLineBatcher : InWorld->LineBatcher)) : NULL);
 }
 
+static float GetDebugLineLifeTime(ULineBatchComponent* LineBatcher, float LifeTime, bool bPersistent)
+{
+	return bPersistent ? -1.0f : ((LifeTime > 0.f) ? LifeTime : LineBatcher->DefaultLifeTime);
+}
+
 void DrawDebugLine(const UWorld* InWorld, FVector const& LineStart, FVector const& LineEnd, FColor const& Color, bool bPersistentLines, float LifeTime, uint8 DepthPriority, float Thickness)
 {
 	// no debug line drawing on dedicated server
@@ -38,7 +43,7 @@ void DrawDebugLine(const UWorld* InWorld, FVector const& LineStart, FVector cons
 		ULineBatchComponent* const LineBatcher = GetDebugLineBatcher( InWorld, bPersistentLines, LifeTime, (DepthPriority == SDPG_Foreground) );
 		if(LineBatcher != NULL)
 		{
-			float const LineLifeTime = (LifeTime > 0.f) ? LifeTime : LineBatcher->DefaultLifeTime;
+			float const LineLifeTime = GetDebugLineLifeTime(LineBatcher, LifeTime, bPersistentLines);
 			LineBatcher->DrawLine(LineStart, LineEnd, Color, DepthPriority, Thickness, LineLifeTime);
 		}
 	}
@@ -53,7 +58,7 @@ void DrawDebugPoint(const UWorld* InWorld, FVector const& Position, float Size,
 		ULineBatchComponent* const LineBatcher = GetDebugLineBatcher( InWorld, bPersistentLines, LifeTime, (DepthPriority == SDPG_Foreground) );
 		if(LineBatcher != NULL)
 		{
-			const float PointLifeTime = (LifeTime > 0.f ? LifeTime : LineBatcher->DefaultLifeTime);
+			const float PointLifeTime = GetDebugLineLifeTime(LineBatcher, LifeTime, bPersistentLines);
 			LineBatcher->DrawPoint(Position, Color.ReinterpretAsLinear(), Size, DepthPriority, PointLifeTime);
 		}
 	}
@@ -101,7 +106,7 @@ void DrawDebugBox(const UWorld* InWorld, FVector const& Center, FVector const& B
 		ULineBatchComponent* const LineBatcher = GetDebugLineBatcher( InWorld, bPersistentLines, LifeTime, (DepthPriority == SDPG_Foreground) );		
 		if(LineBatcher != NULL)
 		{
-			float LineLifeTime = (LifeTime > 0.f)? LifeTime : LineBatcher->DefaultLifeTime;
+			float LineLifeTime = GetDebugLineLifeTime(LineBatcher, LifeTime, bPersistentLines);
 
 			LineBatcher->DrawLine(Center + FVector( Box.X,  Box.Y,  Box.Z), Center + FVector( Box.X, -Box.Y, Box.Z), Color, DepthPriority, Thickness, LineLifeTime);
 			LineBatcher->DrawLine(Center + FVector( Box.X, -Box.Y,  Box.Z), Center + FVector(-Box.X, -Box.Y, Box.Z), Color, DepthPriority, Thickness, LineLifeTime);
@@ -130,7 +135,7 @@ void DrawDebugBox(const UWorld* InWorld, FVector const& Center, FVector const& B
 		ULineBatchComponent* const LineBatcher = GetDebugLineBatcher( InWorld, bPersistentLines, LifeTime, (DepthPriority == SDPG_Foreground) );
 		if(LineBatcher != NULL)
 		{
-			float const LineLifeTime = (LifeTime > 0.f) ? LifeTime : LineBatcher->DefaultLifeTime;
+			float const LineLifeTime = GetDebugLineLifeTime(LineBatcher, LifeTime, bPersistentLines);
 			TArray<struct FBatchedLine> Lines;
 
 			FTransform const Transform(Rotation);
@@ -196,7 +201,7 @@ void DrawDebugMesh(const UWorld* InWorld, TArray<FVector> const& Verts, TArray<i
 		ULineBatchComponent* const LineBatcher = GetDebugLineBatcher( InWorld, bPersistent, LifeTime, false );
 		if(LineBatcher != NULL)
 		{
-			float const ActualLifetime = (LifeTime > 0.f) ? LifeTime : LineBatcher->DefaultLifeTime;
+			float const ActualLifetime = GetDebugLineLifeTime(LineBatcher, LifeTime, bPersistent);
 			LineBatcher->DrawMesh(Verts, Indices, Color, DepthPriority, ActualLifetime);
 		}
 	}
@@ -210,7 +215,7 @@ void DrawDebugSolidBox(const UWorld* InWorld, FBox const& Box, FColor const& Col
 		ULineBatchComponent* const LineBatcher = GetDebugLineBatcher( InWorld, bPersistent, LifeTime, false );
 		if(LineBatcher != NULL)
 		{
-			float const ActualLifetime = (LifeTime > 0.f) ? LifeTime : LineBatcher->DefaultLifeTime;
+			float const ActualLifetime = GetDebugLineLifeTime(LineBatcher, LifeTime, bPersistent);
 			LineBatcher->DrawSolidBox(Box, Transform, Color, DepthPriority, ActualLifetime);
 		}
 	}
@@ -310,7 +315,7 @@ void DrawDebugCoordinateSystem(const UWorld* InWorld, FVector const& AxisLoc, FR
 		ULineBatchComponent* const LineBatcher = GetDebugLineBatcher( InWorld, bPersistentLines, LifeTime, (DepthPriority == SDPG_Foreground) );
 		if(LineBatcher != NULL)
 		{
-			const float LineLifeTime = (LifeTime > 0.f ? LifeTime : LineBatcher->DefaultLifeTime);
+			const float LineLifeTime = GetDebugLineLifeTime(LineBatcher, LifeTime, bPersistentLines);
 			LineBatcher->DrawLine(AxisLoc, AxisLoc + X*Scale, FColor::Red, DepthPriority, Thickness, LineLifeTime );
 			LineBatcher->DrawLine(AxisLoc, AxisLoc + Y*Scale, FColor::Green, DepthPriority, Thickness, LineLifeTime );
 			LineBatcher->DrawLine(AxisLoc, AxisLoc + Z*Scale, FColor::Blue, DepthPriority, Thickness, LineLifeTime );
@@ -332,7 +337,7 @@ ENGINE_API void DrawDebugCrosshairs(const UWorld* InWorld, FVector const& AxisLo
 		ULineBatchComponent* const LineBatcher = GetDebugLineBatcher(InWorld, bPersistentLines, LifeTime, (DepthPriority == SDPG_Foreground));
 		if (LineBatcher != NULL)
 		{
-			const float LineLifeTime = (LifeTime > 0.f) ? LifeTime : LineBatcher->DefaultLifeTime;
+			const float LineLifeTime = GetDebugLineLifeTime(LineBatcher, LifeTime, bPersistentLines);
 			LineBatcher->DrawLine(AxisLoc - X*Scale, AxisLoc + X*Scale, Color, DepthPriority, 0.f, LineLifeTime);
 			LineBatcher->DrawLine(AxisLoc - Y*Scale, AxisLoc + Y*Scale, Color, DepthPriority, 0.f, LineLifeTime);
 			LineBatcher->DrawLine(AxisLoc - Z*Scale, AxisLoc + Z*Scale, Color, DepthPriority, 0.f, LineLifeTime);
@@ -349,7 +354,7 @@ static void InternalDrawDebugCircle(const UWorld* InWorld, const FMatrix& Transf
 		ULineBatchComponent* const LineBatcher = GetDebugLineBatcher( InWorld, bPersistentLines, LifeTime, (DepthPriority == SDPG_Foreground) );
 		if(LineBatcher != NULL)
 		{
-			const float LineLifeTime = (LifeTime > 0.f) ? LifeTime : LineBatcher->DefaultLifeTime;
+			const float LineLifeTime = GetDebugLineLifeTime(LineBatcher, LifeTime, bPersistentLines);
 
 			// Need at least 4 segments
 			Segments = FMath::Max(Segments, 4);
@@ -383,7 +388,7 @@ void DrawDebugCircle(const UWorld* InWorld, const FMatrix& TransformMatrix, floa
 		ULineBatchComponent* const LineBatcher = GetDebugLineBatcher( InWorld, bPersistentLines, LifeTime, (DepthPriority == SDPG_Foreground) );
 		if(LineBatcher != NULL)
 		{
-			const float LineLifeTime = (LifeTime > 0.f) ? LifeTime : LineBatcher->DefaultLifeTime;
+			const float LineLifeTime = GetDebugLineLifeTime(LineBatcher, LifeTime, bPersistentLines);
 
 			// Need at least 4 segments
 			Segments = FMath::Max((Segments - 2) / 2, 4);
@@ -435,7 +440,7 @@ void DrawDebug2DDonut(const UWorld* InWorld, const FMatrix& TransformMatrix, flo
 		ULineBatchComponent* const LineBatcher = GetDebugLineBatcher( InWorld, bPersistentLines, LifeTime, (DepthPriority == SDPG_Foreground) );
 		if(LineBatcher != NULL)
 		{
-			const float LineLifeTime = (LifeTime > 0.f) ? LifeTime : LineBatcher->DefaultLifeTime;
+			const float LineLifeTime = GetDebugLineLifeTime(LineBatcher, LifeTime, bPersistentLines);
 
 			// Need at least 4 segments
 			Segments = FMath::Max((Segments - 4) / 2, 4);
@@ -466,7 +471,7 @@ void DrawDebugSphere(const UWorld* InWorld, FVector const& Center, float Radius,
 		ULineBatchComponent* const LineBatcher = GetDebugLineBatcher( InWorld, bPersistentLines, LifeTime, (DepthPriority == SDPG_Foreground) );
 		if (LineBatcher != NULL)
 		{
-			float LineLifeTime = (LifeTime > 0.f) ? LifeTime : LineBatcher->DefaultLifeTime;
+			float LineLifeTime = GetDebugLineLifeTime(LineBatcher, LifeTime, bPersistentLines);
 
 			// Need at least 4 segments
 			Segments = FMath::Max(Segments, 4);
@@ -550,7 +555,7 @@ void DrawDebugCylinder(const UWorld* InWorld, FVector const& Start, FVector cons
 		ULineBatchComponent* const LineBatcher = GetDebugLineBatcher( InWorld, bPersistentLines, LifeTime, (DepthPriority == SDPG_Foreground) );
 		if(LineBatcher != NULL)
 		{
-			const float LineLifeTime = (LifeTime > 0.f) ? LifeTime : LineBatcher->DefaultLifeTime;
+			const float LineLifeTime = GetDebugLineLifeTime(LineBatcher, LifeTime, bPersistentLines);
 			while( Segments-- )
 			{
 				Segment = Perpendicular.RotateAngleAxis(Angle, Axis) * Radius;
@@ -582,7 +587,7 @@ void DrawDebugAltCone(const UWorld* InWorld, FVector const& Origin, FRotator con
 	if (LineBatcher != NULL)
 	{
 
-		float const LineLifeTime = (LifeTime > 0.f) ? LifeTime : LineBatcher->DefaultLifeTime;
+		float const LineLifeTime = GetDebugLineLifeTime(LineBatcher, LifeTime, bPersistentLines);
 
 		FVector const EndPoint = Origin + AxisX * Length;
 		FVector const Up = FMath::Tan(AngleHeight * 0.5f) * AxisZ * Length;
@@ -687,7 +692,7 @@ void DrawDebugCone(const UWorld* InWorld, FVector const& Origin, FVector const&
 		ULineBatchComponent* const LineBatcher = GetDebugLineBatcher( InWorld, bPersistentLines, LifeTime, (DepthPriority == SDPG_Foreground) );
 		if(LineBatcher != NULL)
 		{
-			float const LineLifeTime = (LifeTime > 0.f) ? LifeTime : LineBatcher->DefaultLifeTime;
+			float const LineLifeTime = GetDebugLineLifeTime(LineBatcher, LifeTime, bPersistentLines);
 
 			TArray<FBatchedLine> Lines;
 			Lines.Empty(NumSides);
