Make UStaticMesh::FindSocket blueprint callable (also make const)

diff --git a/Engine/Source/Runtime/Engine/Classes/Engine/StaticMesh.h b/Engine/Source/Runtime/Engine/Classes/Engine/StaticMesh.h
--- a/Engine/Source/Runtime/Engine/Classes/Engine/StaticMesh.h
+++ b/Engine/Source/Runtime/Engine/Classes/Engine/StaticMesh.h
@@ -1011,7 +1011,9 @@ public:
 	 *	Find a socket object in this StaticMesh by name.
 	 *	Entering NAME_None will return NULL. If there are multiple sockets with the same name, will return the first one.
 	 */
-	ENGINE_API class UStaticMeshSocket* FindSocket(FName InSocketName);
+	UFUNCTION(BlueprintCallable, Category = "StaticMesh")
+	ENGINE_API class UStaticMeshSocket* FindSocket(FName InSocketName) const;
+
 	/**
 	 * Returns vertex color data by position.
 	 * For matching to reimported meshes that may have changed or copying vertex paint data from mesh to mesh.
diff --git a/Engine/Source/Runtime/Engine/Private/StaticMesh.cpp b/Engine/Source/Runtime/Engine/Private/StaticMesh.cpp
--- a/Engine/Source/Runtime/Engine/Private/StaticMesh.cpp
+++ b/Engine/Source/Runtime/Engine/Private/StaticMesh.cpp
@@ -4797,7 +4797,7 @@ void UStaticMesh::GenerateLodsInPackage()
 
 #endif // #if WITH_EDITOR
 
-UStaticMeshSocket* UStaticMesh::FindSocket(FName InSocketName)
+UStaticMeshSocket* UStaticMesh::FindSocket(FName InSocketName) const
 {
 	if(InSocketName == NAME_None)
 	{
