Fix shadowed variable warning

diff --git a/Engine/Source/Runtime/Engine/Private/Animation/AnimMontage.cpp b/Engine/Source/Runtime/Engine/Private/Animation/AnimMontage.cpp
--- a/Engine/Source/Runtime/Engine/Private/Animation/AnimMontage.cpp
+++ b/Engine/Source/Runtime/Engine/Private/Animation/AnimMontage.cpp
@@ -2671,7 +2671,7 @@ UAnimMontage* FAnimMontageInstance::PreviewMatineeSetAnimPositionInner(FName Slo
 	return PlayingMontage;
 }
 
-UAnimMontage* FAnimMontageInstance::SetSequencerMontagePosition(FName SlotName, USkeletalMeshComponent* SkeletalMeshComponent, int32& InOutInstanceId, UAnimSequenceBase* InAnimSequence, float InPosition, float Weight, bool bLooping, bool bPlaying)
+UAnimMontage* FAnimMontageInstance::SetSequencerMontagePosition(FName SlotName, USkeletalMeshComponent* SkeletalMeshComponent, int32& InOutInstanceId, UAnimSequenceBase* InAnimSequence, float InPosition, float Weight, bool bLooping, bool bInPlaying)
 {
 	UAnimInstance* AnimInst = SkeletalMeshComponent->GetAnimInstance();
 	if (AnimInst)
@@ -2699,7 +2699,7 @@ UAnimMontage* FAnimMontageInstance::SetSequencerMontagePosition(FName SlotName,
 			MontageInstanceToUpdate->Blend.SetDesiredValue(Weight);
 			MontageInstanceToUpdate->Blend.SetAlpha(Weight);
 			MontageInstanceToUpdate->SetPosition(InPosition);
-			MontageInstanceToUpdate->bPlaying = bPlaying;
+			MontageInstanceToUpdate->bPlaying = bInPlaying;
 			return PlayingMontage;
 		}
 	}
@@ -2711,7 +2711,7 @@ UAnimMontage* FAnimMontageInstance::SetSequencerMontagePosition(FName SlotName,
 	return nullptr;
 }
 
-UAnimMontage* FAnimMontageInstance::PreviewSequencerMontagePosition(FName SlotName, USkeletalMeshComponent* SkeletalMeshComponent, int32& InOutInstanceId, UAnimSequenceBase* InAnimSequence, float InPosition, float Weight, bool bLooping, bool bFireNotifies, bool bPlaying)
+UAnimMontage* FAnimMontageInstance::PreviewSequencerMontagePosition(FName SlotName, USkeletalMeshComponent* SkeletalMeshComponent, int32& InOutInstanceId, UAnimSequenceBase* InAnimSequence, float InPosition, float Weight, bool bLooping, bool bFireNotifies, bool bInPlaying)
 {
 	UAnimInstance* AnimInst = SkeletalMeshComponent->GetAnimInstance();
 	if (AnimInst)
@@ -2719,7 +2719,7 @@ UAnimMontage* FAnimMontageInstance::PreviewSequencerMontagePosition(FName SlotNa
 		FAnimMontageInstance* MontageInstanceToUpdate = AnimInst->GetMontageInstanceForID(InOutInstanceId);
 		float PreviousPosition = (MontageInstanceToUpdate) ? MontageInstanceToUpdate->GetPosition() : InPosition;
 
-		UAnimMontage* PlayingMontage = SetSequencerMontagePosition(SlotName, SkeletalMeshComponent, InOutInstanceId, InAnimSequence, InPosition, Weight, bLooping, bPlaying);
+		UAnimMontage* PlayingMontage = SetSequencerMontagePosition(SlotName, SkeletalMeshComponent, InOutInstanceId, InAnimSequence, InPosition, Weight, bLooping, bInPlaying);
 		if (PlayingMontage)
 		{
 			// we have to get it again in case if this is new
