Add the ability to have a rumble pattern play while paused.
Based loosely on PR#5337

diff --git a/Engine/Source/Runtime/Engine/Classes/GameFramework/ForceFeedbackEffect.h b/Engine/Source/Runtime/Engine/Classes/GameFramework/ForceFeedbackEffect.h
--- a/Engine/Source/Runtime/Engine/Classes/GameFramework/ForceFeedbackEffect.h
+++ b/Engine/Source/Runtime/Engine/Classes/GameFramework/ForceFeedbackEffect.h
@@ -1,4 +1,4 @@
-﻿// Copyright 1998-2018 Epic Games, Inc. All Rights Reserved.
+// Copyright 1998-2018 Epic Games, Inc. All Rights Reserved.
 #pragma once
 
 #include "CoreMinimal.h"
@@ -39,6 +39,31 @@ struct FForceFeedbackChannelDetails
 	}
 };
 
+/** This structure is used to pass arguments to ClientPlayForceFeedback() client RPC function */
+USTRUCT()
+struct FForceFeedbackParameters
+{
+	GENERATED_BODY()
+
+		FForceFeedbackParameters()
+		: bLooping(false)
+		, bIgnoreTimeDilation(false)
+		, bPlayWhilePaused(false)
+	{}
+
+	UPROPERTY()
+	FName Tag;
+
+	UPROPERTY()
+	bool bLooping;
+
+	UPROPERTY()
+	bool bIgnoreTimeDilation;
+
+	UPROPERTY()
+	bool bPlayWhilePaused;
+};
+
 USTRUCT()
 struct ENGINE_API FActiveForceFeedbackEffect
 {
@@ -47,25 +72,18 @@ struct ENGINE_API FActiveForceFeedbackEffect
 	UPROPERTY()
 	class UForceFeedbackEffect* ForceFeedbackEffect;
 
-	FName Tag;
-	uint32 bLooping:1;
-	uint32 bIgnoreTimeDilation:1;
+	FForceFeedbackParameters Parameters;
 	float PlayTime;
 
 	FActiveForceFeedbackEffect()
 		: ForceFeedbackEffect(nullptr)
-		, Tag(NAME_None)
-		, bLooping(false)
-		, bIgnoreTimeDilation(false)
 		, PlayTime(0.f)
 	{
 	}
 
-	FActiveForceFeedbackEffect(UForceFeedbackEffect* InEffect, const bool bInLooping, const bool bInIgnoreTimeDilation, const FName InTag)
+	FActiveForceFeedbackEffect(UForceFeedbackEffect* InEffect, FForceFeedbackParameters InParameters)
 		: ForceFeedbackEffect(InEffect)
-		, Tag(InTag)
-		, bLooping(bInLooping)
-		, bIgnoreTimeDilation(bInIgnoreTimeDilation)
+		, Parameters(InParameters)
 		, PlayTime(0.f)
 	{
 	}
diff --git a/Engine/Source/Runtime/Engine/Classes/GameFramework/PlayerController.h b/Engine/Source/Runtime/Engine/Classes/GameFramework/PlayerController.h
--- a/Engine/Source/Runtime/Engine/Classes/GameFramework/PlayerController.h
+++ b/Engine/Source/Runtime/Engine/Classes/GameFramework/PlayerController.h
@@ -1031,15 +1031,40 @@ public:
 	 * @param	ForceFeedbackEffect		The force feedback pattern to play
 	 * @param	bLooping				Whether the pattern should be played repeatedly or be a single one shot
 	 * @param	bIgnoreTimeDilation		Whether the pattern should ignore time dilation
+	 * @param	bPlayWhilePaused		Whether the pattern should continue to play while the game is paused
 	 * @param	Tag						A tag that allows stopping of an effect.  If another effect with this Tag is playing, it will be stopped and replaced
 	 */
-	UFUNCTION(unreliable, client, BlueprintCallable, Category="Game|Feedback")
-	void ClientPlayForceFeedback(class UForceFeedbackEffect* ForceFeedbackEffect, bool bLooping, bool bIgnoreTimeDilation, FName Tag);
+	UFUNCTION(BlueprintCallable, Category="Game|Feedback", meta=(DisplayName="Client Play Force Feedback", AdvancedDisplay="bIgnoreTimeDilation,bPlayWhilePaused"))
+	void K2_ClientPlayForceFeedback(class UForceFeedbackEffect* ForceFeedbackEffect, FName Tag, bool bLooping, bool bIgnoreTimeDilation, bool bPlayWhilePaused)
+	{
+		FForceFeedbackParameters Params;
+		Params.Tag = Tag;
+		Params.bLooping = bLooping;
+		Params.bIgnoreTimeDilation = bIgnoreTimeDilation;
+		Params.bPlayWhilePaused = bPlayWhilePaused;
+		ClientPlayForceFeedback(ForceFeedbackEffect, Params);
+	}
+
+	UFUNCTION(unreliable, client)
+	void ClientPlayForceFeedback(class UForceFeedbackEffect* ForceFeedbackEffect, FForceFeedbackParameters Params = FForceFeedbackParameters());
+
+	UE_DEPRECATED(4.22, "Use version that specifies parameters using a struct instead of a list of parameters")
+	void ClientPlayForceFeedback(class UForceFeedbackEffect* ForceFeedbackEffect, bool bLooping, bool bIgnoreTimeDilation, FName Tag)
+	{
+		FForceFeedbackParameters Params;
+		Params.Tag = Tag;
+		Params.bLooping = bLooping;
+		Params.bIgnoreTimeDilation = bIgnoreTimeDilation;
+		ClientPlayForceFeedback(ForceFeedbackEffect, Params);
+	}
 
-	UE_DEPRECATED(4.18, "Use version that specifies whether to ignore time dilation or not")
+	UE_DEPRECATED(4.18, "Use version that specifies parameters using a struct instead of a list of parameters")
 	void ClientPlayForceFeedback(class UForceFeedbackEffect* ForceFeedbackEffect, bool bLooping, FName Tag)
 	{
-		ClientPlayForceFeedback(ForceFeedbackEffect, bLooping, false, Tag);
+		FForceFeedbackParameters Params;
+		Params.Tag = Tag;
+		Params.bLooping = bLooping;
+		ClientPlayForceFeedback(ForceFeedbackEffect, Params);
 	}
 
 	/** 
diff --git a/Engine/Source/Runtime/Engine/Private/GameFramework/ForceFeedbackEffect.cpp b/Engine/Source/Runtime/Engine/Private/GameFramework/ForceFeedbackEffect.cpp
--- a/Engine/Source/Runtime/Engine/Private/GameFramework/ForceFeedbackEffect.cpp
+++ b/Engine/Source/Runtime/Engine/Private/GameFramework/ForceFeedbackEffect.cpp
@@ -1,4 +1,4 @@
-﻿// Copyright 1998-2018 Epic Games, Inc. All Rights Reserved.
+// Copyright 1998-2018 Epic Games, Inc. All Rights Reserved.
 
 #include "GameFramework/ForceFeedbackEffect.h"
 #include "GenericPlatform/IInputInterface.h"
@@ -92,9 +92,9 @@ bool FActiveForceFeedbackEffect::Update(const float DeltaTime, FForceFeedbackVal
 
 	const float Duration = ForceFeedbackEffect->GetDuration();
 
-	PlayTime += (bIgnoreTimeDilation ? FApp::GetDeltaTime() : DeltaTime);
+	PlayTime += (Parameters.bIgnoreTimeDilation ? FApp::GetDeltaTime() : DeltaTime);
 
-	if (PlayTime > Duration && (!bLooping || Duration == 0.f) )
+	if (PlayTime > Duration && (!Parameters.bLooping || Duration == 0.f) )
 	{
 		return false;
 	}
diff --git a/Engine/Source/Runtime/Engine/Private/PlayerController.cpp b/Engine/Source/Runtime/Engine/Private/PlayerController.cpp
--- a/Engine/Source/Runtime/Engine/Private/PlayerController.cpp
+++ b/Engine/Source/Runtime/Engine/Private/PlayerController.cpp
@@ -3051,9 +3051,9 @@ void APlayerController::DisplayDebug(class UCanvas* Canvas, const FDebugDisplayI
 				const FActiveForceFeedbackEffect& LastActiveEffect = ForceFeedbackEffectHistoryEntries[i].LastActiveForceFeedbackEffect;
 				const FString HistoryEntry = FString::Printf(TEXT("%s %s %f %s %f"), 
 															*LastActiveEffect.ForceFeedbackEffect->GetFName().ToString(), 
-															*LastActiveEffect.Tag.ToString(), 
+															*LastActiveEffect.Parameters.Tag.ToString(), 
 															LastActiveEffect.ForceFeedbackEffect->GetDuration(),
-															(LastActiveEffect.bLooping ? TEXT("true") : TEXT("false")),
+															(LastActiveEffect.Parameters.bLooping ? TEXT("true") : TEXT("false")),
 															ForceFeedbackEffectHistoryEntries[i].TimeShown);
 				DisplayDebugManager.DrawString(HistoryEntry);
 			}
@@ -3071,9 +3071,9 @@ void APlayerController::DisplayDebug(class UCanvas* Canvas, const FDebugDisplayI
 			{
 				const FString ActiveEntry = FString::Printf(TEXT("%s %s N/A %.2f %s %.2f - LL: %.2f LS: %.2f RL: %.2f RS: %.2f"),
 					*ActiveEffect.ForceFeedbackEffect->GetFName().ToString(),
-					*ActiveEffect.Tag.ToString(),
+					*ActiveEffect.Parameters.Tag.ToString(),
 					ActiveEffect.ForceFeedbackEffect->GetDuration(),
-					(ActiveEffect.bLooping ? TEXT("true") : TEXT("false")),
+					(ActiveEffect.Parameters.bLooping ? TEXT("true") : TEXT("false")),
 					ActiveEffect.PlayTime,
 					ActiveValues.LeftLarge, ActiveValues.LeftSmall, ActiveValues.RightLarge, ActiveValues.RightSmall);
 				DisplayDebugManager.DrawString(ActiveEntry);
@@ -3662,22 +3662,22 @@ void APlayerController::ClientPrestreamTextures_Implementation( AActor* ForcedAc
 	}
 }
 
-void APlayerController::ClientPlayForceFeedback_Implementation( UForceFeedbackEffect* ForceFeedbackEffect, bool bLooping, bool bIgnoreTimeDilation, FName Tag)
+void APlayerController::ClientPlayForceFeedback_Implementation( UForceFeedbackEffect* ForceFeedbackEffect, FForceFeedbackParameters Params)
 {
 	if (ForceFeedbackEffect)
 	{
-		if (Tag != NAME_None)
+		if (Params.Tag != NAME_None)
 		{
 			for (int32 Index = ActiveForceFeedbackEffects.Num() - 1; Index >= 0; --Index)
 			{
-				if (ActiveForceFeedbackEffects[Index].Tag == Tag)
+				if (ActiveForceFeedbackEffects[Index].Parameters.Tag == Params.Tag)
 				{
 					ActiveForceFeedbackEffects.RemoveAtSwap(Index);
 				}
 			}
 		}
 
-		ActiveForceFeedbackEffects.Emplace(ForceFeedbackEffect, bLooping, bIgnoreTimeDilation, Tag);
+		ActiveForceFeedbackEffects.Emplace(ForceFeedbackEffect, Params);
 
 #if !(UE_BUILD_SHIPPING || UE_BUILD_TEST)
 		ForceFeedbackEffectHistoryEntries.Emplace(ActiveForceFeedbackEffects.Last(), GetWorld()->GetTimeSeconds());
@@ -3696,7 +3696,7 @@ void APlayerController::ClientStopForceFeedback_Implementation( UForceFeedbackEf
 		for (int32 Index = ActiveForceFeedbackEffects.Num() - 1; Index >= 0; --Index)
 		{
 			if (    (ForceFeedbackEffect == NULL || ActiveForceFeedbackEffects[Index].ForceFeedbackEffect == ForceFeedbackEffect)
-				 && (Tag == NAME_None || ActiveForceFeedbackEffects[Index].Tag == Tag) )
+				 && (Tag == NAME_None || ActiveForceFeedbackEffects[Index].Parameters.Tag == Tag) )
 			{
 				ActiveForceFeedbackEffects.RemoveAtSwap(Index);
 			}
@@ -4047,7 +4047,9 @@ void APlayerController::ProcessForceFeedbackAndHaptics(const float DeltaTime, co
 	bool bRightHapticsNeedUpdate = false;
 	bool bGunHapticsNeedUpdate = false;
 
-	bool bProcessFeedback = !bGamePaused;
+	// Always process feedback by default, but if the game is paused then only static
+	// effects that are flagged to play while paused will play
+	bool bProcessFeedback = true;
 #if WITH_EDITOR
 	if (bProcessFeedback)
 	{
@@ -4069,23 +4071,31 @@ void APlayerController::ProcessForceFeedbackAndHaptics(const float DeltaTime, co
 		// --- Force Feedback --------------------------
 		for (int32 Index = ActiveForceFeedbackEffects.Num() - 1; Index >= 0; --Index)
 		{
-			if (!ActiveForceFeedbackEffects[Index].Update(DeltaTime, ForceFeedbackValues))
+			// If the game is paused, only tick force feedback effects that want to ignore time dilation
+			if (!bGamePaused || ActiveForceFeedbackEffects[Index].Parameters.bPlayWhilePaused)
 			{
-				ActiveForceFeedbackEffects.RemoveAtSwap(Index);
+				if (!ActiveForceFeedbackEffects[Index].Update(DeltaTime, ForceFeedbackValues))
+				{
+					ActiveForceFeedbackEffects.RemoveAtSwap(Index);
+				}
 			}
 		}
 
-		for (TSortedMap<uint64, FDynamicForceFeedbackAction>::TIterator It(DynamicForceFeedbacks.CreateIterator()); It; ++It)
+		const bool bProcessDynamicFeedback = !bGamePaused;
+		if (bProcessDynamicFeedback)
 		{
-			if (!It.Value().Update(DeltaTime, ForceFeedbackValues))
+			for (TSortedMap<uint64, FDynamicForceFeedbackAction>::TIterator It(DynamicForceFeedbacks.CreateIterator()); It; ++It)
 			{
-				It.RemoveCurrent();
+				if (!It.Value().Update(DeltaTime, ForceFeedbackValues))
+				{
+					It.RemoveCurrent();
+				}
 			}
-		}
 
-		for (const TPair<int32, FDynamicForceFeedbackDetails*>& DynamicEntry : LatentDynamicForceFeedbacks)
-		{
-			DynamicEntry.Value->Update(ForceFeedbackValues);
+			for (const TPair<int32, FDynamicForceFeedbackDetails*>& DynamicEntry : LatentDynamicForceFeedbacks)
+			{
+				DynamicEntry.Value->Update(ForceFeedbackValues);
+			}
 		}
 
 		if (FForceFeedbackManager* ForceFeedbackManager = FForceFeedbackManager::Get(World))
@@ -4100,39 +4110,41 @@ void APlayerController::ProcessForceFeedbackAndHaptics(const float DeltaTime, co
 		ForceFeedbackValues.RightSmall = FMath::Clamp(ForceFeedbackValues.RightSmall * ForceFeedbackScale, 0.f, 1.f);
 
 		// --- Haptic Feedback -------------------------
-		if (ActiveHapticEffect_Left.IsValid())
+		if (bProcessDynamicFeedback)
 		{
-			const bool bPlaying = ActiveHapticEffect_Left->Update(DeltaTime, LeftHaptics);
-			if (!bPlaying)
+			if (ActiveHapticEffect_Left.IsValid())
 			{
-				ActiveHapticEffect_Left->bLoop ? ActiveHapticEffect_Left->Restart() : ActiveHapticEffect_Left.Reset();
-			}
+				const bool bPlaying = ActiveHapticEffect_Left->Update(DeltaTime, LeftHaptics);
+				if (!bPlaying)
+				{
+					ActiveHapticEffect_Left->bLoop ? ActiveHapticEffect_Left->Restart() : ActiveHapticEffect_Left.Reset();
+				}
 
-			bLeftHapticsNeedUpdate = true;
-		}
+				bLeftHapticsNeedUpdate = true;
+			}
 
-		if (ActiveHapticEffect_Right.IsValid())
-		{
-			const bool bPlaying = ActiveHapticEffect_Right->Update(DeltaTime, RightHaptics);
-			if (!bPlaying)
+			if (ActiveHapticEffect_Right.IsValid())
 			{
-				ActiveHapticEffect_Right->bLoop ? ActiveHapticEffect_Right->Restart() : ActiveHapticEffect_Right.Reset();
-			}
+				const bool bPlaying = ActiveHapticEffect_Right->Update(DeltaTime, RightHaptics);
+				if (!bPlaying)
+				{
+					ActiveHapticEffect_Right->bLoop ? ActiveHapticEffect_Right->Restart() : ActiveHapticEffect_Right.Reset();
+				}
 
-			bRightHapticsNeedUpdate = true;
-		}
+				bRightHapticsNeedUpdate = true;
+			}
 
-		if (ActiveHapticEffect_Gun.IsValid())
-		{
-			const bool bPlaying = ActiveHapticEffect_Gun->Update(DeltaTime, GunHaptics);
-			if (!bPlaying)
+			if (ActiveHapticEffect_Gun.IsValid())
 			{
-				ActiveHapticEffect_Gun->bLoop ? ActiveHapticEffect_Gun->Restart() : ActiveHapticEffect_Gun.Reset();
-			}
+				const bool bPlaying = ActiveHapticEffect_Gun->Update(DeltaTime, GunHaptics);
+				if (!bPlaying)
+				{
+					ActiveHapticEffect_Gun->bLoop ? ActiveHapticEffect_Gun->Restart() : ActiveHapticEffect_Gun.Reset();
+				}
 
-			bGunHapticsNeedUpdate = true;
+				bGunHapticsNeedUpdate = true;
+			}
 		}
-
 	}
 
 	if (FSlateApplication::IsInitialized())
