(4.21.2) PR #5349: Fix wrong stat value being reported for PerformLastStep (Contributed by [REDACTED_PERSON])
#rnx

diff --git a/Engine/Source/Runtime/Engine/Private/World.cpp b/Engine/Source/Runtime/Engine/Private/World.cpp
--- a/Engine/Source/Runtime/Engine/Private/World.cpp
+++ b/Engine/Source/Runtime/Engine/Private/World.cpp
@@ -2426,7 +2426,7 @@ void UWorld::AddToWorld( ULevel* Level, const FTransform& LevelTransform, bool b
 		UE_LOG(LogStreaming, Display, TEXT("Initialize              : %6.2f ms"), RouteActorInitializeTime * 1000 );
 		UE_LOG(LogStreaming, Display, TEXT("Cross Level Refs        : %6.2f ms"), CrossLevelRefsTime * 1000 );
 		UE_LOG(LogStreaming, Display, TEXT("Sort Actor List         : %6.2f ms"), SortActorListTime * 1000 );
-		UE_LOG(LogStreaming, Display, TEXT("Perform Last Step       : %6.2f ms"), SortActorListTime * 1000 );
+		UE_LOG(LogStreaming, Display, TEXT("Perform Last Step       : %6.2f ms"), PerformLastStepTime * 1000 );
 	}
 #endif // PERF_TRACK_DETAILED_ASYNC_STATS
 }
