Add a == node for arrays in blueprint

diff --git a/Engine/Source/Runtime/Engine/Classes/Kismet/KismetArrayLibrary.h b/Engine/Source/Runtime/Engine/Classes/Kismet/KismetArrayLibrary.h
--- a/Engine/Source/Runtime/Engine/Classes/Kismet/KismetArrayLibrary.h
+++ b/Engine/Source/Runtime/Engine/Classes/Kismet/KismetArrayLibrary.h
@@ -46,6 +46,16 @@ class ENGINE_API UKismetArrayLibrary : public UBlueprintFunctionLibrary
 	UFUNCTION(BlueprintCallable, CustomThunk, meta=(DisplayName = "Shuffle", CompactNodeTitle = "SHUFFLE", ArrayParm = "TargetArray"), Category="Utilities|Array")
 	static void Array_Shuffle(const TArray<int32>& TargetArray);
 
+	/** 
+	 * Checks if two arrays are memberwise identical
+	 *
+	 * @param	ArrayA		One of the arrays to compare
+	 * @param	ArrayB		The other array to compare
+	 * @return Whether the two arrays are identical
+	 */
+	UFUNCTION(BlueprintPure, CustomThunk, meta=(DisplayName = "Identical", CompactNodeTitle = "==", ArrayParm = "ArrayA,ArrayB", ArrayTypeDependentParams = "ArrayB"), Category="Utilities|Array")
+	static bool Array_Identical(const TArray<int32>& ArrayA, const TArray<int32>& ArrayB);
+
 	/** 
 	 *Append an array to another array
 	 *
@@ -203,6 +213,7 @@ class ENGINE_API UKismetArrayLibrary : public UBlueprintFunctionLibrary
 	static int32 GenericArray_Add(void* TargetArray, const UArrayProperty* ArrayProp, const void* NewItem);
 	static int32 GenericArray_AddUnique(void* TargetArray, const UArrayProperty* ArrayProp, const void* NewItem);
 	static void GenericArray_Shuffle(void* TargetArray, const UArrayProperty* ArrayProp);
+	static bool GenericArray_Identical(void* ArrayA, const UArrayProperty* ArrayAProp, void* ArrayB, const UArrayProperty* ArrayBProperty);
 	static void GenericArray_Append(void* TargetArray, const UArrayProperty* TargetArrayProp, void* SourceArray, const UArrayProperty* SourceArrayProperty);
 	static void GenericArray_Insert(void* TargetArray, const UArrayProperty* ArrayProp, const void* NewItem, int32 Index);
 	static void GenericArray_Remove(void* TargetArray, const UArrayProperty* ArrayProp, int32 IndexToRemove);
@@ -303,6 +314,35 @@ public:
 		P_NATIVE_END;
 	}
 
+	DECLARE_FUNCTION(execArray_Identical)
+	{
+		// Retrieve the first array
+		Stack.MostRecentProperty = nullptr;
+		Stack.StepCompiledIn<UArrayProperty>(NULL);
+		void* ArrayAAddr = Stack.MostRecentPropertyAddress;
+		UArrayProperty* ArrayAProperty = Cast<UArrayProperty>(Stack.MostRecentProperty);
+		if (!ArrayAProperty)
+		{
+			Stack.bArrayContextFailed = true;
+			return;
+		}
+		// Retrieve the second array
+		Stack.MostRecentProperty = nullptr;
+		Stack.StepCompiledIn<UArrayProperty>(NULL);
+		void* ArrayBAddr = Stack.MostRecentPropertyAddress;
+		UArrayProperty* ArrayBProperty = Cast<UArrayProperty>(Stack.MostRecentProperty);
+		if (!ArrayBProperty)
+		{
+			Stack.bArrayContextFailed = true;
+			return;
+		}
+
+		P_FINISH;
+		P_NATIVE_BEGIN;
+		*(bool*)RESULT_PARAM = GenericArray_Identical(ArrayAAddr, ArrayAProperty, ArrayBAddr, ArrayBProperty);
+		P_NATIVE_END;
+	}
+
 	DECLARE_FUNCTION(execArray_Append)
 	{
 		// Retrieve the target array
diff --git a/Engine/Source/Runtime/Engine/Private/KismetArrayLibrary.cpp b/Engine/Source/Runtime/Engine/Private/KismetArrayLibrary.cpp
--- a/Engine/Source/Runtime/Engine/Private/KismetArrayLibrary.cpp
+++ b/Engine/Source/Runtime/Engine/Private/KismetArrayLibrary.cpp
@@ -1,4 +1,4 @@
-﻿// Copyright 1998-2018 Epic Games, Inc. All Rights Reserved.
+// Copyright 1998-2018 Epic Games, Inc. All Rights Reserved.
 
 #include "Kismet/KismetArrayLibrary.h"
 #include "GameFramework/Actor.h"
@@ -103,6 +103,38 @@ int32 UKismetArrayLibrary::GenericArray_AddUnique(void* TargetArray, const UArra
 	return NewIndex;
 }
 
+bool UKismetArrayLibrary::GenericArray_Identical(void* ArrayA, const UArrayProperty* ArrayAProp, void* ArrayB, const UArrayProperty* ArrayBProp)
+{
+	if (ArrayA && ArrayB)
+	{
+		UProperty* InnerAProp = ArrayAProp->Inner;
+
+		if (InnerAProp->SameType(ArrayBProp->Inner))
+		{
+			FScriptArrayHelper ArrayAHelper(ArrayAProp, ArrayA);
+			FScriptArrayHelper ArrayBHelper(ArrayBProp, ArrayB);
+
+			const int32 ArrayANum = ArrayAHelper.Num();
+			if (ArrayANum != ArrayBHelper.Num())
+			{
+				return false;
+			}
+
+			for (int32 Index = 0; Index < ArrayANum; ++Index)
+			{
+				if (!InnerAProp->Identical(ArrayAHelper.GetRawPtr(Index), ArrayBHelper.GetRawPtr(Index)))
+				{
+					return false;
+				}
+			}
+
+			return true;
+		}
+	}
+
+	return false;
+}
+
 void UKismetArrayLibrary::GenericArray_Append(void* TargetArray, const UArrayProperty* TargetArrayProp, void* SourceArray, const UArrayProperty* SourceArrayProperty)
 {
 	if(TargetArray && SourceArray)
