Fix failure to redirect ClientPlayForceFeedback calls from blueprints to renamed function causing cook errors
#rnx

diff --git a/Engine/Config/BaseEngine.ini b/Engine/Config/BaseEngine.ini
--- a/Engine/Config/BaseEngine.ini
+++ b/Engine/Config/BaseEngine.ini
@@ -639,6 +639,8 @@ bUseStreamingPause=false
 +FunctionRedirects=(OldName="Controller.OnPossess",NewName="Controller.ReceivePossess")
 +FunctionRedirects=(OldName="Controller.OnUnPossess",NewName="Controller.ReceiveUnPossess")
 
++FunctionRedirects=(OldName="PlayerController.ClientPlayForceFeedback",NewName="PlayerController.K2_ClientPlayForceFeedback")
+
 [CoreUObject.Metadata]
 ; Note: UnrealHeaderTool should be rerun after making changes to MetadataRedirects to catch any keys specified in class headers
 +MetadataRedirects=(OldKey="K2Protected", NewKey="BlueprintProtected")
diff --git a/Engine/Source/Runtime/Engine/Classes/GameFramework/PlayerController.h b/Engine/Source/Runtime/Engine/Classes/GameFramework/PlayerController.h
--- a/Engine/Source/Runtime/Engine/Classes/GameFramework/PlayerController.h
+++ b/Engine/Source/Runtime/Engine/Classes/GameFramework/PlayerController.h
@@ -1035,37 +1035,33 @@ public:
 	 * @param	Tag						A tag that allows stopping of an effect.  If another effect with this Tag is playing, it will be stopped and replaced
 	 */
 	UFUNCTION(BlueprintCallable, Category="Game|Feedback", meta=(DisplayName="Client Play Force Feedback", AdvancedDisplay="bIgnoreTimeDilation,bPlayWhilePaused"))
-	void K2_ClientPlayForceFeedback(class UForceFeedbackEffect* ForceFeedbackEffect, FName Tag, bool bLooping, bool bIgnoreTimeDilation, bool bPlayWhilePaused)
-	{
-		FForceFeedbackParameters Params;
-		Params.Tag = Tag;
-		Params.bLooping = bLooping;
-		Params.bIgnoreTimeDilation = bIgnoreTimeDilation;
-		Params.bPlayWhilePaused = bPlayWhilePaused;
-		ClientPlayForceFeedback(ForceFeedbackEffect, Params);
-	}
+	void K2_ClientPlayForceFeedback(class UForceFeedbackEffect* ForceFeedbackEffect, FName Tag, bool bLooping, bool bIgnoreTimeDilation, bool bPlayWhilePaused);
 
+private:
+	/** 
+	 * Internal replicated version of client play force feedback event. 
+	 * Cannot be named ClientPlayForceFeedback as redirector for blueprint function version to K2_... does not work in that case
+	 */
 	UFUNCTION(unreliable, client)
-	void ClientPlayForceFeedback(class UForceFeedbackEffect* ForceFeedbackEffect, FForceFeedbackParameters Params = FForceFeedbackParameters());
+	void ClientPlayForceFeedback_Internal(class UForceFeedbackEffect* ForceFeedbackEffect, FForceFeedbackParameters Params = FForceFeedbackParameters());
 
-	UE_DEPRECATED(4.22, "Use version that specifies parameters using a struct instead of a list of parameters")
-	void ClientPlayForceFeedback(class UForceFeedbackEffect* ForceFeedbackEffect, bool bLooping, bool bIgnoreTimeDilation, FName Tag)
+public:
+
+	/** 
+	 * Play a force feedback pattern on the player's controller
+	 * @param	ForceFeedbackEffect		The force feedback pattern to play
+	 * @param	Params					Parameter struct to customize playback behavior of the feedback effect
+	 */
+	void ClientPlayForceFeedback(class UForceFeedbackEffect* ForceFeedbackEffect, FForceFeedbackParameters Params = FForceFeedbackParameters())
 	{
-		FForceFeedbackParameters Params;
-		Params.Tag = Tag;
-		Params.bLooping = bLooping;
-		Params.bIgnoreTimeDilation = bIgnoreTimeDilation;
-		ClientPlayForceFeedback(ForceFeedbackEffect, Params);
+		ClientPlayForceFeedback_Internal(ForceFeedbackEffect, Params);
 	}
 
+	UE_DEPRECATED(4.22, "Use version that specifies parameters using a struct instead of a list of parameters")
+	void ClientPlayForceFeedback(class UForceFeedbackEffect* ForceFeedbackEffect, bool bLooping, bool bIgnoreTimeDilation, FName Tag);
+
 	UE_DEPRECATED(4.18, "Use version that specifies parameters using a struct instead of a list of parameters")
-	void ClientPlayForceFeedback(class UForceFeedbackEffect* ForceFeedbackEffect, bool bLooping, FName Tag)
-	{
-		FForceFeedbackParameters Params;
-		Params.Tag = Tag;
-		Params.bLooping = bLooping;
-		ClientPlayForceFeedback(ForceFeedbackEffect, Params);
-	}
+	void ClientPlayForceFeedback(class UForceFeedbackEffect* ForceFeedbackEffect, bool bLooping, FName Tag);
 
 	/** 
 	 * Stops a playing force feedback pattern
diff --git a/Engine/Source/Runtime/Engine/Private/PlayerController.cpp b/Engine/Source/Runtime/Engine/Private/PlayerController.cpp
--- a/Engine/Source/Runtime/Engine/Private/PlayerController.cpp
+++ b/Engine/Source/Runtime/Engine/Private/PlayerController.cpp
@@ -3662,7 +3662,7 @@ void APlayerController::ClientPrestreamTextures_Implementation( AActor* ForcedAc
 	}
 }
 
-void APlayerController::ClientPlayForceFeedback_Implementation( UForceFeedbackEffect* ForceFeedbackEffect, FForceFeedbackParameters Params)
+void APlayerController::ClientPlayForceFeedback_Internal_Implementation( UForceFeedbackEffect* ForceFeedbackEffect, FForceFeedbackParameters Params)
 {
 	if (ForceFeedbackEffect)
 	{
@@ -3685,6 +3685,33 @@ void APlayerController::ClientPlayForceFeedback_Implementation( UForceFeedbackEf
 	}
 }
 
+void APlayerController::K2_ClientPlayForceFeedback(class UForceFeedbackEffect* ForceFeedbackEffect, FName Tag, bool bLooping, bool bIgnoreTimeDilation, bool bPlayWhilePaused)
+{
+	FForceFeedbackParameters Params;
+	Params.Tag = Tag;
+	Params.bLooping = bLooping;
+	Params.bIgnoreTimeDilation = bIgnoreTimeDilation;
+	Params.bPlayWhilePaused = bPlayWhilePaused;
+	ClientPlayForceFeedback(ForceFeedbackEffect, Params);
+}
+
+void APlayerController::ClientPlayForceFeedback(class UForceFeedbackEffect* ForceFeedbackEffect, bool bLooping, bool bIgnoreTimeDilation, FName Tag)
+{
+	FForceFeedbackParameters Params;
+	Params.Tag = Tag;
+	Params.bLooping = bLooping;
+	Params.bIgnoreTimeDilation = bIgnoreTimeDilation;
+	ClientPlayForceFeedback(ForceFeedbackEffect, Params);
+}
+
+void APlayerController::ClientPlayForceFeedback(class UForceFeedbackEffect* ForceFeedbackEffect, bool bLooping, FName Tag)
+{
+	FForceFeedbackParameters Params;
+	Params.Tag = Tag;
+	Params.bLooping = bLooping;
+	ClientPlayForceFeedback(ForceFeedbackEffect, Params);
+}
+
 void APlayerController::ClientStopForceFeedback_Implementation( UForceFeedbackEffect* ForceFeedbackEffect, FName Tag)
 {
 	if (ForceFeedbackEffect == NULL && Tag == NAME_None)
