(4.21.2) Don't check on OwningWorld as it is not relevant and totally valid to be null in this context.
#rnx

diff --git a/Engine/Source/Runtime/Engine/Private/LevelScriptBlueprint.cpp b/Engine/Source/Runtime/Engine/Private/LevelScriptBlueprint.cpp
--- a/Engine/Source/Runtime/Engine/Private/LevelScriptBlueprint.cpp
+++ b/Engine/Source/Runtime/Engine/Private/LevelScriptBlueprint.cpp
@@ -50,12 +50,11 @@ FString ULevelScriptBlueprint::GetFriendlyName() const
 	return UBlueprint::GetFriendlyName();
 }
 
-FString ULevelScriptBlueprint::CreateLevelScriptNameFromLevel (const ULevel* Level)
+FString ULevelScriptBlueprint::CreateLevelScriptNameFromLevel(const ULevel* Level)
 {
-	// Since all maps are named "PersistentLevel," check to see if this level is the actual OwningWorld->PersistentLevel, or name it based on the map package.
-	check(Level && Level->OwningWorld);
+	// Since all maps are named "PersistentLevel" the level script name is based on the LevelPackage
+	check(Level);
 	UObject* LevelPackage = Level->GetOutermost();
-
 	return FPackageName::GetShortName(LevelPackage->GetFName().GetPlainNameString());
 
 }
