Make sure the default vertex color is white and not red

diff --git a/Engine/Source/Runtime/Engine/Private/StaticMesh.cpp b/Engine/Source/Runtime/Engine/Private/StaticMesh.cpp
--- a/Engine/Source/Runtime/Engine/Private/StaticMesh.cpp
+++ b/Engine/Source/Runtime/Engine/Private/StaticMesh.cpp
@@ -2598,7 +2598,7 @@ void UStaticMesh::RegisterMeshAttributes( FMeshDescription& MeshDescription )
 	MeshDescription.VertexInstanceAttributes().RegisterAttribute<FVector>( MeshAttribute::VertexInstance::Normal, 1, FVector::ZeroVector, EMeshAttributeFlags::AutoGenerated );
 	MeshDescription.VertexInstanceAttributes().RegisterAttribute<FVector>( MeshAttribute::VertexInstance::Tangent, 1, FVector::ZeroVector, EMeshAttributeFlags::AutoGenerated );
 	MeshDescription.VertexInstanceAttributes().RegisterAttribute<float>( MeshAttribute::VertexInstance::BinormalSign, 1, 0.0f, EMeshAttributeFlags::AutoGenerated );
-	MeshDescription.VertexInstanceAttributes().RegisterAttribute<FVector4>( MeshAttribute::VertexInstance::Color, 1, FVector4( 1.0f ), EMeshAttributeFlags::Lerpable );
+	MeshDescription.VertexInstanceAttributes().RegisterAttribute<FVector4>( MeshAttribute::VertexInstance::Color, 1, FVector4( 1.0f, 1.0f, 1.0f, 1.0f ), EMeshAttributeFlags::Lerpable );
 
 	// Add basic edge attributes
 	MeshDescription.EdgeAttributes().RegisterAttribute<bool>( MeshAttribute::Edge::IsHard, 1, false );
