PR #5064: Take out redundant assign operation (Contributed by [REDACTED_PERSON])

diff --git a/Engine/Source/Runtime/Engine/Private/Components/BillboardComponent.cpp b/Engine/Source/Runtime/Engine/Private/Components/BillboardComponent.cpp
--- a/Engine/Source/Runtime/Engine/Private/Components/BillboardComponent.cpp
+++ b/Engine/Source/Runtime/Engine/Private/Components/BillboardComponent.cpp
@@ -297,7 +297,6 @@ UBillboardComponent::UBillboardComponent(const FObjectInitializer& ObjectInitial
 	bUseEditorCompositing = true;
 
 #if WITH_EDITORONLY_DATA
-	Sprite = ConstructorStatics.SpriteTexture.Object;
 	SpriteInfo.Category = ConstructorStatics.ID_Misc;
 	SpriteInfo.DisplayName = ConstructorStatics.NAME_Misc;
 	bUseInEditorScaling = true;
diff --git a/Engine/Source/Runtime/Engine/Private/TriggerBase.cpp b/Engine/Source/Runtime/Engine/Private/TriggerBase.cpp
--- a/Engine/Source/Runtime/Engine/Private/TriggerBase.cpp
+++ b/Engine/Source/Runtime/Engine/Private/TriggerBase.cpp
@@ -43,7 +43,6 @@ ATriggerBase::ATriggerBase(const FObjectInitializer& ObjectInitializer)
 		SpriteComponent->Sprite = ConstructorStatics.TriggerTextureObject.Get();
 		SpriteComponent->RelativeScale3D = FVector(0.5f, 0.5f, 0.5f);
 		SpriteComponent->bHiddenInGame = false;
-		SpriteComponent->Sprite = ConstructorStatics.TriggerTextureObject.Get();
 		SpriteComponent->SpriteInfo.Category = ConstructorStatics.ID_Triggers;
 		SpriteComponent->SpriteInfo.DisplayName = ConstructorStatics.NAME_Triggers;
 		SpriteComponent->bIsScreenSizeScaled = true;
