[NO NON-PII COMMIT MESSAGE]

diff --git a/Engine/Source/Runtime/Engine/Private/HierarchicalInstancedStaticMesh.cpp b/Engine/Source/Runtime/Engine/Private/HierarchicalInstancedStaticMesh.cpp
--- a/Engine/Source/Runtime/Engine/Private/HierarchicalInstancedStaticMesh.cpp
+++ b/Engine/Source/Runtime/Engine/Private/HierarchicalInstancedStaticMesh.cpp
@@ -2121,7 +2121,7 @@ bool UHierarchicalInstancedStaticMeshComponent::UpdateInstanceTransform(int32 In
 	// if we are only updating rotation/scale we update the instance directly in the cluster tree
 	const bool bIsOmittedInstance = (RenderIndex == INDEX_NONE);
 	const bool bIsBuiltInstance = !bIsOmittedInstance && RenderIndex < NumBuiltRenderInstances;
-	const bool bDoInPlaceUpdate = bIsBuiltInstance && NewLocalLocation.Equals(OldTransform.GetOrigin());
+	const bool bDoInPlaceUpdate = bIsBuiltInstance && NewLocalLocation.Equals(OldTransform.GetOrigin()) && (PerInstanceRenderData.IsValid() && PerInstanceRenderData->InstanceBuffer.RequireCPUAccess);
 
 	bool Result = Super::UpdateInstanceTransform(InstanceIndex, NewInstanceTransform, bWorldSpace, bMarkRenderStateDirty, bTeleport);
 	
@@ -2142,7 +2142,10 @@ bool UHierarchicalInstancedStaticMeshComponent::UpdateInstanceTransform(int32 In
 			{
 				BuiltInstanceBounds += NewInstanceBounds;
 
-				MarkRenderStateDirty();
+				if (bMarkRenderStateDirty)
+				{
+					MarkRenderStateDirty();
+				}
 			}
 		}
 		else
