Wrapped functions to query the available and most suitable languages for Blueprints

 - GetNativeCulture, wraps FTextLocalizationManager::GetNativeCultureName.
 - GetLocalizedCultures, wraps FTextLocalizationManager::GetLocalizedCultureNames.
 - GetSuitableCulture, uses FInternationalization::GetPrioritizedCultureNames to find the best match from an array of cultures.

diff --git a/Engine/Source/Runtime/Engine/Classes/Kismet/KismetInternationalizationLibrary.h b/Engine/Source/Runtime/Engine/Classes/Kismet/KismetInternationalizationLibrary.h
--- a/Engine/Source/Runtime/Engine/Classes/Kismet/KismetInternationalizationLibrary.h
+++ b/Engine/Source/Runtime/Engine/Classes/Kismet/KismetInternationalizationLibrary.h
@@ -8,6 +8,9 @@
 #include "Kismet/BlueprintFunctionLibrary.h"
 #include "KismetInternationalizationLibrary.generated.h"
 
+enum class ELocalizationLoadFlags : uint8;
+enum class ELocalizedTextSourceCategory : uint8;
+
 UCLASS(meta=(BlueprintThreadSafe, ScriptName="InternationalizationLibrary"))
 class ENGINE_API UKismetInternationalizationLibrary : public UBlueprintFunctionLibrary
 {
@@ -110,4 +113,42 @@ public:
 	 */
 	UFUNCTION(BlueprintCallable, Category="Utilities|Internationalization", meta=(AdvancedDisplay="1"))
 	static void ClearCurrentAssetGroupCulture(const FName AssetGroup, const bool SaveToConfig = false);
+
+	/**
+	 * Get the native culture for the given localization category.
+	 * @param Category The localization category to query.
+	 * @return The culture as an IETF language tag (eg, "zh-Hans-CN").
+	 */
+	UFUNCTION(BlueprintPure, Category="Utilities|Internationalization")
+	static FString GetNativeCulture(const ELocalizedTextSourceCategory TextCategory);
+
+	/**
+	 * Get the list of cultures that have localization data available for them.
+	 * @param IncludeGame Check for localized game resources.
+	 * @param IncludeEngine Check for localized engine resources.
+	 * @param IncludeEditor Check for localized editor resources.
+	 * @param IncludeAdditional Check for localized additional (eg, plugin) resources.
+	 * @return The list of cultures as IETF language tags (eg, "zh-Hans-CN").
+	 */
+	UFUNCTION(BlueprintPure, Category="Utilities|Internationalization")
+	static TArray<FString> GetLocalizedCultures(const bool IncludeGame = true, const bool IncludeEngine = false, const bool IncludeEditor = false, const bool IncludeAdditional = false);
+
+	/**
+	 * Get the list of cultures that have localization data available for them.
+	 * @param LoadFlags Controls which resource groups should be checked.
+	 * @return The list of cultures as IETF language tags (eg, "zh-Hans-CN").
+	 */
+	static TArray<FString> GetLocalizedCultures(const ELocalizationLoadFlags LoadFlags);
+
+	/**
+	 * Given a list of available cultures, try and find the most suitable culture from the list based on culture prioritization.
+	 *   eg) If your list was [en, fr, de] and the given culture was "en-US", this function would return "en".
+	 *   eg) If your list was [zh, ar, pl] and the given culture was "en-US", this function would return the fallback culture.
+	 * @param AvailableCultures List of available cultures to filter against (see GetLocalizedCultures).
+	 * @param CultureToMatch Culture to try and match (see GetCurrentLanguage).
+	 * @param FallbackCulture The culture to return if there is no suitable match in the list (typically your native culture, see GetNativeCulture).
+	 * @return The culture as an IETF language tag (eg, "zh-Hans-CN").
+	 */
+	UFUNCTION(BlueprintPure, Category="Utilities|Internationalization")
+	static FString GetSuitableCulture(const TArray<FString>& AvailableCultures, const FString& CultureToMatch, const FString& FallbackCulture = TEXT("en"));
 };
diff --git a/Engine/Source/Runtime/Engine/Private/KismetInternationalizationLibrary.cpp b/Engine/Source/Runtime/Engine/Private/KismetInternationalizationLibrary.cpp
--- a/Engine/Source/Runtime/Engine/Private/KismetInternationalizationLibrary.cpp
+++ b/Engine/Source/Runtime/Engine/Private/KismetInternationalizationLibrary.cpp
@@ -1,6 +1,7 @@
 // Copyright 1998-2018 Epic Games, Inc. All Rights Reserved.
 
 #include "Kismet/KismetInternationalizationLibrary.h"
+#include "Internationalization/TextLocalizationManager.h"
 #include "Internationalization/Internationalization.h"
 #include "Internationalization/Culture.h"
 #include "Misc/ConfigCacheIni.h"
@@ -122,4 +123,38 @@ void UKismetInternationalizationLibrary::ClearCurrentAssetGroupCulture(const FNa
 	}
 }
 
+FString UKismetInternationalizationLibrary::GetNativeCulture(const ELocalizedTextSourceCategory TextCategory)
+{
+	return FTextLocalizationManager::Get().GetNativeCultureName(TextCategory);
+}
+
+TArray<FString> UKismetInternationalizationLibrary::GetLocalizedCultures(const bool IncludeGame, const bool IncludeEngine, const bool IncludeEditor, const bool IncludeAdditional)
+{
+	ELocalizationLoadFlags LoadFlags = ELocalizationLoadFlags::None;
+	LoadFlags |= (IncludeGame ? ELocalizationLoadFlags::Game : ELocalizationLoadFlags::None);
+	LoadFlags |= (IncludeEngine ? ELocalizationLoadFlags::Engine : ELocalizationLoadFlags::None);
+	LoadFlags |= (IncludeEditor ? ELocalizationLoadFlags::Editor : ELocalizationLoadFlags::None);
+	LoadFlags |= (IncludeAdditional ? ELocalizationLoadFlags::Additional : ELocalizationLoadFlags::None);
+
+	return GetLocalizedCultures(LoadFlags);
+}
+
+TArray<FString> UKismetInternationalizationLibrary::GetLocalizedCultures(const ELocalizationLoadFlags LoadFlags)
+{
+	return FTextLocalizationManager::Get().GetLocalizedCultureNames(LoadFlags);
+}
+
+FString UKismetInternationalizationLibrary::GetSuitableCulture(const TArray<FString>& AvailableCultures, const FString& CultureToMatch, const FString& FallbackCulture)
+{
+	const TArray<FString> PrioritizedCulturesToMatch = FInternationalization::Get().GetPrioritizedCultureNames(CultureToMatch);
+	for (const FString& PrioritizedCulture : PrioritizedCulturesToMatch)
+	{
+		if (AvailableCultures.Contains(PrioritizedCulture))
+		{
+			return PrioritizedCulture;
+		}
+	}
+	return FallbackCulture;
+}
+
 #undef LOCTEXT_NAMESPACE
