Wrapped FCulture::GetDisplayName for Blueprints

diff --git a/Engine/Source/Runtime/Engine/Classes/Kismet/KismetInternationalizationLibrary.h b/Engine/Source/Runtime/Engine/Classes/Kismet/KismetInternationalizationLibrary.h
--- a/Engine/Source/Runtime/Engine/Classes/Kismet/KismetInternationalizationLibrary.h
+++ b/Engine/Source/Runtime/Engine/Classes/Kismet/KismetInternationalizationLibrary.h
@@ -151,4 +151,13 @@ public:
 	 */
 	UFUNCTION(BlueprintPure, Category="Utilities|Internationalization")
 	static FString GetSuitableCulture(const TArray<FString>& AvailableCultures, const FString& CultureToMatch, const FString& FallbackCulture = TEXT("en"));
+
+	/**
+	 * Get the display name of the given culture.
+	 * @param Culture The culture to get the display name of, as an IETF language tag (eg, "zh-Hans-CN")
+	 * @param Localized True to get the localized display name (the name of the culture in its own language), or False to get the display name in the current language.
+	 * @return The display name of the culture, or the given culture code on failure.
+	 */
+	UFUNCTION(BlueprintPure, Category="Utilities|Internationalization", meta=(AdvancedDisplay="1"))
+	static FString GetCultureDisplayName(const FString& Culture, const bool Localized = true);
 };
diff --git a/Engine/Source/Runtime/Engine/Private/KismetInternationalizationLibrary.cpp b/Engine/Source/Runtime/Engine/Private/KismetInternationalizationLibrary.cpp
--- a/Engine/Source/Runtime/Engine/Private/KismetInternationalizationLibrary.cpp
+++ b/Engine/Source/Runtime/Engine/Private/KismetInternationalizationLibrary.cpp
@@ -157,4 +157,16 @@ FString UKismetInternationalizationLibrary::GetSuitableCulture(const TArray<FStr
 	return FallbackCulture;
 }
 
+FString UKismetInternationalizationLibrary::GetCultureDisplayName(const FString& Culture, const bool Localized)
+{
+	const FCulturePtr CulturePtr = FInternationalization::Get().GetCulture(Culture);
+	if (CulturePtr.IsValid())
+	{
+		return Localized
+			? CulturePtr->GetDisplayName()
+			: CulturePtr->GetNativeName();
+	}
+	return Culture;
+}
+
 #undef LOCTEXT_NAMESPACE
