Fixed issue where FRuntimeCurveLinearColor::GetLinearColorValue would crash if an ExternalCurve was not set

diff --git a/Engine/Source/Runtime/Engine/Private/Curves/CurveLinearColor.cpp b/Engine/Source/Runtime/Engine/Private/Curves/CurveLinearColor.cpp
--- a/Engine/Source/Runtime/Engine/Private/Curves/CurveLinearColor.cpp
+++ b/Engine/Source/Runtime/Engine/Private/Curves/CurveLinearColor.cpp
@@ -9,7 +9,27 @@
 
 FLinearColor FRuntimeCurveLinearColor::GetLinearColorValue(float InTime) const
 {
-	return ExternalCurve->GetLinearColorValue(InTime);
+	if (ExternalCurve)
+	{
+		return ExternalCurve->GetLinearColorValue(InTime);
+	}
+
+	FLinearColor Result;
+	Result.R = ColorCurves[0].Eval(InTime);
+	Result.G = ColorCurves[1].Eval(InTime);
+	Result.B = ColorCurves[2].Eval(InTime);
+
+	// No alpha keys means alpha should be 1
+	if (ColorCurves[3].GetNumKeys() == 0)
+	{
+		Result.A = 1.0f;
+	}
+	else
+	{
+		Result.A = ColorCurves[3].Eval(InTime);
+	}
+
+	return Result;
 }
 
 UCurveLinearColor::UCurveLinearColor(const FObjectInitializer& ObjectInitializer)
