Fix for material functions input types being labeled incorrectly

diff --git a/Engine/Source/Runtime/Engine/Private/Materials/MaterialExpressions.cpp b/Engine/Source/Runtime/Engine/Private/Materials/MaterialExpressions.cpp
--- a/Engine/Source/Runtime/Engine/Private/Materials/MaterialExpressions.cpp
+++ b/Engine/Source/Runtime/Engine/Private/Materials/MaterialExpressions.cpp
@@ -10935,6 +10935,7 @@ static const TCHAR* GetInputTypeName(uint8 InputType)
 		TEXT("V4"),
 		TEXT("T2d"),
 		TEXT("TCube"),
+		TEXT("TVol"),
 		TEXT("B"),
 		TEXT("MA"),
 		TEXT("TExt")
