Merging //UE4/Dev-Main to Dev-Editor (//UE4/Dev-Editor)

diff --git a/Engine/Build/Build.version b/Engine/Build/Build.version
--- a/Engine/Build/Build.version
+++ b/Engine/Build/Build.version
@@ -1,6 +1,6 @@
 {
 	"MajorVersion": 4,
-	"MinorVersion": 21,
+	"MinorVersion": 22,
 	"PatchVersion": 0,
 	"Changelist": 0,
 	"CompatibleChangelist": 0,
diff --git a/Engine/Plugins/Runtime/GoogleARCore/Source/GoogleARCoreBase/Public/GoogleARCoreAugmentedImageDatabase.h b/Engine/Plugins/Runtime/GoogleARCore/Source/GoogleARCoreBase/Public/GoogleARCoreAugmentedImageDatabase.h
--- a/Engine/Plugins/Runtime/GoogleARCore/Source/GoogleARCoreBase/Public/GoogleARCoreAugmentedImageDatabase.h
+++ b/Engine/Plugins/Runtime/GoogleARCore/Source/GoogleARCoreBase/Public/GoogleARCoreAugmentedImageDatabase.h
@@ -34,7 +34,9 @@ struct FGoogleARCoreAugmentedImageDatabaseEntry
 	float Width;
 
 	FGoogleARCoreAugmentedImageDatabaseEntry()
-		: Width(0.0f) { }
+		: ImageAsset(nullptr)
+	    , Width(0.0f) 
+	{ }
 };
 
 UCLASS(BlueprintType)
diff --git a/Engine/Plugins/Runtime/OptionalMobileFeaturesBPLibrary/Source/OptionalMobileFeaturesBPLibrary/Private/OptionalMobileFeaturesBPLibrary.cpp b/Engine/Plugins/Runtime/OptionalMobileFeaturesBPLibrary/Source/OptionalMobileFeaturesBPLibrary/Private/OptionalMobileFeaturesBPLibrary.cpp
--- a/Engine/Plugins/Runtime/OptionalMobileFeaturesBPLibrary/Source/OptionalMobileFeaturesBPLibrary/Private/OptionalMobileFeaturesBPLibrary.cpp
+++ b/Engine/Plugins/Runtime/OptionalMobileFeaturesBPLibrary/Source/OptionalMobileFeaturesBPLibrary/Private/OptionalMobileFeaturesBPLibrary.cpp
@@ -18,7 +18,7 @@ int UOptionalMobileFeaturesBPLibrary::GetVolumeState()
 	int scaledVolume = (baseVolume * 100) / 15;
 	return scaledVolume;
 #elif PLATFORM_IOS
-	return FIOSPlatformMisc::GetAudioVolume();
+	return FPlatformMisc::GetDeviceVolume();
 #else
 	return 0;
 #endif
diff --git a/Engine/Source/Runtime/AudioMixer/Private/AudioMixerDevice.cpp b/Engine/Source/Runtime/AudioMixer/Private/AudioMixerDevice.cpp
--- a/Engine/Source/Runtime/AudioMixer/Private/AudioMixerDevice.cpp
+++ b/Engine/Source/Runtime/AudioMixer/Private/AudioMixerDevice.cpp
@@ -909,6 +909,10 @@ namespace Audio
 		}
 	}
 
+	void FMixerDevice::FlushAudioRenderingCommands()
+	{
+		SourceManager.FlushCommandQueue();
+	}
 
 	bool FMixerDevice::IsMasterSubmixType(USoundSubmix* InSubmix) const
 	{
diff --git a/Engine/Source/Runtime/AudioMixer/Private/AudioMixerSourceBuffer.cpp b/Engine/Source/Runtime/AudioMixer/Private/AudioMixerSourceBuffer.cpp
--- a/Engine/Source/Runtime/AudioMixer/Private/AudioMixerSourceBuffer.cpp
+++ b/Engine/Source/Runtime/AudioMixer/Private/AudioMixerSourceBuffer.cpp
@@ -105,9 +105,11 @@ namespace Audio
 	bool FMixerSourceBuffer::Init()
 	{
 		check(SoundWave);
-
-		// We flag that this sound wave is active for the lifetime of this object since we use it for decoding, etc.
-		SoundWave->IncrementNumSounds();
+		if (SoundWave->bProcedural && SoundWave->IsGenerating())
+		{
+			UE_LOG(LogAudioMixer, Warning, TEXT("Procedural sound wave is reinitializing even though it is currently actively generating audio. Please stop sound before trying to play it again."));
+			return false;
+		}
 
 		// We have successfully initialized which means our SoundWave has been flagged as bIsActive
 		// GC can run between PreInit and Init so when cleaning up FMixerSourceBuffer, we don't want to touch SoundWave unless bInitailized is true.
@@ -399,9 +401,9 @@ namespace Audio
 		{
 			if (SoundWave->bProcedural)
 			{
+				SoundWave->SetGenerating(true);
 				SoundWave->OnBeginGenerate();
 			}
-
 		}
 	}
 
@@ -416,9 +418,9 @@ namespace Audio
 			if (SoundWave->bProcedural)
 			{
 				SoundWave->OnEndGenerate();
+				SoundWave->SetGenerating(false);
 			}
 
-			SoundWave->DecrementNumSounds();
 			SoundWave = nullptr;
 			bInitialized = false;
 		}
@@ -430,7 +432,6 @@ namespace Audio
 			{
 				delete MixerBuffer;
 			}
-
 			MixerBuffer = nullptr;
 		}
 	}
diff --git a/Engine/Source/Runtime/AudioMixer/Private/AudioMixerSourceManager.cpp b/Engine/Source/Runtime/AudioMixer/Private/AudioMixerSourceManager.cpp
--- a/Engine/Source/Runtime/AudioMixer/Private/AudioMixerSourceManager.cpp
+++ b/Engine/Source/Runtime/AudioMixer/Private/AudioMixerSourceManager.cpp
@@ -2137,6 +2137,16 @@ namespace Audio
 		RenderThreadCommandBufferIndex.Set(!CurrentRenderThreadIndex);
 	}
 
+	void FMixerSourceManager::FlushCommandQueue()
+	{
+		bPumpQueue = true;
+		while (bPumpQueue)
+		{
+			FPlatformProcess::Sleep(0);
+		}
+		bPumpQueue = true;
+	}
+
 	void FMixerSourceManager::UpdatePendingReleaseData(bool bForceWait)
 	{
 		// Check for any pending delete procedural sound waves
diff --git a/Engine/Source/Runtime/AudioMixer/Private/AudioMixerSourceManager.h b/Engine/Source/Runtime/AudioMixer/Private/AudioMixerSourceManager.h
--- a/Engine/Source/Runtime/AudioMixer/Private/AudioMixerSourceManager.h
+++ b/Engine/Source/Runtime/AudioMixer/Private/AudioMixerSourceManager.h
@@ -269,7 +269,7 @@ namespace Audio
 		bool IsBus(const int32 SourceId) const;
 		void PumpCommandQueue();
 		void UpdatePendingReleaseData(bool bForceWait = false);
-
+		void FlushCommandQueue();
 	private:
 
 		void ReleaseSource(const int32 SourceId);
diff --git a/Engine/Source/Runtime/AudioMixer/Private/Components/SynthComponent.cpp b/Engine/Source/Runtime/AudioMixer/Private/Components/SynthComponent.cpp
--- a/Engine/Source/Runtime/AudioMixer/Private/Components/SynthComponent.cpp
+++ b/Engine/Source/Runtime/AudioMixer/Private/Components/SynthComponent.cpp
@@ -280,7 +280,7 @@ void USynthComponent::OnUnregister()
 bool USynthComponent::IsReadyForOwnerToAutoDestroy() const
 {
 	const bool bIsAudioComponentReadyForDestroy = !AudioComponent || (AudioComponent && !AudioComponent->IsPlaying());
-	const bool bIsSynthSoundReadyForDestroy = !Synth || !Synth->GetNumSoundsActive();
+	const bool bIsSynthSoundReadyForDestroy = !Synth || !Synth->IsGenerating();
 	return bIsAudioComponentReadyForDestroy && bIsSynthSoundReadyForDestroy;
 }
 
diff --git a/Engine/Source/Runtime/AudioMixer/Public/AudioMixerDevice.h b/Engine/Source/Runtime/AudioMixer/Public/AudioMixerDevice.h
--- a/Engine/Source/Runtime/AudioMixer/Public/AudioMixerDevice.h
+++ b/Engine/Source/Runtime/AudioMixer/Public/AudioMixerDevice.h
@@ -123,6 +123,8 @@ namespace Audio
 		// Submix buffer listener callbacks
 		virtual void RegisterSubmixBufferListener(ISubmixBufferListener* InSubmixBufferListener, USoundSubmix* InSubmix = nullptr) override;
 		virtual void UnregisterSubmixBufferListener(ISubmixBufferListener* InSubmixBufferListener, USoundSubmix* InSubmix = nullptr) override;
+
+		virtual void FlushAudioRenderingCommands() override;
 		//~ End FAudioDevice
 
 		//~ Begin IAudioMixer
diff --git a/Engine/Source/Runtime/AugmentedReality/Public/ARTypes.h b/Engine/Source/Runtime/AugmentedReality/Public/ARTypes.h
--- a/Engine/Source/Runtime/AugmentedReality/Public/ARTypes.h
+++ b/Engine/Source/Runtime/AugmentedReality/Public/ARTypes.h
@@ -326,7 +326,13 @@ struct AUGMENTEDREALITY_API FARVideoFormat
 	GENERATED_BODY()
 
 public:
-	FARVideoFormat() { }
+	FARVideoFormat()
+		: FPS(0)
+		, Width(0)
+		, Height(0)
+	{
+		
+	}
 	
 	FARVideoFormat(int32 InFPS, int32 InWidth, int32 InHeight)
 		: FPS(InFPS)
diff --git a/Engine/Source/Runtime/Engine/Classes/Sound/SoundWave.h b/Engine/Source/Runtime/Engine/Classes/Sound/SoundWave.h
--- a/Engine/Source/Runtime/Engine/Classes/Sound/SoundWave.h
+++ b/Engine/Source/Runtime/Engine/Classes/Sound/SoundWave.h
@@ -220,10 +220,7 @@ private:
 	/** What state the precache decompressor is in. */
 	FThreadSafeCounter PrecacheState;
 
-	/** Number of sounds actively using this sound wave by the audio renderer (in audio mixer). Prevents GC issues with rendering realtime audio. */
-	FThreadSafeCounter NumSoundsActive;
-	int32 GCTimeOutCount;
-
+	FThreadSafeBool bGenerating;
 #if !WITH_EDITOR
 	// This is the sample rate gotten from platform settings.
 	float CachedSampleRateOverride;
@@ -391,13 +388,8 @@ public:
 	// Called when the procedural sound wave is done generating on the render thread. Only used in the audio mixer and when bProcedural is true..
 	virtual void OnEndGenerate() {};
 
-	// Returns number of sounds using this sound wave (audio mixer only)
-	int32 GetNumSoundsActive();
-
-	// Increment and decrement num sounds (used in audio mixer)
-	void IncrementNumSounds();
-	void DecrementNumSounds();
-
+	bool IsGenerating() const { return bGenerating; }
+	void SetGenerating(bool bInGenerating) { bGenerating = bInGenerating; }
 	/**
 	* Overwrite sample rate. Used for procedural soundwaves, as well as sound waves that are resampled on compress/decompress.
 	*/
diff --git a/Engine/Source/Runtime/Engine/Private/Animation/PoseAsset.cpp b/Engine/Source/Runtime/Engine/Private/Animation/PoseAsset.cpp
--- a/Engine/Source/Runtime/Engine/Private/Animation/PoseAsset.cpp
+++ b/Engine/Source/Runtime/Engine/Private/Animation/PoseAsset.cpp
@@ -1454,5 +1454,4 @@ bool UPoseAsset::GetBasePoseTransform(TArray<FTransform>& OutBasePose, TArray<fl
 }
 #endif // WITH_EDITOR
 
-PRAGMA_ENABLE_OPTIMIZATION
 #undef LOCTEXT_NAMESPACE 
diff --git a/Engine/Source/Runtime/Engine/Private/AudioDevice.cpp b/Engine/Source/Runtime/Engine/Private/AudioDevice.cpp
--- a/Engine/Source/Runtime/Engine/Private/AudioDevice.cpp
+++ b/Engine/Source/Runtime/Engine/Private/AudioDevice.cpp
@@ -3833,7 +3833,7 @@ void FAudioDevice::AddNewActiveSound(const FActiveSound& NewActiveSound)
 	}
 
 	USoundWave* SoundWave = Cast<USoundWave>(Sound);
-	if (SoundWave && SoundWave->bProcedural && SoundWave->GetNumSoundsActive() > 0)
+	if (SoundWave && SoundWave->bProcedural && SoundWave->IsGenerating())
 	{
 		FString SoundWaveName;
 		SoundWave->GetName(SoundWaveName);
@@ -4668,6 +4668,9 @@ void FAudioDevice::Flush(UWorld* WorldToFlush, bool bClearActivatedReverb)
 	// Immediately stop all pending active sounds
 	ProcessingPendingActiveSoundStops(WorldToFlush == nullptr || WorldToFlush->bIsTearingDown);
 
+	// Make sure any in-flight audio rendering commands get executed.
+	FlushAudioRenderingCommands();
+
 	// Anytime we flush, make sure to clear all the listeners.  We'll get the right ones soon enough.
 	Listeners.Reset();
 	Listeners.Add(FListener(this));
diff --git a/Engine/Source/Runtime/Engine/Private/PhysicsEngine/Experimental/PhysInterface_Apeiron.cpp b/Engine/Source/Runtime/Engine/Private/PhysicsEngine/Experimental/PhysInterface_Apeiron.cpp
--- a/Engine/Source/Runtime/Engine/Private/PhysicsEngine/Experimental/PhysInterface_Apeiron.cpp
+++ b/Engine/Source/Runtime/Engine/Private/PhysicsEngine/Experimental/PhysInterface_Apeiron.cpp
@@ -778,6 +778,14 @@ FTransform FPhysInterface_Apeiron::GetComTransform_AssumesLocked(const FPhysicsA
     return GlobalTransform * ComTransform;
 }
 
+FTransform FPhysInterface_Apeiron::GetComTransformLocal_AssumesLocked(const FPhysicsActorReference_Apeiron& InActorReference)
+{
+	uint32 Index = UINT32_MAX;
+	const auto& LocalParticles = GetParticlesAndIndex(InActorReference, Index);
+	Apeiron::TRigidTransform<float, 3> ComTransform(LocalParticles.Geometry(Index) ? LocalParticles.Geometry(Index)->BoundingBox().Center() : Apeiron::TVector<float, 3>(0), Apeiron::TRotation<float, 3>(FQuat(0, 0, 0, 1)));
+	return ComTransform;
+}
+
 FVector FPhysInterface_Apeiron::GetLocalInertiaTensor_AssumesLocked(const FPhysicsActorReference_Apeiron& InActorReference)
 {
     // @todo(mlentine): Just return directly once we implement DiagonalMatrix
diff --git a/Engine/Source/Runtime/Engine/Private/PhysicsEngine/PhysicsReplication.cpp b/Engine/Source/Runtime/Engine/Private/PhysicsEngine/PhysicsReplication.cpp
--- a/Engine/Source/Runtime/Engine/Private/PhysicsEngine/PhysicsReplication.cpp
+++ b/Engine/Source/Runtime/Engine/Private/PhysicsEngine/PhysicsReplication.cpp
@@ -16,6 +16,7 @@
 #include "Engine/Classes/Engine/Player.h"
 #include "Physics/PhysicsInterfaceCore.h"
 #include "Physics/PhysScene_PhysX.h"
+#include "Components/SkeletalMeshComponent.h"
 
 #if WITH_PHYSX
 #include "PhysXPublic.h"
diff --git a/Engine/Source/Runtime/Engine/Private/SoundWave.cpp b/Engine/Source/Runtime/Engine/Private/SoundWave.cpp
--- a/Engine/Source/Runtime/Engine/Private/SoundWave.cpp
+++ b/Engine/Source/Runtime/Engine/Private/SoundWave.cpp
@@ -639,9 +639,6 @@ void USoundWave::BeginDestroy()
 	// the audio render thread stops the sound before GC hits.
 	bIsBeginDestroy = true;
 	
-	// Arbitrary count to wait before forcing IsFinishDestroy to return true. Temporary fallback to hunt a deadlock.
-	GCTimeOutCount = 10;
-
 #if WITH_EDITOR
 	// Flush any async results so we dont leak them in the DDC
 	if (GetDerivedDataCache() && AsyncLoadingDataFormats.Num() > 0)
@@ -859,22 +856,6 @@ FWaveInstance* USoundWave::HandleStart( FActiveSound& ActiveSound, const UPTRINT
 	return WaveInstance;
 }
 
-int32 USoundWave::GetNumSoundsActive()
-{
-	return NumSoundsActive.GetValue();
-}
-
-void USoundWave::IncrementNumSounds()
-{
-	NumSoundsActive.Increment();
-}
-
-void USoundWave::DecrementNumSounds()
-{
-	int32 NewValue = NumSoundsActive.Decrement();
-	check(NewValue >= 0);
-}
-
 bool USoundWave::IsReadyForFinishDestroy()
 {
 	const bool bIsStreamingInProgress = IStreamingManager::Get().GetAudioStreamingManager().IsStreamingInProgress(this);
@@ -894,23 +875,7 @@ bool USoundWave::IsReadyForFinishDestroy()
 			}, GET_STATID(STAT_AudioFreeResources));
 		}
 	
-	bool AllowFinishDestroy = (ResourceState == ESoundWaveResourceState::Freed);
-
-	if (AllowFinishDestroy)
-	{
-		AllowFinishDestroy = (NumSoundsActive.GetValue() == 0);
-		if (!AllowFinishDestroy && --GCTimeOutCount <= 0)
-		{
-			AllowFinishDestroy = true;
-
-			if (AllowFinishDestroy && NumSoundsActive.GetValue() > 0)
-			{
-				UE_LOG(LogAudio, Warning, TEXT("Sound wave '%s' destroyed while num sounds active was '%d'."), *GetName(), NumSoundsActive.GetValue());
-			}
-		}
-	}
-
-	return AllowFinishDestroy;
+	return !bGenerating && ResourceState == ESoundWaveResourceState::Freed;
 }
 
 
diff --git a/Engine/Source/Runtime/Engine/Private/StreamableManager.cpp b/Engine/Source/Runtime/Engine/Private/StreamableManager.cpp
--- a/Engine/Source/Runtime/Engine/Private/StreamableManager.cpp
+++ b/Engine/Source/Runtime/Engine/Private/StreamableManager.cpp
@@ -1353,4 +1353,3 @@ FSoftObjectPath FStreamableManager::HandleLoadedRedirector(UObjectRedirector* Lo
 	return NewPath;
 }
 
-PRAGMA_ENABLE_OPTIMIZATION
diff --git a/Engine/Source/Runtime/Engine/Public/Audio.h b/Engine/Source/Runtime/Engine/Public/Audio.h
--- a/Engine/Source/Runtime/Engine/Public/Audio.h
+++ b/Engine/Source/Runtime/Engine/Public/Audio.h
@@ -717,7 +717,7 @@ public:
 	ENGINE_API FSpatializationParams GetSpatializationParams();
 
 	/** Returns the contained sound buffer object. */
-	const FSoundBuffer* GetBuffer() const { return Buffer; }
+	virtual const FSoundBuffer* GetBuffer() const { return Buffer; }
 
 	/** Initializes any source effects for this sound source. */
 	virtual void InitializeSourceEffects(uint32 InEffectVoiceId)
diff --git a/Engine/Source/Runtime/Engine/Public/AudioDevice.h b/Engine/Source/Runtime/Engine/Public/AudioDevice.h
--- a/Engine/Source/Runtime/Engine/Public/AudioDevice.h
+++ b/Engine/Source/Runtime/Engine/Public/AudioDevice.h
@@ -582,6 +582,9 @@ public:
 	 */
 	void Flush(UWorld* WorldToFlush, bool bClearActivatedReverb = true);
 
+	/** Allows audio rendering command queue to flush during audio device flush. */
+	virtual void FlushAudioRenderingCommands() {}
+
 	/**
 	 * Stop any playing sounds that are using a particular SoundWave
 	 *
diff --git a/Engine/Source/Runtime/Engine/Public/GPUSkinVertexFactory.h b/Engine/Source/Runtime/Engine/Public/GPUSkinVertexFactory.h
--- a/Engine/Source/Runtime/Engine/Public/GPUSkinVertexFactory.h
+++ b/Engine/Source/Runtime/Engine/Public/GPUSkinVertexFactory.h
@@ -721,7 +721,7 @@ public:
 				Index = 1 - Index;
 			}
 
-			check(ClothSimulPositionNormalBuffer[Index].VertexBufferRHI.IsValid());
+			checkf(ClothSimulPositionNormalBuffer[Index].VertexBufferRHI.IsValid(), TEXT("Index: %i Buffer0: %s Frame0: %i Buffer1: %s Frame1: %i"), Index,  ClothSimulPositionNormalBuffer[0].VertexBufferRHI.IsValid() ? TEXT("true") : TEXT("false"), BufferFrameNumber[0], ClothSimulPositionNormalBuffer[1].VertexBufferRHI.IsValid() ? TEXT("true") : TEXT("false"), BufferFrameNumber[1]);
 			return ClothSimulPositionNormalBuffer[Index];
 		}
 		
diff --git a/Engine/Source/Runtime/Engine/Public/Physics/Experimental/PhysInterface_Apeiron.h b/Engine/Source/Runtime/Engine/Public/Physics/Experimental/PhysInterface_Apeiron.h
--- a/Engine/Source/Runtime/Engine/Public/Physics/Experimental/PhysInterface_Apeiron.h
+++ b/Engine/Source/Runtime/Engine/Public/Physics/Experimental/PhysInterface_Apeiron.h
@@ -313,6 +313,7 @@ public:
 	static ENGINE_API FVector GetWorldVelocityAtPoint_AssumesLocked(const FPhysicsActorReference_Apeiron& InActorReference, const FVector& InPoint);
 
 	static ENGINE_API FTransform GetComTransform_AssumesLocked(const FPhysicsActorReference_Apeiron& InActorReference);
+	static ENGINE_API FTransform GetComTransformLocal_AssumesLocked(const FPhysicsActorReference_Apeiron& InActorReference);
 
 	static ENGINE_API FVector GetLocalInertiaTensor_AssumesLocked(const FPhysicsActorReference_Apeiron& InActorReference);
 	static ENGINE_API FBox GetBounds_AssumesLocked(const FPhysicsActorReference_Apeiron& InActorReference);
diff --git a/Engine/Source/Runtime/Launch/Resources/Version.h b/Engine/Source/Runtime/Launch/Resources/Version.h
--- a/Engine/Source/Runtime/Launch/Resources/Version.h
+++ b/Engine/Source/Runtime/Launch/Resources/Version.h
@@ -55,7 +55,7 @@
 // These numbers define the banner UE4 version, and are the most significant numbers when ordering two engine versions (that is, a 4.12.* version is always 
 // newer than a 4.11.* version, regardless of the changelist that it was built with)
 #define ENGINE_MAJOR_VERSION	4
-#define ENGINE_MINOR_VERSION	21
+#define ENGINE_MINOR_VERSION	22
 #define ENGINE_PATCH_VERSION	0
 
 // Macros for encoding strings
diff --git a/Engine/Source/Runtime/LiveLinkInterface/LiveLinkInterface.Build.cs b/Engine/Source/Runtime/LiveLinkInterface/LiveLinkInterface.Build.cs
--- a/Engine/Source/Runtime/LiveLinkInterface/LiveLinkInterface.Build.cs
+++ b/Engine/Source/Runtime/LiveLinkInterface/LiveLinkInterface.Build.cs
@@ -9,8 +9,7 @@ public class LiveLinkInterface : ModuleRules
         PublicDependencyModuleNames.AddRange(
             new string[] {
                 "Core",
-                "CoreUObject",
-				"TimeManagement"
+                "CoreUObject"
             }
         );
 	}
