Fix the merge instancing, Material was wrong when SectionInfoMap swap materials in the source mesh of the instancing

diff --git a/Engine/Source/Runtime/Engine/Private/Components/StaticMeshComponent.cpp b/Engine/Source/Runtime/Engine/Private/Components/StaticMeshComponent.cpp
--- a/Engine/Source/Runtime/Engine/Private/Components/StaticMeshComponent.cpp
+++ b/Engine/Source/Runtime/Engine/Private/Components/StaticMeshComponent.cpp
@@ -2205,13 +2205,35 @@ void UStaticMeshComponent::GetUsedMaterials(TArray<UMaterialInterface*>& OutMate
 {
 	if( GetStaticMesh() && GetStaticMesh()->RenderData )
 	{
+		TMap<int32, UMaterialInterface*> MapOfMaterials;
 		for (int32 LODIndex = 0; LODIndex < GetStaticMesh()->RenderData->LODResources.Num(); LODIndex++)
 		{
 			FStaticMeshLODResources& LODResources = GetStaticMesh()->RenderData->LODResources[LODIndex];
 			for (int32 SectionIndex = 0; SectionIndex < LODResources.Sections.Num(); SectionIndex++)
 			{
 				// Get the material for each element at the current lod index
-				OutMaterials.AddUnique(GetMaterial(LODResources.Sections[SectionIndex].MaterialIndex));
+				int32 MaterialIndex = LODResources.Sections[SectionIndex].MaterialIndex;
+				if (!MapOfMaterials.Contains(MaterialIndex))
+				{
+					MapOfMaterials.Add(MaterialIndex, GetMaterial(MaterialIndex));
+				}
+			}
+		}
+		if (MapOfMaterials.Num() > 0)
+		{
+			//We need to output the material in the correct order (follow the material index)
+			//So we sort the map with the material index
+			MapOfMaterials.KeySort([](int32 A, int32 B) {
+				return A < B; // sort keys in order
+			});
+
+			//Preadd all the material item in the array
+			OutMaterials.AddZeroed(MapOfMaterials.Num());
+			//Set the value in the correct order
+			int32 MaterialIndex = 0;
+			for (auto Kvp : MapOfMaterials)
+			{
+				OutMaterials[MaterialIndex++] = Kvp.Value;
 			}
 		}
 	}
