Fix for stack memory overwrite.

[REDACTED_HANDLE]: This removes a bit of logic added in CL 4272387 for the Oculus 1.28 Plugin changes not sure if this will be problematic.

diff --git a/Engine/Source/Runtime/Engine/Private/Debug/DebugDrawService.cpp b/Engine/Source/Runtime/Engine/Private/Debug/DebugDrawService.cpp
--- a/Engine/Source/Runtime/Engine/Private/Debug/DebugDrawService.cpp
+++ b/Engine/Source/Runtime/Engine/Private/Debug/DebugDrawService.cpp
@@ -64,10 +64,8 @@ void UDebugDrawService::Draw(const FEngineShowFlags Flags, FViewport* Viewport,
 		CanvasObject->AddToRoot();
 	}
 
-	if ((Canvas && GEngine->XRSystem) && !(Canvas->IsStereoRendering() && GEngine->XRSystem.IsValid() && GEngine->XRSystem->GetHMDDevice()))
-	{
-		CanvasObject->Init(View->UnscaledViewRect.Width(), View->UnscaledViewRect.Height(), View, Canvas);
-	}
+	// Canvas must be initialize every draw because the FCanvas passed in is on the stack in some senarioes.
+	CanvasObject->Init(View->UnscaledViewRect.Width(), View->UnscaledViewRect.Height(), View, Canvas);
 
 	CanvasObject->Update();	
 	CanvasObject->SetView(View);
