Fixing texture coord and property nodes not showing up in search

diff --git a/Engine/Source/Runtime/Engine/Classes/Materials/MaterialExpressionTextureCoordinate.h b/Engine/Source/Runtime/Engine/Classes/Materials/MaterialExpressionTextureCoordinate.h
--- a/Engine/Source/Runtime/Engine/Classes/Materials/MaterialExpressionTextureCoordinate.h
+++ b/Engine/Source/Runtime/Engine/Classes/Materials/MaterialExpressionTextureCoordinate.h
@@ -39,6 +39,7 @@ class UMaterialExpressionTextureCoordinate : public UMaterialExpression
 	virtual int32 Compile(class FMaterialCompiler* Compiler, int32 OutputIndex) override;
 	virtual void GetCaption(TArray<FString>& OutCaptions) const override;
 	virtual FText GetKeywords() const override { return FText::FromString(TEXT("texcoord")); }
+	virtual bool MatchesSearchQuery(const TCHAR* SearchQuery) override;
 #endif
 	//~ End UMaterialExpression Interface
 };
diff --git a/Engine/Source/Runtime/Engine/Classes/Materials/MaterialExpressionTextureProperty.h b/Engine/Source/Runtime/Engine/Classes/Materials/MaterialExpressionTextureProperty.h
--- a/Engine/Source/Runtime/Engine/Classes/Materials/MaterialExpressionTextureProperty.h
+++ b/Engine/Source/Runtime/Engine/Classes/Materials/MaterialExpressionTextureProperty.h
@@ -42,6 +42,7 @@ class ENGINE_API UMaterialExpressionTextureProperty : public UMaterialExpression
 	virtual void GetTexturesForceMaterialRecompile(TArray<UTexture *> &Textures) const override;
 	virtual void GetCaption(TArray<FString>& OutCaptions) const override;
 	virtual uint32 GetInputType(int32 InputIndex) override;
+	virtual bool MatchesSearchQuery(const TCHAR* SearchQuery) override;
 #endif //WITH_EDITOR
 	//~ End UMaterialExpression Interface
 };
diff --git a/Engine/Source/Runtime/Engine/Private/Materials/MaterialExpressions.cpp b/Engine/Source/Runtime/Engine/Private/Materials/MaterialExpressions.cpp
--- a/Engine/Source/Runtime/Engine/Private/Materials/MaterialExpressions.cpp
+++ b/Engine/Source/Runtime/Engine/Private/Materials/MaterialExpressions.cpp
@@ -2201,6 +2201,23 @@ uint32 UMaterialExpressionTextureProperty::GetInputType(int32 InputIndex)
 	return MCT_Unknown;
 }
 #undef IF_INPUT_RETURN
+
+
+bool UMaterialExpressionTextureProperty::MatchesSearchQuery(const TCHAR* SearchQuery)
+{
+	TArray<FString> Captions;
+	GetCaption(Captions);
+	for (const FString Caption : Captions)
+	{
+		if (Caption.Contains(SearchQuery))
+		{
+			return true;
+		}
+	}
+
+	return Super::MatchesSearchQuery(SearchQuery);
+}
+
 #endif
 
 //
@@ -3231,6 +3248,22 @@ void UMaterialExpressionTextureCoordinate::GetCaption(TArray<FString>& OutCaptio
 {
 	OutCaptions.Add(FString::Printf(TEXT("TexCoord[%i]"), CoordinateIndex));
 }
+
+
+bool UMaterialExpressionTextureCoordinate::MatchesSearchQuery(const TCHAR* SearchQuery)
+{
+	TArray<FString> Captions;
+	GetCaption(Captions);
+	for (const FString Caption : Captions)
+	{
+		if (Caption.Contains(SearchQuery))
+		{
+			return true;
+		}
+	}
+
+	return Super::MatchesSearchQuery(SearchQuery);
+}
 #endif // WITH_EDITOR
 
 UMaterialExpressionDotProduct::UMaterialExpressionDotProduct(const FObjectInitializer& ObjectInitializer)
