Deprecate FRawMesh - Convert welding threshold to use FMeshDescription

diff --git a/Engine/Source/Developer/MeshDescriptionOperations/Private/MeshDescriptionOperations.cpp b/Engine/Source/Developer/MeshDescriptionOperations/Private/MeshDescriptionOperations.cpp
--- a/Engine/Source/Developer/MeshDescriptionOperations/Private/MeshDescriptionOperations.cpp
+++ b/Engine/Source/Developer/MeshDescriptionOperations/Private/MeshDescriptionOperations.cpp
@@ -1941,4 +1941,49 @@ bool FMeshDescriptionOperations::HasVertexColor(const FMeshDescription& MeshDesc
 	return bHasVertexColor;
 }
 
+void FMeshDescriptionOperations::BuildWeldedVertexIDRemap(const FMeshDescription& MeshDescription, const float WeldingThreshold, TMap<FVertexID, FVertexID>& OutVertexIDRemap)
+{
+	TVertexAttributesConstRef<FVector> VertexPositions = MeshDescription.VertexAttributes().GetAttributesRef<FVector>(MeshAttribute::Vertex::Position);
+
+	int32 NumVertex = MeshDescription.Vertices().Num();
+	OutVertexIDRemap.Reserve(NumVertex);
+
+	// Create a list of vertex Z/index pairs
+	TArray<MeshDescriptionOperationNamespace::FIndexAndZ> VertIndexAndZ;
+	VertIndexAndZ.Reserve(NumVertex);
+
+	for (const FVertexID VertexID : MeshDescription.Vertices().GetElementIDs())
+	{
+		new(VertIndexAndZ)MeshDescriptionOperationNamespace::FIndexAndZ(VertexID.GetValue(), VertexPositions[VertexID]);
+	}
+
+	// Sort the vertices by z value
+	VertIndexAndZ.Sort(MeshDescriptionOperationNamespace::FCompareIndexAndZ());
+
+	// Search for duplicates, quickly!
+	for (int32 i = 0; i < VertIndexAndZ.Num(); i++)
+	{
+		FVertexID Index_i = FVertexID(VertIndexAndZ[i].Index);
+		if (OutVertexIDRemap.Contains(Index_i))
+		{
+			continue;
+		}
+		OutVertexIDRemap.FindOrAdd(Index_i) = Index_i;
+		// only need to search forward, since we add pairs both ways
+		for (int32 j = i + 1; j < VertIndexAndZ.Num(); j++)
+		{
+			if (FMath::Abs(VertIndexAndZ[j].Z - VertIndexAndZ[i].Z) > WeldingThreshold)
+				break; // can't be any more dups
+
+			const FVector& PositionA = *(VertIndexAndZ[i].OriginalVector);
+			const FVector& PositionB = *(VertIndexAndZ[j].OriginalVector);
+
+			if (PositionA.Equals(PositionB, WeldingThreshold))
+			{
+				OutVertexIDRemap.FindOrAdd(FVertexID(VertIndexAndZ[j].Index)) = Index_i;
+			}
+		}
+	}
+}
+
 #undef LOCTEXT_NAMESPACE
diff --git a/Engine/Source/Developer/MeshDescriptionOperations/Public/MeshDescriptionOperations.h b/Engine/Source/Developer/MeshDescriptionOperations/Public/MeshDescriptionOperations.h
--- a/Engine/Source/Developer/MeshDescriptionOperations/Public/MeshDescriptionOperations.h
+++ b/Engine/Source/Developer/MeshDescriptionOperations/Public/MeshDescriptionOperations.h
@@ -10,6 +10,7 @@ struct FUVMapSettings;
 struct FOverlappingCorners;
 enum class ELightmapUVVersion : int32;
 struct FPolygonGroupID;
+struct FVertexID;
 
 DECLARE_LOG_CATEGORY_EXTERN(LogMeshDescriptionOperations, Log, All);
 
@@ -118,4 +119,6 @@ public:
 	static void ConvertSmoothGroupToHardEdges(const TArray<uint32>& FaceSmoothingMasks, FMeshDescription& DestinationMeshDescription);
 
 	static bool HasVertexColor(const FMeshDescription& MeshDescription);
+
+	static void BuildWeldedVertexIDRemap(const FMeshDescription& MeshDescription, const float WeldingThreshold, TMap<FVertexID, FVertexID>& OutVertexIDRemap);
 };
\ No newline at end of file
diff --git a/Engine/Source/Developer/MeshSimplifier/Private/QuadricMeshReduction.cpp b/Engine/Source/Developer/MeshSimplifier/Private/QuadricMeshReduction.cpp
--- a/Engine/Source/Developer/MeshSimplifier/Private/QuadricMeshReduction.cpp
+++ b/Engine/Source/Developer/MeshSimplifier/Private/QuadricMeshReduction.cpp
@@ -5,7 +5,6 @@
 #include "Templates/ScopedPointer.h"
 #include "MeshUtilities.h"
 #include "MeshBuild.h"
-#include "RawMesh.h"
 #include "MeshSimplify.h"
 #include "OverlappingCorners.h"
 #include "Templates/UniquePtr.h"
@@ -14,6 +13,8 @@
 #include "MeshDescription.h"
 #include "MeshAttributes.h"
 #include "RenderUtils.h"
+#include "Engine/StaticMesh.h"
+#include "MeshDescriptionOperations.h"
 
 class FQuadricSimplifierMeshReductionModule : public IMeshReductionModule
 {
@@ -190,6 +191,14 @@ public:
 
 		const uint32 NumTexCoords = MAX_STATIC_TEXCOORDS;
 		int32 InMeshNumTexCoords = 1;
+		
+		TMap<FVertexID, FVertexID> VertexIDRemap;
+
+		bool bWeldVertices = ReductionSettings.WeldingThreshold > 0.0f;
+		if (bWeldVertices)
+		{
+			FMeshDescriptionOperations::BuildWeldedVertexIDRemap(InMesh, ReductionSettings.WeldingThreshold, VertexIDRemap);
+		}
 
 		TArray< TVertSimp< NumTexCoords > >	Verts;
 		TArray< uint32 >					Indexes;
@@ -202,14 +211,14 @@ public:
 			NumFaces += InMesh.GetPolygonTriangles(PolygonID).Num();
 		}
 		int32 NumWedges = NumFaces * 3;
-
-		TVertexAttributesConstRef<FVector> InVertexPositions = InMesh.VertexAttributes().GetAttributesRef<FVector>(MeshAttribute::Vertex::Position);
-		TVertexInstanceAttributesConstRef<FVector> InVertexNormals = InMesh.VertexInstanceAttributes().GetAttributesRef<FVector>(MeshAttribute::VertexInstance::Normal);
-		TVertexInstanceAttributesConstRef<FVector> InVertexTangents = InMesh.VertexInstanceAttributes().GetAttributesRef<FVector>(MeshAttribute::VertexInstance::Tangent);
-		TVertexInstanceAttributesConstRef<float> InVertexBinormalSigns = InMesh.VertexInstanceAttributes().GetAttributesRef<float>(MeshAttribute::VertexInstance::BinormalSign);
-		TVertexInstanceAttributesConstRef<FVector4> InVertexColors = InMesh.VertexInstanceAttributes().GetAttributesRef<FVector4>(MeshAttribute::VertexInstance::Color);
-		TVertexInstanceAttributesConstRef<FVector2D> InVertexUVs = InMesh.VertexInstanceAttributes().GetAttributesRef<FVector2D>(MeshAttribute::VertexInstance::TextureCoordinate);
-		TPolygonGroupAttributesConstRef<FName> InPolygonGroupMaterialNames = InMesh.PolygonGroupAttributes().GetAttributesRef<FName>(MeshAttribute::PolygonGroup::ImportedMaterialSlotName);
+		FStaticMeshDescriptionConstAttributeGetter InMeshAttribute(&InMesh);
+		TVertexAttributesConstRef<FVector> InVertexPositions = InMeshAttribute.GetPositions();
+		TVertexInstanceAttributesConstRef<FVector> InVertexNormals = InMeshAttribute.GetNormals();
+		TVertexInstanceAttributesConstRef<FVector> InVertexTangents = InMeshAttribute.GetTangents();
+		TVertexInstanceAttributesConstRef<float> InVertexBinormalSigns = InMeshAttribute.GetBinormalSigns();
+		TVertexInstanceAttributesConstRef<FVector4> InVertexColors = InMeshAttribute.GetColors();
+		TVertexInstanceAttributesConstRef<FVector2D> InVertexUVs = InMeshAttribute.GetUVs();
+		TPolygonGroupAttributesConstRef<FName> InPolygonGroupMaterialNames = InMeshAttribute.GetPolygonGroupImportedMaterialSlotNames();
 
 		TPolygonGroupAttributesRef<FName> OutPolygonGroupMaterialNames = OutReducedMesh.PolygonGroupAttributes().GetAttributesRef<FName>(MeshAttribute::PolygonGroup::ImportedMaterialSlotName);
 
@@ -230,7 +239,8 @@ public:
 				for (int32 CornerIndex = 0; CornerIndex < 3; ++CornerIndex)
 				{
 					VertexInstanceIDs[CornerIndex] = MeshTriangle.GetVertexInstanceID(CornerIndex);
-					VertexIDs[CornerIndex] = InMesh.GetVertexInstanceVertex(VertexInstanceIDs[CornerIndex]);
+					FVertexID TmpVertexID = InMesh.GetVertexInstanceVertex(VertexInstanceIDs[CornerIndex]);
+					VertexIDs[CornerIndex] = bWeldVertices ? VertexIDRemap[TmpVertexID] : TmpVertexID;
 					Positions[CornerIndex] = InVertexPositions[VertexIDs[CornerIndex]];
 				}
 
@@ -567,54 +577,8 @@ public:
 	{
 		return new FQuadricSimplifierMeshReduction;
 	}
-private:
-	void WeldVertexPositions(const FRawMesh& InMesh, const float WeldingThreshold, TArray<FVector>& OutVertexPositions, TArray<uint32>& OutIndices)
-	{
-		//The remap use to fix the indices after welding the vertex position buffer
-		TArray<int32> VertexRemap;
-		//Initialize some arrays
-		VertexRemap.AddZeroed(InMesh.VertexPositions.Num());
-		for (int32 VertexIndexRef = 0; VertexIndexRef < InMesh.VertexPositions.Num(); ++VertexIndexRef)
-		{
-			VertexRemap[VertexIndexRef] = INDEX_NONE;
-		}
-		OutVertexPositions.Reserve(InMesh.VertexPositions.Num());
-		//Weld overlapping vertex position
-		for (int32 VertexIndexRef = 0; VertexIndexRef < InMesh.VertexPositions.Num(); ++VertexIndexRef)
-		{
-			//Skip already remap vertex
-			if (VertexRemap[VertexIndexRef] != INDEX_NONE)
-			{
-				continue;
-			}
-			const FVector& PositionA = InMesh.VertexPositions[VertexIndexRef];
-			//Add this vertex to the new vertex buffer
-			VertexRemap[VertexIndexRef] = OutVertexPositions.Add(InMesh.VertexPositions[VertexIndexRef]);
-			//Find vertex to weld, search forward VertexIndexRef
-			for (int32 VertexIndex = VertexIndexRef + 1; VertexIndex < InMesh.VertexPositions.Num(); ++VertexIndex)
-			{
-				//skip already remap vertex
-				if (VertexRemap[VertexIndex] != INDEX_NONE)
-				{
-					continue;
-				}
-				const FVector& PositionB = InMesh.VertexPositions[VertexIndex];
-				if (PositionA.Equals(PositionB, WeldingThreshold))
-				{
-					//Remap this vertex to the "reference remapped vertex"
-					VertexRemap[VertexIndex] = VertexRemap[VertexIndexRef];
-				}
-			}
-		}
-		//Remap the indices to the new vertex position buffer
-		OutIndices.AddZeroed(InMesh.WedgeIndices.Num());
-		for (int32 WedgeIndex = 0; WedgeIndex < InMesh.WedgeIndices.Num(); ++WedgeIndex)
-		{
-			int32 VertexIndex = InMesh.WedgeIndices[WedgeIndex];
-			OutIndices[WedgeIndex] = VertexRemap[VertexIndex] == INDEX_NONE ? VertexIndex : VertexRemap[VertexIndex];
-		}
-	}
 };
+
 TUniquePtr<FQuadricSimplifierMeshReduction> GQuadricSimplifierMeshReduction;
 
 void FQuadricSimplifierMeshReductionModule::StartupModule()
diff --git a/Engine/Source/Developer/MeshSimplifier/QuadricMeshReduction.Build.cs b/Engine/Source/Developer/MeshSimplifier/QuadricMeshReduction.Build.cs
--- a/Engine/Source/Developer/MeshSimplifier/QuadricMeshReduction.Build.cs
+++ b/Engine/Source/Developer/MeshSimplifier/QuadricMeshReduction.Build.cs
@@ -11,8 +11,8 @@ public class QuadricMeshReduction : ModuleRules
 		PrivateDependencyModuleNames.Add("CoreUObject");
 		PrivateDependencyModuleNames.Add("Engine");
 		PrivateDependencyModuleNames.Add("RenderCore");
-		PrivateDependencyModuleNames.Add("RawMesh");
         PrivateDependencyModuleNames.Add("MeshDescription");
+        PrivateDependencyModuleNames.Add("MeshDescriptionOperations");
         PrivateDependencyModuleNames.Add("MeshUtilitiesCommon");
 
         PrivateIncludePathModuleNames.AddRange(
diff --git a/Engine/Source/Runtime/Engine/Classes/Engine/StaticMesh.h b/Engine/Source/Runtime/Engine/Classes/Engine/StaticMesh.h
--- a/Engine/Source/Runtime/Engine/Classes/Engine/StaticMesh.h
+++ b/Engine/Source/Runtime/Engine/Classes/Engine/StaticMesh.h
@@ -449,40 +449,67 @@ struct FMaterialRemapIndex
 	TArray<int32> MaterialRemap;
 };
 
+struct FStaticMeshDescriptionConstAttributeGetter
+{
+	ENGINE_API FStaticMeshDescriptionConstAttributeGetter(const FMeshDescription* InMeshDescription)
+		: MeshDescription(InMeshDescription)
+	{}
+
+	const FMeshDescription* MeshDescription;
+
+	ENGINE_API TVertexAttributesConstRef<FVector> GetPositions() const { return MeshDescription->VertexAttributes().GetAttributesRef<FVector>(MeshAttribute::Vertex::Position); }
+
+	ENGINE_API TVertexInstanceAttributesConstRef<FVector> GetNormals() const { return MeshDescription->VertexInstanceAttributes().GetAttributesRef<FVector>(MeshAttribute::VertexInstance::Normal); }
+
+	ENGINE_API TVertexInstanceAttributesConstRef<FVector> GetTangents() const { return MeshDescription->VertexInstanceAttributes().GetAttributesRef<FVector>(MeshAttribute::VertexInstance::Tangent); }
+
+	ENGINE_API TVertexInstanceAttributesConstRef<float> GetBinormalSigns() const { return MeshDescription->VertexInstanceAttributes().GetAttributesRef<float>(MeshAttribute::VertexInstance::BinormalSign); }
+
+	ENGINE_API TVertexInstanceAttributesConstRef<FVector4> GetColors() const { return MeshDescription->VertexInstanceAttributes().GetAttributesRef<FVector4>(MeshAttribute::VertexInstance::Color); }
+
+	ENGINE_API TVertexInstanceAttributesConstRef<FVector2D> GetUVs() const { return MeshDescription->VertexInstanceAttributes().GetAttributesRef<FVector2D>(MeshAttribute::VertexInstance::TextureCoordinate); }
+
+	ENGINE_API TEdgeAttributesConstRef<bool> GetEdgeHardnesses() const { return MeshDescription->EdgeAttributes().GetAttributesRef<bool>(MeshAttribute::Edge::IsHard); }
+
+	ENGINE_API TEdgeAttributesConstRef<float> GetEdgeCreaseSharpnesses() const { return MeshDescription->EdgeAttributes().GetAttributesRef<float>(MeshAttribute::Edge::CreaseSharpness); }
+
+	ENGINE_API TPolygonGroupAttributesConstRef<FName> GetPolygonGroupImportedMaterialSlotNames() { return MeshDescription->PolygonGroupAttributes().GetAttributesRef<FName>(MeshAttribute::PolygonGroup::ImportedMaterialSlotName); }
+};
+
 struct FStaticMeshDescriptionAttributeGetter
 {
 	ENGINE_API FStaticMeshDescriptionAttributeGetter(FMeshDescription* InMeshDescription)
 		: MeshDescription(InMeshDescription)
 	{}
-
+	
 	FMeshDescription* MeshDescription;
 
-	TVertexAttributesRef<FVector> GetPositions() const { return MeshDescription->VertexAttributes().GetAttributesRef<FVector>(MeshAttribute::Vertex::Position); }
-	TVertexAttributesConstRef<FVector> GetPositionsConst() const { return MeshDescription->VertexAttributes().GetAttributesRef<FVector>(MeshAttribute::Vertex::Position); }
+	ENGINE_API TVertexAttributesRef<FVector> GetPositions() const { return MeshDescription->VertexAttributes().GetAttributesRef<FVector>(MeshAttribute::Vertex::Position); }
+	ENGINE_API TVertexAttributesConstRef<FVector> GetPositionsConst() const { return MeshDescription->VertexAttributes().GetAttributesRef<FVector>(MeshAttribute::Vertex::Position); }
 
-	TVertexInstanceAttributesRef<FVector> GetNormals() const { return MeshDescription->VertexInstanceAttributes().GetAttributesRef<FVector>(MeshAttribute::VertexInstance::Normal); }
-	TVertexInstanceAttributesConstRef<FVector> GetNormalsConst() const { return MeshDescription->VertexInstanceAttributes().GetAttributesRef<FVector>(MeshAttribute::VertexInstance::Normal); }
+	ENGINE_API TVertexInstanceAttributesRef<FVector> GetNormals() const { return MeshDescription->VertexInstanceAttributes().GetAttributesRef<FVector>(MeshAttribute::VertexInstance::Normal); }
+	ENGINE_API TVertexInstanceAttributesConstRef<FVector> GetNormalsConst() const { return MeshDescription->VertexInstanceAttributes().GetAttributesRef<FVector>(MeshAttribute::VertexInstance::Normal); }
 
-	TVertexInstanceAttributesRef<FVector> GetTangents() const { return MeshDescription->VertexInstanceAttributes().GetAttributesRef<FVector>(MeshAttribute::VertexInstance::Tangent); }
-	TVertexInstanceAttributesConstRef<FVector> GetTangentsConst() const { return MeshDescription->VertexInstanceAttributes().GetAttributesRef<FVector>(MeshAttribute::VertexInstance::Tangent); }
+	ENGINE_API TVertexInstanceAttributesRef<FVector> GetTangents() const { return MeshDescription->VertexInstanceAttributes().GetAttributesRef<FVector>(MeshAttribute::VertexInstance::Tangent); }
+	ENGINE_API TVertexInstanceAttributesConstRef<FVector> GetTangentsConst() const { return MeshDescription->VertexInstanceAttributes().GetAttributesRef<FVector>(MeshAttribute::VertexInstance::Tangent); }
 
-	TVertexInstanceAttributesRef<float> GetBinormalSigns() const { return MeshDescription->VertexInstanceAttributes().GetAttributesRef<float>(MeshAttribute::VertexInstance::BinormalSign); }
-	TVertexInstanceAttributesConstRef<float> GetBinormalSignsConst() const { return MeshDescription->VertexInstanceAttributes().GetAttributesRef<float>(MeshAttribute::VertexInstance::BinormalSign); }
+	ENGINE_API TVertexInstanceAttributesRef<float> GetBinormalSigns() const { return MeshDescription->VertexInstanceAttributes().GetAttributesRef<float>(MeshAttribute::VertexInstance::BinormalSign); }
+	ENGINE_API TVertexInstanceAttributesConstRef<float> GetBinormalSignsConst() const { return MeshDescription->VertexInstanceAttributes().GetAttributesRef<float>(MeshAttribute::VertexInstance::BinormalSign); }
 
-	TVertexInstanceAttributesRef<FVector4> GetColors() const { return MeshDescription->VertexInstanceAttributes().GetAttributesRef<FVector4>(MeshAttribute::VertexInstance::Color); }
-	TVertexInstanceAttributesConstRef<FVector4> GetColorsConst() const { return MeshDescription->VertexInstanceAttributes().GetAttributesRef<FVector4>(MeshAttribute::VertexInstance::Color); }
+	ENGINE_API TVertexInstanceAttributesRef<FVector4> GetColors() const { return MeshDescription->VertexInstanceAttributes().GetAttributesRef<FVector4>(MeshAttribute::VertexInstance::Color); }
+	ENGINE_API TVertexInstanceAttributesConstRef<FVector4> GetColorsConst() const { return MeshDescription->VertexInstanceAttributes().GetAttributesRef<FVector4>(MeshAttribute::VertexInstance::Color); }
 
-	TVertexInstanceAttributesRef<FVector2D> GetUVs() const { return MeshDescription->VertexInstanceAttributes().GetAttributesRef<FVector2D>(MeshAttribute::VertexInstance::TextureCoordinate); }
-	TVertexInstanceAttributesConstRef<FVector2D> GetUVsConst() const { return MeshDescription->VertexInstanceAttributes().GetAttributesRef<FVector2D>(MeshAttribute::VertexInstance::TextureCoordinate); }
+	ENGINE_API TVertexInstanceAttributesRef<FVector2D> GetUVs() const { return MeshDescription->VertexInstanceAttributes().GetAttributesRef<FVector2D>(MeshAttribute::VertexInstance::TextureCoordinate); }
+	ENGINE_API TVertexInstanceAttributesConstRef<FVector2D> GetUVsConst() const { return MeshDescription->VertexInstanceAttributes().GetAttributesRef<FVector2D>(MeshAttribute::VertexInstance::TextureCoordinate); }
 
-	TEdgeAttributesRef<bool> GetEdgeHardnesses() const { return MeshDescription->EdgeAttributes().GetAttributesRef<bool>(MeshAttribute::Edge::IsHard); }
-	TEdgeAttributesConstRef<bool> GetEdgeHardnessesConst() const { return MeshDescription->EdgeAttributes().GetAttributesRef<bool>(MeshAttribute::Edge::IsHard); }
+	ENGINE_API TEdgeAttributesRef<bool> GetEdgeHardnesses() const { return MeshDescription->EdgeAttributes().GetAttributesRef<bool>(MeshAttribute::Edge::IsHard); }
+	ENGINE_API TEdgeAttributesConstRef<bool> GetEdgeHardnessesConst() const { return MeshDescription->EdgeAttributes().GetAttributesRef<bool>(MeshAttribute::Edge::IsHard); }
 	
-	TEdgeAttributesRef<float> GetEdgeCreaseSharpnesses() const { return MeshDescription->EdgeAttributes().GetAttributesRef<float>(MeshAttribute::Edge::CreaseSharpness); }
-	TEdgeAttributesConstRef<float> GetEdgeCreaseSharpnessesConst() const { return MeshDescription->EdgeAttributes().GetAttributesRef<float>(MeshAttribute::Edge::CreaseSharpness); }
+	ENGINE_API TEdgeAttributesRef<float> GetEdgeCreaseSharpnesses() const { return MeshDescription->EdgeAttributes().GetAttributesRef<float>(MeshAttribute::Edge::CreaseSharpness); }
+	ENGINE_API TEdgeAttributesConstRef<float> GetEdgeCreaseSharpnessesConst() const { return MeshDescription->EdgeAttributes().GetAttributesRef<float>(MeshAttribute::Edge::CreaseSharpness); }
 
-	TPolygonGroupAttributesRef<FName> GetPolygonGroupImportedMaterialSlotNames() { return MeshDescription->PolygonGroupAttributes().GetAttributesRef<FName>(MeshAttribute::PolygonGroup::ImportedMaterialSlotName); }
-	TPolygonGroupAttributesConstRef<FName> GetPolygonGroupImportedMaterialSlotNamesConst() { return MeshDescription->PolygonGroupAttributes().GetAttributesRef<FName>(MeshAttribute::PolygonGroup::ImportedMaterialSlotName); }
+	ENGINE_API TPolygonGroupAttributesRef<FName> GetPolygonGroupImportedMaterialSlotNames() { return MeshDescription->PolygonGroupAttributes().GetAttributesRef<FName>(MeshAttribute::PolygonGroup::ImportedMaterialSlotName); }
+	ENGINE_API TPolygonGroupAttributesConstRef<FName> GetPolygonGroupImportedMaterialSlotNamesConst() { return MeshDescription->PolygonGroupAttributes().GetAttributesRef<FName>(MeshAttribute::PolygonGroup::ImportedMaterialSlotName); }
 };
 
 /**
