Exposed GetWorldSettings to scripting

diff --git a/Engine/Source/Runtime/Engine/Classes/Engine/World.h b/Engine/Source/Runtime/Engine/Classes/Engine/World.h
--- a/Engine/Source/Runtime/Engine/Classes/Engine/World.h
+++ b/Engine/Source/Runtime/Engine/Classes/Engine/World.h
@@ -2153,6 +2153,8 @@ public:
 	 *
 	 * @return AWorldSettings actor associated with this world
 	 */
+	UFUNCTION(BlueprintCallable, Category="Utilities|World", meta=(DisplayName="GetWorldSettings", ScriptName="GetWorldSettings"))
+	AWorldSettings* K2_GetWorldSettings();
 	AWorldSettings* GetWorldSettings( bool bCheckStreamingPersistent = false, bool bChecked = true ) const;
 
 	/**
diff --git a/Engine/Source/Runtime/Engine/Private/World.cpp b/Engine/Source/Runtime/Engine/Private/World.cpp
--- a/Engine/Source/Runtime/Engine/Private/World.cpp
+++ b/Engine/Source/Runtime/Engine/Private/World.cpp
@@ -4150,6 +4150,12 @@ ALevelScriptActor* UWorld::GetLevelScriptActor( ULevel* OwnerLevel ) const
 }
 
 
+AWorldSettings* UWorld::K2_GetWorldSettings()
+{
+	return GetWorldSettings(/*bCheckStreamingPersistent*/false, /*bChecked*/false);
+}
+
+
 AWorldSettings* UWorld::GetWorldSettings( const bool bCheckStreamingPersistent, const bool bChecked ) const
 {
 	checkSlow(!IsInActualRenderingThread());
