Fixing SetMaterialAttributes Input serialization by giving input pins names

diff --git a/Engine/Source/Runtime/Engine/Private/Materials/MaterialExpressions.cpp b/Engine/Source/Runtime/Engine/Private/Materials/MaterialExpressions.cpp
--- a/Engine/Source/Runtime/Engine/Private/Materials/MaterialExpressions.cpp
+++ b/Engine/Source/Runtime/Engine/Private/Materials/MaterialExpressions.cpp
@@ -5010,6 +5010,7 @@ void UMaterialExpressionSetMaterialAttributes::PostEditChangeProperty(FPropertyC
 		
 			// Copy final defaults to new input
 			Inputs.Add(FExpressionInput());
+			Inputs.Last().ExpressionName = FName(*FMaterialAttributeDefinitionMap::GetDisplayName(AttributeSetTypes.Last()));
 			GraphNode->ReconstructNode();
 		}	 
 		else if (PreEditAttributeSetTypes.Num() > AttributeSetTypes.Num())
@@ -5051,6 +5052,10 @@ void UMaterialExpressionSetMaterialAttributes::PostEditChangeProperty(FPropertyC
 		else
 		{
 			// Type changed, update pin names
+			for (int i = 1; i < Inputs.Num(); ++i)
+			{
+				Inputs[i].InputName = *FMaterialAttributeDefinitionMap::GetDisplayName(AttributeSetTypes[i - 1]);
+			}
 			GraphNode->ReconstructNode();
 		}
 	}
