[NO NON-PII COMMIT MESSAGE]

diff --git a/Engine/Source/Runtime/Engine/Private/HierarchicalInstancedStaticMesh.cpp b/Engine/Source/Runtime/Engine/Private/HierarchicalInstancedStaticMesh.cpp
--- a/Engine/Source/Runtime/Engine/Private/HierarchicalInstancedStaticMesh.cpp
+++ b/Engine/Source/Runtime/Engine/Private/HierarchicalInstancedStaticMesh.cpp
@@ -2198,7 +2198,7 @@ int32 UHierarchicalInstancedStaticMeshComponent::AddInstance(const FTransform& I
 
 	int32 InstanceIndex = UInstancedStaticMeshComponent::AddInstance(InstanceTransform);
 
-	if (InstanceIndex != INDEX_NONE)
+	if (InstanceIndex != INDEX_NONE && GetStaticMesh() && GetStaticMesh()->HasValidRenderData())
 	{	
 		check(InstanceIndex == InstanceReorderTable.Num());
 
@@ -2213,12 +2213,9 @@ int32 UHierarchicalInstancedStaticMeshComponent::AddInstance(const FTransform& I
 
 		InstanceUpdateCmdBuffer.AddInstance(InstanceTransform.ToMatrixWithScale());
 
-		if (GetStaticMesh())
-		{
-			const FBox NewInstanceBounds = GetStaticMesh()->GetBounds().GetBox().TransformBy(InstanceTransform);
-			UnbuiltInstanceBounds += NewInstanceBounds;
-			UnbuiltInstanceBoundsList.Add(NewInstanceBounds);
-		}
+		const FBox NewInstanceBounds = GetStaticMesh()->GetBounds().GetBox().TransformBy(InstanceTransform);
+		UnbuiltInstanceBounds += NewInstanceBounds;
+		UnbuiltInstanceBoundsList.Add(NewInstanceBounds);
 
 		if (bAutoRebuildTreeOnInstanceChanges)
 		{
