PR #5294: Fixed grammar and typo in comment. (Contributed by [REDACTED_PERSON])

diff --git a/Engine/Source/Runtime/Engine/Classes/Engine/NetSerialization.h b/Engine/Source/Runtime/Engine/Classes/Engine/NetSerialization.h
--- a/Engine/Source/Runtime/Engine/Classes/Engine/NetSerialization.h
+++ b/Engine/Source/Runtime/Engine/Classes/Engine/NetSerialization.h
@@ -179,7 +179,7 @@ struct TStructOpsTypeTraits< FExampleArray > : public TStructOpsTypeTraitsBase2<
  *
  *		Everything originates in UNetDriver::ServerReplicateActors.
  *		Actors are chosen to replicate, create actor channels, and UActorChannel::ReplicateActor is called.
- *		ReplicateActor is ultimately responsible for deciding what properties have changed, and constructing a FOutBUnch to send to clients.
+ *		ReplicateActor is ultimately responsible for deciding what properties have changed, and constructing an FOutBunch to send to clients.
  *
  *	The UActorChannel has 2 ways to decide what properties need to be sent.
  *		The traditional way, which is a flat TArray<uint8> buffer: UActorChannel::Recent. This represents a flat block of the actor properties.
