Improved consistency between in-game and output log console inputs, and various other QoL improvements

diff --git a/Engine/Source/Developer/OutputLog/Private/SOutputLog.cpp b/Engine/Source/Developer/OutputLog/Private/SOutputLog.cpp
--- a/Engine/Source/Developer/OutputLog/Private/SOutputLog.cpp
+++ b/Engine/Source/Developer/OutputLog/Private/SOutputLog.cpp
@@ -241,7 +241,28 @@ void SConsoleInputBox::OnTextChanged(const FText& InText)
 
 			IConsoleManager::Get().ForEachConsoleObjectThatContains(FConsoleObjectVisitor::CreateLambda(OnConsoleVariable), *InputTextStr);
 		}
-		AutoCompleteList.Sort();
+		AutoCompleteList.Sort([InputTextStr](const FString& A, const FString& B)
+		{ 
+			if (A.StartsWith(InputTextStr))
+			{
+				if (!B.StartsWith(InputTextStr))
+				{
+					return true;
+				}
+			}
+			else
+			{
+				if (B.StartsWith(InputTextStr))
+				{
+					return false;
+				}
+			}
+
+			return A < B;
+
+		});
+
+
 		SetSuggestions(AutoCompleteList, FText::FromString(InputTextStr));
 	}
 	else
@@ -304,14 +325,13 @@ FReply SConsoleInputBox::OnPreviewKeyDown(const FGeometry& MyGeometry, const FKe
 			{
 				if (Suggestions.HasSelectedSuggestion())
 				{
-					MarkActiveSuggestion();
-					SuggestionBox->SetIsOpen(false);
+					Suggestions.StepSelectedSuggestion(KeyEvent.IsShiftDown() ? -1 : +1);
 				}
 				else
 				{
 					Suggestions.SelectedSuggestion = 0;
-					MarkActiveSuggestion();
 				}
+				MarkActiveSuggestion();
 			}
 
 			bConsumeTab = true;
@@ -409,6 +429,10 @@ void SConsoleInputBox::SetSuggestions(TArray<FString>& Elements, FText Highlight
 		{
 			SuggestionListView->RequestScrollIntoView(Suggestions.GetSelectedSuggestion());
 		}
+		else
+		{
+			SuggestionListView->ScrollToTop();
+		}
 	}
 	else
 	{
diff --git a/Engine/Source/Runtime/Engine/Classes/Engine/Console.h b/Engine/Source/Runtime/Engine/Classes/Engine/Console.h
--- a/Engine/Source/Runtime/Engine/Classes/Engine/Console.h
+++ b/Engine/Source/Runtime/Engine/Classes/Engine/Console.h
@@ -110,6 +110,9 @@ class ENGINE_API UConsole
 	/** True while a control key is pressed. */
 	uint32 bCtrl:1;
 
+	/** True while a shift key is pressed. */
+	uint32 bShift:1;
+
 	/** Full list of auto-complete commands and info */
 	TArray<FAutoCompleteCommand> AutoCompleteList;
 
@@ -177,6 +180,9 @@ class ENGINE_API UConsole
 	
 	/** looks for Control key presses and the copy/paste combination that apply to both the console bar and the full open console */
 	virtual bool ProcessControlKey(FKey Key, EInputEvent Event);
+
+	/** looks for Shift key presses */
+	virtual bool ProcessShiftKey(FKey Key, EInputEvent Event);
 	
 	/** appends the specified text to the string of typed text */
 	virtual void AppendInputText(const FString& Text);
diff --git a/Engine/Source/Runtime/Engine/Private/UserInterface/Console.cpp b/Engine/Source/Runtime/Engine/Private/UserInterface/Console.cpp
--- a/Engine/Source/Runtime/Engine/Private/UserInterface/Console.cpp
+++ b/Engine/Source/Runtime/Engine/Private/UserInterface/Console.cpp
@@ -683,6 +683,25 @@ bool UConsole::ProcessControlKey(FKey Key, EInputEvent Event)
 	return false;
 }
 
+bool UConsole::ProcessShiftKey(FKey Key, EInputEvent Event)
+{
+	if (Key == EKeys::LeftShift || Key == EKeys::RightShift)
+	{
+		if (Event == IE_Released)
+		{
+			bShift = false;
+		}
+		else if (Event == IE_Pressed)
+		{
+			bShift = true;
+		}
+
+		return true;
+	}
+
+	return false;
+}
+
 
 void UConsole::AppendInputText(const FString& Text)
 {
@@ -813,6 +832,10 @@ bool UConsole::InputKey_InputLine( int32 ControllerId, FKey Key, EInputEvent Eve
 		{
 			return true;
 		}
+		else if (ProcessShiftKey(Key, Event))
+		{
+			return true;
+		}
 		else if( bGamepad )
 		{
 			return false;
@@ -885,12 +908,18 @@ bool UConsole::InputKey_InputLine( int32 ControllerId, FKey Key, EInputEvent Eve
 			}
 			return false;
 		}
-		else if( Key == EKeys::Up )
+		else if( Key == EKeys::Up || (Key == EKeys::Tab && bShift))
 		{
 			if (!bCtrl)
 			{
 				if(AutoComplete.Num())
 				{
+
+					if (Key == EKeys::Tab)
+					{
+						bCaptureKeyInput = true;
+					}
+
 					if (ConsoleSettings->bOrderTopToBottom)
 					{
 						DecrementCursor();
@@ -907,7 +936,7 @@ bool UConsole::InputKey_InputLine( int32 ControllerId, FKey Key, EInputEvent Eve
 				SetInputLineFromAutoComplete();
 			}
 		}
-		else if( Key == EKeys::Down || Key == EKeys::Tab )
+		else if( Key == EKeys::Down || (Key == EKeys::Tab && !bShift) )
 		{
 			if (!bCtrl)
 			{
@@ -1235,6 +1264,7 @@ void UConsole::BeginState_Open(FName PreviousStateName)
 
 	SBPos = 0;
 	bCtrl = false;
+	bShift = false;
 
 	if ( PreviousStateName == NAME_None )
 	{
@@ -1554,6 +1584,7 @@ void UConsole::FakeGotoState(FName NextStateName)
 		// to SetKeyboardFocus the console is still considered active
 		ConsoleState = NAME_None;
 		bCtrl = false;
+		bShift = false;
 
 		TSharedPtr<SWidget> WidgetToFocus;
 		if (PreviousFocusedWidget.IsValid())
