Moved ASceneCapture's visual components (mesh and frustum) to USceneComponentCapture and subclasses
Deprecated root UStaticMeshComponent of ASceneCapture and switched to USceneComponent
Deprecated bShowPreviewPlane on APlanarReflection and moved to UPlanarReflectionComponent
Sync frustum's FOV to scene capture's FOV

diff --git a/Engine/Source/Runtime/Core/Public/UObject/EditorObjectVersion.h b/Engine/Source/Runtime/Core/Public/UObject/EditorObjectVersion.h
--- a/Engine/Source/Runtime/Core/Public/UObject/EditorObjectVersion.h
+++ b/Engine/Source/Runtime/Core/Public/UObject/EditorObjectVersion.h
@@ -63,6 +63,8 @@ struct CORE_API FEditorObjectVersion
 		MeshDescriptionNewSerialization_MovedToRelease,
 		// New format for mesh description attributes
 		MeshDescriptionNewAttributeFormat,
+		// Switch root component of SceneCapture actors from MeshComponent to SceneComponent
+		ChangeSceneCaptureRootComponent,
 		// -----<new versions can be added above this line>-------------------------------------------------
 		VersionPlusOne,
 		LatestVersion = VersionPlusOne - 1
diff --git a/Engine/Source/Runtime/Engine/Classes/Components/PlanarReflectionComponent.h b/Engine/Source/Runtime/Engine/Classes/Components/PlanarReflectionComponent.h
--- a/Engine/Source/Runtime/Engine/Classes/Components/PlanarReflectionComponent.h
+++ b/Engine/Source/Runtime/Engine/Classes/Components/PlanarReflectionComponent.h
@@ -68,6 +68,9 @@ public:
 	UPROPERTY(EditAnywhere, Category = PlanarReflection, meta = (UIMin = "0", UIMax = "90.0"))
 	float AngleFromPlaneFadeEnd;
 
+	UPROPERTY(EditAnywhere, Category = PlanarReflection)
+	bool bShowPreviewPlane;
+
 	/** 
 	 * Whether to render the scene as two-sided, which can be useful to hide artifacts where normal distortion would read 'under' an object that has been clipped by the reflection plane. 
 	 * With this setting enabled, the backfaces of a mesh would be displayed in the clipped region instead of the background which is potentially a bright sky.
@@ -87,6 +90,7 @@ public:
 	virtual void CreateRenderState_Concurrent() override;
 	virtual void SendRenderTransform_Concurrent() override;
 	virtual void DestroyRenderState_Concurrent() override;
+	virtual void OnRegister() override;
 	//~ Begin UActorComponent Interface
 
 #if WITH_EDITOR
@@ -106,6 +110,13 @@ public:
 		return PlanarReflectionId;
 	}
 
+protected:
+#if WITH_EDITORONLY_DATA
+	/** The material to use on ProxyMeshComponent */
+	UPROPERTY(transient)
+	class UMaterial* CaptureMaterial;
+#endif
+
 private:
 
 	/** Fence used to track progress of releasing resources on the rendering thread. */
diff --git a/Engine/Source/Runtime/Engine/Classes/Components/SceneCaptureComponent.h b/Engine/Source/Runtime/Engine/Classes/Components/SceneCaptureComponent.h
--- a/Engine/Source/Runtime/Engine/Classes/Components/SceneCaptureComponent.h
+++ b/Engine/Source/Runtime/Engine/Classes/Components/SceneCaptureComponent.h
@@ -122,6 +122,7 @@ public:
 	FString ProfilingEventName;
 
 	//~ Begin UActorComponent Interface
+	virtual void OnComponentDestroyed(bool bDestroyingHierarchy) override;
 	virtual void OnRegister() override;
 	//~ End UActorComponent Interface
 
@@ -192,5 +193,14 @@ protected:
 	 * NOTE: It is not safe to put a FSceneViewStateReference in a TArray, which moves its contents around without calling element constructors during realloc.
 	 */
 	TIndirectArray<FSceneViewStateReference> ViewStates;
+
+#if WITH_EDITORONLY_DATA
+	/** The mesh used by ProxyMeshComponent */
+	UPROPERTY(transient)
+	class UStaticMesh* CaptureMesh;
+
+	/** The mesh to show visually where the camera is placed */
+	class UStaticMeshComponent* ProxyMeshComponent;
+#endif
 };
 
diff --git a/Engine/Source/Runtime/Engine/Classes/Components/SceneCaptureComponent2D.h b/Engine/Source/Runtime/Engine/Classes/Components/SceneCaptureComponent2D.h
--- a/Engine/Source/Runtime/Engine/Classes/Components/SceneCaptureComponent2D.h
+++ b/Engine/Source/Runtime/Engine/Classes/Components/SceneCaptureComponent2D.h
@@ -82,6 +82,7 @@ class ENGINE_API USceneCaptureComponent2D : public USceneCaptureComponent
 
 
 	//~ Begin UActorComponent Interface
+	virtual void OnComponentDestroyed(bool bDestroyingHierarchy) override;
 	virtual void OnRegister() override;
 	virtual void SendRenderTransform_Concurrent() override;
 	virtual bool RequiresGameThreadEndOfFrameUpdates() const override
@@ -97,7 +98,8 @@ class ENGINE_API USceneCaptureComponent2D : public USceneCaptureComponent
 	virtual bool CanEditChange(const UProperty* InProperty) const override;
 	virtual void PostEditChangeProperty(FPropertyChangedEvent& PropertyChangedEvent) override;
 #endif // WITH_EDITOR
-	
+
+	static void AddReferencedObjects(UObject* InThis, FReferenceCollector& Collector);
 	virtual void Serialize(FArchive& Ar);
 
 	//~ End UObject Interface
@@ -120,4 +122,11 @@ class ENGINE_API USceneCaptureComponent2D : public USceneCaptureComponent
 	void CaptureScene();
 
 	void UpdateSceneCaptureContents(FSceneInterface* Scene) override;
+
+#if WITH_EDITORONLY_DATA
+	void UpdateDrawFrustum();
+
+	/** The frustum component used to show visually where the camera field of view is */
+	class UDrawFrustumComponent* DrawFrustum;
+#endif
 };
diff --git a/Engine/Source/Runtime/Engine/Classes/Components/SceneCaptureComponentCube.h b/Engine/Source/Runtime/Engine/Classes/Components/SceneCaptureComponentCube.h
--- a/Engine/Source/Runtime/Engine/Classes/Components/SceneCaptureComponentCube.h
+++ b/Engine/Source/Runtime/Engine/Classes/Components/SceneCaptureComponentCube.h
@@ -47,12 +47,14 @@ class ENGINE_API USceneCaptureComponentCube : public USceneCaptureComponent
 
 public:
 	//~ Begin UActorComponent Interface
+	virtual void OnComponentDestroyed(bool bDestroyingHierarchy) override;
 	virtual void OnRegister() override;
 	virtual void SendRenderTransform_Concurrent() override;
 	virtual void TickComponent(float DeltaTime, enum ELevelTick TickType, FActorComponentTickFunction* ThisTickFunction) override;
 	//~ End UActorComponent Interface
 
 	//~ Begin UObject Interface
+	static void AddReferencedObjects(UObject* InThis, FReferenceCollector& Collector);
 #if WITH_EDITOR
 	virtual void PostEditChangeProperty(FPropertyChangedEvent& PropertyChangedEvent) override;
 #endif // WITH_EDITOR
@@ -72,4 +74,11 @@ public:
 	void UpdateContent() { CaptureSceneDeferred(); }
 
 	void UpdateSceneCaptureContents(FSceneInterface* Scene) override;
+
+#if WITH_EDITORONLY_DATA
+	void UpdateDrawFrustum();
+
+	/** The frustum component used to show visually where the camera field of view is */
+	class UDrawFrustumComponent* DrawFrustum;
+#endif
 };
diff --git a/Engine/Source/Runtime/Engine/Classes/Engine/PlanarReflection.h b/Engine/Source/Runtime/Engine/Classes/Engine/PlanarReflection.h
--- a/Engine/Source/Runtime/Engine/Classes/Engine/PlanarReflection.h
+++ b/Engine/Source/Runtime/Engine/Classes/Engine/PlanarReflection.h
@@ -29,16 +29,15 @@ private:
 #endif
 
 public:
+	UE_DEPRECATED(4.22, "Please use ShowPreviewPlane on PlanarReflectionComponent instead")
+	UPROPERTY()
+	bool bShowPreviewPlane_DEPRECATED;
 
-	UPROPERTY(EditAnywhere, Category = PlanarReflection)
-	bool bShowPreviewPlane;
+	virtual void PostLoad() override;
 
 	//~ Begin AActor Interface
-	ENGINE_API virtual void PostActorCreated() override;
-
 #if WITH_EDITOR
 	ENGINE_API virtual void EditorApplyScale(const FVector& DeltaScale, const FVector* PivotLocation, bool bAltDown, bool bShiftDown, bool bCtrlDown) override;
-	ENGINE_API virtual void PostEditChangeProperty(FPropertyChangedEvent& PropertyChangedEvent) override;
 #endif
 	//~ End AActor Interface.
 
diff --git a/Engine/Source/Runtime/Engine/Classes/Engine/SceneCapture.h b/Engine/Source/Runtime/Engine/Classes/Engine/SceneCapture.h
--- a/Engine/Source/Runtime/Engine/Classes/Engine/SceneCapture.h
+++ b/Engine/Source/Runtime/Engine/Classes/Engine/SceneCapture.h
@@ -19,12 +19,25 @@ class ASceneCapture : public AActor
 
 private:
 	/** To display the 3d camera in the editor. */
+	UE_DEPRECATED(4.22, "SceneCapture's mesh and frustum components should now be accessed through the SceneCaptureComponent instead of the Actor")
 	UPROPERTY()
-	class UStaticMeshComponent* MeshComp;
+	class UStaticMeshComponent* MeshComp_DEPRECATED;
+
+	UPROPERTY(Category = Components, VisibleAnywhere, BlueprintReadOnly, meta = (AllowPrivateAccess = "true"))
+	class USceneComponent* SceneComponent;
 
 public:
+	virtual void Serialize(FArchive& Ar) override;
+	virtual void PostLoad() override;
+
 	/** Returns MeshComp subobject **/
-	ENGINE_API class UStaticMeshComponent* GetMeshComp() const { return MeshComp; }
+	UE_DEPRECATED(4.22, "SceneCapture's mesh and frustum components should now be accessed through the SceneCaptureComponent instead of the Actor")
+	ENGINE_API class UStaticMeshComponent* GetMeshComp() const
+	{
+		PRAGMA_DISABLE_DEPRECATION_WARNINGS
+		return MeshComp_DEPRECATED;
+		PRAGMA_ENABLE_DEPRECATION_WARNINGS
+	}
 };
 
 
diff --git a/Engine/Source/Runtime/Engine/Classes/Engine/SceneCapture2D.h b/Engine/Source/Runtime/Engine/Classes/Engine/SceneCapture2D.h
--- a/Engine/Source/Runtime/Engine/Classes/Engine/SceneCapture2D.h
+++ b/Engine/Source/Runtime/Engine/Classes/Engine/SceneCapture2D.h
@@ -21,26 +21,8 @@ private:
 	/** Scene capture component. */
 	UPROPERTY(Category = DecalActor, VisibleAnywhere, BlueprintReadOnly, meta=(AllowPrivateAccess = "true"))
 	class USceneCaptureComponent2D* CaptureComponent2D;
-#if WITH_EDITORONLY_DATA
-	/** To allow drawing the camera frustum in the editor. */
-	UPROPERTY()
-	class UDrawFrustumComponent* DrawFrustum;
-#endif
 
 public:
-
-#if WITH_EDITOR
-	//~ Begin AActor Interface
-	ENGINE_API virtual void PostActorCreated() override;
-	//~ End AActor Interface.
-
-	/** Used to synchronize the DrawFrustumComponent with the SceneCaptureComponent2D settings. */
-	void UpdateDrawFrustum();
-
-	/** Returns DrawFrustum subobject **/
-	ENGINE_API class UDrawFrustumComponent* GetDrawFrustum() const { return DrawFrustum; }
-#endif
-
 	UFUNCTION(BlueprintCallable, Category="Rendering")
 	void OnInterpToggle(bool bEnable);
 
diff --git a/Engine/Source/Runtime/Engine/Classes/Engine/SceneCaptureCube.h b/Engine/Source/Runtime/Engine/Classes/Engine/SceneCaptureCube.h
--- a/Engine/Source/Runtime/Engine/Classes/Engine/SceneCaptureCube.h
+++ b/Engine/Source/Runtime/Engine/Classes/Engine/SceneCaptureCube.h
@@ -22,27 +22,7 @@ private:
 	UPROPERTY(Category = DecalActor, VisibleAnywhere, BlueprintReadOnly, meta=(AllowPrivateAccess = "true"))
 	class USceneCaptureComponentCube* CaptureComponentCube;
 
-#if WITH_EDITORONLY_DATA
-	/** To allow drawing the camera frustum in the editor. */
-	UPROPERTY()
-	class UDrawFrustumComponent* DrawFrustum;
-#endif
-
 public:
-
-#if WITH_EDITOR
-	//~ Begin AActor Interface
-	virtual void PostActorCreated() override;
-	virtual void PostEditMove(bool bFinished) override;
-	//~ End AActor Interface.
-
-	/** Used to synchronize the DrawFrustumComponent with the SceneCaptureComponentCube settings. */
-	void UpdateDrawFrustum();
-
-	/** Returns DrawFrustum subobject **/
-	class UDrawFrustumComponent* GetDrawFrustum() const { return DrawFrustum; }
-#endif
-
 	UFUNCTION(BlueprintCallable, Category = "Rendering")
 	void OnInterpToggle(bool bEnable);
 
diff --git a/Engine/Source/Runtime/Engine/Private/Components/SceneCaptureComponent.cpp b/Engine/Source/Runtime/Engine/Private/Components/SceneCaptureComponent.cpp
--- a/Engine/Source/Runtime/Engine/Private/Components/SceneCaptureComponent.cpp
+++ b/Engine/Source/Runtime/Engine/Private/Components/SceneCaptureComponent.cpp
@@ -36,85 +36,51 @@ static TMultiMap<TWeakObjectPtr<UWorld>, TWeakObjectPtr<USceneCaptureComponent>
 ASceneCapture::ASceneCapture(const FObjectInitializer& ObjectInitializer)
 	: Super(ObjectInitializer)
 {
-	MeshComp = CreateDefaultSubobject<UStaticMeshComponent>(TEXT("CamMesh0"));
+	SceneComponent = CreateDefaultSubobject<USceneComponent>(TEXT("SceneComponent"));
+	RootComponent = SceneComponent;
+}
+
+void ASceneCapture::PostLoad()
+{
+	Super::PostLoad();
+
+	if (GetLinkerCustomVersion(FEditorObjectVersion::GUID) < FEditorObjectVersion::ChangeSceneCaptureRootComponent)
+	{
+		PRAGMA_DISABLE_DEPRECATION_WARNINGS
+		if (MeshComp_DEPRECATED)
+		{
+			MeshComp_DEPRECATED->SetStaticMesh(nullptr);
+		}
+		PRAGMA_ENABLE_DEPRECATION_WARNINGS
+	}
+}
 
-	MeshComp->SetCollisionProfileName(UCollisionProfile::NoCollision_ProfileName);
+void ASceneCapture::Serialize(FArchive& Ar)
+{
+	Super::Serialize(Ar);
 
-	MeshComp->bHiddenInGame = true;
-	MeshComp->CastShadow = false;
-	MeshComp->PostPhysicsComponentTick.bCanEverTick = false;
-	RootComponent = MeshComp;
+	Ar.UsingCustomVersion(FEditorObjectVersion::GUID);
 }
 // -----------------------------------------------
 
 ASceneCapture2D::ASceneCapture2D(const FObjectInitializer& ObjectInitializer)
 	: Super(ObjectInitializer)
 {
-#if WITH_EDITOR
-	DrawFrustum = CreateDefaultSubobject<UDrawFrustumComponent>(TEXT("DrawFrust0"));
-	DrawFrustum->SetIsVisualizationComponent(true);
-	DrawFrustum->SetupAttachment(GetMeshComp());
-#endif
-
 	CaptureComponent2D = CreateDefaultSubobject<USceneCaptureComponent2D>(TEXT("NewSceneCaptureComponent2D"));
-	CaptureComponent2D->SetupAttachment(GetMeshComp());
+	CaptureComponent2D->SetupAttachment(RootComponent);
 }
 
 void ASceneCapture2D::OnInterpToggle(bool bEnable)
 {
 	CaptureComponent2D->SetVisibility(bEnable);
 }
-
-#if WITH_EDITOR
-void ASceneCapture2D::UpdateDrawFrustum()
-{
-	if(DrawFrustum && CaptureComponent2D)
-	{
-		DrawFrustum->FrustumStartDist = GNearClippingPlane;
-		
-		// 1000 is the default frustum distance, ideally this would be infinite but that might cause rendering issues
-		DrawFrustum->FrustumEndDist = (CaptureComponent2D->MaxViewDistanceOverride > DrawFrustum->FrustumStartDist)
-			? CaptureComponent2D->MaxViewDistanceOverride : 1000.0f;
-
-		DrawFrustum->FrustumAngle = CaptureComponent2D->FOVAngle;
-		//DrawFrustum->FrustumAspectRatio = CaptureComponent2D->AspectRatio;
-	}
-}
-
-void ASceneCapture2D::PostActorCreated()
-{
-	Super::PostActorCreated();
-
-	// no need load the editor mesh when there is no editor
-	if(GetMeshComp())
-	{
-		if (!IsRunningCommandlet())
-		{
-			if( !GetMeshComp()->GetStaticMesh())
-			{
-				UStaticMesh* CamMesh = LoadObject<UStaticMesh>(NULL, TEXT("/Engine/EditorMeshes/MatineeCam_SM.MatineeCam_SM"), NULL, LOAD_None, NULL);
-				GetMeshComp()->SetStaticMesh(CamMesh);
-			}
-		}
-	}
-
-	// Sync component with CameraActor frustum settings.
-	UpdateDrawFrustum();
-}
-#endif
 // -----------------------------------------------
 
 ASceneCaptureCube::ASceneCaptureCube(const FObjectInitializer& ObjectInitializer)
 	: Super(ObjectInitializer)
 {
-#if WITH_EDITOR
-	DrawFrustum = CreateDefaultSubobject<UDrawFrustumComponent>(TEXT("DrawFrust0"));
-	DrawFrustum->SetIsVisualizationComponent(true);
-	DrawFrustum->SetupAttachment(GetMeshComp());
-#endif
-
 	CaptureComponentCube = CreateDefaultSubobject<USceneCaptureComponentCube>(TEXT("NewSceneCaptureComponentCube"));
-	CaptureComponentCube->SetupAttachment(GetMeshComp());
+	CaptureComponentCube->SetupAttachment(RootComponent);
 }
 
 void ASceneCaptureCube::OnInterpToggle(bool bEnable)
@@ -122,51 +88,6 @@ void ASceneCaptureCube::OnInterpToggle(bool bEnable)
 	CaptureComponentCube->SetVisibility(bEnable);
 }
 
-#if WITH_EDITOR
-void ASceneCaptureCube::UpdateDrawFrustum()
-{
-	if(DrawFrustum && CaptureComponentCube)
-	{
-		DrawFrustum->FrustumStartDist = GNearClippingPlane;
-
-		// 1000 is the default frustum distance, ideally this would be infinite but that might cause rendering issues
-		DrawFrustum->FrustumEndDist = (CaptureComponentCube->MaxViewDistanceOverride > DrawFrustum->FrustumStartDist)
-			? CaptureComponentCube->MaxViewDistanceOverride : 1000.0f;
-
-		DrawFrustum->FrustumAngle = 90;
-	}
-}
-
-void ASceneCaptureCube::PostActorCreated()
-{
-	Super::PostActorCreated();
-
-	// no need load the editor mesh when there is no editor
-	if(GetMeshComp())
-	{
-		if (!IsRunningCommandlet())
-		{
-			if( !GetMeshComp()->GetStaticMesh())
-			{
-				UStaticMesh* CamMesh = LoadObject<UStaticMesh>(NULL, TEXT("/Engine/EditorMeshes/MatineeCam_SM.MatineeCam_SM"), NULL, LOAD_None, NULL);
-				GetMeshComp()->SetStaticMesh(CamMesh);
-			}
-		}
-	}
-
-	// Sync component with CameraActor frustum settings.
-	UpdateDrawFrustum();
-}
-
-void ASceneCaptureCube::PostEditMove(bool bFinished)
-{
-	Super::PostEditMove(bFinished);
-	if(bFinished && CaptureComponentCube->bCaptureOnMovement)
-	{
-		CaptureComponentCube->CaptureSceneDeferred();
-	}
-}
-#endif
 // -----------------------------------------------
 USceneCaptureComponent::USceneCaptureComponent(const FObjectInitializer& ObjectInitializer)
 	: Super(ObjectInitializer), ShowFlags(FEngineShowFlags(ESFIM_Game))
@@ -188,12 +109,43 @@ USceneCaptureComponent::USceneCaptureComponent(const FObjectInitializer& ObjectI
 
 void USceneCaptureComponent::OnRegister()
 {
+#if WITH_EDITORONLY_DATA
+	if (AActor* MyOwner = GetOwner())
+	{
+		if (ProxyMeshComponent == nullptr)
+		{
+			ProxyMeshComponent = NewObject<UStaticMeshComponent>(MyOwner, NAME_None, RF_Transactional | RF_TextExportTransient);
+			ProxyMeshComponent->SetupAttachment(this);
+			ProxyMeshComponent->SetIsVisualizationComponent(true);
+			ProxyMeshComponent->SetStaticMesh(CaptureMesh);
+			ProxyMeshComponent->SetCollisionProfileName(UCollisionProfile::NoCollision_ProfileName);
+			ProxyMeshComponent->bHiddenInGame = true;
+			ProxyMeshComponent->CastShadow = false;
+			ProxyMeshComponent->PostPhysicsComponentTick.bCanEverTick = false;
+			ProxyMeshComponent->CreationMethod = CreationMethod;
+			ProxyMeshComponent->RegisterComponentWithWorld(GetWorld());
+		}
+	}
+#endif
+
 	Super::OnRegister();
 
 	// Make sure any loaded saved flag settings are reflected in our FEngineShowFlags
 	UpdateShowFlags();
 }
 
+void USceneCaptureComponent::OnComponentDestroyed(bool bDestroyingHierarchy)
+{
+	Super::OnComponentDestroyed(bDestroyingHierarchy);
+
+#if WITH_EDITORONLY_DATA
+	if (ProxyMeshComponent)
+	{
+		ProxyMeshComponent->DestroyComponent();
+	}
+#endif
+}
+
 void USceneCaptureComponent::AddReferencedObjects(UObject* InThis, FReferenceCollector& Collector)
 {
 	USceneCaptureComponent* This = CastChecked<USceneCaptureComponent>(InThis);
@@ -207,6 +159,9 @@ void USceneCaptureComponent::AddReferencedObjects(UObject* InThis, FReferenceCol
 		}
 	}
 
+#if WITH_EDITORONLY_DATA
+	Collector.AddReferencedObject(This->ProxyMeshComponent);
+#endif
 	Super::AddReferencedObjects(This, Collector);
 }
 
@@ -467,6 +422,14 @@ USceneCaptureComponent2D::USceneCaptureComponent2D(const FObjectInitializer& Obj
 		// previous behavior was not exposing MotionBlur and Temporal AA in scene capture 2d.
 		ShowFlags.TemporalAA = false;
 		ShowFlags.MotionBlur = false;
+
+#if WITH_EDITORONLY_DATA
+		if (!IsRunningCommandlet())
+		{
+			static ConstructorHelpers::FObjectFinder<UStaticMesh> EditorMesh(TEXT("/Engine/EditorMeshes/MatineeCam_SM"));
+			CaptureMesh = EditorMesh.Object;
+		}
+#endif
 	}
 }
 
@@ -474,6 +437,21 @@ void USceneCaptureComponent2D::OnRegister()
 {
 	Super::OnRegister();
 
+#if WITH_EDITORONLY_DATA
+	if (AActor* MyOwner = GetOwner())
+	{
+		if (DrawFrustum == nullptr)
+		{
+			DrawFrustum = NewObject<UDrawFrustumComponent>(MyOwner, NAME_None, RF_Transactional | RF_TextExportTransient);
+			DrawFrustum->SetupAttachment(this);
+			DrawFrustum->SetIsVisualizationComponent(true);
+			DrawFrustum->CreationMethod = CreationMethod;
+			DrawFrustum->RegisterComponentWithWorld(GetWorld());
+			UpdateDrawFrustum();
+		}
+	}
+#endif
+
 #if WITH_EDITOR
 	// Update content on register to have at least one frames worth of good data.
 	// Without updating here this component would not work in a blueprint construction script which recreates the component after each move in the editor
@@ -530,6 +508,52 @@ void USceneCaptureComponent2D::CaptureScene()
 	}
 }
 
+void USceneCaptureComponent2D::AddReferencedObjects(UObject* InThis, FReferenceCollector& Collector)
+{
+#if WITH_EDITORONLY_DATA
+	USceneCaptureComponent2D* This = CastChecked<USceneCaptureComponent2D>(InThis);
+	Collector.AddReferencedObject(This->DrawFrustum);
+#endif
+
+	Super::AddReferencedObjects(InThis, Collector);
+}
+
+void USceneCaptureComponent2D::OnComponentDestroyed(bool bDestroyingHierarchy)
+{
+	Super::OnComponentDestroyed(bDestroyingHierarchy);
+
+#if WITH_EDITORONLY_DATA
+	if (DrawFrustum)
+	{
+		DrawFrustum->DestroyComponent();
+	}
+#endif
+}
+
+#if WITH_EDITORONLY_DATA
+void USceneCaptureComponent2D::UpdateDrawFrustum()
+{
+	if (DrawFrustum != nullptr)
+	{
+		const float FrustumDrawDistance = 1000.0f;
+		if (ProjectionType == ECameraProjectionMode::Perspective)
+		{
+			DrawFrustum->FrustumAngle = FOVAngle;
+		}
+		else
+		{
+			DrawFrustum->FrustumAngle = -OrthoWidth;
+		}
+
+		DrawFrustum->FrustumStartDist = GNearClippingPlane;
+		// 1000 is the default frustum distance, ideally this would be infinite but that might cause rendering issues
+		DrawFrustum->FrustumEndDist = (MaxViewDistanceOverride > DrawFrustum->FrustumStartDist)
+			? MaxViewDistanceOverride : 1000.0f;
+		DrawFrustum->MarkRenderStateDirty();
+	}
+}
+#endif
+
 #if WITH_EDITOR
 
 bool USceneCaptureComponent2D::CanEditChange(const UProperty* InProperty) const
@@ -587,6 +611,8 @@ void USceneCaptureComponent2D::PostEditChangeProperty(FPropertyChangedEvent& Pro
 	Super::PostEditChangeProperty(PropertyChangedEvent);
 
 	CaptureSceneDeferred();
+
+	UpdateDrawFrustum();
 }
 #endif // WITH_EDITOR
 
@@ -619,8 +645,9 @@ void USceneCaptureComponent2D::UpdateSceneCaptureContents(FSceneInterface* Scene
 APlanarReflection::APlanarReflection(const FObjectInitializer& ObjectInitializer)
 	: Super(ObjectInitializer)
 {
-	bShowPreviewPlane = true;
-
+	PRAGMA_DISABLE_DEPRECATION_WARNINGS
+	bShowPreviewPlane_DEPRECATED = true;
+	PRAGMA_ENABLE_DEPRECATION_WARNINGS
 	PlanarReflectionComponent = CreateDefaultSubobject<UPlanarReflectionComponent>(TEXT("NewPlanarReflectionComponent"));
 	RootComponent = PlanarReflectionComponent;
 
@@ -630,10 +657,6 @@ APlanarReflection::APlanarReflection(const FObjectInitializer& ObjectInitializer
 	DrawInfluenceBox->SetCollisionProfileName(UCollisionProfile::NoCollision_ProfileName);
 	PlanarReflectionComponent->PreviewBox = DrawInfluenceBox;
 
-	GetMeshComp()->SetWorldRotation(FRotator(0, 0, 0));
-	GetMeshComp()->SetWorldScale3D(FVector(4, 4, 1));
-	GetMeshComp()->SetupAttachment(PlanarReflectionComponent);
-
 #if WITH_EDITORONLY_DATA
 	SpriteComponent = CreateEditorOnlyDefaultSubobject<UBillboardComponent>(TEXT("Sprite"));
 	if (!IsRunningCommandlet() && (SpriteComponent != nullptr))
@@ -661,33 +684,24 @@ APlanarReflection::APlanarReflection(const FObjectInitializer& ObjectInitializer
 #endif
 }
 
-void APlanarReflection::OnInterpToggle(bool bEnable)
-{
-	PlanarReflectionComponent->SetVisibility(bEnable);
-}
-
-void APlanarReflection::PostActorCreated()
+void APlanarReflection::PostLoad()
 {
-	Super::PostActorCreated();
+	Super::PostLoad();
 
-	// no need load the editor mesh when there is no editor
-#if WITH_EDITOR
-	if(GetMeshComp())
+	if (GetLinkerCustomVersion(FEditorObjectVersion::GUID) < FEditorObjectVersion::ChangeSceneCaptureRootComponent)
 	{
-		if (!IsRunningCommandlet())
+		if (PlanarReflectionComponent)
 		{
-			if( !GetMeshComp()->GetStaticMesh())
-			{
-				UStaticMesh* PlaneMesh = LoadObject<UStaticMesh>(NULL, TEXT("/Engine/EditorMeshes/PlanarReflectionPlane.PlanarReflectionPlane"), NULL, LOAD_None, NULL);
-				GetMeshComp()->SetStaticMesh(PlaneMesh);
-				UMaterial* PlaneMaterial = LoadObject<UMaterial>(NULL, TEXT("/Engine/EditorMeshes/ColorCalibrator/M_ChromeBall.M_ChromeBall"), NULL, LOAD_None, NULL);
-				GetMeshComp()->SetMaterial(0, PlaneMaterial);
-			}
+			PRAGMA_DISABLE_DEPRECATION_WARNINGS
+			PlanarReflectionComponent->bShowPreviewPlane = bShowPreviewPlane_DEPRECATED;
+			PRAGMA_ENABLE_DEPRECATION_WARNINGS
 		}
-
-		GetMeshComp()->bVisible = bShowPreviewPlane;
 	}
-#endif
+}
+
+void APlanarReflection::OnInterpToggle(bool bEnable)
+{
+	PlanarReflectionComponent->SetVisibility(bEnable);
 }
 
 #if WITH_EDITOR
@@ -703,17 +717,6 @@ void APlanarReflection::EditorApplyScale(const FVector& DeltaScale, const FVecto
 	PostEditChange();
 }
 
-void APlanarReflection::PostEditChangeProperty(FPropertyChangedEvent& PropertyChangedEvent)
-{
-	Super::PostEditChangeProperty(PropertyChangedEvent);
-
-	if (GetMeshComp())
-	{
-		GetMeshComp()->bVisible = bShowPreviewPlane;
-		GetMeshComp()->MarkRenderStateDirty();
-	}
-}
-
 #endif
 
 // -----------------------------------------------
@@ -724,6 +727,8 @@ int32 NextPlanarReflectionId = 0;
 UPlanarReflectionComponent::UPlanarReflectionComponent(const FObjectInitializer& ObjectInitializer)
 	: Super(ObjectInitializer)
 {
+
+	bShowPreviewPlane = true;
 	bCaptureEveryFrame = true;
 	PrimaryComponentTick.bCanEverTick = true;
 	PrimaryComponentTick.TickGroup = TG_DuringPhysics;
@@ -751,8 +756,31 @@ UPlanarReflectionComponent::UPlanarReflectionComponent(const FObjectInitializer&
 
 	NextPlanarReflectionId++;
 	PlanarReflectionId = NextPlanarReflectionId;
+
+#if WITH_EDITORONLY_DATA
+	if (!IsRunningCommandlet())
+	{
+		static ConstructorHelpers::FObjectFinder<UStaticMesh> EditorMesh(TEXT("/Engine/EditorMeshes/PlanarReflectionPlane.PlanarReflectionPlane"));
+		CaptureMesh = EditorMesh.Object;
+		static ConstructorHelpers::FObjectFinder<UMaterial> EditorMaterial(TEXT("/Engine/EditorMeshes/ColorCalibrator/M_ChromeBall.M_ChromeBall"));
+		CaptureMaterial = EditorMaterial.Object;
+	}
+#endif
 }
 
+void UPlanarReflectionComponent::OnRegister()
+{
+	Super::OnRegister();
+
+#if WITH_EDITORONLY_DATA
+	if (ProxyMeshComponent)
+	{
+		ProxyMeshComponent->SetMaterial(0, CaptureMaterial);
+		ProxyMeshComponent->bVisible = bShowPreviewPlane;
+		ProxyMeshComponent->SetRelativeScale3D(FVector(4, 4, 1));
+	}
+#endif
+}
 
 void UPlanarReflectionComponent::Serialize(FArchive& Ar)
 {
@@ -823,6 +851,9 @@ void UPlanarReflectionComponent::PostEditChangeProperty(FPropertyChangedEvent& P
 		ViewStates[ViewIndex].Destroy();
 		ViewStates[ViewIndex].Allocate();
 	}
+
+	ProxyMeshComponent->bVisible = bShowPreviewPlane;
+	ProxyMeshComponent->MarkRenderStateDirty();
 }
 
 #endif
@@ -874,11 +905,35 @@ USceneCaptureComponentCube::USceneCaptureComponentCube(const FObjectInitializer&
 	PrimaryComponentTick.bAllowTickOnDedicatedServer = false;
 	bTickInEditor = true;
 	IPD = 6.2f;
+
+#if WITH_EDITORONLY_DATA
+	if (!IsRunningCommandlet())
+	{
+		static ConstructorHelpers::FObjectFinder<UStaticMesh> EditorMesh(TEXT("/Engine/EditorMeshes/MatineeCam_SM"));
+		CaptureMesh = EditorMesh.Object;
+	}
+#endif
 }
 
 void USceneCaptureComponentCube::OnRegister()
 {
 	Super::OnRegister();
+
+#if WITH_EDITORONLY_DATA
+	if (AActor* MyOwner = GetOwner())
+	{
+		if (DrawFrustum == nullptr)
+		{
+			DrawFrustum = NewObject<UDrawFrustumComponent>(MyOwner, NAME_None, RF_Transactional | RF_TextExportTransient);
+			DrawFrustum->SetupAttachment(this);
+			DrawFrustum->SetIsVisualizationComponent(true);
+			DrawFrustum->CreationMethod = CreationMethod;
+			DrawFrustum->RegisterComponentWithWorld(GetWorld());
+			UpdateDrawFrustum();
+		}
+	}
+#endif
+
 #if WITH_EDITOR
 	// Update content on register to have at least one frames worth of good data.
 	// Without updating here this component would not work in a blueprint construction script which recreates the component after each move in the editor
@@ -906,6 +961,51 @@ void USceneCaptureComponentCube::TickComponent(float DeltaTime, enum ELevelTick
 	}
 }
 
+void USceneCaptureComponentCube::AddReferencedObjects(UObject* InThis, FReferenceCollector& Collector)
+{
+
+#if WITH_EDITORONLY_DATA
+	USceneCaptureComponentCube* This = CastChecked<USceneCaptureComponentCube>(InThis);
+	Collector.AddReferencedObject(This->DrawFrustum);
+#endif
+
+	Super::AddReferencedObjects(InThis, Collector);
+}
+
+void USceneCaptureComponentCube::OnComponentDestroyed(bool bDestroyingHierarchy)
+{
+	Super::OnComponentDestroyed(bDestroyingHierarchy);
+
+#if WITH_EDITORONLY_DATA
+	if (DrawFrustum)
+	{
+		DrawFrustum->DestroyComponent();
+	}
+#endif
+}
+
+#if WITH_EDITORONLY_DATA
+void USceneCaptureComponentCube::UpdateDrawFrustum()
+{
+	if (DrawFrustum != nullptr)
+	{
+		DrawFrustum->FrustumStartDist = GNearClippingPlane;
+
+		// 1000 is the default frustum distance, ideally this would be infinite but that might cause rendering issues
+		const float OldEndDist = DrawFrustum->FrustumEndDist;
+		DrawFrustum->FrustumEndDist = (MaxViewDistanceOverride > DrawFrustum->FrustumStartDist)
+			? MaxViewDistanceOverride : 1000.0f;
+
+		DrawFrustum->FrustumAngle = 90;
+
+		if (OldEndDist != DrawFrustum->FrustumEndDist)
+		{
+			DrawFrustum->MarkRenderStateDirty();
+		}
+	}
+}
+#endif
+
 void USceneCaptureComponentCube::CaptureSceneDeferred()
 {
 	UWorld* World = GetWorld();
@@ -947,6 +1047,8 @@ void USceneCaptureComponentCube::PostEditChangeProperty(FPropertyChangedEvent& P
 	Super::PostEditChangeProperty(PropertyChangedEvent);
 
 	CaptureSceneDeferred();
+
+	UpdateDrawFrustum();
 }
 #endif // WITH_EDITOR
 
