Add ENGINE_API to a few functions from 4623573 to get linker working again

diff --git a/Engine/Source/Runtime/Engine/Classes/Engine/PlanarReflection.h b/Engine/Source/Runtime/Engine/Classes/Engine/PlanarReflection.h
--- a/Engine/Source/Runtime/Engine/Classes/Engine/PlanarReflection.h
+++ b/Engine/Source/Runtime/Engine/Classes/Engine/PlanarReflection.h
@@ -33,7 +33,7 @@ public:
 	UPROPERTY()
 	bool bShowPreviewPlane_DEPRECATED;
 
-	virtual void PostLoad() override;
+	ENGINE_API virtual void PostLoad() override;
 
 	//~ Begin AActor Interface
 #if WITH_EDITOR
diff --git a/Engine/Source/Runtime/Engine/Classes/Engine/SceneCapture.h b/Engine/Source/Runtime/Engine/Classes/Engine/SceneCapture.h
--- a/Engine/Source/Runtime/Engine/Classes/Engine/SceneCapture.h
+++ b/Engine/Source/Runtime/Engine/Classes/Engine/SceneCapture.h
@@ -27,8 +27,8 @@ private:
 	class USceneComponent* SceneComponent;
 
 public:
-	virtual void Serialize(FArchive& Ar) override;
-	virtual void PostLoad() override;
+	ENGINE_API virtual void Serialize(FArchive& Ar) override;
+	ENGINE_API virtual void PostLoad() override;
 
 	/** Returns MeshComp subobject **/
 	UE_DEPRECATED(4.22, "SceneCapture's mesh and frustum components should now be accessed through the SceneCaptureComponent instead of the Actor")
