When changing focus with Mouse For Touch, don't disable events when focus changes to a widget within the game viewport

diff --git a/Engine/Source/Runtime/Engine/Private/GameViewportClient.cpp b/Engine/Source/Runtime/Engine/Private/GameViewportClient.cpp
--- a/Engine/Source/Runtime/Engine/Private/GameViewportClient.cpp
+++ b/Engine/Source/Runtime/Engine/Private/GameViewportClient.cpp
@@ -753,8 +753,11 @@ void UGameViewportClient::MouseLeave(FViewport* InViewport)
 			InViewport->GetMousePos(LastViewportCursorPos, false);
 
 #if PLATFORM_DESKTOP || PLATFORM_HTML5
-			FVector2D CursorPos(LastViewportCursorPos.X, LastViewportCursorPos.Y);
-			FSlateApplication::Get().SetGameIsFakingTouchEvents(false, &CursorPos);
+			if (!GetGameViewportWidget()->HasFocusedDescendants())
+			{
+				FVector2D CursorPos(LastViewportCursorPos.X, LastViewportCursorPos.Y);
+				FSlateApplication::Get().SetGameIsFakingTouchEvents(false, &CursorPos);
+			}
 #endif
 		}
 	}
