In ASceneCapture, update parent component name for children of the deprecated root component

diff --git a/Engine/Source/Runtime/Engine/Private/Components/SceneCaptureComponent.cpp b/Engine/Source/Runtime/Engine/Private/Components/SceneCaptureComponent.cpp
--- a/Engine/Source/Runtime/Engine/Private/Components/SceneCaptureComponent.cpp
+++ b/Engine/Source/Runtime/Engine/Private/Components/SceneCaptureComponent.cpp
@@ -28,6 +28,9 @@
 #include "PlanarReflectionSceneProxy.h"
 #include "Components/BoxComponent.h"
 #include "Logging/MessageLog.h"
+#include "Engine/BlueprintGeneratedClass.h"
+#include "Engine/SimpleConstructionScript.h"
+#include "Engine/SCS_Node.h"
 
 #define LOCTEXT_NAMESPACE "SceneCaptureComponent"
 
@@ -44,8 +47,26 @@ void ASceneCapture::PostLoad()
 {
 	Super::PostLoad();
 
+#if WITH_EDITORONLY_DATA
 	if (GetLinkerCustomVersion(FEditorObjectVersion::GUID) < FEditorObjectVersion::ChangeSceneCaptureRootComponent)
 	{
+		if (IsTemplate())
+		{
+			if (UBlueprintGeneratedClass* BPClass = Cast<UBlueprintGeneratedClass>(GetClass()))
+			{
+				for (USCS_Node* RootNode : BPClass->SimpleConstructionScript->GetRootNodes())
+				{
+					static const FName OldMeshName(TEXT("CamMesh0"));
+					static const FName OldFrustumName(TEXT("DrawFrust0"));
+					static const FName NewRootName(TEXT("SceneComponent"));
+					if (RootNode->ParentComponentOrVariableName == OldMeshName || RootNode->ParentComponentOrVariableName == OldFrustumName)
+					{
+						RootNode->ParentComponentOrVariableName = NewRootName;
+					}
+				}
+			}
+		}
+
 		PRAGMA_DISABLE_DEPRECATION_WARNINGS
 		if (MeshComp_DEPRECATED)
 		{
@@ -53,6 +74,7 @@ void ASceneCapture::PostLoad()
 		}
 		PRAGMA_ENABLE_DEPRECATION_WARNINGS
 	}
+#endif
 }
 
 void ASceneCapture::Serialize(FArchive& Ar)
