Added GUID serialized with MeshDescription bulk data. This is used when generating DDC keys for render data for new assets following RawMesh deprecation.
To do: handle generated LODs.

diff --git a/Engine/Source/Runtime/Core/Public/UObject/EditorObjectVersion.h b/Engine/Source/Runtime/Core/Public/UObject/EditorObjectVersion.h
--- a/Engine/Source/Runtime/Core/Public/UObject/EditorObjectVersion.h
+++ b/Engine/Source/Runtime/Core/Public/UObject/EditorObjectVersion.h
@@ -67,6 +67,8 @@ struct CORE_API FEditorObjectVersion
 		ChangeSceneCaptureRootComponent,
 		// StaticMesh serializes MeshDescription instead of RawMesh
 		StaticMeshDeprecatedRawMesh,
+		// MeshDescriptionBulkData contains a Guid used as a DDC key
+		MeshDescriptionBulkDataGuid,
 		// -----<new versions can be added above this line>-------------------------------------------------
 		VersionPlusOne,
 		LatestVersion = VersionPlusOne - 1
diff --git a/Engine/Source/Runtime/Engine/Private/StaticMesh.cpp b/Engine/Source/Runtime/Engine/Private/StaticMesh.cpp
--- a/Engine/Source/Runtime/Engine/Private/StaticMesh.cpp
+++ b/Engine/Source/Runtime/Engine/Private/StaticMesh.cpp
@@ -1483,7 +1483,18 @@ static FString BuildStaticMeshDerivedDataKey(UStaticMesh* Mesh, const FStaticMes
 	for (int32 LODIndex = 0; LODIndex < NumLODs; ++LODIndex)
 	{
 		FStaticMeshSourceModel& SrcModel = Mesh->SourceModels[LODIndex];
-		KeySuffix += SrcModel.RawMeshBulkData->GetIdString();
+		
+		if (SrcModel.MeshDescriptionBulkData.IsValid())
+		{
+			KeySuffix += "MD";
+			KeySuffix += SrcModel.MeshDescriptionBulkData->GetIdString();
+			// @todo: handle generated LODs? They will not have a valid MeshDescription bulk data
+		}
+		else
+		{
+			// Legacy path for old assets
+			KeySuffix += SrcModel.RawMeshBulkData->GetIdString();
+		}
 
 		// Serialize the build and reduction settings into a temporary array. The archive
 		// is flagged as persistent so that machines of different endianness produce
@@ -2993,7 +3004,7 @@ void UStaticMesh::FixupMaterialSlotName()
 // differences, etc.) replace the version GUID below with a new one.
 // In case of merge conflicts with DDC versions, you MUST generate a new GUID
 // and set this new GUID as the version.                                       
-#define MESHDATAKEY_STATICMESH_DERIVEDDATA_VER TEXT("E8B980F2DF8F4BB292D0FAA928277269")
+#define MESHDATAKEY_STATICMESH_DERIVEDDATA_VER TEXT("7B4DCBC8C34E4C6E97A9A84BA61CADCD")
 
 static const FString& GetMeshDataKeyStaticMeshDerivedDataVersion()
 {
diff --git a/Engine/Source/Runtime/MeshDescription/Private/MeshDescription.cpp b/Engine/Source/Runtime/MeshDescription/Private/MeshDescription.cpp
--- a/Engine/Source/Runtime/MeshDescription/Private/MeshDescription.cpp
+++ b/Engine/Source/Runtime/MeshDescription/Private/MeshDescription.cpp
@@ -1199,6 +1199,15 @@ void FMeshDescription::ReverseAllPolygonFacing()
 void FMeshDescriptionBulkData::Serialize( FArchive& Ar, UObject* Owner )
 {
 	BulkData.Serialize( Ar, Owner );
+
+	if( Ar.IsLoading() && Ar.CustomVer( FEditorObjectVersion::GUID ) < FEditorObjectVersion::MeshDescriptionBulkDataGuid )
+	{
+		FPlatformMisc::CreateGuid( Guid );
+	}
+	else
+	{
+		Ar << Guid;
+	}
 }
 
 
@@ -1212,6 +1221,8 @@ void FMeshDescriptionBulkData::SaveMeshDescription( FMeshDescription& MeshDescri
 		FBulkDataWriter Ar( BulkData, bIsPersistent );
 		Ar << MeshDescription;
 	}
+
+	FPlatformMisc::CreateGuid( Guid );
 }
 
 
@@ -1240,4 +1251,10 @@ void FMeshDescriptionBulkData::Empty()
 	BulkData.RemoveBulkData();
 }
 
+
+FString FMeshDescriptionBulkData::GetIdString() const
+{
+	return Guid.ToString();
+}
+
 #endif // #if WITH_EDITORONLY_DATA
diff --git a/Engine/Source/Runtime/MeshDescription/Public/MeshDescription.h b/Engine/Source/Runtime/MeshDescription/Public/MeshDescription.h
--- a/Engine/Source/Runtime/MeshDescription/Public/MeshDescription.h
+++ b/Engine/Source/Runtime/MeshDescription/Public/MeshDescription.h
@@ -948,11 +948,17 @@ public:
 
 	/** Returns true if there is no bulk data available */
 	bool IsEmpty() const { return BulkData.GetBulkDataSize() == 0; }
+
+	/** Return unique ID string for this bulk data */
+	FString GetIdString() const;
 #endif
 
 private:
 	/** Internally store bulk data as bytes */
 	FByteBulkData BulkData;
+
+	/** GUID associated with the data stored herein. */
+	FGuid Guid;
 };
 
 
