Fixed non-editor builds.

diff --git a/Engine/Source/Runtime/Engine/Private/StaticMesh.cpp b/Engine/Source/Runtime/Engine/Private/StaticMesh.cpp
--- a/Engine/Source/Runtime/Engine/Private/StaticMesh.cpp
+++ b/Engine/Source/Runtime/Engine/Private/StaticMesh.cpp
@@ -3904,11 +3904,13 @@ void UStaticMesh::PostLoad()
 		CreateBodySetup();
 	}
 
+#if WITH_EDITOR
 	// Release cached mesh descriptions until they are loaded on demand
 	for (int LODIndex = 0; LODIndex < SourceModels.Num(); LODIndex++)
 	{
 		ClearMeshDescription(LODIndex);
 	}
+#endif
 
 	CreateNavCollision();
 }
