Fixed crash when cooking due to legacy RawMesh serialization occurring before the StaticMeshOwner had been set in the FStaticMeshSourceModel.

diff --git a/Engine/Source/Runtime/Engine/Private/StaticMesh.cpp b/Engine/Source/Runtime/Engine/Private/StaticMesh.cpp
--- a/Engine/Source/Runtime/Engine/Private/StaticMesh.cpp
+++ b/Engine/Source/Runtime/Engine/Private/StaticMesh.cpp
@@ -1489,13 +1489,17 @@ static FString BuildStaticMeshDerivedDataKey(UStaticMesh* Mesh, const FStaticMes
 		{
 			KeySuffix += "MD";
 			KeySuffix += SrcModel.MeshDescriptionBulkData->GetIdString();
-			// @todo: handle generated LODs? They will not have a valid MeshDescription bulk data
 		}
-		else
+		else if (!SrcModel.RawMeshBulkData->IsEmpty())
 		{
 			// Legacy path for old assets
 			KeySuffix += SrcModel.RawMeshBulkData->GetIdString();
 		}
+		else
+		{
+			// If neither mesh description nor raw mesh bulk data are valid, this is a generated LOD
+			KeySuffix += "_";
+		}
 
 		// Serialize the build and reduction settings into a temporary array. The archive
 		// is flagged as persistent so that machines of different endianness produce
@@ -2606,6 +2610,12 @@ void FStaticMeshSourceModel::SerializeBulkData(FArchive& Ar, UObject* Owner)
 		RawMeshBulkData->Serialize(Ar, Owner);
 	}
 
+	// Initialize the StaticMeshOwner
+	if (Ar.IsLoading())
+	{
+		StaticMeshOwner = Owner;
+	}
+
 	// Always serialize the MeshDescription bulk data when transacting (so undo/redo works correctly).
 	// Now RawMesh is deprecated, always serialize unless we're loading an old archive.
 	if (Ar.IsTransacting() || !bIsLoadingLegacyArchive)
