Fixed crash when cooking due to legacy RawMesh serialization occurring before the StaticMeshOwner had been set in the FStaticMeshSourceModel.

diff --git a/Engine/Source/Runtime/Engine/Private/StaticMesh.cpp b/Engine/Source/Runtime/Engine/Private/StaticMesh.cpp
--- a/Engine/Source/Runtime/Engine/Private/StaticMesh.cpp
+++ b/Engine/Source/Runtime/Engine/Private/StaticMesh.cpp
@@ -2613,7 +2613,7 @@ void FStaticMeshSourceModel::SerializeBulkData(FArchive& Ar, UObject* Owner)
 	// Initialize the StaticMeshOwner
 	if (Ar.IsLoading())
 	{
-		StaticMeshOwner = Owner;
+		StaticMeshOwner = Cast<UStaticMesh>(Owner);
 	}
 
 	// Always serialize the MeshDescription bulk data when transacting (so undo/redo works correctly).
