Fixed issue where source data was getting lost entirely if a static mesh was saved without having converted and committed the MeshDescription converted from RawMesh. Now saving forces the conversion if nothing has yet committed the MeshDescription, as RawMesh is no longer serialized in new assets.

diff --git a/Engine/Source/Runtime/Engine/Private/StaticMesh.cpp b/Engine/Source/Runtime/Engine/Private/StaticMesh.cpp
--- a/Engine/Source/Runtime/Engine/Private/StaticMesh.cpp
+++ b/Engine/Source/Runtime/Engine/Private/StaticMesh.cpp
@@ -2620,6 +2620,16 @@ void FStaticMeshSourceModel::SerializeBulkData(FArchive& Ar, UObject* Owner)
 	// Now RawMesh is deprecated, always serialize unless we're loading an old archive.
 	if (Ar.IsTransacting() || !bIsLoadingLegacyArchive)
 	{
+		if (Ar.IsSaving() && !MeshDescriptionBulkData.IsValid() && !RawMeshBulkData->IsEmpty())
+		{
+			// If saving a legacy asset which hasn't yet been committed as MeshDescription, perform the conversion now
+			// so it can be loaded sucessfully as MeshDescription next time.
+			// Note: even if there's a more recent cached MeshDescription, if it hasn't been committed, it will not be saved.
+			FRawMesh RawMesh;
+			LoadRawMesh(RawMesh);
+			SaveRawMesh(RawMesh);
+		}
+
 		bool bIsValid = MeshDescriptionBulkData.IsValid();
 		Ar << bIsValid;
 
