(4.21.2) Fix cases where using AudioVolume as a child actor template would crash
#rnx

diff --git a/Engine/Source/Runtime/Engine/Private/AudioVolume.cpp b/Engine/Source/Runtime/Engine/Private/AudioVolume.cpp
--- a/Engine/Source/Runtime/Engine/Private/AudioVolume.cpp
+++ b/Engine/Source/Runtime/Engine/Private/AudioVolume.cpp
@@ -210,18 +210,19 @@ FAudioVolumeProxy::FAudioVolumeProxy(const AAudioVolume* AudioVolume)
 
 void AAudioVolume::AddProxy() const
 {
-	UWorld* World = GetWorld();
-
-	if (FAudioDevice* AudioDevice = World->GetAudioDevice())
+	if (UWorld* World = GetWorld())
 	{
-		DECLARE_CYCLE_STAT(TEXT("FAudioThreadTask.AddAudioVolumeProxy"), STAT_AudioAddAudioVolumeProxy, STATGROUP_TaskGraphTasks);
+		if (FAudioDevice* AudioDevice = World->GetAudioDevice())
+		{
+			DECLARE_CYCLE_STAT(TEXT("FAudioThreadTask.AddAudioVolumeProxy"), STAT_AudioAddAudioVolumeProxy, STATGROUP_TaskGraphTasks);
 
-		FAudioVolumeProxy Proxy(this);
+			FAudioVolumeProxy Proxy(this);
 
-		FAudioThread::RunCommandOnAudioThread([AudioDevice, Proxy]()
-		{
-			AudioDevice->AddAudioVolumeProxy(Proxy);
-		}, GET_STATID(STAT_AudioAddAudioVolumeProxy));
+			FAudioThread::RunCommandOnAudioThread([AudioDevice, Proxy]()
+			{
+				AudioDevice->AddAudioVolumeProxy(Proxy);
+			}, GET_STATID(STAT_AudioAddAudioVolumeProxy));
+		}
 	}
 }
 
@@ -266,18 +267,19 @@ void FAudioDevice::RemoveAudioVolumeProxy(const uint32 AudioVolumeID)
 
 void AAudioVolume::UpdateProxy() const
 {
-	UWorld* World = GetWorld();
-
-	if (FAudioDevice* AudioDevice = World->GetAudioDevice())
+	if (UWorld* World = GetWorld())
 	{
-		DECLARE_CYCLE_STAT(TEXT("FAudioThreadTask.UpdateAudioVolumeProxy"), STAT_AudioUpdateAudioVolumeProxy, STATGROUP_TaskGraphTasks);
+		if (FAudioDevice* AudioDevice = World->GetAudioDevice())
+		{
+			DECLARE_CYCLE_STAT(TEXT("FAudioThreadTask.UpdateAudioVolumeProxy"), STAT_AudioUpdateAudioVolumeProxy, STATGROUP_TaskGraphTasks);
 
-		FAudioVolumeProxy Proxy(this);
+			FAudioVolumeProxy Proxy(this);
 
-		FAudioThread::RunCommandOnAudioThread([AudioDevice, Proxy]()
-		{
-			AudioDevice->UpdateAudioVolumeProxy(Proxy);
-		}, GET_STATID(STAT_AudioUpdateAudioVolumeProxy));
+			FAudioThread::RunCommandOnAudioThread([AudioDevice, Proxy]()
+			{
+				AudioDevice->UpdateAudioVolumeProxy(Proxy);
+			}, GET_STATID(STAT_AudioUpdateAudioVolumeProxy));
+		}
 	}
 }
 
@@ -367,10 +369,13 @@ void AAudioVolume::SetPriority(const float NewPriority)
 	if (NewPriority != Priority)
 	{
 		Priority = NewPriority;
-		GetWorld()->AudioVolumes.Sort([](const AAudioVolume& A, const AAudioVolume& B) { return (A.GetPriority() > B.GetPriority()); });
-		if (bEnabled)
+		if (UWorld* World = GetWorld())
 		{
-			UpdateProxy();
+			World->AudioVolumes.Sort([](const AAudioVolume& A, const AAudioVolume& B) { return (A.GetPriority() > B.GetPriority()); });
+			if (bEnabled)
+			{
+				UpdateProxy();
+			}
 		}
 	}
 }
@@ -413,7 +418,10 @@ void AAudioVolume::PostEditChangeProperty(FPropertyChangedEvent& PropertyChanged
 
 	if (PropertyChangedEvent.Property && PropertyChangedEvent.Property->GetFName() == GET_MEMBER_NAME_CHECKED(AAudioVolume, Priority))
 	{
-		GetWorld()->AudioVolumes.Sort([](const AAudioVolume& A, const AAudioVolume& B) { return (A.GetPriority() > B.GetPriority()); });
+		if (UWorld* World = GetWorld())
+		{
+			World->AudioVolumes.Sort([](const AAudioVolume& A, const AAudioVolume& B) { return (A.GetPriority() > B.GetPriority()); });
+		}
 	}
 
 	if (PropertyChangedEvent.Property && PropertyChangedEvent.Property->GetFName() == GET_MEMBER_NAME_CHECKED(AAudioVolume, bEnabled))
