(4.21.2) Prevent crash clearing instance components if there happens to be a null in the array
#rnx

diff --git a/Engine/Source/Runtime/Engine/Private/Actor.cpp b/Engine/Source/Runtime/Engine/Private/Actor.cpp
--- a/Engine/Source/Runtime/Engine/Private/Actor.cpp
+++ b/Engine/Source/Runtime/Engine/Private/Actor.cpp
@@ -2661,7 +2661,10 @@ void AActor::ClearInstanceComponents(const bool bDestroyComponents)
 		// Run in reverse to reduce memory churn when the components are removed from InstanceComponents
 		for (int32 Index=CachedComponents.Num()-1; Index >= 0; --Index)
 		{
-			CachedComponents[Index]->DestroyComponent();
+			if (CachedComponents[Index])
+			{
+				CachedComponents[Index]->DestroyComponent();
+			}
 		}
 	}
 	else
@@ -2995,9 +2998,9 @@ void AActor::FinishSpawning(const FTransform& UserTransform, bool bIsDefaultTran
 				}
 			}
 
-			// should be fast and relatively rare
-			ValidateDeferredTransformCache();
-		}
+				// should be fast and relatively rare
+				ValidateDeferredTransformCache();
+			}
 
 		FinalRootComponentTransform.GetLocation().DiagnosticCheckNaN(TEXT("AActor::FinishSpawning: FinalRootComponentTransform.GetLocation()"));
 		FinalRootComponentTransform.GetRotation().DiagnosticCheckNaN(TEXT("AActor::FinishSpawning: FinalRootComponentTransform.GetRotation()"));
