Remove SubLevelClassesToStrip feature. Never really panned out to work the way we had hoped it would and not useful in its current form.
#rnx

diff --git a/Engine/Config/BaseEngine.ini b/Engine/Config/BaseEngine.ini
--- a/Engine/Config/BaseEngine.ini
+++ b/Engine/Config/BaseEngine.ini
@@ -6,7 +6,6 @@ EditorStartupMap=/Engine/Maps/Templates/Template_Default
 GameDefaultMap=/Engine/Maps/Entry
 ServerDefaultMap=/Engine/Maps/Entry
 GlobalDefaultGameMode="/Script/Engine.GameModeBase"
-+SubLevelClassesToStrip=(ClassToStrip=/Script/Engine.WorldSettings,StripMode=ExactClass)
 ; Example for map prefix and game mode name setting
 ; +GameModeMapPrefixes=(Name="DM-",GameMode="/Script/GamePackage.DMGameMode")
 ; +GameModeClassAliases=(Name="DM",GameMode="/Script/GamePackage.DMGameMode")
diff --git a/Engine/Source/Runtime/Engine/Private/World.cpp b/Engine/Source/Runtime/Engine/Private/World.cpp
--- a/Engine/Source/Runtime/Engine/Private/World.cpp
+++ b/Engine/Source/Runtime/Engine/Private/World.cpp
@@ -2089,12 +2089,6 @@ private:
 
 #endif // PERF_TRACK_DETAILED_ASYNC_STATS
 
-static TAutoConsoleVariable<int32> CVarStripSubLevelClasses(
-	TEXT("level.StripSubLevelClasses"),
-	0,
-	TEXT("0 - The classes specified in Game Maps Settings in sublevels will not be stripped in game worlds. ")
-	TEXT("1 - The classes specified in Game Maps Settings found in sublevels will be marked pending kill when the level is added to a game world. "));
-
 void UWorld::AddToWorld( ULevel* Level, const FTransform& LevelTransform, bool bConsiderTimeLimit )
 {
 	SCOPE_CYCLE_COUNTER(STAT_AddToWorldTime);
@@ -2124,61 +2118,6 @@ void UWorld::AddToWorld( ULevel* Level, const FTransform& LevelTransform, bool b
 		// Mark level as being the one in process of being made visible.
 		CurrentLevelPendingVisibility = Level;
 
-		if (bIsGameWorld && CVarStripSubLevelClasses.GetValueOnGameThread() != 0)
-		{
-			QUICK_SCOPE_CYCLE_COUNTER(STAT_AddToWorldTime_StripSubLevelClasses);
-
-			const TArray<FSubLevelStrippingInfo>& ClassPathsToStrip = GetDefault<UGameMapsSettings>()->SubLevelClassesToStrip;
-			if (ClassPathsToStrip.Num() > 0)
-			{
-				TArray<UClass*> ExactClassesToStrip;
-				TArray<UClass*> IsChildOfClassesToStrip; // not reserving as this is expected to be the infrequent usage
-				ExactClassesToStrip.Reserve(ClassPathsToStrip.Num());
-				for (const FSubLevelStrippingInfo& StrippingInfo : ClassPathsToStrip)
-				{
-					if (UClass* ClassToStrip = StrippingInfo.ClassToStrip.ResolveClass())
-					{
-						if (StrippingInfo.StripMode == ESubLevelStripMode::ExactClass)
-						{
-							ExactClassesToStrip.Add(ClassToStrip);
-						}
-						else //if (StrippingInfo.StripMode == ESubLevelStripMode::IsChildOf)
-						{
-							IsChildOfClassesToStrip.Add(ClassToStrip);
-						}
-					}
-				}
-				if (ExactClassesToStrip.Num() > 0 || IsChildOfClassesToStrip.Num() > 0)
-				{
-					for (AActor*& Actor : Level->Actors)
-					{
-						if (Actor)
-						{
-							if (ExactClassesToStrip.Contains(Actor->GetClass()))
-							{
-								UE_LOG(LogStreaming, Verbose, TEXT("Stripped sub level actor '%s'"), *Actor->GetFullName());
-								Actor->MarkPendingKill(); // We do not need to go through DestroyActor lifecycle as these objects haven't been initialized yet
-								Actor = nullptr;
-							}
-							else
-							{
-								for (UClass* StripClass : IsChildOfClassesToStrip)
-								{
-									if (Actor->GetClass()->IsChildOf(StripClass))
-									{
-										UE_LOG(LogStreaming, Verbose, TEXT("Stripped sub level actor '%s'"), *Actor->GetFullName());
-										Actor->MarkPendingKill(); // We do not need to go through DestroyActor lifecycle as these objects haven't been initialized yet
-										Actor = nullptr;
-										break;
-									}
-								}
-							}
-						}
-					}
-				}
-			}
-		}
-
 		// Add to the UWorld's array of levels, which causes it to be rendered et al.
 		Levels.AddUnique( Level );
 		
diff --git a/Engine/Source/Runtime/EngineSettings/Classes/GameMapsSettings.h b/Engine/Source/Runtime/EngineSettings/Classes/GameMapsSettings.h
--- a/Engine/Source/Runtime/EngineSettings/Classes/GameMapsSettings.h
+++ b/Engine/Source/Runtime/EngineSettings/Classes/GameMapsSettings.h
@@ -66,20 +66,6 @@ enum class ESubLevelStripMode : uint8
 	IsChildOf
 };
 
-USTRUCT()
-struct FSubLevelStrippingInfo
-{
-	GENERATED_BODY()
-
-	// Actor class that should be removed from a level when it is added to a world as a sublevel
-	UPROPERTY(config, EditAnywhere, Category = LevelStreaming, meta = (MetaClass = "Actor"))
-	FSoftClassPath ClassToStrip;
-
-	// Whether the specified class should be be stripped only if an exact class or for any child of that class
-	UPROPERTY(config, EditAnywhere, Category = LevelStreaming)
-	ESubLevelStripMode StripMode;
-};
-
 UCLASS(config=Engine, defaultconfig)
 class ENGINESETTINGS_API UGameMapsSettings
 	: public UObject
@@ -174,10 +160,6 @@ public:
 	UPROPERTY(config, noclear, EditAnywhere, Category=GameInstance, meta=(MetaClass="GameInstance"))
 	FSoftClassPath GameInstanceClass;
 
-	/** A list of classes that should be stripped from a level's actor array when added to a game world as a sublevel. */
-	UPROPERTY(config, EditAnywhere, Category=LevelStreaming, AdvancedDisplay)
-	TArray<FSubLevelStrippingInfo> SubLevelClassesToStrip;
-
 private:
 
 	/** The map that will be loaded by default when no other map is loaded. */
