Ensure that an initially loaded but not initially visible dynamic streaming level is added to the consideration list.

diff --git a/Engine/Source/Runtime/Engine/Private/LevelStreaming.cpp b/Engine/Source/Runtime/Engine/Private/LevelStreaming.cpp
--- a/Engine/Source/Runtime/Engine/Private/LevelStreaming.cpp
+++ b/Engine/Source/Runtime/Engine/Private/LevelStreaming.cpp
@@ -1627,7 +1627,7 @@ void ULevelStreamingDynamic::PostLoad()
 	// Initialize startup state of the streaming level
 	if ( GetWorld()->IsGameWorld() )
 	{
-		bShouldBeLoaded = bInitiallyLoaded;
+		SetShouldBeLoaded(bInitiallyLoaded);
 		SetShouldBeVisible(bInitiallyVisible);
 	}
 }
