Making sure the height fog component isn't disabled when there is a valid density value in either of the two height fog settings

diff --git a/Engine/Source/Runtime/Engine/Private/Components/HeightFogComponent.cpp b/Engine/Source/Runtime/Engine/Private/Components/HeightFogComponent.cpp
--- a/Engine/Source/Runtime/Engine/Private/Components/HeightFogComponent.cpp
+++ b/Engine/Source/Runtime/Engine/Private/Components/HeightFogComponent.cpp
@@ -49,7 +49,7 @@ UExponentialHeightFogComponent::UExponentialHeightFogComponent(const FObjectInit
 
 void UExponentialHeightFogComponent::AddFogIfNeeded()
 {
-	if (ShouldComponentAddToScene() && ShouldRender() && IsRegistered() && FogDensity * 1000 > DELTA && FogMaxOpacity > DELTA
+	if (ShouldComponentAddToScene() && ShouldRender() && IsRegistered() && ((FogDensity + SecondFogData.FogDensity) * 1000) > DELTA && FogMaxOpacity > DELTA
 		&& (GetOuter() == NULL || !GetOuter()->HasAnyFlags(RF_ClassDefaultObject)))
 	{
 		GetWorld()->Scene->AddExponentialHeightFog(this);
