* Remove the profiler marker for preFiltering query results.  This doesn't tell us a whole lot of useful information and adds a ton of markers per-frame.

diff --git a/Engine/Source/Runtime/Engine/Private/PhysicsEngine/PxQueryFilterCallback.cpp b/Engine/Source/Runtime/Engine/Private/PhysicsEngine/PxQueryFilterCallback.cpp
--- a/Engine/Source/Runtime/Engine/Private/PhysicsEngine/PxQueryFilterCallback.cpp
+++ b/Engine/Source/Runtime/Engine/Private/PhysicsEngine/PxQueryFilterCallback.cpp
@@ -10,8 +10,6 @@
 
 PxQueryHitType::Enum FPxQueryFilterCallback::preFilter(const PxFilterData& filterData, const PxShape* shape, const PxRigidActor* actor, PxHitFlags& queryFlags)
 {
-	SCOPE_CYCLE_COUNTER(STAT_Collision_PreFilter);
-
 	ensureMsgf(shape, TEXT("Invalid shape encountered in FPxQueryFilterCallback::preFilter, actor: %p, filterData: %x %x %x %x"), actor, filterData.word0, filterData.word1, filterData.word2, filterData.word3);
 
 	if (!shape)
