[NO NON-PII COMMIT MESSAGE]

diff --git a/Engine/Source/Runtime/Engine/Private/VorbisAudioInfo.cpp b/Engine/Source/Runtime/Engine/Private/VorbisAudioInfo.cpp
--- a/Engine/Source/Runtime/Engine/Private/VorbisAudioInfo.cpp
+++ b/Engine/Source/Runtime/Engine/Private/VorbisAudioInfo.cpp
@@ -541,17 +541,25 @@ bool FVorbisAudioInfo::StreamCompressedData(uint8* InDestination, bool bLooping,
 			// We've reached the end
 			bLooped = true;
 
+			// Clean up decoder state:
+			BufferOffset = 0;
+			ov_clear(&VFWrapper->vf);
+			FMemory::Memzero(&VFWrapper->vf, sizeof(OggVorbis_File));
+
 			// If we're looping, then we need to make sure we wrap the stream chunks back to 0
 			if (bLooping)
 			{
 				NextStreamingChunkIndex = 0;
 			}
+			else
+			{
+				// Need to clear out the remainder of the buffer
+				FMemory::Memzero(InDestination, BufferSize);
+				BytesActuallyRead = BufferSize;
+				break;
+			}
 
-			BufferOffset = 0;
-
-			// Since we can't tell a streaming file to go back to the start of the stream (there is no seek) we have to close and reopen it which is a bummer
-			ov_clear(&VFWrapper->vf);
-			FMemory::Memzero( &VFWrapper->vf, sizeof( OggVorbis_File ) );
+			// Since we can't tell a streaming file to go back to the start of the stream (there is no seek) we have to close and reopen it.
 			ov_callbacks Callbacks;
 			Callbacks.read_func = OggReadStreaming;
 			Callbacks.close_func = OggCloseStreaming;
@@ -564,12 +572,6 @@ bool FVorbisAudioInfo::StreamCompressedData(uint8* InDestination, bool bLooping,
 				break;
 			}
 
-			if( !bLooping )
-			{
-				// Need to clear out the remainder of the buffer
-				FMemory::Memzero(InDestination, BufferSize);
-				break;
-			}
 			// else we start over to get the samples from the start of the compressed audio data
 			continue;
 		}
