#ue4 - Don't delay sending saved character moves that set bForceNoCombine, since those are generally important to send immediately to avoid latency and misprediction.

diff --git a/Engine/Source/Runtime/Engine/Private/Components/CharacterMovementComponent.cpp b/Engine/Source/Runtime/Engine/Private/Components/CharacterMovementComponent.cpp
--- a/Engine/Source/Runtime/Engine/Private/Components/CharacterMovementComponent.cpp
+++ b/Engine/Source/Runtime/Engine/Private/Components/CharacterMovementComponent.cpp
@@ -10466,6 +10466,12 @@ void FSavedMove_Character::PostUpdate(ACharacter* Character, FSavedMove_Characte
 		{
 			bForceNoCombine = true;
 		}
+
+		// Don't combine or delay moves where velocity changes to/from zero.
+		if (StartVelocity.IsZero() != SavedVelocity.IsZero())
+		{
+			bForceNoCombine = true;
+		}
 	}
 	else if (PostUpdateMode == PostUpdate_Replay)
 	{
@@ -10538,6 +10544,12 @@ bool UCharacterMovementComponent::CanDelaySendingMove(const FSavedMovePtr& NewMo
 		return false;
 	}
 
+	// If we know we don't want to combine this move, reduce latency and avoid misprediction by flushing immediately.
+	if (NewMove->bForceNoCombine)
+	{
+		return false;
+	}
+
 	return true;
 }
 
