Fixed invalid memory access with particle trails when wrapping around LastSelectedParticleIndex

diff --git a/Engine/Source/Runtime/Engine/Private/Particles/ParticleTrail2EmitterInstance.cpp b/Engine/Source/Runtime/Engine/Private/Particles/ParticleTrail2EmitterInstance.cpp
--- a/Engine/Source/Runtime/Engine/Private/Particles/ParticleTrail2EmitterInstance.cpp
+++ b/Engine/Source/Runtime/Engine/Private/Particles/ParticleTrail2EmitterInstance.cpp
@@ -2232,7 +2232,7 @@ bool FParticleRibbonEmitterInstance::ResolveSourcePoint(int32 InTrailIdx,
 								LastSelectedParticleIndex++;
 								if (LastSelectedParticleIndex >= SourceEmitter->ActiveParticles)
 								{
-									LastSelectedParticleIndex = -1;
+									LastSelectedParticleIndex = 0;
 								}
 
 								// need to check if the next source index is in use, and go through until we find one that isn't
@@ -2250,7 +2250,7 @@ bool FParticleRibbonEmitterInstance::ResolveSourcePoint(int32 InTrailIdx,
 									{
 										LastSelectedParticleIndex++;
 									}								
-								} while (bInUse && LastSelectedParticleIndex<SourceEmitter->ActiveParticles);
+								} while (bInUse && LastSelectedParticleIndex < SourceEmitter->ActiveParticles);
 
 								if (LastSelectedParticleIndex >= SourceEmitter->ActiveParticles)
 								{
