Enforce EDL as a requirement before cooking optimized Blueprint component data for runtime use.

Change summary:
- Updated the tooltip comment on UCookerSettings::BlueprintComponentDataCookingMethod to note the dependency on EDL.
- Modified UBlueprint::BeginCacheForCookedPlatformData() to ensure that EDL is enabled before cooking Blueprint component data.
    - In the "fast path component instancing" use case, we'll emit a warning at cook time if EDL is disabled, and the fast path will not be used at runtime (no data will be generated and the fast path will not be used, but cooking will succeed).
    - In the "non-nativized component override inherited from nativized parent" use case, we'll emit an error at cook time if EDL is disabled. We must fail the cook, otherwise non-nativized Blueprints can have invalid component data at runtime.

diff --git a/Engine/Source/Editor/UnrealEd/Public/CookerSettings.h b/Engine/Source/Editor/UnrealEd/Public/CookerSettings.h
--- a/Engine/Source/Editor/UnrealEd/Public/CookerSettings.h
+++ b/Engine/Source/Editor/UnrealEd/Public/CookerSettings.h
@@ -89,7 +89,7 @@ public:
 	UPROPERTY(GlobalConfig, EditAnywhere, Category = Cooker, AdvancedDisplay)
 	bool bCompileBlueprintsInDevelopmentMode;
 
-	/** Generate optimized component data to speed up Blueprint construction at runtime. This option can increase the overall Blueprint memory usage in a cooked build. */
+	/** Generate optimized component data to speed up Blueprint construction at runtime. This option can increase the overall Blueprint memory usage in a cooked build. Requires Event-Driven Loading (EDL), which is enabled by default. */
 	UPROPERTY(GlobalConfig, EditAnywhere, Category = Cooker, AdvancedDisplay, meta = (DisplayName = "Generate optimized Blueprint component data"))
 	EBlueprintComponentDataCookingMethod BlueprintComponentDataCookingMethod;
 
diff --git a/Engine/Source/Runtime/Engine/Private/Blueprint.cpp b/Engine/Source/Runtime/Engine/Private/Blueprint.cpp
--- a/Engine/Source/Runtime/Engine/Private/Blueprint.cpp
+++ b/Engine/Source/Runtime/Engine/Private/Blueprint.cpp
@@ -1337,10 +1337,24 @@ void UBlueprint::BeginCacheForCookedPlatformData(const ITargetPlatform *TargetPl
 
 									if (bIsOwnerClassTargetedForReplacement)
 									{
-										// Use the template's archetype for the delta serialization here; remaining properties will have already been set via native subobject instancing at runtime.
-										constexpr bool bUseTemplateArchetype = true;
-										FBlueprintEditorUtils::BuildComponentInstancingData(RecordIt->ComponentTemplate, RecordIt->CookedComponentInstancingData, bUseTemplateArchetype);
-										++NumCookedComponents;
+										// EDL is required, because we need to enforce a preload dependency on the CDO (see UBlueprintGeneratedClass::GetDefaultObjectPreloadDependencies). This is
+										// difficult to support in the non-EDL case, because we have to enforce the dependency at runtime, which can lead to unpredictable results in a cooked build.
+										if (IsEventDrivenLoaderEnabledInCookedBuilds())
+										{
+											// Use the template's archetype for the delta serialization here; remaining properties will have already been set via native subobject instancing at runtime.
+											constexpr bool bUseTemplateArchetype = true;
+											FBlueprintEditorUtils::BuildComponentInstancingData(RecordIt->ComponentTemplate, RecordIt->CookedComponentInstancingData, bUseTemplateArchetype);
+											++NumCookedComponents;
+										}
+										else
+										{
+											UE_LOG(LogBlueprint, Error, TEXT("%s overrides component \'%s\' inherited from %s, which will be converted to C++. This requires Event-Driven Loading (EDL) to be enabled; otherwise, %s must be excluded from Blueprint nativization."),
+												*GetName(),
+												*RecordIt->ComponentKey.GetSCSVariableName().ToString(),
+												*ComponentTemplateOwnerClass->GetName(),
+												*ComponentTemplateOwnerClass->GetName()
+											);
+										}
 									}
 								}
 							}
@@ -1374,6 +1388,20 @@ void UBlueprint::BeginCacheForCookedPlatformData(const ITargetPlatform *TargetPl
 			default:
 				break;
 			}
+			
+			// EDL is required, because we need to enforce a preload dependency on the CDO (see UBlueprintGeneratedClass::GetDefaultObjectPreloadDependencies). This is
+			// difficult to support in the non-EDL case, because we have to enforce the dependency at runtime, which can lead to unpredictable results in a cooked build.
+			if(bResult && !IsEventDrivenLoaderEnabledInCookedBuilds())
+			{
+				bResult = false;
+
+				static bool bWarnOnEDLDisabled = true;
+				if (bWarnOnEDLDisabled)
+				{
+					UE_LOG(LogBlueprint, Warning, TEXT("Cannot cook Blueprint component data for faster instancing at runtime, because Event-Driven Loading (EDL) has been disabled. Re-enable EDL to support this feature, or disable the option to cook Blueprint component data."));
+					bWarnOnEDLDisabled = false;
+				}
+			}
 
 			return bResult;
 		};
