PR: Remove StaticMesh == nullptr map check. (Contributed by [REDACTED_PERSON])

- Mirrored from //UE4/Release-4.21 HF (4721186).

diff --git a/Engine/Source/Runtime/Engine/Private/StaticMeshActor.cpp b/Engine/Source/Runtime/Engine/Private/StaticMeshActor.cpp
--- a/Engine/Source/Runtime/Engine/Private/StaticMeshActor.cpp
+++ b/Engine/Source/Runtime/Engine/Private/StaticMeshActor.cpp
@@ -176,16 +176,7 @@ void AStaticMeshActor::CheckForErrors()
 			->AddToken(FTextToken::Create(FText::Format(LOCTEXT( "MapCheck_Message_StaticMeshComponent", "Static mesh actor {ActorName} has NULL StaticMeshComponent property - please delete" ), Arguments)))
 			->AddToken(FMapErrorToken::Create(FMapErrors::StaticMeshComponent));
 	}
-	else if( StaticMeshComponent->GetStaticMesh() == nullptr )
-	{
-		FFormatNamedArguments Arguments;
-		Arguments.Add(TEXT("ActorName"), FText::FromString(GetName()));
-		MapCheck.Warning()
-			->AddToken(FUObjectToken::Create(this))
-			->AddToken(FTextToken::Create(FText::Format(LOCTEXT( "MapCheck_Message_StaticMeshNull", "Static mesh actor {ActorName} has NULL StaticMesh property" ), Arguments)))
-			->AddToken(FMapErrorToken::Create(FMapErrors::StaticMeshNull));
-	}
-	else
+	else if( StaticMeshComponent->GetStaticMesh() != nullptr )
 	{
 		FCollisionQueryParams SphereParams(SCENE_QUERY_STAT(CheckForErrors), false, this);
 
@@ -194,22 +185,18 @@ void AStaticMeshActor::CheckForErrors()
 
 		for ( int32 OverlapIdx=0; OverlapIdx<Overlaps.Num(); OverlapIdx++ )
 		{
-			AActor* CurTouchingActor = Overlaps[OverlapIdx].GetActor();
-			if( CurTouchingActor != NULL )
+			AStaticMeshActor *A = Cast<AStaticMeshActor>(Overlaps[OverlapIdx].GetActor());
+			if ( A && (A != this) && (A->GetActorLocation() - GetActorLocation()).IsNearlyZero() && A->StaticMeshComponent
+				&& (A->StaticMeshComponent->GetStaticMesh() == StaticMeshComponent->GetStaticMesh()) && (A->GetActorRotation() == GetActorRotation())
+				&& (A->StaticMeshComponent->RelativeScale3D == StaticMeshComponent->RelativeScale3D) )
 			{
-				AStaticMeshActor *A = Cast<AStaticMeshActor>( CurTouchingActor );
-				if ( A && (A != this) && (A->GetActorLocation() - GetActorLocation()).IsNearlyZero() && A->StaticMeshComponent
-					&& (A->StaticMeshComponent->GetStaticMesh() == StaticMeshComponent->GetStaticMesh()) && (A->GetActorRotation() == GetActorRotation())
-					&&  (A->StaticMeshComponent->RelativeScale3D == StaticMeshComponent->RelativeScale3D) )
-				{
-					FFormatNamedArguments Arguments;
-					Arguments.Add(TEXT("ActorName0"), FText::FromString(GetName()));
-					Arguments.Add(TEXT("ActorName1"), FText::FromString(A->GetName()));
-					MapCheck.Warning()
-						->AddToken(FUObjectToken::Create(this))
-						->AddToken(FTextToken::Create(FText::Format( LOCTEXT( "MapCheck_Message_SameLocation", "{ActorName0} is in the same location as {ActorName1}" ), Arguments ) ))
-						->AddToken(FMapErrorToken::Create(FMapErrors::SameLocation));
-				}
+				FFormatNamedArguments Arguments;
+				Arguments.Add(TEXT("ActorName0"), FText::FromString(GetName()));
+				Arguments.Add(TEXT("ActorName1"), FText::FromString(A->GetName()));
+				MapCheck.Warning()
+					->AddToken(FUObjectToken::Create(this))
+					->AddToken(FTextToken::Create(FText::Format( LOCTEXT( "MapCheck_Message_SameLocation", "{ActorName0} is in the same location as {ActorName1}" ), Arguments ) ))
+					->AddToken(FMapErrorToken::Create(FMapErrors::SameLocation));
 			}
 		}
 
