Merging

//Tasks/UE4/Dev-UEA-180-CurveCompression/Engine/Source/Runtime/Engine/Private/AudioDerivedData.cpp

to //UE4/Dev-Anim/Engine/Source/Runtime/Engine/Private/AudioDerivedData.cpp

diff --git a/Engine/Source/Runtime/Engine/Private/AudioDerivedData.cpp b/Engine/Source/Runtime/Engine/Private/AudioDerivedData.cpp
--- a/Engine/Source/Runtime/Engine/Private/AudioDerivedData.cpp
+++ b/Engine/Source/Runtime/Engine/Private/AudioDerivedData.cpp
@@ -732,6 +732,13 @@ bool FStreamedAudioPlatformData::AreDerivedChunksAvailable() const
 
 void FStreamedAudioPlatformData::Serialize(FArchive& Ar, USoundWave* Owner)
 {
+#if WITH_EDITORONLY_DATA
+	if (Owner)
+	{
+		Owner->RawDataCriticalSection.Lock();
+	}
+#endif
+
 	Ar << NumChunks;
 	Ar << AudioFormat;
 
@@ -747,6 +754,13 @@ void FStreamedAudioPlatformData::Serialize(FArchive& Ar, USoundWave* Owner)
 	{
 		Chunks[ChunkIndex].Serialize(Ar, Owner, ChunkIndex);
 	}
+
+#if WITH_EDITORONLY_DATA
+	if (Owner)
+	{
+		Owner->RawDataCriticalSection.Unlock();
+	}
+#endif
 }
 
 /**
@@ -1072,7 +1086,8 @@ static void CookSurroundWave( USoundWave* SoundWave, FName FormatName, const IAu
 			{
 				for (int ChannelIndex = 1; ChannelIndex < ChannelCount; ChannelIndex++)
 				{
-					ResampleWaveData(SourceBuffers[ChannelIndex], SampleDataSize, 1, WaveSampleRate, SampleRateOverride);
+					size_t DataSize = SourceBuffers[ChannelIndex].Num();
+					ResampleWaveData(SourceBuffers[ChannelIndex], DataSize, 1, WaveSampleRate, SampleRateOverride);
 				}	
 
 				WaveSampleRate = SampleRateOverride;
