Merging using Dev-Anim<->Release4.21

Fix crash when resetting animation mode when an anim blueprint uses the rigid body node

diff --git a/Engine/Source/Runtime/Engine/Private/SkeletalMeshComponentPhysics.cpp b/Engine/Source/Runtime/Engine/Private/SkeletalMeshComponentPhysics.cpp
--- a/Engine/Source/Runtime/Engine/Private/SkeletalMeshComponentPhysics.cpp
+++ b/Engine/Source/Runtime/Engine/Private/SkeletalMeshComponentPhysics.cpp
@@ -750,7 +750,8 @@ void USkeletalMeshComponent::InstantiatePhysicsAsset_Internal(const UPhysicsAsse
 		ConInst->CopyConstraintParamsFrom(&OutConstraintsetup->DefaultInstance);
 		ConInst->ConstraintIndex = ConstraintIdx; // Set the ConstraintIndex property in the ConstraintInstance.
 #if WITH_EDITOR
-		if(GetWorld()->IsGameWorld())
+		UWorld* World = GetWorld();
+		if(World && World->IsGameWorld())
 		{
 			//In the editor we may be currently editing the physics asset, so make sure to use the default profile
 			OutConstraintsetup->ApplyConstraintProfile(NAME_None, *ConInst, /*bDefaultIfNotFound=*/true);
