Merging using Dev-Anim<->Release4.21

Fix crash hot-reloading a material applied to an existing skeletal mesh component

No repro, so speculative fix

diff --git a/Engine/Source/Runtime/Engine/Private/SkeletalRenderGPUSkin.cpp b/Engine/Source/Runtime/Engine/Private/SkeletalRenderGPUSkin.cpp
--- a/Engine/Source/Runtime/Engine/Private/SkeletalRenderGPUSkin.cpp
+++ b/Engine/Source/Runtime/Engine/Private/SkeletalRenderGPUSkin.cpp
@@ -155,7 +155,7 @@ void FSkeletalMeshObjectGPUSkin::InitResources(USkinnedMeshComponent* InMeshComp
 		FSkeletalMeshObjectLOD& SkelLOD = LODs[LODIndex];
 		
 		// Skip LODs that have their render data stripped
-		if (SkelLOD.SkelMeshRenderData->LODRenderData[LODIndex].GetNumVertices() > 0)
+		if (SkelLOD.SkelMeshRenderData && SkelLOD.SkelMeshRenderData->LODRenderData.IsValidIndex(LODIndex) && SkelLOD.SkelMeshRenderData->LODRenderData[LODIndex].GetNumVertices() > 0)
 		{
 			const FSkelMeshObjectLODInfo& MeshLODInfo = LODInfo[LODIndex];
 
@@ -177,7 +177,7 @@ void FSkeletalMeshObjectGPUSkin::ReleaseResources()
 		FSkeletalMeshObjectLOD& SkelLOD = LODs[LODIndex];
 		
 		// Skip LODs that have their render data stripped
-		if (SkelLOD.SkelMeshRenderData->LODRenderData[LODIndex].GetNumVertices() > 0)
+		if (SkelLOD.SkelMeshRenderData && SkelLOD.SkelMeshRenderData->LODRenderData.IsValidIndex(LODIndex) && SkelLOD.SkelMeshRenderData->LODRenderData[LODIndex].GetNumVertices() > 0)
 		{
 			SkelLOD.ReleaseResources();
 		}
