Add csv stats for evaluating curves

diff --git a/Engine/Source/Runtime/Engine/Private/Animation/AnimSequence.cpp b/Engine/Source/Runtime/Engine/Private/Animation/AnimSequence.cpp
--- a/Engine/Source/Runtime/Engine/Private/Animation/AnimSequence.cpp
+++ b/Engine/Source/Runtime/Engine/Private/Animation/AnimSequence.cpp
@@ -4996,6 +4996,7 @@ void UAnimSequence::EvaluateCurveData(FBlendedCurve& OutCurve, float CurrentTime
 	}
 	else
 	{
+		CSV_SCOPED_TIMING_STAT(Animation, EvaluateCurveData);
 		CompressedCurveData.EvaluateCurveData(OutCurve, CurrentTime);
 	}
 }
diff --git a/Engine/Source/Runtime/Engine/Private/Animation/AnimSequenceBase.cpp b/Engine/Source/Runtime/Engine/Private/Animation/AnimSequenceBase.cpp
--- a/Engine/Source/Runtime/Engine/Private/Animation/AnimSequenceBase.cpp
+++ b/Engine/Source/Runtime/Engine/Private/Animation/AnimSequenceBase.cpp
@@ -12,6 +12,7 @@
 #include "UObject/FortniteMainBranchObjectVersion.h"
 
 DEFINE_LOG_CATEGORY(LogAnimMarkerSync);
+CSV_DECLARE_CATEGORY_EXTERN(Animation);
 
 #define LOCTEXT_NAMESPACE "AnimSequenceBase"
 /////////////////////////////////////////////////////
@@ -720,6 +721,7 @@ void UAnimSequenceBase::RefreshParentAssetData()
 /** Add curve data to Instance at the time of CurrentTime **/
 void UAnimSequenceBase::EvaluateCurveData(FBlendedCurve& OutCurve, float CurrentTime, bool bForceUseRawData) const
 {
+	CSV_SCOPED_TIMING_STAT(Animation, EvaluateCurveData);
 	RawCurveData.EvaluateCurveData(OutCurve, CurrentTime);
 }
 
