Fixed properties copied into sub anim instances being one frame out of date

diff --git a/Engine/Source/Runtime/Engine/Private/Animation/AnimNode_SubInstance.cpp b/Engine/Source/Runtime/Engine/Private/Animation/AnimNode_SubInstance.cpp
--- a/Engine/Source/Runtime/Engine/Private/Animation/AnimNode_SubInstance.cpp
+++ b/Engine/Source/Runtime/Engine/Private/Animation/AnimNode_SubInstance.cpp
@@ -34,12 +34,7 @@ void FAnimNode_SubInstance::Update_AnyThread(const FAnimationUpdateContext& Cont
 	{
 		FAnimInstanceProxy& Proxy = InstanceToRun->GetProxyOnAnyThread<FAnimInstanceProxy>();
 
-		// Only update if we've not had a single-threaded update already
-		if(InstanceToRun->bNeedsUpdate)
-		{
-			Proxy.UpdateAnimation();
-		}
-
+		// First copy properties
 		check(InstanceProperties.Num() == SubInstanceProperties.Num());
 		for(int32 PropIdx = 0; PropIdx < InstanceProperties.Num(); ++PropIdx)
 		{
@@ -58,6 +53,12 @@ void FAnimNode_SubInstance::Update_AnyThread(const FAnimationUpdateContext& Cont
 				CallerProperty->CopyCompleteValue(DestPtr, SrcPtr);
 			}
 		}
+
+		// Only update if we've not had a single-threaded update already
+		if(InstanceToRun->bNeedsUpdate)
+		{
+			Proxy.UpdateAnimation();
+		}
 	}
 }
 
