Expose FAnimMontageInstance ENGINE_API

PR #5310: Make FAnimMontageInstance to be fully accessed outside the engine (add ENGINE_API) (Contributed by [REDACTED_PERSON])

diff --git a/Engine/Source/Runtime/Engine/Classes/Animation/AnimMontage.h b/Engine/Source/Runtime/Engine/Classes/Animation/AnimMontage.h
--- a/Engine/Source/Runtime/Engine/Classes/Animation/AnimMontage.h
+++ b/Engine/Source/Runtime/Engine/Classes/Animation/AnimMontage.h
@@ -280,7 +280,7 @@ private:
 };
 
 USTRUCT()
-struct FAnimMontageInstance
+struct ENGINE_API FAnimMontageInstance
 {
 	GENERATED_USTRUCT_BODY()
 
@@ -384,7 +384,7 @@ public:
 	 * If Follower gets ticked after Leader, then synchronization will be exact and support more complex cases (i.e. timeline jumps).
 	 *		This can be enforced by setting up tick pre-requisites if desired.
 	 */
-	ENGINE_API void MontageSync_Follow(struct FAnimMontageInstance* NewLeaderMontageInstance);
+	void MontageSync_Follow(struct FAnimMontageInstance* NewLeaderMontageInstance);
 	/** Stop leading, release all followers. */
 	void MontageSync_StopLeading();
 	/** Stop following our leader */
@@ -428,7 +428,7 @@ public:
 
 	//~ Begin montage instance Interfaces
 	void Play(float InPlayRate = 1.f);
-	ENGINE_API void Stop(const FAlphaBlend& InBlendOut, bool bInterrupt=true);
+	void Stop(const FAlphaBlend& InBlendOut, bool bInterrupt=true);
 	void Pause();
 	void Initialize(class UAnimMontage * InMontage);
 
@@ -482,10 +482,10 @@ public:
 	 * second is normal tick. This tick has to happen later when all node ticks
 	 * to accumulate and update curve data/notifies/branching points
 	 */
-	ENGINE_API void UpdateWeight(float DeltaTime);
+	void UpdateWeight(float DeltaTime);
 
 #if WITH_EDITOR	
-	ENGINE_API void EditorOnly_PreAdvance();
+	void EditorOnly_PreAdvance();
 #endif
 
 	/** Simulate is same as Advance, but without calling any events or touching any of the instance data. So it performs a simulation of advancing the timeline. */
@@ -494,7 +494,7 @@ public:
 
 	FName GetCurrentSection() const;
 	FName GetNextSection() const;
-	ENGINE_API int32 GetNextSectionID(int32 const & CurrentSectionID) const;
+	int32 GetNextSectionID(int32 const & CurrentSectionID) const;
 	FName GetSectionNameFromID(int32 const & SectionID) const;
 
 	// reference has to be managed manually
@@ -502,7 +502,7 @@ public:
 
 	/** Delegate function handlers
 	 */
-	ENGINE_API void HandleEvents(float PreviousTrackPos, float CurrentTrackPos, const FBranchingPointMarker* BranchingPointMarker);
+	void HandleEvents(float PreviousTrackPos, float CurrentTrackPos, const FBranchingPointMarker* BranchingPointMarker);
 
 private:
 	/** Called by blueprint functions that modify the montages current position. */
@@ -519,11 +519,11 @@ private:
 
 public:
 	/** static functions that are used by matinee functionality */
-	ENGINE_API static UAnimMontage* SetMatineeAnimPositionInner(FName SlotName, USkeletalMeshComponent* SkeletalMeshComponent, UAnimSequenceBase* InAnimSequence, float InPosition, bool bLooping);
-	ENGINE_API static UAnimMontage* PreviewMatineeSetAnimPositionInner(FName SlotName, USkeletalMeshComponent* SkeletalMeshComponent, UAnimSequenceBase* InAnimSequence, float InPosition, bool bLooping, bool bFireNotifies, float DeltaTime);
+	static UAnimMontage* SetMatineeAnimPositionInner(FName SlotName, USkeletalMeshComponent* SkeletalMeshComponent, UAnimSequenceBase* InAnimSequence, float InPosition, bool bLooping);
+	static UAnimMontage* PreviewMatineeSetAnimPositionInner(FName SlotName, USkeletalMeshComponent* SkeletalMeshComponent, UAnimSequenceBase* InAnimSequence, float InPosition, bool bLooping, bool bFireNotifies, float DeltaTime);
 	/** static functions that are used by sequencer montage support*/
-	ENGINE_API static UAnimMontage* SetSequencerMontagePosition(FName SlotName, USkeletalMeshComponent* SkeletalMeshComponent, int32& InOutInstanceId, UAnimSequenceBase* InAnimSequence, float InPosition, float Weight, bool bLooping, bool bPlaying);
-	ENGINE_API static UAnimMontage* PreviewSequencerMontagePosition(FName SlotName, USkeletalMeshComponent* SkeletalMeshComponent, int32& InOutInstanceId, UAnimSequenceBase* InAnimSequence, float InPosition, float Weight, bool bLooping, bool bFireNotifies, bool bPlaying);
+	static UAnimMontage* SetSequencerMontagePosition(FName SlotName, USkeletalMeshComponent* SkeletalMeshComponent, int32& InOutInstanceId, UAnimSequenceBase* InAnimSequence, float InPosition, float Weight, bool bLooping, bool bPlaying);
+	static UAnimMontage* PreviewSequencerMontagePosition(FName SlotName, USkeletalMeshComponent* SkeletalMeshComponent, int32& InOutInstanceId, UAnimSequenceBase* InAnimSequence, float InPosition, float Weight, bool bLooping, bool bFireNotifies, bool bPlaying);
 private:
 	static UAnimMontage* InitializeMatineeControl(FName SlotName, USkeletalMeshComponent* SkeletalMeshComponent, UAnimSequenceBase* InAnimSequence, bool bLooping);
 };
