Copying //Tasks/UE4/Dev-UEA-180-CurveCompression to Dev-Anim (//UE4/Dev-Anim)

diff --git a/Engine/Config/BaseEngine.ini b/Engine/Config/BaseEngine.ini
--- a/Engine/Config/BaseEngine.ini
+++ b/Engine/Config/BaseEngine.ini
@@ -2150,6 +2150,10 @@ HMDWornMovementThreshold = 50.0
 
 [/Script/Engine.AnimationSettings]
 bStripAnimationDataOnDedicatedServer=False
+
+[Animation.DefaultObjectSettings]
+CurveCompressionSettings="/Engine/Animation/DefaultAnimCurveCompressionSettings"
+
 [/Script/Engine.MeshSimplificationSettings]
 r.MeshReductionModule="QuadricMeshReduction"
 
diff --git a/Engine/Source/Developer/AssetTools/Private/AssetTools.cpp b/Engine/Source/Developer/AssetTools/Private/AssetTools.cpp
--- a/Engine/Source/Developer/AssetTools/Private/AssetTools.cpp
+++ b/Engine/Source/Developer/AssetTools/Private/AssetTools.cpp
@@ -38,6 +38,7 @@
 #include "AssetTypeActions/AssetTypeActions_AnimationAsset.h"
 #include "AssetTypeActions/AssetTypeActions_AnimBlueprint.h"
 #include "AssetTypeActions/AssetTypeActions_AnimComposite.h"
+#include "AssetTypeActions/AssetTypeActions_AnimCurveCompressionSettings.h"
 #include "AssetTypeActions/AssetTypeActions_AnimMontage.h"
 #include "AssetTypeActions/AssetTypeActions_AnimSequence.h"
 #include "AssetTypeActions/AssetTypeActions_BlendSpace.h"
@@ -162,6 +163,7 @@ UAssetToolsImpl::UAssetToolsImpl(const FObjectInitializer& ObjectInitializer)
 	RegisterAssetTypeActions(MakeShareable(new FAssetTypeActions_AnimationAsset));
 	RegisterAssetTypeActions(MakeShareable(new FAssetTypeActions_AnimBlueprint));
 	RegisterAssetTypeActions(MakeShareable(new FAssetTypeActions_AnimComposite));
+	RegisterAssetTypeActions(MakeShareable(new FAssetTypeActions_AnimCurveCompressionSettings));
 	RegisterAssetTypeActions(MakeShareable(new FAssetTypeActions_AnimMontage));
 	RegisterAssetTypeActions(MakeShareable(new FAssetTypeActions_AnimSequence));
 	RegisterAssetTypeActions(MakeShareable(new FAssetTypeActions_AimOffset));
diff --git a/Engine/Source/Developer/AssetTools/Private/AssetTypeActions/AssetTypeActions_AnimCurveCompressionSettings.cpp b/Engine/Source/Developer/AssetTools/Private/AssetTypeActions/AssetTypeActions_AnimCurveCompressionSettings.cpp
new file mode 100644
--- /dev/null
+++ b/Engine/Source/Developer/AssetTools/Private/AssetTypeActions/AssetTypeActions_AnimCurveCompressionSettings.cpp
@@ -0,0 +1,110 @@
+// Copyright 1998-2018 Epic Games, Inc. All Rights Reserved.
+
+#include "AssetTypeActions/AssetTypeActions_AnimCurveCompressionSettings.h"
+#include "Animation/AnimSequence.h"
+#include "Dialogs/Dialogs.h"
+#include "EditorStyleSet.h"
+#include "Framework/MultiBox/MultiBoxBuilder.h"
+#include "Misc/ScopedSlowTask.h"
+#include "UObject/UObjectIterator.h"
+
+#define LOCTEXT_NAMESPACE "AssetTypeActions"
+
+void FAssetTypeActions_AnimCurveCompressionSettings::OpenAssetEditor(const TArray<UObject*>& InObjects, TSharedPtr<class IToolkitHost> EditWithinLevelEditor)
+{
+	TSharedRef<FSimpleAssetEditor> AssetEditor = FSimpleAssetEditor::CreateEditor(EToolkitMode::Standalone, EditWithinLevelEditor, InObjects);
+
+	auto SettingAssets = GetTypedWeakObjectPtrs<UAnimCurveCompressionSettings>(InObjects);
+	if (SettingAssets.Num() == 1)
+	{
+		TSharedPtr<class FUICommandList> PluginCommands = MakeShareable(new FUICommandList);
+		TSharedPtr<FExtender> ToolbarExtender = MakeShareable(new FExtender);
+		ToolbarExtender->AddToolBarExtension("Asset", EExtensionHook::After, PluginCommands, FToolBarExtensionDelegate::CreateRaw(this, &FAssetTypeActions_AnimCurveCompressionSettings::AddToolbarExtension, SettingAssets[0]));
+		AssetEditor->AddToolbarExtender(ToolbarExtender);
+
+		AssetEditor->RegenerateMenusAndToolbars();
+	}
+}
+
+void FAssetTypeActions_AnimCurveCompressionSettings::AddToolbarExtension(FToolBarBuilder& Builder, TWeakObjectPtr<UAnimCurveCompressionSettings> CurveSettings)
+{
+	Builder.BeginSection("Compress");
+	Builder.AddToolBarButton(
+		FUIAction(
+			FExecuteAction::CreateSP(this, &FAssetTypeActions_AnimCurveCompressionSettings::ExecuteCompression, CurveSettings)
+		),
+		NAME_None,
+		LOCTEXT("AnimCurveCompressionSettings_Compress", "Compress"),
+		LOCTEXT("AnimCurveCompressionSettings_CompressTooltip", "All animation sequences that use these settings will be compressed."),
+		FSlateIcon(FEditorStyle::GetStyleSetName(), "Persona.ApplyCompression")
+		);
+	Builder.EndSection();
+}
+
+void FAssetTypeActions_AnimCurveCompressionSettings::GetActions(const TArray<UObject*>& InObjects, FMenuBuilder& MenuBuilder)
+{
+	auto SettingAssets = GetTypedWeakObjectPtrs<UAnimCurveCompressionSettings>(InObjects);
+
+	if (SettingAssets.Num() != 1)
+	{
+		return;
+	}
+
+	MenuBuilder.AddMenuEntry(
+		LOCTEXT("AnimCurveCompressionSettings_Compress", "Compress"),
+		LOCTEXT("AnimCurveCompressionSettings_CompressTooltip", "All animation sequences that use these settings will be compressed."),
+		FSlateIcon(FEditorStyle::GetStyleSetName(), "Persona.ApplyCompression.Small"),
+		FUIAction(
+			FExecuteAction::CreateSP(this, &FAssetTypeActions_AnimCurveCompressionSettings::ExecuteCompression, SettingAssets[0])
+		)
+	);
+}
+
+void FAssetTypeActions_AnimCurveCompressionSettings::ExecuteCompression(TWeakObjectPtr<UAnimCurveCompressionSettings> CurveSettings)
+{
+	if (!CurveSettings.IsValid())
+	{
+		return;
+	}
+
+	UAnimCurveCompressionSettings* Settings = CurveSettings.Get();
+
+	TArray<UAnimSequence*> AnimSeqsToRecompress;
+	for (TObjectIterator<UAnimSequence> It; It; ++It)
+	{
+		UAnimSequence* AnimSeq = *It;
+		if (AnimSeq->GetOutermost() == GetTransientPackage())
+		{
+			continue;
+		}
+
+		if (AnimSeq->CurveCompressionSettings == Settings)
+		{
+			AnimSeqsToRecompress.Add(AnimSeq);
+		}
+	}
+
+	if (AnimSeqsToRecompress.Num() == 0)
+	{
+		return;
+	}
+
+	FFormatNamedArguments Arguments;
+	Arguments.Add(TEXT("NumAnimSequences"), FText::AsNumber(AnimSeqsToRecompress.Num()));
+	FText DialogText = FText::Format(LOCTEXT("AnimCurveCompressionSettings_CompressWarningText", "{NumAnimSequences} animation sequences are about to compress."), Arguments);
+	const EAppReturnType::Type DlgResult = OpenMsgDlgInt(EAppMsgType::OkCancel, DialogText, LOCTEXT("AnimCurveCompressionSettings_CompressWarning", "Warning"));
+	if (DlgResult != EAppReturnType::Ok)
+	{
+		return;
+	}
+
+	const FText StatusText = FText::Format(LOCTEXT("AnimCurveCompressionSettings_Compressing", "Compressing '{0}' animations"), FText::AsNumber(AnimSeqsToRecompress.Num()));
+	FScopedSlowTask LoadingAnimSlowTask(AnimSeqsToRecompress.Num(), StatusText);
+	LoadingAnimSlowTask.MakeDialog();
+
+	for (UAnimSequence* AnimSeq : AnimSeqsToRecompress)
+	{
+		LoadingAnimSlowTask.EnterProgressFrame();
+		AnimSeq->RequestSyncAnimRecompression(false);
+	}
+}
diff --git a/Engine/Source/Developer/AssetTools/Private/AssetTypeActions/AssetTypeActions_AnimCurveCompressionSettings.h b/Engine/Source/Developer/AssetTools/Private/AssetTypeActions/AssetTypeActions_AnimCurveCompressionSettings.h
new file mode 100644
--- /dev/null
+++ b/Engine/Source/Developer/AssetTools/Private/AssetTypeActions/AssetTypeActions_AnimCurveCompressionSettings.h
@@ -0,0 +1,27 @@
+// Copyright 1998-2018 Epic Games, Inc. All Rights Reserved.
+
+#pragma once
+
+#include "CoreMinimal.h"
+#include "AssetTypeActions_Base.h"
+#include "Animation/AnimCurveCompressionSettings.h"
+
+class FAssetTypeActions_AnimCurveCompressionSettings : public FAssetTypeActions_Base
+{
+public:
+	// IAssetTypeActions Implementation
+	virtual FText GetName() const override { return NSLOCTEXT("AssetTypeActions", "AssetTypeActions_AnimCurveCompressionSettings", "Curve Compression Settings"); }
+	virtual FColor GetTypeColor() const override { return FColor(255, 255, 0); }
+	virtual UClass* GetSupportedClass() const override { return UAnimCurveCompressionSettings::StaticClass(); }
+	virtual bool CanFilter() override { return true; }
+	virtual uint32 GetCategories() override { return EAssetTypeCategories::Animation; }
+
+	virtual void OpenAssetEditor(const TArray<UObject*>& InObjects, TSharedPtr<class IToolkitHost> EditWithinLevelEditor = TSharedPtr<IToolkitHost>()) override;
+
+	virtual bool HasActions(const TArray<UObject*>& InObjects) const override { return true; }
+	virtual void GetActions(const TArray<UObject*>& InObjects, FMenuBuilder& MenuBuilder) override;
+
+private:
+	void AddToolbarExtension(FToolBarBuilder& Builder, TWeakObjectPtr<UAnimCurveCompressionSettings> CurveSettings);
+	void ExecuteCompression(TWeakObjectPtr<UAnimCurveCompressionSettings> CurveSettings);
+};
diff --git a/Engine/Source/Editor/Persona/Private/SAnimCurvePanel.cpp b/Engine/Source/Editor/Persona/Private/SAnimCurvePanel.cpp
--- a/Engine/Source/Editor/Persona/Private/SAnimCurvePanel.cpp
+++ b/Engine/Source/Editor/Persona/Private/SAnimCurvePanel.cpp
@@ -164,7 +164,8 @@ public:
 		UpdateNameInternal(AnimSequenceBase.Get()->RawCurveData, RequestedNameUID, RequestedName);
 		if (UAnimSequence* Seq = AnimSequence.Get())
 		{
-			UpdateNameInternal(Seq->CompressedCurveData, RequestedNameUID, RequestedName);
+			FSmartName NewCurveName(RequestedName, RequestedNameUID);
+			Seq->UpdateCompressedCurveName(CurveUID, NewCurveName);
 		}
 		CurveUID = RequestedNameUID;
 	}
@@ -175,13 +176,6 @@ public:
 	void SetCurveTypeFlag(EAnimAssetCurveFlags InFlag, bool bValue)
 	{
 		AnimSequenceBase.Get()->RawCurveData.GetCurveData(CurveUID)->SetCurveTypeFlag(InFlag, bValue);
-		if (UAnimSequence* Seq = AnimSequence.Get())
-		{
-			if (FAnimCurveBase* CompressedCurve = Seq->CompressedCurveData.GetCurveData(CurveUID))
-			{
-				CompressedCurve->SetCurveTypeFlag(InFlag, bValue);
-			}
-		}
 	}
 
 	/**
@@ -190,13 +184,6 @@ public:
 	void ToggleCurveTypeFlag(EAnimAssetCurveFlags InFlag)
 	{
 		AnimSequenceBase.Get()->RawCurveData.GetCurveData(CurveUID)->ToggleCurveTypeFlag(InFlag);
-		if (UAnimSequence* Seq = AnimSequence.Get())
-		{
-			if (FAnimCurveBase* CompressedCurve = Seq->CompressedCurveData.GetCurveData(CurveUID))
-			{
-				CompressedCurve->ToggleCurveTypeFlag(InFlag);
-			}
-		}
 	}
 
 	/**
diff --git a/Engine/Source/Editor/SkeletonEditor/Private/EditableSkeleton.cpp b/Engine/Source/Editor/SkeletonEditor/Private/EditableSkeleton.cpp
--- a/Engine/Source/Editor/SkeletonEditor/Private/EditableSkeleton.cpp
+++ b/Engine/Source/Editor/SkeletonEditor/Private/EditableSkeleton.cpp
@@ -297,7 +297,9 @@ void FEditableSkeleton::RenameSmartname(const FName InContainerName, SmartName::
 						if (UAnimSequence* Seq = Cast<UAnimSequence>(SequenceBase))
 						{
 							SequencesToRecompress.Add(Seq);
-							Seq->CompressedCurveData.Empty();
+
+							Seq->CompressedCurveByteStream.Empty();
+							Seq->CurveCompressionCodec = nullptr;
 						}
 					}
 				}
diff --git a/Engine/Source/Editor/UnrealEd/Classes/Factories/AnimCurveCompressionSettingsFactory.h b/Engine/Source/Editor/UnrealEd/Classes/Factories/AnimCurveCompressionSettingsFactory.h
new file mode 100644
--- /dev/null
+++ b/Engine/Source/Editor/UnrealEd/Classes/Factories/AnimCurveCompressionSettingsFactory.h
@@ -0,0 +1,17 @@
+// Copyright 1998-2018 Epic Games, Inc. All Rights Reserved.
+
+#pragma once
+
+#include "CoreMinimal.h"
+#include "Factories/Factory.h"
+#include "AnimCurveCompressionSettingsFactory.generated.h"
+
+UCLASS(HideCategories = Object, MinimalAPI)
+class UAnimCurveCompressionSettingsFactory : public UFactory
+{
+	GENERATED_UCLASS_BODY()
+
+	//~ Begin UFactory Interface
+	virtual UObject* FactoryCreateNew(UClass* Class, UObject* InParent, FName Name, EObjectFlags Flags, UObject* Context, FFeedbackContext* Warn) override;
+	//~ Begin UFactory Interface
+};
diff --git a/Engine/Source/Editor/UnrealEd/Private/Factories/AnimCurveCompressionSettingsFactory.cpp b/Engine/Source/Editor/UnrealEd/Private/Factories/AnimCurveCompressionSettingsFactory.cpp
new file mode 100644
--- /dev/null
+++ b/Engine/Source/Editor/UnrealEd/Private/Factories/AnimCurveCompressionSettingsFactory.cpp
@@ -0,0 +1,20 @@
+// Copyright 1998-2018 Epic Games, Inc. All Rights Reserved.
+
+/*=============================================================================
+AnimCurveCompressionSettingsFactory.cpp: Factory for animation curve compression settings assets
+=============================================================================*/
+
+#include "Factories/AnimCurveCompressionSettingsFactory.h"
+#include "Animation/AnimCurveCompressionSettings.h"
+
+UAnimCurveCompressionSettingsFactory::UAnimCurveCompressionSettingsFactory(const FObjectInitializer& ObjectInitializer)
+	: Super(ObjectInitializer)
+{
+	bCreateNew = true;
+	SupportedClass = UAnimCurveCompressionSettings::StaticClass();
+}
+
+UObject* UAnimCurveCompressionSettingsFactory::FactoryCreateNew(UClass* Class, UObject* InParent, FName Name, EObjectFlags Flags, UObject* Context, FFeedbackContext* Warn)
+{
+	return NewObject<UAnimCurveCompressionSettings>(InParent, Class, Name, Flags);
+}
diff --git a/Engine/Source/Runtime/Engine/Classes/Animation/AnimCurveCompressionCodec.h b/Engine/Source/Runtime/Engine/Classes/Animation/AnimCurveCompressionCodec.h
new file mode 100644
--- /dev/null
+++ b/Engine/Source/Runtime/Engine/Classes/Animation/AnimCurveCompressionCodec.h
@@ -0,0 +1,75 @@
+// Copyright 1998-2018 Epic Games, Inc. All Rights Reserved.
+
+#pragma once
+
+#include "CoreMinimal.h"
+#include "Animation/AnimTypes.h"
+#include "AnimCurveCompressionCodec.generated.h"
+
+class UAnimCurveCompressionCodec;
+class UAnimSequence;
+struct FBlendedCurve;
+
+#if WITH_EDITORONLY_DATA
+/** Holds the result from animation curve compression */
+struct FAnimCurveCompressionResult
+{
+	/** The animation curves as raw compressed bytes */
+	TArray<uint8> CompressedBytes;
+
+	/** The codec used by the compressed bytes */
+	UAnimCurveCompressionCodec* Codec;
+
+	/** Default constructor */
+	FAnimCurveCompressionResult() : CompressedBytes(), Codec(nullptr) {}
+};
+#endif
+
+/*
+ * Base class for all curve compression codecs.
+ */
+UCLASS(abstract, hidecategories = Object, EditInlineNew)
+class ENGINE_API UAnimCurveCompressionCodec : public UObject
+{
+	GENERATED_UCLASS_BODY()
+
+#if WITH_EDITORONLY_DATA
+	/** A GUID that is unique to this codec instance. After creation, it never changes. */
+	FGuid InstanceGuid;
+#endif
+
+	//////////////////////////////////////////////////////////////////////////
+
+#if WITH_EDITORONLY_DATA
+	// UObject overrides
+	virtual void PostInitProperties() override;
+	virtual void PostDuplicate(bool bDuplicateForPIE) override;
+	virtual void Serialize(FArchive& Ar) override;
+
+	/** Returns whether or not we can use this codec to compress. */
+	virtual bool IsCodecValid() const { return true; }
+
+	/** Compresses the curve data from an animation sequence. */
+	virtual bool Compress(const UAnimSequence& AnimSeq, FAnimCurveCompressionResult& OutResult) PURE_VIRTUAL(UAnimCurveCompressionCodec::Compress, return false;);
+
+	/*
+	 * Called to generate a unique DDC key for this codec instance.
+	 * A suitable key should be generated from: the InstanceGuid, a codec version, and all relevant properties that drive the behavior.
+	 */
+	virtual void PopulateDDCKey(FArchive& Ar);
+#endif
+
+	/*
+	 * Decompresses all the active blended curves.
+	 * Note: Codecs should _NOT_ rely on any member properties during decompression. Decompression
+	 * behavior should entirely be driven by code and the compressed data.
+	 */
+	virtual void DecompressCurves(const UAnimSequence& AnimSeq, FBlendedCurve& Curves, float CurrentTime) const PURE_VIRTUAL(UAnimCurveCompressionCodec::DecompressCurves, );
+
+	/*
+	 * Decompress a single curve.
+	 * Note: Codecs should _NOT_ rely on any member properties during decompression. Decompression
+	 * behavior should entirely be driven by code and the compressed data.
+	 */
+	virtual float DecompressCurve(const UAnimSequence& AnimSeq, SmartName::UID_Type CurveUID, float CurrentTime) const PURE_VIRTUAL(UAnimCurveCompressionCodec::DecompressCurve, return 0.0f;);
+};
diff --git a/Engine/Source/Runtime/Engine/Classes/Animation/AnimCurveCompressionCodec_CompressedRichCurve.h b/Engine/Source/Runtime/Engine/Classes/Animation/AnimCurveCompressionCodec_CompressedRichCurve.h
new file mode 100644
--- /dev/null
+++ b/Engine/Source/Runtime/Engine/Classes/Animation/AnimCurveCompressionCodec_CompressedRichCurve.h
@@ -0,0 +1,42 @@
+// Copyright 1998-2018 Epic Games, Inc. All Rights Reserved.
+
+#pragma once
+
+/**
+* Stores the raw rich curves as FCompressedRichCurve internally with optional key reduction and key time quantization.
+*/
+
+#include "CoreMinimal.h"
+#include "Animation/AnimCurveCompressionCodec.h"
+#include "AnimCurveCompressionCodec_CompressedRichCurve.generated.h"
+
+UCLASS(meta = (DisplayName = "Compressed Rich Curves"))
+class ENGINE_API UAnimCurveCompressionCodec_CompressedRichCurve : public UAnimCurveCompressionCodec
+{
+	GENERATED_UCLASS_BODY()
+
+#if WITH_EDITORONLY_DATA
+	/** Max error allowed when compressing the rich curves */
+	UPROPERTY(Category = Compression, EditAnywhere, meta = (ClampMin = "0"))
+	float MaxCurveError;
+
+	/** Whether to use the animation sequence sample rate or an explicit value */
+	UPROPERTY(Category = Compression, EditAnywhere)
+	bool UseAnimSequenceSampleRate;
+
+	/** Sample rate to use when measuring the curve error */
+	UPROPERTY(Category = Compression, EditAnywhere, meta = (ClampMin = "0", EditCondition = "!UseAnimSequenceSampleRate"))
+	float ErrorSampleRate;
+#endif
+
+	//////////////////////////////////////////////////////////////////////////
+
+#if WITH_EDITORONLY_DATA
+	// UAnimCurveCompressionCodec overrides
+	virtual bool Compress(const UAnimSequence& AnimSeq, FAnimCurveCompressionResult& OutResult) override;
+	virtual void PopulateDDCKey(FArchive& Ar) override;
+#endif
+
+	virtual void DecompressCurves(const UAnimSequence& AnimSeq, FBlendedCurve& Curves, float CurrentTime) const override;
+	virtual float DecompressCurve(const UAnimSequence& AnimSeq, SmartName::UID_Type CurveUID, float CurrentTime) const override;
+};
diff --git a/Engine/Source/Runtime/Engine/Classes/Animation/AnimCurveCompressionCodec_UniformlySampled.h b/Engine/Source/Runtime/Engine/Classes/Animation/AnimCurveCompressionCodec_UniformlySampled.h
new file mode 100644
--- /dev/null
+++ b/Engine/Source/Runtime/Engine/Classes/Animation/AnimCurveCompressionCodec_UniformlySampled.h
@@ -0,0 +1,38 @@
+// Copyright 1998-2018 Epic Games, Inc. All Rights Reserved.
+
+#pragma once
+
+/**
+* Stores the raw rich curves as FCompressedRichCurve internally with optional key reduction and key time quantization.
+*/
+
+#include "CoreMinimal.h"
+#include "Animation/AnimCurveCompressionCodec.h"
+#include "AnimCurveCompressionCodec_UniformlySampled.generated.h"
+
+UCLASS(meta = (DisplayName = "Uniformly Sampled"))
+class ENGINE_API UAnimCurveCompressionCodec_UniformlySampled : public UAnimCurveCompressionCodec
+{
+	GENERATED_UCLASS_BODY()
+
+#if WITH_EDITORONLY_DATA
+	/** Whether to use the animation sequence sample rate or an explicit value */
+	UPROPERTY(Category = Compression, EditAnywhere)
+	bool UseAnimSequenceSampleRate;
+
+	/** Sample rate to use when uniformly sampling */
+	UPROPERTY(Category = Compression, EditAnywhere, meta = (ClampMin = "0", EditCondition = "!UseAnimSequenceSampleRate"))
+	float SampleRate;
+#endif
+
+	//////////////////////////////////////////////////////////////////////////
+
+#if WITH_EDITORONLY_DATA
+	// UAnimCurveCompressionCodec overrides
+	virtual bool Compress(const UAnimSequence& AnimSeq, FAnimCurveCompressionResult& OutResult) override;
+	virtual void PopulateDDCKey(FArchive& Ar) override;
+#endif
+
+	virtual void DecompressCurves(const UAnimSequence& AnimSeq, FBlendedCurve& Curves, float CurrentTime) const override;
+	virtual float DecompressCurve(const UAnimSequence& AnimSeq, SmartName::UID_Type CurveUID, float CurrentTime) const override;
+};
diff --git a/Engine/Source/Runtime/Engine/Classes/Animation/AnimCurveCompressionSettings.h b/Engine/Source/Runtime/Engine/Classes/Animation/AnimCurveCompressionSettings.h
new file mode 100644
--- /dev/null
+++ b/Engine/Source/Runtime/Engine/Classes/Animation/AnimCurveCompressionSettings.h
@@ -0,0 +1,45 @@
+// Copyright 1998-2018 Epic Games, Inc. All Rights Reserved.
+
+#pragma once
+
+#include "CoreMinimal.h"
+#include "AnimCurveCompressionSettings.generated.h"
+
+class UAnimCurveCompressionCodec;
+class UAnimSequence;
+
+/*
+ * This object is used to wrap a curve compression codec. It allows a clean integration in the editor by avoiding the need
+ * to create asset types and factory wrappers for every codec.
+ */
+UCLASS(hidecategories = Object)
+class ENGINE_API UAnimCurveCompressionSettings : public UObject
+{
+	GENERATED_UCLASS_BODY()
+
+#if WITH_EDITORONLY_DATA
+	/** An animation curve compression codec. */
+	UPROPERTY(Category = Compression, Instanced, EditAnywhere, NoClear)
+	UAnimCurveCompressionCodec* Codec;
+#endif
+
+	//////////////////////////////////////////////////////////////////////////
+
+#if WITH_EDITORONLY_DATA
+	// UObject overrides
+	virtual void PostInitProperties() override;
+
+	/** Returns whether or not we can use these settings to compress. */
+	bool AreSettingsValid() const;
+
+	/*
+	 * Compresses the animation curves inside the supplied sequence.
+	 * Note that this will modify the animation sequence by populating the compressed bytes
+	 * and the codec used but it is left unchanged if compression fails.
+	 */
+	bool Compress(UAnimSequence& AnimSeq) const;
+
+	/** Generates a DDC key that takes into account the current settings and selected codec. */
+	FString MakeDDCKey() const;
+#endif
+};
diff --git a/Engine/Source/Runtime/Engine/Classes/Animation/AnimSequence.h b/Engine/Source/Runtime/Engine/Classes/Animation/AnimSequence.h
--- a/Engine/Source/Runtime/Engine/Classes/Animation/AnimSequence.h
+++ b/Engine/Source/Runtime/Engine/Classes/Animation/AnimSequence.h
@@ -448,6 +448,12 @@ protected:
 	 */
 	TArray<struct FTrackToSkeletonMap> CompressedTrackToSkeletonMapTable;
 
+	/**
+	 * Much like track indices above, we need to be able to remap curve names. The FName inside FSmartName
+	 * is serialized and the UID is generated at runtime. Stored in the DDC.
+	 */
+	TArray<struct FSmartName> CompressedCurveNames;
+
 	/**
 	 * Raw uncompressed keyframe data. 
 	 */
@@ -484,8 +490,15 @@ public:
 	 */
 	UPROPERTY(Category=Compression, VisibleAnywhere)
 	class UAnimCompress* CompressionScheme;
+
+	/** The curve compression settings used to compress curves in this sequence. */
+	UPROPERTY(Category=Compression, EditAnywhere)
+	class UAnimCurveCompressionSettings* CurveCompressionSettings;
 #endif // WITH_EDITORONLY_DATA
 
+	/** The codec used by the compressed data as determined by the compression settings. */
+	class UAnimCurveCompressionCodec* CurveCompressionCodec;
+
 	/** The compression format that was used to compress translation tracks. */
 	TEnumAsByte<enum AnimationCompressionFormat> TranslationCompressionFormat;
 
@@ -541,8 +554,8 @@ public:
 
 	class AnimEncoding* ScaleCodec;
 
-	// Built during compression, could be baked additive or original curve data
-	FRawCurveTracks CompressedCurveData;
+	/* Compressed curve data stream used by AnimCurveCompressionCodec */
+	TArray<uint8> CompressedCurveByteStream;
 
 	// The size of the raw data used to create the compressed data
 	int32 CompressedRawDataSize;
@@ -659,6 +672,7 @@ public:
 #endif // WITH_EDITOR
 	virtual void BeginDestroy() override;
 	virtual void GetAssetRegistryTags(TArray<FAssetRegistryTag>& OutTags) const override;
+	static void AddReferencedObjects(UObject* This, FReferenceCollector& Collector);
 	//~ End UObject Interface
 
 	//~ Begin UAnimationAsset Interface
@@ -676,8 +690,11 @@ public:
 	virtual bool HasRootMotion() const override { return bEnableRootMotion; }
 	virtual void RefreshCacheData() override;
 	virtual EAdditiveAnimationType GetAdditiveAnimType() const override { return AdditiveAnimType; }
-	virtual void EvaluateCurveData(FBlendedCurve& OutCurve, float CurrentTime, bool bForceUseRawData=false) const;
-	virtual const FRawCurveTracks& GetCurveData() const;
+
+	virtual void EvaluateCurveData(FBlendedCurve& OutCurve, float CurrentTime, bool bForceUseRawData = false) const override;
+	virtual float EvaluateCurveData(SmartName::UID_Type CurveUID, float CurrentTime, bool bForceUseRawData = false) const override;
+	virtual bool HasCurveData(SmartName::UID_Type CurveUID, bool bForceUseRawData) const override;
+
 #if WITH_EDITOR
 	virtual void MarkRawDataAsModified(bool bForceNewRawDatGuid = true) override
 	{
@@ -731,6 +748,8 @@ public:
 	const TArray<FRawAnimSequenceTrack>& GetAdditiveBaseAnimationData() const { return TemporaryAdditiveBaseAnimationData; }
 	void  UpdateCompressedTrackMapFromRaw() { CompressedTrackToSkeletonMapTable = TrackToSkeletonMapTable; }
 	void  UpdateCompressedNumFramesFromRaw() { CompressedNumFrames = NumFrames; }
+	void  UpdateCompressedCurveNames();
+	void  UpdateCompressedCurveName(SmartName::UID_Type CurveUID, const struct FSmartName& NewCurveName);
 	
 	// Adds a new track (if no track of the supplied name is found) to the raw animation data, optionally setting it to TrackData.
 	int32 AddNewRawTrack(FName TrackName, FRawAnimSequenceTrack* TrackData = nullptr);
@@ -738,6 +757,7 @@ public:
 
 	const TArray<FTrackToSkeletonMap>& GetRawTrackToSkeletonMapTable() const { return TrackToSkeletonMapTable; }
 	const TArray<FTrackToSkeletonMap>& GetCompressedTrackToSkeletonMapTable() const { return CompressedTrackToSkeletonMapTable; }
+	const TArray<struct FSmartName>& GetCompressedCurveNames() const { return CompressedCurveNames; }
 
 	FRawAnimSequenceTrack& GetRawAnimationTrack(int32 TrackIndex) { return RawAnimationData[TrackIndex]; }
 	const FRawAnimSequenceTrack& GetRawAnimationTrack(int32 TrackIndex) const { return RawAnimationData[TrackIndex]; }
@@ -867,6 +887,7 @@ public:
 	void RequestAnimCompression(FRequestAnimCompressionParams Params);
 	void RequestSyncAnimRecompression(bool bOutput = false) { RequestAnimCompression(FRequestAnimCompressionParams(false, false, bOutput)); }
 	bool IsCompressedDataValid() const;
+	bool IsCurveCompressedDataValid() const;
 
 	// Write the compressed data to the supplied FArchive
 	void SerializeCompressedData(FArchive& Ar, bool bDDCData);
diff --git a/Engine/Source/Runtime/Engine/Classes/Animation/AnimSequenceBase.h b/Engine/Source/Runtime/Engine/Classes/Animation/AnimSequenceBase.h
--- a/Engine/Source/Runtime/Engine/Classes/Animation/AnimSequenceBase.h
+++ b/Engine/Source/Runtime/Engine/Classes/Animation/AnimSequenceBase.h
@@ -101,9 +101,11 @@ class ENGINE_API UAnimSequenceBase : public UAnimationAsset
 	virtual void GetAnimNotifiesFromDeltaPositions(const float& PreviousPosition, const float & CurrentPosition, TArray<FAnimNotifyEventReference>& OutActiveNotifies) const;
 
 	/** Evaluate curve data to Instance at the time of CurrentTime **/
-	virtual void EvaluateCurveData(FBlendedCurve& OutCurve, float CurrentTime, bool bForceUseRawData=false) const;
+	virtual void EvaluateCurveData(FBlendedCurve& OutCurve, float CurrentTime, bool bForceUseRawData = false) const;
+	virtual float EvaluateCurveData(SmartName::UID_Type CurveUID, float CurrentTime, bool bForceUseRawData = false) const;
 
 	virtual const FRawCurveTracks& GetCurveData() const { return RawCurveData; }
+	virtual bool HasCurveData(SmartName::UID_Type CurveUID, bool bForceUseRawData = false) const;
 
 	/** Return Number of Frames **/
 	virtual int32 GetNumberOfFrames() const;
diff --git a/Engine/Source/Runtime/Engine/Classes/Curves/RichCurve.h b/Engine/Source/Runtime/Engine/Classes/Curves/RichCurve.h
--- a/Engine/Source/Runtime/Engine/Classes/Curves/RichCurve.h
+++ b/Engine/Source/Runtime/Engine/Classes/Curves/RichCurve.h
@@ -73,6 +73,37 @@ enum ERichCurveExtrapolation
 };
 
 
+/** Enumerates curve compression options. */
+UENUM()
+enum ERichCurveCompressionFormat
+{
+	/** No keys are present */
+	RCCF_Empty UMETA(DisplayName = "Empty"),
+
+	/** All keys use constant interpolation */
+	RCCF_Constant UMETA(DisplayName = "Constant"),
+
+	/** All keys use linear interpolation */
+	RCCF_Linear UMETA(DisplayName = "Linear"),
+
+	/** All keys use cubic interpolation */
+	RCCF_Cubic UMETA(DisplayName = "Cubic"),
+
+	/** Keys use mixed interpolation modes */
+	RCCF_Mixed UMETA(DisplayName = "Mixed"),
+};
+
+/** Enumerates key time compression options. */
+UENUM()
+enum ERichCurveKeyTimeCompressionFormat
+{
+	/** Key time is quantized to 16 bits */
+	RCKTCF_uint16 UMETA(DisplayName = "uint16"),
+
+	/** Key time uses full precision */
+	RCKTCF_float32 UMETA(DisplayName = "float32"),
+};
+
 /** One key in a rich, editable float curve */
 USTRUCT()
 struct ENGINE_API FRichCurveKey
@@ -212,7 +243,7 @@ public:
 	 */
 	bool HasAnyData() const
 	{
-		return DefaultValue != MAX_flt || Keys.Num();
+		return DefaultValue != MAX_flt || Keys.Num() != 0;
 	}
 
 	/** Gets a copy of the keys, so indices and handles can't be meddled with */
@@ -294,6 +325,12 @@ public:
 	/** Removes the default value for this curve. */
 	void ClearDefaultValue() { DefaultValue = MAX_flt; }
 
+	/** Returns whether the curve is constant or not */
+	bool IsConstant(float Tolerance = SMALL_NUMBER) const;
+
+	/** Returns whether the curve is empty or not */
+	bool IsEmpty() const { return Keys.Num() == 0; }
+
 	/** Shifts all keys forwards or backwards in time by an even amount, preserving order */
 	void ShiftCurve(float DeltaTime);
 	void ShiftCurve(float DeltaTime, TSet<FKeyHandle>& KeyHandles);
@@ -349,6 +386,9 @@ public:
 	void RemoveRedundantKeys(float Tolerance);
 	void RemoveRedundantKeys(float Tolerance, float FirstKeyTime, float LastKeyTime);
 
+	/** Compresses a rich curve for efficient runtime storage and evaluation */
+	void CompressCurve(struct FCompressedRichCurve& OutCurve, float ErrorThreshold = 0.0001f, float SampleRate = 120.0f) const;
+
 private:
 	void RemoveRedundantKeysInternal(float Tolerance, int32 InStartKeepKey, int32 InEndKeepKey);
 	int32 FindKeyInternal(float KeyTime, float KeyTimeTolerance) const;
@@ -379,6 +419,85 @@ public:
 	TArray<FRichCurveKey> Keys;
 };
 
+/**
+ * A runtime optimized representation of a FRichCurve. It consumes less memory and evaluates faster.
+ */
+USTRUCT()
+struct ENGINE_API FCompressedRichCurve
+{
+	GENERATED_USTRUCT_BODY()
+
+	/** Compression format used by CompressedKeys */
+	TEnumAsByte<ERichCurveCompressionFormat> CompressionFormat;
+
+	/** Compression format used to pack the key time */
+	TEnumAsByte<ERichCurveKeyTimeCompressionFormat> KeyTimeCompressionFormat;
+
+	/** Pre-infinity extrapolation state */
+	TEnumAsByte<ERichCurveExtrapolation> PreInfinityExtrap;
+
+	/** Post-infinity extrapolation state */
+	TEnumAsByte<ERichCurveExtrapolation> PostInfinityExtrap;
+
+	union TConstantValueNumKeys
+	{
+		float ConstantValue;
+		int32 NumKeys;
+
+		TConstantValueNumKeys() : NumKeys(0) {}
+	};
+
+	/**
+	* If the compression format is constant, this is the value returned
+	* Inline here to reduce the likelihood of accessing the compressed keys data for the common case of constant/zero/empty curves
+	* When a curve is linear/cubic/mixed, the constant float value isn't used and instead we use the number of keys
+	*/
+	TConstantValueNumKeys ConstantValueNumKeys;
+
+	/** Compressed keys, used only outside of the editor */
+	TArray<uint8> CompressedKeys;
+
+	FCompressedRichCurve()
+		: CompressionFormat(RCCF_Empty)
+		, KeyTimeCompressionFormat(RCKTCF_float32)
+		, PreInfinityExtrap(RCCE_None)
+		, PostInfinityExtrap(RCCE_None)
+		, ConstantValueNumKeys()
+		, CompressedKeys()
+	{}
+
+	/** Evaluate this rich curve at the specified time */
+	float Eval(float InTime, float InDefaultValue = 0.0f) const;
+
+	/** Evaluate this rich curve at the specified time */
+	static float StaticEval(ERichCurveCompressionFormat CompressionFormat, ERichCurveKeyTimeCompressionFormat KeyTimeCompressionFormat, ERichCurveExtrapolation PreInfinityExtrap, ERichCurveExtrapolation PostInfinityExtrap, TConstantValueNumKeys ConstantValueNumKeys, const uint8* CompressedKeys, float InTime, float InDefaultValue = 0.0f);
+
+	/** ICPPStructOps interface */
+	bool Serialize(FArchive& Ar);
+	bool operator==(const FCompressedRichCurve& Other) const;
+	bool operator!=(const FCompressedRichCurve& Other) const { return !(*this == Other); }
+
+	friend FArchive& operator<<(FArchive& Ar, FCompressedRichCurve& Curve)
+	{
+		Curve.Serialize(Ar);
+		return Ar;
+	}
+};
+
+/*
+ * Override serialization for compressed rich curves to handle the union
+ */
+template<>
+struct TStructOpsTypeTraits<FCompressedRichCurve>
+	: public TStructOpsTypeTraitsBase2<FCompressedRichCurve>
+{
+	enum
+	{
+		WithSerializer = true,
+		WithCopy = false,
+		WithIdenticalViaEquality = true,
+	};
+};
 
 /**
  * Info about a curve to be edited.
diff --git a/Engine/Source/Runtime/Engine/Private/Animation/AnimCompressionDerivedData.cpp b/Engine/Source/Runtime/Engine/Private/Animation/AnimCompressionDerivedData.cpp
--- a/Engine/Source/Runtime/Engine/Private/Animation/AnimCompressionDerivedData.cpp
+++ b/Engine/Source/Runtime/Engine/Private/Animation/AnimCompressionDerivedData.cpp
@@ -6,6 +6,7 @@
 #include "AnimationUtils.h"
 #include "AnimEncoding.h"
 #include "Animation/AnimCompress.h"
+#include "Animation/AnimCurveCompressionSettings.h"
 #include "AnimationCompression.h"
 #include "UObject/Package.h"
 
@@ -70,6 +71,7 @@ FString FDerivedDataAnimationCompression::GetPluginSpecificCacheKeySuffix() cons
 	//	* Baked Additive Flag
 	//	* Additive ref pose GUID or hardcoded string if not available
 	//	* Compression Settings
+	//	* Curve compression settings
 
 	uint8 AdditiveSettings = bCanBakeAdditive ? (OriginalAnimSequence->RefPoseType << 4) + OriginalAnimSequence->AdditiveAnimType : 0;
 
@@ -78,7 +80,7 @@ FString FDerivedDataAnimationCompression::GetPluginSpecificCacheKeySuffix() cons
 
 	const int32 StripFrame = bPerformStripping ? 1 : 0;
 
-	FString Ret = FString::Printf(TEXT("%i_%i_%i_%i_%s%s%s_%c%c%i_%s_%s"),
+	FString Ret = FString::Printf(TEXT("%i_%i_%i_%i_%s%s%s_%c%c%i_%s_%s_%s"),
 		(int32)UE_ANIMCOMPRESSION_DERIVEDDATA_VER,
 		(int32)CURRENT_ANIMATION_ENCODING_PACKAGE_VERSION,
 		OriginalAnimSequence->CompressCommandletVersion,
@@ -90,7 +92,8 @@ FString FDerivedDataAnimationCompression::GetPluginSpecificCacheKeySuffix() cons
 		RefType,
 		OriginalAnimSequence->RefFrameIndex,
 		(bCanBakeAdditive && AdditiveBase) ? *AdditiveBase->GetRawDataGuid().ToString() : TEXT("NoAdditiveBase"),
-		*OriginalAnimSequence->CompressionScheme->MakeDDCKey()
+		*OriginalAnimSequence->CompressionScheme->MakeDDCKey(),
+		*OriginalAnimSequence->CurveCompressionSettings->MakeDDCKey()
 		);
 
 	return Ret;
@@ -155,13 +158,9 @@ bool FDerivedDataAnimationCompression::Build( TArray<uint8>& OutData )
 		}
 
 		AnimToOperateOn->UpdateCompressedTrackMapFromRaw();
-		AnimToOperateOn->CompressedCurveData = AnimToOperateOn->RawCurveData; //Curves don't actually get compressed, but could have additives baked in
+		AnimToOperateOn->UpdateCompressedCurveNames();
 
-		const float MaxCurveError = AnimToOperateOn->CompressionScheme->MaxCurveError;
-		for (FFloatCurve& Curve : AnimToOperateOn->CompressedCurveData.FloatCurves)
-		{
-			Curve.FloatCurve.RemoveRedundantKeys(MaxCurveError);
-		}
+		bool bCurveCompressionSuccess = FAnimationUtils::CompressAnimCurves(*AnimToOperateOn);
 
 #if DO_CHECK
 		FString CompressionName = AnimToOperateOn->CompressionScheme->GetFullName();
@@ -172,7 +171,7 @@ bool FDerivedDataAnimationCompression::Build( TArray<uint8>& OutData )
 		AnimToOperateOn->UpdateCompressedNumFramesFromRaw(); //Do this before compression so compress code can read the correct value
 
 		FAnimationUtils::CompressAnimSequence(AnimToOperateOn, *CompressContext.Get());
-		bCompressionSuccessful = AnimToOperateOn->IsCompressedDataValid();
+		bCompressionSuccessful = AnimToOperateOn->IsCompressedDataValid() && bCurveCompressionSuccess;
 
 		ensureMsgf(bCompressionSuccessful, TEXT("Anim Compression failed for Sequence '%s' with compression scheme '%s': compressed data empty\n\tAnimIndex: %i\n\tMaxAnim:%i\n\tAllowAltCompressor:%s\n\tOutput:%s"), 
 											*AnimToOperateOn->GetFullName(), 
@@ -191,6 +190,7 @@ bool FDerivedDataAnimationCompression::Build( TArray<uint8>& OutData )
 	{
 		CA_SUPPRESS(6011); // See https://connect.microsoft.com/VisualStudio/feedback/details/3007725
 		OriginalAnimSequence->CompressionScheme = static_cast<UAnimCompress*>(StaticDuplicateObject(AnimToOperateOn->CompressionScheme, OriginalAnimSequence));
+		OriginalAnimSequence->CurveCompressionSettings = AnimToOperateOn->CurveCompressionSettings;
 	}
 
 	if (bCompressionSuccessful)
diff --git a/Engine/Source/Runtime/Engine/Private/Animation/AnimCompressionDerivedData.h b/Engine/Source/Runtime/Engine/Private/Animation/AnimCompressionDerivedData.h
--- a/Engine/Source/Runtime/Engine/Private/Animation/AnimCompressionDerivedData.h
+++ b/Engine/Source/Runtime/Engine/Private/Animation/AnimCompressionDerivedData.h
@@ -47,7 +47,7 @@ public:
 		// This is a version string that mimics the old versioning scheme. If you
 		// want to bump this version, generate a new guid using VS->Tools->Create GUID and
 		// return it here. Ex.
-		return TEXT("EFCFC622A8794B758D53CDE253471CBD");
+		return TEXT("D58C7C3E48274C04BFC4D7823AF46F5E");
 	}
 
 	virtual FString GetPluginSpecificCacheKeySuffix() const override;
diff --git a/Engine/Source/Runtime/Engine/Private/Animation/AnimCurveCompressionCodec.cpp b/Engine/Source/Runtime/Engine/Private/Animation/AnimCurveCompressionCodec.cpp
new file mode 100644
--- /dev/null
+++ b/Engine/Source/Runtime/Engine/Private/Animation/AnimCurveCompressionCodec.cpp
@@ -0,0 +1,42 @@
+// Copyright 1998-2018 Epic Games, Inc. All Rights Reserved.
+
+#include "Animation/AnimCurveCompressionCodec.h"
+
+UAnimCurveCompressionCodec::UAnimCurveCompressionCodec(const FObjectInitializer& ObjectInitializer)
+	: Super(ObjectInitializer)
+{
+}
+
+#if WITH_EDITORONLY_DATA
+void UAnimCurveCompressionCodec::PostInitProperties()
+{
+	Super::PostInitProperties();
+
+	if (!IsTemplate())
+	{
+		InstanceGuid = FGuid::NewGuid();
+	}
+}
+
+void UAnimCurveCompressionCodec::PostDuplicate(bool bDuplicateForPIE)
+{
+	Super::PostDuplicate(bDuplicateForPIE);
+
+	InstanceGuid = FGuid::NewGuid();
+}
+
+void UAnimCurveCompressionCodec::Serialize(FArchive& Ar)
+{
+	Super::Serialize(Ar);
+
+	if (!Ar.IsCooking())
+	{
+		Ar << InstanceGuid;
+	}
+}
+
+void UAnimCurveCompressionCodec::PopulateDDCKey(FArchive& Ar)
+{
+	Ar << InstanceGuid;
+}
+#endif
diff --git a/Engine/Source/Runtime/Engine/Private/Animation/AnimCurveCompressionCodec_CompressedRichCurve.cpp b/Engine/Source/Runtime/Engine/Private/Animation/AnimCurveCompressionCodec_CompressedRichCurve.cpp
new file mode 100644
--- /dev/null
+++ b/Engine/Source/Runtime/Engine/Private/Animation/AnimCurveCompressionCodec_CompressedRichCurve.cpp
@@ -0,0 +1,137 @@
+// Copyright 1998-2018 Epic Games, Inc. All Rights Reserved.
+
+#include "Animation/AnimCurveCompressionCodec_CompressedRichCurve.h"
+#include "Animation/AnimSequence.h"
+#include "Serialization/MemoryWriter.h"
+
+UAnimCurveCompressionCodec_CompressedRichCurve::UAnimCurveCompressionCodec_CompressedRichCurve(const FObjectInitializer& ObjectInitializer)
+	: Super(ObjectInitializer)
+{
+#if WITH_EDITORONLY_DATA
+	MaxCurveError = 0.0f;
+	UseAnimSequenceSampleRate = true;
+	ErrorSampleRate = 60.0f;
+#endif
+}
+
+// This mirrors in part the FCompressedRichCurve
+struct FCurveDesc
+{
+	TEnumAsByte<ERichCurveCompressionFormat> CompressionFormat;
+	TEnumAsByte<ERichCurveKeyTimeCompressionFormat> KeyTimeCompressionFormat;
+	TEnumAsByte<ERichCurveExtrapolation> PreInfinityExtrap;
+	TEnumAsByte<ERichCurveExtrapolation> PostInfinityExtrap;
+	FCompressedRichCurve::TConstantValueNumKeys ConstantValueNumKeys;
+	int32 KeyDataOffset;
+};
+
+#if WITH_EDITORONLY_DATA
+bool UAnimCurveCompressionCodec_CompressedRichCurve::Compress(const UAnimSequence& AnimSeq, FAnimCurveCompressionResult& OutResult)
+{
+	int32 NumCurves = AnimSeq.RawCurveData.FloatCurves.Num();
+
+	TArray<FCurveDesc> Curves;
+	Curves.Reserve(NumCurves);
+	Curves.AddUninitialized(NumCurves);
+
+	int32 KeyDataOffset = 0;
+	KeyDataOffset += sizeof(FCurveDesc) * NumCurves;
+
+	const FAnimKeyHelper Helper(AnimSeq.SequenceLength, AnimSeq.GetRawNumberOfFrames());
+	const float SampleRate = UseAnimSequenceSampleRate ? Helper.KeysPerSecond() : ErrorSampleRate;
+
+	TArray<uint8> KeyData;
+
+	for (int32 CurveIndex = 0; CurveIndex < NumCurves; ++CurveIndex)
+	{
+		const FFloatCurve& Curve = AnimSeq.RawCurveData.FloatCurves[CurveIndex];
+
+		FRichCurve RawCurve = Curve.FloatCurve;	// Copy
+		RawCurve.RemoveRedundantKeys(MaxCurveError);
+
+		FCompressedRichCurve CompressedCurve;
+		RawCurve.CompressCurve(CompressedCurve, MaxCurveError, SampleRate);
+
+		FCurveDesc& CurveDesc = Curves[CurveIndex];
+		CurveDesc.CompressionFormat = CompressedCurve.CompressionFormat;
+		CurveDesc.KeyTimeCompressionFormat = CompressedCurve.KeyTimeCompressionFormat;
+		CurveDesc.PreInfinityExtrap = CompressedCurve.PreInfinityExtrap;
+		CurveDesc.PostInfinityExtrap = CompressedCurve.PostInfinityExtrap;
+		CurveDesc.ConstantValueNumKeys = CompressedCurve.ConstantValueNumKeys;
+		CurveDesc.KeyDataOffset = KeyDataOffset;
+
+		KeyDataOffset += CompressedCurve.CompressedKeys.Num();
+		KeyData.Append(CompressedCurve.CompressedKeys);
+	}
+
+	TArray<uint8> TempBytes;
+	TempBytes.Reserve(KeyDataOffset);
+
+	// Serialize the compression settings into a temporary array. The archive
+	// is flagged as persistent so that machines of different endianness produce
+	// identical binary results.
+	FMemoryWriter Ar(TempBytes, /*bIsPersistent=*/ true);
+
+	Ar.Serialize(Curves.GetData(), sizeof(FCurveDesc) * NumCurves);
+	Ar.Serialize(KeyData.GetData(), KeyData.Num());
+
+	OutResult.CompressedBytes = TempBytes;
+	OutResult.Codec = this;
+
+	return true;
+}
+
+void UAnimCurveCompressionCodec_CompressedRichCurve::PopulateDDCKey(FArchive& Ar)
+{
+	Super::PopulateDDCKey(Ar);
+
+	int32 CodecVersion = 0;
+
+	Ar << CodecVersion;
+	Ar << MaxCurveError;
+	Ar << UseAnimSequenceSampleRate;
+	Ar << ErrorSampleRate;
+}
+#endif
+
+void UAnimCurveCompressionCodec_CompressedRichCurve::DecompressCurves(const UAnimSequence& AnimSeq, FBlendedCurve& Curves, float CurrentTime) const
+{
+	const uint8* Buffer = AnimSeq.CompressedCurveByteStream.GetData();
+	const FCurveDesc* CurveDescriptions = (const FCurveDesc*)(Buffer);
+
+	const TArray<FSmartName>& CompressedCurveNames = AnimSeq.GetCompressedCurveNames();
+	const int32 NumCurves = CompressedCurveNames.Num();
+	for (int32 CurveIndex = 0; CurveIndex < NumCurves; ++CurveIndex)
+	{
+		const FSmartName& CurveName = CompressedCurveNames[CurveIndex];
+		if (Curves.IsEnabled(CurveName.UID))
+		{
+			const FCurveDesc& Curve = CurveDescriptions[CurveIndex];
+			const uint8* CompressedKeys = Buffer + Curve.KeyDataOffset;
+			const float Value = FCompressedRichCurve::StaticEval(Curve.CompressionFormat, Curve.KeyTimeCompressionFormat, Curve.PreInfinityExtrap, Curve.PostInfinityExtrap, Curve.ConstantValueNumKeys, CompressedKeys, CurrentTime);
+			Curves.Set(CurveName.UID, Value);
+		}
+	}
+}
+
+float UAnimCurveCompressionCodec_CompressedRichCurve::DecompressCurve(const UAnimSequence& AnimSeq, SmartName::UID_Type CurveUID, float CurrentTime) const
+{
+	const uint8* Buffer = AnimSeq.CompressedCurveByteStream.GetData();
+	const FCurveDesc* CurveDescriptions = (const FCurveDesc*)(Buffer);
+
+	const TArray<FSmartName>& CompressedCurveNames = AnimSeq.GetCompressedCurveNames();
+	const int32 NumCurves = CompressedCurveNames.Num();
+	for (int32 CurveIndex = 0; CurveIndex < NumCurves; ++CurveIndex)
+	{
+		const FSmartName& CurveName = CompressedCurveNames[CurveIndex];
+		if (CurveName.UID == CurveUID)
+		{
+			const FCurveDesc& Curve = CurveDescriptions[CurveIndex];
+			const uint8* CompressedKeys = Buffer + Curve.KeyDataOffset;
+			const float Value = FCompressedRichCurve::StaticEval(Curve.CompressionFormat, Curve.KeyTimeCompressionFormat, Curve.PreInfinityExtrap, Curve.PostInfinityExtrap, Curve.ConstantValueNumKeys, CompressedKeys, CurrentTime);
+			return Value;
+		}
+	}
+
+	return 0.0f;
+}
diff --git a/Engine/Source/Runtime/Engine/Private/Animation/AnimCurveCompressionCodec_UniformlySampled.cpp b/Engine/Source/Runtime/Engine/Private/Animation/AnimCurveCompressionCodec_UniformlySampled.cpp
new file mode 100644
--- /dev/null
+++ b/Engine/Source/Runtime/Engine/Private/Animation/AnimCurveCompressionCodec_UniformlySampled.cpp
@@ -0,0 +1,293 @@
+// Copyright 1998-2018 Epic Games, Inc. All Rights Reserved.
+
+#include "Animation/AnimCurveCompressionCodec_UniformlySampled.h"
+#include "Animation/AnimSequence.h"
+#include "Serialization/MemoryWriter.h"
+
+UAnimCurveCompressionCodec_UniformlySampled::UAnimCurveCompressionCodec_UniformlySampled(const FObjectInitializer& ObjectInitializer)
+	: Super(ObjectInitializer)
+{
+#if WITH_EDITORONLY_DATA
+	UseAnimSequenceSampleRate = true;
+	SampleRate = 30.0f;
+#endif
+}
+
+#if WITH_EDITORONLY_DATA
+bool UAnimCurveCompressionCodec_UniformlySampled::Compress(const UAnimSequence& AnimSeq, FAnimCurveCompressionResult& OutResult)
+{
+	const int32 NumCurves = AnimSeq.RawCurveData.FloatCurves.Num();
+	const float Duration = AnimSeq.SequenceLength;
+
+	int32 NumSamples;
+	float SampleRate_;
+	if (UseAnimSequenceSampleRate)
+	{
+		const FAnimKeyHelper Helper(AnimSeq.SequenceLength, AnimSeq.GetRawNumberOfFrames());
+		SampleRate_ = Helper.KeysPerSecond();
+		NumSamples = FMath::RoundToInt(Duration * SampleRate_) + 1;
+	}
+	else
+	{
+		// If our duration isn't an exact multiple of the sample rate, we'll round
+		// and end up with a sample rate slightly corrected to make sure we spread
+		// the resulting error over the whole duration
+		NumSamples = FMath::RoundToInt(Duration * SampleRate) + 1;
+		SampleRate_ = (NumSamples - 1) / Duration;
+	}
+
+	int32 NumConstantCurves = 0;
+	for (const FFloatCurve& Curve : AnimSeq.RawCurveData.FloatCurves)
+	{
+		if (Curve.FloatCurve.IsConstant())
+		{
+			NumConstantCurves++;
+		}
+	}
+
+	const int32 NumAnimatedCurves = NumCurves - NumConstantCurves;
+
+	// 1 bit per curve, round up
+	const int32 ConstantBitsetSize = sizeof(uint32) * ((NumCurves + 31) / 32);
+
+	int32 BufferSize = 0;
+	BufferSize += sizeof(int32);	// NumConstantCurves
+	BufferSize += sizeof(int32);	// NumSamples
+	BufferSize += sizeof(float);	// SampleRate
+	BufferSize += ConstantBitsetSize;								// Constant curve bitset
+	BufferSize += sizeof(float) * NumConstantCurves;				// Constant curve samples
+	BufferSize += sizeof(float) * NumAnimatedCurves * NumSamples;	// Animated curve samples
+
+	TArray<uint8> Buffer;
+	Buffer.Reserve(BufferSize);
+	Buffer.AddUninitialized(BufferSize);
+
+	int32 BufferOffset = 0;
+	int32* NumConstantCurvesPtr = (int32*)&Buffer[BufferOffset];
+	BufferOffset += sizeof(int32);
+
+	int32* NumSamplesPtr = (int32*)&Buffer[BufferOffset];
+	BufferOffset += sizeof(int32);
+
+	float* SampleRatePtr = (float*)&Buffer[BufferOffset];
+	BufferOffset += sizeof(float);
+
+	*NumConstantCurvesPtr = NumConstantCurves;
+	*NumSamplesPtr = NumSamples;
+	*SampleRatePtr = SampleRate_;
+
+	if (NumCurves > 0 && NumSamples > 0)
+	{
+		uint32* ConstantCurvesBitsetPtr = (uint32*)&Buffer[BufferOffset];
+		BufferOffset += ConstantBitsetSize;
+		FMemory::Memzero(ConstantCurvesBitsetPtr, ConstantBitsetSize);
+
+		float* ConstantSamplesPtr = NumConstantCurves > 0 ? (float*)&Buffer[BufferOffset] : nullptr;
+		BufferOffset += sizeof(float) * NumConstantCurves;
+
+		float* AnimatedSamplesPtr = NumAnimatedCurves > 0 ? (float*)&Buffer[BufferOffset] : nullptr;
+		BufferOffset += sizeof(float) * NumAnimatedCurves * NumSamples;
+
+		for (int32 CurveIndex = 0, ConstantCurveIndex = 0; CurveIndex < NumCurves; ++CurveIndex)
+		{
+			const FFloatCurve& Curve = AnimSeq.RawCurveData.FloatCurves[CurveIndex];
+			if (Curve.FloatCurve.IsConstant())
+			{
+				if (Curve.FloatCurve.IsEmpty())
+				{
+					ConstantSamplesPtr[ConstantCurveIndex] = Curve.FloatCurve.DefaultValue;
+				}
+				else
+				{
+					ConstantSamplesPtr[ConstantCurveIndex] = Curve.FloatCurve.Keys[0].Value;
+				}
+
+				// Bitset uses little-endian bit ordering
+				ConstantCurvesBitsetPtr[CurveIndex / 32] |= 1 << (CurveIndex % 32);
+				ConstantCurveIndex++;
+			}
+		}
+
+		// Write out samples sorted by time first in order to have everything contiguous in memory
+		// for improved cache locality
+		// Curve 0 Key 0, Curve 0 Key 1, Curve 0 Key N, Curve 1 Key 0, Curve 1 Key 1, Curve 1 Key N, Curve M Key 0, ...
+		const float InvSampleRate = 1.0f / SampleRate_;
+		for (int32 SampleIndex = 0; SampleIndex < NumSamples; ++SampleIndex)
+		{
+			const float SampleTime = FMath::Clamp(SampleIndex * InvSampleRate, 0.0f, Duration);
+			float* AnimatedSamples = AnimatedSamplesPtr + (SampleIndex * NumAnimatedCurves);
+
+			for (int32 CurveIndex = 0, AnimatedCurveIndex = 0; CurveIndex < NumCurves; ++CurveIndex)
+			{
+				const FFloatCurve& Curve = AnimSeq.RawCurveData.FloatCurves[CurveIndex];
+				if (Curve.FloatCurve.IsConstant())
+				{
+					// Skip constant curves, their data has already been written
+					continue;
+				}
+
+				const float SampleValue = Curve.FloatCurve.Eval(SampleTime);
+
+				AnimatedSamples[AnimatedCurveIndex] = SampleValue;
+				AnimatedCurveIndex++;
+			}
+		}
+	}
+
+	check(BufferOffset == BufferSize);
+
+	OutResult.CompressedBytes = Buffer;
+	OutResult.Codec = this;
+
+	return true;
+}
+
+void UAnimCurveCompressionCodec_UniformlySampled::PopulateDDCKey(FArchive& Ar)
+{
+	Super::PopulateDDCKey(Ar);
+
+	int32 CodecVersion = 0;
+
+	Ar << CodecVersion;
+	Ar << UseAnimSequenceSampleRate;
+	Ar << SampleRate;
+}
+#endif
+
+void UAnimCurveCompressionCodec_UniformlySampled::DecompressCurves(const UAnimSequence& AnimSeq, FBlendedCurve& Curves, float CurrentTime) const
+{
+	const TArray<FSmartName>& CompressedCurveNames = AnimSeq.GetCompressedCurveNames();
+	const int32 NumCurves = CompressedCurveNames.Num();
+
+	if (NumCurves == 0)
+	{
+		return;
+	}
+
+	const uint8* Buffer = AnimSeq.CompressedCurveByteStream.GetData();
+
+	int32 BufferOffset = 0;
+	const int32 NumConstantCurves = *(const int32*)&Buffer[BufferOffset];
+	BufferOffset += sizeof(int32);
+
+	const int32 NumSamples = *(const int32*)&Buffer[BufferOffset];
+	BufferOffset += sizeof(int32);
+
+	if (NumSamples == 0)
+	{
+		return;
+	}
+
+	const float SampleRate_ = *(const float*)&Buffer[BufferOffset];
+	BufferOffset += sizeof(float);
+
+	const uint32* ConstantCurvesBitsetPtr = (const uint32*)&Buffer[BufferOffset];
+	BufferOffset += sizeof(uint32) * NumCurves;
+
+	const int32 ConstantBitsetSize = sizeof(uint32) * ((NumCurves + 31) / 32);
+	const float* ConstantSamplesPtr = (const float*)&Buffer[BufferOffset];
+	BufferOffset += ConstantBitsetSize;
+
+	const float* AnimatedSamplesPtr = (const float*)&Buffer[BufferOffset];
+	//const int32 NumAnimatedCurves = NumCurves - NumConstantCurves;
+	//BufferOffset += sizeof(float) * NumAnimatedCurves * NumSamples;
+
+	const float SamplePoint = CurrentTime * SampleRate_;
+	const int32 SampleIndex0 = FMath::Clamp(FMath::FloorToInt(SamplePoint), 0, NumSamples - 1);
+	const int32 SampleIndex1 = FMath::Min(SampleIndex0 + 1, NumSamples - 1);
+	const float InterpolationAlpha = SamplePoint - float(SampleIndex0);
+	const int32 NumAnimatedCurves = NumCurves - NumConstantCurves;
+
+	const float* AnimatedSamples0 = AnimatedSamplesPtr + (SampleIndex0 * NumAnimatedCurves);
+	const float* AnimatedSamples1 = AnimatedSamplesPtr + (SampleIndex1 * NumAnimatedCurves);
+
+	for (int32 CurveIndex = 0, ConstantCurveIndex = 0, AnimatedCurveIndex = 0; CurveIndex < NumCurves; ++CurveIndex)
+	{
+		const FSmartName& CurveName = CompressedCurveNames[CurveIndex];
+		const bool bIsConstant = (ConstantCurvesBitsetPtr[CurveIndex / 32] & (1 << (CurveIndex % 32))) != 0;
+		if (Curves.IsEnabled(CurveName.UID))
+		{
+			float Sample;
+			if (bIsConstant)
+			{
+				Sample = ConstantSamplesPtr[ConstantCurveIndex];
+			}
+			else
+			{
+				const float Sample0 = AnimatedSamples0[AnimatedCurveIndex];
+				const float Sample1 = AnimatedSamples1[AnimatedCurveIndex];
+				Sample = FMath::Lerp(Sample0, Sample1, InterpolationAlpha);
+			}
+
+			Curves.Set(CurveName.UID, Sample);
+		}
+
+		(bIsConstant ? ConstantCurveIndex : AnimatedCurveIndex)++;
+	}
+}
+
+float UAnimCurveCompressionCodec_UniformlySampled::DecompressCurve(const UAnimSequence& AnimSeq, SmartName::UID_Type CurveUID, float CurrentTime) const
+{
+	const TArray<FSmartName>& CompressedCurveNames = AnimSeq.GetCompressedCurveNames();
+	const int32 NumCurves = CompressedCurveNames.Num();
+
+	if (NumCurves == 0)
+	{
+		return 0.0f;
+	}
+
+	const uint8* Buffer = AnimSeq.CompressedCurveByteStream.GetData();
+
+	int32 BufferOffset = 0;
+	const int32 NumConstantCurves = *(const int32*)&Buffer[BufferOffset];
+	BufferOffset += sizeof(int32);
+
+	const int32 NumSamples = *(const int32*)&Buffer[BufferOffset];
+	BufferOffset += sizeof(int32);
+
+	const float SampleRate_ = *(const float*)&Buffer[BufferOffset];
+	BufferOffset += sizeof(float);
+
+	const uint32* ConstantCurvesBitsetPtr = (const uint32*)&Buffer[BufferOffset];
+	BufferOffset += sizeof(uint32) * NumCurves;
+
+	const int32 ConstantBitsetSize = sizeof(uint32) * ((NumCurves + 31) / 32);
+	const float* ConstantSamplesPtr = (const float*)&Buffer[BufferOffset];
+	BufferOffset += ConstantBitsetSize;
+
+	const float* AnimatedSamplesPtr = (const float*)&Buffer[BufferOffset];
+	//const int32 NumAnimatedCurves = NumCurves - NumConstantCurves;
+	//BufferOffset += sizeof(float) * NumAnimatedCurves * NumSamples;
+
+	for (int32 CurveIndex = 0, ConstantCurveIndex = 0, AnimatedCurveIndex = 0; CurveIndex < NumCurves; ++CurveIndex)
+	{
+		const FSmartName& CurveName = CompressedCurveNames[CurveIndex];
+		const bool bIsConstant = (ConstantCurvesBitsetPtr[CurveIndex / 32] & (1 << (CurveIndex % 32))) != 0;
+		if (CurveName.UID == CurveUID)
+		{
+			float Sample;
+			if (bIsConstant)
+			{
+				Sample = ConstantSamplesPtr[ConstantCurveIndex];
+			}
+			else
+			{
+				const float SamplePoint = CurrentTime * SampleRate_;
+				const int32 SampleIndex0 = FMath::Clamp(FMath::FloorToInt(SamplePoint), 0, NumSamples - 1);
+				const int32 SampleIndex1 = FMath::Min(SampleIndex0 + 1, NumSamples - 1);
+				const float InterpolationAlpha = SamplePoint - float(SampleIndex0);
+				const int32 NumAnimatedCurves = NumCurves - NumConstantCurves;
+
+				const float Sample0 = AnimatedSamplesPtr[(SampleIndex0 * NumAnimatedCurves) + CurveIndex];
+				const float Sample1 = AnimatedSamplesPtr[(SampleIndex1 * NumAnimatedCurves) + CurveIndex];
+				Sample = FMath::Lerp(Sample0, Sample1, InterpolationAlpha);
+			}
+
+			return Sample;
+		}
+
+		(bIsConstant ? ConstantCurveIndex : AnimatedCurveIndex)++;
+	}
+
+	return 0.0f;
+}
diff --git a/Engine/Source/Runtime/Engine/Private/Animation/AnimCurveCompressionSettings.cpp b/Engine/Source/Runtime/Engine/Private/Animation/AnimCurveCompressionSettings.cpp
new file mode 100644
--- /dev/null
+++ b/Engine/Source/Runtime/Engine/Private/Animation/AnimCurveCompressionSettings.cpp
@@ -0,0 +1,74 @@
+// Copyright 1998-2018 Epic Games, Inc. All Rights Reserved.
+
+#include "Animation/AnimCurveCompressionSettings.h"
+#include "Animation/AnimCurveCompressionCodec_CompressedRichCurve.h"
+#include "Animation/AnimSequence.h"
+#include "Serialization/MemoryWriter.h"
+
+UAnimCurveCompressionSettings::UAnimCurveCompressionSettings(const FObjectInitializer& ObjectInitializer)
+	: Super(ObjectInitializer)
+{
+}
+
+#if WITH_EDITORONLY_DATA
+void UAnimCurveCompressionSettings::PostInitProperties()
+{
+	Super::PostInitProperties();
+
+	if (!IsTemplate())
+	{
+		// Ensure we are never null
+		Codec = NewObject<UAnimCurveCompressionCodec_CompressedRichCurve>(this, NAME_None, RF_Public);
+	}
+}
+
+bool UAnimCurveCompressionSettings::AreSettingsValid() const
+{
+	return Codec != nullptr && Codec->IsCodecValid();
+}
+
+bool UAnimCurveCompressionSettings::Compress(UAnimSequence& AnimSeq) const
+{
+	if (Codec == nullptr || !AreSettingsValid())
+	{
+		return false;
+	}
+
+	FAnimCurveCompressionResult CompressionResult;
+	bool Success = Codec->Compress(AnimSeq, CompressionResult);
+	if (Success)
+	{
+		AnimSeq.CompressedCurveByteStream = CompressionResult.CompressedBytes;
+		AnimSeq.CurveCompressionCodec = CompressionResult.Codec;
+	}
+
+	return Success;
+}
+
+FString UAnimCurveCompressionSettings::MakeDDCKey() const
+{
+	if (Codec == nullptr)
+	{
+		return TEXT("<Missing Codec>");
+	}
+
+	TArray<uint8> TempBytes;
+	TempBytes.Reserve(64);
+
+	// Serialize the compression settings into a temporary array. The archive
+	// is flagged as persistent so that machines of different endianness produce
+	// identical binary results.
+	FMemoryWriter Ar(TempBytes, /*bIsPersistent=*/ true);
+
+	Codec->PopulateDDCKey(Ar);
+
+	FString Key;
+	Key.Reserve(TempBytes.Num() + 1);
+	for (int32 ByteIndex = 0; ByteIndex < TempBytes.Num(); ++ByteIndex)
+	{
+		ByteToHex(TempBytes[ByteIndex], Key);
+	}
+
+	return Key;
+}
+#endif
diff --git a/Engine/Source/Runtime/Engine/Private/Animation/AnimCurveTypes.cpp b/Engine/Source/Runtime/Engine/Private/Animation/AnimCurveTypes.cpp
--- a/Engine/Source/Runtime/Engine/Private/Animation/AnimCurveTypes.cpp
+++ b/Engine/Source/Runtime/Engine/Private/Animation/AnimCurveTypes.cpp
@@ -3,7 +3,7 @@
 #include "Animation/AnimCurveTypes.h"
 #include "UObject/FrameworkObjectVersion.h"
 
-DECLARE_CYCLE_STAT(TEXT("AnimSeq EvalCurveData"), STAT_AnimSeq_EvalCurveData, STATGROUP_Anim);
+DECLARE_CYCLE_STAT(TEXT("EvalRawCurveData"), STAT_EvalRawCurveData, STATGROUP_Anim);
 
 /////////////////////////////////////////////////////
 // FFloatCurve
@@ -293,7 +293,7 @@ void FTransformCurve::Resize(float NewLength, bool bInsert/* whether insert or r
 
 void FRawCurveTracks::EvaluateCurveData( FBlendedCurve& Curves, float CurrentTime ) const
 {
-	SCOPE_CYCLE_COUNTER(STAT_AnimSeq_EvalCurveData);
+	SCOPE_CYCLE_COUNTER(STAT_EvalRawCurveData);
 	if (Curves.NumValidCurveCount > 0)
 	{
 		// evaluate the curve data at the CurrentTime and add to Instance
@@ -302,7 +302,8 @@ void FRawCurveTracks::EvaluateCurveData( FBlendedCurve& Curves, float CurrentTim
 			const FFloatCurve& Curve = *CurveIter;
 			if (Curves.IsEnabled(Curve.Name.UID))
 			{
-				Curves.Set(Curve.Name.UID, Curve.Evaluate(CurrentTime));
+				float Value = Curve.Evaluate(CurrentTime);
+				Curves.Set(Curve.Name.UID, Value);
 			}
 		}
 	}
diff --git a/Engine/Source/Runtime/Engine/Private/Animation/AnimSequence.cpp b/Engine/Source/Runtime/Engine/Private/Animation/AnimSequence.cpp
--- a/Engine/Source/Runtime/Engine/Private/Animation/AnimSequence.cpp
+++ b/Engine/Source/Runtime/Engine/Private/Animation/AnimSequence.cpp
@@ -21,6 +21,8 @@
 #include "Animation/AnimNotifies/AnimNotify.h"
 #include "Animation/Rig.h"
 #include "Animation/AnimationSettings.h"
+#include "Animation/AnimCurveCompressionCodec.h"
+#include "Animation/AnimCurveCompressionSettings.h"
 #include "EditorFramework/AssetImportData.h"
 #include "Logging/TokenizedMessage.h"
 #include "Logging/MessageLog.h"
@@ -41,6 +43,7 @@
 CSV_DEFINE_CATEGORY(Animation, true);
 
 DECLARE_CYCLE_STAT(TEXT("AnimSeq GetBonePose"), STAT_AnimSeq_GetBonePose, STATGROUP_Anim);
+DECLARE_CYCLE_STAT(TEXT("AnimSeq EvalCurveData"), STAT_AnimSeq_EvalCurveData, STATGROUP_Anim);
 DECLARE_CYCLE_STAT(TEXT("Build Anim Track Pairs"), STAT_BuildAnimTrackPairs, STATGROUP_Anim);
 DECLARE_CYCLE_STAT(TEXT("Extract Pose From Anim Data"), STAT_ExtractPoseFromAnimData, STATGROUP_Anim);
 
@@ -276,9 +279,24 @@ void UAnimSequence::GetAssetRegistryTags(TArray<FAssetRegistryTag>& OutTags) con
 	Super::GetAssetRegistryTags(OutTags);
 }
 
+void UAnimSequence::AddReferencedObjects(UObject* This, FReferenceCollector& Collector)
+{
+	Super::AddReferencedObjects(This, Collector);
+
+	UAnimSequence* AnimSeq = CastChecked<UAnimSequence>(This);
+	Collector.AddReferencedObject(AnimSeq->CurveCompressionCodec);
+}
+
 int32 UAnimSequence::GetUncompressedRawSize() const
 {
-	return ((sizeof(FVector) + sizeof(FQuat) + sizeof(FVector)) * RawAnimationData.Num() * NumFrames);
+	int32 BoneRawSize = ((sizeof(FVector) + sizeof(FQuat) + sizeof(FVector)) * RawAnimationData.Num() * NumFrames);
+	int32 CurveRawSize = 0;
+	for (const FFloatCurve& Curve : RawCurveData.FloatCurves)
+	{
+		CurveRawSize += sizeof(FFloatCurve);
+		CurveRawSize += sizeof(FRichCurveKey) * Curve.FloatCurve.Keys.Num();
+	}
+	return BoneRawSize + CurveRawSize;
 }
 
 int32 UAnimSequence::GetApproxRawSize() const
@@ -292,13 +310,19 @@ int32 UAnimSequence::GetApproxRawSize() const
 			sizeof( FQuat ) * RawTrack.RotKeys.Num() + 
 			sizeof( FVector ) * RawTrack.ScaleKeys.Num(); 
 	}
+	for (const FFloatCurve& Curve : RawCurveData.FloatCurves)
+	{
+		Total += sizeof(FFloatCurve);
+		Total += sizeof(FRichCurveKey) * Curve.FloatCurve.Keys.Num();
+	}
 	return Total;
 }
 
 int32 UAnimSequence::GetApproxCompressedSize() const
 {
-	const int32 Total = sizeof(int32)*CompressedTrackOffsets.Num() + CompressedByteStream.Num() + CompressedScaleOffsets.GetMemorySize() + sizeof(FCompressedSegment)*CompressedSegments.Num();
-	return Total;
+	int32 BoneTotal = sizeof(int32)*CompressedTrackOffsets.Num() + CompressedByteStream.Num() + CompressedScaleOffsets.GetMemorySize() + sizeof(FCompressedSegment)*CompressedSegments.Num();
+	int32 CurveTotal = CompressedCurveByteStream.Num();
+	return BoneTotal + CurveTotal;
 }
 
 /**
@@ -653,7 +677,10 @@ void UAnimSequence::PostLoad()
 
 	if (USkeleton* CurrentSkeleton = GetSkeleton())
 	{
-		VerifyCurveNames<FFloatCurve>(*CurrentSkeleton, USkeleton::AnimCurveMappingName, CompressedCurveData.FloatCurves);
+		for (FSmartName& CurveName : CompressedCurveNames)
+		{
+			CurrentSkeleton->VerifySmartName(USkeleton::AnimCurveMappingName, CurveName);
+		}
 
 #if WITH_EDITOR
 		VerifyCurveNames<FTransformCurve>(*CurrentSkeleton, USkeleton::AnimTrackCurveMappingName, RawCurveData.TransformCurves);
@@ -665,18 +692,6 @@ void UAnimSequence::PostLoad()
 #endif
 	}
 
-#if WITH_EDITOR
-
-	// Compressed curve flags are not authoritative (they come from the DDC). Keep them up to date with
-	// actual anim flags 
-	for (FFloatCurve& Curve : RawCurveData.FloatCurves)
-	{
-		if (FAnimCurveBase* CompressedCurve = CompressedCurveData.GetCurveData(Curve.Name.UID))
-		{
-			CompressedCurve->SetCurveTypeFlags(Curve.GetCurveTypeFlags());
-		}
-	}
-#endif // WITH_EDITOR
 	AddAnimLoadingDebugEntry(TEXT("PostLoadEnd"));
 }
 
@@ -844,11 +859,17 @@ void UAnimSequence::PostEditChangeProperty(FPropertyChangedEvent& PropertyChange
 		}
 		
 	}
+
 	// @Todo fix me: This is temporary fix to make sure they always have compressed data
 	if (RawAnimationData.Num() > 0 && (!IsCompressedDataValid() || bAdditiveSettingsChanged))
 	{
 		PostProcessSequence();
 	}
+
+	if (PropertyChangedEvent.Property != nullptr && PropertyChangedEvent.Property->GetFName() == GET_MEMBER_NAME_CHECKED(UAnimSequence, CurveCompressionSettings))
+	{
+		RequestSyncAnimRecompression(false);
+	}
 }
 
 void UAnimSequence::PostDuplicate(bool bDuplicateForPIE)
@@ -1460,6 +1481,7 @@ void UAnimSequence::BuildPoseFromRawData(const TArray<FRawAnimSequenceTrack>& In
 void UAnimSequence::GetBonePose(FCompactPose& OutPose, FBlendedCurve& OutCurve, const FAnimExtractContext& ExtractionContext, bool bForceUseRawData) const
 {
 	SCOPE_CYCLE_COUNTER(STAT_AnimSeq_GetBonePose);
+	CSV_SCOPED_TIMING_STAT(Animation, AnimSeq_GetBonePose);
 
 	const FBoneContainer& RequiredBones = OutPose.GetBoneContainer();
 	const bool bUseRawDataForPoseExtraction = bForceUseRawData || UseRawDataForPoseExtraction(RequiredBones);
@@ -1735,6 +1757,31 @@ void UAnimSequence::GetBonePose(FCompactPose& OutPose, FBlendedCurve& OutCurve,
 }
 
 #if WITH_EDITORONLY_DATA
+void UAnimSequence::UpdateCompressedCurveNames()
+{
+	// Copy our curve names to the compressed array
+	const int32 NumCurves = RawCurveData.FloatCurves.Num();
+	CompressedCurveNames.Reset(NumCurves);
+	CompressedCurveNames.AddUninitialized(NumCurves);
+	for (int32 CurveIndex = 0; CurveIndex < NumCurves; ++CurveIndex)
+	{
+		const FFloatCurve& Curve = RawCurveData.FloatCurves[CurveIndex];
+		CompressedCurveNames[CurveIndex] = Curve.Name;
+	}
+}
+
+void UAnimSequence::UpdateCompressedCurveName(SmartName::UID_Type CurveUID, const struct FSmartName& NewCurveName)
+{
+	for (FSmartName& CurveName : CompressedCurveNames)
+	{
+		if (CurveName.UID == CurveUID)
+		{
+			CurveName = NewCurveName;
+			break;
+		}
+	}
+}
+
 int32 UAnimSequence::AddNewRawTrack(FName TrackName, FRawAnimSequenceTrack* TrackData)
 {
 	const int32 SkeletonIndex = GetSkeleton() ? GetSkeleton()->GetReferenceSkeleton().FindBoneIndex(TrackName) : INDEX_NONE;
@@ -2427,6 +2474,11 @@ void UAnimSequence::RequestAnimCompression(FRequestAnimCompressionParams Params)
 		CompressionScheme = FAnimationUtils::GetDefaultAnimationCompressionAlgorithm();
 	}
 
+	if (CurveCompressionSettings == nullptr || !CurveCompressionSettings->AreSettingsValid())
+	{
+		CurveCompressionSettings = FAnimationUtils::GetDefaultAnimationCurveCompressionSettings();
+	}
+
 	if (!RawDataGuid.IsValid())
 	{
 		RawDataGuid = GenerateGuidFromRawData();
@@ -2523,7 +2575,7 @@ void UAnimSequence::SerializeCompressedData(FArchive& Ar, bool bDDCData)
 	Ar << CompressedSegments;
 
 	Ar << CompressedTrackToSkeletonMapTable;
-	Ar << CompressedCurveData;
+	Ar << CompressedCurveNames;
 
 	Ar << CompressedRawDataSize;
 	Ar << CompressedNumFrames;
@@ -2549,6 +2601,31 @@ void UAnimSequence::SerializeCompressedData(FArchive& Ar, bool bDDCData)
 		// and then use the codecs to byte swap
 		check(RotationCodec != NULL);
 		((AnimEncoding*)RotationCodec)->ByteSwapIn(*this, MemoryReader);
+
+#if WITH_EDITOR
+		if (bDDCData)
+		{
+			FString CurveCodecPath;
+			Ar << CurveCodecPath;
+
+			CurveCompressionCodec = LoadObject<UAnimCurveCompressionCodec>(nullptr, *CurveCodecPath);
+		}
+		else
+#else
+		check(!bDDCData);
+#endif
+		{
+			UAnimCurveCompressionCodec* CurveCodec = nullptr;
+			Ar << CurveCodec;
+			CurveCompressionCodec = CurveCodec;
+		}
+
+		int32 NumCurveBytes;
+		Ar << NumCurveBytes;
+
+		CompressedCurveByteStream.Empty(NumCurveBytes);
+		CompressedCurveByteStream.AddUninitialized(NumCurveBytes);
+		Ar.Serialize(CompressedCurveByteStream.GetData(), NumCurveBytes);
 	}
 	else if (Ar.IsSaving() || Ar.IsCountingMemory())
 	{
@@ -2575,6 +2652,24 @@ void UAnimSequence::SerializeCompressedData(FArchive& Ar, bool bDDCData)
 
 		// Count compressed data.
 		Ar.CountBytes(SerializedData.Num(), SerializedData.Num());
+
+#if WITH_EDITOR
+		if (bDDCData)
+		{
+			FString CurveCodecPath = CurveCompressionCodec->GetPathName();
+			Ar << CurveCodecPath;
+		}
+		else
+#else
+		check(!bDDCData);
+#endif
+		{
+			Ar << CurveCompressionCodec;
+		}
+
+		int32 NumCurveBytes = CompressedCurveByteStream.Num();
+		Ar << NumCurveBytes;
+		Ar.Serialize(CompressedCurveByteStream.GetData(), NumCurveBytes);
 	}
 
 #if WITH_EDITOR
@@ -2587,8 +2682,14 @@ void UAnimSequence::SerializeCompressedData(FArchive& Ar, bool bDDCData)
 		{
 			if(USkeleton* CurrentSkeleton = GetSkeleton())
 			{
-				VerifyCurveNames<FFloatCurve>(*CurrentSkeleton, USkeleton::AnimCurveMappingName, CompressedCurveData.FloatCurves);
-				bUseRawDataOnly = !IsCompressedDataValid();
+				// Refresh the compressed curve names since the IDs might have changed since
+				for (FSmartName& CurveName : CompressedCurveNames)
+				{
+					CurrentSkeleton->VerifySmartName(USkeleton::AnimCurveMappingName, CurveName);
+				}
+
+				bUseRawDataOnly = !IsCompressedDataValid() || !IsCurveCompressedDataValid();
+
 				ensureMsgf(!bUseRawDataOnly, TEXT("Anim Compression failed for Sequence '%s' Guid:%s CompressedDebugData:\n\tOriginal Anim:%s\n\tAdditiveSetting:%i\n\tCompression Scheme:%s\n\tRawDataGuid:%s"),
 					*GetFullName(),
 					*RawDataGuid.ToString(),
@@ -3020,12 +3121,14 @@ void UAnimSequence::RecycleAnimSequence()
 	CompressedSegments.Empty(0);
 	SourceRawAnimationData.Empty(0);
 	RawCurveData.Empty();
-	CompressedCurveData.Empty();
+	CompressedCurveByteStream.Empty(0);
+	CompressedCurveNames.Empty(0);
 	AuthoredSyncMarkers.Empty();
 	UniqueMarkerNames.Empty();
 	Notifies.Empty();
 	AnimNotifyTracks.Empty();
 	CompressionScheme = nullptr;
+	CurveCompressionCodec = nullptr;
 	TranslationCompressionFormat = RotationCompressionFormat = ScaleCompressionFormat = ACF_None;
 #endif // WITH_EDITORONLY_DATA
 }
@@ -4990,27 +5093,54 @@ void UAnimSequence::RefreshCacheData()
 
 void UAnimSequence::EvaluateCurveData(FBlendedCurve& OutCurve, float CurrentTime, bool bForceUseRawData) const
 {
-	if (bUseRawDataOnly || bForceUseRawData)
+	SCOPE_CYCLE_COUNTER(STAT_AnimSeq_EvalCurveData);
+
+	if (OutCurve.NumValidCurveCount == 0)
+	{
+		return;
+	}
+
+	if (bUseRawDataOnly || bForceUseRawData || !IsCurveCompressedDataValid())
 	{
-		Super::EvaluateCurveData(OutCurve, CurrentTime);
+		Super::EvaluateCurveData(OutCurve, CurrentTime, bForceUseRawData);
 	}
 	else
 	{
 		CSV_SCOPED_TIMING_STAT(Animation, EvaluateCurveData);
-		CompressedCurveData.EvaluateCurveData(OutCurve, CurrentTime);
+		CurveCompressionCodec->DecompressCurves(*this, OutCurve, CurrentTime);
 	}
 }
 
-const FRawCurveTracks& UAnimSequence::GetCurveData() const
+float UAnimSequence::EvaluateCurveData(SmartName::UID_Type CurveUID, float CurrentTime, bool bForceUseRawData) const
 {
-	if (bUseRawDataOnly)
+	SCOPE_CYCLE_COUNTER(STAT_AnimSeq_EvalCurveData);
+
+	if (bUseRawDataOnly || bForceUseRawData || !IsCurveCompressedDataValid())
 	{
-		return Super::GetCurveData();
+		return Super::EvaluateCurveData(CurveUID, CurrentTime, bForceUseRawData);
 	}
 	else
 	{
-		return CompressedCurveData;
+		return CurveCompressionCodec->DecompressCurve(*this, CurveUID, CurrentTime);
+	}
+}
+
+bool UAnimSequence::HasCurveData(SmartName::UID_Type CurveUID, bool bForceUseRawData) const
+{
+	if (bUseRawDataOnly || bForceUseRawData || !IsCurveCompressedDataValid())
+	{
+		return Super::HasCurveData(CurveUID, bForceUseRawData);
 	}
+
+	for (const FSmartName& CurveName : CompressedCurveNames)
+	{
+		if (CurveName.UID == CurveUID)
+		{
+			return true;
+		}
+	}
+
+	return false;
 }
 
 void UAnimSequence::RefreshSyncMarkerDataFromAuthored()
@@ -5738,13 +5868,13 @@ void UAnimSequence::OnRawDataChanged()
 	CompressedTrackOffsets.Empty();
 	CompressedScaleOffsets.Empty();
 	CompressedByteStream.Empty();
+	CompressedSegments.Empty();
 	bUseRawDataOnly = true;
 
 	RequestSyncAnimRecompression(false);
 	//MDW - Once we have async anim ddc requests we should do this too
 	//RequestDependentAnimRecompression();
 }
-
 #endif
 
 bool UAnimSequence::IsCompressedDataValid() const
@@ -5753,6 +5883,29 @@ bool UAnimSequence::IsCompressedDataValid() const
 		   (TranslationCompressionFormat == ACF_Identity && RotationCompressionFormat == ACF_Identity && ScaleCompressionFormat == ACF_Identity);
 }
 
+bool UAnimSequence::IsCurveCompressedDataValid() const
+{
+	if (CurveCompressionCodec == nullptr)
+	{
+		// No codec
+		return false;
+	}
+
+	if (CompressedCurveByteStream.Num() == 0 && RawCurveData.FloatCurves.Num() != 0)
+	{
+		// No compressed data but we have raw data
+		if (!IsValidAdditive())
+		{
+			return false;
+		}
+
+		// Additive sequences can have raw curves that all end up being 0.0 (e.g. they 100% match the base sequence curves)
+		// in which case there will be no compressed curve data.
+	}
+
+	return true;
+}
+
 /*-----------------------------------------------------------------------------
 	AnimNotify& subclasses
 -----------------------------------------------------------------------------*/
diff --git a/Engine/Source/Runtime/Engine/Private/Animation/AnimSequenceBase.cpp b/Engine/Source/Runtime/Engine/Private/Animation/AnimSequenceBase.cpp
--- a/Engine/Source/Runtime/Engine/Private/Animation/AnimSequenceBase.cpp
+++ b/Engine/Source/Runtime/Engine/Private/Animation/AnimSequenceBase.cpp
@@ -725,6 +725,17 @@ void UAnimSequenceBase::EvaluateCurveData(FBlendedCurve& OutCurve, float Current
 	RawCurveData.EvaluateCurveData(OutCurve, CurrentTime);
 }
 
+float UAnimSequenceBase::EvaluateCurveData(SmartName::UID_Type CurveUID, float CurrentTime, bool bForceUseRawData) const
+{
+	const FFloatCurve* Curve = (const FFloatCurve*)RawCurveData.GetCurveData(CurveUID, ERawCurveTrackTypes::RCT_Float);
+	return Curve != nullptr ? Curve->Evaluate(CurrentTime) : 0.0f;
+}
+
+bool UAnimSequenceBase::HasCurveData(SmartName::UID_Type CurveUID, bool bForceUseRawData) const
+{
+	return RawCurveData.GetCurveData(CurveUID) != nullptr;
+}
+
 void UAnimSequenceBase::Serialize(FArchive& Ar)
 {
 	Ar.UsingCustomVersion(FFrameworkObjectVersion::GUID);
diff --git a/Engine/Source/Runtime/Engine/Private/Animation/AnimationUtils.cpp b/Engine/Source/Runtime/Engine/Private/Animation/AnimationUtils.cpp
--- a/Engine/Source/Runtime/Engine/Private/Animation/AnimationUtils.cpp
+++ b/Engine/Source/Runtime/Engine/Private/Animation/AnimationUtils.cpp
@@ -16,6 +16,7 @@
 #include "Animation/AnimCompress_Automatic.h"
 #include "Animation/AnimSet.h"
 #include "Animation/AnimationSettings.h"
+#include "Animation/AnimCurveCompressionSettings.h"
 #include "AnimationCompression.h"
 #include "Engine/SkeletalMeshSocket.h"
 #include "AnimEncoding.h"
@@ -2233,3 +2234,79 @@ void FAnimationUtils::TallyErrorsFromPerturbation(
 		TrackError.MaxErrorInScaleDueToScale = MaxErrorS_DueToR;
 	}
 }
+
+#if WITH_EDITOR
+static UAnimCurveCompressionSettings* DefaultCurveCompressionSettings = nullptr;
+
+UAnimCurveCompressionSettings* FAnimationUtils::GetDefaultAnimationCurveCompressionSettings()
+{
+	if (DefaultCurveCompressionSettings == nullptr)
+	{
+		FConfigSection* AnimDefaultObjectSettingsSection = GConfig->GetSectionPrivate(TEXT("Animation.DefaultObjectSettings"), false, true, GEngineIni);
+		const FConfigValue* Value = AnimDefaultObjectSettingsSection != nullptr ? AnimDefaultObjectSettingsSection->Find(TEXT("CurveCompressionSettings")) : nullptr;
+
+		if (Value != nullptr)
+		{
+			const FString& CurveCompressionSettingsName = Value->GetValue();
+			DefaultCurveCompressionSettings = LoadObject<UAnimCurveCompressionSettings>(nullptr, *CurveCompressionSettingsName);
+		}
+
+		if (DefaultCurveCompressionSettings == nullptr)
+		{
+			UE_LOG(LogAnimationCompression, Fatal, TEXT("Couldn't find default curve compression settings under '[Animation.DefaultObjectSettings]'"));
+		}
+
+		// Force load the default settings and all its dependencies just in case it hasn't happened yet
+		bool bLoadDependencies = false;
+		if (DefaultCurveCompressionSettings->HasAnyFlags(RF_NeedLoad))
+		{
+			DefaultCurveCompressionSettings->GetLinker()->Preload(DefaultCurveCompressionSettings);
+			bLoadDependencies = true;
+		}
+
+		if (DefaultCurveCompressionSettings->HasAnyFlags(RF_NeedPostLoad))
+		{
+			DefaultCurveCompressionSettings->ConditionalPostLoad();
+			bLoadDependencies = true;
+		}
+
+		if (bLoadDependencies)
+		{
+			TArray<UObject*> ObjectReferences;
+			FReferenceFinder(ObjectReferences, nullptr, false, true, false, true).FindReferences(DefaultCurveCompressionSettings);
+
+			for (UObject* Dependency : ObjectReferences)
+			{
+				if (Dependency->HasAnyFlags(RF_NeedLoad))
+				{
+					Dependency->GetLinker()->Preload(Dependency);
+				}
+
+				if (Dependency->HasAnyFlags(RF_NeedPostLoad))
+				{
+					Dependency->ConditionalPostLoad();
+				}
+			}
+		}
+
+		DefaultCurveCompressionSettings->AddToRoot();
+	}
+
+	return DefaultCurveCompressionSettings;
+}
+
+bool FAnimationUtils::CompressAnimCurves(UAnimSequence& AnimSeq)
+{
+	// Clear any previous data we might have even if we end up failing to compress
+	AnimSeq.CompressedCurveByteStream.Empty();
+	AnimSeq.CurveCompressionCodec = nullptr;
+
+	if (AnimSeq.CurveCompressionSettings == nullptr || !AnimSeq.CurveCompressionSettings->AreSettingsValid())
+	{
+		return false;
+	}
+
+	check(AnimSeq.CurveCompressionSettings->AreSettingsValid());
+	return AnimSeq.CurveCompressionSettings->Compress(AnimSeq);
+}
+#endif
diff --git a/Engine/Source/Runtime/Engine/Private/Curves/RichCurve.cpp b/Engine/Source/Runtime/Engine/Private/Curves/RichCurve.cpp
--- a/Engine/Source/Runtime/Engine/Private/Curves/RichCurve.cpp
+++ b/Engine/Source/Runtime/Engine/Private/Curves/RichCurve.cpp
@@ -1,7 +1,7 @@
 // Copyright 1998-2019 Epic Games, Inc. All Rights Reserved.
 
 #include "Curves/RichCurve.h"
-
+#include "Templates/Function.h"
 
 DECLARE_CYCLE_STAT(TEXT("RichCurve Eval"), STAT_RichCurve_Eval, STATGROUP_Engine);
 
@@ -458,6 +458,24 @@ float FRichCurve::GetKeyValue(FKeyHandle KeyHandle) const
 	return GetKey(KeyHandle).Value;
 }
 
+bool FRichCurve::IsConstant(float Tolerance) const
+{
+	if (Keys.Num() <= 1)
+	{
+		return true;
+	}
+
+	const FRichCurveKey& RefKey = Keys[0];
+	for (const FRichCurveKey& Key : Keys)
+	{
+		if (!FMath::IsNearlyEqual(Key.Value, RefKey.Value, Tolerance))
+		{
+			return false;
+		}
+	}
+
+	return true;
+}
 
 void FRichCurve::ShiftCurve(float DeltaTime)
 {
@@ -930,7 +948,7 @@ void FRichCurve::RemoveRedundantKeys(float Tolerance, float FirstKeyTime, float
 }
 
 /** Util to find float value on bezier defined by 4 control points */ 
-static float BezierInterp(float P0, float P1, float P2, float P3, float Alpha)
+FORCEINLINE_DEBUGGABLE static float BezierInterp(float P0, float P1, float P2, float P3, float Alpha)
 {
 	const float P01 = FMath::Lerp(P0, P1, Alpha);
 	const float P12 = FMath::Lerp(P1, P2, Alpha);
@@ -1273,3 +1291,772 @@ bool FRichCurve::operator==(const FRichCurve& Curve) const
 
 	return true;
 }
+
+static ERichCurveCompressionFormat FindRichCurveCompressionFormat(const FRichCurve& Curve)
+{
+	if (Curve.Keys.Num() == 0)
+	{
+		return RCCF_Empty;
+	}
+
+	if (Curve.IsConstant())
+	{
+		return RCCF_Constant;
+	}
+
+	const FRichCurveKey& RefKey = Curve.Keys[0];
+	for (const FRichCurveKey& Key : Curve.Keys)
+	{
+		if (Key.InterpMode != RefKey.InterpMode)
+		{
+			return RCCF_Mixed;
+		}
+	}
+
+	switch (RefKey.InterpMode)
+	{
+	case RCIM_Constant:
+	case RCIM_None:
+	default:
+		return RCCF_Constant;
+	case RCIM_Linear:
+		return RCCF_Linear;
+	case RCIM_Cubic:
+		return RCCF_Cubic;
+	}
+}
+
+static ERichCurveKeyTimeCompressionFormat FindRichCurveKeyFormat(const FRichCurve& Curve, float ErrorThreshold, float SampleRate, ERichCurveCompressionFormat CompressionFormat)
+{
+	const int32 NumKeys = Curve.Keys.Num();
+	if (NumKeys == 0 || CompressionFormat == RCCF_Constant || CompressionFormat == RCCF_Empty || ErrorThreshold <= 0.0f || SampleRate <= 0.0f)
+	{
+		return RCKTCF_float32;
+	}
+
+	auto EvalForTwoKeys = [](const FRichCurveKey& Key1, float KeyTime1, const FRichCurveKey& Key2, float KeyTime2, float InTime)
+	{
+		const float Diff = KeyTime2 - KeyTime1;
+
+		if (Diff > 0.f && Key1.InterpMode != RCIM_Constant)
+		{
+			const float Alpha = (InTime - KeyTime1) / Diff;
+			const float P0 = Key1.Value;
+			const float P3 = Key2.Value;
+
+			if (Key1.InterpMode == RCIM_Linear)
+			{
+				return FMath::Lerp(P0, P3, Alpha);
+			}
+			else
+			{
+				const float OneThird = 1.0f / 3.0f;
+				const float P1 = P0 + (Key1.LeaveTangent * Diff * OneThird);
+				const float P2 = P3 - (Key2.ArriveTangent * Diff * OneThird);
+
+				return BezierInterp(P0, P1, P2, P3, Alpha);
+			}
+		}
+		else
+		{
+			return Key1.Value;
+		}
+	};
+
+	auto DecayTime = [](const FRichCurveKey& Key, float MinTime, float DeltaTime, float InvDeltaTime, float QuantizationScale, float InvQuantizationScale)
+	{
+		// 0.0f -> 0, 1.0f -> 255 for 8 bits
+		const float NormalizedTime = FMath::Clamp((Key.Time - MinTime) * InvDeltaTime, 0.0f, 1.0f);
+		const float QuantizedTime = FMath::RoundHalfFromZero(NormalizedTime * QuantizationScale);
+		const float LossyNormalizedTime = QuantizedTime * InvQuantizationScale;
+		return (LossyNormalizedTime * DeltaTime) + MinTime;
+	};
+
+	const float MinTime = Curve.Keys[0].Time;
+	const float MaxTime = Curve.Keys.Last().Time;
+	const float DeltaTime = MaxTime - MinTime;
+	const float InvDeltaTime = 1.0f / DeltaTime;
+	const float SampleRateIncrement = 1.0f / SampleRate;
+
+	// This is only acceptable if the maximum error is within a reasonable value
+	bool bFitsOn16Bits = true;
+
+	int32 CurrentLossyKey = 0;
+	int32 CurrentRefKey = 0;
+	float CurrentTime = MinTime;
+	while (CurrentTime <= MaxTime && bFitsOn16Bits)
+	{
+		if (CurrentTime > Curve.Keys[CurrentRefKey + 1].Time)
+		{
+			CurrentRefKey++;
+
+			if (CurrentRefKey >= NumKeys)
+			{
+				break;
+			}
+		}
+
+		float LossyTime1_16;
+		float LossyTime2_16 = DecayTime(Curve.Keys[CurrentLossyKey + 1], MinTime, DeltaTime, InvDeltaTime, 65535.0f, 1.0f / 65535.0f);
+		if (CurrentTime > LossyTime2_16)
+		{
+			CurrentLossyKey++;
+
+			if (CurrentLossyKey >= NumKeys)
+			{
+				break;
+			}
+
+			LossyTime1_16 = LossyTime2_16;
+			LossyTime2_16 = DecayTime(Curve.Keys[CurrentLossyKey + 1], MinTime, DeltaTime, InvDeltaTime, 65535.0f, 1.0f / 65535.0f);
+		}
+		else
+		{
+			LossyTime1_16 = DecayTime(Curve.Keys[CurrentLossyKey], MinTime, DeltaTime, InvDeltaTime, 65535.0f, 1.0f / 65535.0f);
+		}
+
+		const float Result_16 = EvalForTwoKeys(Curve.Keys[CurrentLossyKey], LossyTime1_16, Curve.Keys[CurrentLossyKey + 1], LossyTime2_16, CurrentTime);
+		const float Result_Ref = ::EvalForTwoKeys(Curve.Keys[CurrentRefKey], Curve.Keys[CurrentRefKey + 1], CurrentTime);
+
+		const float Error_16 = FMath::Abs(Result_Ref - Result_16);
+
+		bFitsOn16Bits &= Error_16 <= ErrorThreshold;
+
+		CurrentTime += SampleRateIncrement;
+	}
+
+	// In order to normalize time values, we need to store the MinTime and the DeltaTime with full precision
+	// This means we need 8 bytes of overhead
+	// If the number of keys is too small, the overhead is larger or equal to the space we save and isn't
+	// worth it.
+
+	// For 8 bits, the formula is: 8 + (N * sizeof(uint8))
+	// With 8 bits, 2 keys need 8 bytes with full precision and 10 bytes packed. No savings, not worth using.
+	// With 8 bits, 3 keys need 12 bytes with full precision and 11 bytes packed. We save 1 byte, worth using.
+	// For 16 bits, the formula is: 8 + (N * sizeof(uint16))
+	// With 16 bits, 6 keys need 20 bytes with full precision and 20 bytes packed. No savings, not worth using.
+	// With 16 bits, 7 keys need 24 bytes with full precision and 22 bytes packed. We save 2 bytes, worth using.
+	// Alignment and padding must also be taken into account
+
+	// Note: Support for storing key time on 8 bits was attempted but it was rarely selected and wasn't worth the complexity
+
+	int32 SizeInterpMode = 0;
+	if (CompressionFormat == RCCF_Mixed)
+	{
+		SizeInterpMode += NumKeys * sizeof(uint8);
+	}
+
+	const int32 SizeUInt16 = Align(Align(SizeInterpMode, sizeof(uint16)) + (NumKeys * sizeof(uint16)), sizeof(float)) + (2 * sizeof(float));
+	const int32 SizeFloat32 = Align(SizeInterpMode, sizeof(float)) + NumKeys * sizeof(float);
+
+	if (bFitsOn16Bits && SizeUInt16 < SizeFloat32)
+	{
+		return RCKTCF_uint16;
+	}
+	else
+	{
+		return RCKTCF_float32;
+	}
+}
+
+void FRichCurve::CompressCurve(FCompressedRichCurve& OutCurve, float ErrorThreshold, float SampleRate) const
+{
+	ERichCurveCompressionFormat CompressionFormat = FindRichCurveCompressionFormat(*this);
+	OutCurve.CompressionFormat = CompressionFormat;
+
+	ERichCurveKeyTimeCompressionFormat KeyFormat = FindRichCurveKeyFormat(*this, ErrorThreshold, SampleRate, CompressionFormat);
+	OutCurve.KeyTimeCompressionFormat = KeyFormat;
+
+	OutCurve.PreInfinityExtrap = PreInfinityExtrap;
+	OutCurve.PostInfinityExtrap = PostInfinityExtrap;
+
+	if (CompressionFormat == RCCF_Empty)
+	{
+		OutCurve.ConstantValueNumKeys.ConstantValue = DefaultValue;
+		OutCurve.CompressedKeys.Empty();
+	}
+	else if (CompressionFormat == RCCF_Constant)
+	{
+		OutCurve.ConstantValueNumKeys.ConstantValue = Keys[0].Value;
+		OutCurve.CompressedKeys.Empty();
+	}
+	else
+	{
+		int32 PackedDataSize = 0;
+
+		// If we are mixed, we need to store the interp mode for every key, this data comes first following the header
+		// Next comes the quantized time values followed by the normalization range
+		// And the values/tangents follow last
+
+		if (CompressionFormat == RCCF_Mixed)
+		{
+			PackedDataSize += Keys.Num() * sizeof(uint8);
+		}
+
+		if (KeyFormat == RCKTCF_uint16)
+		{
+			PackedDataSize = Align(PackedDataSize, sizeof(uint16));
+			PackedDataSize += Keys.Num() * sizeof(uint16);
+			PackedDataSize = Align(PackedDataSize, sizeof(float));
+			PackedDataSize += 2 * sizeof(float);
+		}
+		else
+		{
+			check(KeyFormat == RCKTCF_float32);
+			PackedDataSize = Align(PackedDataSize, sizeof(float));
+			PackedDataSize += Keys.Num() * sizeof(float);
+		}
+
+		PackedDataSize += Keys.Num() * sizeof(float);	// Key values
+
+		// Key tangents
+		if (CompressionFormat == RCCF_Cubic)
+		{
+			PackedDataSize += Keys.Num() * 2 * sizeof(float);
+		}
+		else if (CompressionFormat == RCCF_Mixed)
+		{
+			for (const FRichCurveKey& Key : Keys)
+			{
+				if (Key.InterpMode == RCIM_Cubic)
+				{
+					PackedDataSize += 2 * sizeof(float);
+				}
+			}
+		}
+
+		OutCurve.CompressedKeys.Empty(PackedDataSize);
+		OutCurve.CompressedKeys.AddUninitialized(PackedDataSize);
+
+		int32 WriteOffset = 0;
+		uint8* BasePtr = OutCurve.CompressedKeys.GetData();
+
+		OutCurve.ConstantValueNumKeys.NumKeys = Keys.Num();
+
+		// Key interp modes
+		if (CompressionFormat == RCCF_Mixed)
+		{
+			uint8* InterpModes = BasePtr + WriteOffset;
+			WriteOffset += Keys.Num() * sizeof(uint8);
+
+			for (const FRichCurveKey& Key : Keys)
+			{
+				if (Key.InterpMode == RCIM_Linear)
+				{
+					*InterpModes++ = (uint8)RCCF_Linear;
+				}
+				else if (Key.InterpMode == RCIM_Cubic)
+				{
+					*InterpModes++ = (uint8)RCCF_Cubic;
+				}
+				else
+				{
+					*InterpModes++ = (uint8)RCCF_Constant;
+				}
+			}
+		}
+
+		// Key times
+		if (KeyFormat == RCKTCF_uint16)
+		{
+			const float MinTime = Keys[0].Time;
+			const float MaxTime = Keys.Last().Time;
+			const float DeltaTime = MaxTime - MinTime;
+			const float InvDeltaTime = 1.0f / DeltaTime;
+
+			const int32 KeySize = sizeof(uint16);
+
+			if (KeyFormat == RCKTCF_uint16)
+			{
+				WriteOffset = Align(WriteOffset, sizeof(uint16));
+			}
+
+			uint8* KeyTimes8 = BasePtr + WriteOffset;
+			uint16* KeyTimes16 = reinterpret_cast<uint16*>(KeyTimes8);
+			WriteOffset += Keys.Num() * KeySize;
+
+			for (const FRichCurveKey& Key : Keys)
+			{
+				const float NormalizedTime = FMath::Clamp((Key.Time - MinTime) * InvDeltaTime, 0.0f, 1.0f);
+				const uint16 QuantizedTime = (uint16)FMath::RoundHalfFromZero(NormalizedTime * 65535.0f);
+				*KeyTimes16++ = QuantizedTime;
+			}
+
+			WriteOffset = Align(WriteOffset, sizeof(float));
+			float* RangeData = reinterpret_cast<float*>(BasePtr + WriteOffset);
+			WriteOffset += 2 * sizeof(float);
+
+			RangeData[0] = MinTime;
+			RangeData[1] = DeltaTime;
+		}
+		else
+		{
+			WriteOffset = Align(WriteOffset, sizeof(float));
+			float* KeyTimes = reinterpret_cast<float*>(BasePtr + WriteOffset);
+			WriteOffset += Keys.Num() * sizeof(float);
+
+			for (const FRichCurveKey& Key : Keys)
+			{
+				*KeyTimes++ = Key.Time;
+			}
+		}
+
+		// Key values and tangents
+		float* KeyData = reinterpret_cast<float*>(BasePtr + WriteOffset);
+		for (const FRichCurveKey& Key : Keys)
+		{
+			*KeyData++ = Key.Value;
+
+			if (Key.InterpMode == RCIM_Cubic)
+			{
+				check(CompressionFormat == RCCF_Cubic || CompressionFormat == RCCF_Mixed);
+				*KeyData++ = Key.ArriveTangent;
+				*KeyData++ = Key.LeaveTangent;
+			}
+		}
+
+		check(((uint8*)KeyData - BasePtr) == PackedDataSize);
+	}
+}
+
+struct Quantized16BitKeyTimeAdapter
+{
+	static constexpr float QuantizationScale = 1.0f / 65535.0f;
+	static constexpr int32 KeySize = sizeof(uint16);
+	static constexpr int32 RangeDataSize = 2 * sizeof(float);
+
+	using KeyTimeType = uint16;
+
+	const uint16* KeyTimes;
+	float MinTime;
+	float DeltaTime;
+	int32 KeyDataOffset;
+
+	Quantized16BitKeyTimeAdapter(const uint8* BasePtr, int32 KeyTimesOffset, int32 NumKeys)
+	{
+		const int32 RangeDataOffset = Align(KeyTimesOffset + (NumKeys * sizeof(uint16)), sizeof(float));
+		KeyDataOffset = RangeDataOffset + RangeDataSize;
+		const float* RangeData = reinterpret_cast<const float*>(BasePtr + RangeDataOffset);
+
+		KeyTimes = reinterpret_cast<const uint16*>(BasePtr + KeyTimesOffset);
+		MinTime = RangeData[0];
+		DeltaTime = RangeData[1];
+	}
+
+	float GetTime(int32 KeyIndex) const
+	{
+		const float KeyNormalizedTime = KeyTimes[KeyIndex] * QuantizationScale;
+		return (KeyNormalizedTime * DeltaTime) + MinTime;
+	};
+};
+
+struct Float32BitKeyTimeAdapter
+{
+	static constexpr int32 KeySize = sizeof(float);
+	static constexpr int32 RangeDataSize = 0;
+
+	using KeyTimeType = float;
+
+	const float* KeyTimes;
+	int32 KeyDataOffset;
+
+	Float32BitKeyTimeAdapter(const uint8* BasePtr, int32 KeyTimesOffset, int32 NumKeys)
+	{
+		KeyTimes = reinterpret_cast<const float*>(BasePtr + KeyTimesOffset);
+		KeyDataOffset = Align(KeyTimesOffset + (NumKeys * sizeof(float)), sizeof(float));
+	}
+
+	constexpr float GetTime(int32 KeyIndex) const
+	{
+		return KeyTimes[KeyIndex];
+	};
+};
+
+using KeyDataHandle = int32;
+
+template<ERichCurveCompressionFormat Format>
+struct UniformKeyDataAdapter
+{
+	const float* KeyData;
+
+	template<typename KeyTimeAdapterType>
+	
+	constexpr UniformKeyDataAdapter(const uint8* BasePtr, const KeyTimeAdapterType& KeyTimeAdapter)
+    : KeyData(reinterpret_cast<const float*>(BasePtr + KeyTimeAdapter.KeyDataOffset))
+    {}
+
+	constexpr KeyDataHandle GetKeyDataHandle(int32 KeyIndexToQuery) const
+	{
+		return Format == RCCF_Cubic ? (KeyIndexToQuery * 3) : KeyIndexToQuery;
+	};
+
+	constexpr float GetKeyValue(KeyDataHandle Handle) const
+	{
+		return KeyData[Handle];
+	}
+
+	constexpr float GetKeyArriveTangent(KeyDataHandle Handle) const
+	{
+		return KeyData[Handle + 1];
+	}
+
+	constexpr float GetKeyLeaveTangent(KeyDataHandle Handle) const
+	{
+		return KeyData[Handle + 2];
+	}
+
+	constexpr ERichCurveCompressionFormat GetKeyInterpMode(int32 KeyIndex) const
+	{
+		return Format;
+	}
+};
+
+struct MixedKeyDataAdapter
+{
+	const uint8* InterpModes;
+	const float* KeyData;
+
+	template<typename KeyTimeAdapterType>
+	MixedKeyDataAdapter(const uint8* BasePtr, int32 InterpModesOffset, const KeyTimeAdapterType& KeyTimeAdapter)
+	{
+		InterpModes = BasePtr + InterpModesOffset;
+		KeyData = reinterpret_cast<const float*>(BasePtr + KeyTimeAdapter.KeyDataOffset);
+	}
+
+	KeyDataHandle GetKeyDataHandle(int32 KeyIndexToQuery) const
+	{
+		int32 Offset = 0;
+		for (int32 KeyIndex = 0; KeyIndex < KeyIndexToQuery; ++KeyIndex)
+		{
+			Offset += InterpModes[KeyIndex] == RCCF_Cubic ? 3 : 1;
+		}
+
+		return Offset;
+	};
+
+	constexpr float GetKeyValue(KeyDataHandle Handle) const
+	{
+		return KeyData[Handle];
+	}
+
+	constexpr float GetKeyArriveTangent(KeyDataHandle Handle) const
+	{
+		return KeyData[Handle + 1];
+	}
+
+	constexpr float GetKeyLeaveTangent(KeyDataHandle Handle) const
+	{
+		return KeyData[Handle + 2];
+	}
+
+	constexpr ERichCurveCompressionFormat GetKeyInterpMode(int32 KeyIndex) const
+	{
+		return (ERichCurveCompressionFormat)InterpModes[KeyIndex];
+	}
+};
+
+template<typename KeyTimeAdapterType, typename KeyDataAdapterType>
+static float RemapTimeValue(float InTime, const KeyTimeAdapterType& KeyTimeAdapter, const KeyDataAdapterType& KeyDataAdapter, int32 NumKeys, ERichCurveExtrapolation InfinityExtrap, int32 KeyIndex0, int32 KeyIndex1, float& CycleValueOffset)
+{
+	// For Pre-infinity, key0 and key1 are the actual key 0 and key 1
+	// For Post-infinity, key0 and key1 are the last and second to last key
+	const float MinTime = KeyTimeAdapter.GetTime(0);
+	const float MaxTime = KeyTimeAdapter.GetTime(NumKeys - 1);
+
+	int CycleCount = 0;
+	CycleTime(MinTime, MaxTime, InTime, CycleCount);
+
+	if (InfinityExtrap == RCCE_CycleWithOffset)
+	{
+		const KeyDataHandle ValueHandle0 = KeyDataAdapter.GetKeyDataHandle(KeyIndex0);
+		const float KeyValue0 = KeyDataAdapter.GetKeyValue(ValueHandle0);
+		const KeyDataHandle ValueHandle1 = KeyDataAdapter.GetKeyDataHandle(KeyIndex1);
+		const float KeyValue1 = KeyDataAdapter.GetKeyValue(ValueHandle1);
+
+		const float DV = KeyValue0 - KeyValue1;
+		CycleValueOffset = DV * CycleCount;
+	}
+	else if (InfinityExtrap == RCCE_Oscillate)
+	{
+		if (CycleCount % 2 == 1)
+		{
+			InTime = MinTime + (MaxTime - InTime);
+		}
+	}
+
+	return InTime;
+}
+
+template<typename KeyTimeAdapterType, typename KeyDataAdapterType>
+static float InterpEvalExtrapolate(float InTime, const KeyTimeAdapterType& KeyTimeAdapter, const KeyDataAdapterType& KeyDataAdapter, ERichCurveExtrapolation InfinityExtrap, int32 KeyIndex0, int32 KeyIndex1, float KeyTime0)
+{
+	// For Pre-infinity, key0 and key1 are the actual key 0 and key 1
+	// For Post-infinity, key0 and key1 are the last and second to last key
+	const KeyDataHandle ValueHandle0 = KeyDataAdapter.GetKeyDataHandle(KeyIndex0);
+	const float KeyValue0 = KeyDataAdapter.GetKeyValue(ValueHandle0);
+
+	if (InfinityExtrap == RCCE_Linear)
+	{
+		const float KeyTime1 = KeyTimeAdapter.GetTime(KeyIndex1);
+		const float DT = KeyTime1 - KeyTime0;
+
+		if (FMath::IsNearlyZero(DT))
+		{
+			return KeyValue0;
+		}
+		else
+		{
+			const KeyDataHandle ValueHandle1 = KeyDataAdapter.GetKeyDataHandle(KeyIndex1);
+			const float KeyValue1 = KeyDataAdapter.GetKeyValue(ValueHandle1);
+			const float DV = KeyValue1 - KeyValue0;
+			const float Slope = DV / DT;
+
+			return Slope * (InTime - KeyTime0) + KeyValue0;
+		}
+	}
+	else
+	{
+		// Otherwise if constant or in a cycle or oscillate, always use the first key value
+		return KeyValue0;
+	}
+}
+
+// Each template permutation is only called from one place, force inline it to avoid issues
+template<typename KeyTimeAdapterType, typename KeyDataAdapterType>
+FORCEINLINE_DEBUGGABLE static float InterpEval(float InTime, const KeyTimeAdapterType& KeyTimeAdapter, const KeyDataAdapterType& KeyDataAdapter, int32 NumKeys, ERichCurveExtrapolation PreInfinityExtrap, ERichCurveExtrapolation PostInfinityExtrap)
+{
+	float CycleValueOffset = 0.0;
+
+	const float FirstKeyTime = KeyTimeAdapter.GetTime(0);
+	if (InTime <= FirstKeyTime)
+	{
+		if (PreInfinityExtrap != RCCE_Linear && PreInfinityExtrap != RCCE_Constant)
+		{
+			InTime = RemapTimeValue(InTime, KeyTimeAdapter, KeyDataAdapter, NumKeys, PreInfinityExtrap, 0, NumKeys - 1, CycleValueOffset);
+		}
+		else
+		{
+			return InterpEvalExtrapolate(InTime, KeyTimeAdapter, KeyDataAdapter, PreInfinityExtrap, 0, 1, FirstKeyTime);
+		}
+	}
+
+	const float LastKeyTime = KeyTimeAdapter.GetTime(NumKeys - 1);
+	if (InTime >= LastKeyTime)
+	{
+		if (PostInfinityExtrap != RCCE_Linear && PostInfinityExtrap != RCCE_Constant)
+		{
+			InTime = RemapTimeValue(InTime, KeyTimeAdapter, KeyDataAdapter, NumKeys, PostInfinityExtrap, NumKeys - 1, 0, CycleValueOffset);
+		}
+		else
+		{
+			return InterpEvalExtrapolate(InTime, KeyTimeAdapter, KeyDataAdapter, PostInfinityExtrap, NumKeys - 1, NumKeys - 2, LastKeyTime);
+		}
+	}
+
+	// perform a lower bound to get the second of the interpolation nodes
+	int32 First = 1;
+	int32 Last = NumKeys - 1;
+	int32 Count = Last - First;
+
+	while (Count > 0)
+	{
+		const int32 Step = Count / 2;
+		const int32 Middle = First + Step;
+
+		// TODO: Can we do the search with integers? In order to do so, we need to Floorf(..) the key times
+		const float KeyTime = KeyTimeAdapter.GetTime(Middle);
+		if (InTime >= KeyTime)
+		{
+			First = Middle + 1;
+			Count -= Step + 1;
+		}
+		else
+		{
+			Count = Step;
+		}
+	}
+
+	const float KeyTime0 = KeyTimeAdapter.GetTime(First - 1);
+	const float KeyTime1 = KeyTimeAdapter.GetTime(First);
+	const float Diff = KeyTime1 - KeyTime0;
+
+	const KeyDataHandle KeyValueHandle0 = KeyDataAdapter.GetKeyDataHandle(First - 1);
+	const float KeyValue0 = KeyDataAdapter.GetKeyValue(KeyValueHandle0);
+
+	// const value here allows the code to be stripped statically if the data is uniform
+	// which it is most of the time
+	const ERichCurveCompressionFormat KeyInterpMode0 = KeyDataAdapter.GetKeyInterpMode(First - 1);
+	float InterpolatedValue;
+	if (Diff > 0.0f && KeyInterpMode0 != RCCF_Constant)
+	{
+		const KeyDataHandle KeyValueHandle1 = KeyDataAdapter.GetKeyDataHandle(First);
+		const float KeyValue1 = KeyDataAdapter.GetKeyValue(KeyValueHandle1);
+
+		const float Alpha = (InTime - KeyTime0) / Diff;
+		const float P0 = KeyValue0;
+		const float P3 = KeyValue1;
+
+		if (KeyInterpMode0 == RCCF_Linear)
+		{
+			InterpolatedValue = FMath::Lerp(P0, P3, Alpha);
+		}
+		else
+		{
+			const float OneThird = 1.0f / 3.0f;
+			const float ScaledDiff = Diff * OneThird;
+			const float KeyLeaveTangent0 = KeyDataAdapter.GetKeyLeaveTangent(KeyValueHandle0);
+			const float KeyArriveTangent1 = KeyDataAdapter.GetKeyArriveTangent(KeyValueHandle1);
+			const float P1 = P0 + (KeyLeaveTangent0 * ScaledDiff);
+			const float P2 = P3 - (KeyArriveTangent1 * ScaledDiff);
+
+			InterpolatedValue = BezierInterp(P0, P1, P2, P3, Alpha);
+		}
+	}
+	else
+	{
+		InterpolatedValue = KeyValue0;
+	}
+
+	return InterpolatedValue + CycleValueOffset;
+}
+
+static TFunction<float(ERichCurveExtrapolation PreInfinityExtrap, ERichCurveExtrapolation PostInfinityExtrap, FCompressedRichCurve::TConstantValueNumKeys ConstantValueNumKeys, const uint8* CompressedKeys, float InTime, float InDefaultValue)> InterpEvalMap[5][2]
+{
+	// RCCF_Empty
+	{
+		// RCKTCF_uint16
+		[](ERichCurveExtrapolation PreInfinityExtrap, ERichCurveExtrapolation PostInfinityExtrap, FCompressedRichCurve::TConstantValueNumKeys ConstantValueNumKeys, const uint8* CompressedKeys, float InTime, float InDefaultValue)
+		{
+			return ConstantValueNumKeys.ConstantValue == MAX_flt ? InDefaultValue : ConstantValueNumKeys.ConstantValue;
+		},
+		// RCKTCF_float32
+		[](ERichCurveExtrapolation PreInfinityExtrap, ERichCurveExtrapolation PostInfinityExtrap, FCompressedRichCurve::TConstantValueNumKeys ConstantValueNumKeys, const uint8* CompressedKeys, float InTime, float InDefaultValue)
+		{
+			return ConstantValueNumKeys.ConstantValue == MAX_flt ? InDefaultValue : ConstantValueNumKeys.ConstantValue;
+		},
+	},
+	// RCCF_Constant
+	{
+		// RCKTCF_uint16
+		[](ERichCurveExtrapolation PreInfinityExtrap, ERichCurveExtrapolation PostInfinityExtrap, FCompressedRichCurve::TConstantValueNumKeys ConstantValueNumKeys, const uint8* CompressedKeys, float InTime, float InDefaultValue) { return ConstantValueNumKeys.ConstantValue; },
+		// RCKTCF_float32
+		[](ERichCurveExtrapolation PreInfinityExtrap, ERichCurveExtrapolation PostInfinityExtrap, FCompressedRichCurve::TConstantValueNumKeys ConstantValueNumKeys, const uint8* CompressedKeys, float InTime, float InDefaultValue) { return ConstantValueNumKeys.ConstantValue; },
+	},
+	// RCCF_Linear
+	{
+		// RCKTCF_uint16
+		[](ERichCurveExtrapolation PreInfinityExtrap, ERichCurveExtrapolation PostInfinityExtrap, FCompressedRichCurve::TConstantValueNumKeys ConstantValueNumKeys, const uint8* CompressedKeys, float InTime, float InDefaultValue)
+		{
+			const int32 KeyTimesOffset = 0;
+			Quantized16BitKeyTimeAdapter KeyTimeAdapter(CompressedKeys, KeyTimesOffset, ConstantValueNumKeys.NumKeys);
+			UniformKeyDataAdapter<RCCF_Linear> KeyDataAdapter(CompressedKeys, KeyTimeAdapter);
+			return InterpEval(InTime, KeyTimeAdapter, KeyDataAdapter, ConstantValueNumKeys.NumKeys, PreInfinityExtrap, PostInfinityExtrap);
+		},
+		// RCKTCF_float32
+		[](ERichCurveExtrapolation PreInfinityExtrap, ERichCurveExtrapolation PostInfinityExtrap, FCompressedRichCurve::TConstantValueNumKeys ConstantValueNumKeys, const uint8* CompressedKeys, float InTime, float InDefaultValue)
+		{
+			const int32 KeyTimesOffset = 0;
+			Float32BitKeyTimeAdapter KeyTimeAdapter(CompressedKeys, KeyTimesOffset, ConstantValueNumKeys.NumKeys);
+			UniformKeyDataAdapter<RCCF_Linear> KeyDataAdapter(CompressedKeys, KeyTimeAdapter);
+			return InterpEval(InTime, KeyTimeAdapter, KeyDataAdapter, ConstantValueNumKeys.NumKeys, PreInfinityExtrap, PostInfinityExtrap);
+		},
+	},
+	// RCCF_Cubic
+	{
+		// RCKTCF_uint16
+		[](ERichCurveExtrapolation PreInfinityExtrap, ERichCurveExtrapolation PostInfinityExtrap, FCompressedRichCurve::TConstantValueNumKeys ConstantValueNumKeys, const uint8* CompressedKeys, float InTime, float InDefaultValue)
+		{
+			const int32 KeyTimesOffset = 0;
+			Quantized16BitKeyTimeAdapter KeyTimeAdapter(CompressedKeys, KeyTimesOffset, ConstantValueNumKeys.NumKeys);
+			UniformKeyDataAdapter<RCCF_Cubic> KeyDataAdapter(CompressedKeys, KeyTimeAdapter);
+			return InterpEval(InTime, KeyTimeAdapter, KeyDataAdapter, ConstantValueNumKeys.NumKeys, PreInfinityExtrap, PostInfinityExtrap);
+		},
+		// RCKTCF_float32
+		[](ERichCurveExtrapolation PreInfinityExtrap, ERichCurveExtrapolation PostInfinityExtrap, FCompressedRichCurve::TConstantValueNumKeys ConstantValueNumKeys, const uint8* CompressedKeys, float InTime, float InDefaultValue)
+		{
+			const int32 KeyTimesOffset = 0;
+			Float32BitKeyTimeAdapter KeyTimeAdapter(CompressedKeys, KeyTimesOffset, ConstantValueNumKeys.NumKeys);
+			UniformKeyDataAdapter<RCCF_Cubic> KeyDataAdapter(CompressedKeys, KeyTimeAdapter);
+			return InterpEval(InTime, KeyTimeAdapter, KeyDataAdapter, ConstantValueNumKeys.NumKeys, PreInfinityExtrap, PostInfinityExtrap);
+		},
+	},
+	// RCCF_Mixed
+	{
+		// RCKTCF_uint16
+		[](ERichCurveExtrapolation PreInfinityExtrap, ERichCurveExtrapolation PostInfinityExtrap, FCompressedRichCurve::TConstantValueNumKeys ConstantValueNumKeys, const uint8* CompressedKeys, float InTime, float InDefaultValue)
+		{
+			const int32 InterpModesOffset = 0;
+			const int32 KeyTimesOffset = InterpModesOffset + Align(ConstantValueNumKeys.NumKeys * sizeof(uint8), sizeof(uint16));
+			Quantized16BitKeyTimeAdapter KeyTimeAdapter(CompressedKeys, KeyTimesOffset, ConstantValueNumKeys.NumKeys);
+			MixedKeyDataAdapter KeyDataAdapter(CompressedKeys, InterpModesOffset, KeyTimeAdapter);
+			return InterpEval(InTime, KeyTimeAdapter, KeyDataAdapter, ConstantValueNumKeys.NumKeys, PreInfinityExtrap, PostInfinityExtrap);
+		},
+		// RCKTCF_float32
+		[](ERichCurveExtrapolation PreInfinityExtrap, ERichCurveExtrapolation PostInfinityExtrap, FCompressedRichCurve::TConstantValueNumKeys ConstantValueNumKeys, const uint8* CompressedKeys, float InTime, float InDefaultValue)
+		{
+			const int32 InterpModesOffset = 0;
+			const int32 KeyTimesOffset = InterpModesOffset + Align(ConstantValueNumKeys.NumKeys * sizeof(uint8), sizeof(float));
+			Float32BitKeyTimeAdapter KeyTimeAdapter(CompressedKeys, KeyTimesOffset, ConstantValueNumKeys.NumKeys);
+			MixedKeyDataAdapter KeyDataAdapter(CompressedKeys, InterpModesOffset, KeyTimeAdapter);
+			return InterpEval(InTime, KeyTimeAdapter, KeyDataAdapter, ConstantValueNumKeys.NumKeys, PreInfinityExtrap, PostInfinityExtrap);
+		},
+	},
+};
+
+float FCompressedRichCurve::Eval(float InTime, float InDefaultValue) const
+{
+	SCOPE_CYCLE_COUNTER(STAT_RichCurve_Eval);
+
+	// Dynamic dispatch into a template optimized code path
+	const float Value = InterpEvalMap[CompressionFormat][KeyTimeCompressionFormat](PreInfinityExtrap, PostInfinityExtrap, ConstantValueNumKeys, CompressedKeys.GetData(), InTime, InDefaultValue);
+	return Value;
+}
+
+float FCompressedRichCurve::StaticEval(ERichCurveCompressionFormat CompressionFormat, ERichCurveKeyTimeCompressionFormat KeyTimeCompressionFormat, ERichCurveExtrapolation PreInfinityExtrap, ERichCurveExtrapolation PostInfinityExtrap, FCompressedRichCurve::TConstantValueNumKeys ConstantValueNumKeys, const uint8* CompressedKeys, float InTime, float InDefaultValue)
+{
+	SCOPE_CYCLE_COUNTER(STAT_RichCurve_Eval);
+
+	// Dynamic dispatch into a template optimized code path
+	const float Value = InterpEvalMap[CompressionFormat][KeyTimeCompressionFormat](PreInfinityExtrap, PostInfinityExtrap, ConstantValueNumKeys, CompressedKeys, InTime, InDefaultValue);
+	return Value;
+}
+
+bool FCompressedRichCurve::Serialize(FArchive& Ar)
+{
+	Ar << CompressionFormat;
+	Ar << KeyTimeCompressionFormat;
+	Ar << PreInfinityExtrap;
+	Ar << PostInfinityExtrap;
+
+	int32 NumKeysOrConstant = ConstantValueNumKeys.NumKeys;
+	Ar << NumKeysOrConstant;
+	ConstantValueNumKeys.NumKeys = NumKeysOrConstant;
+
+	if (Ar.IsLoading())
+	{
+		int32 NumBytes;
+		Ar << NumBytes;
+
+		CompressedKeys.Empty(NumBytes);
+		CompressedKeys.AddUninitialized(NumBytes);
+		Ar.Serialize(CompressedKeys.GetData(), NumBytes);
+	}
+	else
+	{
+		int32 NumBytes = CompressedKeys.Num();
+		Ar << NumBytes;
+		Ar.Serialize(CompressedKeys.GetData(), NumBytes);
+	}
+
+	return true;
+}
+
+bool FCompressedRichCurve::operator==(const FCompressedRichCurve& Other) const
+{
+	return CompressionFormat == Other.CompressionFormat
+		&& KeyTimeCompressionFormat == Other.KeyTimeCompressionFormat
+		&& PreInfinityExtrap == Other.PreInfinityExtrap
+		&& PostInfinityExtrap == Other.PostInfinityExtrap
+		&& ConstantValueNumKeys.NumKeys == Other.ConstantValueNumKeys.NumKeys
+		&& CompressedKeys == Other.CompressedKeys;
+}
diff --git a/Engine/Source/Runtime/Engine/Private/UnrealEngine.cpp b/Engine/Source/Runtime/Engine/Private/UnrealEngine.cpp
--- a/Engine/Source/Runtime/Engine/Private/UnrealEngine.cpp
+++ b/Engine/Source/Runtime/Engine/Private/UnrealEngine.cpp
@@ -5329,7 +5329,7 @@ bool UEngine::HandleListAnimsCommand(const TCHAR* Cmd, FOutputDevice& Ar)
 			NumKeys = AnimSeq->GetRawNumberOfFrames();
 			SequenceLength = AnimSeq->GetPlayLength();
 			RateScale = AnimSeq->RateScale;
-			NumCurves = AnimSeq->CompressedCurveData.FloatCurves.Num();
+			NumCurves = AnimSeq->RawCurveData.FloatCurves.Num();
 
 			TranslationFormat = FAnimationUtils::GetAnimationCompressionFormatString(AnimSeq->TranslationCompressionFormat);
 			RotationFormat = FAnimationUtils::GetAnimationCompressionFormatString(AnimSeq->RotationCompressionFormat);
diff --git a/Engine/Source/Runtime/Engine/Public/Animation/AnimTypes.h b/Engine/Source/Runtime/Engine/Public/Animation/AnimTypes.h
--- a/Engine/Source/Runtime/Engine/Public/Animation/AnimTypes.h
+++ b/Engine/Source/Runtime/Engine/Public/Animation/AnimTypes.h
@@ -625,6 +625,12 @@ struct FAnimKeyHelper
 		return (NumKeys > 1) ? Length / (float)(NumKeys - 1) : MINIMUM_ANIMATION_LENGTH;
 	}
 
+	// Returns the FPS for the given length and number of keys
+	float KeysPerSecond() const
+	{
+		return NumKeys > 0 ? (float(NumKeys - 1) / Length) : 0.0f;
+	}
+
 	int32 LastKey() const
 	{
 		return (NumKeys > 1) ? NumKeys - 1 : 0;
diff --git a/Engine/Source/Runtime/Engine/Public/AnimationUtils.h b/Engine/Source/Runtime/Engine/Public/AnimationUtils.h
--- a/Engine/Source/Runtime/Engine/Public/AnimationUtils.h
+++ b/Engine/Source/Runtime/Engine/Public/AnimationUtils.h
@@ -276,4 +276,16 @@ public:
 	 */
 	ENGINE_API static UAnimCompress* GetDefaultAnimationCompressionAlgorithm();
 
+#if WITH_EDITOR
+	/** Returns the default animation curve compression settings, can never by null. */
+	ENGINE_API static UAnimCurveCompressionSettings* GetDefaultAnimationCurveCompressionSettings();
+
+	/**
+	 * Compresses the animation curves within a sequence with the chosen settings.
+	 * Note: This modifies the sequence.
+	 *
+	 * @return		Returns true on success, false it we fail to compress
+	 */
+	ENGINE_API static bool CompressAnimCurves(UAnimSequence& AnimSeq);
+#endif
 };
