Fix merge error from Main

diff --git a/Engine/Source/Runtime/Engine/Private/Curves/RichCurve.cpp b/Engine/Source/Runtime/Engine/Private/Curves/RichCurve.cpp
--- a/Engine/Source/Runtime/Engine/Private/Curves/RichCurve.cpp
+++ b/Engine/Source/Runtime/Engine/Private/Curves/RichCurve.cpp
@@ -1567,6 +1567,35 @@ struct MixedKeyDataAdapter
 	}
 };
 
+static void CycleTime(float MinTime, float MaxTime, float& InTime, int& CycleCount)
+{
+	float InitTime = InTime;
+	float Duration = MaxTime - MinTime;
+
+	if (InTime > MaxTime)
+	{
+		CycleCount = FMath::FloorToInt((MaxTime-InTime)/Duration);
+		InTime = InTime + Duration*CycleCount;
+	}
+	else if (InTime < MinTime)
+	{
+		CycleCount = FMath::FloorToInt((InTime-MinTime)/Duration);
+		InTime = InTime - Duration*CycleCount;
+	}
+
+	if (InTime == MaxTime && InitTime < MinTime)
+	{
+		InTime = MinTime;
+	}
+
+	if (InTime == MinTime && InitTime > MaxTime)
+	{
+		InTime = MaxTime;
+	}
+
+	CycleCount = FMath::Abs(CycleCount);
+}
+
 template<typename KeyTimeAdapterType, typename KeyDataAdapterType>
 static float RemapTimeValue(float InTime, const KeyTimeAdapterType& KeyTimeAdapter, const KeyDataAdapterType& KeyDataAdapter, int32 NumKeys, ERichCurveExtrapolation InfinityExtrap, int32 KeyIndex0, int32 KeyIndex1, float& CycleValueOffset)
 {
