Improved acks Part 1/2

 - These are the supporting code required for changes to NetConnection which iis submitted in a separate changelist

Original  changes from dev-networking
CL 4687651, CL 4687649, CL 4684239, 4688073, 4647801

diff --git a/Engine/Source/Runtime/Engine/Private/Net/NetPacketNotify.cpp b/Engine/Source/Runtime/Engine/Private/Net/NetPacketNotify.cpp
new file mode 100644
--- /dev/null
+++ b/Engine/Source/Runtime/Engine/Private/Net/NetPacketNotify.cpp
@@ -0,0 +1,139 @@
+// Copyright 1998-2019 Epic Games, Inc. All Rights Reserved.
+
+#include "Net/NetPacketNotify.h"
+#include "Serialization/BitWriter.h"
+#include "Serialization/BitReader.h"
+
+SIZE_T FNetPacketNotify::GetCurrentSequenceHistoryLength() const
+{
+	if (InAckSeq >= InAckSeqAck)
+	{
+		return (SIZE_T)SequenceNumberT::Diff(InAckSeq, InAckSeqAck);
+	}
+	else
+	{
+		// Worst case send full history
+		return SequenceHistoryT::Size;
+	}
+}
+
+void FNetPacketNotify::Init(SequenceNumberT InitialInSeq, SequenceNumberT InitialOutSeq)
+{
+	InSeqHistory.Reset();
+	InSeq = InitialInSeq;
+	InAckSeq = InitialInSeq;
+	InAckSeqAck = InitialInSeq;
+	OutSeq = InitialOutSeq;
+	OutAckSeq = SequenceNumberT(InitialOutSeq.Get() - 1);
+}
+
+void FNetPacketNotify::AckSeq(SequenceNumberT AckedSeq, bool IsAck)
+{
+	check( AckedSeq == InSeq);
+
+	while (AckedSeq > InAckSeq)
+	{
+		++InAckSeq;
+
+		const bool bReportAcked = InAckSeq == AckedSeq ? IsAck : false;
+
+		UE_LOG_PACKET_NOTIFY(TEXT("FNetPacketNotify::AckSeq - AckedSeq: %u, IsAck %u"), InAckSeq.Get(), bReportAcked ? 1u : 0u);
+
+		InSeqHistory.AddDeliveryStatus(bReportAcked);		
+	}
+}
+
+namespace 
+{
+	struct FPackedHeader
+	{
+		using SequenceNumberT = FNetPacketNotify::SequenceNumberT;
+
+		static_assert(FNetPacketNotify::SequenceNumberBits <= 14, "SequenceNumbers must be smaller than 14 bits to fit history word count");
+
+		enum { HistoryWordCountBits = 4 };
+		enum { SeqMask				= (1 << FNetPacketNotify::SequenceNumberBits) - 1 };
+		enum { HistoryWordCountMask	= (1 << HistoryWordCountBits) - 1 };
+		enum { AckSeqShift			= HistoryWordCountBits };
+		enum { SeqShift				= AckSeqShift + FNetPacketNotify::SequenceNumberBits };
+		
+		static uint32 Pack(SequenceNumberT Seq, SequenceNumberT AckedSeq, SIZE_T HistoryWordCount)
+		{
+			uint32 Packed = 0u;
+
+			Packed |= Seq.Get() << SeqShift;
+			Packed |= AckedSeq.Get() << AckSeqShift;
+			Packed |= HistoryWordCount & HistoryWordCountMask;
+
+			return Packed;
+		}
+
+		static SequenceNumberT GetSeq(uint32 Packed) { return SequenceNumberT(Packed >> SeqShift & SeqMask); }
+		static SequenceNumberT GetAckedSeq(uint32 Packed) { return SequenceNumberT(Packed >> AckSeqShift & SeqMask); }
+		static SIZE_T GetHistoryWordCount(uint32 Packed) { return (Packed & HistoryWordCountMask); }
+	};
+}
+
+// These methods must always write and read the exact same number of bits, that is the reason for not using WriteInt/WrittedWrappedInt
+void FNetPacketNotify::WriteHeader(FBitWriter& Writer, bool bRefresh)
+{
+	// we always write at least 1 word
+	SIZE_T CurrentHistoryWordCount = FMath::Clamp<SIZE_T>((GetCurrentSequenceHistoryLength() + SequenceHistoryT::BitsPerWord - 1u) / SequenceHistoryT::BitsPerWord, 1u, SequenceHistoryT::WordCount);
+
+	// We can only do a refresh if we do not need more space for the history
+	if (bRefresh && (CurrentHistoryWordCount > WrittenHistoryWordCount))
+	{
+		return;
+	}
+
+	// How many words of ack data should we write? If this is a refresh we must write the same size as the original header
+	WrittenHistoryWordCount = bRefresh ? WrittenHistoryWordCount : CurrentHistoryWordCount;
+	// This is the last InAck we have acknowledged at this time
+	WrittenInAckSeq = InAckSeq;
+
+	SequenceNumberT::SequenceT Seq = OutSeq.Get();
+	SequenceNumberT::SequenceT AckedSeq = InAckSeq.Get();
+
+	// Pack data into a uint
+	uint32 PackedHeader = FPackedHeader::Pack(Seq, AckedSeq, WrittenHistoryWordCount - 1);
+
+	// Write packed header
+	Writer << PackedHeader;
+
+	// Write ack history
+	InSeqHistory.Write(Writer, WrittenHistoryWordCount);
+
+	UE_LOG_PACKET_NOTIFY(TEXT("FNetPacketNotify::WriteHeader - Seq %u, AckedSeq %u bReFresh %u HistorySizeInWords %u"), Seq, AckedSeq, bRefresh ? 1u : 0u, WrittenHistoryWordCount);
+}
+
+bool FNetPacketNotify::ReadHeader(FNotificationHeader& Data, FBitReader& Reader) const
+{
+	// Read packed header
+	uint32 PackedHeader = 0;	
+	Reader << PackedHeader;
+
+	// unpack
+	Data.Seq = FPackedHeader::GetSeq(PackedHeader);
+	Data.AckedSeq = FPackedHeader::GetAckedSeq(PackedHeader);
+	Data.HistoryWordCount = FPackedHeader::GetHistoryWordCount(PackedHeader) + 1;
+
+	// Read ack history
+	Data.History.Read(Reader, Data.HistoryWordCount);
+
+	UE_LOG_PACKET_NOTIFY(TEXT("FNetPacketNotify::ReadHeader - Seq %u, AckedSeq %u HistorySizeInWords %u"), Data.Seq.Get(), Data.AckedSeq.Get(), Data.HistoryWordCount);
+
+	return Reader.IsError() == false;
+}
+
+FNetPacketNotify::SequenceNumberT FNetPacketNotify::CommitAndIncrementOutSeq()
+{
+	// we have not written a header...this is a fail.
+	check(WrittenHistoryWordCount != 0);
+
+	// Add entry to the ack-record so that we can update the InAckSeqAck when we received the ack for this OutSeq.
+	AckRecord.Enqueue( {OutSeq, WrittenInAckSeq} );
+	WrittenHistoryWordCount = 0u;
+	
+	return ++OutSeq;
+}
+
diff --git a/Engine/Source/Runtime/Engine/Private/Net/Tests/NetPacketNotifyTests.cpp b/Engine/Source/Runtime/Engine/Private/Net/Tests/NetPacketNotifyTests.cpp
new file mode 100644
--- /dev/null
+++ b/Engine/Source/Runtime/Engine/Private/Net/Tests/NetPacketNotifyTests.cpp
@@ -0,0 +1,303 @@
+// Copyright 1998-2019 Epic Games, Inc. All Rights Reserved.
+
+#include "Misc/AutomationTest.h"
+#include "Net/NetPacketNotify.h"
+
+#if WITH_DEV_AUTOMATION_TESTS
+
+IMPLEMENT_SIMPLE_AUTOMATION_TEST(FNetPacketNotifyTest, "Network.PacketNotifyTest", EAutomationTestFlags::ApplicationContextMask | EAutomationTestFlags::EngineFilter)
+
+struct FNetPacketNotifyTestUtil
+{
+	FNetPacketNotify DefaultNotify;
+	FNetPacketNotifyTestUtil()
+	{
+		DefaultNotify.Init(FNetPacketNotify::SequenceNumberT(-1),  FNetPacketNotify::SequenceNumberT(0));
+	}
+
+	// Helper to fill in SequenceHistory with expected result
+	template <typename T>
+	static void InitHistory(FNetPacketNotify::SequenceHistoryT& History, const T& DataToSet)
+	{
+		const SIZE_T Count = sizeof(T) / sizeof(DataToSet[0]);
+		static_assert(Count < FNetPacketNotify::SequenceHistoryT::WordCount, "DataToSet must be smaller than HistoryBuffer");
+
+		for (SIZE_T It=0; It < Count; ++It)
+		{
+			History.Data()[It] = DataToSet[It];
+		}
+	}
+
+	// Pretend to receive and acknowledge incoming packet to generate ackdata
+	static int32 PretendReceiveSeq(FNetPacketNotify& PacketNotify, FNetPacketNotify::SequenceNumberT Seq, bool Ack = true)
+	{
+		FNetPacketNotify::FNotificationHeader Data;
+		Data.Seq = Seq;
+		Data.AckedSeq = PacketNotify.GetOutAckSeq();
+		Data.History = FNetPacketNotify::SequenceHistoryT(0);
+		Data.HistoryWordCount = 1;
+		
+		FNetPacketNotify::SequenceNumberT::DifferenceT SeqDelta = PacketNotify.Update(Data, [](FNetPacketNotify::SequenceNumberT AckedSequence, bool delivered) {});
+		if (SeqDelta > 0)
+		{
+			if (Ack)
+				PacketNotify.AckSeq(Seq);
+		}
+
+		return SeqDelta;
+	}
+
+	// Pretend to send packet
+	static void PretendSendSeq(FNetPacketNotify& PacketNotify, FNetPacketNotify::SequenceNumberT LastAckSeq)
+	{
+		// set last InAcqSeq that we know that the remote end knows that we know (AckAck)
+		PacketNotify.WrittenHistoryWordCount = 1;
+		PacketNotify.WrittenInAckSeq = LastAckSeq;
+
+		// Store data
+		PacketNotify.CommitAndIncrementOutSeq();
+	}
+
+	// pretend to ack array of sequence numbers
+	template<typename T>
+	static void PretendAckSequenceNumbers(FNetPacketNotify& PacketNotify, const T& InSequenceNumbers)
+	{
+		SIZE_T SequenceNumberCount = sizeof(InSequenceNumbers) / sizeof(InSequenceNumbers[0]);
+
+		for (SIZE_T I=0; I<SequenceNumberCount; ++I)
+		{
+			FNetPacketNotifyTestUtil::PretendReceiveSeq(PacketNotify, InSequenceNumbers[I]);
+		}
+	}
+	
+	// Pretend that we received a packet
+	template<typename T>
+	static SIZE_T PretendReceivedPacket(FNetPacketNotify& PacketNotify, const FNetPacketNotify::FNotificationHeader Data, T& OutSequenceNumbers)
+	{
+		SIZE_T NotificationCount = 0;
+
+		auto HandleAck = [&OutSequenceNumbers, &NotificationCount](FNetPacketNotify::SequenceNumberT Seq, bool delivered)
+		{
+			const SIZE_T MaxSequenceNumberCount = sizeof(OutSequenceNumbers) / sizeof(OutSequenceNumbers[0]);
+
+			if (delivered)
+			{
+				if (NotificationCount < MaxSequenceNumberCount)
+				{
+					OutSequenceNumbers[NotificationCount] = Seq;
+				}
+				++NotificationCount;		
+			}
+		};
+		return PacketNotify.Update(Data, HandleAck);
+	}
+
+	// Test to fake sending and receiving an array of sequence numbers and test if we get the expected notifications back
+	template<typename T>
+	static bool TestNotificationSequence(const T& InSequenceNumbers, FNetPacketNotify::SequenceNumberT FirstSequence = 0)
+	{
+		T NotifiedSequenceNumbers = { 0 };
+
+		// Sender, which we will also be the receiver of the acks
+		FNetPacketNotify Sender;
+		Sender.Init(FNetPacketNotify::SequenceNumberT(FirstSequence.Get() - 1), FirstSequence);
+	
+		// pretend that we have sent the InSequenceNumbers
+		SIZE_T Count = sizeof(T) / sizeof(InSequenceNumbers[0]);
+		for (SIZE_T It=0; It < Count; ++It)
+		{			
+			do
+			{
+				FNetPacketNotifyTestUtil::PretendSendSeq(Sender, 0);
+			}
+			while (InSequenceNumbers[It] >= Sender.GetOutSeq());
+		}
+
+		// Receiver which we fake have received the packets sent from sender
+		FNetPacketNotify Receiver;
+		Receiver.Init(FNetPacketNotify::SequenceNumberT(FirstSequence.Get() - 1), FirstSequence);
+		PretendAckSequenceNumbers(Receiver, InSequenceNumbers);
+
+		// Fake Header with acks sent from receiver back to sender
+		FNetPacketNotify::FNotificationHeader Data;
+		Data.Seq = Receiver.GetOutSeq();
+		Data.AckedSeq = Receiver.GetInAckSeq();
+		Data.HistoryWordCount = FNetPacketNotify::SequenceHistoryT::WordCount;
+		Data.History = Receiver.GetInSeqHistory();
+
+		// Process the received ack information
+		PretendReceivedPacket(Sender, Data, NotifiedSequenceNumbers);
+
+		// Check that it matches the expected result
+		return FPlatformMemory::Memcmp(&InSequenceNumbers[0], &NotifiedSequenceNumbers[0], sizeof(InSequenceNumbers)) == 0;
+	}
+};
+
+bool FNetPacketNotifyTest::RunTest(const FString& Parameters)
+{
+	FNetPacketNotifyTestUtil Util;
+	// Test fill
+	{
+		FNetPacketNotify::SequenceNumberT ExpectedInSeq(31);
+		FNetPacketNotify::SequenceHistoryT ExpectedInSeqHistory(0xffffffffu, 1);
+
+		FNetPacketNotify Acks = Util.DefaultNotify;
+				
+		for (int32 I=0; I<32; ++I)
+		{
+			FNetPacketNotifyTestUtil::PretendReceiveSeq(Acks, I, true);
+		}
+
+		TestEqual(TEXT("Test fill - InSeq"), Acks.GetInSeq(), ExpectedInSeq);
+		TestEqual(TEXT("Test fill - History"), Acks.GetInSeqHistory(), ExpectedInSeqHistory);
+	}
+
+	// Test drop every other
+	{
+		FNetPacketNotify::SequenceNumberT ExpectedInSeq(30);
+		FNetPacketNotify::SequenceHistoryT ExpectedInSeqHistory(0x55555555u, 1);
+
+		FNetPacketNotify Acks = Util.DefaultNotify;
+				
+		for (int32 I=0; I<16; ++I)
+		{
+			FNetPacketNotifyTestUtil::PretendReceiveSeq(Acks, I*2);
+		}
+		
+		TestEqual(TEXT("Test drop every other - InSeq"), Acks.GetInSeq(), ExpectedInSeq);
+		TestEqual(TEXT("Test drop every other - History"), Acks.GetInSeqHistory(), ExpectedInSeqHistory);
+	}
+
+	// Test burst drop
+	{
+		FNetPacketNotify::SequenceNumberT ExpectedInSeq(128);
+		FNetPacketNotify::SequenceHistoryT ExpectedInSeqHistory;
+		uint32 ExpectedArray[] = {0x1, 0, 0, 0x20000000 };
+		FNetPacketNotifyTestUtil::InitHistory(ExpectedInSeqHistory, ExpectedArray );
+
+		FNetPacketNotify Acks = Util.DefaultNotify;
+
+		// Drop early
+		FNetPacketNotifyTestUtil::PretendReceiveSeq(Acks, 3);
+
+		// Large gap until next seq
+		FNetPacketNotifyTestUtil::PretendReceiveSeq(Acks, 128);
+
+		TestEqual(TEXT("Test burst drop - InSeq"), Acks.GetInSeq(), ExpectedInSeq);
+		TestEqual(TEXT("Test burst drop - History"), Acks.GetInSeqHistory(), ExpectedInSeqHistory);
+	}
+
+	// Test window overflow
+	{
+		FNetPacketNotify Acks = Util.DefaultNotify;
+
+		const FNetPacketNotify::SequenceNumberT ExpectedInSeq(0);
+		const FNetPacketNotify::SequenceNumberT MaxWindowSeq(FNetPacketNotify::SequenceNumberT::SeqNumberHalf);
+
+		FNetPacketNotifyTestUtil::PretendReceiveSeq(Acks, ExpectedInSeq, true);
+		TestEqual(TEXT("Test window overflow - Expect InSeq 0"), Acks.GetInSeq(), ExpectedInSeq);
+		TestEqual(TEXT("Test window overflow - Expect InAckSeq 0"), Acks.GetInAckSeq(), ExpectedInSeq);
+
+		FNetPacketNotifyTestUtil::PretendReceiveSeq(Acks, MaxWindowSeq, true);
+		TestEqual(TEXT("Test window overflow - Expect Seq reject"), Acks.GetInSeq(), ExpectedInSeq);
+		TestEqual(TEXT("Test window overflow - Expect Seq reject"), Acks.GetInAckSeq(), ExpectedInSeq);
+
+		const FNetPacketNotify::SequenceNumberT NextExpectedInSeq(1);
+		FNetPacketNotifyTestUtil::PretendReceiveSeq(Acks, NextExpectedInSeq, true);
+		TestEqual(TEXT("Test window overflow - Expect InSeq 1"), Acks.GetInSeq(), NextExpectedInSeq);
+		TestEqual(TEXT("Test window overflow - Expect InAckSeq 1"), Acks.GetInAckSeq(), NextExpectedInSeq);
+
+		FNetPacketNotifyTestUtil::PretendReceiveSeq(Acks, MaxWindowSeq, true);
+		TestEqual(TEXT("Test window overflow - Expect InSeq MaxWindowSeq"), Acks.GetInSeq(), MaxWindowSeq);
+		TestEqual(TEXT("Test window overflow - Expect InSeq MaxWindowSeq"), Acks.GetInAckSeq(), MaxWindowSeq);
+	}
+
+	// create history
+	{
+		const FNetPacketNotify::SequenceNumberT ExpectedInSeq(18);
+		const FNetPacketNotify::SequenceHistoryT ExpectedInSeqHistory(0x8853u);
+
+		const FNetPacketNotify::SequenceNumberT AckdPacketIds[] = {3, 7, 12, 14, 17, 18};
+		const SIZE_T ExpectedCount = sizeof(AckdPacketIds)/sizeof((AckdPacketIds)[0]);		
+
+		FNetPacketNotify Acks = Util.DefaultNotify;
+		FNetPacketNotifyTestUtil::PretendAckSequenceNumbers(Acks, AckdPacketIds);
+ 
+		TestEqual(TEXT("Create history - InSeq"), Acks.GetInSeq(), ExpectedInSeq);
+		TestEqual(TEXT("Create history - InSeqHistory"), Acks.GetInSeqHistory(), ExpectedInSeqHistory);
+	}
+	
+	// test notifications
+	{
+		static const FNetPacketNotify::SequenceNumberT ExpectedAckdPacketIds[] = {3, 7, 12, 14, 17, 18};
+		static const SIZE_T ExpectedCount = sizeof(ExpectedAckdPacketIds)/sizeof((ExpectedAckdPacketIds)[0]);		
+
+		FNetPacketNotify::SequenceNumberT RcvdAcks[ExpectedCount] = { 0 };
+		SIZE_T RcvdCount = 0;
+
+		// Create src data
+		FNetPacketNotify Acks = Util.DefaultNotify;
+
+		// Fill in some data
+		FNetPacketNotify::FNotificationHeader Data;
+		Data.Seq = FNetPacketNotify::SequenceNumberT(0);
+		Data.AckedSeq = FNetPacketNotify::SequenceNumberT(18);
+		Data.History = FNetPacketNotify::SequenceHistoryT(0x8853u);
+		Data.HistoryWordCount = 1;
+
+		// Need to fake ack record as well.
+		for (SIZE_T It=0; It <= 18; ++It)
+		{
+			FNetPacketNotifyTestUtil::PretendSendSeq(Acks, 0);
+		}
+	
+		SIZE_T DeltaSeq = FNetPacketNotifyTestUtil::PretendReceivedPacket(Acks, Data, RcvdAcks);
+
+		TestEqual(TEXT("Notifications - Create sequence delta"), DeltaSeq, 1);
+		TestEqual(TEXT("Notifications - Create sequence"), FPlatformMemory::Memcmp(ExpectedAckdPacketIds, RcvdAcks, sizeof(ExpectedAckdPacketIds)), 0u);
+	}
+
+	// test various sequence
+	{
+		static const FNetPacketNotify::SequenceNumberT TestSeqs[] = {3, 7, 12, 14, 17, 18};	
+		TestTrue(TEXT("Test Seq {3, 7, 12, 14, 17, 18}"), FNetPacketNotifyTestUtil::TestNotificationSequence(TestSeqs));
+	}
+	{
+		static const FNetPacketNotify::SequenceNumberT TestSeqs[] = {0, 1, 2, 3, 4, 5, 6, 7, 8, 9, 10, 11, 12, 13, 14, 15, 16, 17, 18, 19, 20, 21, 22, 23, 24, 25, 26, 27, 28, 29, 30, 31};	
+		TestTrue(TEXT("Test Seq {0, 1, 2, 3, 4, 5, 6, 7, 8, 9, 10, 11, 12, 13, 14, 15, 16, 17, 18, 19, 20, 21, 22, 23, 24, 25, 26, 27, 28, 29, 30, 31};"), FNetPacketNotifyTestUtil::TestNotificationSequence(TestSeqs));
+	}
+	{
+		static const FNetPacketNotify::SequenceNumberT TestSeqs[] = {2, 3, 4, 5, 6, 7, 8, 9, 10, 12, 13, 14, 16, 17, 18, 19, 20, 21, 22, 23, 24, 25, 26, 27, 28, 29, 30, 31};	
+		TestTrue(TEXT("{2, 3, 4, 5, 6, 7, 8, 9, 10, 12, 13, 14, 16, 17, 18, 19, 20, 21, 22, 23, 24, 25, 26, 27, 28, 29, 30, 31}"), FNetPacketNotifyTestUtil::TestNotificationSequence(TestSeqs));
+	}
+	{
+		static const FNetPacketNotify::SequenceNumberT TestSeqs[] = {0, FNetPacketNotify::MaxSequenceHistoryLength - 1};	
+		TestTrue(TEXT("Test Seq {0, FNetPacketNotify::MaxSequenceHistoryLength - 1}"), FNetPacketNotifyTestUtil::TestNotificationSequence(TestSeqs));
+	}
+	{
+		static const FNetPacketNotify::SequenceNumberT TestSeqs[] = {0, FNetPacketNotify::MaxSequenceHistoryLength};	
+		TestFalse(TEXT("Test Seq {0, FNetPacketNotify::MaxSequenceHistoryLength}"), FNetPacketNotifyTestUtil::TestNotificationSequence(TestSeqs));
+	}
+	{
+		static const FNetPacketNotify::SequenceNumberT TestSeqs[] = {0, FNetPacketNotify::SequenceNumberT::SeqNumberHalf };
+		TestFalse(TEXT("Test Seq {FNetPacketNotify::SequenceNumberT::SeqNumberMax, 0};"), FNetPacketNotifyTestUtil::TestNotificationSequence(TestSeqs));
+	}
+	{
+		static const FNetPacketNotify::SequenceNumberT TestSeqs[] = {0, FNetPacketNotify::SequenceNumberT::SeqNumberHalf - 1};
+		TestFalse(TEXT("Test Seq {FNetPacketNotify::SequenceNumberT::SeqNumberMax, 0};"), FNetPacketNotifyTestUtil::TestNotificationSequence(TestSeqs));
+	}
+	{
+		static const FNetPacketNotify::SequenceNumberT TestSeqs[] = {FNetPacketNotify::SequenceNumberT::SeqNumberMax, 0};
+		TestFalse(TEXT("Test Seq {FNetPacketNotify::SequenceNumberT::SeqNumberMax, 0} From 0;"), FNetPacketNotifyTestUtil::TestNotificationSequence(TestSeqs));
+	}
+	{
+		static const FNetPacketNotify::SequenceNumberT TestSeqs[] = {FNetPacketNotify::SequenceNumberT::SeqNumberMax, 0};
+		TestTrue(TEXT("Test Seq {FNetPacketNotify::SequenceNumberT::SeqNumberMax, 0} From SeqNumberHalf + 1;"), FNetPacketNotifyTestUtil::TestNotificationSequence(TestSeqs, FNetPacketNotify::SequenceNumberT::SeqNumberHalf + 2));
+	}
+
+	check(!HasAnyErrors());
+
+	return true;
+}
+
+#endif // #if WITH_DEV_AUTOMATION_TESTS
\ No newline at end of file
diff --git a/Engine/Source/Runtime/Engine/Private/Net/Tests/Util/ResizableCircularQueueTests.cpp b/Engine/Source/Runtime/Engine/Private/Net/Tests/Util/ResizableCircularQueueTests.cpp
new file mode 100644
--- /dev/null
+++ b/Engine/Source/Runtime/Engine/Private/Net/Tests/Util/ResizableCircularQueueTests.cpp
@@ -0,0 +1,218 @@
+// Copyright 1998-2019 Epic Games, Inc. All Rights Reserved.
+
+#include "Misc/AutomationTest.h"
+#include "Net/Util/ResizableCircularQueue.h"
+
+#if WITH_DEV_AUTOMATION_TESTS
+
+IMPLEMENT_SIMPLE_AUTOMATION_TEST(FResizableCircularQueueTest, "Network.ResizableCircularQueueTest", EAutomationTestFlags::ApplicationContextMask | EAutomationTestFlags::EngineFilter)
+
+struct FResizableCircularQueueTestUtil
+{
+	typedef TResizableCircularQueue<uint32> QueueT;
+
+	static bool VerifyQueueIntegrity(const QueueT& Queue, uint32 ExpectedValueAtFront, uint32 Increment)
+	{
+		bool bSuccess = true;
+		SIZE_T Offset = Queue.Count() - 1;
+	
+		uint32 ExpectedValue = ExpectedValueAtFront;
+
+		// Peek elements in queue at given offset, peek from back to front
+		for (SIZE_T It = 0; It < Queue.Count(); ++It)
+		{
+			bSuccess = bSuccess && (ExpectedValue == Queue.PeekAtOffset(It));
+			ExpectedValue += Increment;
+		}
+
+		return bSuccess;
+	}
+
+	static void OverideHeadAndTail(QueueT& Queue, uint32 Head, uint32 Tail)
+	{
+		Queue.Head = Head;
+		Queue.Tail = Tail;
+	}
+};
+
+bool FResizableCircularQueueTest::RunTest(const FString& Parameters)
+{
+	// Test empty
+	{
+		FResizableCircularQueueTestUtil::QueueT Q(0);
+
+		TestEqual(TEXT("Test empty - Size"), Q.Count(), 0);
+		TestTrue(TEXT("Test empty - IsEmpty"), Q.IsEmpty());
+		TestEqual(TEXT("Test empty - Capacity"), Q.AllocatedCapacity(), 0);
+	}
+
+	// Test Push to Capacity
+	{
+		const SIZE_T ElementsToPush = 8;
+
+		FResizableCircularQueueTestUtil::QueueT Q(ElementsToPush);
+
+		for (SIZE_T It=0; It < ElementsToPush; ++It)
+		{
+			Q.Enqueue(It);
+		}
+		
+		TestEqual(TEXT("Test Push to Capacity - Size"), Q.Count(), ElementsToPush);
+		TestEqual(TEXT("Test Push to Capacity - Capacity"), Q.AllocatedCapacity(), ElementsToPush);
+ 		TestTrue(TEXT("Test Push to Capacity - Expected"), FResizableCircularQueueTestUtil::VerifyQueueIntegrity(Q, 0, 1));
+	}
+
+	// Test Push over Capacity
+	{
+		const SIZE_T ElementsToPush = 32;
+
+		FResizableCircularQueueTestUtil::QueueT Q(ElementsToPush);
+
+		for (SIZE_T It=0; It < ElementsToPush; ++It)
+		{
+			Q.Enqueue(It);
+		}
+		
+		TestEqual(TEXT("Test Push over Capacity - Size"), Q.Count(), ElementsToPush);
+		TestEqual(TEXT("Test Push over Capacity - Capacity"), Q.AllocatedCapacity(), ElementsToPush);
+ 		TestTrue(TEXT("Test Push over Capacity - Expected"), FResizableCircularQueueTestUtil::VerifyQueueIntegrity(Q, 0, 1));
+	}
+
+	// Test Push and Pop
+	{
+		const SIZE_T ElementsToPush = 256;
+		const SIZE_T ElementPopMod = 16;
+		const SIZE_T ExpectedSize = 256 - ElementPopMod;
+		const SIZE_T ExpectedCapacity = 256;
+
+		FResizableCircularQueueTestUtil::QueueT Q(4);
+
+		uint32 ExpectedPoppedValue = 0;
+		for (uint32 It = 0; It < 256; ++It)
+		{
+			Q.Enqueue(It);
+			TestEqual(TEXT("Test Push and pop - Push"), It, Q.PeekAtOffset(Q.Count() - 1));
+
+			if (It % ElementPopMod == 0)
+			{
+				const uint32 PoppedValue = Q.PeekAtOffset(0);
+				TestEqual(TEXT("Test Push and pop - Pop"), ExpectedPoppedValue, PoppedValue);
+				++ExpectedPoppedValue;
+				Q.Pop();
+			}
+		}
+
+		TestEqual(TEXT("Test Push and pop - Size"), Q.Count(), ExpectedSize);
+		TestEqual(TEXT("Test Push and pop - Capacity"), Q.AllocatedCapacity(), ExpectedCapacity);
+ 		TestTrue (TEXT("Test Push and pop - Expected"), FResizableCircularQueueTestUtil::VerifyQueueIntegrity(Q, ElementPopMod, 1));	
+	}
+
+	// Test Push and pop all
+	{
+		const SIZE_T ElementsToPush = 256;
+
+		FResizableCircularQueueTestUtil::QueueT Q(ElementsToPush);
+
+		TestTrue(TEXT ("Test Push and pop all - IsEmpty before"), Q.IsEmpty());
+		TestEqual(TEXT("Test Push and pop all - Size before"), Q.Count(), 0);
+
+		for (SIZE_T It=0; It < ElementsToPush; ++It)
+		{
+			Q.Enqueue(It);
+		}
+		
+		TestEqual(TEXT("Test Push and pop all - Size"), Q.Count(), ElementsToPush);
+		TestEqual(TEXT("Test Push and pop all - Capacity"), Q.AllocatedCapacity(), ElementsToPush);
+ 		TestTrue (TEXT("Test Push and pop all - Expected"), FResizableCircularQueueTestUtil::VerifyQueueIntegrity(Q, 0, 1));
+
+		for (SIZE_T It=0; It < ElementsToPush; ++It)
+		{
+			Q.Pop();
+		}
+		
+		TestTrue( TEXT("Test Push and pop all - IsEmpty after"), Q.IsEmpty());
+		TestEqual(TEXT("Test Push and pop all - Size after"), Q.Count(), 0);
+		TestEqual(TEXT("Test Push and pop all - Capacity after"), Q.AllocatedCapacity(), ElementsToPush);
+	}
+
+	// Test index wrap
+	{
+		const SIZE_T ElementsToPush = 256;
+		const SIZE_T ElementPopMod = 16;
+		const SIZE_T ExpectedSize = 256 - ElementPopMod;
+		const SIZE_T ExpectedCapacity = 256;
+
+		FResizableCircularQueueTestUtil::QueueT Q(4);
+
+		// Set head and tail at the end of the space to test index wraparound
+		FResizableCircularQueueTestUtil::OverideHeadAndTail(Q, uint32(-2), uint32(-2));
+
+		TestTrue(TEXT ("Test index wrap - IsEmpty before"), Q.IsEmpty());
+		TestEqual(TEXT("Test index wrap - Size before"), Q.Count(), 0);
+
+		for (SIZE_T It=0; It < ElementsToPush; ++It)
+		{
+			Q.Enqueue(It);
+		}
+		
+		TestEqual(TEXT("Test index wrap - Size"), Q.Count(), ElementsToPush);
+		TestEqual(TEXT("Test index wrap - Capacity"), Q.AllocatedCapacity(), ElementsToPush);
+ 		TestTrue (TEXT("Test index wrap - Expected"), FResizableCircularQueueTestUtil::VerifyQueueIntegrity(Q, 0, 1));
+
+		for (SIZE_T It=0; It < ElementsToPush; ++It)
+		{
+			Q.Pop();
+		}
+		
+		TestTrue( TEXT("Test index wrap - IsEmpty after"), Q.IsEmpty());
+		TestEqual(TEXT("Test index wrap - Size after"), Q.Count(), 0);
+		TestEqual(TEXT("Test index wrap - Capacity after"), Q.AllocatedCapacity(), ElementsToPush);
+	}
+
+	// Test Trim
+	{
+		const SIZE_T ElementsToPush = 9;
+		const SIZE_T ElementsToPop = 5;
+		const SIZE_T ExpectedCapacity = 16;
+		const SIZE_T ExpectedCapacityAfterTrim = 4;
+
+		FResizableCircularQueueTestUtil::QueueT Q(0);
+
+		for (SIZE_T It=0; It < ElementsToPush; ++It)
+		{
+			Q.Enqueue(It);
+		}
+		
+		TestEqual(TEXT("Test Trim - Size"), Q.Count(), ElementsToPush);
+		TestEqual(TEXT("Test Trim - Capacity"), Q.AllocatedCapacity(), ExpectedCapacity);
+ 		TestTrue(TEXT("Test Trim - Expected"), FResizableCircularQueueTestUtil::VerifyQueueIntegrity(Q, 0, 1));
+
+		for (SIZE_T It=0; It < ElementsToPop; ++It)
+		{
+			Q.Pop();
+		}
+
+		Q.Trim();
+
+		TestEqual(TEXT("Test Trim - Size"), Q.Count(), ElementsToPush - ElementsToPop);
+		TestEqual(TEXT("Test Trim - Capacity"), Q.AllocatedCapacity(), ExpectedCapacityAfterTrim);
+ 		TestTrue(TEXT("Test Trim - Expected"), FResizableCircularQueueTestUtil::VerifyQueueIntegrity(Q, ElementsToPop, 1));
+	}
+
+	// Test Trim empty
+	{
+		FResizableCircularQueueTestUtil::QueueT Q(0);
+
+		Q.Trim();
+
+		TestEqual(TEXT("Test trim empty - Size"), Q.Count(), 0);
+		TestEqual(TEXT("Test trim empty - Capacity"), Q.AllocatedCapacity(), 0);
+	}
+
+
+	check(!HasAnyErrors());
+
+	return true;
+}
+
+#endif // #if WITH_DEV_AUTOMATION_TESTS
\ No newline at end of file
diff --git a/Engine/Source/Runtime/Engine/Private/Net/Tests/Util/SequenceNumberTests.cpp b/Engine/Source/Runtime/Engine/Private/Net/Tests/Util/SequenceNumberTests.cpp
new file mode 100644
--- /dev/null
+++ b/Engine/Source/Runtime/Engine/Private/Net/Tests/Util/SequenceNumberTests.cpp
@@ -0,0 +1,111 @@
+// Copyright 1998-2019 Epic Games, Inc. All Rights Reserved.
+
+#include "Misc/AutomationTest.h"
+#include "Net/Util/SequenceNumber.h"
+
+#if WITH_DEV_AUTOMATION_TESTS
+
+IMPLEMENT_SIMPLE_AUTOMATION_TEST(FNetSequenceNumberTest, "Network.SequenceNumberTest", EAutomationTestFlags::ApplicationContextMask | EAutomationTestFlags::EngineFilter)
+
+bool FNetSequenceNumberTest::RunTest(const FString& Parameters)
+{
+	typedef TSequenceNumber<3, uint16> FSequence3;
+
+	// Test construct
+	{
+		FSequence3 Seq;
+		TestEqual(TEXT("SequenceNumbers - Seq() == 0"), 0u, Seq.Get());		
+	}
+
+	// Test construct with value
+	{
+		FSequence3 Seq(8);
+		TestEqual(TEXT("SequenceNumbers - Seq(8) == 0"), 0u, Seq.Get());		
+	}
+
+	// Test postIncrement
+	{
+		FSequence3 Seq;
+		Seq++;
+		TestEqual("SequenceNumbers - Seq()++ == 1", 1, Seq.Get());
+	}
+
+	// Test preIncrement
+	{
+		FSequence3 Seq;
+		++Seq;
+		TestEqual("SequenceNumbers - ++Seq == 1", 1, Seq.Get());
+	}
+
+	// Test wraparound
+	{
+		FSequence3 Seq(7);
+		Seq++;
+		TestEqual("SequenceNumbers - Seq(7) + 1 == 0", 0, Seq.Get());
+	}
+
+	// Test equal and not equal
+	{
+		FSequence3 A(2);
+		FSequence3 B(2);
+		FSequence3 C(1);
+
+		TestTrue("SequenceNumbers - operator ==", A == B);
+		TestTrue("SequenceNumbers - operator ==", B == A);
+		TestTrue("SequenceNumbers - operator !=", A != C);
+		TestTrue("SequenceNumbers - operator !=", C != A);
+		TestFalse("SequenceNumbers - operator !=", A != B);
+	}
+
+	// Test operator >
+	{
+		FSequence3 Zero(0);
+		FSequence3 HalfMinusOne(FSequence3::SeqNumberHalf-1);
+		FSequence3 Half(FSequence3::SeqNumberHalf);
+		FSequence3 MaxSeq(FSequence3::SeqNumberCount - 1);
+
+		TestTrue("SequenceNumbers - operator> HalfMinusOne > Zero", HalfMinusOne > Zero);
+		TestTrue("SequenceNumbers - operator> Half > HalfMinusOne", Half > HalfMinusOne);
+		TestFalse("SequenceNumbers - operator> Half > Zero", Half > Zero);
+
+		TestFalse("SequenceNumbers - operator> MaxSeq > Zero", MaxSeq > Zero);
+		TestTrue("SequenceNumbers - operator> Zero > MaxSeq", Zero > MaxSeq);
+	}
+	
+	// Test diff sequence numbers
+	{		
+		// Valid sequence = 0-7, max distance between sequence number in order to determine order is half the sequence space (0-3)
+		for (int32 I = 0; I < FSequence3::SeqNumberCount; ++I)
+		{
+			FSequence3 Seq(I);
+			FSequence3 Ref(I);
+			
+			for (int32 U = 0; U < FSequence3::SeqNumberCount; ++U)
+			{	
+				FSequence3::DifferenceT Diff = FSequence3::Diff(Seq, Ref);
+
+				if (U < FSequence3::SeqNumberHalf)
+				{
+					if (!(Diff == U))
+					{
+						TestTrue(TEXT("SequenceNumbers - Expected Diff "), Diff == U);
+					}
+				}
+				else
+				{
+					if (!(Diff == (U - FSequence3::SeqNumberCount)))
+					{
+						TestTrue(TEXT("SequenceNumbers - Expected Diff"), Diff == (U - FSequence3::SeqNumberCount));
+					}
+				}
+				++Seq;
+			}
+		}
+
+		check(!HasAnyErrors());
+	}
+
+	return true;
+}
+
+#endif // #if WITH_DEV_AUTOMATION_TESTS
\ No newline at end of file
diff --git a/Engine/Source/Runtime/Engine/Public/Net/NetPacketNotify.h b/Engine/Source/Runtime/Engine/Public/Net/NetPacketNotify.h
new file mode 100644
--- /dev/null
+++ b/Engine/Source/Runtime/Engine/Public/Net/NetPacketNotify.h
@@ -0,0 +1,216 @@
+// Copyright 1998-2019 Epic Games, Inc. All Rights Reserved.
+
+#pragma once
+
+#include "EngineLogs.h"
+#include "HAL/UnrealMemory.h"
+#include "Logging/LogMacros.h"
+#include "Templates/IsSigned.h"
+#include "Util/ResizableCircularQueue.h"
+#include "Util/SequenceNumber.h"
+#include "Util/SequenceHistory.h"
+
+#if (UE_BUILD_SHIPPING || UE_BUILD_TEST)
+#	define UE_NET_ENABLE_PACKET_NOTIFY_LOG 0
+#else
+#	define UE_NET_ENABLE_PACKET_NOTIFY_LOG 1
+#endif 
+
+#if UE_NET_ENABLE_PACKET_NOTIFY_LOG
+#	define UE_LOG_PACKET_NOTIFY(Format, ...)  UE_LOG(LogNetTraffic, Verbose, Format, ##__VA_ARGS__)
+#else
+#	define UE_LOG_PACKET_NOTIFY(...)
+#endif
+
+#define UE_LOG_PACKET_NOTIFY_WARNING(Format, ...)  UE_LOG(LogNetTraffic, Warning, Format, ##__VA_ARGS__)
+
+struct FBitWriter;
+struct FBitReader;
+
+/** 
+	FNetPacketNotify - Drives delivery of sequence numbers, acknowledgments and notifications of delivery sequence numbers
+*/
+class FNetPacketNotify
+{
+public:
+	enum { SequenceNumberBits = 14 };
+	enum { MaxSequenceHistoryLength = 256 };
+
+	typedef TSequenceNumber<SequenceNumberBits, uint16> SequenceNumberT;
+	typedef TSequenceHistory<MaxSequenceHistoryLength> SequenceHistoryT;
+
+	struct FNotificationHeader
+	{
+		SequenceHistoryT History;
+		SIZE_T HistoryWordCount;
+		SequenceNumberT Seq;
+		SequenceNumberT AckedSeq;
+	};
+
+	/** Constructor */
+	FNetPacketNotify()
+	: AckRecord(64)
+	, WrittenHistoryWordCount(0)
+	{}
+
+	/** Init notification with expected initial sequence numbers */
+	void Init(SequenceNumberT InitialInSeq, SequenceNumberT InitialOutSeq);
+
+	/** Mark Seq as received and update current InSeq, missing sequence numbers will be marked as lost */
+	void AckSeq(SequenceNumberT Seq) { AckSeq(Seq, true); }
+
+	/** Explicitly mark Seq as not received and update current InSeq, additional missing sequence numbers will be marked as lost */
+	void NakSeq(SequenceNumberT Seq) { AckSeq(Seq, false); }
+
+	/** Increment outgoing seq number and commit data*/
+	SequenceNumberT CommitAndIncrementOutSeq();
+
+	/** Write NotificationHeader, and update outgoing ack record 
+	 if bRefresh is true we will attempt to refresh a previously written header 
+	if the size remains the same */
+	void WriteHeader(FBitWriter& Writer, bool bRefresh = false);
+	
+	/** Read header from stream */
+	bool ReadHeader(FNotificationHeader& Data, FBitReader& Reader) const;
+
+	/** Update state of PacketNotification based on received header and invoke packet notifications for received acks.
+		InFunc is a function in the format void)(FNetPacketNotify::SequenceNumberT AckedSequence, bool delivered) to handle packetNotifications.
+
+		Returns the positive delta of the incoming seq number if it is within half the sequence number space
+		Returns 0 if the received sequence number is out of the current window or if the acknowledged seq received by the remote is invalid.
+	*/
+	template<class Functor>
+	SequenceNumberT::DifferenceT Update(const FNotificationHeader& NotificationData, Functor&& InFunc)
+	{		
+		if (NotificationData.Seq > InSeq && NotificationData.AckedSeq >= OutAckSeq)
+		{
+			UE_LOG_PACKET_NOTIFY(TEXT("FNetPacketNotify::Update - Seq %u, InSeq %u"), NotificationData.Seq.Get(), InSeq.Get());
+
+			const SequenceNumberT::DifferenceT InSeqDelta = SequenceNumberT::Diff(NotificationData.Seq, InSeq);
+			
+			ProcessReceivedAcks(NotificationData, InFunc);
+
+			// accept sequence
+			InSeq = NotificationData.Seq;
+
+			return InSeqDelta;
+		}
+		else
+		{
+			return 0;
+		}
+	}
+
+	/** Get the current SequenceHistory */
+	const SequenceHistoryT& GetInSeqHistory() const { return InSeqHistory; }
+
+	/** Get the last received in sequence number */
+	SequenceNumberT GetInSeq() const { return InSeq; }
+
+	/** Get the last received sequence number that we have accepted, InAckSeq cannot be larger than InSeq */
+	SequenceNumberT GetInAckSeq() const { return InAckSeq; }
+
+	/** Get the current outgoing sequence number */
+	SequenceNumberT GetOutSeq() const { return OutSeq; }
+
+	/** Get the last outgoing sequence number acknowledged by remote */
+	SequenceNumberT GetOutAckSeq() const { return OutAckSeq; }
+
+	/** If we do have more unacknowledged sequence numbers in-flight than our maximum sendwindow we should not send more as the receiving end will not be able to detect if the sequence number has wrapped around */
+	bool CanSend() const { SequenceNumberT NextOutSeq = OutSeq; ++NextOutSeq; return NextOutSeq >= OutAckSeq; }
+
+	/** Get the current sequenceHistory length in bits */
+	SIZE_T GetCurrentSequenceHistoryLength() const;
+
+private:
+	struct FSentAckData
+	{
+		SequenceNumberT OutSeq;	// Not needed... just to verify that things work as expected
+		SequenceNumberT InAckSeq;
+	};
+	typedef TResizableCircularQueue<FSentAckData, TInlineAllocator<128>> AckRecordT;
+
+	AckRecordT AckRecord;				// Track acked seq for each sent packet to track size of ack history
+	SIZE_T WrittenHistoryWordCount;		// Bookkeeping to track if we can update data
+	SequenceNumberT WrittenInAckSeq;	// When we call CommitAndIncrementOutSequence this will be committed along with the current outgoing sequence number for bookkeeping
+
+	// Track incoming sequence data
+	SequenceHistoryT InSeqHistory;		// BitBuffer containing a bitfield describing the history of received packets
+	SequenceNumberT InSeq;				// Last sequence number received and accepted from remote
+	SequenceNumberT InAckSeq;			// Last sequence number received from remote that we have acknowledged, this is needed since we support accepting a packet but explicitly not acknowledge it as received.
+	SequenceNumberT InAckSeqAck;		// Last sequence number received from remote that we have acknowledged and also knows that the remote has received the ack, used to calculate how big our history must be
+
+	// Track outgoing sequence data
+	SequenceNumberT OutSeq;				// Outgoing sequence number
+	SequenceNumberT OutAckSeq;			// Last sequence number that we know that the remote side have received.
+
+private:
+	SequenceNumberT UpdateInAckSeqAck(SequenceNumberT::DifferenceT AckCount, SequenceNumberT AckedSeq)
+	{
+		check((SIZE_T)AckCount <= AckRecord.Count());
+
+		if (AckCount > 1)
+		{
+			AckRecord.PopNoCheck(AckCount - 1);
+		}
+		
+		FSentAckData AckData = AckRecord.PeekNoCheck();
+		AckRecord.PopNoCheck();
+
+		// verify that we have a matching sequence number
+		if (AckData.OutSeq == AckedSeq)
+		{
+			return AckData.InAckSeq;
+		}
+		else
+		{
+			UE_LOG_PACKET_NOTIFY_WARNING(TEXT("FNetPacketNotify::UpdateInAckSeqAck - Failed to find matching AckRecord for %u, (Found %u)"), AckedSeq.Get(), AckData.OutSeq.Get());
+
+			// Pessimistic view, should never occur
+			return SequenceNumberT(AckedSeq.Get() - MaxSequenceHistoryLength);
+		}
+	}
+
+	template<class Functor>
+	void ProcessReceivedAcks(const FNotificationHeader& NotificationData, Functor&& InFunc)
+	{
+		if (NotificationData.AckedSeq > OutAckSeq)
+		{
+			UE_LOG_PACKET_NOTIFY(TEXT("Notification::ProcessReceivedAcks - AckedSeq: %u, OutAckSeq: %u"), NotificationData.AckedSeq.Get(), OutAckSeq.Get());
+
+			SequenceNumberT::DifferenceT AckCount = SequenceNumberT::Diff(NotificationData.AckedSeq, OutAckSeq);
+
+			// Update InAckSeqAck used to track the needed number of bits to transmit our ack history
+			InAckSeqAck = UpdateInAckSeqAck(AckCount, NotificationData.AckedSeq);
+
+			// ExpectedAck = OutAckSeq + 1
+			SequenceNumberT CurrentAck(OutAckSeq);
+			++CurrentAck;
+
+			// Everything not found in the history buffer is treated as lost
+			while (AckCount > (SequenceNumberT::DifferenceT)(SequenceHistoryT::Size))
+			{
+				--AckCount;
+				UE_LOG_PACKET_NOTIFY_WARNING(TEXT("Notification::ProcessReceivedAcks Seq: %u - IsAck: 0 HistoryIndex: N/A"), CurrentAck.Get());
+				InFunc(CurrentAck, false);
+				++CurrentAck;
+			}
+
+			// For sequence numbers contained in the history we lookup the delivery status from the history
+			while (AckCount > 0)
+			{
+				--AckCount;
+				UE_LOG_PACKET_NOTIFY(TEXT("Notification::ProcessReceivedAcks Seq: %u - IsAck: %u HistoryIndex: %u"), CurrentAck.Get(), NotificationData.History.IsDelivered(AckCount) ? 1u : 0u, AckCount);
+				InFunc(CurrentAck, NotificationData.History.IsDelivered(AckCount));
+				++CurrentAck;
+			}
+			OutAckSeq = NotificationData.AckedSeq;
+		}
+	}
+
+	void AckSeq(SequenceNumberT AckedSeq, bool IsAck);
+
+#if WITH_DEV_AUTOMATION_TESTS
+	friend struct FNetPacketNotifyTestUtil;
+#endif
+};
diff --git a/Engine/Source/Runtime/Engine/Public/Net/Util/ResizableCircularQueue.h b/Engine/Source/Runtime/Engine/Public/Net/Util/ResizableCircularQueue.h
new file mode 100644
--- /dev/null
+++ b/Engine/Source/Runtime/Engine/Public/Net/Util/ResizableCircularQueue.h
@@ -0,0 +1,207 @@
+// Copyright 1998-2019 Epic Games, Inc. All Rights Reserved.
+
+#pragma once
+
+#include "CoreTypes.h"
+#include "Misc/AssertionMacros.h"
+#include "Containers/Array.h"
+
+/**
+ * Simple ResisableCircularQueue for trivially copyable types
+ * Relies on unsigned arithmetics and ever increasing head and tail indices to avoid having to store 
+ * an extra element or maintain explicit empty state.
+ * Capacity most be power of two and it supports growing
+ */
+class FDefaultAllocator;
+
+template<typename T, typename AllocatorT = FDefaultAllocator>
+class TResizableCircularQueue
+{
+public:
+	typedef T ElementT;
+
+	/** Construct Empty Queue with the given initial Capacity, Capacity must be a power of two since we rely on unsigned arithmetics for wraparound */
+	explicit TResizableCircularQueue(SIZE_T InitialCapacity);
+
+	/** Returns true if the queue is empty */
+	bool IsEmpty() const { return Head == Tail; }
+
+	/** Gets the number of elements in the queue. */
+	SIZE_T Count() const { return Head - Tail; }
+
+	/** Current allocated Capacity */
+	SIZE_T AllocatedCapacity() const { return Storage.Num(); }
+
+	/** Push single element to the back of the Queue */
+	void Enqueue(const ElementT& SrcData);
+
+	/** Pop elements from the front of the queue */
+	void Pop();
+
+	/** Pop Count elements from the front of the queue */
+	void Pop(SIZE_T Count);
+
+	/**  Unchecked version, Pop a single element from the front of the queue */
+	void PopNoCheck() { ++Tail; }
+
+	/** Unchecked version, Pop Count elements from the front of the queue */
+	void PopNoCheck(SIZE_T Count) { Tail += Count; }
+
+	/** Peek with the given offset from the front of the queue */
+	const ElementT& PeekAtOffset(SIZE_T Offset = 0) const { check(Offset < Count()); return PeekAtOffsetNoCheck(Offset); }
+
+	/** Peek at front element */
+	const ElementT& Peek() const { return PeekAtOffset(0); }
+
+	/**  Unchecked version, Peek with the given offset from the front of the queue */
+	const ElementT& PeekAtOffsetNoCheck(SIZE_T Offset = 0) const { return Storage.GetData()[(Tail + Offset) & IndexMask]; }
+
+	/** Peek at front element with no check */
+	const ElementT& PeekNoCheck() const { return PeekAtOffsetNoCheck(0); }
+
+	/**  Trim memory usage to the next power of two for the current size */
+	void Trim();
+
+	/** Empty queue without releasing memory */
+	void Reset() { Head = Tail = 0u; }
+
+	/** Empty queue and release memory */
+	void Empty() { Reset(); Storage.Empty(); IndexMask = IndexT(-1); }
+	
+private:
+
+#if WITH_DEV_AUTOMATION_TESTS
+	friend struct FResizableCircularQueueTestUtil;
+#endif
+
+	// Resize buffer maintaining validity of stored data.
+	void SetCapacity(SIZE_T NewCapacity);
+
+	// TIsTriviallyCopyable (a.k.a. std::is_trivially_copyable) should be used but
+	// TIsTriviallyCopyable is not provided by the core module at the moment. Core's
+	// implementation of TIsTrivial is actually equivalent to std::is_trivially_copyable
+	// since std::is_trivial<T> requires T to have a trivial constructor but TIsTrivial
+	// doesn't
+	static_assert(TIsTrivial<ElementT>::Value, "ResizableCircularQueue only supports trivially copyable elements");
+
+	typedef uint32 IndexT;
+	typedef TArray<ElementT, AllocatorT> StorageT;
+
+	IndexT Head;
+	IndexT Tail;
+
+	IndexT IndexMask;
+	StorageT Storage;
+};
+
+
+template<typename T, typename AllocatorT>
+TResizableCircularQueue<T, AllocatorT>::TResizableCircularQueue(SIZE_T InitialCapacity)
+: Head(0u)
+, Tail(0u)
+, IndexMask(0u)
+{
+	// Capacity should be power of two, it will be rounded up but lets warn the user anyway in debug builds.
+	checkSlow((InitialCapacity & (InitialCapacity - 1)) == 0);
+	if (InitialCapacity > 0)
+	{
+		SetCapacity(InitialCapacity);
+	}
+}
+
+template<typename T, typename AllocatorT>
+void TResizableCircularQueue<T, AllocatorT>::Enqueue(const ElementT& SrcData)
+{ 
+	const SIZE_T RequiredCapacity = Count() + 1;
+	if (RequiredCapacity > AllocatedCapacity())
+	{
+		// Capacity must be power of two
+		SetCapacity(RequiredCapacity);
+	}
+
+	T* DstData = Storage.GetData();
+	const IndexT MaskedIndex = Head & IndexMask;
+	DstData[MaskedIndex] = SrcData;
+	++Head;		
+}
+
+template<typename T, typename AllocatorT>
+void
+TResizableCircularQueue<T, AllocatorT>::Pop(SIZE_T PopCount)
+{
+	if (ensure(Count() > PopCount))
+	{
+		PopNoCheck(PopCount);
+	}
+}
+
+template<typename T, typename AllocatorT>
+void TResizableCircularQueue<T, AllocatorT>::Pop()
+{
+	if (ensure(Count() > 0))
+	{
+		PopNoCheck();
+	}
+}
+
+template<typename T, typename AllocatorT>
+void TResizableCircularQueue<T, AllocatorT>::SetCapacity(SIZE_T RequiredCapacity)
+{
+	SIZE_T NewCapacity = FMath::RoundUpToPowerOfTwo(RequiredCapacity);
+
+	if ((NewCapacity == Storage.Num()) || (NewCapacity < Count()))
+	{
+		return;
+	}
+
+	if (Storage.Num() > 0)
+	{
+		StorageT NewStorage;
+		NewStorage.AddUninitialized(NewCapacity);
+		
+		// copy data to new storage
+		const IndexT MaskedTail = Tail & IndexMask;
+		const IndexT MaskedHead = Head & IndexMask;
+
+		const ElementT* SrcData = Storage.GetData();
+		const SIZE_T SrcCapacity = Storage.Num();
+		const SIZE_T SrcSize = Count();
+		ElementT* DstData = NewStorage.GetData();
+
+		if (MaskedTail >= MaskedHead)
+		{
+			const SIZE_T CopyCount = (SrcCapacity - MaskedTail);
+			FPlatformMemory::Memcpy(DstData, SrcData + MaskedTail, CopyCount * sizeof(ElementT));
+			FPlatformMemory::Memcpy(DstData + CopyCount, SrcData, MaskedHead * sizeof(ElementT));
+		}
+		else
+		{
+			FPlatformMemory::Memcpy(DstData, SrcData + MaskedTail, SrcSize * sizeof(ElementT));
+		}
+
+		this->Storage = MoveTemp(NewStorage);
+		IndexMask = NewCapacity - 1;
+		Tail = 0u;
+		Head = SrcSize;
+	}
+	else
+	{
+		IndexMask = NewCapacity - 1;
+		Storage.AddUninitialized(NewCapacity);
+	}
+}
+
+template<typename T, typename AllocatorT>
+void TResizableCircularQueue<T, AllocatorT>::Trim()
+{
+	if (IsEmpty())
+	{
+		Empty();
+	}
+	else
+	{
+		SetCapacity(Count());
+	}
+}
+
+
diff --git a/Engine/Source/Runtime/Engine/Public/Net/Util/SequenceHistory.h b/Engine/Source/Runtime/Engine/Public/Net/Util/SequenceHistory.h
new file mode 100644
--- /dev/null
+++ b/Engine/Source/Runtime/Engine/Public/Net/Util/SequenceHistory.h
@@ -0,0 +1,126 @@
+// Copyright 1998-2019 Epic Games, Inc. All Rights Reserved.
+
+#pragma once
+
+#include "CoreTypes.h"
+#include "HAL/UnrealMemory.h"
+#include "Misc/AssertionMacros.h"
+#include "Templates/IsSigned.h"
+#include "Serialization/BitWriter.h"
+#include "Serialization/BitReader.h"
+
+/** Util class to manage history of received sequence numbers */
+template <SIZE_T HistorySize>
+class TSequenceHistory
+{
+public:	
+	typedef uint32 WordT;
+	
+	static const SIZE_T BitsPerWord = sizeof(WordT) * 8;
+	static const SIZE_T WordCount = HistorySize / BitsPerWord;
+	static const SIZE_T Size = HistorySize;
+
+	static_assert(HistorySize > 0, "HistorySize must be > 0");
+	static_assert(HistorySize % BitsPerWord == 0, "InMaxHistorySize must be a modulo of the wordsize");
+
+public:	
+	TSequenceHistory();
+
+#if WITH_DEV_AUTOMATION_TESTS
+	explicit TSequenceHistory(WordT Value, SIZE_T Count = 1);
+	WordT* Data() { return &Storage[0]; }
+#endif
+	/** Reset */
+	void Reset();
+
+	/** Store delivery status, oldest will be dropped */
+	void AddDeliveryStatus(bool Delivered);
+
+	/** Query the status of a specific index, index 0 is last stored status */
+	bool IsDelivered(SIZE_T Index) const;
+	
+	bool operator==(const TSequenceHistory& Other) const { return FMemory::Memcmp(Storage, Other.Storage, WordCount * sizeof (WordT)) == 0; }
+
+	bool operator!=(const TSequenceHistory& Other) const { return FMemory::Memcmp(Storage, Other.Storage, WordCount * sizeof (WordT)) != 0; }
+
+	/** Write history to BitStream */
+	void Write(FBitWriter& Writer, SIZE_T NumWords) const;
+
+	/** Read history from BitStream */
+	void Read(FBitReader& Reader, SIZE_T NumWords);
+
+private:
+	WordT Storage[WordCount];
+};
+
+template <SIZE_T HistorySize>
+TSequenceHistory<HistorySize>::TSequenceHistory()
+{
+	Reset();
+}
+
+#if WITH_DEV_AUTOMATION_TESTS
+template <SIZE_T HistorySize>
+TSequenceHistory<HistorySize>::TSequenceHistory(WordT Value, SIZE_T Count)
+{
+	Reset();
+	for (SIZE_T CurrentWordIt = 0; CurrentWordIt < Count; ++CurrentWordIt)
+	{
+		Storage[CurrentWordIt] = Value;
+	}	
+}
+#endif 
+
+template <SIZE_T HistorySize>
+void TSequenceHistory<HistorySize>::Reset()
+{
+	FPlatformMemory::Memset(&Storage[0], 0, WordCount * sizeof(WordT));
+}
+
+template <SIZE_T HistorySize>
+void TSequenceHistory<HistorySize>::AddDeliveryStatus(bool Delivered)
+{
+	WordT Carry = Delivered ? 1u : 0u;
+	const WordT ValueMask = 1u << (BitsPerWord - 1);
+	
+	for (SIZE_T CurrentWordIt = 0; CurrentWordIt < WordCount; ++CurrentWordIt)
+	{
+		const WordT OldValue = Carry;
+		
+		// carry over highest bit in each word to the next word
+		Carry = (Storage[CurrentWordIt] & ValueMask) >> (BitsPerWord - 1);
+		Storage[CurrentWordIt] = (Storage[CurrentWordIt] << 1u) | OldValue;
+	}
+}
+
+template <SIZE_T HistorySize>
+bool TSequenceHistory<HistorySize>::IsDelivered(SIZE_T Index) const
+{
+	check(Index < Size);
+
+	const SIZE_T WordIndex = Index / BitsPerWord;
+	const WordT WordMask = (WordT(1) << (Index & (BitsPerWord - 1)));
+	
+	return (Storage[WordIndex] & WordMask) != 0u;
+}
+
+template <SIZE_T HistorySize>
+void TSequenceHistory<HistorySize>::Write(FBitWriter& Writer, SIZE_T NumWords) const
+{
+	NumWords = FPlatformMath::Min(NumWords, WordCount);
+	for (SIZE_T CurrentWordIt = 0; CurrentWordIt < NumWords; ++CurrentWordIt)
+	{
+		WordT temp = Storage[CurrentWordIt];
+		Writer << temp;
+	}
+}
+
+template <SIZE_T HistorySize>
+void TSequenceHistory<HistorySize>::Read(FBitReader& Reader, SIZE_T NumWords)
+{
+	NumWords = FPlatformMath::Min(NumWords, WordCount);
+	for (SIZE_T CurrentWordIt = 0; CurrentWordIt < NumWords; ++CurrentWordIt)
+	{
+		Reader << Storage[CurrentWordIt];
+	}
+}
diff --git a/Engine/Source/Runtime/Engine/Public/Net/Util/SequenceNumber.h b/Engine/Source/Runtime/Engine/Public/Net/Util/SequenceNumber.h
new file mode 100644
--- /dev/null
+++ b/Engine/Source/Runtime/Engine/Public/Net/Util/SequenceNumber.h
@@ -0,0 +1,70 @@
+// Copyright 1998-2019 Epic Games, Inc. All Rights Reserved.
+
+#pragma once
+
+#include "CoreTypes.h"
+#include "HAL/UnrealMemory.h"
+#include "Misc/AssertionMacros.h"
+#include "Templates/IsSigned.h"
+
+/** Helper class to work with sequence numbers */
+template <SIZE_T NumBits, typename SequenceType>
+class TSequenceNumber
+{
+	static_assert(TIsSigned<SequenceType>::Value == false, "The base type for sequence numbers must be unsigned");
+
+public:
+	typedef SequenceType SequenceT;
+	typedef int32 DifferenceT;
+
+	// Constants
+	enum { SeqNumberBits = NumBits };
+	enum { SeqNumberCount = SequenceT(1) << NumBits };
+	enum { SeqNumberHalf = SequenceT(1) << (NumBits - 1) };
+	enum { SeqNumberMax = SeqNumberCount - 1u };
+	enum { SeqNumberMask = SeqNumberMax };
+
+	/** Default constructor */
+	TSequenceNumber() : Value(0u) {}
+
+	/** Constructor with given value */
+	TSequenceNumber(SequenceT ValueIn) : Value(ValueIn & SeqNumberMask) {}
+	
+	/** Get Current Value */	
+	SequenceT Get() const { return Value; }
+
+	/** Diff between sequence numbers (A - B) only valid if (A - B) < SeqNumberHalf */
+	static DifferenceT Diff(TSequenceNumber A, TSequenceNumber B);
+	
+	/** return true if this is > Other, this is only considered to be the case if (A - B) < SeqNumberHalf since we have to be able to detect wraparounds */
+	bool operator>(const TSequenceNumber& Other) const { return (Value != Other.Value) && (((Value - Other.Value) & SeqNumberMask) < SeqNumberHalf); }
+
+	/** Check if this is >= Other, See above */
+	bool operator>=(const TSequenceNumber& Other) const { return ((Value - Other.Value) & SeqNumberMask) < SeqNumberHalf; }
+
+	/** Equals, NOTE that sequence numbers wrap around so 0 == 0 + SequenceNumberCount */
+	bool operator==(const TSequenceNumber& Other) const { return Value == Other.Value; }
+
+	bool operator!=(const TSequenceNumber& Other) const { return Value != Other.Value; }
+
+	/** Pre-increment and wrap around */
+	TSequenceNumber& operator++() { Increment(1u); return *this; }
+	
+	/** Post-increment and wrap around */
+	TSequenceNumber operator++(int) { TSequenceNumber Tmp(*this); Increment(1u); return Tmp; }
+
+private:
+	void Increment(SequenceT InValue) { *this = TSequenceNumber(Value + InValue); }
+	SequenceT Value;
+};
+
+template <SIZE_T NumBits, typename SequenceType>
+typename TSequenceNumber<NumBits, SequenceType>::DifferenceT TSequenceNumber<NumBits, SequenceType>::Diff(TSequenceNumber A, TSequenceNumber B) 
+{ 
+	const SIZE_T ShiftValue = sizeof(DifferenceT)*8 - NumBits;
+
+	const SequenceT ValueA = A.Value;
+	const SequenceT ValueB = B.Value;
+
+	return (DifferenceT)((ValueA - ValueB) << ShiftValue) >> ShiftValue;
+};
