Improved acks Part 2/2

 - These are the code changes to enable the new acks submitted in CL

Original changes from dev-networking
CL 4687651, CL 4687649, CL 4684239

diff --git a/Engine/Source/Runtime/Core/Private/Misc/NetworkVersion.cpp b/Engine/Source/Runtime/Core/Private/Misc/NetworkVersion.cpp
--- a/Engine/Source/Runtime/Core/Private/Misc/NetworkVersion.cpp
+++ b/Engine/Source/Runtime/Core/Private/Misc/NetworkVersion.cpp
@@ -18,7 +18,7 @@ FString FNetworkVersion::ProjectVersion = TEXT("1.0.0");
 bool FNetworkVersion::bHasCachedNetworkChecksum			= false;
 uint32 FNetworkVersion::CachedNetworkChecksum			= 0;
 
-uint32 FNetworkVersion::EngineNetworkProtocolVersion	= HISTORY_CHANNEL_CLOSE_REASON;
+uint32 FNetworkVersion::EngineNetworkProtocolVersion	= HISTORY_ACKS_INCLUDED_IN_HEADER;
 uint32 FNetworkVersion::GameNetworkProtocolVersion		= 0;
 
 uint32 FNetworkVersion::EngineCompatibleNetworkProtocolVersion		= HISTORY_REPLAY_BACKWARDS_COMPAT;
diff --git a/Engine/Source/Runtime/Core/Public/Misc/NetworkVersion.h b/Engine/Source/Runtime/Core/Public/Misc/NetworkVersion.h
--- a/Engine/Source/Runtime/Core/Public/Misc/NetworkVersion.h
+++ b/Engine/Source/Runtime/Core/Public/Misc/NetworkVersion.h
@@ -39,6 +39,7 @@ enum EEngineNetworkVersionHistory
 	HISTORY_NEW_ACTOR_OVERRIDE_LEVEL = 5,			// Bump version since a level reference was added to the new actor information
 	HISTORY_CHANNEL_NAMES = 6,						// Bump version since channel type is now an fname
 	HISTORY_CHANNEL_CLOSE_REASON = 7,				// Bump version to serialize a channel close reason in bunches instead of bDormant
+	HISTORY_ACKS_INCLUDED_IN_HEADER = 8,			// Bump version since acks are now sent as part of the header
 };
 
 struct CORE_API FNetworkVersion
diff --git a/Engine/Source/Runtime/Engine/Classes/Engine/ChildConnection.h b/Engine/Source/Runtime/Engine/Classes/Engine/ChildConnection.h
--- a/Engine/Source/Runtime/Engine/Classes/Engine/ChildConnection.h
+++ b/Engine/Source/Runtime/Engine/Classes/Engine/ChildConnection.h
@@ -52,10 +52,6 @@ public:
 		Parent->AssertValid();
 	}
 
-	virtual void SendAck(int32 PacketId, bool FirstTime=1) override
-	{
-	}
-
 	virtual void FlushNet(bool bIgnoreSimulation = false) override
 	{
 		Parent->FlushNet(bIgnoreSimulation);
diff --git a/Engine/Source/Runtime/Engine/Classes/Engine/NetConnection.h b/Engine/Source/Runtime/Engine/Classes/Engine/NetConnection.h
--- a/Engine/Source/Runtime/Engine/Classes/Engine/NetConnection.h
+++ b/Engine/Source/Runtime/Engine/Classes/Engine/NetConnection.h
@@ -14,6 +14,7 @@
 #include "GameFramework/OnlineReplStructs.h"
 #include "Engine/NetDriver.h"
 #include "Net/DataBunch.h"
+#include "Net/NetPacketNotify.h"
 #include "Engine/Player.h"
 #include "Engine/Channel.h"
 #include "ProfilingDebugging/Histogram.h"
@@ -38,10 +39,11 @@ typedef TMap<TWeakObjectPtr<AActor>, UActorChannel*, FDefaultSetAllocator, TWeak
 	Types.
 -----------------------------------------------------------------------------*/
 enum { RELIABLE_BUFFER = 256 }; // Power of 2 >= 1.
-enum { MAX_PACKETID = 16384 };  // Power of 2 >= 1, covering guaranteed loss/misorder time.
+enum { MAX_PACKETID = FNetPacketNotify::SequenceNumberT::SeqNumberCount };  // Power of 2 >= 1, covering guaranteed loss/misorder time.
 enum { MAX_CHSEQUENCE = 1024 }; // Power of 2 >RELIABLE_BUFFER, covering loss/misorder time.
 enum { MAX_BUNCH_HEADER_BITS = 256 };
-enum { MAX_PACKET_HEADER_BITS = 15 }; // = FMath::CeilLogTwo(MAX_PACKETID) + 1 (IsAck)
+enum { UE_PACKET_MAX_SEQUENCE_HISTORY_BITS = FNetPacketNotify::MaxSequenceHistoryLength };
+enum { MAX_PACKET_HEADER_BITS = 49 + UE_PACKET_MAX_SEQUENCE_HISTORY_BITS }; // = [(Header)32 + (SeqHistory)32-MaxSequenceHistoryLength|(bool)1|?(ServerTime)8|(ReceviedRateByte)8]
 enum { MAX_PACKET_TRAILER_BITS = 1 };
 
 // 
@@ -389,7 +391,7 @@ public:
 	FBitWriter		SendBuffer;						// Queued up bits waiting to send
 	double			OutLagTime[256];				// For lag measuring.
 	int32			OutLagPacketId[256];			// For lag measuring.
-	int32			OutBytesPerSecondHistory[256];	// For saturation measuring.
+	uint8			OutBytesPerSecondHistory[256];	// For saturation measuring.
 	float			RemoteSaturation;
 	int32			InPacketId;						// Full incoming packet index.
 	int32			OutPacketId;					// Most recently sent packet.
@@ -405,8 +407,6 @@ public:
 	TArray<int32>		OutReliable;
 	TArray<int32>		InReliable;
 	TArray<int32>		PendingOutRec;	// Outgoing reliable unacked data from previous (now destroyed) channel in this slot.  This contains the first chsequence not acked
-	TArray<int32> QueuedAcks, ResendAcks;
-
 	int32				InitOutReliable;
 	int32				InitInReliable;
 
@@ -693,6 +693,7 @@ public:
 	ENGINE_API virtual void AssertValid();
 
 	/** Send an acknowledgment. */
+	UE_DEPRECATED(4.22, "This method will be removed")
 	ENGINE_API virtual void SendAck( int32 PacketId, bool FirstTime=1);
 
 	/**
@@ -857,6 +858,7 @@ public:
 	// Functions.
 
 	/** Resend any pending acks. */
+	UE_DEPRECATED(4.22, "This method will be removed.")
 	void PurgeAcks();
 
 	/** Send package map to the remote. */
@@ -1054,6 +1056,18 @@ private:
 	/** Returns true if an outgoing packet should be dropped due to packet simulation settings, including loss burst simulation. */
 	bool ShouldDropOutgoingPacketForLossSimulation(int64 NumBits) const;
 
+	/** Write packetHeader */
+	void WritePacketHeader(FBitWriter& Writer);
+
+	/** Write extended packet header information (ServerFrameTime, SaturationData) */
+	void WritePacketInfo(FBitWriter& Writer) const;
+	
+	/** Read extended packet header information (ServerFrameTime, SaturationData) */
+	bool ReadPacketInfo(FBitReader& Reader);
+
+	/** Packet was acknowledged as delivered */
+	void ReceivedAck(int32 AckPacketId);
+
 	/**
 	 * on the server, the world the client has told us it has loaded
 	 * used to make sure the client has traveled correctly, prevent replicating actors before level transitions are done, etc
@@ -1072,6 +1086,15 @@ private:
 
 	/** Bookkeeping for connections using internal ack to avoid acking all open channels */
 	TArray<int32> ChannelsWaitingForInternalAck;
+
+	/** Sequence data used to implement reliability */
+	FNetPacketNotify PacketNotify;
+
+	/** Full PacketId  of last sent packet that we have received notification for (i.e. we know if it was delivered or not). Related to OutAckPacketId which is tha last successfully delivered PacketId */
+	int32 LastNotifiedPacketId;
+
+	/** Keep old behavior where we send a packet with only acks even if we have no other outgoing data if we got incoming data */
+	uint32 HasDirtyAcks;
 };
 
 
diff --git a/Engine/Source/Runtime/Engine/Private/NetConnection.cpp b/Engine/Source/Runtime/Engine/Private/NetConnection.cpp
--- a/Engine/Source/Runtime/Engine/Private/NetConnection.cpp
+++ b/Engine/Source/Runtime/Engine/Private/NetConnection.cpp
@@ -19,6 +19,7 @@
 #include "Net/UnrealNetwork.h"
 #include "Net/NetworkProfiler.h"
 #include "Net/DataReplication.h"
+#include "Net/NetPacketNotify.h"
 #include "Engine/ActorChannel.h"
 #include "Engine/ChildConnection.h"
 #include "Engine/VoiceChannel.h"
@@ -26,7 +27,6 @@
 #include "Engine/PackageMapClient.h"
 #include "Engine/NetworkObjectList.h"
 #include "EncryptionComponent.h"
-
 #include "Net/PerfCountersHelpers.h"
 #include "GameDelegates.h"
 #include "Misc/PackageName.h"
@@ -133,6 +133,8 @@ UNetConnection::UNetConnection(const FObjectInitializer& ObjectInitializer)
 #endif
 ,	PlayerOnlinePlatformName( NAME_None )
 ,	ClientWorldPackageName( NAME_None )
+,	LastNotifiedPacketId( -1 )
+,	HasDirtyAcks(0u)
 {
 	MaxChannelSize = CVarMaxChannelSize.GetValueOnAnyThread();
 	if (MaxChannelSize <= 0)
@@ -145,6 +147,8 @@ UNetConnection::UNetConnection(const FObjectInitializer& ObjectInitializer)
 	OutReliable.AddDefaulted(MaxChannelSize);
 	InReliable.AddDefaulted(MaxChannelSize);
 	PendingOutRec.AddDefaulted(MaxChannelSize);
+
+	PacketNotify.Init(InPacketId, OutPacketId);
 }
 
 /**
@@ -338,6 +342,7 @@ void UNetConnection::InitSequence(int32 IncomingSequence, int32 OutgoingSequence
 		InPacketId = IncomingSequence - 1;
 		OutPacketId = OutgoingSequence;
 		OutAckPacketId = OutgoingSequence - 1;
+		LastNotifiedPacketId = OutAckPacketId;
 
 		// Initialize the reliable packet sequence (more useful/effective at preventing attacks)
 		InitInReliable = IncomingSequence & (MAX_CHSEQUENCE - 1);
@@ -346,6 +351,8 @@ void UNetConnection::InitSequence(int32 IncomingSequence, int32 OutgoingSequence
 		InReliable.Init(InitInReliable, InReliable.Num());
 		OutReliable.Init(InitOutReliable, OutReliable.Num());
 
+		PacketNotify.Init(InPacketId, OutPacketId);
+
 		UE_LOG(LogNet, Verbose, TEXT("InitSequence: IncomingSequence: %i, OutgoingSequence: %i, InitInReliable: %i, InitOutReliable: %i"), IncomingSequence, OutgoingSequence, InitInReliable, InitOutReliable);
 	}
 }
@@ -456,7 +463,6 @@ bool UNetConnection::IsEncryptionEnabled() const
 void UNetConnection::Serialize( FArchive& Ar )
 {
 	UObject::Serialize( Ar );
-	
 	Ar << PackageMap;
 	for (UChannel* Channel : Channels)
 	{
@@ -1095,7 +1101,7 @@ void UNetConnection::FlushNet(bool bIgnoreSimulation)
 	TimeSensitive = 0;
 
 	// If there is any pending data to send, send it.
-	if (SendBuffer.GetNumBits() || ( Driver->Time-LastSendTime > Driver->KeepAliveTime && !InternalAck && State != USOCK_Closed))
+	if (SendBuffer.GetNumBits() || HasDirtyAcks || ( Driver->Time-LastSendTime > Driver->KeepAliveTime && !InternalAck && State != USOCK_Closed))
 	{
 		// Due to the PacketHandler handshake code, servers must never send the client data,
 		// before first receiving a client control packet (which is taken as an indication of a complete handshake).
@@ -1110,10 +1116,10 @@ void UNetConnection::FlushNet(bool bIgnoreSimulation)
 
 		FOutPacketTraits Traits;
 
-		// If sending keepalive packet, still write the packet id
+		// If sending keepalive packet or just acks, still write the packet header
 		if (SendBuffer.GetNumBits() == 0)
 		{
-			WriteBitsToSendBuffer( NULL, 0 );		// This will force the packet id to be written
+			WriteBitsToSendBuffer( NULL, 0 );		// This will force the packet header to be written
 
 			Traits.bIsKeepAlive = true;
 			AnalyticsVars.OutKeepAliveCount++;
@@ -1130,6 +1136,14 @@ void UNetConnection::FlushNet(bool bIgnoreSimulation)
 		// Write the UNetConnection-level termination bit
 		SendBuffer.WriteBit(1);
 
+		// Refresh outgoing header with latest data
+		if ( !InternalAck )
+		{
+			// if we update ack, we also update received ack associated with outgoing seq
+			// so we know how many ack bits we need to write (which is updated in received packet)
+			WritePacketHeader(SendBuffer);
+		}
+
 		ValidateSendBuffer();
 
 		const int32 NumStrayBits = SendBuffer.GetNumBits();
@@ -1209,9 +1223,17 @@ void UNetConnection::FlushNet(bool bIgnoreSimulation)
 		// Remember the actual time this packet was sent out, so we can compute ping when the ack comes back
 		OutLagPacketId[Index]			= OutPacketId;
 		OutLagTime[Index]				= FPlatformTime::Seconds();
-		OutBytesPerSecondHistory[Index]	= OutBytesPerSecond / 1024;
+	
+		OutBytesPerSecondHistory[Index]	= FMath::Min(OutBytesPerSecond / 1024, 255);
+
+		// Increase outgoing sequence number
+		if (!InternalAck)
+		{
+			PacketNotify.CommitAndIncrementOutSeq();
+			HasDirtyAcks = 0u;
+		}
+		++OutPacketId; 
 
-		OutPacketId++;
 		++OutPackets;
 		++OutTotalPackets;
 		Driver->OutPackets++;
@@ -1237,13 +1259,6 @@ void UNetConnection::FlushNet(bool bIgnoreSimulation)
 
 		InitSendBuffer();
 	}
-
-	// Move acks around.
-	for (int32 i=0; i<QueuedAcks.Num(); i++)
-	{
-		ResendAcks.Add(QueuedAcks[i]);
-	}
-	QueuedAcks.Empty(32);
 }
 
 bool UNetConnection::ShouldDropOutgoingPacketForLossSimulation(int64 NumBits) const
@@ -1271,8 +1286,64 @@ int32 UNetConnection::IsNetReady( bool Saturate )
 void UNetConnection::ReadInput( float DeltaSeconds )
 {}
 
+void UNetConnection::ReceivedAck(int32 AckPacketId)
+{
+	UE_LOG(LogNetTraffic, Verbose, TEXT("   Received ack %i"), AckPacketId);
+
+	// Advance OutAckPacketId
+	OutAckPacketId = AckPacketId;
+
+	// Process the bunch.
+	LastRecvAckTime = Driver->Time;
+
+	if ( PackageMap != NULL )
+	{
+		PackageMap->ReceivedAck( AckPacketId );
+	}
+
+	// Forward the ack to the channel.
+	{
+		SCOPE_CYCLE_COUNTER(Stat_NetConnectionReceivedAcks);
+	
+		for (int32 i = OpenChannels.Num() - 1; i >= 0; i--)
+		{
+			UChannel* const Channel = OpenChannels[i];
+				
+			if (Channel)
+			{
+				if (Channel->OpenPacketId.Last == AckPacketId) // Necessary for unreliable "bNetTemporary" channels.
+				{
+					Channel->OpenAcked = 1;
+				}
+				
+				for (FOutBunch* OutBunch = Channel->OutRec; OutBunch; OutBunch = OutBunch->Next)
+				{
+					if (OutBunch->bOpen)
+					{
+						UE_LOG(LogNet, VeryVerbose, TEXT("Channel %i reset Ackd because open is reliable. "), Channel->ChIndex );
+						Channel->OpenAcked  = 0; // We have a reliable open bunch, don't let the above code set the OpenAcked state,
+												// it must be set in UChannel::ReceivedAcks to verify all open bunches were received.
+					}
+
+					if (OutBunch->PacketId == AckPacketId)
+					{
+						OutBunch->ReceivedAck = 1;
+					}
+				}
+				Channel->ReceivedAcks(); //warning: May destroy Channel.
+			}
+			else
+			{
+				UE_LOG(LogNet, Warning, TEXT("UNetConnection::ReceivedPacket: null channel in OpenChannels array while processing ack. %s"), *Describe());
+			}
+		}
+	}
+}
+
 void UNetConnection::ReceivedNak( int32 NakPacketId )
 {
+	UE_LOG(LogNetTraffic, Verbose, TEXT("   Received nak %i"), NakPacketId);
+
 	SCOPE_CYCLE_COUNTER(Stat_NetConnectionReceivedNak);
 
 	// Update pending NetGUIDs
@@ -1285,7 +1356,7 @@ void UNetConnection::ReceivedNak( int32 NakPacketId )
 		if (Channel)
 		{
 			Channel->ReceivedNak(NakPacketId);
-			if(Channel->OpenPacketId.InRange(NakPacketId))
+			if (Channel->OpenPacketId.InRange(NakPacketId))
 			{
 				Channel->ReceivedAcks(); //warning: May destroy Channel.
 			}
@@ -1295,6 +1366,124 @@ void UNetConnection::ReceivedNak( int32 NakPacketId )
 			UE_LOG(LogNet, Warning, TEXT("UNetConnection::ReceivedNak: null channel in OpenChannels array. %s"), *Describe());
 		}
 	}
+
+	// Stats
+	++OutPacketsLost;
+	++OutTotalPacketsLost;
+	++Driver->OutTotalPacketsLost;
+}
+
+// IMPORTANT:
+// WritePacketHeader must ALWAYS write the exact same number of bits as we go back and rewrite the header
+// right before we put the packet on the wire.
+void UNetConnection::WritePacketHeader(FBitWriter& Writer)
+{
+	// If this is a header refresh, we only serialize the updated serial number information
+	const bool bIsHeaderUpdate = Writer.GetNumBits() > 0u;
+
+	// Header is always written first in the packet
+	FBitWriterMark Reset;
+	FBitWriterMark Restore(Writer);
+	Reset.PopWithoutClear(Writer);
+	
+	// Write notification header or refresh the header if used space is the same.
+	PacketNotify.WriteHeader(Writer, bIsHeaderUpdate);
+
+	// Jump back to where we came from.
+	if (bIsHeaderUpdate)
+	{
+		Restore.PopWithoutClear(Writer);
+	}
+}
+
+void UNetConnection::WritePacketInfo(FBitWriter& Writer) const
+{
+	const uint8 bHasServerFrameTime = Driver->IsServer() ? bLastHasServerFrameTime : ( CVarPingExcludeFrameTime.GetValueOnGameThread() > 0 ? 1u : 0u );
+
+	// Write data used to calculate link latency
+	Writer.WriteBit(bHasServerFrameTime);
+	if (bHasServerFrameTime && Driver->IsServer())
+	{
+		uint8 FrameTimeByte = FMath::Min( FMath::FloorToInt( FrameTime * 1000 ), 255 );
+		Writer << FrameTimeByte;
+	}
+
+	// Notify server of our current rate per second at this time
+	uint8 InKBytesPerSecondByte = FMath::Min(InBytesPerSecond / 1024, 255);
+	Writer << InKBytesPerSecondByte;
+}
+
+bool UNetConnection::ReadPacketInfo(FBitReader& Reader)
+{
+	const bool bHasServerFrameTime = Reader.ReadBit() == 1u;
+	double ServerFrameTime = 0;
+
+	if ( !Driver->IsServer() )
+	{
+		if ( bHasServerFrameTime )
+		{
+			uint8 FrameTimeByte	= 0;
+			Reader << FrameTimeByte;
+			// As a client, our request was granted, read the frame time
+			ServerFrameTime = ( double )FrameTimeByte / 1000;
+		}
+	}
+	else
+	{
+		bLastHasServerFrameTime = bHasServerFrameTime;
+	}
+
+	// limit to known size to know the size of the packet header
+	uint8 RemoteInKBytesPerSecondByte = 0;
+	Reader << RemoteInKBytesPerSecondByte;
+
+	if ( Reader.IsError() )
+	{
+		return false;
+	}
+
+	// Update ping
+	// At this time we have updated OutAckPacketId to the latest received ack.
+	const int32 Index = OutAckPacketId & (ARRAY_COUNT(OutLagPacketId)-1);
+
+	if ( OutLagPacketId[Index] == OutAckPacketId )
+	{
+		OutLagPacketId[Index] = -1;		// Only use the ack once
+
+#if !UE_BUILD_SHIPPING
+		if ( CVarPingDisplayServerTime.GetValueOnAnyThread() > 0 )
+		{
+			UE_LOG( LogNetTraffic, Warning, TEXT( "ServerFrameTime: %2.2f" ), ServerFrameTime * 1000.0f );
+		}
+#endif
+
+		// use FApp's time because it is set closer to the beginning of the frame - we don't care about the time so far of the current frame to process the packet
+		const double CurrentTime = FApp::GetCurrentTime();
+		const double GameTime	 = ServerFrameTime;
+		const double RTT		 = (CurrentTime - OutLagTime[Index] ) - ( CVarPingExcludeFrameTime.GetValueOnAnyThread() ? GameTime : 0.0 );
+		const double NewLag		 = FMath::Max( RTT, 0.0 );
+
+		if ( OutBytesPerSecondHistory[Index] > 0 )
+		{
+			RemoteSaturation = ( 1.0f - FMath::Min( ( float )RemoteInKBytesPerSecondByte / ( float )OutBytesPerSecondHistory[Index], 1.0f ) ) * 100.0f;
+		}
+		else
+		{
+			RemoteSaturation = 0.0f;
+		}
+
+		//UE_LOG( LogNet, Warning, TEXT( "Out: %i, InRemote: %i, Saturation: %f" ), OutBytesPerSecondHistory[Index], RemoteInKBytesPerSecond, RemoteSaturation );
+
+		LagAcc += NewLag;
+		LagCount++;
+
+		if (PlayerController != NULL)
+		{
+			PlayerController->UpdatePing(NewLag);
+		}
+	}
+
+	return true;
 }
 
 void UNetConnection::ReceivedPacket( FBitReader& Reader )
@@ -1319,201 +1508,107 @@ void UNetConnection::ReceivedPacket( FBitReader& Reader )
 	LastReceiveTime		= Driver->Time;
 	LastReceiveRealtime = FPlatformTime::Seconds();
 
-	// Check packet ordering.
-	const int32 PacketId = InternalAck ? InPacketId + 1 : MakeRelative(Reader.ReadInt(MAX_PACKETID),InPacketId,MAX_PACKETID);
-	if( PacketId > InPacketId )
+	if (InternalAck)
 	{
-		const int32 PacketsLost = PacketId - InPacketId - 1;
-		
-		if ( PacketsLost > 10 )
-		{
-			UE_LOG( LogNetTraffic, Log, TEXT( "High single frame packet loss. PacketsLost: %i %s" ), PacketsLost, *Describe() );
-		}
-
-		InPacketsLost += PacketsLost;
-		InTotalPacketsLost += PacketsLost;
-		Driver->InPacketsLost += PacketsLost;
-		Driver->InTotalPacketsLost += PacketsLost;
-		InPacketId = PacketId;
-	}
+		++InPacketId;
+	}	
 	else
 	{
-		Driver->InOutOfOrderPackets++;
-		// Protect against replay attacks
-		// We already protect against this for reliable bunches, and unreliable properties
-		// The only bunch we would process would be unreliable RPC's, which could allow for replay attacks
-		// So rather than add individual protection for unreliable RPC's as well, just kill it at the source, 
-		// which protects everything in one fell swoop
-		return;
-	}
-
-	const bool bIgnoreRPCs = Driver->ShouldIgnoreRPCs();
-
-	bool bSkipAck = false;
-
-	// Track channels that were rejected while processing this packet - used to avoid sending multiple close-channel bunches,
-	// which would cause a disconnect serverside
-	TArray<int32> RejectedChans;
-
-	// Disassemble and dispatch all bunches in the packet.
-	while( !Reader.AtEnd() && State!=USOCK_Closed )
-	{
-		// Parse the bunch.
-		int32 StartPos = Reader.GetPosBits();
-		bool IsAck = !!Reader.ReadBit();
-		if ( Reader.IsError() )
+		// Read packet header
+		FNetPacketNotify::FNotificationHeader Header;
+		if (!PacketNotify.ReadHeader(Header, Reader))
 		{
-			CLOSE_CONNECTION_DUE_TO_SECURITY_VIOLATION(this, ESecurityEvent::Malformed_Packet, TEXT("Bunch missing ack flag"));
+			CLOSE_CONNECTION_DUE_TO_SECURITY_VIOLATION(this, ESecurityEvent::Malformed_Packet, TEXT("Failed to read PacketHeader"));
 			return;
 		}
 
-		// Process the bunch.
-		if( IsAck )
+		// Lambda to dispatch delivery notifications, 
+		auto HandlePacketNotification = [&Header, this](FNetPacketNotify::SequenceNumberT AckedSequence, bool bDelivered)
 		{
-			LastRecvAckTime = Driver->Time;
-
-			// This is an acknowledgment.
-			const int32 AckPacketId = MakeRelative(Reader.ReadInt(MAX_PACKETID),OutAckPacketId,MAX_PACKETID);
+			// Increase LastNotifiedPacketId, this is a full packet Id
+			++LastNotifiedPacketId;
 
-			if( Reader.IsError() )
+			// Sanity check
+			if (FNetPacketNotify::SequenceNumberT(LastNotifiedPacketId) != AckedSequence)
 			{
-				CLOSE_CONNECTION_DUE_TO_SECURITY_VIOLATION(this, ESecurityEvent::Malformed_Packet, TEXT("Bunch missing ack"));
+				CLOSE_CONNECTION_DUE_TO_SECURITY_VIOLATION(this, ESecurityEvent::Malformed_Packet, TEXT("LastNotifiedPacketId != AckedSequence"));
 				return;
 			}
 
-			double ServerFrameTime = 0;
-
-			// If this is the server, we're reading in the request to send them our frame time
-			// If this is the client, we're reading in confirmation that our request to get frame time from server is granted
-			const bool bHasServerFrameTime = !!Reader.ReadBit();
-
-			if ( !Driver->IsServer() )
+			if (bDelivered)
 			{
-				if ( bHasServerFrameTime )
-				{
-					// As a client, our request was granted, read the frame time
-					uint8 FrameTimeByte	= 0;
-					Reader << FrameTimeByte;
-					ServerFrameTime = ( double )FrameTimeByte / 1000;
-				}
+				ReceivedAck(LastNotifiedPacketId);
 			}
 			else
 			{
-				// Server remembers so he can use during SendAck to notify to client of his frame time
-				bLastHasServerFrameTime = bHasServerFrameTime;
-			}
-
-			uint32 RemoteInKBytesPerSecond = 0;
-			Reader.SerializeIntPacked( RemoteInKBytesPerSecond );
+				ReceivedNak(LastNotifiedPacketId);
+			};
+		};
 
-			// Resend any old reliable packets that the receiver hasn't acknowledged.
-			if( AckPacketId>OutAckPacketId )
-			{
-				for (int32 NakPacketId = OutAckPacketId + 1; NakPacketId<AckPacketId; NakPacketId++, OutPacketsLost++, OutTotalPacketsLost++, Driver->OutTotalPacketsLost++)
-				{
-					UE_LOG(LogNetTraffic, Verbose, TEXT("   Received virtual nak %i (%.1f)"), NakPacketId, (Reader.GetPosBits()-StartPos)/8.f );
-					ReceivedNak( NakPacketId );
-				}
-				OutAckPacketId = AckPacketId;
-			}
-			else if( AckPacketId<OutAckPacketId )
+		// Update incoming sequence data and deliver packet notifications
+		// Packet is only accepted if both the incoming sequence number and incoming ack data are valid
+		int PacketSequenceDelta = PacketNotify.Update(Header, HandlePacketNotification);
+		if (PacketSequenceDelta > 0)
+		{
+			const int32 PacketsLost = PacketSequenceDelta - 1;
+		
+			if ( PacketsLost > 10 )
 			{
-				//warning: Double-ack logic makes this unmeasurable.
-				//OutOrdAcc++;
+				UE_LOG( LogNetTraffic, Verbose, TEXT( "High single frame packet loss. PacketsLost: %i %s" ), PacketsLost, *Describe() );
 			}
 
-			// Update ping
-			const int32 Index = AckPacketId & (ARRAY_COUNT(OutLagPacketId)-1);
-
-			if ( OutLagPacketId[Index] == AckPacketId )
-			{
-				OutLagPacketId[Index] = -1;		// Only use the ack once
-
-#if !UE_BUILD_SHIPPING
-				if ( CVarPingDisplayServerTime.GetValueOnAnyThread() > 0 )
-				{
-					UE_LOG( LogNetTraffic, Warning, TEXT( "ServerFrameTime: %2.2f" ), ServerFrameTime * 1000.0f );
-				}
-#endif
-
-				// use FApp's time because it is set closer to the beginning of the frame - we don't care about the time so far of the current frame to process the packet
-				const double CurrentTime = FApp::GetCurrentTime();
-				const double GameTime	 = ServerFrameTime;
-				const double RTT		 = (CurrentTime - OutLagTime[Index] ) - ( CVarPingExcludeFrameTime.GetValueOnAnyThread() ? GameTime : 0.0 );
-				const double NewLag		 = FMath::Max( RTT, 0.0 );
-
-				if ( OutBytesPerSecondHistory[Index] > 0 )
-				{
-					RemoteSaturation = ( 1.0f - FMath::Min( ( float )RemoteInKBytesPerSecond / ( float )OutBytesPerSecondHistory[Index], 1.0f ) ) * 100.0f;
-				}
-				else
-				{
-					RemoteSaturation = 0.0f;
-				}
-
-				//UE_LOG( LogNet, Warning, TEXT( "Out: %i, InRemote: %i, Saturation: %f" ), OutBytesPerSecondHistory[Index], RemoteInKBytesPerSecond, RemoteSaturation );
-
-				LagAcc += NewLag;
-				LagCount++;
-
-				if (PlayerController != NULL)
-				{
-					PlayerController->UpdatePing(NewLag);
-				}
-			}
+			InPacketsLost += PacketsLost;
+			InTotalPacketsLost += PacketsLost;
+			Driver->InPacketsLost += PacketsLost;
+			Driver->InTotalPacketsLost += PacketsLost;
+			InPacketId += PacketSequenceDelta;
 
-			if ( PackageMap != NULL )
+			// Extra information associated with the header
+			if (!ReadPacketInfo(Reader))
 			{
-				PackageMap->ReceivedAck( AckPacketId );
+				CLOSE_CONNECTION_DUE_TO_SECURITY_VIOLATION(this, ESecurityEvent::Malformed_Packet, TEXT("Failed to read PacketHeader"));
+				return;
 			}
+		}
+		else
+		{
+			Driver->InOutOfOrderPackets++;
+			// Protect against replay attacks
+			// We already protect against this for reliable bunches, and unreliable properties
+			// The only bunch we would process would be unreliable RPC's, which could allow for replay attacks
+			// So rather than add individual protection for unreliable RPC's as well, just kill it at the source, 
+			// which protects everything in one fell swoop
+			return;
+		}
+	}
 
-			// Forward the ack to the channel.
-			UE_LOG(LogNetTraffic, Verbose, TEXT("   Received ack %i (%.1f)"), AckPacketId, (Reader.GetPosBits()-StartPos)/8.f );
-
-			{
-				SCOPE_CYCLE_COUNTER(Stat_NetConnectionReceivedAcks);
-
-				for (int32 i = OpenChannels.Num() - 1; i >= 0; i--)
-				{
-					UChannel* const Channel = OpenChannels[i];
+	const bool bIgnoreRPCs = Driver->ShouldIgnoreRPCs();
 
-					if (Channel)
-					{
-						if (Channel->OpenPacketId.Last==AckPacketId) // Necessary for unreliable "bNetTemporary" channels.
-						{
-							Channel->OpenAcked = 1;
-						}
+	bool bSkipAck = false;
 
-						for (FOutBunch* OutBunch = Channel->OutRec; OutBunch; OutBunch = OutBunch->Next)
-						{
-							if (OutBunch->bOpen)
-							{
-								UE_LOG(LogNet, VeryVerbose, TEXT("Channel %i reset Ackd because open is reliable. "), Channel->ChIndex );
-								Channel->OpenAcked  = 0; // We have a reliable open bunch, don't let the above code set the OpenAcked state,
-														 // it must be set in UChannel::ReceivedAcks to verify all open bunches were received.
-							}
+	// Track channels that were rejected while processing this packet - used to avoid sending multiple close-channel bunches,
+	// which would cause a disconnect serverside
+	TArray<int32> RejectedChans;
 
-							if (OutBunch->PacketId==AckPacketId)
-							{
-								OutBunch->ReceivedAck = 1;
-							}
-						}
-						Channel->ReceivedAcks(); //warning: May destroy Channel.
-					}
-					else
-					{
-						UE_LOG(LogNet, Warning, TEXT("UNetConnection::ReceivedPacket: null channel in OpenChannels array while processing ack. %s"), *Describe());
-					}
-				}
-			}	
+	// Disassemble and dispatch all bunches in the packet.
+	while( !Reader.AtEnd() && State!=USOCK_Closed )
+	{
+		// For demo backwards compatibility, old replays still have this bit
+		if (InternalAck && EngineNetworkProtocolVersion < EEngineNetworkVersionHistory::HISTORY_ACKS_INCLUDED_IN_HEADER)
+		{
+			const bool IsAckDummy = Reader.ReadBit() == 1u;
 		}
-		else
+
+		// Parse the bunch.
+		int32 StartPos = Reader.GetPosBits();
+		
+		// Process Received data
 		{
 			// Parse the incoming data.
 			FInBunch Bunch( this );
 			int32 IncomingStartPos		= Reader.GetPosBits();
 			uint8 bControl				= Reader.ReadBit();
-			Bunch.PacketId				= PacketId;
+			Bunch.PacketId				= InPacketId;
 			Bunch.bOpen					= bControl ? Reader.ReadBit() : 0;
 			Bunch.bClose				= bControl ? Reader.ReadBit() : 0;
 			
@@ -1572,7 +1667,7 @@ PRAGMA_ENABLE_DEPRECATION_WARNINGS
 			else if ( Bunch.bPartial )
 			{
 				// If this is an unreliable partial bunch, we simply use packet sequence since we already have it
-				Bunch.ChSequence = PacketId;
+				Bunch.ChSequence = InPacketId;
 			}
 			else
 			{
@@ -1900,8 +1995,36 @@ PRAGMA_ENABLE_DEPRECATION_WARNINGS
 	if ( !bSkipAck )
 	{
 		LastGoodPacketRealtime = FPlatformTime::Seconds();
+	}
+
+	if( !InternalAck )
+	{
+		// We always call AckSequence even if we are explicitly rejecting the packet as this updates the expected InSeq used to drive future acks.
+		if ( bSkipAck )
+		{
+			// Explicit Nak, we treat this packet as dropped but we still report it to the sending side as quickly as possible
+			PacketNotify.NakSeq( InPacketId );
+		}
+		else
+		{
+			PacketNotify.AckSeq( InPacketId );
+
+			// Keep stats happy
+			++OutTotalAcks;
+			++Driver->OutTotalAcks;
+		}
+
+		// We do want to let the other side know about the ack, so even if there are no other outgoing data when we tick the connection we will send an ackpacket.
+		TimeSensitive = 1;
+		++HasDirtyAcks;
 
-		SendAck(PacketId, true);
+		// This is to allow us to recover from hitches were we process more than FNetPacketNotify::SequenceHistoryLength packets in a row withouht sending out any packets.
+		// In most cases this allows us to recover smoothly without reporting any pessimistic naks which is the case is we overshoot our ack history.
+		if (HasDirtyAcks >= FNetPacketNotify::MaxSequenceHistoryLength)
+		{
+			UE_LOG(LogNet, Warning, TEXT("UNetConnection::ReceivedPacket - Too many received acks(%u) since last sent packet. %s "), HasDirtyAcks, *Describe());
+			FlushNet();
+		}
 	}
 }
 
@@ -1945,10 +2068,24 @@ int32 UNetConnection::WriteBitsToSendBuffer(
 	// If this is the start of the queue, make sure to add the packet id
 	if ( SendBuffer.GetNumBits() == 0 && !InternalAck )
 	{
-		SendBuffer.WriteIntWrapped( OutPacketId, MAX_PACKETID );
-		ValidateSendBuffer();
+		// Write Packet Header, before sending the packet we will go back and rewrite the data
+		WritePacketHeader(SendBuffer);
+
+		// Also write server RTT and Received rate
+		WritePacketInfo(SendBuffer);
+
+		// We do not allow the first bunch to merge with the ack data as this will "revert" the ack data.
+		AllowMerge = false;
+	
+		// Update stats for PacketIdBits and ackdata (also including the data used for packet RTT and saturation calculations)
+		int64 BitsWritten = SendBuffer.GetNumBits();
+		NumPacketIdBits += FNetPacketNotify::SequenceNumberT::SeqNumberBits;
+		NumAckBits += BitsWritten - FNetPacketNotify::SequenceNumberT::SeqNumberBits;
+
+		// Report stats to profiler
+		NETWORK_PROFILER( GNetworkProfiler.TrackSendAck( NumAckBits, this ) );
 
-		NumPacketIdBits += SendBuffer.GetNumBits();
+		ValidateSendBuffer();
 	}
 
 	// Add the bits to the queue
@@ -1972,9 +2109,6 @@ int32 UNetConnection::WriteBitsToSendBuffer(
 		case EWriteBitsDataType::Bunch:
 			NumBunchBits += SizeInBits + ExtraSizeInBits;
 			break;
-		case EWriteBitsDataType::Ack:
-			NumAckBits += SizeInBits + ExtraSizeInBits;
-			break;
 		default:
 			break;
 	}
@@ -2022,64 +2156,12 @@ TSharedPtr<FObjectReplicator> UNetConnection::CreateReplicatorForNewActorChannel
 
 void UNetConnection::PurgeAcks()
 {
-	for ( int32 i = 0; i < ResendAcks.Num(); i++ )
-	{
-		SendAck( ResendAcks[i], 0 );
-	}
-
-	ResendAcks.Empty(32);
 }
 
-
 void UNetConnection::SendAck(int32 AckPacketId, bool FirstTime/*=1*/)
 {
-	SCOPE_CYCLE_COUNTER(Stat_NetConnectionSendAck);
-
-	ValidateSendBuffer();
-
-	if( !InternalAck )
-	{
-		if( FirstTime )
-		{
-			PurgeAcks();
-			QueuedAcks.Add(AckPacketId);
-		}
-
-		FBitWriter AckData( 32, true );
-
-		AckData.WriteBit( 1 );
-		AckData.WriteIntWrapped(AckPacketId, MAX_PACKETID);
-		
-		const bool bHasServerFrameTime = Driver->IsServer() ? bLastHasServerFrameTime : ( CVarPingExcludeFrameTime.GetValueOnGameThread() > 0 ? true : false );
-
-		AckData.WriteBit( bHasServerFrameTime ? 1 : 0 );
-
-		if ( Driver->IsServer() && bHasServerFrameTime )
-		{
-			uint8 FrameTimeByte = FMath::Min( FMath::FloorToInt( FrameTime * 1000 ), 255 );
-			AckData << FrameTimeByte;
-		}
-
-		// Notify server of our current rate per second at this time
-		uint32 InKBytesPerSecond = InBytesPerSecond / 1024;
-		AckData.SerializeIntPacked( InKBytesPerSecond );
-
-		NETWORK_PROFILER( GNetworkProfiler.TrackSendAck( AckData.GetNumBits(), this ) );
-
-		WriteBitsToSendBuffer( AckData.GetData(), AckData.GetNumBits(), nullptr, 0, EWriteBitsDataType::Ack );
-
-		AllowMerge = false;
-
-		TimeSensitive = 1;
-
-		UE_LOG(LogNetTraffic, Log, TEXT("   Send ack %i"), AckPacketId);
-	}
-
-	++OutTotalAcks;
-	++Driver->OutTotalAcks;
 }
 
-
 int32 UNetConnection::SendRawBunch( FOutBunch& Bunch, bool InAllowMerge )
 {
 	ValidateSendBuffer();
@@ -2091,7 +2173,7 @@ int32 UNetConnection::SendRawBunch( FOutBunch& Bunch, bool InAllowMerge )
 
 	// Build header.
 	SendBunchHeader.Reset();
-	SendBunchHeader.WriteBit( 0 );
+
 	SendBunchHeader.WriteBit( Bunch.bOpen || Bunch.bClose );
 	if( Bunch.bOpen || Bunch.bClose )
 	{
@@ -2537,8 +2619,6 @@ void UNetConnection::Tick()
 	}
 
 	// Flush.
-	PurgeAcks();
-
 	if ( TimeSensitive || (Driver->Time - LastSendTime) > Driver->KeepAliveTime)
 	{
 		bool bHandlerHandshakeComplete = !Handler.IsValid() || Handler->IsFullyInitialized();
