Fixed compilation error introduced in cl 4718030

diff --git a/Engine/Source/Runtime/Engine/Private/NetConnection.cpp b/Engine/Source/Runtime/Engine/Private/NetConnection.cpp
--- a/Engine/Source/Runtime/Engine/Private/NetConnection.cpp
+++ b/Engine/Source/Runtime/Engine/Private/NetConnection.cpp
@@ -491,8 +491,6 @@ void UNetConnection::Serialize( FArchive& Ar )
 		OutReliable.CountBytes(Ar);
 		InReliable.CountBytes(Ar);
 		PendingOutRec.CountBytes(Ar);
-		QueuedAcks.CountBytes(Ar);
-		ResendAcks.CountBytes(Ar);
 		ActorChannels.CountBytes(Ar);
 		DestroyedStartupOrDormantActorGUIDs.CountBytes(Ar);
 		KeepProcessingActorChannelBunchesMap.CountBytes(Ar);
