Fix static analysis warnings on pointer access (code wouldn't have actually accessed an invalid point originally but added safety code for clarity)

diff --git a/Engine/Source/Runtime/Engine/Private/Animation/AnimCurveCompressionCodec_UniformlySampled.cpp b/Engine/Source/Runtime/Engine/Private/Animation/AnimCurveCompressionCodec_UniformlySampled.cpp
--- a/Engine/Source/Runtime/Engine/Private/Animation/AnimCurveCompressionCodec_UniformlySampled.cpp
+++ b/Engine/Source/Runtime/Engine/Private/Animation/AnimCurveCompressionCodec_UniformlySampled.cpp
@@ -1,4 +1,4 @@
-// Copyright 1998-2018 Epic Games, Inc. All Rights Reserved.
+// Copyright 1998-2019 Epic Games, Inc. All Rights Reserved.
 
 #include "Animation/AnimCurveCompressionCodec_UniformlySampled.h"
 #include "Animation/AnimSequence.h"
@@ -88,48 +88,54 @@ bool UAnimCurveCompressionCodec_UniformlySampled::Compress(const UAnimSequence&
 		float* AnimatedSamplesPtr = NumAnimatedCurves > 0 ? (float*)&Buffer[BufferOffset] : nullptr;
 		BufferOffset += sizeof(float) * NumAnimatedCurves * NumSamples;
 
-		for (int32 CurveIndex = 0, ConstantCurveIndex = 0; CurveIndex < NumCurves; ++CurveIndex)
+		if(ConstantSamplesPtr)
 		{
-			const FFloatCurve& Curve = AnimSeq.RawCurveData.FloatCurves[CurveIndex];
-			if (Curve.FloatCurve.IsConstant())
+			for (int32 CurveIndex = 0, ConstantCurveIndex = 0; CurveIndex < NumCurves; ++CurveIndex)
 			{
-				if (Curve.FloatCurve.IsEmpty())
-				{
-					ConstantSamplesPtr[ConstantCurveIndex] = Curve.FloatCurve.DefaultValue;
-				}
-				else
+				const FFloatCurve& Curve = AnimSeq.RawCurveData.FloatCurves[CurveIndex];
+				if (Curve.FloatCurve.IsConstant())
 				{
-					ConstantSamplesPtr[ConstantCurveIndex] = Curve.FloatCurve.Keys[0].Value;
+					if (Curve.FloatCurve.IsEmpty())
+					{
+						ConstantSamplesPtr[ConstantCurveIndex] = Curve.FloatCurve.DefaultValue;
+					}
+					else
+					{
+						ConstantSamplesPtr[ConstantCurveIndex] = Curve.FloatCurve.Keys[0].Value;
+					}
+
+					// Bitset uses little-endian bit ordering
+					ConstantCurvesBitsetPtr[CurveIndex / 32] |= 1 << (CurveIndex % 32);
+					ConstantCurveIndex++;
 				}
-
-				// Bitset uses little-endian bit ordering
-				ConstantCurvesBitsetPtr[CurveIndex / 32] |= 1 << (CurveIndex % 32);
-				ConstantCurveIndex++;
 			}
 		}
 
-		// Write out samples sorted by time first in order to have everything contiguous in memory
-		// for improved cache locality
-		// Curve 0 Key 0, Curve 0 Key 1, Curve 0 Key N, Curve 1 Key 0, Curve 1 Key 1, Curve 1 Key N, Curve M Key 0, ...
-		const float InvSampleRate = 1.0f / SampleRate_;
-		for (int32 SampleIndex = 0; SampleIndex < NumSamples; ++SampleIndex)
+		if(AnimatedSamplesPtr)
 		{
-			const float SampleTime = FMath::Clamp(SampleIndex * InvSampleRate, 0.0f, Duration);
-			float* AnimatedSamples = AnimatedSamplesPtr + (SampleIndex * NumAnimatedCurves);
-
-			for (int32 CurveIndex = 0, AnimatedCurveIndex = 0; CurveIndex < NumCurves; ++CurveIndex)
+			// Write out samples sorted by time first in order to have everything contiguous in memory
+			// for improved cache locality
+			// Curve 0 Key 0, Curve 0 Key 1, Curve 0 Key N, Curve 1 Key 0, Curve 1 Key 1, Curve 1 Key N, Curve M Key 0, ...
+			const float InvSampleRate = 1.0f / SampleRate_;
+			for (int32 SampleIndex = 0; SampleIndex < NumSamples; ++SampleIndex)
 			{
-				const FFloatCurve& Curve = AnimSeq.RawCurveData.FloatCurves[CurveIndex];
-				if (Curve.FloatCurve.IsConstant())
+				const float SampleTime = FMath::Clamp(SampleIndex * InvSampleRate, 0.0f, Duration);
+				float* AnimatedSamples = AnimatedSamplesPtr + (SampleIndex * NumAnimatedCurves);
+
+				for (int32 CurveIndex = 0, AnimatedCurveIndex = 0; CurveIndex < NumCurves; ++CurveIndex)
 				{
-					// Skip constant curves, their data has already been written
-					continue;
-				}
+					const FFloatCurve& Curve = AnimSeq.RawCurveData.FloatCurves[CurveIndex];
+					if (Curve.FloatCurve.IsConstant())
+					{
+						// Skip constant curves, their data has already been written
+						continue;
+					}
 
-				const float SampleValue = Curve.FloatCurve.Eval(SampleTime);
+					const float SampleValue = Curve.FloatCurve.Eval(SampleTime);
 
-				AnimatedSamples[AnimatedCurveIndex] = SampleValue;
-				AnimatedCurveIndex++;
+					AnimatedSamples[AnimatedCurveIndex] = SampleValue;
+					AnimatedCurveIndex++;
+				}
 			}
 		}
 	}
