GetLifeSpan should not crash if called on actors that no longer have a world

diff --git a/Engine/Source/Runtime/Engine/Private/Actor.cpp b/Engine/Source/Runtime/Engine/Private/Actor.cpp
--- a/Engine/Source/Runtime/Engine/Private/Actor.cpp
+++ b/Engine/Source/Runtime/Engine/Private/Actor.cpp
@@ -4565,9 +4565,14 @@ void AActor::SetLifeSpan( float InLifespan )
 
 float AActor::GetLifeSpan() const
 {
-	// Timer remaining returns -1.0f if there is no such timer - return this as ZERO
-	const float CurrentLifespan = GetWorldTimerManager().GetTimerRemaining(TimerHandle_LifeSpanExpired);
-	return ( CurrentLifespan != -1.0f ) ? CurrentLifespan : 0.0f;
+	if (UWorld* World = GetWorld())
+	{
+		// Timer remaining returns -1.0f if there is no such timer - return this as ZERO
+		const float CurrentLifespan = World->GetTimerManager().GetTimerRemaining(TimerHandle_LifeSpanExpired);
+		return (CurrentLifespan != -1.0f) ? CurrentLifespan : 0.0f;
+	}
+	
+	return 0.0f;
 }
 
 void AActor::PostInitializeComponents()
