[NO NON-PII COMMIT MESSAGE]

diff --git a/Engine/Source/Runtime/Engine/Public/Net/VoiceConfig.h b/Engine/Source/Runtime/Engine/Public/Net/VoiceConfig.h
--- a/Engine/Source/Runtime/Engine/Public/Net/VoiceConfig.h
+++ b/Engine/Source/Runtime/Engine/Public/Net/VoiceConfig.h
@@ -37,7 +37,7 @@ namespace MicSilenceDetectionConfig
 	static const Audio::EPeakMode::Type LevelDetectionMode = Audio::EPeakMode::Type::Peak;
 	static const bool IsAnalog = false;
 
-	static const int32 PacketBufferSlack = 6;
+	static const int32 PacketBufferSlack = 20;
 }
 
 class USoundAttenuation;
