tentative fix for some VOIP clients getting pitched-up audio. UVoipStatics::GetSampleRate could get called too early and return an incorrect sample rate because a UEnum was not loaded yet, thus causing all buffers to be incorrectly truncated.

diff --git a/Engine/Source/Runtime/Engine/Private/VoiceConfig.cpp b/Engine/Source/Runtime/Engine/Private/VoiceConfig.cpp
--- a/Engine/Source/Runtime/Engine/Private/VoiceConfig.cpp
+++ b/Engine/Source/Runtime/Engine/Private/VoiceConfig.cpp
@@ -35,13 +35,21 @@ int32 UVOIPStatics::GetVoiceSampleRate()
 	}
 
 	static FString DesiredSampleRateStr;
-	bRetrievedSampleRate = true;
-
+	
 	if (GConfig->GetString(TEXT("/Script/Engine.AudioSettings"), TEXT("VoiPSampleRate"), DesiredSampleRateStr, GEngineIni))
 	{
-		SampleRate = (int32)GetEnumValueFromString<EVoiceSampleRate>(TEXT("EVoiceSampleRate"), DesiredSampleRateStr);
+		if (DesiredSampleRateStr.Equals(TEXT("Low16000Hz")))
+		{
+			SampleRate = (int32)EVoiceSampleRate::Low16000Hz;
+		}
+		else if (DesiredSampleRateStr.Equals(TEXT("Normal24000Hz")))
+		{
+			SampleRate = (int32)EVoiceSampleRate::Normal24000Hz;
+		}
+
 		if (SampleRate > 0)
 		{		
+			bRetrievedSampleRate = true;
 			return SampleRate;
 		}
 	}
diff --git a/Engine/Source/Runtime/Engine/Public/Net/VoiceConfig.h b/Engine/Source/Runtime/Engine/Public/Net/VoiceConfig.h
--- a/Engine/Source/Runtime/Engine/Public/Net/VoiceConfig.h
+++ b/Engine/Source/Runtime/Engine/Public/Net/VoiceConfig.h
@@ -42,17 +42,6 @@ namespace MicSilenceDetectionConfig
 
 class USoundAttenuation;
 
-template <typename EnumType>
-static FORCEINLINE EnumType GetEnumValueFromString(const FString& EnumName, const FString& String)
-{
-	UEnum* Enum = FindObject<UEnum>((UObject*) ANY_PACKAGE, *EnumName, true);
-	if (!Enum)
-	{
-		return EnumType(0);
-	}
-	return (EnumType)Enum->GetValueByName(FName(*String));
-}
-
 USTRUCT(BlueprintType)
 struct FVoiceSettings
 {
