Fix false static analysis warning on pointer access

diff --git a/Engine/Source/Runtime/Engine/Private/Animation/AnimSequence.cpp b/Engine/Source/Runtime/Engine/Private/Animation/AnimSequence.cpp
--- a/Engine/Source/Runtime/Engine/Private/Animation/AnimSequence.cpp
+++ b/Engine/Source/Runtime/Engine/Private/Animation/AnimSequence.cpp
@@ -867,7 +867,8 @@ void UAnimSequence::PostEditChangeProperty(FPropertyChangedEvent& PropertyChange
 		PostProcessSequence();
 	}
 
-	if (PropertyChangedEvent.Property != nullptr && PropertyChangedEvent.Property->GetFName() == GET_MEMBER_NAME_CHECKED(UAnimSequence, CurveCompressionSettings))
+	UProperty* Property = PropertyChangedEvent.Property;
+	if (Property != nullptr && Property->GetFName() == GET_MEMBER_NAME_CHECKED(UAnimSequence, CurveCompressionSettings))
 	{
 		RequestSyncAnimRecompression(false);
 	}
