Updated copyright notice

diff --git a/Engine/Source/Runtime/Engine/Classes/Engine/ContentEncryptionConfig.h b/Engine/Source/Runtime/Engine/Classes/Engine/ContentEncryptionConfig.h
--- a/Engine/Source/Runtime/Engine/Classes/Engine/ContentEncryptionConfig.h
+++ b/Engine/Source/Runtime/Engine/Classes/Engine/ContentEncryptionConfig.h
@@ -1,4 +1,4 @@
-// Copyright 1998-2018 Epic Games, Inc. All Rights Reserved.
+// Copyright 1998-2019 Epic Games, Inc. All Rights Reserved.
 #pragma once
 
 /** Project specific configuration for content encryption */
