Fix for calling StaticLoadObject during PostLoad of AnimSequence

diff --git a/Engine/Source/Runtime/Engine/Private/Animation/AnimSequence.cpp b/Engine/Source/Runtime/Engine/Private/Animation/AnimSequence.cpp
--- a/Engine/Source/Runtime/Engine/Private/Animation/AnimSequence.cpp
+++ b/Engine/Source/Runtime/Engine/Private/Animation/AnimSequence.cpp
@@ -468,6 +468,13 @@ void UAnimSequence::Serialize(FArchive& Ar)
 		SourceFileTimestamp_DEPRECATED = TEXT("");
 	}
 
+	// Do this is serialize as if the default animation curve compression asset isn't loaded it will
+	// fire a warning if we try and load it in post load
+	if (CurveCompressionSettings == nullptr || !CurveCompressionSettings->AreSettingsValid())
+	{
+		CurveCompressionSettings = FAnimationUtils::GetDefaultAnimationCurveCompressionSettings();
+	}
+
 #endif // WITH_EDITORONLY_DATA
 
 	AddAnimLoadingDebugEntry(TEXT("PostSerialize"));
