Add a pure virtual Duplicate function to allocate a copy of the curve in  FIndexedCurve and implement in the various subclasses.

diff --git a/Engine/Source/Runtime/Engine/Classes/Curves/IndexedCurve.h b/Engine/Source/Runtime/Engine/Classes/Curves/IndexedCurve.h
--- a/Engine/Source/Runtime/Engine/Classes/Curves/IndexedCurve.h
+++ b/Engine/Source/Runtime/Engine/Classes/Curves/IndexedCurve.h
@@ -46,6 +46,9 @@ public:
 	/** Get the time for the Key with the specified index. */
 	virtual float GetKeyTime(FKeyHandle KeyHandle) const PURE_VIRTUAL(FIndexedCurve::GetKeyTime, return 0.f;);
 	
+	/** Allocates a duplicate of the curve */
+	virtual FIndexedCurve* Duplicate() const PURE_VIRTUAL(FIndexedCurve::Duplicate, return nullptr;);
+
 	/** Checks to see if the key handle is valid for this curve. */
 	bool IsKeyHandleValid(FKeyHandle KeyHandle) const;
 
diff --git a/Engine/Source/Runtime/Engine/Classes/Curves/IntegralCurve.h b/Engine/Source/Runtime/Engine/Classes/Curves/IntegralCurve.h
--- a/Engine/Source/Runtime/Engine/Classes/Curves/IntegralCurve.h
+++ b/Engine/Source/Runtime/Engine/Classes/Curves/IntegralCurve.h
@@ -53,6 +53,9 @@ public:
 	/** Get number of keys in curve. */
 	virtual int32 GetNumKeys() const override final { return Keys.Num(); }
 
+	/** Allocates a duplicate of the curve */
+	virtual FIndexedCurve* Duplicate() const final { return new FIntegralCurve(*this); }
+
 	/** Evaluates the value of an array of keys at a time */
 	int32 Evaluate(float Time, int32 InDefaultValue = 0) const;
 
diff --git a/Engine/Source/Runtime/Engine/Classes/Curves/NameCurve.h b/Engine/Source/Runtime/Engine/Classes/Curves/NameCurve.h
--- a/Engine/Source/Runtime/Engine/Classes/Curves/NameCurve.h
+++ b/Engine/Source/Runtime/Engine/Classes/Curves/NameCurve.h
@@ -163,6 +163,9 @@ public:
 
 	virtual int32 GetNumKeys() const override final { return Keys.Num(); }
 
+	/** Allocates a duplicate of the curve */
+	virtual FIndexedCurve* Duplicate() const final { return new FNameCurve(*this); }
+
 public:
 
 	/** Sorted array of keys */
diff --git a/Engine/Source/Runtime/Engine/Classes/Curves/RichCurve.h b/Engine/Source/Runtime/Engine/Classes/Curves/RichCurve.h
--- a/Engine/Source/Runtime/Engine/Classes/Curves/RichCurve.h
+++ b/Engine/Source/Runtime/Engine/Classes/Curves/RichCurve.h
@@ -312,6 +312,9 @@ public:
 	/** Compresses a rich curve for efficient runtime storage and evaluation */
 	void CompressCurve(struct FCompressedRichCurve& OutCurve, float ErrorThreshold = 0.0001f, float SampleRate = 120.0f) const;
 
+	/** Allocates a duplicate of the curve */
+	virtual FIndexedCurve* Duplicate() const final { return new FRichCurve(*this); }
+
 private:
 	void RemoveRedundantKeysInternal(float Tolerance, int32 InStartKeepKey, int32 InEndKeepKey);
 	virtual int32 GetKeyIndex(float KeyTime, float KeyTimeTolerance) const override final;
diff --git a/Engine/Source/Runtime/Engine/Classes/Curves/SimpleCurve.h b/Engine/Source/Runtime/Engine/Classes/Curves/SimpleCurve.h
--- a/Engine/Source/Runtime/Engine/Classes/Curves/SimpleCurve.h
+++ b/Engine/Source/Runtime/Engine/Classes/Curves/SimpleCurve.h
@@ -186,6 +186,9 @@ public:
 	virtual void RemoveRedundantKeys(float Tolerance) final override;
 	virtual void RemoveRedundantKeys(float Tolerance, float FirstKeyTime, float LastKeyTime) final override;
 
+	/** Allocates a duplicate of the curve */
+	virtual FIndexedCurve* Duplicate() const final { return new FSimpleCurve(*this); }
+
 protected:
 	virtual int32 GetKeyIndex(float KeyTime, float KeyTimeTolerance) const override final;
 
diff --git a/Engine/Source/Runtime/Engine/Classes/Curves/StringCurve.h b/Engine/Source/Runtime/Engine/Classes/Curves/StringCurve.h
--- a/Engine/Source/Runtime/Engine/Classes/Curves/StringCurve.h
+++ b/Engine/Source/Runtime/Engine/Classes/Curves/StringCurve.h
@@ -212,6 +212,9 @@ public:
 
 	virtual int32 GetNumKeys() const override final { return Keys.Num(); }
 
+	/** Allocates a duplicate of the curve */
+	virtual FIndexedCurve* Duplicate() const final { return new FStringCurve(*this); }
+
 public:
 
 	/** Default value */
