Material layers cooking race condition fix.

diff --git a/Engine/Source/Runtime/Engine/Private/Materials/MaterialExpressions.cpp b/Engine/Source/Runtime/Engine/Private/Materials/MaterialExpressions.cpp
--- a/Engine/Source/Runtime/Engine/Private/Materials/MaterialExpressions.cpp
+++ b/Engine/Source/Runtime/Engine/Private/Materials/MaterialExpressions.cpp
@@ -10928,7 +10928,10 @@ FString FMaterialLayersFunctions::GetStaticPermutationString() const
 
 void FMaterialLayersFunctions::SerializeForDDC(FArchive& Ar)
 {
-	KeyString = GetStaticPermutationString();
+	if (!Ar.IsCooking())
+	{
+		KeyString = GetStaticPermutationString();
+	}
 	Ar << KeyString;
 }
 
