Copying //UE4/Dev-Framework to Dev-Main (//UE4/Dev-Main) @ 4736279
#rnx

diff --git a/Engine/Source/Runtime/Engine/Private/ActorConstruction.cpp b/Engine/Source/Runtime/Engine/Private/ActorConstruction.cpp
--- a/Engine/Source/Runtime/Engine/Private/ActorConstruction.cpp
+++ b/Engine/Source/Runtime/Engine/Private/ActorConstruction.cpp
@@ -1044,7 +1044,8 @@ UActorComponent* AActor::CreateComponentFromTemplateData(const FBlueprintCookedC
 
 UActorComponent* AActor::AddComponent(FName TemplateName, bool bManualAttachment, const FTransform& RelativeTransform, const UObject* ComponentTemplateContext)
 {
-	if (GetWorld()->bIsTearingDown)
+	UWorld* World = GetWorld();
+	if (World->bIsTearingDown)
 	{
 		UE_LOG(LogActor, Warning, TEXT("AddComponent failed because we are in the process of tearing down the world"));
 		return nullptr;
@@ -1115,7 +1116,6 @@ UActorComponent* AActor::AddComponent(FName TemplateName, bool bManualAttachment
 			NewActorComp->RegisterComponent();
 		}
 
-		UWorld* World = GetWorld();
 		if (!bRunningUserConstructionScript && World && bIsSceneComponent)
 		{
 			UPrimitiveComponent* NewPrimitiveComponent = Cast<UPrimitiveComponent>(NewActorComp);
