Fix server log spam in the demo net driver when there's no maximum recording time during prioritization

diff --git a/Engine/Source/Runtime/Engine/Private/DemoNetDriver.cpp b/Engine/Source/Runtime/Engine/Private/DemoNetDriver.cpp
--- a/Engine/Source/Runtime/Engine/Private/DemoNetDriver.cpp
+++ b/Engine/Source/Runtime/Engine/Private/DemoNetDriver.cpp
@@ -2594,7 +2594,7 @@ void UDemoNetDriver::TickDemoRecordFrame( float DeltaSeconds )
 
 	// Make sure we're under the desired recording time quota, if any.
 	// See ReplicatePriorizeActor.
-	if (TotalPrioritizeActorsTime > RecordTimeLimit)
+	if (RecordTimeLimit > 0.0f && TotalPrioritizeActorsTime > RecordTimeLimit)
 	{
 		DemoNetDriverRecordingPrivate::LogDemoRecordTimeElapsed(TEXT("Exceeded maximum desired recording time (during Prioritization).  Max: %.3fms, TimeSpent: %.3fms, Active Actors: %d, Prioritized Actors: %d"),
 			MaxDesiredRecordTimeMS, TotalPrioritizeActorsTimeMS, NumActiveObjects, NumPrioritizedActors);
