Implementing cook-time analysis data

diff --git a/Engine/Source/Programs/UnrealBuildTool/Configuration/UEBuildModuleCPP.cs b/Engine/Source/Programs/UnrealBuildTool/Configuration/UEBuildModuleCPP.cs
--- a/Engine/Source/Programs/UnrealBuildTool/Configuration/UEBuildModuleCPP.cs
+++ b/Engine/Source/Programs/UnrealBuildTool/Configuration/UEBuildModuleCPP.cs
@@ -82,6 +82,7 @@ namespace UnrealBuildTool
 			new KeyValuePair<string, string>("Engine", "GameplayTags"),
 			new KeyValuePair<string, string>("Engine", "MaterialShaderQualitySettings"),
 			new KeyValuePair<string, string>("Engine", "UnrealEd"),
+			new KeyValuePair<string, string>("Engine", "AudioMixer"),
 			new KeyValuePair<string, string>("PacketHandler", "ReliabilityHandlerComponent"),
 			new KeyValuePair<string, string>("GameplayDebugger", "AIModule"),
 			new KeyValuePair<string, string>("GameplayDebugger", "GameplayTasks"),
diff --git a/Engine/Source/Runtime/AudioMixer/AudioMixer.Build.cs b/Engine/Source/Runtime/AudioMixer/AudioMixer.Build.cs
--- a/Engine/Source/Runtime/AudioMixer/AudioMixer.Build.cs
+++ b/Engine/Source/Runtime/AudioMixer/AudioMixer.Build.cs
@@ -6,7 +6,7 @@ namespace UnrealBuildTool.Rules
 	{
 		public AudioMixer(ReadOnlyTargetRules Target) : base(Target)
 		{
-			OptimizeCode = CodeOptimization.Always;
+			OptimizeCode = CodeOptimization.Never;
 
 			PrivateIncludePathModuleNames.Add("TargetPlatform");
 
@@ -22,7 +22,6 @@ namespace UnrealBuildTool.Rules
 				{
 					"Core",
 					"CoreUObject",
-					"Engine",
 				}
 			);
 
@@ -30,6 +29,7 @@ namespace UnrealBuildTool.Rules
 				new string[]
 				{
 					"CoreUObject",
+					"Engine",
 				}
 			);
 
diff --git a/Engine/Source/Runtime/AudioMixer/Public/DSP/SpectrumAnalyzer.h b/Engine/Source/Runtime/AudioMixer/Public/DSP/SpectrumAnalyzer.h
--- a/Engine/Source/Runtime/AudioMixer/Public/DSP/SpectrumAnalyzer.h
+++ b/Engine/Source/Runtime/AudioMixer/Public/DSP/SpectrumAnalyzer.h
@@ -23,6 +23,7 @@ namespace Audio
 			Small_256 = 256,
 			Medium_512 = 512,
 			Large_1024 = 1024,
+			VeryLarge_2048 = 2048,
 			TestLarge_4096 = 4096
 		};
 
diff --git a/Engine/Source/Runtime/Engine/Classes/Components/AudioComponent.h b/Engine/Source/Runtime/Engine/Classes/Components/AudioComponent.h
--- a/Engine/Source/Runtime/Engine/Classes/Components/AudioComponent.h
+++ b/Engine/Source/Runtime/Engine/Classes/Components/AudioComponent.h
@@ -8,6 +8,7 @@
 #include "Engine/EngineTypes.h"
 #include "Components/SceneComponent.h"
 #include "Sound/SoundAttenuation.h"
+#include "Sound/SoundWave.h"
 
 #include "AudioComponent.generated.h"
 
@@ -427,6 +428,22 @@ public:
 	UFUNCTION(BlueprintCallable, Category = "Audio|Components|Audio")
 	void SetLowPassFilterFrequency(float InLowPassFilterFrequency);
 
+	/** Queries if the sound wave playing in this audio component has cooked FFT data. */
+	UFUNCTION(BlueprintCallable, Category = "Audio|Components|Audio")
+	bool HasCookedFFTData() const;
+
+	/** Queries if the sound wave playing in this audio component has cooked amplitude analyses. */
+	UFUNCTION(BlueprintCallable, Category = "Audio|Components|Audio")
+	bool HasCookedAmplitudeEnvelopeData() const;
+
+	/** Returns the current cooked spectral data of the playing audio component. Returns true if there is data and the audio component is playing. */
+	UFUNCTION(BlueprintCallable, Category = "Audio|Components|Audio")
+	bool GetCookedFFTData(const TArray<float>& FrequenciesToGet, TArray<FSoundWaveSpectralData>& OutSoundWaveSpectralData);
+
+	/** Returns the current cooked envelope data of the playing audio component. Returns true if there is data and the audio component is playing. */
+	UFUNCTION(BlueprintCallable, Category = "Audio|Components|Audio")
+	bool GetCookedEnvelopeData(float& OutEnvelopeData);
+
 	static void PlaybackCompleted(uint64 AudioComponentID, bool bFailedToStart);
 
 private:
@@ -479,7 +496,9 @@ public:
 
 	static UAudioComponent* GetAudioComponentFromID(uint64 AudioComponentID);
 
-	void UpdateInteriorSettings(bool bFullUpdate);
+	// Sets the audio thread playback time as used by the active sound playing this audio component
+	// Will be set if the audio component is using baked FFT or envelope following data so as to be able to feed that data to BP based on playback time
+	void SetPlaybackTimes(const TMap<uint32, float>& InSoundWavePlaybackTimes);
 
 public:
 
@@ -507,6 +526,32 @@ private:
 	FRotator SavedAutoAttachRelativeRotation;
 	FVector SavedAutoAttachRelativeScale3D;
 
+	struct FSoundWavePlaybackTimeData
+	{
+		USoundWave* SoundWave;
+		float PlaybackTime;
+
+		// Cachced indices to boost searching cooked data indices
+		uint32 LastEnvelopeCookedIndex;
+		uint32 LastFFTCookedIndex;
+
+		FSoundWavePlaybackTimeData()
+			: SoundWave(nullptr)
+			, PlaybackTime(0.0f)
+			, LastEnvelopeCookedIndex(INDEX_NONE)
+			, LastFFTCookedIndex(INDEX_NONE)
+		{}
+
+		FSoundWavePlaybackTimeData(USoundWave* InSoundWave)
+			: SoundWave(InSoundWave)
+			, PlaybackTime(0.0f)
+			, LastEnvelopeCookedIndex(INDEX_NONE)
+			, LastFFTCookedIndex(INDEX_NONE)
+		{}
+	};
+	// The current playback times of sound waves in this audio component
+	TMap<uint32, FSoundWavePlaybackTimeData> SoundWavePlaybackTimes;
+
 	/** Restore relative transform from auto attachment and optionally detach from parent (regardless of whether it was an auto attachment). */
 	void CancelAutoAttachment(bool bDetachFromParent);
 
diff --git a/Engine/Source/Runtime/Engine/Classes/Sound/SoundBase.h b/Engine/Source/Runtime/Engine/Classes/Sound/SoundBase.h
--- a/Engine/Source/Runtime/Engine/Classes/Sound/SoundBase.h
+++ b/Engine/Source/Runtime/Engine/Classes/Sound/SoundBase.h
@@ -218,5 +218,8 @@ public:
 	/** Returns the sound concurrency object ID if it exists. If it doesn't exist, returns 0. */
 	uint32 GetSoundConcurrencyObjectID() const;
 
+	/** Returns whether the sound has cooked analysis data (e.g. FFT or envelope following data) and returns sound waves which have cooked data. */
+	virtual bool GetSoundWavesWithCookedAnalysisData(TArray<USoundWave*>& OutSoundWaves);
+
 };
 
diff --git a/Engine/Source/Runtime/Engine/Classes/Sound/SoundCue.h b/Engine/Source/Runtime/Engine/Classes/Sound/SoundCue.h
--- a/Engine/Source/Runtime/Engine/Classes/Sound/SoundCue.h
+++ b/Engine/Source/Runtime/Engine/Classes/Sound/SoundCue.h
@@ -149,6 +149,7 @@ public:
 	virtual float GetDuration() override;
 	virtual const FSoundAttenuationSettings* GetAttenuationSettingsToApply() const override;
 	virtual float GetSubtitlePriority() const override;
+	virtual bool GetSoundWavesWithCookedAnalysisData(TArray<USoundWave*>& OutSoundWaves) override;
 	//~ End USoundBase Interface.
 
 	/** Construct and initialize a node within this Cue */
diff --git a/Engine/Source/Runtime/Engine/Classes/Sound/SoundWave.h b/Engine/Source/Runtime/Engine/Classes/Sound/SoundWave.h
--- a/Engine/Source/Runtime/Engine/Classes/Sound/SoundWave.h
+++ b/Engine/Source/Runtime/Engine/Classes/Sound/SoundWave.h
@@ -129,6 +129,108 @@ struct FStreamedAudioPlatformData
 
 };
 
+// Struct used to retrieve the spectral magnitude of an audio signal at a given frequency to BP
+USTRUCT(BlueprintType)
+struct FSoundWaveSpectralData
+{
+	GENERATED_USTRUCT_BODY()
+
+	// The frequency hz of the spectrum value
+	UPROPERTY(EditAnywhere, BlueprintReadWrite, Category = "SpectralData|Foo")
+	float FrequencyHz;
+
+	// The magnitude of the spectrum at this frequency
+	UPROPERTY(EditAnywhere, BlueprintReadWrite, Category = "SpectralData|Bar")
+	float Magnitude;
+
+	// The normalized magnitude of the spectrum at this frequency
+	UPROPERTY(EditAnywhere, BlueprintReadWrite, Category = "SpectralData|Bar")
+	float NormalizedMagnitude;
+
+	FSoundWaveSpectralData()
+		: FrequencyHz(0.0f)
+		, Magnitude(0.0f)
+		, NormalizedMagnitude(0.0f)
+	{}
+};
+
+// Sort predicate for sorting spectral data by frequency (lowest first)
+struct FCompareSpectralDataByFrequencyHz
+{
+	FORCEINLINE bool operator()(const FSoundWaveSpectralData& A, const FSoundWaveSpectralData& B) const
+	{
+		return A.FrequencyHz < B.FrequencyHz;
+	}
+};
+
+
+// Struct used to store spectral data with time-stamps
+USTRUCT()
+struct FSoundWaveSpectralDataEntry
+{
+	GENERATED_USTRUCT_BODY()
+
+	// The magnitude of the spectrum at this frequency
+	UPROPERTY()
+	float Magnitude;
+
+	// The normalized magnitude of the spectrum at this frequency
+	UPROPERTY()
+	float NormalizedMagnitude;
+};
+
+
+// Struct used to store spectral data with time-stamps
+USTRUCT()
+struct FSoundWaveSpectralTimeData
+{
+	GENERATED_USTRUCT_BODY()
+
+	// The spectral data at the given time. The array indices correspond to the frequencies set to analyze.
+	UPROPERTY()
+	TArray<FSoundWaveSpectralDataEntry> Data;
+
+	// The timestamp associated with this spectral data
+	UPROPERTY()
+	float TimeSec;
+
+	FSoundWaveSpectralTimeData()
+		: TimeSec(0.0f)
+	{}
+};
+
+// Struct used to store time-stamped envelope data
+USTRUCT()
+struct FSoundWaveEnvelopeTimeData
+{
+	GENERATED_USTRUCT_BODY()
+
+	// The normalized linear amplitude of the audio
+	UPROPERTY()
+	float Amplitude;
+
+	// The timestamp of the audio
+	UPROPERTY()
+	float TimeSec;
+
+	FSoundWaveEnvelopeTimeData()
+		: Amplitude(0.0f)
+		, TimeSec(0.0f)
+	{}
+};
+
+// The FFT size (in audio frames) to use for baked FFT analysis
+UENUM(BlueprintType)
+enum class ESoundWaveFFTSize : uint8
+{
+	VerySmall_64,
+	Small_256,
+	Medium_512,
+	Large_1024,
+	VeryLarge_2048,
+};
+
+
 UCLASS(hidecategories=Object, editinlinenew, BlueprintType)
 class ENGINE_API USoundWave : public USoundBase
 {
@@ -196,10 +298,60 @@ class ENGINE_API USoundWave : public USoundBase
 	uint8 bIsAmbisonics : 1;
 
 	/** Whether this SoundWave was decompressed from OGG. */
-	uint8 bDecompressedFromOgg:1;
+	uint8 bDecompressedFromOgg : 1;
+
+#if WITH_EDITORONLY_DATA
+	/** Whether or not to enable cook-time baked FFT analysis. */
+	UPROPERTY(EditAnywhere, Category = "Analysis|FFT")
+	uint8 bEnableBakedFFTAnalysis : 1;
+
+	/** Whether or not to enable cook-time amplitude envelope analysis. */
+	UPROPERTY(EditAnywhere, Category = "Analysis|Envelope")
+	uint8 bEnableAmplitudeEnvelopeAnalysis : 1;
+
+	/** The FFT window size to use for fft analysis. */
+	UPROPERTY(EditAnywhere, Category = "Analysis|FFT", meta = (EditCondition = "bEnableBakedFFTAnalysis"))
+	ESoundWaveFFTSize FFTSize;
+
+	/** How many audio frames analyze at a time. */
+	UPROPERTY(EditAnywhere, Category = "Analysis|FFT", meta = (EditCondition = "bEnableAmplitudeEnvelopeAnalysis", ClampMin = "512", UIMin = "512"))
+	int32 FFTAnalysisFrameSize;
+
+	/** How many audio frames to average a new envelope value. Larger values use less memory for audio envelope data but will result in lower envelope accuracy. */
+	UPROPERTY(EditAnywhere, Category = "Analysis|Envelope", meta = (EditCondition = "bEnableAmplitudeEnvelopeAnalysis", ClampMin = "512", UIMin = "512"))
+	int32 EnvelopeFollowerFrameSize;
+
+	/** How quickly the envelope analyzer responds to increasing amplitudes. */
+	UPROPERTY(EditAnywhere, Category = "Analysis|Envelope", meta = (EditCondition = "bEnableAmplitudeEnvelopeAnalysis", ClampMin = "0", UIMin = "0"))
+	int32 EnvelopeFollowerAttackTime;
+
+	/** How quickly the envelope analyzer responds to decreasing amplitudes. */
+	UPROPERTY(EditAnywhere, Category = "Analysis|Envelope", meta = (EditCondition = "bEnableAmplitudeEnvelopeAnalysis", ClampMin = "0", UIMin = "0"))
+	int32 EnvelopeFollowerReleaseTime;
+#endif
+
+	/** The frequencies (in hz) to analyze when doing baked FFT analysis. */
+	UPROPERTY(EditAnywhere, Category = "Analysis|FFT", meta = (EditCondition = "bEnableBakedFFTAnalysis"))
+	TArray<float> FrequenciesToAnalyze;
+
+	/** The cooked spectral time data. */
+	UPROPERTY()
+	TArray<FSoundWaveSpectralTimeData> CookedSpectralTimeData;
+
+	/** The cooked cooked envelope data. */
+	UPROPERTY()
+	TArray<FSoundWaveEnvelopeTimeData> CookedEnvelopeTimeData;
+
+	/** Helper function to get interpolated cooked FFT data for a given time value. */
+	bool GetInterpolatedCookedFFTDataForTime(float InTime, uint32& InOutLastIndex, TArray<FSoundWaveSpectralData>& OutData, bool bLoop);
+	bool GetInterpolatedCookedEnvelopeDataForTime(float InTime, uint32& InOutLastIndex, float& OutAmplitude, bool bLoop);
 
 private:
 
+	/** Helper functions to search analysis data. Takes starting index to start query. Returns which data index the result was found at. Returns INDEX_NONE if not found. */
+	uint32 GetInterpolatedCookedFFTDataForTimeInternal(float InTime, uint32 StartingIndex, TArray<FSoundWaveSpectralData>& OutData, bool bLoop);
+	uint32 GetInterpolatedCookedEnvelopeDataForTimeInternal(float InTime, uint32 StartingIndex, float& OutAmplitude, bool bLoop);
+
 #if !WITH_EDITOR
 	// This is set to false on initialization, then set to true on non-editor platforms when we cache appropriate sample rate.
 	uint8 bCachedSampleRateFromPlatformSettings : 1;
@@ -387,6 +539,7 @@ public:
 	virtual float GetDuration() override;
 	virtual float GetSubtitlePriority() const override;
 	virtual bool IsAllowedVirtual() const override;
+	virtual bool GetSoundWavesWithCookedAnalysisData(TArray<USoundWave*>& OutSoundWaves) override;
 	//~ End USoundBase Interface.
 
 	// Called  when the procedural sound wave begins on the render thread. Only used in the audio mixer and when bProcedural is true.
@@ -485,9 +638,19 @@ public:
 
 	virtual bool HasCompressedData(FName Format, ITargetPlatform* TargetPlatform = GetRunningPlatform()) const;
 
+#if WITH_EDITOR
+	/** Utility which returns imported PCM data and the parsed header for the file. Returns true if there was data, false if there wasn't. */
+	bool GetImportedSoundWaveData(TArray<uint8>& OutRawPCMData, uint32& OutSampleRate, uint16& OutNumChannels);
+#endif
+
 private:
 	FName GetPlatformSpecificFormat(FName Format, const FPlatformAudioCookOverrides* CompressionOverrides);
 
+#if WITH_EDITOR
+	void BakeFFTAnalysis();
+	void BakeEnvelopeAnalysis();
+#endif
+
 public:
 	virtual void BeginGetCompressedData(FName Format, const FPlatformAudioCookOverrides* CompressionOverrides);
 
@@ -626,6 +789,7 @@ public:
 	{
 		return (ESoundWavePrecacheState)PrecacheState.GetValue();
 	}
+
 };
 
 
diff --git a/Engine/Source/Runtime/Engine/Engine.Build.cs b/Engine/Source/Runtime/Engine/Engine.Build.cs
--- a/Engine/Source/Runtime/Engine/Engine.Build.cs
+++ b/Engine/Source/Runtime/Engine/Engine.Build.cs
@@ -7,7 +7,9 @@ public class Engine : ModuleRules
 {
 	public Engine(ReadOnlyTargetRules Target) : base(Target)
 	{
-		PrivatePCHHeaderFile = "Private/EnginePrivatePCH.h";
+        OptimizeCode = CodeOptimization.Never;
+
+        PrivatePCHHeaderFile = "Private/EnginePrivatePCH.h";
 
 		SharedPCHHeaderFile = "Public/EngineSharedPCH.h";
 
@@ -96,6 +98,7 @@ public class Engine : ModuleRules
 				"CinematicCamera",
 				"Analytics",
 				"AnalyticsET",
+                "AudioMixer",
                 //"CrunchCompression"
             }
 		);
@@ -171,6 +174,7 @@ public class Engine : ModuleRules
 		CircularlyReferencedDependentModules.Add("UMG");
 		CircularlyReferencedDependentModules.Add("MaterialShaderQualitySettings");
 		CircularlyReferencedDependentModules.Add("CinematicCamera");
+        CircularlyReferencedDependentModules.Add("AudioMixer");
 
         // The AnimGraphRuntime module is not needed by Engine proper, but it is loaded in LaunchEngineLoop.cpp,
         // and needs to be listed in an always-included module in order to be compiled into standalone games
diff --git a/Engine/Source/Runtime/Engine/Private/ActiveSound.cpp b/Engine/Source/Runtime/Engine/Private/ActiveSound.cpp
--- a/Engine/Source/Runtime/Engine/Private/ActiveSound.cpp
+++ b/Engine/Source/Runtime/Engine/Private/ActiveSound.cpp
@@ -56,6 +56,7 @@ FActiveSound::FActiveSound()
 	, bUpdatePlayPercentage(false)
 	, bUpdateSingleEnvelopeValue(false)
 	, bUpdateMultiEnvelopeValue(false)
+	, bUpdatePlaybackTime(false)
 	, bIsPlayingAudio(false)
 	, bIsStopping(false)
 	, UserIndex(0)
@@ -543,39 +544,73 @@ void FActiveSound::UpdateWaveInstances( TArray<FWaveInstance*> &InWaveInstances,
 		}
 
 		// Check to see if we need to broadcast the envelope value of sounds playing with this active sound
-		if (bUpdateMultiEnvelopeValue && AudioComponentID > 0)
+		if (AudioComponentID > 0)
 		{
-			int32 NumWaveInstances = ThisSoundsWaveInstances.Num();
-
-			// Add up the envelope value for every wave instance so we get a sum of the envelope value for all sources.
-			float EnvelopeValueSum = 0.0f;
-			float MaxEnvelopeValue = 0.0f;
-			for (FWaveInstance* WaveInstance : ThisSoundsWaveInstances)
+			if (bUpdateMultiEnvelopeValue)
 			{
-				const float WaveInstanceEnvelopeValue = WaveInstance->GetEnvelopeValue();
-				EnvelopeValueSum += WaveInstanceEnvelopeValue;
-				MaxEnvelopeValue = FMath::Max(WaveInstanceEnvelopeValue, MaxEnvelopeValue);
-			}
+				int32 NumWaveInstances = ThisSoundsWaveInstances.Num();
 
-			// Now divide by the number of instances to get the average
-			float AverageEnvelopeValue = EnvelopeValueSum / NumWaveInstances;
-			uint64 AudioComponentIDCopy = AudioComponentID;
-			FAudioThread::RunCommandOnGameThread([AudioComponentIDCopy, AverageEnvelopeValue, MaxEnvelopeValue, NumWaveInstances]()
-			{
-				if (UAudioComponent* AudioComponent = UAudioComponent::GetAudioComponentFromID(AudioComponentIDCopy))
+				// Add up the envelope value for every wave instance so we get a sum of the envelope value for all sources.
+				float EnvelopeValueSum = 0.0f;
+				float MaxEnvelopeValue = 0.0f;
+				for (FWaveInstance* WaveInstance : ThisSoundsWaveInstances)
 				{
-					if (AudioComponent->OnAudioMultiEnvelopeValue.IsBound())
+					const float WaveInstanceEnvelopeValue = WaveInstance->GetEnvelopeValue();
+					EnvelopeValueSum += WaveInstanceEnvelopeValue;
+					MaxEnvelopeValue = FMath::Max(WaveInstanceEnvelopeValue, MaxEnvelopeValue);
+				}
+
+				// Now divide by the number of instances to get the average
+				float AverageEnvelopeValue = EnvelopeValueSum / NumWaveInstances;
+				uint64 AudioComponentIDCopy = AudioComponentID;
+				FAudioThread::RunCommandOnGameThread([AudioComponentIDCopy, AverageEnvelopeValue, MaxEnvelopeValue, NumWaveInstances]()
+				{
+					if (UAudioComponent* AudioComponent = UAudioComponent::GetAudioComponentFromID(AudioComponentIDCopy))
 					{
-						AudioComponent->OnAudioMultiEnvelopeValue.Broadcast(AverageEnvelopeValue, MaxEnvelopeValue, NumWaveInstances);
+						if (AudioComponent->OnAudioMultiEnvelopeValue.IsBound())
+						{
+							AudioComponent->OnAudioMultiEnvelopeValue.Broadcast(AverageEnvelopeValue, MaxEnvelopeValue, NumWaveInstances);
+						}
+
+						if (AudioComponent->OnAudioMultiEnvelopeValueNative.IsBound())
+						{
+							AudioComponent->OnAudioMultiEnvelopeValueNative.Broadcast(AudioComponent, AverageEnvelopeValue, MaxEnvelopeValue, NumWaveInstances);
+						}
 					}
+				});
+			}
+
+			if (bUpdatePlaybackTime)
+			{
+				TMap<uint32, float> WaveInstancePlaybackTimes;
+
+				// Update each of the wave instances playback time based on delta time and the wave instances pitch value
+				for (FWaveInstance* WaveInstance : ThisSoundsWaveInstances)
+				{
+					WaveInstance->PlaybackTime += DeltaTime * WaveInstance->Pitch;
 
-					if (AudioComponent->OnAudioMultiEnvelopeValueNative.IsBound())
+					// For looping sounds, we need to check the wrapping condition
+					if (WaveInstance->LoopingMode != LOOP_Never)
 					{
-						AudioComponent->OnAudioMultiEnvelopeValueNative.Broadcast(AudioComponent, AverageEnvelopeValue, MaxEnvelopeValue, NumWaveInstances);
+						float Duration = WaveInstance->WaveData->Duration;
+						if (WaveInstance->PlaybackTime > Duration)
+						{
+							WaveInstance->PlaybackTime = 0.0f;
+						}
 					}
+					WaveInstancePlaybackTimes.Add(WaveInstance->WaveData->GetUniqueID(), WaveInstance->PlaybackTime);
 				}
-			});
+				uint64 AudioComponentIDCopy = AudioComponentID;
+				FAudioThread::RunCommandOnGameThread([AudioComponentIDCopy, WaveInstancePlaybackTimes]()
+				{
+					if (UAudioComponent* AudioComponent = UAudioComponent::GetAudioComponentFromID(AudioComponentIDCopy))
+					{
+						AudioComponent->SetPlaybackTimes(WaveInstancePlaybackTimes);
+					}
+				});
+			}
 		}
+
 	}
 
 	InWaveInstances.Append(ThisSoundsWaveInstances);
diff --git a/Engine/Source/Runtime/Engine/Private/Components/AudioComponent.cpp b/Engine/Source/Runtime/Engine/Private/Components/AudioComponent.cpp
--- a/Engine/Source/Runtime/Engine/Private/Components/AudioComponent.cpp
+++ b/Engine/Source/Runtime/Engine/Private/Components/AudioComponent.cpp
@@ -13,6 +13,14 @@
 
 DECLARE_CYCLE_STAT(TEXT("AudioComponent Play"), STAT_AudioComp_Play, STATGROUP_Audio);
 
+static float BakedAnalysisTimeShiftCVar = 0.0f;
+FAutoConsoleVariableRef CVarBackedAnalysisTimeShift(
+	TEXT("au.AnalysisTimeShift"),
+	BakedAnalysisTimeShiftCVar,
+	TEXT("Shifts the timeline for baked analysis playback.\n")
+	TEXT("Value: The time in seconds to shift the timeline."),
+	ECVF_Default);
+
 
 /*-----------------------------------------------------------------------------
 UAudioComponent implementation.
@@ -408,6 +416,17 @@ void UAudioComponent::PlayInternal(const float StartTime, const float FadeInDura
 			NewActiveSound.bUpdatePlayPercentage = OnAudioPlaybackPercentNative.IsBound() || OnAudioPlaybackPercent.IsBound();
 			NewActiveSound.bUpdateSingleEnvelopeValue = OnAudioSingleEnvelopeValue.IsBound() || OnAudioSingleEnvelopeValueNative.IsBound();
 			NewActiveSound.bUpdateMultiEnvelopeValue = OnAudioMultiEnvelopeValue.IsBound() || OnAudioMultiEnvelopeValueNative.IsBound();
+			
+			// Setup audio component cooked analysis data playback data set
+			TArray<USoundWave*> SoundWavesWithCookedData;
+			NewActiveSound.bUpdatePlaybackTime = Sound->GetSoundWavesWithCookedAnalysisData(SoundWavesWithCookedData);
+
+			// Reset the audio component's soundwave playback times
+			SoundWavePlaybackTimes.Reset();
+			for (USoundWave* SoundWave : SoundWavesWithCookedData)
+			{
+				SoundWavePlaybackTimes.Add(SoundWave->GetUniqueID(), FSoundWavePlaybackTimeData(SoundWave));
+			}
 
 			NewActiveSound.MaxDistance = MaxDistance;
 
@@ -1110,4 +1129,167 @@ void UAudioComponent::SetLowPassFilterFrequency(float InLowPassFilterFrequency)
 	}
 }
 
+bool UAudioComponent::HasCookedFFTData() const
+{
+	if (Sound)
+	{
+		USoundWave* SoundWave = Cast<USoundWave>(Sound);
+		if (SoundWave)
+		{
+			return SoundWave->CookedSpectralTimeData.Num() > 0;
+		}
+	}
+	return false;
+}
+
+bool UAudioComponent::HasCookedAmplitudeEnvelopeData() const
+{
+	if (Sound)
+	{
+		USoundWave* SoundWave = Cast<USoundWave>(Sound);
+		if (SoundWave)
+		{
+			return SoundWave->CookedEnvelopeTimeData.Num() > 0;
+		}
+	}
+	return false;
+}
+
+void UAudioComponent::SetPlaybackTimes(const TMap<uint32, float>& InSoundWavePlaybackTimes)
+{
+	// Reset the playback times for everything in case the wave instance stops and is not updated
+	for (auto& Elem : SoundWavePlaybackTimes)
+	{
+		Elem.Value.PlaybackTime = 0.0f;
+	}
+
+	for (auto& Elem : InSoundWavePlaybackTimes)
+	{
+		uint32 ObjectId = Elem.Key;
+		FSoundWavePlaybackTimeData* PlaybackTimeData = SoundWavePlaybackTimes.Find(ObjectId);
+		if (PlaybackTimeData)
+		{
+			PlaybackTimeData->PlaybackTime = FMath::Max(Elem.Value - BakedAnalysisTimeShiftCVar, 0.0f);
+		}
+	}
+}
+
+bool UAudioComponent::GetCookedFFTData(const TArray<float>& FrequenciesToGet, TArray<FSoundWaveSpectralData>& OutSoundWaveSpectralData)
+{
+	bool bHadData = false;
+	if (IsPlaying() && SoundWavePlaybackTimes.Num() > 0 && FrequenciesToGet.Num() > 0)
+	{
+		OutSoundWaveSpectralData.Reset();
+		for (float Frequency : FrequenciesToGet)
+		{
+			FSoundWaveSpectralData NewEntry;
+			NewEntry.FrequencyHz = Frequency;
+			OutSoundWaveSpectralData.Add(NewEntry);
+		}
+
+		// Sort by frequency (lowest frequency first).
+		OutSoundWaveSpectralData.Sort(FCompareSpectralDataByFrequencyHz());
+
+		int32 NumEntriesAdded = 0;
+		for (auto& Entry : SoundWavePlaybackTimes)
+		{
+			if (Entry.Value.PlaybackTime > 0.0f && Entry.Value.SoundWave->CookedSpectralTimeData.Num() > 0)
+			{
+				static TArray<FSoundWaveSpectralData> CookedSpectralData;
+				CookedSpectralData.Reset();
+
+				// Find the point in the spectral data that corresponds to the time
+				Entry.Value.SoundWave->GetInterpolatedCookedFFTDataForTime(Entry.Value.PlaybackTime, Entry.Value.LastFFTCookedIndex, CookedSpectralData, Sound->IsLooping());
+
+				if (CookedSpectralData.Num() > 0)
+				{
+					// Find the interpolated values given the frequencies we want to get
+					for (FSoundWaveSpectralData& OutSpectralData : OutSoundWaveSpectralData)
+					{
+						// Check min edge case: we're requesting cooked FFT data lower than what we have cooked
+						if (OutSpectralData.FrequencyHz < CookedSpectralData[0].FrequencyHz)
+						{
+							// Just mix in the lowest value we have cooked
+							OutSpectralData.Magnitude += CookedSpectralData[0].Magnitude;
+							OutSpectralData.NormalizedMagnitude += CookedSpectralData[0].NormalizedMagnitude;
+						}
+						// Check max edge case: we're requesting cooked FFT data at a higher frequency than what we have cooked
+						else if (OutSpectralData.FrequencyHz >= CookedSpectralData.Last().FrequencyHz)
+						{
+							// Just mix in the highest value we have cooked
+							OutSpectralData.Magnitude += CookedSpectralData.Last().Magnitude;
+							OutSpectralData.NormalizedMagnitude += CookedSpectralData.Last().NormalizedMagnitude;
+						}
+						// We need to find the 2 closest cooked results and interpolate those
+						else
+						{
+							for (int32 SpectralDataIndex = 0; SpectralDataIndex < CookedSpectralData.Num() - 1; ++SpectralDataIndex)
+							{
+								const FSoundWaveSpectralData& CurrentSpectralData = CookedSpectralData[SpectralDataIndex];
+								const FSoundWaveSpectralData& NextSpectralData = CookedSpectralData[SpectralDataIndex + 1];
+								if (OutSpectralData.FrequencyHz >= CurrentSpectralData.FrequencyHz && OutSpectralData.FrequencyHz < NextSpectralData.FrequencyHz)
+								{
+									float Alpha = (OutSpectralData.FrequencyHz - CurrentSpectralData.FrequencyHz) / (NextSpectralData.FrequencyHz - CurrentSpectralData.FrequencyHz);
+									OutSpectralData.Magnitude += FMath::Lerp(CurrentSpectralData.Magnitude, NextSpectralData.Magnitude, Alpha);
+									OutSpectralData.NormalizedMagnitude += FMath::Lerp(CurrentSpectralData.NormalizedMagnitude, NextSpectralData.NormalizedMagnitude, Alpha);
+
+									break;
+								}
+							}
+						}
+					}
+					
+					++NumEntriesAdded;
+					bHadData = true;
+				}
+			}
+		}
+
+		// Divide by the number of entries we added (i.e. we are averaging together multiple cooked FFT data in the case of multiple sound waves playing with cooked data)
+		if (NumEntriesAdded > 1)
+		{
+			for (FSoundWaveSpectralData& OutSpectralData : OutSoundWaveSpectralData)
+			{
+				OutSpectralData.Magnitude /= NumEntriesAdded;
+				OutSpectralData.NormalizedMagnitude /= NumEntriesAdded;
+			}
+		}
+	}
+
+	return bHadData;
+}
+
+bool UAudioComponent::GetCookedEnvelopeData(float& OutEnvelopeData)
+{
+	bool bHadData = false;
+	if (IsPlaying() && SoundWavePlaybackTimes.Num() > 0)
+	{
+		static TArray<FSoundWaveEnvelopeTimeData> CookedEnvelopeData;
+		int32 NumEntriesAdded = 0;
+		OutEnvelopeData = 0.0f;
+		for (auto& Entry : SoundWavePlaybackTimes)
+		{
+			if (Entry.Value.SoundWave->CookedEnvelopeTimeData.Num() > 0 && Entry.Value.PlaybackTime > 0.0f)
+			{
+				CookedEnvelopeData.Reset();
 
+				// Find the point in the spectral data that corresponds to the time
+				float SoundWaveAmplitude = 0.0f;
+				if (Entry.Value.SoundWave->GetInterpolatedCookedEnvelopeDataForTime(Entry.Value.PlaybackTime, Entry.Value.LastEnvelopeCookedIndex, SoundWaveAmplitude, Sound->IsLooping()))
+				{
+					OutEnvelopeData += SoundWaveAmplitude;
+					++NumEntriesAdded;
+					bHadData = true;
+				}
+			}
+		}
+
+		// Divide by number of entries we added... get average amplitude envelope
+		if (bHadData)
+		{
+			OutEnvelopeData /= NumEntriesAdded;
+		}
+	}
+
+	return bHadData;
+}
diff --git a/Engine/Source/Runtime/Engine/Private/SoundBase.cpp b/Engine/Source/Runtime/Engine/Private/SoundBase.cpp
--- a/Engine/Source/Runtime/Engine/Private/SoundBase.cpp
+++ b/Engine/Source/Runtime/Engine/Private/SoundBase.cpp
@@ -166,6 +166,11 @@ uint32 USoundBase::GetSoundConcurrencyObjectID() const
 	return 0;
 }
 
+bool USoundBase::GetSoundWavesWithCookedAnalysisData(TArray<USoundWave*>& OutSoundWaves)
+{
+	return false;
+}
+
 void USoundBase::PostLoad()
 {
 	Super::PostLoad();
diff --git a/Engine/Source/Runtime/Engine/Private/SoundCue.cpp b/Engine/Source/Runtime/Engine/Private/SoundCue.cpp
--- a/Engine/Source/Runtime/Engine/Private/SoundCue.cpp
+++ b/Engine/Source/Runtime/Engine/Private/SoundCue.cpp
@@ -487,6 +487,24 @@ float USoundCue::GetSubtitlePriority() const
 	return SubtitlePriority;
 }
 
+bool USoundCue::GetSoundWavesWithCookedAnalysisData(TArray<USoundWave*>& OutSoundWaves)
+{
+	// Check this sound cue's wave players to see if any of their soundwaves have cooked analysis data
+	TArray<USoundNodeWavePlayer*> WavePlayers;
+	RecursiveFindNode<USoundNodeWavePlayer>(FirstNode, WavePlayers);
+
+	bool bHasAnalysisData = false;
+	for (USoundNodeWavePlayer* Player : WavePlayers)
+	{
+		USoundWave* SoundWave = Player->GetSoundWave();
+		if (SoundWave->GetSoundWavesWithCookedAnalysisData(OutSoundWaves))
+		{
+			bHasAnalysisData = true;
+		}
+	}
+	return bHasAnalysisData;
+}
+
 #if WITH_EDITOR
 UEdGraph* USoundCue::GetGraph()
 { 
diff --git a/Engine/Source/Runtime/Engine/Private/SoundWave.cpp b/Engine/Source/Runtime/Engine/Private/SoundWave.cpp
--- a/Engine/Source/Runtime/Engine/Private/SoundWave.cpp
+++ b/Engine/Source/Runtime/Engine/Private/SoundWave.cpp
@@ -20,6 +20,9 @@
 #include "ProfilingDebugging/CookStats.h"
 #include "HAL/LowLevelMemTracker.h"
 #include "AudioCompressionSettingsUtils.h"
+#include "DSP/SpectrumAnalyzer.h"
+#include "DSP/EnvelopeFollower.h"
+#include "DSP/BufferVectorOperations.h"
 
 #if ENABLE_COOK_STATS
 namespace SoundWaveCookStats
@@ -111,6 +114,18 @@ USoundWave::USoundWave(const FObjectInitializer& ObjectInitializer)
 	RawPCMDataSize = 0;
 	SetPrecacheState(ESoundWavePrecacheState::NotStarted);
 
+#if WITH_EDITORONLY_DATA
+	FFTSize = ESoundWaveFFTSize::Medium_512;
+	FrequenciesToAnalyze.Add(100.0f);
+	FrequenciesToAnalyze.Add(500.0f);
+	FrequenciesToAnalyze.Add(1000.0f);
+	FrequenciesToAnalyze.Add(5000.0f);
+	FFTAnalysisFrameSize = 1024;
+	EnvelopeFollowerFrameSize = 1024;
+	EnvelopeFollowerAttackTime = 10;
+	EnvelopeFollowerReleaseTime = 100;
+#endif
+
 #if !WITH_EDITOR
 	bCachedSampleRateFromPlatformSettings = false;
 	bSampleRateManuallyReset = false;
@@ -436,6 +451,42 @@ const FPlatformAudioCookOverrides* USoundWave::GetPlatformCompressionOverridesFo
 	return FPlatformCompressionUtilities::GetCookOverridesForCurrentPlatform();
 }
 
+#if WITH_EDITOR
+bool USoundWave::GetImportedSoundWaveData(TArray<uint8>& OutRawPCMData, uint32& OutSampleRate, uint16& OutNumChannels)
+{
+	// Can only get sound wave data if there is bulk data and if we don't have some weird munging of multi-channel files (e.g. mono stereo only)
+	if (RawData.GetBulkDataSize() > 0)
+	{
+		FWaveModInfo WaveInfo;
+
+		const uint8* RawWaveData = (const uint8*)RawData.LockReadOnly();
+		int32 RawDataSize = RawData.GetBulkDataSize();
+
+		// parse the wave data
+		if (!WaveInfo.ReadWaveHeader(RawWaveData, RawDataSize, 0))
+		{
+			UE_LOG(LogAudio, Warning, TEXT("Only mono or stereo 16 bit waves allowed: %s."), *GetFullName());
+			RawData.Unlock();
+			return false;
+		}
+
+		// Copy the raw PCM data and the header info that was parsed
+		OutRawPCMData.Reset();
+		OutRawPCMData.AddUninitialized(WaveInfo.SampleDataSize);
+		FMemory::Memcpy(OutRawPCMData.GetData(), WaveInfo.SampleDataStart, WaveInfo.SampleDataSize);
+
+		OutSampleRate = *WaveInfo.pSamplesPerSec;
+		OutNumChannels = *WaveInfo.pChannels;
+
+		RawData.Unlock();
+		return true;
+	}
+
+	UE_LOG(LogAudio, Warning, TEXT("Failed to get imported raw data for sound wave '%s'"), *GetFullName());
+	return false;
+}
+#endif
+
 FName USoundWave::GetPlatformSpecificFormat(FName Format, const FPlatformAudioCookOverrides* CompressionOverrides)
 {
 	// Platforms that require compression overrides get concatenated formats.
@@ -732,6 +783,239 @@ float USoundWave::GetSampleRateForTargetPlatform(const ITargetPlatform* TargetPl
 	}
 }
 
+static bool AnyFFTAnalysisPropertiesChanged(const FName& PropertyName)
+{
+	// List of properties which cause re-cooking to get triggered
+	static FName EnableFFTAnalysisFName					= GET_MEMBER_NAME_CHECKED(USoundWave, bEnableBakedFFTAnalysis);
+	static FName FFTSizeFName							= GET_MEMBER_NAME_CHECKED(USoundWave, FFTSize);
+	static FName FFTAnalysisFrameSizeFName				= GET_MEMBER_NAME_CHECKED(USoundWave, FFTAnalysisFrameSize);
+	static FName FrequenciesToAnalyzeFName				= GET_MEMBER_NAME_CHECKED(USoundWave, FrequenciesToAnalyze);
+
+	return	PropertyName == EnableFFTAnalysisFName ||
+			PropertyName == FFTSizeFName ||
+			PropertyName == FFTAnalysisFrameSizeFName ||
+			PropertyName == FrequenciesToAnalyzeFName;
+}
+
+static bool AnyEnvelopeAnalysisPropertiesChanged(const FName& PropertyName)
+{
+	static FName EnableAmplitudeEnvelopeAnalysisFName = GET_MEMBER_NAME_CHECKED(USoundWave, bEnableAmplitudeEnvelopeAnalysis);
+	static FName EnvelopeFollowerFrameSizeFName = GET_MEMBER_NAME_CHECKED(USoundWave, EnvelopeFollowerFrameSize);
+	static FName EnvelopeFollowerAttackTimeFName = GET_MEMBER_NAME_CHECKED(USoundWave, EnvelopeFollowerAttackTime);
+	static FName EnvelopeFollowerReleaseTimeFName = GET_MEMBER_NAME_CHECKED(USoundWave, EnvelopeFollowerReleaseTime);
+
+	return	PropertyName == EnableAmplitudeEnvelopeAnalysisFName ||
+			PropertyName == EnvelopeFollowerFrameSizeFName ||
+			PropertyName == EnvelopeFollowerAttackTimeFName ||
+			PropertyName == EnvelopeFollowerReleaseTimeFName;
+
+}
+
+void USoundWave::BakeFFTAnalysis()
+{
+	// Clear any existing spectral data regardless of if it's enabled. If this was enabled and is now toggled, this will clear previous data.
+	CookedSpectralTimeData.Reset();
+
+	// Perform analysis if enabled on the sound wave
+	if (bEnableBakedFFTAnalysis)
+	{
+		// If there are no frequencies to analyze, we can't do the analysis
+		if (!FrequenciesToAnalyze.Num())
+		{
+			UE_LOG(LogAudio, Warning, TEXT("Soundwave '%s' had baked FFT analysis enabled without specifying any frequencies to analyze."), *GetFullName());
+			return;
+		}
+
+		if (ChannelSizes.Num() > 0)
+		{
+			UE_LOG(LogAudio, Warning, TEXT("Soundwave '%s' has multi-channel audio (channels greater than 2). Baking FFT analysis is not currently supported for this yet."), *GetFullName());
+			return;
+		}
+
+		// Retrieve the raw imported data
+		TArray<uint8> RawImportedWaveData;
+		uint32 RawDataSampleRate = 0;
+		uint16 RawDataNumChannels = 0;
+
+		if (!GetImportedSoundWaveData(RawImportedWaveData, RawDataSampleRate, RawDataNumChannels))
+		{
+			return;
+		}
+
+		if (RawDataSampleRate == 0 || RawDataNumChannels == 0)
+		{
+			UE_LOG(LogAudio, Error, TEXT("Failed to parse the raw imported data for '%s' for baked FFT analysis."), *GetFullName());
+			return;
+		}
+
+		const uint32 NumFrames = (RawImportedWaveData.Num() / sizeof(int16)) / RawDataNumChannels;
+		int16* InputData = (int16*)RawImportedWaveData.GetData();
+
+		Audio::FSpectrumAnalyzerSettings SpectrumAnalyzerSettings;
+		switch (FFTSize)
+		{
+		case ESoundWaveFFTSize::VerySmall_64:
+			SpectrumAnalyzerSettings.FFTSize = Audio::FSpectrumAnalyzerSettings::EFFTSize::Min_64;
+			break;
+
+		case ESoundWaveFFTSize::Small_256:
+			SpectrumAnalyzerSettings.FFTSize = Audio::FSpectrumAnalyzerSettings::EFFTSize::Small_256;
+			break;
+
+		default:
+		case ESoundWaveFFTSize::Medium_512:
+			SpectrumAnalyzerSettings.FFTSize = Audio::FSpectrumAnalyzerSettings::EFFTSize::Medium_512;
+			break;
+
+		case ESoundWaveFFTSize::Large_1024:
+			SpectrumAnalyzerSettings.FFTSize = Audio::FSpectrumAnalyzerSettings::EFFTSize::Large_1024;
+			break;
+
+		case ESoundWaveFFTSize::VeryLarge_2048:
+			SpectrumAnalyzerSettings.FFTSize = Audio::FSpectrumAnalyzerSettings::EFFTSize::VeryLarge_2048;
+			break;
+
+		}
+
+		// Build a new spectrum analyzer
+		Audio::FSpectrumAnalyzer SpectrumAnalyzer(SpectrumAnalyzerSettings, (float)SampleRate);
+
+		// The audio data block to use to submit audio data to the spectrum analyzer
+		Audio::AlignedFloatBuffer AnalysisData;
+		check(FFTAnalysisFrameSize > 256);
+		AnalysisData.Reserve(FFTAnalysisFrameSize);
+
+		float MaximumMagnitude = 0.0f;
+		for (uint32 FrameIndex = 0; FrameIndex < NumFrames; ++FrameIndex)
+		{
+			// Get the averaged sample value of all the channels
+			float SampleValue = 0.0f;
+			for (uint16 ChannelIndex = 0; ChannelIndex < RawDataNumChannels; ++ChannelIndex)
+			{
+				SampleValue += (float)InputData[FrameIndex * RawDataNumChannels] / 32767.0f;
+			}
+			SampleValue /= RawDataNumChannels;
+
+			// Accumate the samples in the scratch buffer
+			AnalysisData.Add(SampleValue);
+
+			// Until we reached the frame size
+			if (AnalysisData.Num() == FFTAnalysisFrameSize)
+			{
+				SpectrumAnalyzer.PushAudio(AnalysisData.GetData(), AnalysisData.Num());
+
+				// Block while the analyzer does the analysis
+				while (SpectrumAnalyzer.PerformAnalysisIfPossible()) {}
+
+				FSoundWaveSpectralTimeData NewData;
+
+				// Don't need to lock here since we're doing this sync, but it's here as that's the expected pattern for the Spectrum analyzer
+				SpectrumAnalyzer.LockOutputBuffer();
+
+				// Get the magntiudes for the specified frequencies
+				for (float Frequency : FrequenciesToAnalyze)
+				{
+					FSoundWaveSpectralDataEntry DataEntry;
+					DataEntry.Magnitude = SpectrumAnalyzer.GetMagnitudeForFrequency(Frequency);
+
+					// Track the max magnitude so we can later set normalized magnitudes
+					if (DataEntry.Magnitude > MaximumMagnitude)
+					{
+						MaximumMagnitude = DataEntry.Magnitude;
+					}
+
+					NewData.Data.Add(DataEntry);
+				}
+
+				SpectrumAnalyzer.UnlockOutputBuffer();
+
+				// The time stamp is derived from the frame index and sample rate
+				NewData.TimeSec = FMath::Max((float)(FrameIndex - FFTAnalysisFrameSize + 1) / RawDataSampleRate, 0.0f);
+
+				CookedSpectralTimeData.Add(NewData);
+
+				AnalysisData.Reset();
+			}
+		}
+
+		// It's possible for the maximum magnitude to be 0.0 if the audio file was silent. 
+		if (MaximumMagnitude > 0.0f)
+		{
+			// Normalize all the magnitude values based on the highest magnitude
+			for (FSoundWaveSpectralTimeData& SpectralTimeData : CookedSpectralTimeData)
+			{
+				for (FSoundWaveSpectralDataEntry& DataEntry : SpectralTimeData.Data)
+				{
+					DataEntry.NormalizedMagnitude = DataEntry.Magnitude / MaximumMagnitude;
+				}
+			}
+		}
+
+	}
+}
+
+void USoundWave::BakeEnvelopeAnalysis()
+{
+	// Clear any existing envelope data regardless of if it's enabled. If this was enabled and is now toggled, this will clear previous data.
+	CookedEnvelopeTimeData.Reset();
+
+	// Perform analysis if enabled on the sound wave
+	if (bEnableAmplitudeEnvelopeAnalysis)
+	{
+		if (ChannelSizes.Num() > 0)
+		{
+			UE_LOG(LogAudio, Warning, TEXT("Soundwave '%s' has multi-channel audio (channels greater than 2). Baking envelope analysis is not currently supported for this yet."), *GetFullName());
+			return;
+		}
+
+		// Retrieve the raw imported data
+		TArray<uint8> RawImportedWaveData;
+		uint32 RawDataSampleRate = 0;
+		uint16 RawDataNumChannels = 0;
+
+		if (!GetImportedSoundWaveData(RawImportedWaveData, RawDataSampleRate, RawDataNumChannels))
+		{
+			return;
+		}
+
+		if (RawDataSampleRate == 0 || RawDataNumChannels == 0)
+		{
+			UE_LOG(LogAudio, Error, TEXT("Failed to parse the raw imported data for '%s' for baked FFT analysis."), *GetFullName());
+			return;
+		}
+
+		const uint32 NumFrames = (RawImportedWaveData.Num() / sizeof(int16)) / RawDataNumChannels;
+		int16* InputData = (int16*)RawImportedWaveData.GetData();
+
+		Audio::FEnvelopeFollower EnvelopeFollower;
+		EnvelopeFollower.Init(RawDataSampleRate, EnvelopeFollowerAttackTime, EnvelopeFollowerReleaseTime);
+
+		for (uint32 FrameIndex = 0; FrameIndex < NumFrames; ++FrameIndex)
+		{
+			// Get the averaged sample value of all the channels
+			float SampleValue = 0.0f;
+			for (uint16 ChannelIndex = 0; ChannelIndex < RawDataNumChannels; ++ChannelIndex)
+			{
+				SampleValue += (float)InputData[FrameIndex * RawDataNumChannels] / 32767.0f;
+			}
+			SampleValue /= RawDataNumChannels;
+
+			float Output = EnvelopeFollower.ProcessAudio(SampleValue);
+
+			// Until we reached the frame size
+			if (FrameIndex % EnvelopeFollowerFrameSize == 0)
+			{
+				FSoundWaveEnvelopeTimeData NewData;
+				NewData.Amplitude = Output;
+				NewData.TimeSec = (float)FrameIndex / RawDataSampleRate;
+				CookedEnvelopeTimeData.Add(NewData);
+			}
+		}
+
+	}
+
+}
+
 void USoundWave::PostEditChangeProperty(FPropertyChangedEvent& PropertyChangedEvent)
 {
 	Super::PostEditChangeProperty(PropertyChangedEvent);
@@ -742,7 +1026,7 @@ void USoundWave::PostEditChangeProperty(FPropertyChangedEvent& PropertyChangedEv
 	// Prevent constant re-compression of SoundWave while properties are being changed interactively
 	if (PropertyChangedEvent.ChangeType != EPropertyChangeType::Interactive)
 	{
-		// Regenerate on save any compressed sound formats
+		// Regenerate on save any compressed sound formats or if analysis needs to be re-done
 		if (UProperty* PropertyThatChanged = PropertyChangedEvent.Property)
 		{
 			const FName& Name = PropertyThatChanged->GetFName();
@@ -753,6 +1037,14 @@ void USoundWave::PostEditChangeProperty(FPropertyChangedEvent& PropertyChangedEv
 				UpdatePlatformData();
 				MarkPackageDirty();
 			}
+			else if (AnyFFTAnalysisPropertiesChanged(Name))
+			{
+				BakeFFTAnalysis();
+			}
+			else if (AnyEnvelopeAnalysisPropertiesChanged(Name))
+			{
+				BakeEnvelopeAnalysis();
+			}
 		}
 	}
 }
@@ -1185,6 +1477,16 @@ bool USoundWave::IsStreaming(const FPlatformAudioCookOverrides* Overrides /* = n
 			&& Duration > Overrides->AutoStreamingThreshold);
 }
 
+bool USoundWave::GetSoundWavesWithCookedAnalysisData(TArray<USoundWave*>& OutSoundWaves)
+{
+	if (CookedSpectralTimeData.Num() > 0 || CookedEnvelopeTimeData.Num() > 0)
+	{
+		OutSoundWaves.Add(this);
+		return true;
+	}
+	return false;
+}
+
 void USoundWave::UpdatePlatformData()
 {
 	if (IsStreaming())
@@ -1281,3 +1583,214 @@ bool USoundWave::GetChunkData(int32 ChunkIndex, uint8** OutChunkData)
 	}
 	return true;
 }
+
+uint32 USoundWave::GetInterpolatedCookedFFTDataForTimeInternal(float InTime, uint32 StartingIndex, TArray<FSoundWaveSpectralData>& OutData, bool bLoop)
+{
+	// Find the two entries on either side of the input time
+	int32 NumDataEntries = CookedSpectralTimeData.Num();
+	for (int32 Index = StartingIndex; Index < NumDataEntries; ++Index)
+	{
+		// Get the current data at this index
+		const FSoundWaveSpectralTimeData& CurrentData = CookedSpectralTimeData[Index];
+
+		// Get the next data, wrap if needed (i.e. if current is last index, we'll lerp to the first index)
+		const FSoundWaveSpectralTimeData& NextData = CookedSpectralTimeData[(Index + 1) % NumDataEntries];
+
+		if (InTime >= CurrentData.TimeSec && InTime < NextData.TimeSec)
+		{
+			// Lerping alpha is fraction from current to next data
+			const float Alpha = (InTime - CurrentData.TimeSec) / (NextData.TimeSec - CurrentData.TimeSec);
+			for (int32 FrequencyIndex = 0; FrequencyIndex < FrequenciesToAnalyze.Num(); ++FrequencyIndex)
+			{
+				FSoundWaveSpectralData InterpolatedData;
+				InterpolatedData.FrequencyHz = FrequenciesToAnalyze[FrequencyIndex];
+				InterpolatedData.Magnitude = FMath::Lerp(CurrentData.Data[FrequencyIndex].Magnitude, NextData.Data[FrequencyIndex].Magnitude, Alpha);
+				InterpolatedData.NormalizedMagnitude = FMath::Lerp(CurrentData.Data[FrequencyIndex].NormalizedMagnitude, NextData.Data[FrequencyIndex].NormalizedMagnitude, Alpha);
+
+				OutData.Add(InterpolatedData);
+			}
+
+			// Sort by frequency (lowest frequency first).
+			OutData.Sort(FCompareSpectralDataByFrequencyHz());
+
+			// We found cooked spectral data which maps to these indices
+			return Index;
+		}
+	}
+
+	return INDEX_NONE;
+}
+
+bool USoundWave::GetInterpolatedCookedFFTDataForTime(float InTime, uint32& InOutLastIndex, TArray<FSoundWaveSpectralData>& OutData, bool bLoop)
+{
+	if (CookedSpectralTimeData.Num() > 0)
+	{
+		// Handle edge cases
+		if (!bLoop)
+		{
+			// Pointer to which data to use
+			FSoundWaveSpectralTimeData* SpectralTimeData = nullptr;
+
+			// We are past the edge
+			if (InTime >= CookedSpectralTimeData.Last().TimeSec)
+			{
+				SpectralTimeData = &CookedSpectralTimeData.Last();
+				InOutLastIndex = CookedPlatformData.Num() - 1;
+			}
+			// We are before the first data point
+			else if (InTime < CookedSpectralTimeData[0].TimeSec)
+			{
+				SpectralTimeData = &CookedSpectralTimeData[0];
+				InOutLastIndex = 0;
+			}
+
+			// If we were either case before we have a non-nullptr here
+			if (SpectralTimeData != nullptr)
+			{
+				// Create an entry for this clamped output
+				for (int32 FrequencyIndex = 0; FrequencyIndex < FrequenciesToAnalyze.Num(); ++FrequencyIndex)
+				{
+					FSoundWaveSpectralData InterpolatedData;
+					InterpolatedData.FrequencyHz = FrequenciesToAnalyze[FrequencyIndex];
+					InterpolatedData.Magnitude = SpectralTimeData->Data[FrequencyIndex].Magnitude;
+					InterpolatedData.NormalizedMagnitude = SpectralTimeData->Data[FrequencyIndex].NormalizedMagnitude;
+
+					OutData.Add(InterpolatedData);
+				}
+
+				return true;
+			}
+		}
+		// We're looping 
+		else 
+		{
+			// Need to check initial wrap-around case (i.e. we're reading earlier than first data point so need to lerp from last data point to first
+			if (InTime >= 0.0f && InTime < CookedSpectralTimeData[0].TimeSec)
+			{
+				const FSoundWaveSpectralTimeData& CurrentData = CookedSpectralTimeData.Last();
+
+				// Get the next data, wrap if needed (i.e. if current is last index, we'll lerp to the first index)
+				const FSoundWaveSpectralTimeData& NextData = CookedSpectralTimeData[0];
+
+				float TimeLeftFromLastDataToEnd = Duration - CurrentData.TimeSec;
+				float Alpha = (TimeLeftFromLastDataToEnd + InTime) / (TimeLeftFromLastDataToEnd + NextData.TimeSec);
+
+				for (int32 FrequencyIndex = 0; FrequencyIndex < FrequenciesToAnalyze.Num(); ++FrequencyIndex)
+				{
+					FSoundWaveSpectralData InterpolatedData;
+					InterpolatedData.FrequencyHz = FrequenciesToAnalyze[FrequencyIndex];
+					InterpolatedData.Magnitude = FMath::Lerp(CurrentData.Data[FrequencyIndex].Magnitude, NextData.Data[FrequencyIndex].Magnitude, Alpha);
+					InterpolatedData.NormalizedMagnitude = FMath::Lerp(CurrentData.Data[FrequencyIndex].NormalizedMagnitude, NextData.Data[FrequencyIndex].NormalizedMagnitude, Alpha);
+
+					OutData.Add(InterpolatedData);
+
+					InOutLastIndex = 0;
+				}
+				return true;
+			}
+			// Or we've been offset a bit in the negative.
+			else if (InTime < 0.0f)
+			{
+				// Wrap the time to the end of the sound wave file 
+				InTime = FMath::Clamp(Duration + InTime, 0.0f, Duration);
+			}
+		}
+
+		uint32 StartingIndex = InOutLastIndex == INDEX_NONE ? 0 : InOutLastIndex;
+
+		InOutLastIndex = GetInterpolatedCookedFFTDataForTimeInternal(InTime, StartingIndex, OutData, bLoop);
+		if (InOutLastIndex == INDEX_NONE && StartingIndex != 0)
+		{
+			InOutLastIndex = GetInterpolatedCookedFFTDataForTimeInternal(InTime, 0, OutData, bLoop);
+		}
+		return InOutLastIndex != INDEX_NONE;
+	}
+
+	return false;
+}
+
+uint32 USoundWave::GetInterpolatedCookedEnvelopeDataForTimeInternal(float InTime, uint32 StartingIndex, float& OutAmplitude, bool bLoop)
+{
+	if (StartingIndex == INDEX_NONE || StartingIndex == CookedEnvelopeTimeData.Num())
+	{
+		StartingIndex = 0;
+	}
+
+	// Find the two entries on either side of the input time
+	int32 NumDataEntries = CookedEnvelopeTimeData.Num();
+	for (int32 Index = StartingIndex; Index < NumDataEntries; ++Index)
+	{
+		const FSoundWaveEnvelopeTimeData& CurrentData = CookedEnvelopeTimeData[Index];
+		const FSoundWaveEnvelopeTimeData& NextData = CookedEnvelopeTimeData[(Index + 1) % NumDataEntries];
+
+		if (InTime >= CurrentData.TimeSec && InTime < NextData.TimeSec)
+		{
+			// Lerping alpha is fraction from current to next data
+			const float Alpha = (InTime - CurrentData.TimeSec) / (NextData.TimeSec - CurrentData.TimeSec);
+			OutAmplitude = FMath::Lerp(CurrentData.Amplitude, NextData.Amplitude, Alpha);
+
+			// We found cooked spectral data which maps to these indices
+			return Index;
+		}
+	}
+
+	// Did not find the data
+	return INDEX_NONE;
+}
+
+bool USoundWave::GetInterpolatedCookedEnvelopeDataForTime(float InTime, uint32& InOutLastIndex, float& OutAmplitude, bool bLoop)
+{
+	InOutLastIndex = INDEX_NONE;
+	if (CookedEnvelopeTimeData.Num() > 0 && InTime >= 0.0f)
+	{
+		// Handle edge cases
+		if (!bLoop)
+		{
+			// We are past the edge
+			if (InTime >= CookedEnvelopeTimeData.Last().TimeSec)
+			{
+				OutAmplitude = CookedEnvelopeTimeData.Last().Amplitude;
+				InOutLastIndex = CookedEnvelopeTimeData.Num() - 1;
+				return true;
+			}
+			// We are before the first data point
+			else if (InTime < CookedEnvelopeTimeData[0].TimeSec)
+			{
+				OutAmplitude = CookedEnvelopeTimeData[0].Amplitude;
+				InOutLastIndex = 0;
+				return true;
+			}
+		}	
+		else
+		{
+			// Need to check initial wrap-around case (i.e. we're reading earlier than first data point so need to lerp from last data point to first
+			if (InTime >= 0.0f && InTime < CookedEnvelopeTimeData[0].TimeSec)
+			{
+				const FSoundWaveEnvelopeTimeData& CurrentData = CookedEnvelopeTimeData.Last();
+				const FSoundWaveEnvelopeTimeData& NextData = CookedEnvelopeTimeData[0];
+
+				float TimeLeftFromLastDataToEnd = Duration - CurrentData.TimeSec;
+				float Alpha = (TimeLeftFromLastDataToEnd + InTime) / (TimeLeftFromLastDataToEnd + NextData.TimeSec);
+
+				OutAmplitude = FMath::Lerp(CurrentData.Amplitude, NextData.Amplitude, Alpha);
+				InOutLastIndex = 0;
+				return true;
+			}
+			// Or we've been offset a bit in the negative.
+			else if (InTime < 0.0f)
+			{
+				// Wrap the time to the end of the sound wave file 
+				InTime = FMath::Clamp(Duration + InTime, 0.0f, Duration);
+			}
+
+			uint32 StartingIndex = InOutLastIndex == INDEX_NONE ? 0 : InOutLastIndex;
+
+			InOutLastIndex = GetInterpolatedCookedEnvelopeDataForTimeInternal(InTime, StartingIndex, OutAmplitude, bLoop);
+			if (InOutLastIndex == INDEX_NONE && StartingIndex != 0)
+			{
+				InOutLastIndex = GetInterpolatedCookedEnvelopeDataForTimeInternal(InTime, 0, OutAmplitude, bLoop);
+			}
+		}
+	}
+	return InOutLastIndex != INDEX_NONE;
+}
diff --git a/Engine/Source/Runtime/Engine/Public/ActiveSound.h b/Engine/Source/Runtime/Engine/Public/ActiveSound.h
--- a/Engine/Source/Runtime/Engine/Public/ActiveSound.h
+++ b/Engine/Source/Runtime/Engine/Public/ActiveSound.h
@@ -366,6 +366,9 @@ public:
 	/** Whether or not this active sound will update the average envelope value of every wave instance that plays a sound source. Based on set delegates on audio component. */
 	uint8 bUpdateMultiEnvelopeValue:1;
 
+	/** Whether or not the active sound should update it's owning audio component's playback time. */
+	uint8 bUpdatePlaybackTime:1;
+
 	/** Whether or not this active sound is playing audio, as in making audible sounds. */
 	uint8 bIsPlayingAudio:1;
 
diff --git a/Engine/Source/Runtime/Engine/Public/Audio.h b/Engine/Source/Runtime/Engine/Public/Audio.h
--- a/Engine/Source/Runtime/Engine/Public/Audio.h
+++ b/Engine/Source/Runtime/Engine/Public/Audio.h
@@ -427,6 +427,9 @@ public:
 	/** The absolute position of the wave instance relative to forward vector of listener. */
 	float AbsoluteAzimuth; 
 
+	/** The playback time of the wave instance. Updated from active sound. */
+	float PlaybackTime;
+
 	/** The reverb send method to use. */
 	EReverbSendMethod ReverbSendMethod;
 
@@ -846,7 +849,6 @@ protected:
 	friend struct FActiveSound;
 };
 
-
 /*-----------------------------------------------------------------------------
 	FWaveModInfo. 
 -----------------------------------------------------------------------------*/
