Remove flag to disable optimization in Engine.build.cs #fyi aaron.mcleran, bob.tellez

diff --git a/Engine/Source/Runtime/Engine/Engine.Build.cs b/Engine/Source/Runtime/Engine/Engine.Build.cs
--- a/Engine/Source/Runtime/Engine/Engine.Build.cs
+++ b/Engine/Source/Runtime/Engine/Engine.Build.cs
@@ -7,8 +7,6 @@ public class Engine : ModuleRules
 {
 	public Engine(ReadOnlyTargetRules Target) : base(Target)
 	{
-        OptimizeCode = CodeOptimization.Never;
-
         PrivatePCHHeaderFile = "Private/EnginePrivatePCH.h";
 
 		SharedPCHHeaderFile = "Public/EngineSharedPCH.h";
