Adding in null check

diff --git a/Engine/Source/Runtime/Engine/Private/SoundCue.cpp b/Engine/Source/Runtime/Engine/Private/SoundCue.cpp
--- a/Engine/Source/Runtime/Engine/Private/SoundCue.cpp
+++ b/Engine/Source/Runtime/Engine/Private/SoundCue.cpp
@@ -497,7 +497,7 @@ bool USoundCue::GetSoundWavesWithCookedAnalysisData(TArray<USoundWave*>& OutSoun
 	for (USoundNodeWavePlayer* Player : WavePlayers)
 	{
 		USoundWave* SoundWave = Player->GetSoundWave();
-		if (SoundWave->GetSoundWavesWithCookedAnalysisData(OutSoundWaves))
+		if (SoundWave && SoundWave->GetSoundWavesWithCookedAnalysisData(OutSoundWaves))
 		{
 			bHasAnalysisData = true;
 		}
