Only log the actor channel lmit reached warning once - this prevents log spam when the limit is reached

diff --git a/Engine/Source/Runtime/Engine/Classes/Engine/NetConnection.h b/Engine/Source/Runtime/Engine/Classes/Engine/NetConnection.h
--- a/Engine/Source/Runtime/Engine/Classes/Engine/NetConnection.h
+++ b/Engine/Source/Runtime/Engine/Classes/Engine/NetConnection.h
@@ -1095,6 +1095,9 @@ private:
 
 	/** Keep old behavior where we send a packet with only acks even if we have no other outgoing data if we got incoming data */
 	uint32 HasDirtyAcks;
+	
+	/** True if we've hit the actor channel limit and logged a warning about it */
+	bool bHasWarnedAboutChannelLimit;
 };
 
 
diff --git a/Engine/Source/Runtime/Engine/Private/NetConnection.cpp b/Engine/Source/Runtime/Engine/Private/NetConnection.cpp
--- a/Engine/Source/Runtime/Engine/Private/NetConnection.cpp
+++ b/Engine/Source/Runtime/Engine/Private/NetConnection.cpp
@@ -135,6 +135,7 @@ UNetConnection::UNetConnection(const FObjectInitializer& ObjectInitializer)
 ,	ClientWorldPackageName( NAME_None )
 ,	LastNotifiedPacketId( -1 )
 ,	HasDirtyAcks(0u)
+,	bHasWarnedAboutChannelLimit(false)
 {
 	MaxChannelSize = CVarMaxChannelSize.GetValueOnAnyThread();
 	if (MaxChannelSize <= 0)
@@ -2300,7 +2301,11 @@ UChannel* UNetConnection::CreateChannelByName( const FName& ChName, EChannelCrea
 		// Fail to create if the channel array is full
 		if( ChIndex == Channels.Num() )
 		{
-			UE_LOG(LogNetTraffic, Warning, TEXT("No free channel could be found in the channel list (current limit is %d channels). Consider increasing the max channels allowed using CVarMaxChannelSize."), MaxChannelSize);
+			if (!bHasWarnedAboutChannelLimit)
+			{
+				bHasWarnedAboutChannelLimit = true;
+				UE_LOG(LogNetTraffic, Warning, TEXT("No free channel could be found in the channel list (current limit is %d channels) for connection with owner %s. Consider increasing the max channels allowed using CVarMaxChannelSize."), MaxChannelSize, *GetNameSafe(OwningActor));
+			}
 			return NULL;
 		}
 	}
