MN - The assert was invalid in this case because the material can be delayed deleted and used one more frame. Nevertheless there is an issue where the AsyncLoadingThread can enter at the same time and mess with the GlobalMap corrupting it and eventually causing a new Material using a deleted Map. (See UE-63799)

diff --git a/Engine/Source/Runtime/Engine/Private/Materials/MaterialShader.cpp b/Engine/Source/Runtime/Engine/Private/Materials/MaterialShader.cpp
--- a/Engine/Source/Runtime/Engine/Private/Materials/MaterialShader.cpp
+++ b/Engine/Source/Runtime/Engine/Private/Materials/MaterialShader.cpp
@@ -35,6 +35,7 @@ namespace MaterialShaderCookStats
 //
 // Globals
 //
+FCriticalSection FMaterialShaderMap::GIdToMaterialShaderMapCS;
 TMap<FMaterialShaderMapId,FMaterialShaderMap*> FMaterialShaderMap::GIdToMaterialShaderMap[SP_NumPlatforms];
 #if ALLOW_SHADERMAP_DEBUG_DATA
 TArray<FMaterialShaderMap*> FMaterialShaderMap::AllMaterialShaderMaps;
@@ -963,10 +964,13 @@ FShader* FMaterialShaderType::FinishCompileShader(
 * @param Platform - The platform to lookup for
 * @return NULL if no cached shader map was found.
 */
-FMaterialShaderMap* FMaterialShaderMap::FindId(const FMaterialShaderMapId& ShaderMapId, EShaderPlatform InPlatform)
+TRefCountPtr<FMaterialShaderMap> FMaterialShaderMap::FindId(const FMaterialShaderMapId& ShaderMapId, EShaderPlatform InPlatform)
 {
+	FScopeLock ScopeLock(&GIdToMaterialShaderMapCS);
 	check(ShaderMapId.BaseMaterialId != FGuid());
-	return GIdToMaterialShaderMap[InPlatform].FindRef(ShaderMapId);
+	TRefCountPtr<FMaterialShaderMap> Result = GIdToMaterialShaderMap[InPlatform].FindRef(ShaderMapId);
+	check(Result == nullptr || !Result->bDeletedThroughDeferredCleanup);
+	return Result;
 }
 
 #if ALLOW_SHADERMAP_DEBUG_DATA
@@ -995,6 +999,8 @@ void FMaterialShaderMap::FlushShaderTypes(TArray<FShaderType*>& ShaderTypesToFlu
 
 void FMaterialShaderMap::FixupShaderTypes(EShaderPlatform Platform, const TMap<FShaderType*, FString>& ShaderTypeNames, const TMap<const FShaderPipelineType*, FString>& ShaderPipelineTypeNames, const TMap<FVertexFactoryType*, FString>& VertexFactoryTypeNames)
 {
+	FScopeLock ScopeLock(&GIdToMaterialShaderMapCS);
+
 	TArray<FMaterialShaderMapId> Keys;
 	FMaterialShaderMap::GIdToMaterialShaderMap[Platform].GenerateKeyArray(Keys);
 
@@ -2134,14 +2140,18 @@ void FMaterialShaderMap::Register(EShaderPlatform InShaderPlatform)
 		INC_DWORD_STAT_BY(STAT_Shaders_ShaderMapMemory, GetSizeBytes());
 	}
 
-	GIdToMaterialShaderMap[GetShaderPlatform()].Add(ShaderMapId,this);
-	bRegistered = true;
+	{
+		FScopeLock ScopeLock(&GIdToMaterialShaderMapCS);
+		GIdToMaterialShaderMap[GetShaderPlatform()].Add(ShaderMapId,this);
+		bRegistered = true;
+	}
 }
 
 void FMaterialShaderMap::AddRef()
 {
 	//#todo-mw: re-enable to try to find potential corruption of the global shader map ID array
 	//check(IsInGameThread());
+	FScopeLock ScopeLock(&GIdToMaterialShaderMapCS);
 	check(!bDeletedThroughDeferredCleanup);
 	++NumRefs;
 }
@@ -2150,20 +2160,27 @@ void FMaterialShaderMap::Release()
 {
 	//#todo-mw: re-enable to try to find potential corruption of the global shader map ID array
 	//check(IsInGameThread());
-	check(NumRefs > 0);
-	if(--NumRefs == 0)
+
 	{
-		if (bRegistered)
+		FScopeLock ScopeLock(&GIdToMaterialShaderMapCS);
+
+		check(NumRefs > 0);
+		if (--NumRefs == 0)
 		{
-			DEC_DWORD_STAT(STAT_Shaders_NumShaderMaps);
-			DEC_DWORD_STAT_BY(STAT_Shaders_ShaderMapMemory, GetSizeBytes());
+			if (bRegistered)
+			{
+				bRegistered = false;
+				DEC_DWORD_STAT(STAT_Shaders_NumShaderMaps);
+				DEC_DWORD_STAT_BY(STAT_Shaders_ShaderMapMemory, GetSizeBytes());
+				GIdToMaterialShaderMap[GetShaderPlatform()].Remove(ShaderMapId);
+			}
 
-			GIdToMaterialShaderMap[GetShaderPlatform()].Remove(ShaderMapId);
-			bRegistered = false;
+			check(!bDeletedThroughDeferredCleanup);
+			bDeletedThroughDeferredCleanup = true;
 		}
-
-		check(!bDeletedThroughDeferredCleanup);
-		bDeletedThroughDeferredCleanup = true;
+	}
+	if (bDeletedThroughDeferredCleanup)
+	{
 		BeginCleanup(this);
 	}
 }
diff --git a/Engine/Source/Runtime/Engine/Private/Materials/MaterialShared.cpp b/Engine/Source/Runtime/Engine/Private/Materials/MaterialShared.cpp
--- a/Engine/Source/Runtime/Engine/Private/Materials/MaterialShared.cpp
+++ b/Engine/Source/Runtime/Engine/Private/Materials/MaterialShared.cpp
@@ -1610,7 +1610,7 @@ bool FMaterial::CacheShaders(const FMaterialShaderMapId& ShaderMapId, EShaderPla
 	// If we loaded this material with inline shaders, use what was loaded (GameThreadShaderMap) instead of looking in the DDC
 	if (bContainsInlineShaders)
 	{
-		FMaterialShaderMap* ExistingShaderMap = nullptr;
+		TRefCountPtr<FMaterialShaderMap> ExistingShaderMap = nullptr;
 		
 		if (GameThreadShaderMap)
 		{
diff --git a/Engine/Source/Runtime/Engine/Public/MaterialShared.h b/Engine/Source/Runtime/Engine/Public/MaterialShared.h
--- a/Engine/Source/Runtime/Engine/Public/MaterialShared.h
+++ b/Engine/Source/Runtime/Engine/Public/MaterialShared.h
@@ -705,7 +705,7 @@ public:
 	 * @param Platform - The platform to lookup for
 	 * @return NULL if no cached shader map was found.
 	 */
-	static FMaterialShaderMap* FindId(const FMaterialShaderMapId& ShaderMapId, EShaderPlatform Platform);
+	static TRefCountPtr<FMaterialShaderMap> FindId(const FMaterialShaderMapId& ShaderMapId, EShaderPlatform Platform);
 
 #if ALLOW_SHADERMAP_DEBUG_DATA
 	/** Flushes the given shader types from any loaded FMaterialShaderMap's. */
@@ -878,7 +878,7 @@ public:
 
 	bool IsValidForRendering(bool bFailOnInvalid = false) const
 	{
-		const bool bValid = bCompilationFinalized && bCompiledSuccessfully && !bDeletedThroughDeferredCleanup;
+		const bool bValid = bCompilationFinalized && bCompiledSuccessfully;// && !bDeletedThroughDeferredCleanup; //deferred actually deletion will prevent the material to go away before we finish rendering
 		checkf(bValid || !bFailOnInvalid, TEXT("FMaterialShaderMap %s invalid for rendering: bCompilationFinalized: %i, bCompiledSuccessfully: %i, bDeletedThroughDeferredCleanup: %i"), *GetFriendlyName(), bCompilationFinalized, bCompiledSuccessfully, bDeletedThroughDeferredCleanup ? 1 : 0);
 		return bValid;
 	}
@@ -912,6 +912,7 @@ private:
 	 * No ref counting needed as these are removed on destruction of the shader map.
 	 */
 	static TMap<FMaterialShaderMapId,FMaterialShaderMap*> GIdToMaterialShaderMap[SP_NumPlatforms];
+	static FCriticalSection GIdToMaterialShaderMapCS;
 
 #if ALLOW_SHADERMAP_DEBUG_DATA
 	/** 
