Implementing BP and console command for setting max channels dynamically.

diff --git a/Engine/Source/Runtime/Engine/Classes/Kismet/GameplayStatics.h b/Engine/Source/Runtime/Engine/Classes/Kismet/GameplayStatics.h
--- a/Engine/Source/Runtime/Engine/Classes/Kismet/GameplayStatics.h
+++ b/Engine/Source/Runtime/Engine/Classes/Kismet/GameplayStatics.h
@@ -705,6 +705,19 @@ public:
 	UFUNCTION(BlueprintCallable, Category = "Audio", meta = (WorldContext = "WorldContextObject"))
 	static class UReverbEffect* GetCurrentReverbEffect(const UObject* WorldContextObject);
 
+	/**
+	 * Sets the max number of voices (also known as "channels") dynamically by percentage. E.g. if you want to temporarily reduce voice count by 50%, use 0.50. Later, you can return to the original max voice count by using 1.0.
+	 */
+	UFUNCTION(BlueprintCallable, Category = "Audio", meta = (WorldContext = "WorldContextObject"))
+	static void SetMaxAudioChannelsScaled(const UObject* WorldContextObject, float MaxChannelCountScale);
+
+	/**
+	 * Retrieves the max voice count currently used by the audio engine.
+	 */
+	UFUNCTION(BlueprintCallable, Category = "Audio", meta = (WorldContext = "WorldContextObject"))
+	static int32 GetMaxAudioChannelCount(const UObject* WorldContextObject);
+
+
 	// --- Decal functions ------------------------------
 
 	/** Spawns a decal at the given location and rotation, fire and forget. Does not replicate.
diff --git a/Engine/Source/Runtime/Engine/Private/AudioDevice.cpp b/Engine/Source/Runtime/Engine/Private/AudioDevice.cpp
--- a/Engine/Source/Runtime/Engine/Private/AudioDevice.cpp
+++ b/Engine/Source/Runtime/Engine/Private/AudioDevice.cpp
@@ -49,6 +49,13 @@ FAutoConsoleVariableRef CVarSetAudioChannelCount(
 	TEXT("0: Disable, >0: Enable"),
 	ECVF_Default);
 
+static float AudioChannelCountScaleCVar = 1.0f;
+FAutoConsoleVariableRef CVarSetAudioChannelScaleCount(
+	TEXT("au.SetAudioChannelScaleCount"),
+	AudioChannelCountScaleCVar,
+	TEXT("Changes the audio channel count by percentage.\n"),
+	ECVF_Default);
+
 static int32 DisableStoppingVoicesCvar = 0;
 FAutoConsoleVariableRef CVarDisableStoppingVoices(
 	TEXT("au.DisableStoppingVoices"),
@@ -159,6 +166,9 @@ void FDynamicParameter::Update(float DeltaTime)
 
 FAudioDevice::FAudioDevice()
 	: MaxChannels(0)
+	, MaxChannels_GameThread(0)
+	, MaxChannelsScale(1.0f)
+	, MaxChannelsScale_GameThread(1.0f)
 	, NumStoppingVoices(32)
 	, MaxWaveInstances(0)
 	, SampleRate(0)
@@ -247,6 +257,7 @@ bool FAudioDevice::Init(int32 InMaxChannels)
 
 	// MaxChannels is the min of the platform-specific value and the max value in the quality settings (InMaxChannels)
 	MaxChannels = PlatformSettings.MaxChannels > 0 ? FMath::Min(PlatformSettings.MaxChannels, InMaxChannels) : InMaxChannels;
+	MaxChannels_GameThread = MaxChannels;
 
 	// Mixed sample rate is set by the platform
 	SampleRate = PlatformSettings.SampleRate;
@@ -428,6 +439,21 @@ void FAudioDevice::PrecacheStartupSounds()
 
 void FAudioDevice::SetMaxChannels(int32 InMaxChannels)
 {
+	if (!IsInAudioThread())
+	{
+		MaxChannels_GameThread = InMaxChannels;
+
+		DECLARE_CYCLE_STAT(TEXT("FAudioThreadTask.SetMaxChannels"), STAT_AudioSetMaxChannels, STATGROUP_AudioThreadCommands);
+
+		FAudioThread::RunCommandOnAudioThread([this, InMaxChannels]()
+		{
+			this->SetMaxChannels(InMaxChannels);
+
+		}, GET_STATID(STAT_AudioSetMaxChannels));
+
+		return;
+	}
+
 	if (InMaxChannels > Sources.Num())
 	{
 		UE_LOG(LogAudio, Warning, TEXT("Can't increase channels past starting number!"));
@@ -437,14 +463,56 @@ void FAudioDevice::SetMaxChannels(int32 InMaxChannels)
 	MaxChannels = InMaxChannels;
 }
 
+void FAudioDevice::SetMaxChannelsScaled(float InScaledChannelCount)
+{
+	if (!IsInAudioThread())
+	{
+		MaxChannelsScale_GameThread = InScaledChannelCount;
+
+		DECLARE_CYCLE_STAT(TEXT("FAudioThreadTask.SetMaxChannelsScaled"), STAT_AudioSetMaxChannelsScaled, STATGROUP_AudioThreadCommands);
+
+		FAudioThread::RunCommandOnAudioThread([this, InScaledChannelCount]()
+		{
+			MaxChannelsScale = FMath::Clamp(InScaledChannelCount, 0.0f, 1.0f);
+
+		}, GET_STATID(STAT_AudioSetMaxChannelsScaled));
+
+		return;
+	}
+	else
+	{
+		MaxChannelsScale = FMath::Clamp(InScaledChannelCount, 0.0f, 1.0f);
+
+		DECLARE_CYCLE_STAT(TEXT("FAudioThreadTask.SetMaxChannelsScaled"), STAT_AudioSetMaxChannelsScaled, STATGROUP_AudioThreadCommands);
+
+		FAudioThread::RunCommandOnGameThread([this, InScaledChannelCount]()
+		{
+			MaxChannelsScale_GameThread = InScaledChannelCount;
+
+		}, GET_STATID(STAT_AudioSetMaxChannelsScaled));
+	}
+}
+
 int32 FAudioDevice::GetMaxChannels() const
 {
-	if (AudioChannelCountCVar > 0 && AudioChannelCountCVar < Sources.Num())
+	if (IsInAudioThread())
 	{
-		return AudioChannelCountCVar;
+		if (AudioChannelCountCVar > 0 && AudioChannelCountCVar < Sources.Num())
+		{
+			return FMath::Max(int32(AudioChannelCountCVar * MaxChannelsScale * AudioChannelCountScaleCVar), 1);
+		}
+
+		return FMath::Max(int32(MaxChannels * MaxChannelsScale * AudioChannelCountScaleCVar), 1);
 	}
+	else
+	{
+		if (AudioChannelCountCVar > 0 && AudioChannelCountCVar < Sources.Num())
+		{
+			return FMath::Max(int32(AudioChannelCountCVar * MaxChannelsScale_GameThread * AudioChannelCountScaleCVar), 1);
+		}
 
-	return MaxChannels;
+		return FMath::Max(int32(MaxChannels_GameThread * MaxChannelsScale_GameThread * AudioChannelCountScaleCVar), 1);
+	}
 }
 
 void FAudioDevice::Teardown()
diff --git a/Engine/Source/Runtime/Engine/Private/GameplayStatics.cpp b/Engine/Source/Runtime/Engine/Private/GameplayStatics.cpp
--- a/Engine/Source/Runtime/Engine/Private/GameplayStatics.cpp
+++ b/Engine/Source/Runtime/Engine/Private/GameplayStatics.cpp
@@ -1600,6 +1600,48 @@ class UReverbEffect* UGameplayStatics::GetCurrentReverbEffect(const UObject* Wor
 	return nullptr;
 }
 
+void UGameplayStatics::SetMaxAudioChannelsScaled(const UObject* WorldContextObject, float MaxChannelCountScale)
+{
+	if (GEngine == nullptr || !GEngine->UseSound())
+	{
+		return;
+	}
+
+	UWorld* ThisWorld = GEngine->GetWorldFromContextObject(WorldContextObject, EGetWorldErrorMode::LogAndReturnNull);
+	if (!ThisWorld || !ThisWorld->bAllowAudioPlayback)
+	{
+		return;
+	}
+
+	if (FAudioDevice* AudioDevice = ThisWorld->GetAudioDevice())
+	{
+		AudioDevice->SetMaxChannelsScaled(MaxChannelCountScale);
+	}
+}
+
+int32 UGameplayStatics::GetMaxAudioChannelCount(const UObject* WorldContextObject)
+{
+	if (GEngine == nullptr || !GEngine->UseSound())
+	{
+		return 0;
+	}
+
+
+	UWorld* ThisWorld = GEngine->GetWorldFromContextObject(WorldContextObject, EGetWorldErrorMode::LogAndReturnNull);
+	if (!ThisWorld || !ThisWorld->bAllowAudioPlayback)
+	{
+		return 0;
+	}
+
+	if (FAudioDevice* AudioDevice = ThisWorld->GetAudioDevice())
+	{
+		return AudioDevice->GetMaxChannels();
+	}
+
+	return 0;
+}
+
+
 UDecalComponent* CreateDecalComponent(class UMaterialInterface* DecalMaterial, FVector DecalSize, UWorld* World, AActor* Actor, float LifeSpan)
 {
 	UDecalComponent* DecalComp = NewObject<UDecalComponent>((Actor ? Actor : (UObject*)World));
diff --git a/Engine/Source/Runtime/Engine/Public/AudioDevice.h b/Engine/Source/Runtime/Engine/Public/AudioDevice.h
--- a/Engine/Source/Runtime/Engine/Public/AudioDevice.h
+++ b/Engine/Source/Runtime/Engine/Public/AudioDevice.h
@@ -624,7 +624,12 @@ public:
 	/** 
 	 * Sets the maximum number of channels dynamically. Can't raise the cap over the initial value but can lower it 
 	 */
-	virtual void SetMaxChannels(int32 InMaxChannels);
+	void SetMaxChannels(int32 InMaxChannels);
+
+	/**
+	 * Sets the maximum number of channels dynamically by scaled percentage.
+	 */
+	void SetMaxChannelsScaled(float InScaledChannelCount);
 
 	/** Returns the max channels used by the audio device. */
 	int32 GetMaxChannels() const;
@@ -1550,6 +1555,11 @@ public:
 
 	/** The maximum number of concurrent audible sounds */
 	int32 MaxChannels;
+	int32 MaxChannels_GameThread;
+
+	/** A scaler on the max channels. */
+	float MaxChannelsScale;
+	float MaxChannelsScale_GameThread;
 
 	/** The number of sources to reserve for stopping sounds. */
 	int32 NumStoppingVoices;
