Fixed up query on audio component for if there is FFT or envelope data.

Properly querying sound cues!

diff --git a/Engine/Source/Runtime/Engine/Classes/Sound/SoundBase.h b/Engine/Source/Runtime/Engine/Classes/Sound/SoundBase.h
--- a/Engine/Source/Runtime/Engine/Classes/Sound/SoundBase.h
+++ b/Engine/Source/Runtime/Engine/Classes/Sound/SoundBase.h
@@ -221,5 +221,8 @@ public:
 	/** Returns whether the sound has cooked analysis data (e.g. FFT or envelope following data) and returns sound waves which have cooked data. */
 	virtual bool GetSoundWavesWithCookedAnalysisData(TArray<USoundWave*>& OutSoundWaves);
 
+	/** Queries if the sound has cooked FFT or envelope data. */
+	virtual bool HasCookedFFTData() const { return false; }
+	virtual bool HasCookedAmplitudeEnvelopeData() const { return false; }
 };
 
diff --git a/Engine/Source/Runtime/Engine/Classes/Sound/SoundCue.h b/Engine/Source/Runtime/Engine/Classes/Sound/SoundCue.h
--- a/Engine/Source/Runtime/Engine/Classes/Sound/SoundCue.h
+++ b/Engine/Source/Runtime/Engine/Classes/Sound/SoundCue.h
@@ -150,6 +150,8 @@ public:
 	virtual const FSoundAttenuationSettings* GetAttenuationSettingsToApply() const override;
 	virtual float GetSubtitlePriority() const override;
 	virtual bool GetSoundWavesWithCookedAnalysisData(TArray<USoundWave*>& OutSoundWaves) override;
+	virtual bool HasCookedFFTData() const override;
+	virtual bool HasCookedAmplitudeEnvelopeData() const override;
 	//~ End USoundBase Interface.
 
 	/** Construct and initialize a node within this Cue */
@@ -180,7 +182,8 @@ public:
 	/**
 	 * Recursively finds sound nodes of type T
 	 */
-	template<typename T> void RecursiveFindNode(USoundNode* Node, TArray<T*>& OutNodes)
+	template<typename T> 
+	void RecursiveFindNode(USoundNode* Node, TArray<T*>& OutNodes)
 	{
 		if (Node)
 		{
@@ -199,6 +202,27 @@ public:
 		}
 	}
 
+	template<typename T>
+	void RecursiveFindNode(const USoundNode* Node, TArray<const T*>& OutNodes) const
+	{
+		if (Node)
+		{
+			// Record the node if it is the desired type
+			if (const T* FoundNode = Cast<T>(Node))
+			{
+				OutNodes.AddUnique(FoundNode);
+			}
+
+			// Recurse.
+			const int32 MaxChildNodes = Node->GetMaxChildNodes();
+			for (int32 ChildIndex = 0; ChildIndex < Node->ChildNodes.Num() && ChildIndex < MaxChildNodes; ++ChildIndex)
+			{
+				RecursiveFindNode<T>(Node->ChildNodes[ChildIndex], OutNodes);
+			}
+		}
+	}
+
+
 	/** Find the path through the sound cue to a node identified by its hash */
 	ENGINE_API bool FindPathToNode(const UPTRINT NodeHashToFind, TArray<USoundNode*>& OutPath) const;
 
diff --git a/Engine/Source/Runtime/Engine/Classes/Sound/SoundWave.h b/Engine/Source/Runtime/Engine/Classes/Sound/SoundWave.h
--- a/Engine/Source/Runtime/Engine/Classes/Sound/SoundWave.h
+++ b/Engine/Source/Runtime/Engine/Classes/Sound/SoundWave.h
@@ -540,6 +540,8 @@ public:
 	virtual float GetSubtitlePriority() const override;
 	virtual bool IsAllowedVirtual() const override;
 	virtual bool GetSoundWavesWithCookedAnalysisData(TArray<USoundWave*>& OutSoundWaves) override;
+	virtual bool HasCookedFFTData() const override;
+	virtual bool HasCookedAmplitudeEnvelopeData() const override;
 	//~ End USoundBase Interface.
 
 	// Called  when the procedural sound wave begins on the render thread. Only used in the audio mixer and when bProcedural is true.
diff --git a/Engine/Source/Runtime/Engine/Private/Components/AudioComponent.cpp b/Engine/Source/Runtime/Engine/Private/Components/AudioComponent.cpp
--- a/Engine/Source/Runtime/Engine/Private/Components/AudioComponent.cpp
+++ b/Engine/Source/Runtime/Engine/Private/Components/AudioComponent.cpp
@@ -1133,11 +1133,7 @@ bool UAudioComponent::HasCookedFFTData() const
 {
 	if (Sound)
 	{
-		USoundWave* SoundWave = Cast<USoundWave>(Sound);
-		if (SoundWave)
-		{
-			return SoundWave->CookedSpectralTimeData.Num() > 0;
-		}
+		return Sound->HasCookedFFTData();
 	}
 	return false;
 }
@@ -1146,11 +1142,7 @@ bool UAudioComponent::HasCookedAmplitudeEnvelopeData() const
 {
 	if (Sound)
 	{
-		USoundWave* SoundWave = Cast<USoundWave>(Sound);
-		if (SoundWave)
-		{
-			return SoundWave->CookedEnvelopeTimeData.Num() > 0;
-		}
+		return Sound->HasCookedAmplitudeEnvelopeData();
 	}
 	return false;
 }
diff --git a/Engine/Source/Runtime/Engine/Private/SoundCue.cpp b/Engine/Source/Runtime/Engine/Private/SoundCue.cpp
--- a/Engine/Source/Runtime/Engine/Private/SoundCue.cpp
+++ b/Engine/Source/Runtime/Engine/Private/SoundCue.cpp
@@ -505,6 +505,40 @@ bool USoundCue::GetSoundWavesWithCookedAnalysisData(TArray<USoundWave*>& OutSoun
 	return bHasAnalysisData;
 }
 
+bool USoundCue::HasCookedFFTData() const
+{
+	// Check this sound cue's wave players to see if any of their soundwaves have cooked analysis data
+	TArray<const USoundNodeWavePlayer*> WavePlayers;
+	RecursiveFindNode<USoundNodeWavePlayer>(FirstNode, WavePlayers);
+
+	for (const USoundNodeWavePlayer* Player : WavePlayers)
+	{
+		const USoundWave* SoundWave = Player->GetSoundWave();
+		if (SoundWave && SoundWave->HasCookedFFTData())
+		{
+			return true;
+		}
+	}
+	return false;
+}
+
+bool USoundCue::HasCookedAmplitudeEnvelopeData() const
+{
+	// Check this sound cue's wave players to see if any of their soundwaves have cooked analysis data
+	TArray<const USoundNodeWavePlayer*> WavePlayers;
+	RecursiveFindNode<USoundNodeWavePlayer>(FirstNode, WavePlayers);
+
+	for (const USoundNodeWavePlayer* Player : WavePlayers)
+	{
+		const USoundWave* SoundWave = Player->GetSoundWave();
+		if (SoundWave && SoundWave->HasCookedAmplitudeEnvelopeData())
+		{
+			return true;
+		}
+	}
+	return false;
+}
+
 #if WITH_EDITOR
 UEdGraph* USoundCue::GetGraph()
 { 
diff --git a/Engine/Source/Runtime/Engine/Private/SoundWave.cpp b/Engine/Source/Runtime/Engine/Private/SoundWave.cpp
--- a/Engine/Source/Runtime/Engine/Private/SoundWave.cpp
+++ b/Engine/Source/Runtime/Engine/Private/SoundWave.cpp
@@ -1487,6 +1487,16 @@ bool USoundWave::GetSoundWavesWithCookedAnalysisData(TArray<USoundWave*>& OutSou
 	return false;
 }
 
+bool USoundWave::HasCookedFFTData() const
+{
+	return CookedSpectralTimeData.Num() > 0;
+}
+
+bool USoundWave::HasCookedAmplitudeEnvelopeData() const
+{
+	return CookedEnvelopeTimeData.Num() > 0;
+}
+
 void USoundWave::UpdatePlatformData()
 {
 	if (IsStreaming())
