Fix CIS warning (introduced by Dev-Anim 8 days ago)
#rnx

diff --git a/Engine/Source/Runtime/Engine/Private/Curves/RichCurve.cpp b/Engine/Source/Runtime/Engine/Private/Curves/RichCurve.cpp
--- a/Engine/Source/Runtime/Engine/Private/Curves/RichCurve.cpp
+++ b/Engine/Source/Runtime/Engine/Private/Curves/RichCurve.cpp
@@ -1378,10 +1378,7 @@ void FRichCurve::CompressCurve(FCompressedRichCurve& OutCurve, float ErrorThresh
 
 			const int32 KeySize = sizeof(uint16);
 
-			if (KeyFormat == RCKTCF_uint16)
-			{
-				WriteOffset = Align(WriteOffset, sizeof(uint16));
-			}
+			WriteOffset = Align(WriteOffset, sizeof(uint16));
 
 			uint8* KeyTimes8 = BasePtr + WriteOffset;
 			uint16* KeyTimes16 = reinterpret_cast<uint16*>(KeyTimes8);
