#ue4 - Fix cached rotator/quat conversion getting out of sync in rare cases where FScopedPreventAttachedComponentMove was used and the component rotation changed within the scope. This could happen for simulated characters running RootMotion.

#udn https://udn.unrealengine.com/questions/478413/but-in-smooth-mesh-locomotion-on-simulated-proxy.html

diff --git a/Engine/Source/Runtime/Engine/Private/Components/SceneComponent.cpp b/Engine/Source/Runtime/Engine/Private/Components/SceneComponent.cpp
--- a/Engine/Source/Runtime/Engine/Private/Components/SceneComponent.cpp
+++ b/Engine/Source/Runtime/Engine/Private/Components/SceneComponent.cpp
@@ -3309,7 +3309,7 @@ FScopedPreventAttachedComponentMove::~FScopedPreventAttachedComponentMove()
 			}
 			if (!bSavedAbsoluteRotation)
 			{
-				Owner->RelativeRotation = ChildRelativeTM.GetRotation().Rotator();
+				Owner->RelativeRotation = Owner->RelativeRotationCache.QuatToRotator(ChildRelativeTM.GetRotation());
 			}
 			if (!bSavedAbsoluteScale)
 			{
