[NO NON-PII COMMIT MESSAGE]

diff --git a/Engine/Source/Runtime/Engine/Classes/Sound/SoundWave.h b/Engine/Source/Runtime/Engine/Classes/Sound/SoundWave.h
--- a/Engine/Source/Runtime/Engine/Classes/Sound/SoundWave.h
+++ b/Engine/Source/Runtime/Engine/Classes/Sound/SoundWave.h
@@ -114,7 +114,7 @@ struct FStreamedAudioPlatformData
 	 *						either be NULL or have enough space for the chunk
 	 * @returns true if requested chunk has been loaded.
 	 */
-	bool TryLoadChunk(int32 ChunkIndex, uint8** OutChunkData);
+	bool TryLoadChunk(int32 ChunkIndex, uint8** OutChunkData, bool bMakeSureChunkIsLoaded = false);
 
 	/** Serialization. */
 	void Serialize(FArchive& Ar, class USoundWave* Owner);
@@ -780,7 +780,7 @@ public:
 	 * @param ChunkIndex	The Chunk index to cache.
 	 * @param OutChunkData	Address of pointer that will store data.
 	 */
-	bool GetChunkData(int32 ChunkIndex, uint8** OutChunkData);
+	bool GetChunkData(int32 ChunkIndex, uint8** OutChunkData, bool bMakeSureChunkIsLoaded = false);
 
 	void SetPrecacheState(ESoundWavePrecacheState InState)
 	{
diff --git a/Engine/Source/Runtime/Engine/Private/AudioDerivedData.cpp b/Engine/Source/Runtime/Engine/Private/AudioDerivedData.cpp
--- a/Engine/Source/Runtime/Engine/Private/AudioDerivedData.cpp
+++ b/Engine/Source/Runtime/Engine/Private/AudioDerivedData.cpp
@@ -643,8 +643,12 @@ FStreamedAudioPlatformData::~FStreamedAudioPlatformData()
 #endif
 }
 
-bool FStreamedAudioPlatformData::TryLoadChunk(int32 ChunkIndex, uint8** OutChunkData)
+bool FStreamedAudioPlatformData::TryLoadChunk(int32 ChunkIndex, uint8** OutChunkData, bool bMakeSureChunkIsLoaded /* = false */)
 {
+	// if bMakeSureChunkIsLoaded is true, we don't actually know the size of the chunk's bulk data,
+	// so it will need to be allocated in GetCopy.
+	check(!bMakeSureChunkIsLoaded || (OutChunkData && (*OutChunkData == nullptr)));
+
 	bool bCachedChunk = false;
 	FStreamedAudioChunk& Chunk = Chunks[ChunkIndex];
 
@@ -658,20 +662,17 @@ bool FStreamedAudioPlatformData::TryLoadChunk(int32 ChunkIndex, uint8** OutChunk
 	}
 #endif // #if WITH_EDITORONLY_DATA
 
-	// Load chunk from bulk data if available.
-	if (Chunk.BulkData.GetBulkDataSize() > 0)
+	// Load chunk from bulk data if available. If the chunk is not loaded, GetCopy will load it synchronously.
+	if (Chunk.BulkData.IsBulkDataLoaded() || bMakeSureChunkIsLoaded)
 	{
 		if (OutChunkData)
 		{
-			if (*OutChunkData == NULL)
-			{
-				*OutChunkData = static_cast<uint8*>(FMemory::Malloc(Chunk.BulkData.GetBulkDataSize()));
-			}
 			Chunk.BulkData.GetCopy((void**)OutChunkData);
 		}
 		bCachedChunk = true;
 	}
 
+
 #if WITH_EDITORONLY_DATA
 	// Wait for async DDC to complete
 	if (Chunk.DerivedDataKey.IsEmpty() == false)
diff --git a/Engine/Source/Runtime/Engine/Private/AudioStreaming.cpp b/Engine/Source/Runtime/Engine/Private/AudioStreaming.cpp
--- a/Engine/Source/Runtime/Engine/Private/AudioStreaming.cpp
+++ b/Engine/Source/Runtime/Engine/Private/AudioStreaming.cpp
@@ -137,8 +137,8 @@ bool FStreamingWaveData::Initialize(USoundWave* InSoundWave, FAudioStreamingMana
 	FLoadedAudioChunk* FirstChunk = &LoadedChunks[FirstLoadedChunkIndex];
 	FirstChunk->Index = 0;
 
-	// If we fail here, we'll just fail the streaming wave data altogether.
-	if (!SoundWave->GetChunkData(0, &FirstChunk->Data))
+	// Make sure we have loaded the 0th chunk before proceeding:
+	if (!SoundWave->GetChunkData(0, &FirstChunk->Data, true))
 	{
 		// Error/warning logging will have already been performed in the GetChunkData function
 		return false;
@@ -270,6 +270,7 @@ void FStreamingWaveData::BeginPendingRequests(const TArray<uint32>& IndicesToLoa
 		{
 			for (int32 ChunkIndex = 0; ChunkIndex < LoadedChunks.Num(); ++ChunkIndex)
 			{
+				check(Index != 0);
 				if (LoadedChunks[ChunkIndex].Index == Index)
 				{
 					FreeLoadedChunk(LoadedChunks[ChunkIndex]);
diff --git a/Engine/Source/Runtime/Engine/Private/SoundWave.cpp b/Engine/Source/Runtime/Engine/Private/SoundWave.cpp
--- a/Engine/Source/Runtime/Engine/Private/SoundWave.cpp
+++ b/Engine/Source/Runtime/Engine/Private/SoundWave.cpp
@@ -1567,16 +1567,16 @@ float USoundWave::GetSampleRateForCompressionOverrides(const FPlatformAudioCookO
 	}
 }
 
-bool USoundWave::GetChunkData(int32 ChunkIndex, uint8** OutChunkData)
+bool USoundWave::GetChunkData(int32 ChunkIndex, uint8** OutChunkData, bool bMakeSureChunkIsLoaded /* = false */)
 {
-	if (RunningPlatformData->TryLoadChunk(ChunkIndex, OutChunkData) == false)
+	if (RunningPlatformData->TryLoadChunk(ChunkIndex, OutChunkData, bMakeSureChunkIsLoaded) == false)
 	{
 #if WITH_EDITORONLY_DATA
 		// Unable to load chunks from the cache. Rebuild the sound and attempt to recache it.
 		UE_LOG(LogAudio, Display, TEXT("GetChunkData failed, rebuilding %s"), *GetPathName());
 
 		ForceRebuildPlatformData();
-		if (RunningPlatformData->TryLoadChunk(ChunkIndex, OutChunkData) == false)
+		if (RunningPlatformData->TryLoadChunk(ChunkIndex, OutChunkData, bMakeSureChunkIsLoaded) == false)
 		{
 			UE_LOG(LogAudio, Display, TEXT("Failed to build sound %s."), *GetPathName());
 		}
