#ue4 - Fix server logspam when player initially joins a server and hasn't moved yet.

diff --git a/Engine/Source/Runtime/Engine/Private/Components/CharacterMovementComponent.cpp b/Engine/Source/Runtime/Engine/Private/Components/CharacterMovementComponent.cpp
--- a/Engine/Source/Runtime/Engine/Private/Components/CharacterMovementComponent.cpp
+++ b/Engine/Source/Runtime/Engine/Private/Components/CharacterMovementComponent.cpp
@@ -7672,7 +7672,11 @@ bool UCharacterMovementComponent::ForcePositionUpdate(float DeltaTime)
 	ServerData->ServerAccumulatedClientTimeStamp += DeltaTime;
 
 #if !(UE_BUILD_SHIPPING)
-	UE_LOG(LogNetPlayerMovement, Warning, TEXT("ForcePositionUpdate %s (DeltaTime %.2f)"), *CharacterOwner->GetName(), DeltaTime);
+	const bool bServerMoveHasOccurred = (ServerData->ServerTimeStampLastServerMove != 0.f);
+	if (bServerMoveHasOccurred)
+	{
+		UE_LOG(LogNetPlayerMovement, Warning, TEXT("ForcePositionUpdate %s (DeltaTime %.2f)"), *CharacterOwner->GetName(), DeltaTime);
+	}
 #endif
 
 	// Force movement update.
