Speculative fix for crash in vorbis stream compressed info

diff --git a/Engine/Source/Runtime/Engine/Private/VorbisAudioInfo.cpp b/Engine/Source/Runtime/Engine/Private/VorbisAudioInfo.cpp
--- a/Engine/Source/Runtime/Engine/Private/VorbisAudioInfo.cpp
+++ b/Engine/Source/Runtime/Engine/Private/VorbisAudioInfo.cpp
@@ -483,6 +483,11 @@ void FVorbisAudioInfo::EnableHalfRate( bool HalfRate )
 
 bool FVorbisAudioInfo::StreamCompressedInfo(USoundWave* Wave, struct FSoundQualityInfo* QualityInfo)
 {
+	if (!bDllLoaded)
+	{
+		return;
+	}
+
 	SCOPE_CYCLE_COUNTER( STAT_VorbisPrepareDecompressionTime );
 
 	FScopeLock ScopeLock(&VorbisCriticalSection);
@@ -511,6 +516,11 @@ bool FVorbisAudioInfo::StreamCompressedInfo(USoundWave* Wave, struct FSoundQuali
 
 bool FVorbisAudioInfo::StreamCompressedData(uint8* InDestination, bool bLooping, uint32 BufferSize)
 {
+	if (!bDllLoaded)
+	{
+		return;
+	}
+
 	check( VFWrapper != NULL );
 
 #if PLATFORM_ANDROID
