#UE4 CIS fix

diff --git a/Engine/Source/Runtime/Engine/Private/VorbisAudioInfo.cpp b/Engine/Source/Runtime/Engine/Private/VorbisAudioInfo.cpp
--- a/Engine/Source/Runtime/Engine/Private/VorbisAudioInfo.cpp
+++ b/Engine/Source/Runtime/Engine/Private/VorbisAudioInfo.cpp
@@ -485,7 +485,7 @@ bool FVorbisAudioInfo::StreamCompressedInfo(USoundWave* Wave, struct FSoundQuali
 {
 	if (!bDllLoaded)
 	{
-		return;
+		return false;
 	}
 
 	SCOPE_CYCLE_COUNTER( STAT_VorbisPrepareDecompressionTime );
@@ -518,7 +518,7 @@ bool FVorbisAudioInfo::StreamCompressedData(uint8* InDestination, bool bLooping,
 {
 	if (!bDllLoaded)
 	{
-		return;
+		return false;
 	}
 
 	check( VFWrapper != NULL );
