Fixed warnings triggered when recovering from long spikes if timeouts are disabled
Cleanup of NetPacketNotify, some methods was exposed as public even thogh they should not have been

diff --git a/Engine/Source/Runtime/Engine/Private/Net/NetPacketNotify.cpp b/Engine/Source/Runtime/Engine/Private/Net/NetPacketNotify.cpp
--- a/Engine/Source/Runtime/Engine/Private/Net/NetPacketNotify.cpp
+++ b/Engine/Source/Runtime/Engine/Private/Net/NetPacketNotify.cpp
@@ -4,6 +4,12 @@
 #include "Serialization/BitWriter.h"
 #include "Serialization/BitReader.h"
 
+FNetPacketNotify::FNetPacketNotify()
+	: AckRecord(64)
+	, WrittenHistoryWordCount(0)
+{
+}
+
 SIZE_T FNetPacketNotify::GetCurrentSequenceHistoryLength() const
 {
 	if (InAckSeq >= InAckSeqAck)
@@ -17,6 +23,32 @@ SIZE_T FNetPacketNotify::GetCurrentSequenceHistoryLength() const
 	}
 }
 
+FNetPacketNotify::SequenceNumberT FNetPacketNotify::UpdateInAckSeqAck(SequenceNumberT::DifferenceT AckCount, SequenceNumberT AckedSeq)
+{
+	check((SIZE_T)AckCount <= AckRecord.Count());
+
+	if (AckCount > 1)
+	{
+		AckRecord.PopNoCheck(AckCount - 1);
+	}
+
+	FSentAckData AckData = AckRecord.PeekNoCheck();
+	AckRecord.PopNoCheck();
+
+	// verify that we have a matching sequence number
+	if (AckData.OutSeq == AckedSeq)
+	{
+		return AckData.InAckSeq;
+	}
+	else
+	{
+		UE_LOG_PACKET_NOTIFY_WARNING(TEXT("FNetPacketNotify::UpdateInAckSeqAck - Failed to find matching AckRecord for %u, (Found %u)"), AckedSeq.Get(), AckData.OutSeq.Get());
+
+		// Pessimistic view, should never occur
+		return SequenceNumberT(AckedSeq.Get() - MaxSequenceHistoryLength);
+	}
+}
+
 void FNetPacketNotify::Init(SequenceNumberT InitialInSeq, SequenceNumberT InitialOutSeq)
 {
 	InSeqHistory.Reset();
@@ -75,7 +107,7 @@ namespace
 }
 
 // These methods must always write and read the exact same number of bits, that is the reason for not using WriteInt/WrittedWrappedInt
-void FNetPacketNotify::WriteHeader(FBitWriter& Writer, bool bRefresh)
+bool FNetPacketNotify::WriteHeader(FBitWriter& Writer, bool bRefresh)
 {
 	// we always write at least 1 word
 	SIZE_T CurrentHistoryWordCount = FMath::Clamp<SIZE_T>((GetCurrentSequenceHistoryLength() + SequenceHistoryT::BitsPerWord - 1u) / SequenceHistoryT::BitsPerWord, 1u, SequenceHistoryT::WordCount);
@@ -83,7 +115,7 @@ void FNetPacketNotify::WriteHeader(FBitWriter& Writer, bool bRefresh)
 	// We can only do a refresh if we do not need more space for the history
 	if (bRefresh && (CurrentHistoryWordCount > WrittenHistoryWordCount))
 	{
-		return;
+		return false;
 	}
 
 	// How many words of ack data should we write? If this is a refresh we must write the same size as the original header
@@ -104,6 +136,8 @@ void FNetPacketNotify::WriteHeader(FBitWriter& Writer, bool bRefresh)
 	InSeqHistory.Write(Writer, WrittenHistoryWordCount);
 
 	UE_LOG_PACKET_NOTIFY(TEXT("FNetPacketNotify::WriteHeader - Seq %u, AckedSeq %u bReFresh %u HistorySizeInWords %u"), Seq, AckedSeq, bRefresh ? 1u : 0u, WrittenHistoryWordCount);
+
+	return true;
 }
 
 bool FNetPacketNotify::ReadHeader(FNotificationHeader& Data, FBitReader& Reader) const
diff --git a/Engine/Source/Runtime/Engine/Private/NetConnection.cpp b/Engine/Source/Runtime/Engine/Private/NetConnection.cpp
--- a/Engine/Source/Runtime/Engine/Private/NetConnection.cpp
+++ b/Engine/Source/Runtime/Engine/Private/NetConnection.cpp
@@ -1228,8 +1228,7 @@ void UNetConnection::FlushNet(bool bIgnoreSimulation)
 		// Increase outgoing sequence number
 		if (!InternalAck)
 		{
-			PacketNotify.CommitAndIncrementOutSeq();
-			HasDirtyAcks = 0u;
+			PacketNotify.CommitAndIncrementOutSeq();			
 		}
 		++OutPacketId; 
 
@@ -1386,12 +1385,18 @@ void UNetConnection::WritePacketHeader(FBitWriter& Writer)
 	Reset.PopWithoutClear(Writer);
 	
 	// Write notification header or refresh the header if used space is the same.
-	PacketNotify.WriteHeader(Writer, bIsHeaderUpdate);
+	bool bWroteHeader = PacketNotify.WriteHeader(Writer, bIsHeaderUpdate);
 
 	// Jump back to where we came from.
 	if (bIsHeaderUpdate)
 	{
 		Restore.PopWithoutClear(Writer);
+
+		// if we wrote the header and successfully refreshed the header status we no longer has any dirty acks
+		if (bWroteHeader)
+		{
+			HasDirtyAcks = 0u;
+		}
 	}
 }
 
@@ -2018,11 +2023,17 @@ PRAGMA_ENABLE_DEPRECATION_WARNINGS
 		++HasDirtyAcks;
 
 		// This is to allow us to recover from hitches were we process more than FNetPacketNotify::SequenceHistoryLength packets in a row withouht sending out any packets.
-		// In most cases this allows us to recover smoothly without reporting any pessimistic naks which is the case is we overshoot our ack history.
+		// In most cases this allows us to recover smoothly without reporting any pessimistic naks which will be the case if we overshoot the ack history.
+		// Note: This should only occur if we are running with no timeouts!
 		if (HasDirtyAcks >= FNetPacketNotify::MaxSequenceHistoryLength)
 		{
-			UE_LOG(LogNet, Warning, TEXT("UNetConnection::ReceivedPacket - Too many received acks(%u) since last sent packet. %s "), HasDirtyAcks, *Describe());
+			UE_LOG(LogNet, Warning, TEXT("UNetConnection::ReceivedPacket - Too many received packets to ack (%u) since last sent packet. InSeq: %u %s NextOutGoingSeq: %u"), HasDirtyAcks, PacketNotify.GetInSeq().Get(), *Describe(), PacketNotify.GetOutSeq().Get());
+
 			FlushNet();
+			if (HasDirtyAcks) // if acks still are dirty, flush again
+			{
+				FlushNet();
+			}
 		}
 	}
 }
@@ -3332,4 +3343,4 @@ static void	PrintActorReportFunc(const TArray<FString>& Args, UWorld* InWorld)
 	ObjKeyMap.Dump(*GLog);
 }
 
-FAutoConsoleCommandWithWorldAndArgs PrintActorReportCmd(TEXT("net.ActorReport"), TEXT(""),	FConsoleCommandWithWorldAndArgsDelegate::CreateStatic(PrintActorReportFunc) );
\ No newline at end of file
+FAutoConsoleCommandWithWorldAndArgs PrintActorReportCmd(TEXT("net.ActorReport"), TEXT(""),	FConsoleCommandWithWorldAndArgsDelegate::CreateStatic(PrintActorReportFunc) );
diff --git a/Engine/Source/Runtime/Engine/Public/Net/NetPacketNotify.h b/Engine/Source/Runtime/Engine/Public/Net/NetPacketNotify.h
--- a/Engine/Source/Runtime/Engine/Public/Net/NetPacketNotify.h
+++ b/Engine/Source/Runtime/Engine/Public/Net/NetPacketNotify.h
@@ -48,10 +48,7 @@ public:
 	};
 
 	/** Constructor */
-	FNetPacketNotify()
-	: AckRecord(64)
-	, WrittenHistoryWordCount(0)
-	{}
+	FNetPacketNotify();
 
 	/** Init notification with expected initial sequence numbers */
 	void Init(SequenceNumberT InitialInSeq, SequenceNumberT InitialOutSeq);
@@ -66,9 +63,10 @@ public:
 	SequenceNumberT CommitAndIncrementOutSeq();
 
 	/** Write NotificationHeader, and update outgoing ack record 
-	 if bRefresh is true we will attempt to refresh a previously written header 
-	if the size remains the same */
-	void WriteHeader(FBitWriter& Writer, bool bRefresh = false);
+		if bRefresh is true we will attempt to refresh a previously written header if the resulting size will be the same as the already written header.
+		returns true if data was written, and false if no data was written which might be the case if we try to rewrite an existing header but the required size differs.
+	*/
+	bool WriteHeader(FBitWriter& Writer, bool bRefresh = false);
 	
 	/** Read header from stream */
 	bool ReadHeader(FNotificationHeader& Data, FBitReader& Reader) const;
@@ -80,26 +78,7 @@ public:
 		Returns 0 if the received sequence number is out of the current window or if the acknowledged seq received by the remote is invalid.
 	*/
 	template<class Functor>
-	SequenceNumberT::DifferenceT Update(const FNotificationHeader& NotificationData, Functor&& InFunc)
-	{		
-		if (NotificationData.Seq > InSeq && NotificationData.AckedSeq >= OutAckSeq)
-		{
-			UE_LOG_PACKET_NOTIFY(TEXT("FNetPacketNotify::Update - Seq %u, InSeq %u"), NotificationData.Seq.Get(), InSeq.Get());
-
-			const SequenceNumberT::DifferenceT InSeqDelta = SequenceNumberT::Diff(NotificationData.Seq, InSeq);
-			
-			ProcessReceivedAcks(NotificationData, InFunc);
-
-			// accept sequence
-			InSeq = NotificationData.Seq;
-
-			return InSeqDelta;
-		}
-		else
-		{
-			return 0;
-		}
-	}
+	SequenceNumberT::DifferenceT Update(const FNotificationHeader& NotificationData, Functor&& InFunc);
 
 	/** Get the current SequenceHistory */
 	const SequenceHistoryT& GetInSeqHistory() const { return InSeqHistory; }
@@ -145,72 +124,78 @@ private:
 	SequenceNumberT OutAckSeq;			// Last sequence number that we know that the remote side have received.
 
 private:
-	SequenceNumberT UpdateInAckSeqAck(SequenceNumberT::DifferenceT AckCount, SequenceNumberT AckedSeq)
+	SequenceNumberT UpdateInAckSeqAck(SequenceNumberT::DifferenceT AckCount, SequenceNumberT AckedSeq);
+
+	template<class Functor>
+	inline void ProcessReceivedAcks(const FNotificationHeader& NotificationData, Functor&& InFunc);
+	void AckSeq(SequenceNumberT AckedSeq, bool IsAck);
+
+#if WITH_DEV_AUTOMATION_TESTS
+	friend struct FNetPacketNotifyTestUtil;
+#endif
+};
+
+template<class Functor>
+FNetPacketNotify::SequenceNumberT::DifferenceT FNetPacketNotify::Update(const FNotificationHeader& NotificationData, Functor&& InFunc)
+{
+	if (NotificationData.Seq > InSeq && NotificationData.AckedSeq >= OutAckSeq)
 	{
-		check((SIZE_T)AckCount <= AckRecord.Count());
+		UE_LOG_PACKET_NOTIFY(TEXT("FNetPacketNotify::Update - Seq %u, InSeq %u"), NotificationData.Seq.Get(), InSeq.Get());
 
-		if (AckCount > 1)
-		{
-			AckRecord.PopNoCheck(AckCount - 1);
-		}
-		
-		FSentAckData AckData = AckRecord.PeekNoCheck();
-		AckRecord.PopNoCheck();
+		const SequenceNumberT::DifferenceT InSeqDelta = SequenceNumberT::Diff(NotificationData.Seq, InSeq);
 
-		// verify that we have a matching sequence number
-		if (AckData.OutSeq == AckedSeq)
-		{
-			return AckData.InAckSeq;
-		}
-		else
-		{
-			UE_LOG_PACKET_NOTIFY_WARNING(TEXT("FNetPacketNotify::UpdateInAckSeqAck - Failed to find matching AckRecord for %u, (Found %u)"), AckedSeq.Get(), AckData.OutSeq.Get());
+		ProcessReceivedAcks(NotificationData, InFunc);
 
-			// Pessimistic view, should never occur
-			return SequenceNumberT(AckedSeq.Get() - MaxSequenceHistoryLength);
-		}
+		// accept sequence
+		InSeq = NotificationData.Seq;
+
+		return InSeqDelta;
+	}
+	else
+	{
+		return 0;
 	}
+}
 
-	template<class Functor>
-	void ProcessReceivedAcks(const FNotificationHeader& NotificationData, Functor&& InFunc)
+template<class Functor>
+void FNetPacketNotify::ProcessReceivedAcks(const FNotificationHeader& NotificationData, Functor&& InFunc)
+{
+	if (NotificationData.AckedSeq > OutAckSeq)
 	{
-		if (NotificationData.AckedSeq > OutAckSeq)
-		{
-			UE_LOG_PACKET_NOTIFY(TEXT("Notification::ProcessReceivedAcks - AckedSeq: %u, OutAckSeq: %u"), NotificationData.AckedSeq.Get(), OutAckSeq.Get());
+		UE_LOG_PACKET_NOTIFY(TEXT("Notification::ProcessReceivedAcks - AckedSeq: %u, OutAckSeq: %u"), NotificationData.AckedSeq.Get(), OutAckSeq.Get());
+
+		SequenceNumberT::DifferenceT AckCount = SequenceNumberT::Diff(NotificationData.AckedSeq, OutAckSeq);
+
+		// Update InAckSeqAck used to track the needed number of bits to transmit our ack history
+		InAckSeqAck = UpdateInAckSeqAck(AckCount, NotificationData.AckedSeq);
 
-			SequenceNumberT::DifferenceT AckCount = SequenceNumberT::Diff(NotificationData.AckedSeq, OutAckSeq);
+		// ExpectedAck = OutAckSeq + 1
+		SequenceNumberT CurrentAck(OutAckSeq);
+		++CurrentAck;
 
-			// Update InAckSeqAck used to track the needed number of bits to transmit our ack history
-			InAckSeqAck = UpdateInAckSeqAck(AckCount, NotificationData.AckedSeq);
+		if (AckCount > (SequenceNumberT::DifferenceT)(SequenceHistoryT::Size))
+		{
+			UE_LOG_PACKET_NOTIFY_WARNING(TEXT("Notification::ProcessReceivedAcks - Missed Acks: AckedSeq: %u, OutAckSeq: %u, FirstMissingSeq: %u Count: %u"), NotificationData.AckedSeq.Get(), OutAckSeq.Get(), CurrentAck.Get(), AckCount - (SequenceNumberT::DifferenceT)(SequenceHistoryT::Size));
+		}
 
-			// ExpectedAck = OutAckSeq + 1
-			SequenceNumberT CurrentAck(OutAckSeq);
+		// Everything not found in the history buffer is treated as lost
+		while (AckCount > (SequenceNumberT::DifferenceT)(SequenceHistoryT::Size))
+		{
+			--AckCount;
+			InFunc(CurrentAck, false);
 			++CurrentAck;
+		}
 
-			// Everything not found in the history buffer is treated as lost
-			while (AckCount > (SequenceNumberT::DifferenceT)(SequenceHistoryT::Size))
-			{
-				--AckCount;
-				UE_LOG_PACKET_NOTIFY_WARNING(TEXT("Notification::ProcessReceivedAcks Seq: %u - IsAck: 0 HistoryIndex: N/A"), CurrentAck.Get());
-				InFunc(CurrentAck, false);
-				++CurrentAck;
-			}
-
-			// For sequence numbers contained in the history we lookup the delivery status from the history
-			while (AckCount > 0)
-			{
-				--AckCount;
-				UE_LOG_PACKET_NOTIFY(TEXT("Notification::ProcessReceivedAcks Seq: %u - IsAck: %u HistoryIndex: %u"), CurrentAck.Get(), NotificationData.History.IsDelivered(AckCount) ? 1u : 0u, AckCount);
-				InFunc(CurrentAck, NotificationData.History.IsDelivered(AckCount));
-				++CurrentAck;
-			}
-			OutAckSeq = NotificationData.AckedSeq;
+		// For sequence numbers contained in the history we lookup the delivery status from the history
+		while (AckCount > 0)
+		{
+			--AckCount;
+			UE_LOG_PACKET_NOTIFY(TEXT("Notification::ProcessReceivedAcks Seq: %u - IsAck: %u HistoryIndex: %u"), CurrentAck.Get(), NotificationData.History.IsDelivered(AckCount) ? 1u : 0u, AckCount);
+			InFunc(CurrentAck, NotificationData.History.IsDelivered(AckCount));
+			++CurrentAck;
 		}
+		OutAckSeq = NotificationData.AckedSeq;
 	}
+}
 
-	void AckSeq(SequenceNumberT AckedSeq, bool IsAck);
 
-#if WITH_DEV_AUTOMATION_TESTS
-	friend struct FNetPacketNotifyTestUtil;
-#endif
-};
