Fixed one-off bone visibility changes causing flicker

diff --git a/Engine/Source/Runtime/Engine/Classes/Components/SkinnedMeshComponent.h b/Engine/Source/Runtime/Engine/Classes/Components/SkinnedMeshComponent.h
--- a/Engine/Source/Runtime/Engine/Classes/Components/SkinnedMeshComponent.h
+++ b/Engine/Source/Runtime/Engine/Classes/Components/SkinnedMeshComponent.h
@@ -518,6 +518,9 @@ protected:
 	UPROPERTY(Transient)
 	mutable uint8 bCachedLocalBoundsUpToDate:1;
 
+	/** Whether we have updated bone visibility this tick */
+	uint8 bBoneVisibilityDirty:1;
+
 private:
 	/** If true, UpdateTransform will always result in a call to MeshObject->Update. */
 	UPROPERTY(transient)
diff --git a/Engine/Source/Runtime/Engine/Private/Components/SkinnedMeshComponent.cpp b/Engine/Source/Runtime/Engine/Private/Components/SkinnedMeshComponent.cpp
--- a/Engine/Source/Runtime/Engine/Private/Components/SkinnedMeshComponent.cpp
+++ b/Engine/Source/Runtime/Engine/Private/Components/SkinnedMeshComponent.cpp
@@ -348,6 +348,7 @@ USkinnedMeshComponent::USkinnedMeshComponent(const FObjectInitializer& ObjectIni
 	CurrentEditableComponentTransforms = 0;
 	CurrentReadComponentTransforms = 1;
 	bNeedToFlipSpaceBaseBuffers = false;
+	bBoneVisibilityDirty = false;
 
 	bCanEverAffectNavigation = false;
 	MasterBoneMapCacheCount = 0;
@@ -944,6 +945,8 @@ void USkinnedMeshComponent::RebuildVisibilityArray()
 				}
 			}
 		}
+
+		bBoneVisibilityDirty = true;
 	}
 }
 
@@ -2798,6 +2801,13 @@ void USkinnedMeshComponent::FlipEditableSpaceBases()
 			{
 				GetEditableComponentSpaceTransforms() = GetComponentSpaceTransforms();
 				GetEditableBoneVisibilityStates() = GetBoneVisibilityStates();
+				bBoneVisibilityDirty = false;
+			}
+			// If we have changed bone visibility, then we need to reflect that next frame
+			else if(bBoneVisibilityDirty)
+			{
+				GetEditableBoneVisibilityStates() = GetBoneVisibilityStates();
+				bBoneVisibilityDirty = false;
 			}
 		}
 		else
