TBA: Migrated a couple of UObject Serialize functions to FStructuredArchive (SoundCue / MaterialExpressions / Editor strip flags)

diff --git a/Engine/Source/Runtime/Engine/Classes/Sound/SoundCue.h b/Engine/Source/Runtime/Engine/Classes/Sound/SoundCue.h
--- a/Engine/Source/Runtime/Engine/Classes/Sound/SoundCue.h
+++ b/Engine/Source/Runtime/Engine/Classes/Sound/SoundCue.h
@@ -130,7 +130,7 @@ public:
 	static void AddReferencedObjects(UObject* InThis, FReferenceCollector& Collector);
 #endif
 	virtual void PostLoad() override;
-	virtual void Serialize(FArchive& Ar) override;
+	virtual void Serialize(FStructuredArchive::FRecord Record) override;
 	//~ End UObject Interface.
 
 	//~ Begin USoundBase Interface.
diff --git a/Engine/Source/Runtime/Engine/Private/EngineUtils.cpp b/Engine/Source/Runtime/Engine/Private/EngineUtils.cpp
--- a/Engine/Source/Runtime/Engine/Private/EngineUtils.cpp
+++ b/Engine/Source/Runtime/Engine/Private/EngineUtils.cpp
@@ -517,3 +517,50 @@ FStripDataFlags::FStripDataFlags( class FArchive& Ar, uint8 InGlobalFlags, uint8
 		Ar << ClassStripFlags;
 	}
 }
+
+/*-----------------------------------------------------------------------------
+Serialized data stripping.
+-----------------------------------------------------------------------------*/
+FStripDataFlags::FStripDataFlags(FStructuredArchive::FSlot Slot, uint8 InClassFlags /*= 0*/, int32 InVersion /*= VER_UE4_OLDEST_LOADABLE_PACKAGE */)
+	: GlobalStripFlags(0)
+	, ClassStripFlags(0)
+{
+	FArchive& UnderlyingArchive = Slot.GetUnderlyingArchive();
+	FStructuredArchive::FRecord Record = Slot.EnterRecord();
+
+	check(InVersion >= VER_UE4_OLDEST_LOADABLE_PACKAGE);
+	if (UnderlyingArchive.UE4Ver() >= InVersion)
+	{
+		if (UnderlyingArchive.IsCooking())
+		{
+			// When cooking GlobalStripFlags are automatically generated based on the current target
+			// platform's properties.
+			GlobalStripFlags |= UnderlyingArchive.CookingTarget()->HasEditorOnlyData() ? FStripDataFlags::None : FStripDataFlags::Editor;
+			GlobalStripFlags |= UnderlyingArchive.CookingTarget()->IsServerOnly() ? FStripDataFlags::Server : FStripDataFlags::None;
+			ClassStripFlags = InClassFlags;
+		}
+		Record << NAMED_FIELD(GlobalStripFlags);
+		Record << NAMED_FIELD(ClassStripFlags);
+	}
+}
+
+FStripDataFlags::FStripDataFlags(FStructuredArchive::FSlot Slot, uint8 InGlobalFlags, uint8 InClassFlags, int32 InVersion /*= VER_UE4_OLDEST_LOADABLE_PACKAGE */)
+	: GlobalStripFlags(0)
+	, ClassStripFlags(0)
+{
+	FArchive& UnderlyingArchive = Slot.GetUnderlyingArchive();
+	FStructuredArchive::FRecord Record = Slot.EnterRecord();
+
+	check(InVersion >= VER_UE4_OLDEST_LOADABLE_PACKAGE);
+	if (UnderlyingArchive.UE4Ver() >= InVersion)
+	{
+		if (UnderlyingArchive.IsCooking())
+		{
+			// Don't generate global strip flags and use the ones passed in by the caller.
+			GlobalStripFlags = InGlobalFlags;
+			ClassStripFlags = InClassFlags;
+		}
+		Record << NAMED_FIELD(GlobalStripFlags);
+		Record << NAMED_FIELD(ClassStripFlags);
+	}
+}
diff --git a/Engine/Source/Runtime/Engine/Private/SoundCue.cpp b/Engine/Source/Runtime/Engine/Private/SoundCue.cpp
--- a/Engine/Source/Runtime/Engine/Private/SoundCue.cpp
+++ b/Engine/Source/Runtime/Engine/Private/SoundCue.cpp
@@ -93,30 +93,32 @@ void USoundCue::CacheAggregateValues()
 	}
 }
 
-void USoundCue::Serialize(FArchive& Ar)
+void USoundCue::Serialize(FStructuredArchive::FRecord Record)
 {
+	FArchive& UnderlyingArchive = Record.GetUnderlyingArchive();
+
 	// Always force the duration to be updated when we are saving or cooking
-	if (Ar.IsSaving() || Ar.IsCooking())
+	if (UnderlyingArchive.IsSaving() || UnderlyingArchive.IsCooking())
 	{
 		CacheAggregateValues();
 	}
 
-	Super::Serialize(Ar);
+	Super::Serialize(Record);
 
-	if (Ar.UE4Ver() >= VER_UE4_COOKED_ASSETS_IN_EDITOR_SUPPORT)
+	if (UnderlyingArchive.UE4Ver() >= VER_UE4_COOKED_ASSETS_IN_EDITOR_SUPPORT)
 	{
-		FStripDataFlags StripFlags(Ar);
+		FStripDataFlags StripFlags(Record.EnterField(FIELD_NAME_TEXT("SoundCueStripFlags")));
 #if WITH_EDITORONLY_DATA
 		if (!StripFlags.IsEditorDataStripped())
 		{
-			Ar << SoundCueGraph;
+			Record << NAMED_FIELD(SoundCueGraph);
 		}
 #endif
 	}
 #if WITH_EDITOR
 	else
 	{
-		Ar << SoundCueGraph;
+		Record << NAMED_FIELD(SoundCueGraph);
 	}
 #endif
 }
diff --git a/Engine/Source/Runtime/Engine/Public/EngineUtils.h b/Engine/Source/Runtime/Engine/Public/EngineUtils.h
--- a/Engine/Source/Runtime/Engine/Public/EngineUtils.h
+++ b/Engine/Source/Runtime/Engine/Public/EngineUtils.h
@@ -852,6 +852,28 @@ public:
 	 */
 	FStripDataFlags( FArchive& Ar, uint8 InGlobalFlags, uint8 InClassFlags, int32 InVersion = VER_UE4_OLDEST_LOADABLE_PACKAGE );
 
+	/**
+	* Constructor.
+	* Serializes strip data flags. Global (engine) flags are automatically generated from target platform
+	* when saving. Class flags need to be defined by the user.
+	*
+	* @param Ar - Archive to serialize with.
+	* @param InClassFlags - User defined per class flags .
+	* @param InVersion - Minimal strip version required to serialize strip flags
+	*/
+	FStripDataFlags(FStructuredArchive::FSlot Slot, uint8 InClassFlags = 0, int32 InVersion = VER_UE4_OLDEST_LOADABLE_PACKAGE);
+
+	/**
+	* Constructor.
+	* Serializes strip data flags. Global (engine) flags are user defined and will not be automatically generated
+	* when saving. Class flags also need to be defined by the user.
+	*
+	* @param Ar - Archive to serialize with.
+	* @param InClassFlags - User defined per class flags.
+	* @param InVersion - Minimal version required to serialize strip flags
+	*/
+	FStripDataFlags(FStructuredArchive::FSlot Slot, uint8 InGlobalFlags, uint8 InClassFlags, int32 InVersion = VER_UE4_OLDEST_LOADABLE_PACKAGE);
+
 	/**
 	 * Checks if FStripDataFlags::Editor flag is set or not
 	 *
diff --git a/Engine/Source/Runtime/Landscape/Classes/Materials/MaterialExpressionLandscapeLayerBlend.h b/Engine/Source/Runtime/Landscape/Classes/Materials/MaterialExpressionLandscapeLayerBlend.h
--- a/Engine/Source/Runtime/Landscape/Classes/Materials/MaterialExpressionLandscapeLayerBlend.h
+++ b/Engine/Source/Runtime/Landscape/Classes/Materials/MaterialExpressionLandscapeLayerBlend.h
@@ -73,7 +73,7 @@ class LANDSCAPE_API UMaterialExpressionLandscapeLayerBlend : public UMaterialExp
 
 
 	//~ Begin UObject Interface
-	virtual void Serialize(FArchive& Ar) override;
+	virtual void Serialize(FStructuredArchive::FRecord Record) override;
 	virtual bool NeedsLoadForClient() const override;
 #if WITH_EDITOR
 	virtual void PostEditChangeProperty(FPropertyChangedEvent& PropertyChangedEvent) override;
diff --git a/Engine/Source/Runtime/Landscape/Classes/Materials/MaterialExpressionLandscapeLayerSwitch.h b/Engine/Source/Runtime/Landscape/Classes/Materials/MaterialExpressionLandscapeLayerSwitch.h
--- a/Engine/Source/Runtime/Landscape/Classes/Materials/MaterialExpressionLandscapeLayerSwitch.h
+++ b/Engine/Source/Runtime/Landscape/Classes/Materials/MaterialExpressionLandscapeLayerSwitch.h
@@ -38,7 +38,7 @@ class LANDSCAPE_API UMaterialExpressionLandscapeLayerSwitch : public UMaterialEx
 public:
 
 	//~ Begin UObject Interface
-	virtual void Serialize(FArchive& Ar) override;
+	virtual void Serialize(FStructuredArchive::FRecord Record) override;
 	//~ End UObject Interface
 
 	//~ Begin UMaterialExpression Interface
diff --git a/Engine/Source/Runtime/Landscape/Private/Materials/MaterialExpressionLandscapeLayerBlend.cpp b/Engine/Source/Runtime/Landscape/Private/Materials/MaterialExpressionLandscapeLayerBlend.cpp
--- a/Engine/Source/Runtime/Landscape/Private/Materials/MaterialExpressionLandscapeLayerBlend.cpp
+++ b/Engine/Source/Runtime/Landscape/Private/Materials/MaterialExpressionLandscapeLayerBlend.cpp
@@ -42,11 +42,12 @@ FGuid& UMaterialExpressionLandscapeLayerBlend::GetParameterExpressionId()
 }
 
 
-void UMaterialExpressionLandscapeLayerBlend::Serialize(FArchive& Ar)
+void UMaterialExpressionLandscapeLayerBlend::Serialize(FStructuredArchive::FRecord Record)
 {
-	Super::Serialize(Ar);
+	Super::Serialize(Record);
+	FArchive& UnderlyingArchive = Record.GetUnderlyingArchive();
 
-	if (Ar.IsLoading() && Ar.UE4Ver() < VER_UE4_ADD_LB_WEIGHTBLEND)
+	if (UnderlyingArchive.IsLoading() && UnderlyingArchive.UE4Ver() < VER_UE4_ADD_LB_WEIGHTBLEND)
 	{
 		// convert any LB_AlphaBlend entries to LB_WeightBlend
 		for (FLayerBlendInput& LayerInput : Layers)
diff --git a/Engine/Source/Runtime/Landscape/Private/Materials/MaterialExpressionLandscapeLayerSwitch.cpp b/Engine/Source/Runtime/Landscape/Private/Materials/MaterialExpressionLandscapeLayerSwitch.cpp
--- a/Engine/Source/Runtime/Landscape/Private/Materials/MaterialExpressionLandscapeLayerSwitch.cpp
+++ b/Engine/Source/Runtime/Landscape/Private/Materials/MaterialExpressionLandscapeLayerSwitch.cpp
@@ -92,11 +92,11 @@ void UMaterialExpressionLandscapeLayerSwitch::GetCaption(TArray<FString>& OutCap
 }
 #endif // WITH_EDITOR
 
-void UMaterialExpressionLandscapeLayerSwitch::Serialize(FArchive& Ar)
+void UMaterialExpressionLandscapeLayerSwitch::Serialize(FStructuredArchive::FRecord Record)
 {
-	Super::Serialize(Ar);
+	Super::Serialize(Record);
 
-	if (Ar.UE4Ver() < VER_UE4_FIX_TERRAIN_LAYER_SWITCH_ORDER)
+	if (Record.GetUnderlyingArchive().UE4Ver() < VER_UE4_FIX_TERRAIN_LAYER_SWITCH_ORDER)
 	{
 		Swap(LayerUsed, LayerNotUsed);
 	}
