Static analysis warning fix

diff --git a/Engine/Source/Runtime/Engine/Private/GameFramework/HapticFeedbackEffect.cpp b/Engine/Source/Runtime/Engine/Private/GameFramework/HapticFeedbackEffect.cpp
--- a/Engine/Source/Runtime/Engine/Private/GameFramework/HapticFeedbackEffect.cpp
+++ b/Engine/Source/Runtime/Engine/Private/GameFramework/HapticFeedbackEffect.cpp
@@ -175,11 +175,12 @@ void UHapticFeedbackEffect_SoundWave::PrepareSoundWaveBuffer()
 	}
 	AD->Precache(SoundWave, true, false);
 	SoundWave->InitAudioResource(AD->GetRuntimeFormat(SoundWave));
-	check(!SoundWave->RawPCMData || SoundWave->RawPCMDataSize);
 	uint8* PCMData = SoundWave->RawPCMData;
+	int32 RawPCMDataSize = SoundWave->RawPCMDataSize;
+	check((PCMData != nullptr) || (RawPCMDataSize == 0));
 	int32 SampleRate = SoundWave->GetSampleRateForCurrentPlatform();
 	int TargetFrequency = 320;
-	int TargetBufferSize = (SoundWave->RawPCMDataSize * TargetFrequency) / (SampleRate * 2) + 1; //2 because we're only using half of the 16bit source PCM buffer
+	int TargetBufferSize = (RawPCMDataSize * TargetFrequency) / (SampleRate * 2) + 1; //2 because we're only using half of the 16bit source PCM buffer
 	HapticBuffer.BufferLength = TargetBufferSize;
 	HapticBuffer.RawData.AddUninitialized(TargetBufferSize);
 	HapticBuffer.CurrentPtr = 0;
@@ -187,7 +188,7 @@ void UHapticFeedbackEffect_SoundWave::PrepareSoundWaveBuffer()
 
 	int previousTargetIndex = -1;
 	int currentMin = 0;
-	for (int i = 1; i < SoundWave->RawPCMDataSize; i += 2)
+	for (int i = 1; i < RawPCMDataSize; i += 2)
 	{
 		int targetIndex = i * TargetFrequency / (SampleRate * 2);
 		int val = PCMData[i];
