Merging //UE4/Dev-Main @ 4118128 to Dev-Core (//UE4/Dev-Core)

diff --git a/Engine/Source/Runtime/ApplicationCore/Private/HTML5/HTML5InputInterface.cpp b/Engine/Source/Runtime/ApplicationCore/Private/HTML5/HTML5InputInterface.cpp
--- a/Engine/Source/Runtime/ApplicationCore/Private/HTML5/HTML5InputInterface.cpp
+++ b/Engine/Source/Runtime/ApplicationCore/Private/HTML5/HTML5InputInterface.cpp
@@ -100,8 +100,11 @@ void FHTML5InputInterface::Tick(float DeltaTime, const SDL_Event& Event,TSharedR
 				const SDL_Keycode KeyCode = KeyEvent.keysym.scancode;
 				const bool bIsRepeated = KeyEvent.repeat != 0;
 
-				// First KeyDown, then KeyChar. This is important, as in-game console ignores first character otherwise
-				MessageHandler->OnKeyDown(KeyCode, KeyEvent.keysym.sym, bIsRepeated);
+				if ( KeyCode != SDL_SCANCODE_F5 ) // UE-58440 -- note: this will be removed when new toolchain comes in
+				{
+					// First KeyDown, then KeyChar. This is important, as in-game console ignores first character otherwise
+					MessageHandler->OnKeyDown(KeyCode, KeyEvent.keysym.sym, bIsRepeated);
+				}
 
 				// Backspace/Return input caught here.
 				// Note that TextInput still seems to get characters messages too but slate seems to not process them.
@@ -120,7 +123,10 @@ void FHTML5InputInterface::Tick(float DeltaTime, const SDL_Event& Event,TSharedR
 				const SDL_Keycode KeyCode = keyEvent.keysym.scancode;
 				const bool IsRepeat = keyEvent.repeat != 0;
 
-				MessageHandler->OnKeyUp( KeyCode, keyEvent.keysym.sym, IsRepeat );
+				if ( KeyCode != SDL_SCANCODE_F5 ) // UE-58440 -- note: this will be removed when new toolchain comes in
+				{
+					MessageHandler->OnKeyUp( KeyCode, keyEvent.keysym.sym, IsRepeat );
+				}
 				UE_LOG(LogHTML5Input, Verbose, TEXT("KeyUp Code:%d"), KeyCode);
 			}
 			break;
diff --git a/Engine/Source/Runtime/Engine/Private/InstancedStaticMesh.cpp b/Engine/Source/Runtime/Engine/Private/InstancedStaticMesh.cpp
--- a/Engine/Source/Runtime/Engine/Private/InstancedStaticMesh.cpp
+++ b/Engine/Source/Runtime/Engine/Private/InstancedStaticMesh.cpp
@@ -236,7 +236,7 @@ void FStaticMeshInstanceBuffer::InitRHI()
 		QUICK_SCOPE_CYCLE_COUNTER(STAT_FStaticMeshInstanceBuffer_InitRHI);
 		auto AccessFlags = BUF_Static;
 		CreateVertexBuffer(InstanceData->GetOriginResourceArray(), AccessFlags | BUF_ShaderResource, 16, PF_A32B32G32R32F, InstanceOriginBuffer.VertexBufferRHI, InstanceOriginSRV);
-		CreateVertexBuffer(InstanceData->GetTransformResourceArray(), AccessFlags | BUF_ShaderResource, InstanceData->GetTranslationUsesHalfs() ? 8 : 16, InstanceData->GetTranslationUsesHalfs() ? PF_FloatRGBA : PF_A32B32G32R32F, InstanceTransformBuffer.VertexBufferRHI, InstanceTransformSRV);
+		CreateVertexBuffer(InstanceData->GetTransformResourceArray(), AccessFlags | BUF_ShaderResource, InstanceData->GetTranslationUsesHalfs() && GVertexElementTypeSupport.IsSupported(VET_Half2) ? 8 : 16, InstanceData->GetTranslationUsesHalfs() && GVertexElementTypeSupport.IsSupported(VET_Half2) ? PF_FloatRGBA : PF_A32B32G32R32F, InstanceTransformBuffer.VertexBufferRHI, InstanceTransformSRV);
 		CreateVertexBuffer(InstanceData->GetLightMapResourceArray(), AccessFlags | BUF_ShaderResource, 8, PF_R16G16B16A16_SNORM, InstanceLightmapBuffer.VertexBufferRHI, InstanceLightmapSRV);
 	}
 }
@@ -305,8 +305,8 @@ void FStaticMeshInstanceBuffer::BindInstanceVertexBuffer(const class FVertexFact
 			EVertexStreamUsage::ManualFetch | EVertexStreamUsage::Instancing
 		);
 
-		EVertexElementType TransformType = InstanceData->GetTranslationUsesHalfs() ? VET_Half4 : VET_Float4;
-		uint32 TransformStride = InstanceData->GetTranslationUsesHalfs() ? 8 : 16;
+		EVertexElementType TransformType = InstanceData->GetTranslationUsesHalfs() && GVertexElementTypeSupport.IsSupported(VET_Half2) ? VET_Half4 : VET_Float4;
+		uint32 TransformStride = InstanceData->GetTranslationUsesHalfs() && GVertexElementTypeSupport.IsSupported(VET_Half2) ? 8 : 16;
 
 		InstancedStaticMeshData.InstanceTransformComponent[0] = FVertexStreamComponent(
 			&InstanceTransformBuffer,
