AnimScriptInstances created at runtime will now also be added to the owning omponent's cluster to avoid GC issues.

diff --git a/Engine/Source/Runtime/Engine/Private/Components/SkeletalMeshComponent.cpp b/Engine/Source/Runtime/Engine/Private/Components/SkeletalMeshComponent.cpp
--- a/Engine/Source/Runtime/Engine/Private/Components/SkeletalMeshComponent.cpp
+++ b/Engine/Source/Runtime/Engine/Private/Components/SkeletalMeshComponent.cpp
@@ -743,6 +743,11 @@ bool USkeletalMeshComponent::InitializeAnimScriptInstance(bool bForceReinit)
 					// otherwise, initialize with AnimationData
 					AnimationData.Initialize(Cast<UAnimSingleNodeInstance>(AnimScriptInstance));
 				}
+
+				if (AnimScriptInstance)
+				{
+					AnimScriptInstance->AddToCluster(this);
+				}
 			}
 		}
 
