The engine will no longer be resetting all linkers causing massive load times when renaming the world package when entering Play In Editor

diff --git a/Engine/Source/Runtime/Engine/Private/World.cpp b/Engine/Source/Runtime/Engine/Private/World.cpp
--- a/Engine/Source/Runtime/Engine/Private/World.cpp
+++ b/Engine/Source/Runtime/Engine/Private/World.cpp
@@ -2577,7 +2577,7 @@ void UWorld::RenameToPIEWorld(int32 PIEInstanceID)
 	WorldPackage->SetPackageFlags(PKG_PlayInEditor);
 
 	const FString PIEPackageName = *UWorld::ConvertToPIEPackageName(WorldPackage->GetName(), PIEInstanceID);
-	WorldPackage->Rename(*PIEPackageName);
+	WorldPackage->Rename(*PIEPackageName, nullptr, REN_ForceNoResetLoaders);
 	FSoftObjectPath::AddPIEPackageName(FName(*PIEPackageName));
 
 	StreamingLevelsPrefix = UWorld::BuildPIEPackagePrefix(PIEInstanceID);
