Simplified TBitArray::operator<.
Moved the original implementation to SkeletalMeshLODRenderData.cpp where the more complex sorting behavior is required, and fixed the buffer overrun.

diff --git a/Engine/Source/Runtime/Core/Public/Containers/BitArray.h b/Engine/Source/Runtime/Core/Public/Containers/BitArray.h
--- a/Engine/Source/Runtime/Core/Public/Containers/BitArray.h
+++ b/Engine/Source/Runtime/Core/Public/Containers/BitArray.h
@@ -258,21 +258,8 @@ public:
 		const uint32* Data0 = GetData();
 		const uint32* Data1 = Other.GetData();
 
-		//sort by num bits active
-		int32 Count0 = 0, Count1 = 0;
+		//lexicographically compare
 		for (uint32 i = 0; i < NumWords; i++)
-		{
-			Count0 += FPlatformMath::CountBits(Data0[i]);
-			Count1 += FPlatformMath::CountBits(Data1[i]);
-		}
-
-		if (Count0 != Count1)
-		{
-			return Count0 < Count1;
-		}
-
-		//sort by big-num value
-		for (uint32 i = NumWords; i != ~0u; i--)
 		{
 			if (Data0[i] != Data1[i])
 			{
diff --git a/Engine/Source/Runtime/Engine/Private/SkeletalMeshLODRenderData.cpp b/Engine/Source/Runtime/Engine/Private/SkeletalMeshLODRenderData.cpp
--- a/Engine/Source/Runtime/Engine/Private/SkeletalMeshLODRenderData.cpp
+++ b/Engine/Source/Runtime/Engine/Private/SkeletalMeshLODRenderData.cpp
@@ -22,6 +22,48 @@ static FAutoConsoleVariableRef CVarStripSkeletalMeshLodsBelowMinLod(
 	TEXT("If set will strip skeletal mesh LODs under the minimum renderable LOD for the target platform during cooking.")
 );
 
+namespace
+{
+	struct FReverseOrderBitArraysBySetBits
+	{
+		FORCEINLINE bool operator()(const TBitArray<>& Lhs, const TBitArray<>& Rhs) const
+		{
+			//sort by length
+			if (Lhs.Num() != Rhs.Num())
+			{
+				return Lhs.Num() > Rhs.Num();
+			}
+
+			uint32 NumWords = FMath::DivideAndRoundUp(Lhs.Num(), NumBitsPerDWORD);
+			const uint32* Data0 = Lhs.GetData();
+			const uint32* Data1 = Rhs.GetData();
+
+			//sort by num bits active
+			int32 Count0 = 0, Count1 = 0;
+			for (uint32 i = 0; i < NumWords; i++)
+			{
+				Count0 += FPlatformMath::CountBits(Data0[i]);
+				Count1 += FPlatformMath::CountBits(Data1[i]);
+			}
+
+			if (Count0 != Count1)
+			{
+				return Count0 > Count1;
+			}
+
+			//sort by big-num value
+			for (uint32 i = NumWords - 1; i != ~0u; i--)
+			{
+				if (Data0[i] != Data1[i])
+				{
+					return Data0[i] > Data1[i];
+				}
+			}
+			return false;
+		}
+	};
+}
+
 // Serialization.
 FArchive& operator<<(FArchive& Ar, FSkelMeshRenderSection& S)
 {
@@ -283,7 +325,7 @@ void FSkeletalMeshLODRenderData::InitResources(bool bNeedsVertexColors, int32 LO
 							MorphAnimIndicies[Index].Add(VertexIndex);
 						}
 					}
-					//MorphBitToIndex.KeySort(TGreater<TBitArray<>>());
+					//MorphBitToIndex.KeySort(FReverseOrderBitArraysBySetBits());
 				}
 			}
 
@@ -330,7 +372,7 @@ void FSkeletalMeshLODRenderData::InitResources(bool bNeedsVertexColors, int32 LO
 				MorphTargetVertexInfoBuffers.TempStoreSize += MorphIndexToBit.Num();
 
 				//slightly better rule and cache locality when sorting by smallest number of active morphs (remember that the bits are inversed -> greater)
-				MorphIndexToBit.Sort(TGreater<TBitArray<>>());
+				MorphIndexToBit.Sort(FReverseOrderBitArraysBySetBits());
 				for (const auto& BitIndicies : MorphIndexToBit)
 				{
 					//continue splitting all requested outputs permutations until they are derived by the weights. 
