Fix for TMap property key iteration in FFindInstancedReferenceSubobjectHelper::GetInstancedSubObjects_Inner(), missed from CL# 3968969.

diff --git a/Engine/Source/Runtime/Engine/Private/InstancedReferenceSubobjectHelper.cpp b/Engine/Source/Runtime/Engine/Private/InstancedReferenceSubobjectHelper.cpp
--- a/Engine/Source/Runtime/Engine/Private/InstancedReferenceSubobjectHelper.cpp
+++ b/Engine/Source/Runtime/Engine/Private/InstancedReferenceSubobjectHelper.cpp
@@ -142,7 +142,8 @@ void FFindInstancedReferenceSubobjectHelper::GetInstancedSubObjects_Inner(FInsta
 			if (const UStruct* Struct = KeyStructProperty->Struct)
 			{
 				FScriptMapHelper MapHelper(MapProperty, MapProperty->ContainerPtrToValuePtr<void>(ContainerAddress));
-				for (int32 ElementIndex = 0; ElementIndex < MapHelper.Num(); ++ElementIndex)
+				int32 Num = MapHelper.Num();
+				for (int32 ElementIndex = 0; Num; ++ElementIndex)
 				{
 					if (MapHelper.IsValidIndex(ElementIndex))
 					{
@@ -153,6 +154,8 @@ void FFindInstancedReferenceSubobjectHelper::GetInstancedSubObjects_Inner(FInsta
 							GetInstancedSubObjects_Inner(PropertyPath, KeyAddress, OutObjects);
 							PropertyPath.Pop();
 						}
+
+						--Num;
 					}
 				}
 			}
