PR #2806: Added optimization to stat uobjects to prevent excessive rendering off of the screen, this gave us at least a 10x improvement. (Contributed by [REDACTED_PERSON])

diff --git a/Engine/Source/Runtime/Engine/Private/StatsRender2.cpp b/Engine/Source/Runtime/Engine/Private/StatsRender2.cpp
--- a/Engine/Source/Runtime/Engine/Private/StatsRender2.cpp
+++ b/Engine/Source/Runtime/Engine/Private/StatsRender2.cpp
@@ -720,6 +720,12 @@ void RenderArrayOfStats( FCanvas* Canvas, const int32 X, int32& Y, const TArray<
 		}
 
 		Y += FunctionToCall( ViewData, ComplexStat, Canvas, X, Y, TotalGroupBudget, bIsBudgetIgnored );
+
+		// Stop drawing after we have gone off the bottom of the screen, if in budget mode continue to gather avg/max stats
+		if (!bBudget && Y > (Globals.SizeXY.Y + Globals.GetFontHeight()))
+		{
+			break;
+		}
 	}
 
 	if (MaxStatsPerGroup < Aggregates.Num())
