Added default engine textures to root set in cooked builds (even though they are referenced by a root set engine object) to prevent them from being added to clusters which would result in these clusters always staying in memory.

diff --git a/Engine/Source/Runtime/Engine/Private/UnrealEngine.cpp b/Engine/Source/Runtime/Engine/Private/UnrealEngine.cpp
--- a/Engine/Source/Runtime/Engine/Private/UnrealEngine.cpp
+++ b/Engine/Source/Runtime/Engine/Private/UnrealEngine.cpp
@@ -2231,6 +2231,23 @@ void LoadEngineClass(FSoftClassPath& ClassName, TSubclassOf<ClassType>& EngineCl
 	}
 }
 
+/** Special function that loads an engine texture and adds it to root set in cooked builds. 
+  * This is to prevent the textures from being added to GC clusters (root set objects are not clustered) 
+	* because otherwise since they're always going to be referenced by UUnrealEngine object which will 
+	* prevent the clusters from being GC'd. */
+template <typename TextureType>
+static void LoadEngineTexture(TextureType*& InOutTexture, const TCHAR* InName)
+{
+	if (!InOutTexture)
+	{
+		InOutTexture = LoadObject<TextureType>(nullptr, InName, nullptr, LOAD_None, nullptr);
+	}
+	if (FPlatformProperties::RequiresCookedData() && InOutTexture)
+	{
+		InOutTexture->AddToRoot();
+	}
+}
+
 UEngineCustomTimeStep* InitializeCustomTimeStep(UEngine* InEngine, FSoftClassPath InCustomTimeStepClassName)
 {
 	UEngineCustomTimeStep* NewCustomTimeStep = nullptr;
@@ -2341,50 +2358,15 @@ void UEngine::InitializeObjectReferences()
 		}
 	}
 
-	if( DefaultTexture == NULL )
-	{
-		DefaultTexture = LoadObject<UTexture2D>(NULL, *DefaultTextureName.ToString(), NULL, LOAD_None, NULL);
-	}
-
-	if( DefaultDiffuseTexture == NULL )
-	{
-		DefaultDiffuseTexture = LoadObject<UTexture2D>(NULL, *DefaultDiffuseTextureName.ToString(), NULL, LOAD_None, NULL);
-	}
-
-	if( HighFrequencyNoiseTexture == NULL )
-	{
-		HighFrequencyNoiseTexture = LoadObject<UTexture2D>(NULL, *HighFrequencyNoiseTextureName.ToString(), NULL, LOAD_None, NULL);
-	}
-
-	if( DefaultBokehTexture == NULL )
-	{
-		DefaultBokehTexture = LoadObject<UTexture2D>(NULL, *DefaultBokehTextureName.ToString(), NULL, LOAD_None, NULL);
-	}
-
-	if (DefaultBloomKernelTexture == NULL)
-	{
-		DefaultBloomKernelTexture = LoadObject<UTexture2D>(NULL, *DefaultBloomKernelTextureName.ToString(), NULL, LOAD_None, NULL);
-	}
-
-	if( PreIntegratedSkinBRDFTexture == NULL )
-	{
-		PreIntegratedSkinBRDFTexture = LoadObject<UTexture2D>(NULL, *PreIntegratedSkinBRDFTextureName.ToString(), NULL, LOAD_None, NULL);
-	}
-
-	if( MiniFontTexture == NULL )
-	{
-		MiniFontTexture = LoadObject<UTexture2D>(NULL, *MiniFontTextureName.ToString(), NULL, LOAD_None, NULL);
-	}
-
-	if( WeightMapPlaceholderTexture == NULL )
-	{
-		WeightMapPlaceholderTexture = LoadObject<UTexture2D>(NULL, *WeightMapPlaceholderTextureName.ToString(), NULL, LOAD_None, NULL);
-	}
-
-	if (LightMapDensityTexture == NULL)
-	{
-		LightMapDensityTexture = LoadObject<UTexture2D>(NULL, *LightMapDensityTextureName.ToString(), NULL, LOAD_None, NULL);
-	}
+	LoadEngineTexture(DefaultTexture, *DefaultTextureName.ToString());
+	LoadEngineTexture(DefaultDiffuseTexture, *DefaultDiffuseTextureName.ToString());
+	LoadEngineTexture(HighFrequencyNoiseTexture, *HighFrequencyNoiseTextureName.ToString());
+	LoadEngineTexture(DefaultBokehTexture, *DefaultBokehTextureName.ToString());
+	LoadEngineTexture(DefaultBloomKernelTexture, *DefaultBloomKernelTextureName.ToString());
+	LoadEngineTexture(PreIntegratedSkinBRDFTexture, *PreIntegratedSkinBRDFTextureName.ToString());
+	LoadEngineTexture(MiniFontTexture, *MiniFontTextureName.ToString());
+	LoadEngineTexture(WeightMapPlaceholderTexture, *WeightMapPlaceholderTextureName.ToString());
+	LoadEngineTexture(LightMapDensityTexture, *LightMapDensityTextureName.ToString());
 
 	if ( DefaultPhysMaterial == NULL )
 	{
