Move semantics for timer data.
Move support for TSparseArray::Add().

diff --git a/Engine/Source/Runtime/Core/Public/Containers/SparseArray.h b/Engine/Source/Runtime/Core/Public/Containers/SparseArray.h
--- a/Engine/Source/Runtime/Core/Public/Containers/SparseArray.h
+++ b/Engine/Source/Runtime/Core/Public/Containers/SparseArray.h
@@ -126,13 +126,21 @@ public:
 	}
 
 	/** Adds an element to the array. */
-	int32 Add(typename TTypeTraits<ElementType>::ConstInitType Element)
+	int32 Add(const ElementType& Element)
 	{
 		FSparseArrayAllocationInfo Allocation = AddUninitialized();
 		new(Allocation) ElementType(Element);
 		return Allocation.Index;
 	}
 
+	/** Adds an element to the array. */
+	int32 Add(ElementType&& Element)
+	{
+		FSparseArrayAllocationInfo Allocation = AddUninitialized();
+		new(Allocation) ElementType(MoveTemp(Element));
+		return Allocation.Index;
+	}
+
 	/**
 	 * Allocates space for an element in the array at a given index.  The element is not initialized, and you must use the corresponding placement new operator
 	 * to construct the element in the allocated memory.
diff --git a/Engine/Source/Runtime/Engine/Private/TimerManager.cpp b/Engine/Source/Runtime/Engine/Private/TimerManager.cpp
--- a/Engine/Source/Runtime/Engine/Private/TimerManager.cpp
+++ b/Engine/Source/Runtime/Engine/Private/TimerManager.cpp
@@ -207,7 +207,7 @@ FTimerHandle FTimerManager::K2_FindDynamicTimerHandle(FTimerDynamicDelegate InDy
 	return Result;
 }
 
-void FTimerManager::InternalSetTimer(FTimerHandle& InOutHandle, FTimerUnifiedDelegate const& InDelegate, float InRate, bool InbLoop, float InFirstDelay)
+void FTimerManager::InternalSetTimer(FTimerHandle& InOutHandle, FTimerUnifiedDelegate&& InDelegate, float InRate, bool InbLoop, float InFirstDelay)
 {
 	SCOPE_CYCLE_COUNTER(STAT_SetTimer);
 
@@ -225,7 +225,7 @@ void FTimerManager::InternalSetTimer(FTimerHandle& InOutHandle, FTimerUnifiedDel
 	{
 		// set up the new timer
 		FTimerData NewTimerData;
-		NewTimerData.TimerDelegate = InDelegate;
+		NewTimerData.TimerDelegate = MoveTemp(InDelegate);
 
 		NewTimerData.Rate = InRate;
 		NewTimerData.bLoop = InbLoop;
@@ -261,7 +261,7 @@ void FTimerManager::InternalSetTimer(FTimerHandle& InOutHandle, FTimerUnifiedDel
 	}
 }
 
-void FTimerManager::InternalSetTimerForNextTick(FTimerUnifiedDelegate const& InDelegate)
+void FTimerManager::InternalSetTimerForNextTick(FTimerUnifiedDelegate&& InDelegate)
 {
 	SCOPE_CYCLE_COUNTER(STAT_SetTimerForNextTick);
 
@@ -272,7 +272,7 @@ void FTimerManager::InternalSetTimerForNextTick(FTimerUnifiedDelegate const& InD
 	NewTimerData.Rate = 0.f;
 	NewTimerData.bLoop = false;
 	NewTimerData.bRequiresDelegate = true;
-	NewTimerData.TimerDelegate = InDelegate;
+	NewTimerData.TimerDelegate = MoveTemp(InDelegate);
 	NewTimerData.ExpireTime = InternalTime;
 	NewTimerData.Status = ETimerStatus::Active;
 
diff --git a/Engine/Source/Runtime/Engine/Public/TimerManager.h b/Engine/Source/Runtime/Engine/Public/TimerManager.h
--- a/Engine/Source/Runtime/Engine/Public/TimerManager.h
+++ b/Engine/Source/Runtime/Engine/Public/TimerManager.h
@@ -83,6 +83,12 @@ struct FTimerUnifiedDelegate
 
 	/** Utility to output info about delegate as a string. */
 	FString ToString() const;
+
+	// Movable only
+	FTimerUnifiedDelegate(FTimerUnifiedDelegate&&) = default;
+	FTimerUnifiedDelegate(const FTimerUnifiedDelegate&) = delete;
+	FTimerUnifiedDelegate& operator=(FTimerUnifiedDelegate&&) = default;
+	FTimerUnifiedDelegate& operator=(const FTimerUnifiedDelegate&) = delete;
 };
 
 enum class ETimerStatus : uint8
@@ -135,6 +141,12 @@ struct FTimerData
 		, ExpireTime(0)
 		, LevelCollection(ELevelCollectionType::DynamicSourceLevels)
 	{}
+
+	// Movable only
+	FTimerData(FTimerData&&) = default;
+	FTimerData(const FTimerData&) = delete;
+	FTimerData& operator=(FTimerData&&) = default;
+	FTimerData& operator=(const FTimerData&) = delete;
 };
 
 
@@ -392,8 +404,8 @@ protected:
 	FTimerData* FindTimer( FTimerHandle const& InHandle );
 
 private:
-	void InternalSetTimer( FTimerHandle& InOutHandle, FTimerUnifiedDelegate const& InDelegate, float InRate, bool InbLoop, float InFirstDelay );
-	void InternalSetTimerForNextTick( FTimerUnifiedDelegate const& InDelegate );
+	void InternalSetTimer( FTimerHandle& InOutHandle, FTimerUnifiedDelegate&& InDelegate, float InRate, bool InbLoop, float InFirstDelay );
+	void InternalSetTimerForNextTick( FTimerUnifiedDelegate&& InDelegate );
 	void InternalClearTimer( FTimerHandle const& InDelegate );
 	void InternalClearAllTimers( void const* Object );
 	float InternalGetTimerRate( FTimerData const* const TimerData ) const;
