Improved reporting of bad TimerManager internal state.
Fix for object->timer map entry being wrongly removed for ActivePendingRemoval timers.
Correct removal of timer at the time it should be removed.

diff --git a/Engine/Source/Runtime/Engine/Private/TimerManager.cpp b/Engine/Source/Runtime/Engine/Private/TimerManager.cpp
--- a/Engine/Source/Runtime/Engine/Private/TimerManager.cpp
+++ b/Engine/Source/Runtime/Engine/Private/TimerManager.cpp
@@ -25,16 +25,24 @@ uint64 FTimerManager::LastAssignedSerialNumber = 0;
 
 namespace
 {
-	void DescribeFTimerDataSafely(const FTimerData& Data)
+	void DescribeFTimerDataSafely(FOutputDevice& Ar, const FTimerData& Data)
 	{
-		UE_LOG(LogEngine, Log, TEXT("TimerData %p : bLoop=%s, bRequiresDelegate=%s, Status=%d, Rate=%f, ExpireTime=%f"),
+		Ar.Logf(
+			TEXT("TimerData %p : bLoop=%s, bRequiresDelegate=%s, Status=%d, Rate=%f, ExpireTime=%f"),
 			&Data,
 			Data.bLoop ? TEXT("true") : TEXT("false"),
 			Data.bRequiresDelegate ? TEXT("true") : TEXT("false"),
 			static_cast<int32>(Data.Status),
 			Data.Rate,
 			Data.ExpireTime
-			)
+		);
+	}
+
+	FString GetFTimerDataSafely(const FTimerData& Data)
+	{
+		FStringOutputDevice Output;
+		DescribeFTimerDataSafely(Output, Data);
+		return Output;
 	}
 }
 
@@ -96,7 +104,7 @@ void FTimerManager::OnCrash()
 		}
 		else
 		{
-			DescribeFTimerDataSafely(Timer);
+			DescribeFTimerDataSafely(*GLog, Timer);
 		}
 	}
 	UE_LOG(LogEngine, Log, TEXT("------- %d Expired Active Timers -------"), ExpiredActiveTimerCount);
@@ -105,14 +113,14 @@ void FTimerManager::OnCrash()
 	for (FTimerHandle Handle : PausedTimerSet)
 	{
 		const FTimerData& Timer = GetTimer(Handle);
-		DescribeFTimerDataSafely(Timer);
+		DescribeFTimerDataSafely(*GLog, Timer);
 	}
 
 	UE_LOG(LogEngine, Log, TEXT("------- %d Pending Timers -------"), PendingTimerSet.Num());
 	for (FTimerHandle Handle : PendingTimerSet)
 	{
 		const FTimerData& Timer = GetTimer(Handle);
-		DescribeFTimerDataSafely(Timer);
+		DescribeFTimerDataSafely(*GLog, Timer);
 	}
 
 	UE_LOG(LogEngine, Log, TEXT("------- %d Total Timers -------"), PendingTimerSet.Num() + PausedTimerSet.Num() + ActiveTimerHeap.Num() - ExpiredActiveTimerCount);
@@ -355,8 +363,6 @@ void FTimerManager::InternalClearAllTimers(void const* Object)
 	{
 		InternalClearTimer(TimerToRemove);
 	}
-
-	ObjectToTimers.Remove(Object);
 }
 
 float FTimerManager::InternalGetTimerRemaining(FTimerData const* const TimerData) const
@@ -554,6 +560,7 @@ void FTimerManager::Tick(float DeltaTime)
 		if (Top->Status == ETimerStatus::ActivePendingRemoval)
 		{
 			ActiveTimerHeap.HeapPop(TopHandle, FTimerHeapOrder(Timers), /*bAllowShrinking=*/ false);
+			RemoveTimer(TopHandle);
 			continue;
 		}
 
@@ -700,7 +707,11 @@ FTimerHandle FTimerManager::AddTimer(FTimerData&& TimerData)
 
 	if (TimerIndicesByObjectKey)
 	{
-		ObjectToTimers.FindOrAdd(TimerIndicesByObjectKey).Add(Result);
+		TSet<FTimerHandle>& HandleSet = ObjectToTimers.FindOrAdd(TimerIndicesByObjectKey);
+
+		bool bAlreadyExists = false;
+		HandleSet.Add(Result, &bAlreadyExists);
+		checkf(!bAlreadyExists, TEXT("A timer with this handle and object has already been added! (%s)"), *GetFTimerDataSafely(TimerData));
 	}
 
 	return Result;
@@ -714,10 +725,10 @@ void FTimerManager::RemoveTimer(FTimerHandle Handle)
 	if (const void* TimerIndicesByObjectKey = Data.TimerIndicesByObjectKey)
 	{
 		TSet<FTimerHandle>* TimersForObject = ObjectToTimers.Find(TimerIndicesByObjectKey);
-		checkf(TimersForObject, TEXT("Removed timer was bound to an object which is not tracked by TimerIndicesByObject!"));
+		checkf(TimersForObject, TEXT("Removed timer was bound to an object which is not tracked by ObjectToTimers! (%s)"), *GetFTimerDataSafely(Data));
 
 		int32 NumRemoved = TimersForObject->Remove(Handle);
-		checkf(NumRemoved == 1, TEXT("Removed timer was bound to an object which is not tracked by TimerIndicesByObject!"));
+		checkf(NumRemoved == 1, TEXT("Removed timer was bound to an object which is not tracked by ObjectToTimers! (%s)"), *GetFTimerDataSafely(Data));
 
 		if (TimersForObject->Num() == 0)
 		{
