Removing redundant BeginLoad call to fix a crash in ShooterGame

diff --git a/Engine/Source/Runtime/Engine/Private/UnrealEngine.cpp b/Engine/Source/Runtime/Engine/Private/UnrealEngine.cpp
--- a/Engine/Source/Runtime/Engine/Private/UnrealEngine.cpp
+++ b/Engine/Source/Runtime/Engine/Private/UnrealEngine.cpp
@@ -8558,8 +8558,6 @@ FGuid UEngine::GetPackageGuid(FName PackageName, bool bForPIE)
 {
 	FGuid Result(0,0,0,0);
 
-	TRefCountPtr<FUObjectSerializeContext> LoadContext(new FUObjectSerializeContext());
-	BeginLoad(LoadContext, *PackageName.ToString());
 	uint32 LoadFlags = LOAD_NoWarn | LOAD_NoVerify;
 	if (bForPIE)
 	{
