A few drive-by optimizations to game startup times that yielded ~15% speedup

Wall time from starting Test client with cooked data til you can choose game mode went from ~22s to ~19s on my machine.

* TSet/TMap::Reset() and Empty()
-- Use simpler code that compile to rep instruction instead of loop
-- Avoid needless clearing of FSetElementId when set is empty
* AllowThreadHeartBeat() parses commandline once
* Avoided some needless work in FStreamableHandle::HasLoadCompletedOrStalled() and UpdateCombinedHandle()

diff --git a/Engine/Source/Runtime/Core/Private/GenericPlatform/GenericPlatformMisc.cpp b/Engine/Source/Runtime/Core/Private/GenericPlatform/GenericPlatformMisc.cpp
--- a/Engine/Source/Runtime/Core/Private/GenericPlatform/GenericPlatformMisc.cpp
+++ b/Engine/Source/Runtime/Core/Private/GenericPlatform/GenericPlatformMisc.cpp
@@ -1022,12 +1022,8 @@ bool FGenericPlatformMisc::UseRenderThread()
 
 bool FGenericPlatformMisc::AllowThreadHeartBeat()
 {
-	if (FParse::Param(FCommandLine::Get(), TEXT("noheartbeatthread")))
-	{
-		return false;
-	}
-
-	return true;
+	static bool bHeartbeat = !FParse::Param(FCommandLine::Get(), TEXT("noheartbeatthread"));
+	return bHeartbeat;
 }
 
 int32 FGenericPlatformMisc::NumberOfCoresIncludingHyperthreads()
diff --git a/Engine/Source/Runtime/Core/Public/Containers/Set.h b/Engine/Source/Runtime/Core/Public/Containers/Set.h
--- a/Engine/Source/Runtime/Core/Public/Containers/Set.h
+++ b/Engine/Source/Runtime/Core/Public/Containers/Set.h
@@ -117,6 +117,15 @@ public:
 	}
 
 private:
+	
+	/** Reset a range of FSetElementIds to invalid */
+	FORCEINLINE static void ResetRange(FSetElementId* Range, int32 Count)
+	{
+		for (int32 I = 0; I < Count; ++I)
+		{
+			Range[I] = FSetElementId();
+		}
+	}
 
 	/** The index of the element in the set's element array. */
 	int32 Index;
@@ -365,14 +374,16 @@ public:
 	/** Efficiently empties out the set but preserves all allocations and capacities */
     void Reset()
     {
+		if (Num() == 0)
+		{
+			return;
+		}
+    	
 		// Reset the elements array.
 		Elements.Reset();
 
 		// Clear the references to the elements that have now been removed.
-		for (int32 HashIndex = 0, LocalHashSize = HashSize; HashIndex < LocalHashSize; ++HashIndex)
-		{
-			GetTypedHash(HashIndex) = FSetElementId();
-		}
+		FSetElementId::ResetRange(Hash.GetAllocation(), HashSize);
     }
 
 	/** Shrinks the set's element storage to avoid slack. */
@@ -1374,10 +1385,7 @@ public:
 			Hash.ResizeAllocation(0, HashSize, sizeof(FSetElementId));
 		}
 
-		for (auto* It = (FSetElementId*)Hash.GetAllocation(), *End = It + HashSize; It != End; ++It)
-		{
-			*It = FSetElementId();
-		}
+		FSetElementId::ResetRange(Hash.GetAllocation(), HashSize);
 	}
 
 	void RemoveAt(int32 Index, const FScriptSetLayout& Layout)
diff --git a/Engine/Source/Runtime/Engine/Private/StreamableManager.cpp b/Engine/Source/Runtime/Engine/Private/StreamableManager.cpp
--- a/Engine/Source/Runtime/Engine/Private/StreamableManager.cpp
+++ b/Engine/Source/Runtime/Engine/Private/StreamableManager.cpp
@@ -205,30 +205,21 @@ EAsyncPackageState::Type FStreamableHandle::WaitUntilComplete(float Timeout, boo
 
 bool FStreamableHandle::HasLoadCompletedOrStalled() const
 {
-	// We need to recursively look for complete or stalled handles
-	TArray<TSharedRef< FStreamableHandle>> HandlesToCheck;
-
-	HandlesToCheck.Add(const_cast<FStreamableHandle*>(this)->AsShared());
-
-	for (int32 i = 0; i < HandlesToCheck.Num(); i++)
+	// Check this handle
+	if (!IsCombinedHandle() && !HasLoadCompleted() && !IsStalled())
 	{
-		TSharedRef<const FStreamableHandle> Handle = HandlesToCheck[i];
+		return false;
+	}
 
-		if (!Handle->IsCombinedHandle() && !Handle->HasLoadCompleted() && !Handle->IsStalled())
+	// Check children recursively
+	for (const TSharedPtr<FStreamableHandle>& ChildHandle : ChildHandles)
+	{
+		if (ChildHandle.IsValid() && !ChildHandle->HasLoadCompletedOrStalled())
 		{
 			return false;
 		}
-
-		for (TSharedPtr<FStreamableHandle> ChildHandle : Handle->ChildHandles)
-		{
-			if (ChildHandle.IsValid())
-			{
-				HandlesToCheck.Add(ChildHandle.ToSharedRef());
-			}
-		}
 	}
 
-	// All handles are either completed or stalled
 	return true;
 }
 
@@ -518,16 +509,10 @@ void FStreamableHandle::UpdateCombinedHandle()
 	// Check all our children, complete if done
 	bool bAllCompleted = true;
 	bool bAllCanceled = true;
-	for (TSharedPtr<FStreamableHandle> ChildHandle : ChildHandles)
+	for (TSharedPtr<FStreamableHandle>& ChildHandle : ChildHandles)
 	{
-		if (ChildHandle->IsLoadingInProgress())
-		{
-			bAllCompleted = false;
-		}
-		if (!ChildHandle->WasCanceled())
-		{
-			bAllCanceled = false;
-		}
+		bAllCompleted = bAllCompleted && !ChildHandle->IsLoadingInProgress();
+		bAllCanceled = bAllCanceled && ChildHandle->WasCanceled();
 	}
 
 	// If all our sub handles were canceled, cancel us. Otherwise complete us if at least one was completed and there are none in progress
