TIndirectArray's placement new syntax deprecated, which can cause new/free or malloc/delete mismatches.

Also see UE-64535.

diff --git a/Engine/Plugins/2D/Paper2D/Source/Paper2D/Private/PaperSprite.cpp b/Engine/Plugins/2D/Paper2D/Source/Paper2D/Private/PaperSprite.cpp
--- a/Engine/Plugins/2D/Paper2D/Source/Paper2D/Private/PaperSprite.cpp
+++ b/Engine/Plugins/2D/Paper2D/Source/Paper2D/Private/PaperSprite.cpp
@@ -479,7 +479,7 @@ protected:
 			FlushRenderingCommands();
 		}
 
-		new (ComponentContexts) FComponentReregisterContext(Component);
+		ComponentContexts.Add(new FComponentReregisterContext(Component));
 	}
 
 private:
diff --git a/Engine/Plugins/2D/Paper2D/Source/Paper2D/Private/PaperTileMap.cpp b/Engine/Plugins/2D/Paper2D/Source/Paper2D/Private/PaperTileMap.cpp
--- a/Engine/Plugins/2D/Paper2D/Source/Paper2D/Private/PaperTileMap.cpp
+++ b/Engine/Plugins/2D/Paper2D/Source/Paper2D/Private/PaperTileMap.cpp
@@ -156,7 +156,7 @@ protected:
 			FlushRenderingCommands();
 		}
 
-		new (ComponentContexts) FComponentReregisterContext(Component);
+		ComponentContexts.Add(new FComponentReregisterContext(Component));
 	}
 
 private:
diff --git a/Engine/Plugins/Experimental/AlembicImporter/Source/AlembicLibrary/Private/AbcImporter.cpp b/Engine/Plugins/Experimental/AlembicImporter/Source/AlembicLibrary/Private/AbcImporter.cpp
--- a/Engine/Plugins/Experimental/AlembicImporter/Source/AlembicLibrary/Private/AbcImporter.cpp
+++ b/Engine/Plugins/Experimental/AlembicImporter/Source/AlembicLibrary/Private/AbcImporter.cpp
@@ -646,7 +646,7 @@ TArray<UObject*> FAbcImporter::ImportAsSkeletalMesh(UObject* InParent, EObjectFl
 		FSkeletalMeshModel* ImportedModel = SkeletalMesh->GetImportedModel();
 		check(ImportedModel->LODModels.Num() == 0);
 		ImportedModel->LODModels.Empty();
-		new(ImportedModel->LODModels)FSkeletalMeshLODModel();
+		ImportedModel->LODModels.Add(new FSkeletalMeshLODModel());
 		SkeletalMesh->ResetLODInfo();
 		SkeletalMesh->AddLODInfo();
 		FSkeletalMeshLODModel& LODModel = ImportedModel->LODModels[0];
diff --git a/Engine/Plugins/FX/Niagara/Source/Niagara/Private/NiagaraRenderer.cpp b/Engine/Plugins/FX/Niagara/Source/Niagara/Private/NiagaraRenderer.cpp
--- a/Engine/Plugins/FX/Niagara/Source/Niagara/Private/NiagaraRenderer.cpp
+++ b/Engine/Plugins/FX/Niagara/Source/Niagara/Private/NiagaraRenderer.cpp
@@ -440,7 +440,8 @@ FNiagaraGlobalReadBuffer::FAllocation FNiagaraGlobalReadBuffer::AllocateFloat(ui
 		if (Buffer == NULL)
 		{
 			uint32 NewBufferSize = FMath::Max(Num, (uint32)GMinNiagaraRenderingBufferSize);
-			Buffer = new(FloatBufferPool->Buffers) FDynamicAllocReadBuffer();
+			Buffer = new FDynamicAllocReadBuffer();
+			FloatBufferPool->Buffers.Add(Buffer);
 			Buffer->Initialize(sizeof(float), NewBufferSize, PF_R32_FLOAT, BUF_Dynamic);
 		}
 
@@ -494,7 +495,8 @@ FNiagaraGlobalReadBuffer::FAllocation FNiagaraGlobalReadBuffer::AllocateInt32(ui
 		if (Buffer == NULL)
 		{
 			uint32 NewBufferSize = FMath::Max(Num, (uint32)GMinNiagaraRenderingBufferSize);
-			Buffer = new(Int32BufferPool->Buffers) FDynamicAllocReadBuffer();
+			Buffer = new FDynamicAllocReadBuffer();
+			Int32BufferPool->Buffers.Add(Buffer);
 			Buffer->Initialize(sizeof(int32), NewBufferSize, PF_R32_SINT, BUF_Dynamic);
 		}
 
diff --git a/Engine/Plugins/Online/OnlineSubsystemSteam/Source/Private/OnlineSessionAsyncServerSteam.cpp b/Engine/Plugins/Online/OnlineSubsystemSteam/Source/Private/OnlineSessionAsyncServerSteam.cpp
--- a/Engine/Plugins/Online/OnlineSubsystemSteam/Source/Private/OnlineSessionAsyncServerSteam.cpp
+++ b/Engine/Plugins/Online/OnlineSubsystemSteam/Source/Private/OnlineSessionAsyncServerSteam.cpp
@@ -952,7 +952,8 @@ void FOnlineAsyncTaskSteamFindServerBase::ParseSearchResult(class gameserveritem
 		if (ServerBuildId == BuildUniqueId)
 		{
 			// Create a new pending search result 
-			FPendingSearchResultSteam* NewPendingSearch = new (PendingSearchResults) FPendingSearchResultSteam(this);
+			FPendingSearchResultSteam* NewPendingSearch = new FPendingSearchResultSteam(this);
+			PendingSearchResults.Add(NewPendingSearch);
 			NewPendingSearch->ServerId = FUniqueNetIdSteam(ServerDetails->m_steamID);
 			NewPendingSearch->HostAddr = ServerAddr;
 
diff --git a/Engine/Plugins/Runtime/ApexDestruction/Source/ApexDestruction/Private/ApexDestructibleAssetImport.cpp b/Engine/Plugins/Runtime/ApexDestruction/Source/ApexDestruction/Private/ApexDestructibleAssetImport.cpp
--- a/Engine/Plugins/Runtime/ApexDestruction/Source/ApexDestruction/Private/ApexDestructibleAssetImport.cpp
+++ b/Engine/Plugins/Runtime/ApexDestruction/Source/ApexDestruction/Private/ApexDestructibleAssetImport.cpp
@@ -877,7 +877,7 @@ bool SetApexDestructibleAsset(UDestructibleMesh& DestructibleMesh, apex::Destruc
 	FSkeletalMeshModel& DestructibleMeshResource = *DestructibleMesh.GetImportedModel();
 	check(DestructibleMeshResource.LODModels.Num() == 0);
 	DestructibleMeshResource.LODModels.Empty();
-	new(DestructibleMeshResource.LODModels)FSkeletalMeshLODModel();
+	DestructibleMeshResource.LODModels.Add(new FSkeletalMeshLODModel());
 
 	DestructibleMesh.ResetLODInfo();
 	DestructibleMesh.AddLODInfo();
diff --git a/Engine/Source/Developer/MeshUtilities/Private/MeshUtilities.cpp b/Engine/Source/Developer/MeshUtilities/Private/MeshUtilities.cpp
--- a/Engine/Source/Developer/MeshUtilities/Private/MeshUtilities.cpp
+++ b/Engine/Source/Developer/MeshUtilities/Private/MeshUtilities.cpp
@@ -2424,8 +2424,10 @@ public:
 		for (int32 LODIndex = 0; LODIndex < SourceModels.Num(); ++LODIndex)
 		{
 			FStaticMeshSourceModel& SrcModel = SourceModels[LODIndex];
-			FRawMesh& RawMesh = *new(LODMeshes)FRawMesh;
-			FOverlappingCorners& OverlappingCorners = *new(LODOverlappingCorners)FOverlappingCorners;
+			FRawMesh& RawMesh = *new FRawMesh;
+			LODMeshes.Add(&RawMesh);
+			FOverlappingCorners& OverlappingCorners = *new FOverlappingCorners;
+			LODOverlappingCorners.Add(&OverlappingCorners);
 
 			if (!SrcModel.IsRawMeshEmpty())
 			{
@@ -3930,7 +3932,8 @@ public:
 
 		BeginSlowTask();
 
-		FOverlappingCorners& OverlappingCorners = *new(LODOverlappingCorners)FOverlappingCorners;
+		FOverlappingCorners& OverlappingCorners = *new FOverlappingCorners;
+		LODOverlappingCorners.Add(&OverlappingCorners);
 
 		float ComparisonThreshold = THRESH_POINTS_ARE_SAME;//GetComparisonThreshold(LODBuildSettings[LODIndex]);
 		int32 NumWedges = BuildData->GetNumWedges();
diff --git a/Engine/Source/Developer/MeshUtilities/Private/SkeletalMeshTools.cpp b/Engine/Source/Developer/MeshUtilities/Private/SkeletalMeshTools.cpp
--- a/Engine/Source/Developer/MeshUtilities/Private/SkeletalMeshTools.cpp
+++ b/Engine/Source/Developer/MeshUtilities/Private/SkeletalMeshTools.cpp
@@ -251,7 +251,8 @@ namespace SkeletalMeshTools
 					Chunks.Add(DestChunk);
 					DestChunk->MaterialIndex = SrcChunk->MaterialIndex;
 					DestChunk->OriginalSectionIndex = SrcChunk->OriginalSectionIndex;
-					TArray<int32>& IndexMap = *new(IndexMaps) TArray<int32>();
+					TArray<int32>& IndexMap = *new TArray<int32>();
+					IndexMaps.Add(&IndexMap);
 					IndexMap.AddUninitialized(SrcChunk->Vertices.Num());
 					FMemory::Memset(IndexMap.GetData(),0xff,IndexMap.GetTypeSize()*IndexMap.Num());
 				}
diff --git a/Engine/Source/Developer/ShaderCompilerCommon/Private/HlslLexer.cpp b/Engine/Source/Developer/ShaderCompilerCommon/Private/HlslLexer.cpp
--- a/Engine/Source/Developer/ShaderCompilerCommon/Private/HlslLexer.cpp
+++ b/Engine/Source/Developer/ShaderCompilerCommon/Private/HlslLexer.cpp
@@ -945,7 +945,7 @@ namespace CrossCompiler
 	void FHlslScanner::Clear(const FString& Filename)
 	{
 		Tokens.Empty();
-		new (SourceFilenames) FString(Filename);
+		SourceFilenames.Add(new FString(Filename));
 	}
 
 	bool FHlslScanner::Lex(const FString& String, const FString& Filename)
@@ -965,7 +965,7 @@ namespace CrossCompiler
 				FTokenizer::ProcessDirective(Tokenizer, CompilerMessages, *this);
 				if (Tokenizer.Filename != SourceFilenames.Last())
 				{
-					new(SourceFilenames) FString(Tokenizer.Filename);
+					SourceFilenames.Add(new FString(Tokenizer.Filename));
 				}
 			}
 			else
diff --git a/Engine/Source/Developer/SimplygonSwarm/Private/SimplygonRESTClient.cpp b/Engine/Source/Developer/SimplygonSwarm/Private/SimplygonRESTClient.cpp
--- a/Engine/Source/Developer/SimplygonSwarm/Private/SimplygonRESTClient.cpp
+++ b/Engine/Source/Developer/SimplygonSwarm/Private/SimplygonRESTClient.cpp
@@ -71,7 +71,8 @@ void FSimplygonSwarmTask::CreateUploadParts(const int32 MaxUploadPartSize)
 
 			for (int32 PartIndex = 0; PartIndex < NumberOfPartsRequried; ++PartIndex)
 			{
-				FSwarmUploadPart* UploadPartData = new (UploadParts) FSwarmUploadPart();
+				FSwarmUploadPart* UploadPartData = new FSwarmUploadPart();
+				UploadParts.Add(UploadPartData);
 				int32 Offset = PartIndex * MaxUploadPartSize * sizeof(uint8);
 				UploadPartData->Data.AddUninitialized(MaxUploadPartSize);
 				FMemory::Memcpy(UploadPartData->Data.GetData(), fileBlob.GetData() + Offset, MaxUploadPartSize * sizeof(uint8));
@@ -83,7 +84,8 @@ void FSimplygonSwarmTask::CreateUploadParts(const int32 MaxUploadPartSize)
 			{
 				//NOTE: need to set Offset before doing a new on UploadParts
 				int32 Offset = UploadParts.Num() * MaxUploadPartSize;
-				FSwarmUploadPart* UploadPartData = new (UploadParts) FSwarmUploadPart();
+				FSwarmUploadPart* UploadPartData = new FSwarmUploadPart();
+				UploadParts.Add(UploadPartData);
 				UploadPartData->Data.AddUninitialized(RemainingBytes);
 				FMemory::Memcpy(UploadPartData->Data.GetData(), fileBlob.GetData() + Offset, RemainingBytes * sizeof(uint8));
 				UploadPartData->PartNumber = NumberOfPartsRequried + 1;
@@ -92,7 +94,8 @@ void FSimplygonSwarmTask::CreateUploadParts(const int32 MaxUploadPartSize)
 		}
 		else
 		{
-			FSwarmUploadPart* UploadPartData = new (UploadParts) FSwarmUploadPart();
+			FSwarmUploadPart* UploadPartData = new FSwarmUploadPart();
+			UploadParts.Add(UploadPartData);
 			UploadPartData->Data.AddUninitialized(fileBlob.Num());
 			FMemory::Memcpy(UploadPartData->Data.GetData(), fileBlob.GetData(), fileBlob.Num() * sizeof(uint8));
 			UploadPartData->PartNumber = 1;
diff --git a/Engine/Source/Developer/TextureCompressor/Private/TextureCompressorModule.cpp b/Engine/Source/Developer/TextureCompressor/Private/TextureCompressorModule.cpp
--- a/Engine/Source/Developer/TextureCompressor/Private/TextureCompressorModule.cpp
+++ b/Engine/Source/Developer/TextureCompressor/Private/TextureCompressorModule.cpp
@@ -1261,7 +1261,8 @@ static void GenerateAngularFilteredMip(FImage* DestMip, FImage& SrcMip, float Co
 
 		for (int32 Face = 0; Face < 6; ++Face)
 		{
-			auto* AsyncTask = new(AsyncTasks) FAsyncGenerateMipsPerFaceTask(Face, DestMip, MipExtent, ConeAngle, TexelAreaArray.GetData(), &SrcMip);
+			auto* AsyncTask = new FAsyncGenerateMipsPerFaceTask(Face, DestMip, MipExtent, ConeAngle, TexelAreaArray.GetData(), &SrcMip);
+			AsyncTasks.Add(AsyncTask);
 			AsyncTask->StartBackgroundTask();
 		}
 
@@ -1785,12 +1786,13 @@ static bool CompressMipChain(
 				FCompressedImage2D& DestMip = *new(OutMips) FCompressedImage2D;
 				if (bAllowParallelBuild && FMath::Min(SrcMip.SizeX, SrcMip.SizeY) >= MinAsyncCompressionSize)
 				{
-					FAsyncCompressionTask* AsyncTask = new(AsyncCompressionTasks) FAsyncCompressionTask(
+					FAsyncCompressionTask* AsyncTask = new FAsyncCompressionTask(
 						TextureFormat,
 						&SrcMip,
 						Settings,
 						bImageHasAlphaChannel
 						);
+					AsyncCompressionTasks.Add(AsyncTask);
 #if WITH_EDITOR
 					AsyncTask->StartBackgroundTask(GLargeThreadPool);
 #else
diff --git a/Engine/Source/Developer/TextureFormatDXT/Private/TextureFormatDXT.cpp b/Engine/Source/Developer/TextureFormatDXT/Private/TextureFormatDXT.cpp
--- a/Engine/Source/Developer/TextureFormatDXT/Private/TextureFormatDXT.cpp
+++ b/Engine/Source/Developer/TextureFormatDXT/Private/TextureFormatDXT.cpp
@@ -354,7 +354,7 @@ static bool CompressImageUsingNVTT(
 		uint8* Dest = OutCompressedData.GetData();
 		for (int32 BatchIndex = 0; BatchIndex < NumBatches; ++BatchIndex)
 		{
-			new(Compressors) FNVTTCompressor(
+			Compressors.Add(new FNVTTCompressor(
 				Src,
 				PixelFormat,
 				SizeX,
@@ -363,7 +363,7 @@ static bool CompressImageUsingNVTT(
 				bIsNormalMap,
 				Dest,
 				CompressedStride
-				);
+			));
 			Src += UncompressedStride;
 			Dest += CompressedStride;
 		}
@@ -375,7 +375,8 @@ static bool CompressImageUsingNVTT(
 		TIndirectArray<FAsyncNVTTTask> AsyncTasks;
 		for (int32 BatchIndex = 0; BatchIndex < NumBatches; ++BatchIndex)
 		{
-			FAsyncNVTTTask* AsyncTask = new(AsyncTasks) FAsyncNVTTTask(&Compressors[BatchIndex]);
+			FAsyncNVTTTask* AsyncTask = new FAsyncNVTTTask(&Compressors[BatchIndex]);
+			AsyncTasks.Add(AsyncTask);
 #if WITH_EDITOR
 			AsyncTask->StartBackgroundTask(GLargeThreadPool);
 #else
diff --git a/Engine/Source/Developer/TextureFormatIntelISPCTexComp/Private/TextureFormatIntelISPCTexComp.cpp b/Engine/Source/Developer/TextureFormatIntelISPCTexComp/Private/TextureFormatIntelISPCTexComp.cpp
--- a/Engine/Source/Developer/TextureFormatIntelISPCTexComp/Private/TextureFormatIntelISPCTexComp.cpp
+++ b/Engine/Source/Developer/TextureFormatIntelISPCTexComp/Private/TextureFormatIntelISPCTexComp.cpp
@@ -282,7 +282,8 @@ struct FMultithreadedCompression
 					// Create a new task unless it's the last task in the last slice (that one will run on current thread, after these threads have been started)
 					if (SliceIndex < (Image.NumSlices - 1) || iTask < (NumStasksPerSlice - 1))
 					{
-						auto* AsyncTask = new(CompressionTasks) FIntelCompressTask(&EncoderSettings, &Image, &OutCompressedImage, iTask * MultithreadSettings.iScansPerTask, (iTask + 1) * MultithreadSettings.iScansPerTask, SliceIndex, FunctionCallback);
+						auto* AsyncTask = new FIntelCompressTask(&EncoderSettings, &Image, &OutCompressedImage, iTask * MultithreadSettings.iScansPerTask, (iTask + 1) * MultithreadSettings.iScansPerTask, SliceIndex, FunctionCallback);
+						CompressionTasks.Add(AsyncTask);
 						AsyncTask->StartBackgroundTask();
 					}
 				}
diff --git a/Engine/Source/Developer/Windows/ShaderFormatD3D/Private/D3D11ShaderCompiler.cpp b/Engine/Source/Developer/Windows/ShaderFormatD3D/Private/D3D11ShaderCompiler.cpp
--- a/Engine/Source/Developer/Windows/ShaderFormatD3D/Private/D3D11ShaderCompiler.cpp
+++ b/Engine/Source/Developer/Windows/ShaderFormatD3D/Private/D3D11ShaderCompiler.cpp
@@ -617,7 +617,7 @@ static bool CompileAndProcessD3DShader(FString& PreprocessedShaderSource, const
 					if (ParamDesc.SystemValueType == D3D_NAME_UNDEFINED && ParamDesc.Mask != 0)
 					{
 						++NumInterpolants;
-						new(InterpolantNames) FString(FString::Printf(TEXT("%s%d"), ANSI_TO_TCHAR(ParamDesc.SemanticName), ParamDesc.SemanticIndex));
+						InterpolantNames.Add(new FString(FString::Printf(TEXT("%s%d"), ANSI_TO_TCHAR(ParamDesc.SemanticName), ParamDesc.SemanticIndex)));
 						ShaderOutputs.Add(*InterpolantNames.Last());
 					}
 				}
diff --git a/Engine/Source/Editor/BlueprintGraph/Private/DelegateNodeHandlers.cpp b/Engine/Source/Editor/BlueprintGraph/Private/DelegateNodeHandlers.cpp
--- a/Engine/Source/Editor/BlueprintGraph/Private/DelegateNodeHandlers.cpp
+++ b/Engine/Source/Editor/BlueprintGraph/Private/DelegateNodeHandlers.cpp
@@ -81,7 +81,8 @@ struct FKCHandlerDelegateHelper
 	{
 		check(SelfPin && NetPin && BoundProperty && DelegateNode);
 
-		FBPTerminal* Term = new(Context.VariableReferences) FBPTerminal();
+		FBPTerminal* Term = new FBPTerminal();
+		Context.VariableReferences.Add(Term);
 		Term->CopyFromPin(SelfPin, BoundProperty->GetName());
 		Term->AssociatedVarProperty = BoundProperty;
 
diff --git a/Engine/Source/Editor/BlueprintGraph/Private/K2Node_DelegateSet.cpp b/Engine/Source/Editor/BlueprintGraph/Private/K2Node_DelegateSet.cpp
--- a/Engine/Source/Editor/BlueprintGraph/Private/K2Node_DelegateSet.cpp
+++ b/Engine/Source/Editor/BlueprintGraph/Private/K2Node_DelegateSet.cpp
@@ -69,7 +69,8 @@ public:
 		// Create a term for this property
 		if( BoundProperty != NULL )
 		{
-			FBPTerminal* Term = new(Context.VariableReferences) FBPTerminal();
+			FBPTerminal* Term = new FBPTerminal();
+			Context.VariableReferences.Add(Term);
 			Term->CopyFromPin(DelegatePin, DelegatePin->PinName);
 			Term->AssociatedVarProperty = BoundProperty;
 
diff --git a/Engine/Source/Editor/BlueprintGraph/Private/K2Node_FunctionEntry.cpp b/Engine/Source/Editor/BlueprintGraph/Private/K2Node_FunctionEntry.cpp
--- a/Engine/Source/Editor/BlueprintGraph/Private/K2Node_FunctionEntry.cpp
+++ b/Engine/Source/Editor/BlueprintGraph/Private/K2Node_FunctionEntry.cpp
@@ -41,7 +41,8 @@ public:
 	void RegisterFunctionInput(FKismetFunctionContext& Context, UEdGraphPin* Net, UFunction* Function)
 	{
 		// This net is a parameter into the function
-		FBPTerminal* Term = new (Context.Parameters) FBPTerminal();
+		FBPTerminal* Term = new FBPTerminal();
+		Context.Parameters.Add(Term);
 		Term->CopyFromPin(Net, Net->PinName);
 
 		// Flag pass by reference parameters specially
@@ -111,7 +112,8 @@ public:
 						// node wired into our function graph, know that it
 						// will first check to see if this already exists 
 						// for it to use (rather than creating one of its own)
-						FBPTerminal* ResultTerm = new (Context.Results) FBPTerminal();
+						FBPTerminal* ResultTerm = new FBPTerminal();
+						Context.Results.Add(ResultTerm);
 						ResultTerm->Name = ParamName;
 
 						ResultTerm->Type = ParamType;
diff --git a/Engine/Source/Editor/BlueprintGraph/Private/K2Node_FunctionResult.cpp b/Engine/Source/Editor/BlueprintGraph/Private/K2Node_FunctionResult.cpp
--- a/Engine/Source/Editor/BlueprintGraph/Private/K2Node_FunctionResult.cpp
+++ b/Engine/Source/Editor/BlueprintGraph/Private/K2Node_FunctionResult.cpp
@@ -79,7 +79,8 @@ public:
 					return;
 				}
 			}
-			FBPTerminal* Term = new (Context.Results) FBPTerminal();
+			FBPTerminal* Term = new FBPTerminal();
+			Context.Results.Add(Term);
 			Term->CopyFromPin(Net, MoveTemp(NetPinName));
 			Context.NetMap.Add(Net, Term);
 		}
diff --git a/Engine/Source/Editor/BlueprintGraph/Private/K2Node_GetArrayItem.cpp b/Engine/Source/Editor/BlueprintGraph/Private/K2Node_GetArrayItem.cpp
--- a/Engine/Source/Editor/BlueprintGraph/Private/K2Node_GetArrayItem.cpp
+++ b/Engine/Source/Editor/BlueprintGraph/Private/K2Node_GetArrayItem.cpp
@@ -49,7 +49,8 @@ public:
 		}
 
 		{
-			FBPTerminal* Term = new (Context.InlineGeneratedValues) FBPTerminal();
+			FBPTerminal* Term = new FBPTerminal();
+			Context.InlineGeneratedValues.Add(Term);
 			Term->CopyFromPin(Node->Pins[2], Context.NetNameMap->MakeValidName(Node->Pins[2]));
 			Context.NetMap.Add(Node->Pins[2], Term);
 		}
diff --git a/Engine/Source/Editor/BlueprintGraph/Private/K2Node_Literal.cpp b/Engine/Source/Editor/BlueprintGraph/Private/K2Node_Literal.cpp
--- a/Engine/Source/Editor/BlueprintGraph/Private/K2Node_Literal.cpp
+++ b/Engine/Source/Editor/BlueprintGraph/Private/K2Node_Literal.cpp
@@ -54,7 +54,8 @@ public:
 		if( !Term )
 		{
 			FString RefPropName = (TargetObject ? TargetObject->GetName() : TEXT("None")) + TEXT("_") + (Context.SourceGraph ? *Context.SourceGraph->GetName() : TEXT("None")) + TEXT("_RefProperty");
-			Term = new (Context.LevelActorReferences) FBPTerminal();
+			Term = new FBPTerminal();
+			Context.LevelActorReferences.Add(Term);
 			Term->CopyFromPin(Net, Context.NetNameMap->MakeValidName(Net));
 			Term->Name = RefPropName;
 			Term->PropertyDefault = TargetObjectName;
diff --git a/Engine/Source/Editor/BlueprintGraph/Private/K2Node_Select.cpp b/Engine/Source/Editor/BlueprintGraph/Private/K2Node_Select.cpp
--- a/Engine/Source/Editor/BlueprintGraph/Private/K2Node_Select.cpp
+++ b/Engine/Source/Editor/BlueprintGraph/Private/K2Node_Select.cpp
@@ -59,7 +59,8 @@ public:
 		}
 
 		{
-			FBPTerminal* Term = new (Context.InlineGeneratedValues) FBPTerminal();
+			FBPTerminal* Term = new FBPTerminal();
+			Context.InlineGeneratedValues.Add(Term);
 			Term->CopyFromPin(ReturnPin, Context.NetNameMap->MakeValidName(ReturnPin));
 			Context.NetMap.Add(ReturnPin, Term);
 		}
diff --git a/Engine/Source/Editor/BlueprintGraph/Private/K2Node_Self.cpp b/Engine/Source/Editor/BlueprintGraph/Private/K2Node_Self.cpp
--- a/Engine/Source/Editor/BlueprintGraph/Private/K2Node_Self.cpp
+++ b/Engine/Source/Editor/BlueprintGraph/Private/K2Node_Self.cpp
@@ -27,7 +27,8 @@ public:
 		UEdGraphPin* VarPin = SelfNode->FindPin(UEdGraphSchema_K2::PN_Self);
 		check( VarPin );
 
-		FBPTerminal* Term = new (Context.Literals) FBPTerminal();
+		FBPTerminal* Term = new FBPTerminal();
+		Context.Literals.Add(Term);
 		Term->CopyFromPin(VarPin, VarPin->PinName);
 		Term->bIsLiteral = true;
 		Context.NetMap.Add(VarPin, Term);		
diff --git a/Engine/Source/Editor/BlueprintGraph/Private/K2Node_SetVariableOnPersistentFrame.cpp b/Engine/Source/Editor/BlueprintGraph/Private/K2Node_SetVariableOnPersistentFrame.cpp
--- a/Engine/Source/Editor/BlueprintGraph/Private/K2Node_SetVariableOnPersistentFrame.cpp
+++ b/Engine/Source/Editor/BlueprintGraph/Private/K2Node_SetVariableOnPersistentFrame.cpp
@@ -42,7 +42,8 @@ public:
 			}
 
 			// Create the term in the list
-			FBPTerminal* Term = new (Context.PersistentFrameVariableReferences)FBPTerminal();
+			FBPTerminal* Term = new FBPTerminal();
+			Context.PersistentFrameVariableReferences.Add(Term);
 			Term->CopyFromPin(Pin, Pin->PinName);
 			Term->AssociatedVarProperty = BoundProperty;
 			Context.NetMap.Add(Pin, Term);
diff --git a/Engine/Source/Editor/BlueprintGraph/Private/K2Node_Switch.cpp b/Engine/Source/Editor/BlueprintGraph/Private/K2Node_Switch.cpp
--- a/Engine/Source/Editor/BlueprintGraph/Private/K2Node_Switch.cpp
+++ b/Engine/Source/Editor/BlueprintGraph/Private/K2Node_Switch.cpp
@@ -102,7 +102,8 @@ public:
 				if ((Pin->Direction == EGPD_Output) && (Pin != DefaultPin) && (!bCanSkipUnlinkedCase || Pin->LinkedTo.Num() > 0))
 				{
 					// Create a term for the switch case value
-					FBPTerminal* CaseValueTerm = new (Context.Literals) FBPTerminal();
+					FBPTerminal* CaseValueTerm = new FBPTerminal();
+					Context.Literals.Add(CaseValueTerm);
 					CaseValueTerm->Name = Pin->PinName.ToString();
 					CaseValueTerm->Type = SwitchNode->GetInnerCaseType();
 					CaseValueTerm->SourcePin = Pin;
diff --git a/Engine/Source/Editor/BlueprintGraph/Private/MakeStructHandler.cpp b/Engine/Source/Editor/BlueprintGraph/Private/MakeStructHandler.cpp
--- a/Engine/Source/Editor/BlueprintGraph/Private/MakeStructHandler.cpp
+++ b/Engine/Source/Editor/BlueprintGraph/Private/MakeStructHandler.cpp
@@ -141,7 +141,8 @@ void FKCHandler_MakeStruct::Compile(FKismetFunctionContext& Context, UEdGraphNod
 				Schema->ConvertPropertyToPinType(OverrideProperty, /*out*/ PinType);
 
 				// Create the term in the list
-				FBPTerminal* OverrideTerm = new (Context.VariableReferences) FBPTerminal();
+				FBPTerminal* OverrideTerm = new FBPTerminal();
+				Context.VariableReferences.Add(OverrideTerm);
 				OverrideTerm->Type = PinType;
 				OverrideTerm->AssociatedVarProperty = OverrideProperty;
 				OverrideTerm->Context = OutputStructTerm;
diff --git a/Engine/Source/Editor/BlueprintGraph/Private/MathExpressionHandler.cpp b/Engine/Source/Editor/BlueprintGraph/Private/MathExpressionHandler.cpp
--- a/Engine/Source/Editor/BlueprintGraph/Private/MathExpressionHandler.cpp
+++ b/Engine/Source/Editor/BlueprintGraph/Private/MathExpressionHandler.cpp
@@ -176,7 +176,8 @@ void FKCHandler_MathExpression::RegisterNets(FKismetFunctionContext& Context, UE
 				UEdGraphNode* LinkedOwnerNode = Linked ? Linked->GetOwningNodeUnchecked() : nullptr;
 				if (LinkedOwnerNode && (InnerExitNode != LinkedOwnerNode))
 				{
-					FBPTerminal* Term = new (Context.InlineGeneratedValues) FBPTerminal();
+					FBPTerminal* Term = new FBPTerminal();
+					Context.InlineGeneratedValues.Add(Term);
 					Term->CopyFromPin(Pin, Context.NetNameMap->MakeValidName(Pin));
 					Context.NetMap.Add(Pin, Term);
 				}
diff --git a/Engine/Source/Editor/BlueprintGraph/Private/StructMemberNodeHandlers.cpp b/Engine/Source/Editor/BlueprintGraph/Private/StructMemberNodeHandlers.cpp
--- a/Engine/Source/Editor/BlueprintGraph/Private/StructMemberNodeHandlers.cpp
+++ b/Engine/Source/Editor/BlueprintGraph/Private/StructMemberNodeHandlers.cpp
@@ -23,7 +23,8 @@ static FBPTerminal* RegisterStructVar(FCompilerResultsLog& MessageLog, FKismetFu
 	if (UProperty* BoundProperty = FKismetCompilerUtilities::FindNamedPropertyInScope(SearchScope, MemberSetNode->GetVarName()))
 	{
 		// Create the term in the list
-		FBPTerminal* Term = new (Context.VariableReferences) FBPTerminal();
+		FBPTerminal* Term = new FBPTerminal();
+		Context.VariableReferences.Add(Term);
 
 		Schema->ConvertPropertyToPinType(BoundProperty, /*out*/ Term->Type);
 		Term->Source = MemberSetNode;
@@ -61,7 +62,8 @@ static void ResolveAndRegisterScopedStructTerm(FCompilerResultsLog& MessageLog,
 	if (UProperty* BoundProperty = FindField<UProperty>(StructType, Net->PinName))
 	{
 		// Create the term in the list
-		FBPTerminal* Term = new (Context.VariableReferences) FBPTerminal();
+		FBPTerminal* Term = new FBPTerminal();
+		Context.VariableReferences.Add(Term);
 		Term->CopyFromPin(Net, Net->PinName.ToString());
 		Term->AssociatedVarProperty = BoundProperty;
 		Context.NetMap.Add(Net, Term);
diff --git a/Engine/Source/Editor/KismetCompiler/Private/KismetCompiler.cpp b/Engine/Source/Editor/KismetCompiler/Private/KismetCompiler.cpp
--- a/Engine/Source/Editor/KismetCompiler/Private/KismetCompiler.cpp
+++ b/Engine/Source/Editor/KismetCompiler/Private/KismetCompiler.cpp
@@ -2813,7 +2813,8 @@ void FKismetCompilerContext::CreateFunctionStubForEvent(UK2Node_Event* SrcEventN
 	ChildStubGraph->SetFlags(RF_Transient);
 	MessageLog.NotifyIntermediateObjectCreation(ChildStubGraph, SrcEventNode);
 
-	FKismetFunctionContext& StubContext = *new (FunctionList) FKismetFunctionContext(MessageLog, Schema, NewClass, Blueprint, CompileOptions.DoesRequireCppCodeGeneration());
+	FKismetFunctionContext& StubContext = *new FKismetFunctionContext(MessageLog, Schema, NewClass, Blueprint, CompileOptions.DoesRequireCppCodeGeneration());
+	FunctionList.Add(&StubContext);
 	StubContext.SourceGraph = ChildStubGraph;
 
 	StubContext.SourceEventFromStubGraph = SrcEventNode;
@@ -3273,7 +3274,8 @@ void FKismetCompilerContext::CreateAndProcessUbergraph()
 
 		// Do some cursory validation (pin types match, inputs to outputs, pins never point to their parent node, etc...)
 		{
-			UbergraphContext = new (FunctionList) FKismetFunctionContext(MessageLog, Schema, NewClass, Blueprint, CompileOptions.DoesRequireCppCodeGeneration());
+			UbergraphContext = new FKismetFunctionContext(MessageLog, Schema, NewClass, Blueprint, CompileOptions.DoesRequireCppCodeGeneration());
+			FunctionList.Add(UbergraphContext);
 			UbergraphContext->SourceGraph = ConsolidatedEventGraph;
 			UbergraphContext->MarkAsEventGraph();
 			UbergraphContext->MarkAsInternalOrCppUseOnly();
@@ -3655,7 +3657,8 @@ void FKismetCompilerContext::ProcessOneFunctionGraph(UEdGraph* SourceGraph, bool
 	if ((CompileOptions.CompileType == EKismetCompileType::SkeletonOnly) || ValidateGraphIsWellFormed(FunctionGraph))
 	{
 		const UEdGraphSchema_K2* FunctionGraphSchema = CastChecked<const UEdGraphSchema_K2>(FunctionGraph->GetSchema());
-		FKismetFunctionContext& Context = *new (FunctionList) FKismetFunctionContext(MessageLog, FunctionGraphSchema, NewClass, Blueprint, CompileOptions.DoesRequireCppCodeGeneration());
+		FKismetFunctionContext& Context = *new FKismetFunctionContext(MessageLog, FunctionGraphSchema, NewClass, Blueprint, CompileOptions.DoesRequireCppCodeGeneration());
+		FunctionList.Add(&Context);
 		Context.SourceGraph = FunctionGraph;
 
 		if(FBlueprintEditorUtils::IsDelegateSignatureGraph(SourceGraph))
@@ -3763,7 +3766,9 @@ void FKismetCompilerContext::CreateFunctionList()
 
 FKismetFunctionContext* FKismetCompilerContext::CreateFunctionContext()
 {
-	return new (FunctionList) FKismetFunctionContext(MessageLog, Schema, NewClass, Blueprint, CompileOptions.DoesRequireCppCodeGeneration());
+	FKismetFunctionContext* Result = new FKismetFunctionContext(MessageLog, Schema, NewClass, Blueprint, CompileOptions.DoesRequireCppCodeGeneration());
+	FunctionList.Add(Result);
+	return Result;
 }
 
 /** Compile a blueprint into a class and a set of functions */
diff --git a/Engine/Source/Editor/KismetCompiler/Private/KismetCompilerMisc.cpp b/Engine/Source/Editor/KismetCompiler/Private/KismetCompilerMisc.cpp
--- a/Engine/Source/Editor/KismetCompiler/Private/KismetCompilerMisc.cpp
+++ b/Engine/Source/Editor/KismetCompiler/Private/KismetCompilerMisc.cpp
@@ -1566,7 +1566,8 @@ void FNodeHandlingFunctor::ResolveAndRegisterScopedTerm(FKismetFunctionContext&
 	{
 		UBlueprintEditorSettings* Settings = GetMutableDefault<UBlueprintEditorSettings>();
 		// Create the term in the list
-		FBPTerminal* Term = new (NetArray) FBPTerminal();
+		FBPTerminal* Term = new FBPTerminal();
+		NetArray.Add(Term);
 		Term->CopyFromPin(Net, Net->PinName);
 		Term->AssociatedVarProperty = BoundProperty;
 		Term->bPassedByReference = true;
@@ -2230,15 +2231,25 @@ FBPTerminal* FKismetFunctionContext::CreateLocalTerminal(ETerminalSpecification
 	{
 	case ETerminalSpecification::TS_ForcedShared:
 		ensure(IsEventGraph());
-		Result = new (EventGraphLocals)FBPTerminal();
+		Result = new FBPTerminal();
+		EventGraphLocals.Add(Result);
 		break;
 	case ETerminalSpecification::TS_Literal:
-		Result = new (Literals) FBPTerminal();
+		Result = new FBPTerminal();
+		Literals.Add(Result);
 		Result->bIsLiteral = true;
 		break;
 	default:
 		const bool bIsLocal = !IsEventGraph();
-		Result = new (bIsLocal ? Locals : EventGraphLocals) FBPTerminal();
+		Result = new FBPTerminal();
+		if (bIsLocal)
+		{
+			Locals.Add(Result);
+		}
+		else
+		{
+			EventGraphLocals.Add(Result);
+		}
 		Result->SetVarTypeLocal(bIsLocal);
 		break;
 	}
@@ -2262,7 +2273,15 @@ FBPTerminal* FKismetFunctionContext::CreateLocalTerminalFromPinAutoChooseScope(U
 		// Pin's connections are checked, to tell if created terminal is shared, or if it could be a local variable.
 		bSharedTerm = FEventGraphUtils::PinRepresentsSharedTerminal(*Net, MessageLog);
 	}
-	FBPTerminal* Term = new (bSharedTerm ? EventGraphLocals : Locals) FBPTerminal();
+	FBPTerminal* Term = new FBPTerminal();
+	if (bSharedTerm)
+	{
+		EventGraphLocals.Add(Term);
+	}
+	else
+	{
+		Locals.Add(Term);
+	}
 	Term->CopyFromPin(Net, MoveTemp(NewName));
 	return Term;
 }
diff --git a/Engine/Source/Editor/KismetCompiler/Public/KismetCompiledFunctionContext.h b/Engine/Source/Editor/KismetCompiler/Public/KismetCompiledFunctionContext.h
--- a/Engine/Source/Editor/KismetCompiler/Public/KismetCompiledFunctionContext.h
+++ b/Engine/Source/Editor/KismetCompiler/Public/KismetCompiledFunctionContext.h
@@ -199,7 +199,8 @@ public:
 
 	FBPTerminal* RegisterLiteral(UEdGraphPin* Net)
 	{
-		FBPTerminal* Term = new (Literals) FBPTerminal();
+		FBPTerminal* Term = new FBPTerminal();
+		Literals.Add(Term);
 		Term->CopyFromPin(Net, Net->DefaultValue);
 		Term->ObjectLiteral = Net->DefaultObject;
 		Term->TextLiteral = Net->DefaultTextValue;
diff --git a/Engine/Source/Editor/MaterialEditor/Private/MaterialEditor.cpp b/Engine/Source/Editor/MaterialEditor/Private/MaterialEditor.cpp
--- a/Engine/Source/Editor/MaterialEditor/Private/MaterialEditor.cpp
+++ b/Engine/Source/Editor/MaterialEditor/Private/MaterialEditor.cpp
@@ -4204,7 +4204,8 @@ FMatExpressionPreview* FMaterialEditor::GetExpressionPreview(UMaterialExpression
 		if( !Preview )
 		{
 			bNewlyCreated = true;
-			Preview = new(ExpressionPreviews) FMatExpressionPreview(MaterialExpression);
+			Preview = new FMatExpressionPreview(MaterialExpression);
+			ExpressionPreviews.Add(Preview);
 			Preview->CacheShaders(GMaxRHIShaderPlatform, true);
 		}
 		return Preview;
diff --git a/Engine/Source/Editor/UnrealEd/Private/Factories/SkeletalMeshImport.cpp b/Engine/Source/Editor/UnrealEd/Private/Factories/SkeletalMeshImport.cpp
--- a/Engine/Source/Editor/UnrealEd/Private/Factories/SkeletalMeshImport.cpp
+++ b/Engine/Source/Editor/UnrealEd/Private/Factories/SkeletalMeshImport.cpp
@@ -951,7 +951,8 @@ void RestoreExistingSkelMeshData(ExistingSkelMeshData* MeshData, USkeletalMesh*
 						}
 						else
 						{
-							FSkeletalMeshLODModel* NewLODModel = new(SkeletalMesh->GetImportedModel()->LODModels) FSkeletalMeshLODModel(LODModel);
+							FSkeletalMeshLODModel* NewLODModel = new FSkeletalMeshLODModel(LODModel);
+							SkeletalMesh->GetImportedModel()->LODModels.Add(NewLODModel);
 							SkeletalMesh->AddLODInfo(LODInfo);
 						}
 					}
@@ -967,7 +968,8 @@ void RestoreExistingSkelMeshData(ExistingSkelMeshData* MeshData, USkeletalMesh*
 					int32 LODIndex = Index + 1;
 					FSkeletalMeshLODModel& LODModel = MeshData->ExistingLODModels[Index];
 					FSkeletalMeshLODInfo& LODInfo = MeshData->ExistingLODInfo[Index];
-					FSkeletalMeshLODModel* NewLODModel = new(SkeletalMesh->GetImportedModel()->LODModels) FSkeletalMeshLODModel(LODModel);
+					FSkeletalMeshLODModel* NewLODModel = new FSkeletalMeshLODModel(LODModel);
+					SkeletalMesh->GetImportedModel()->LODModels.Add(NewLODModel);
 					// add LOD info back
 					SkeletalMesh->AddLODInfo(LODInfo);
 					//Apply the new skinning to the existing LOD geometry
diff --git a/Engine/Source/Editor/UnrealEd/Private/Fbx/FbxSkeletalMeshImport.cpp b/Engine/Source/Editor/UnrealEd/Private/Fbx/FbxSkeletalMeshImport.cpp
--- a/Engine/Source/Editor/UnrealEd/Private/Fbx/FbxSkeletalMeshImport.cpp
+++ b/Engine/Source/Editor/UnrealEd/Private/Fbx/FbxSkeletalMeshImport.cpp
@@ -1595,7 +1595,7 @@ USkeletalMesh* UnFbx::FFbxImporter::ImportSkeletalMesh(FImportSkeletalMeshArgs &
 	FSkeletalMeshModel *ImportedResource = SkeletalMesh->GetImportedModel();
 	check(ImportedResource->LODModels.Num() == 0);
 	ImportedResource->LODModels.Empty();
-	new(ImportedResource->LODModels)FSkeletalMeshLODModel();
+	ImportedResource->LODModels.Add(new FSkeletalMeshLODModel());
 
 	FSkeletalMeshLODModel& LODModel = ImportedResource->LODModels[0];
 
@@ -1723,7 +1723,7 @@ USkeletalMesh* UnFbx::FFbxImporter::ImportSkeletalMesh(FImportSkeletalMeshArgs &
 			USkinnedMeshComponent* SkinComp = *It;
 			if (SkinComp->SkeletalMesh == SkeletalMesh)
 			{
-				new(ComponentContexts) FComponentReregisterContext(SkinComp);
+				ComponentContexts.Add(new FComponentReregisterContext(SkinComp));
 			}
 		}
 	}
@@ -3454,7 +3454,7 @@ void UnFbx::FFbxImporter::InsertNewLODToBaseSkeletalMesh(USkeletalMesh* InSkelet
 	// If we want to add this as a new LOD to this mesh - add to LODModels/LODInfo array.
 	if (DesiredLOD == DestImportedResource->LODModels.Num())
 	{
-		new(DestImportedResource->LODModels)FSkeletalMeshLODModel();
+		DestImportedResource->LODModels.Add(new FSkeletalMeshLODModel());
 
 		// Add element to LODInfo array.
 		BaseSkeletalMesh->AddLODInfo();
@@ -4217,8 +4217,9 @@ void UnFbx::FFbxImporter::ImportMorphTargetsInternal( TArray<FbxNode*>& SkelMesh
 
 			TArray< FMorphTargetDelta >* Deltas = Results[ NewMorphDeltasIdx ];
 
-			FAsyncTask<FAsyncImportMorphTargetWork>* NewWork = new (PendingWork)FAsyncTask<FAsyncImportMorphTargetWork>( &BaseLODModel, BaseSkelMesh->RefSkeleton, BaseImportData,
+			FAsyncTask<FAsyncImportMorphTargetWork>* NewWork = new FAsyncTask<FAsyncImportMorphTargetWork>( &BaseLODModel, BaseSkelMesh->RefSkeleton, BaseImportData,
 				MoveTemp( ShapeImportData.Points ), ImportOptions, *Deltas, BaseIndexData, BaseWedgePointIndices, WedgePointToVertexIndexMap, OverlappingVertices, MoveTemp( ModifiedPoints ), WedgeToFaces, MeshDataBundle, TangentZ);
+			PendingWork.Add(NewWork);
 
 			NewWork->StartBackgroundTask(GLargeThreadPool);
 			CurrentNumTasks++;
diff --git a/Engine/Source/Programs/PixelStreaming/WebRTCProxy/src/Logging.cpp b/Engine/Source/Programs/PixelStreaming/WebRTCProxy/src/Logging.cpp
--- a/Engine/Source/Programs/PixelStreaming/WebRTCProxy/src/Logging.cpp
+++ b/Engine/Source/Programs/PixelStreaming/WebRTCProxy/src/Logging.cpp
@@ -103,4 +103,6 @@ void ILogOutput::LogToAll(
 	{
 		Out->Log(File, Line, Category, Verbosity, Buf);
 	}
+
+	va_end()
 }
diff --git a/Engine/Source/Programs/UnrealLightmass/Private/Lighting/LightingSystem.cpp b/Engine/Source/Programs/UnrealLightmass/Private/Lighting/LightingSystem.cpp
--- a/Engine/Source/Programs/UnrealLightmass/Private/Lighting/LightingSystem.cpp
+++ b/Engine/Source/Programs/UnrealLightmass/Private/Lighting/LightingSystem.cpp
@@ -542,7 +542,8 @@ void FStaticLightingSystem::MultithreadProcess()
 	// Spawn the static lighting threads.
 	for(int32 ThreadIndex = 0;ThreadIndex < NumStaticLightingThreads;ThreadIndex++)
 	{
-		FMappingProcessingThreadRunnable* ThreadRunnable = new(Threads) FMappingProcessingThreadRunnable(this, ThreadIndex, StaticLightingTask_ProcessMappings);
+		FMappingProcessingThreadRunnable* ThreadRunnable = new FMappingProcessingThreadRunnable(this, ThreadIndex, StaticLightingTask_ProcessMappings);
+		Threads.Add(ThreadRunnable);
 		const FString ThreadName = FString::Printf(TEXT("MappingProcessingThread%u"), ThreadIndex);
 		ThreadRunnable->Thread = FRunnableThread::Create(ThreadRunnable, *ThreadName);
 	}
diff --git a/Engine/Source/Programs/UnrealLightmass/Private/Lighting/PhotonMapping.cpp b/Engine/Source/Programs/UnrealLightmass/Private/Lighting/PhotonMapping.cpp
--- a/Engine/Source/Programs/UnrealLightmass/Private/Lighting/PhotonMapping.cpp
+++ b/Engine/Source/Programs/UnrealLightmass/Private/Lighting/PhotonMapping.cpp
@@ -261,7 +261,8 @@ void FStaticLightingSystem::EmitDirectPhotons(
 	TIndirectArray<FDirectPhotonEmittingThreadRunnable> DirectPhotonEmittingThreads;
 	for (int32 ThreadIndex = 1; ThreadIndex < NumStaticLightingThreads; ThreadIndex++)
 	{
-		FDirectPhotonEmittingThreadRunnable* ThreadRunnable = new(DirectPhotonEmittingThreads) FDirectPhotonEmittingThreadRunnable(this, ThreadIndex, Input);
+		FDirectPhotonEmittingThreadRunnable* ThreadRunnable = new FDirectPhotonEmittingThreadRunnable(this, ThreadIndex, Input);
+		DirectPhotonEmittingThreads.Add(ThreadRunnable);
 		const FString ThreadName = FString::Printf(TEXT("DirectPhotonEmittingThread%u"), ThreadIndex);
 		ThreadRunnable->Thread = FRunnableThread::Create(ThreadRunnable, *ThreadName);
 	}
@@ -731,7 +732,8 @@ void FStaticLightingSystem::EmitIndirectPhotons(
 	TIndirectArray<FIndirectPhotonEmittingThreadRunnable> IndirectPhotonEmittingThreads;
 	for (int32 ThreadIndex = 1; ThreadIndex < NumStaticLightingThreads; ThreadIndex++)
 	{
-		FIndirectPhotonEmittingThreadRunnable* ThreadRunnable = new(IndirectPhotonEmittingThreads) FIndirectPhotonEmittingThreadRunnable(this, ThreadIndex, Input);
+		FIndirectPhotonEmittingThreadRunnable* ThreadRunnable = new FIndirectPhotonEmittingThreadRunnable(this, ThreadIndex, Input);
+		IndirectPhotonEmittingThreads.Add(ThreadRunnable);
 		const FString ThreadName = FString::Printf(TEXT("IndirectPhotonEmittingThread%u"), ThreadIndex);
 		ThreadRunnable->Thread = FRunnableThread::Create(ThreadRunnable, *ThreadName);
 	}
@@ -1226,7 +1228,8 @@ void FStaticLightingSystem::MarkIrradiancePhotons(const FBoxSphereBounds& Import
 	IrradiancePhotonMarkingThreads.Empty(NumStaticLightingThreads);
 	for(int32 ThreadIndex = 1; ThreadIndex < NumStaticLightingThreads; ThreadIndex++)
 	{
-		FIrradiancePhotonMarkingThreadRunnable* ThreadRunnable = new(IrradiancePhotonMarkingThreads) FIrradiancePhotonMarkingThreadRunnable(this, ThreadIndex, IrradiancePhotons);
+		FIrradiancePhotonMarkingThreadRunnable* ThreadRunnable = new FIrradiancePhotonMarkingThreadRunnable(this, ThreadIndex, IrradiancePhotons);
+		IrradiancePhotonMarkingThreads.Add(ThreadRunnable);
 		const FString ThreadName = FString::Printf(TEXT("IrradiancePhotonMarkingThread%u"), ThreadIndex);
 		ThreadRunnable->Thread = FRunnableThread::Create(ThreadRunnable, *ThreadName);
 	}
@@ -1390,7 +1393,8 @@ void FStaticLightingSystem::CalculateIrradiancePhotons(const FBoxSphereBounds& I
 	IrradiancePhotonThreads.Empty(NumStaticLightingThreads);
 	for(int32 ThreadIndex = 1; ThreadIndex < NumStaticLightingThreads; ThreadIndex++)
 	{
-		FIrradiancePhotonCalculatingThreadRunnable* ThreadRunnable = new(IrradiancePhotonThreads) FIrradiancePhotonCalculatingThreadRunnable(this, ThreadIndex, IrradiancePhotons);
+		FIrradiancePhotonCalculatingThreadRunnable* ThreadRunnable = new FIrradiancePhotonCalculatingThreadRunnable(this, ThreadIndex, IrradiancePhotons);
+		IrradiancePhotonThreads.Add(ThreadRunnable);
 		const FString ThreadName = FString::Printf(TEXT("IrradiancePhotonCalculatingThread%u"), ThreadIndex);
 		ThreadRunnable->Thread = FRunnableThread::Create(ThreadRunnable, *ThreadName);
 	}
@@ -1589,7 +1593,8 @@ void FStaticLightingSystem::CacheIrradiancePhotons()
 	check(PhotonMappingSettings.bCacheIrradiancePhotonsOnSurfaces);
 	for(int32 ThreadIndex = 1; ThreadIndex < NumStaticLightingThreads; ThreadIndex++)
 	{
-		FMappingProcessingThreadRunnable* ThreadRunnable = new(IrradiancePhotonCachingThreads) FMappingProcessingThreadRunnable(this, ThreadIndex, StaticLightingTask_CacheIrradiancePhotons);
+		FMappingProcessingThreadRunnable* ThreadRunnable = new FMappingProcessingThreadRunnable(this, ThreadIndex, StaticLightingTask_CacheIrradiancePhotons);
+		IrradiancePhotonCachingThreads.Add(ThreadRunnable);
 		const FString ThreadName = FString::Printf(TEXT("IrradiancePhotonCachingThread%u"), ThreadIndex);
 		ThreadRunnable->Thread = FRunnableThread::Create(ThreadRunnable, *ThreadName);
 	}
diff --git a/Engine/Source/Programs/UnrealLightmass/Private/Lighting/Radiosity.cpp b/Engine/Source/Programs/UnrealLightmass/Private/Lighting/Radiosity.cpp
--- a/Engine/Source/Programs/UnrealLightmass/Private/Lighting/Radiosity.cpp
+++ b/Engine/Source/Programs/UnrealLightmass/Private/Lighting/Radiosity.cpp
@@ -166,7 +166,8 @@ void FStaticLightingSystem::SetupRadiosity()
 
 	for(int32 ThreadIndex = 1; ThreadIndex < NumStaticLightingThreads; ThreadIndex++)
 	{
-		FMappingProcessingThreadRunnable* ThreadRunnable = new(RadiositySetupThreads) FMappingProcessingThreadRunnable(this, ThreadIndex, StaticLightingTask_RadiositySetup);
+		FMappingProcessingThreadRunnable* ThreadRunnable = new FMappingProcessingThreadRunnable(this, ThreadIndex, StaticLightingTask_RadiositySetup);
+		RadiositySetupThreads.Add(ThreadRunnable);
 		const FString ThreadName = FString::Printf(TEXT("RadiositySetupThread%u"), ThreadIndex);
 		ThreadRunnable->Thread = FRunnableThread::Create(ThreadRunnable, *ThreadName);
 	}
@@ -456,7 +457,8 @@ void FStaticLightingSystem::RunRadiosityIterations()
 
 		for(int32 ThreadIndex = 1; ThreadIndex < NumStaticLightingThreads; ThreadIndex++)
 		{
-			FMappingProcessingThreadRunnable* ThreadRunnable = new(RadiosityIterationThreads) FMappingProcessingThreadRunnable(this, ThreadIndex, StaticLightingTask_RadiosityIterations);
+			FMappingProcessingThreadRunnable* ThreadRunnable = new FMappingProcessingThreadRunnable(this, ThreadIndex, StaticLightingTask_RadiosityIterations);
+			RadiosityIterationThreads.Add(ThreadRunnable);
 			const FString ThreadName = FString::Printf(TEXT("RadiosityIterationThread%u"), ThreadIndex);
 			ThreadRunnable->Thread = FRunnableThread::Create(ThreadRunnable, *ThreadName);
 		}
diff --git a/Engine/Source/Programs/UnrealLightmass/Private/Lighting/TextureMapping.cpp b/Engine/Source/Programs/UnrealLightmass/Private/Lighting/TextureMapping.cpp
--- a/Engine/Source/Programs/UnrealLightmass/Private/Lighting/TextureMapping.cpp
+++ b/Engine/Source/Programs/UnrealLightmass/Private/Lighting/TextureMapping.cpp
@@ -206,7 +206,8 @@ void FStaticLightingSystem::FinalizeSurfaceCache()
 {
 	for(int32 ThreadIndex = 1; ThreadIndex < NumStaticLightingThreads; ThreadIndex++)
 	{
-		FMappingProcessingThreadRunnable* ThreadRunnable = new(FinalizeSurfaceCacheThreads) FMappingProcessingThreadRunnable(this, ThreadIndex, StaticLightingTask_FinalizeSurfaceCache);
+		FMappingProcessingThreadRunnable* ThreadRunnable = new FMappingProcessingThreadRunnable(this, ThreadIndex, StaticLightingTask_FinalizeSurfaceCache);
+		FinalizeSurfaceCacheThreads.Add(ThreadRunnable);
 		const FString ThreadName = FString::Printf(TEXT("FinalizeSurfaceCacheThread%u"), ThreadIndex);
 		ThreadRunnable->Thread = FRunnableThread::Create(ThreadRunnable, *ThreadName);
 	}
diff --git a/Engine/Source/Runtime/ClothingSystemRuntime/Private/Assets/ClothingAsset.cpp b/Engine/Source/Runtime/ClothingSystemRuntime/Private/Assets/ClothingAsset.cpp
--- a/Engine/Source/Runtime/ClothingSystemRuntime/Private/Assets/ClothingAsset.cpp
+++ b/Engine/Source/Runtime/ClothingSystemRuntime/Private/Assets/ClothingAsset.cpp
@@ -241,7 +241,7 @@ bool UClothingAsset::BindToSkeletalMesh(USkeletalMesh* InSkelMesh, int32 InMeshL
 		USkeletalMeshComponent* Component = *It;
 		if (Component && !Component->IsTemplate() && Component->SkeletalMesh == InSkelMesh)
 		{
-			new(ComponentContexts) FComponentReregisterContext(Component);
+			ComponentContexts.Add(new FComponentReregisterContext(Component));
 		}
 	}
 
diff --git a/Engine/Source/Runtime/Core/Private/Stats/StatsCommand.cpp b/Engine/Source/Runtime/Core/Private/Stats/StatsCommand.cpp
--- a/Engine/Source/Runtime/Core/Private/Stats/StatsCommand.cpp
+++ b/Engine/Source/Runtime/Core/Private/Stats/StatsCommand.cpp
@@ -1425,7 +1425,7 @@ struct FHUDGroupManager
 				FInternalGroup& InternalGroup = GroupIt.Value();
 
 				// Create a new hud group.
-				new(ToGame->ActiveStatGroups) FActiveStatGroupInfo();
+				ToGame->ActiveStatGroups.Add(new FActiveStatGroupInfo());
 				FActiveStatGroupInfo& HudGroup = ToGame->ActiveStatGroups.Last();
 
 				ToGame->GroupNames.Add( GroupName );
diff --git a/Engine/Source/Runtime/Core/Public/Containers/ChunkedArray.h b/Engine/Source/Runtime/Core/Public/Containers/ChunkedArray.h
--- a/Engine/Source/Runtime/Core/Public/Containers/ChunkedArray.h
+++ b/Engine/Source/Runtime/Core/Public/Containers/ChunkedArray.h
@@ -65,7 +65,7 @@ public:
 		Chunks.Empty(NumChunks);
 		for(int32 ChunkIndex = 0;ChunkIndex < NumChunks;ChunkIndex++)
 		{
-			new(Chunks) FChunk;
+			Chunks.Add(new FChunk);
 		}
 	}
 
@@ -203,7 +203,7 @@ public:
 		{
 			if (NumElements % NumElementsPerChunk == 0)
 			{
-				new(Chunks) FChunk;
+				Chunks.Add(new FChunk);
 			}
 			NumElements++;
 		}
diff --git a/Engine/Source/Runtime/Core/Public/Containers/IndirectArray.h b/Engine/Source/Runtime/Core/Public/Containers/IndirectArray.h
--- a/Engine/Source/Runtime/Core/Public/Containers/IndirectArray.h
+++ b/Engine/Source/Runtime/Core/Public/Containers/IndirectArray.h
@@ -193,7 +193,7 @@ public:
 			Empty(NewNum);
 			for (int32 Index = 0; Index < NewNum; Index++)
 			{
-				new(*this) T;
+				Add(new T);
 			}
 			for (int32 Index = 0; Index < NewNum; Index++)
 			{
@@ -230,7 +230,9 @@ public:
 			A.Empty(NewNum);
 			for (int32 Index = 0; Index < NewNum; Index++)
 			{
-				Ar << *new(A)T;
+				T* NewElement = new T;
+				Ar << *NewElement;
+				A.Add(NewElement);
 			}
 		}
 		else
@@ -273,11 +275,7 @@ public:
 		T** Element = GetData() + Index;
 		for (int32 ElementId = Count; ElementId; --ElementId)
 		{
-			// We need a typedef here because VC won't compile the destructor call below if T itself has a member called T
-			typedef T IndirectArrayDestructElementType;
-
-			(*Element)->IndirectArrayDestructElementType::~IndirectArrayDestructElementType();
-			FMemory::Free(*Element);
+			delete *Element;
 			++Element;
 		}
 		Array.RemoveAt(Index, Count, bAllowShrinking);
@@ -303,11 +301,7 @@ public:
 		T** Element = GetData() + Index;
 		for (int32 ElementId = Count; ElementId; --ElementId)
 		{
-			// We need a typedef here because VC won't compile the destructor call below if T itself has a member called T
-			typedef T IndirectArrayDestructElementType;
-
-			(*Element)->IndirectArrayDestructElementType::~IndirectArrayDestructElementType();
-			FMemory::Free(*Element);
+			delete *Element;
 			++Element;
 		}
 		Array.RemoveAtSwap(Index, Count, bAllowShrinking);
@@ -428,11 +422,7 @@ private:
 		T** Element = GetData();
 		for (int32 Index = Array.Num(); Index; --Index)
 		{
-			// We need a typedef here because VC won't compile the destructor call below if T itself has a member called T
-			typedef T IndirectArrayDestructElementType;
-
-			(*Element)->IndirectArrayDestructElementType::~IndirectArrayDestructElementType();
-			FMemory::Free(*Element);
+			delete *Element;
 			++Element;
 		}
 	}
@@ -457,16 +447,18 @@ struct TContainerTraits<TIndirectArray<T, Allocator> >
 	enum { MoveWillEmptyContainer = TContainerTraitsBase<typename TIndirectArray<T, Allocator>::InternalArrayType>::MoveWillEmptyContainer };
 };
 
-
-template <typename T,typename Allocator> void* operator new( size_t Size, TIndirectArray<T,Allocator>& Array )
+template <typename T,typename Allocator>
+DEPRECATED(4.22, "Placement new on TIndirectArray has been deprecated - users should call Add() passing a pointer to an object created with new.")
+void* operator new( size_t Size, TIndirectArray<T,Allocator>& Array )
 {
 	check(Size == sizeof(T));
 	const int32 Index = Array.Add((T*)FMemory::Malloc(Size));
 	return &Array[Index];
 }
 
-
-template <typename T,typename Allocator> void* operator new( size_t Size, TIndirectArray<T,Allocator>& Array, int32 Index )
+template <typename T,typename Allocator>
+DEPRECATED(4.22, "Placement new on TIndirectArray has been deprecated - users should call Insert() passing a pointer to an object created with new.")
+void* operator new( size_t Size, TIndirectArray<T,Allocator>& Array, int32 Index )
 {
 	check(Size == sizeof(T));
 	Array.Insert((T*)FMemory::Malloc(Size), Index);
diff --git a/Engine/Source/Runtime/Engine/Private/AudioDerivedData.cpp b/Engine/Source/Runtime/Engine/Private/AudioDerivedData.cpp
--- a/Engine/Source/Runtime/Engine/Private/AudioDerivedData.cpp
+++ b/Engine/Source/Runtime/Engine/Private/AudioDerivedData.cpp
@@ -314,7 +314,8 @@ class FStreamedAudioCacheDerivedDataWorker : public FNonAbandonableTask
 						const int32 AudioDataSize = ChunkBuffers[ChunkIndex].Num();
 						const int32 ZeroPadBytes = FMath::Max(MaxChunkSize - AudioDataSize, 0);
 
-						FStreamedAudioChunk* NewChunk = new(DerivedData->Chunks) FStreamedAudioChunk();
+						FStreamedAudioChunk* NewChunk = new FStreamedAudioChunk();
+						DerivedData->Chunks.Add(NewChunk);
 
 						// Store both the audio data size and the data size so decoders will know what portion of the bulk data is real audio
 						NewChunk->AudioDataSize = AudioDataSize;
@@ -339,7 +340,8 @@ class FStreamedAudioCacheDerivedDataWorker : public FNonAbandonableTask
 				else
 				{
 					// Could not split so copy compressed data into a single chunk
-					FStreamedAudioChunk* NewChunk = new(DerivedData->Chunks) FStreamedAudioChunk();
+					FStreamedAudioChunk* NewChunk = new FStreamedAudioChunk();
+					DerivedData->Chunks.Add(NewChunk);
 					NewChunk->DataSize = CompressedBuffer.Num();
 					NewChunk->AudioDataSize = NewChunk->DataSize;
 
@@ -740,7 +742,7 @@ void FStreamedAudioPlatformData::Serialize(FArchive& Ar, USoundWave* Owner)
 		Chunks.Empty(NumChunks);
 		for (int32 ChunkIndex = 0; ChunkIndex < NumChunks; ++ChunkIndex)
 		{
-			new(Chunks) FStreamedAudioChunk();
+			Chunks.Add(new FStreamedAudioChunk());
 		}
 	}
 	for (int32 ChunkIndex = 0; ChunkIndex < NumChunks; ++ChunkIndex)
diff --git a/Engine/Source/Runtime/Engine/Private/AudioStreaming.cpp b/Engine/Source/Runtime/Engine/Private/AudioStreaming.cpp
--- a/Engine/Source/Runtime/Engine/Private/AudioStreaming.cpp
+++ b/Engine/Source/Runtime/Engine/Private/AudioStreaming.cpp
@@ -314,12 +314,13 @@ void FStreamingWaveData::BeginPendingRequests(const TArray<uint32>& IndicesToLoa
 				INC_DWORD_STAT_BY(STAT_AudioMemorySize, ChunkSize);
 				INC_DWORD_STAT_BY(STAT_AudioMemory, ChunkSize);
 
-				FAsyncStreamDerivedChunkTask* Task = new(PendingAsyncStreamDerivedChunkTasks)FAsyncStreamDerivedChunkTask(
+				FAsyncStreamDerivedChunkTask* Task = new FAsyncStreamDerivedChunkTask(
 					Chunk.DerivedDataKey,
 					ChunkStorage->Data,
 					ChunkSize,
 					&PendingChunkChangeRequestStatus
 					);
+				PendingAsyncStreamDerivedChunkTasks.Add(Task);
 				Task->StartBackgroundTask();
 			}
 			else
diff --git a/Engine/Source/Runtime/Engine/Private/BlueprintGeneratedClass.cpp b/Engine/Source/Runtime/Engine/Private/BlueprintGeneratedClass.cpp
--- a/Engine/Source/Runtime/Engine/Private/BlueprintGeneratedClass.cpp
+++ b/Engine/Source/Runtime/Engine/Private/BlueprintGeneratedClass.cpp
@@ -397,7 +397,8 @@ bool UBlueprintGeneratedClass::BuildCustomPropertyListForPostConstruction(FCusto
 				if (UStructProperty* StructProperty = Cast<UStructProperty>(Property))
 				{
 					// Create a new node for the struct property.
-					*CurrentNodePtr = new(CustomPropertyListForPostConstruction) FCustomPropertyListNode(Property, Idx);
+					*CurrentNodePtr = new FCustomPropertyListNode(Property, Idx);
+					CustomPropertyListForPostConstruction.Add(*CurrentNodePtr);
 
 					// Recursively gather up all struct fields that differ and assign to the current node's sub property list.
 					if (BuildCustomPropertyListForPostConstruction((*CurrentNodePtr)->SubPropertyList, StructProperty->Struct, PropertyValue, DefaultPropertyValue))
@@ -417,7 +418,8 @@ bool UBlueprintGeneratedClass::BuildCustomPropertyListForPostConstruction(FCusto
 				else if (UArrayProperty* ArrayProperty = Cast<UArrayProperty>(Property))
 				{
 					// Create a new node for the array property.
-					*CurrentNodePtr = new(CustomPropertyListForPostConstruction) FCustomPropertyListNode(Property, Idx);
+					*CurrentNodePtr = new FCustomPropertyListNode(Property, Idx);
+					CustomPropertyListForPostConstruction.Add(*CurrentNodePtr);
 
 					// Recursively gather up all array item indices that differ and assign to the current node's sub property list.
 					if (BuildCustomArrayPropertyListForPostConstruction(ArrayProperty, (*CurrentNodePtr)->SubPropertyList, PropertyValue, DefaultPropertyValue))
@@ -437,7 +439,8 @@ bool UBlueprintGeneratedClass::BuildCustomPropertyListForPostConstruction(FCusto
 				else if (!Property->Identical(PropertyValue, DefaultPropertyValue))
 				{
 					// Create a new node, link it into the chain and add it into the array.
-					*CurrentNodePtr = new(CustomPropertyListForPostConstruction) FCustomPropertyListNode(Property, Idx);
+					*CurrentNodePtr = new FCustomPropertyListNode(Property, Idx);
+					CustomPropertyListForPostConstruction.Add(*CurrentNodePtr);
 
 					// Advance to the next node ptr.
 					CurrentNodePtr = &(*CurrentNodePtr)->PropertyListNext;
@@ -468,7 +471,8 @@ bool UBlueprintGeneratedClass::BuildCustomArrayPropertyListForPostConstruction(U
 			if (UStructProperty* InnerStructProperty = Cast<UStructProperty>(ArrayProperty->Inner))
 			{
 				// Create a new node for the item value at this index.
-				*CurrentArrayNodePtr = new(CustomPropertyListForPostConstruction) FCustomPropertyListNode(ArrayProperty, ArrayValueIndex);
+				*CurrentArrayNodePtr = new FCustomPropertyListNode(ArrayProperty, ArrayValueIndex);
+				CustomPropertyListForPostConstruction.Add(*CurrentArrayNodePtr);
 
 				// Recursively gather up all struct fields that differ and assign to the array item value node's sub property list.
 				if (BuildCustomPropertyListForPostConstruction((*CurrentArrayNodePtr)->SubPropertyList, InnerStructProperty->Struct, ArrayPropertyValue, DefaultArrayPropertyValue))
@@ -488,7 +492,8 @@ bool UBlueprintGeneratedClass::BuildCustomArrayPropertyListForPostConstruction(U
 			else if (UArrayProperty* InnerArrayProperty = Cast<UArrayProperty>(ArrayProperty->Inner))
 			{
 				// Create a new node for the item value at this index.
-				*CurrentArrayNodePtr = new(CustomPropertyListForPostConstruction) FCustomPropertyListNode(ArrayProperty, ArrayValueIndex);
+				*CurrentArrayNodePtr = new FCustomPropertyListNode(ArrayProperty, ArrayValueIndex);
+				CustomPropertyListForPostConstruction.Add(*CurrentArrayNodePtr);
 
 				// Recursively gather up all array item indices that differ and assign to the array item value node's sub property list.
 				if (BuildCustomArrayPropertyListForPostConstruction(InnerArrayProperty, (*CurrentArrayNodePtr)->SubPropertyList, ArrayPropertyValue, DefaultArrayPropertyValue))
@@ -508,7 +513,8 @@ bool UBlueprintGeneratedClass::BuildCustomArrayPropertyListForPostConstruction(U
 			else if (!ArrayProperty->Inner->Identical(ArrayPropertyValue, DefaultArrayPropertyValue))
 			{
 				// Create a new node, link it into the chain and add it into the array.
-				*CurrentArrayNodePtr = new(CustomPropertyListForPostConstruction) FCustomPropertyListNode(ArrayProperty, ArrayValueIndex);
+				*CurrentArrayNodePtr = new FCustomPropertyListNode(ArrayProperty, ArrayValueIndex);
+				CustomPropertyListForPostConstruction.Add(*CurrentArrayNodePtr);
 
 				// Advance to the next array item node ptr.
 				CurrentArrayNodePtr = &(*CurrentArrayNodePtr)->PropertyListNext;
@@ -1610,7 +1616,8 @@ void FBlueprintCookedComponentInstancingData::BuildCachedPropertyList(FCustomPro
 		}
 
 		// Create a new node to hold property info.
-		FCustomPropertyListNode* NewNode = new(CachedPropertyListForSerialization) FCustomPropertyListNode(Property, ChangedPropertyInfo.ArrayIndex);
+		FCustomPropertyListNode* NewNode = new FCustomPropertyListNode(Property, ChangedPropertyInfo.ArrayIndex);
+		CachedPropertyListForSerialization.Add(NewNode);
 
 		// Link the new node into the current property list.
 		if (CurrentNode)
@@ -1644,7 +1651,8 @@ void FBlueprintCookedComponentInstancingData::BuildCachedArrayPropertyList(const
 		const FBlueprintComponentChangedPropertyInfo& ChangedArrayPropertyInfo = ChangedPropertyList[(*CurrentSourceIdx)++];
 		UProperty* InnerProperty = ChangedArrayPropertyInfo.PropertyName != NAME_None ? ArrayProperty->Inner : nullptr;
 
-		*ArraySubPropertyNode = new(CachedPropertyListForSerialization) FCustomPropertyListNode(InnerProperty, ChangedArrayPropertyInfo.ArrayIndex);
+		*ArraySubPropertyNode = new FCustomPropertyListNode(InnerProperty, ChangedArrayPropertyInfo.ArrayIndex);
+		CachedPropertyListForSerialization.Add(*ArraySubPropertyNode);
 
 		// If this is a UStruct property, recursively build a sub-property list.
 		if (const UStructProperty* InnerStructProperty = Cast<UStructProperty>(InnerProperty))
diff --git a/Engine/Source/Runtime/Engine/Private/Components/ActorComponent.cpp b/Engine/Source/Runtime/Engine/Private/Components/ActorComponent.cpp
--- a/Engine/Source/Runtime/Engine/Private/Components/ActorComponent.cpp
+++ b/Engine/Source/Runtime/Engine/Private/Components/ActorComponent.cpp
@@ -83,7 +83,7 @@ FGlobalComponentReregisterContext::FGlobalComponentReregisterContext()
 	// Detach all actor components.
 	for(UActorComponent* Component : TObjectRange<UActorComponent>())
 	{
-		new(ComponentContexts) FComponentReregisterContext(Component);
+		ComponentContexts.Add(new FComponentReregisterContext(Component));
 	}
 }
 
@@ -109,7 +109,7 @@ FGlobalComponentReregisterContext::FGlobalComponentReregisterContext(const TArra
 		}
 		if( bShouldReregister )
 		{
-			new(ComponentContexts) FComponentReregisterContext(Component);		
+			ComponentContexts.Add(new FComponentReregisterContext(Component));
 		}
 	}
 }
@@ -130,7 +130,7 @@ FGlobalComponentRecreateRenderStateContext::FGlobalComponentRecreateRenderStateC
 	// recreate render state for all components.
 	for (UActorComponent* Component : TObjectRange<UActorComponent>())
 	{
-		new(ComponentContexts) FComponentRecreateRenderStateContext(Component);
+		ComponentContexts.Add(new FComponentRecreateRenderStateContext(Component));
 	}
 }
 
diff --git a/Engine/Source/Runtime/Engine/Private/Components/ModelComponent.cpp b/Engine/Source/Runtime/Engine/Private/Components/ModelComponent.cpp
--- a/Engine/Source/Runtime/Engine/Private/Components/ModelComponent.cpp
+++ b/Engine/Source/Runtime/Engine/Private/Components/ModelComponent.cpp
@@ -201,7 +201,8 @@ void UModelComponent::CommitSurfaces()
 		// If no matching element was found, create a new element.
 		if(!NewElement)
 		{
-			NewElement = new(Elements) FModelElement(this,Surf.Material);
+			NewElement = new FModelElement(this,Surf.Material);
+			Elements.Add(NewElement);
 			NewElement->MapBuildDataId = OldElement->MapBuildDataId;
 		}
 
@@ -524,7 +525,9 @@ bool UModelComponent::GenerateElements(bool bBuildRenderData)
 			if(!Element)
 			{
 				// If there's no matching element, create a new element.
-				Element = MaterialToElementMap.Add(Surf.Material,new(Elements) FModelElement(this,Surf.Material));
+				FModelElement* NewElement = new FModelElement(this,Surf.Material);
+				Elements.Add(NewElement);
+				Element = MaterialToElementMap.Add(Surf.Material, NewElement);
 			}
 
 			// Add the node to the element.
@@ -563,7 +566,9 @@ bool UModelComponent::GenerateElements(bool bBuildRenderData)
 			FModelElement** Element = NodeGroupToElementMap.Find(NodeGroupKey);
 			if (Element == nullptr)
 			{
-				Element = &NodeGroupToElementMap.Add(NodeGroupKey, new(Elements)FModelElement(this, Surf.Material));
+				FModelElement* NewElement = new FModelElement(this, Surf.Material);
+				Elements.Add(NewElement);
+				Element = &NodeGroupToElementMap.Add(NodeGroupKey, NewElement);
 			}
 			check(Element);
 
@@ -587,7 +592,8 @@ void UModelComponent::CopyElementsFrom(UModelComponent* SrcComponent)
 	for (int32 ElementIndex = 0; ElementIndex < SrcComponent->Elements.Num(); ++ElementIndex)
 	{
 		FModelElement& SrcElement = SrcComponent->Elements[ElementIndex];
-		FModelElement& DestElement = *new(Elements) FModelElement(SrcElement);
+		FModelElement& DestElement = *new FModelElement(SrcElement);
+		Elements.Add(&DestElement);
 		DestElement.Component = this;
 	}
 
diff --git a/Engine/Source/Runtime/Engine/Private/Materials/MaterialShader.cpp b/Engine/Source/Runtime/Engine/Private/Materials/MaterialShader.cpp
--- a/Engine/Source/Runtime/Engine/Private/Materials/MaterialShader.cpp
+++ b/Engine/Source/Runtime/Engine/Private/Materials/MaterialShader.cpp
@@ -1473,7 +1473,8 @@ void FMaterialShaderMap::Compile(
 					{
 						// Create a new mesh material shader map.
 						MeshShaderMapIndex = MeshShaderMaps.Num();
-						MeshShaderMap = new(MeshShaderMaps) FMeshMaterialShaderMap(InPlatform, VertexFactoryType);
+						MeshShaderMap = new FMeshMaterialShaderMap(InPlatform, VertexFactoryType);
+						MeshShaderMaps.Add(MeshShaderMap);
 					}
   
 					// Enqueue compilation all mesh material shaders for this material and vertex factory type combo.
@@ -2365,7 +2366,7 @@ void FMaterialShaderMap::Serialize(FArchive& Ar, bool bInlineShaderResources, bo
 
 			if (VertexFactoryType->IsUsedWithMaterials())
 			{
-				new(MeshShaderMaps) FMeshMaterialShaderMap(GetShaderPlatform(), VertexFactoryType);
+				MeshShaderMaps.Add(new FMeshMaterialShaderMap(GetShaderPlatform(), VertexFactoryType));
 			}
 		}
 
diff --git a/Engine/Source/Runtime/Engine/Private/ModelLight.cpp b/Engine/Source/Runtime/Engine/Private/ModelLight.cpp
--- a/Engine/Source/Runtime/Engine/Private/ModelLight.cpp
+++ b/Engine/Source/Runtime/Engine/Private/ModelLight.cpp
@@ -198,7 +198,8 @@ FModelElement* UModelComponent::CreateNewTempElement(UModelComponent* Component)
 	}
 
 	// make it in the temp array
-	FModelElement* Element = new(*TempElements) FModelElement(Component, NULL);
+	FModelElement* Element = new FModelElement(Component, NULL);
+	TempElements->Add(Element);
 	return Element;
 }
 
@@ -240,7 +241,7 @@ void UModelComponent::ApplyTempElements(bool bLightingWasSuccessful)
 		{
 			UModelComponent* Component = UpdatedComponents[ComponentIndex];
 
-			new(ComponentContexts) FComponentReregisterContext(Component);
+			ComponentContexts.Add(new FComponentReregisterContext(Component));
 		}
 
 		// Release all index buffers since they will be modified by BuildRenderData()
diff --git a/Engine/Source/Runtime/Engine/Private/Particles/ParticleSystemRender.cpp b/Engine/Source/Runtime/Engine/Private/Particles/ParticleSystemRender.cpp
--- a/Engine/Source/Runtime/Engine/Private/Particles/ParticleSystemRender.cpp
+++ b/Engine/Source/Runtime/Engine/Private/Particles/ParticleSystemRender.cpp
@@ -6910,7 +6910,8 @@ FMeshBatch* FParticleSystemSceneProxy::GetPooledMeshBatch()
 	}
 	else
 	{
-		Batch = new(MeshBatchPool) FMeshBatch();
+		Batch = new FMeshBatch();
+		MeshBatchPool.Add(Batch);
 	}
 	FirstFreeMeshBatch++;
 	return Batch;
diff --git a/Engine/Source/Runtime/Engine/Private/ShaderCompiler/ShaderCompiler.cpp b/Engine/Source/Runtime/Engine/Private/ShaderCompiler/ShaderCompiler.cpp
--- a/Engine/Source/Runtime/Engine/Private/ShaderCompiler/ShaderCompiler.cpp
+++ b/Engine/Source/Runtime/Engine/Private/ShaderCompiler/ShaderCompiler.cpp
@@ -2152,7 +2152,7 @@ void FShaderCompilingManager::PropagateMaterialChangesToPrimitives(const TMap<FM
 
 				if (bPrimitiveIsDependentOnMaterial)
 				{
-					new(ComponentContexts) FComponentRecreateRenderStateContext(PrimitiveComponent);
+					ComponentContexts.Add(new FComponentRecreateRenderStateContext(PrimitiveComponent));
 				}
 			}
 		}
diff --git a/Engine/Source/Runtime/Engine/Private/SkeletalMeshMerge.cpp b/Engine/Source/Runtime/Engine/Private/SkeletalMeshMerge.cpp
--- a/Engine/Source/Runtime/Engine/Private/SkeletalMeshMerge.cpp
+++ b/Engine/Source/Runtime/Engine/Private/SkeletalMeshMerge.cpp
@@ -433,7 +433,8 @@ void FSkeletalMeshMerge::GenerateLODModel( int32 LODIdx )
 	FSkeletalMeshRenderData* MergeResource = MergeMesh->GetResourceForRendering();
 	check(MergeResource);
 
-	FSkeletalMeshLODRenderData& MergeLODData = *new(MergeResource->LODRenderData) FSkeletalMeshLODRenderData;
+	FSkeletalMeshLODRenderData& MergeLODData = *new FSkeletalMeshLODRenderData;
+	MergeResource->LODRenderData.Add(&MergeLODData);
 	// add the new LOD info entry
 	FSkeletalMeshLODInfo& MergeLODInfo = MergeMesh->AddLODInfo();
 	MergeLODInfo.ScreenSize = MergeLODInfo.LODHysteresis = MAX_FLT;
diff --git a/Engine/Source/Runtime/Engine/Private/SkeletalMeshRenderData.cpp b/Engine/Source/Runtime/Engine/Private/SkeletalMeshRenderData.cpp
--- a/Engine/Source/Runtime/Engine/Private/SkeletalMeshRenderData.cpp
+++ b/Engine/Source/Runtime/Engine/Private/SkeletalMeshRenderData.cpp
@@ -98,7 +98,8 @@ void FSkeletalMeshRenderData::Cache(USkeletalMesh* Owner)
 			for (int32 LODIndex = 0; LODIndex < SkelMeshModel->LODModels.Num(); LODIndex++)
 			{
 				FSkeletalMeshLODModel& LODModel = SkelMeshModel->LODModels[LODIndex];
-				FSkeletalMeshLODRenderData* LODData = new(LODRenderData) FSkeletalMeshLODRenderData();
+				FSkeletalMeshLODRenderData* LODData = new FSkeletalMeshLODRenderData();
+				LODRenderData.Add(LODData);
 
 				LODData->BuildFromLODModel(&LODModel, VertexBufferBuildFlags);
 			}
diff --git a/Engine/Source/Runtime/Engine/Private/StaticMesh.cpp b/Engine/Source/Runtime/Engine/Private/StaticMesh.cpp
--- a/Engine/Source/Runtime/Engine/Private/StaticMesh.cpp
+++ b/Engine/Source/Runtime/Engine/Private/StaticMesh.cpp
@@ -772,7 +772,7 @@ void FStaticMeshRenderData::Serialize(FArchive& Ar, UStaticMesh* Owner, bool bCo
 		LODVertexFactories.Empty(LODResources.Num());
 		for (int i = 0; i < LODResources.Num(); i++)
 		{
-			new(LODVertexFactories) FStaticMeshVertexFactories(ERHIFeatureLevel::Num);
+			LODVertexFactories.Add(new FStaticMeshVertexFactories(ERHIFeatureLevel::Num));
 		}
 	}
 
@@ -893,8 +893,8 @@ void FStaticMeshRenderData::AllocateLODResources(int32 NumLODs)
 	check(LODResources.Num() == 0);
 	while (LODResources.Num() < NumLODs)
 	{
-		new(LODResources) FStaticMeshLODResources;
-		new(LODVertexFactories) FStaticMeshVertexFactories(ERHIFeatureLevel::Num);
+		LODResources.Add(new FStaticMeshLODResources);
+		LODVertexFactories.Add(new FStaticMeshVertexFactories(ERHIFeatureLevel::Num));
 	}
 }
 
diff --git a/Engine/Source/Runtime/Engine/Private/StaticMeshRender.cpp b/Engine/Source/Runtime/Engine/Private/StaticMeshRender.cpp
--- a/Engine/Source/Runtime/Engine/Private/StaticMeshRender.cpp
+++ b/Engine/Source/Runtime/Engine/Private/StaticMeshRender.cpp
@@ -185,7 +185,8 @@ FStaticMeshSceneProxy::FStaticMeshSceneProxy(UStaticMeshComponent* InComponent,
 	const bool bLODsShareStaticLighting = RenderData->bLODsShareStaticLighting || bForceLODsShareStaticLighting;
 	for(int32 LODIndex = 0;LODIndex < RenderData->LODResources.Num();LODIndex++)
 	{
-		FLODInfo* NewLODInfo = new(LODs) FLODInfo(InComponent, RenderData->LODVertexFactories,LODIndex,bLODsShareStaticLighting);
+		FLODInfo* NewLODInfo = new FLODInfo(InComponent, RenderData->LODVertexFactories,LODIndex,bLODsShareStaticLighting);
+		LODs.Add(NewLODInfo);
 
 		// Under certain error conditions an LOD's material will be set to 
 		// DefaultMaterial. Ensure our material view relevance is set properly.
diff --git a/Engine/Source/Runtime/Engine/Private/Streaming/StreamingManagerTexture.cpp b/Engine/Source/Runtime/Engine/Private/Streaming/StreamingManagerTexture.cpp
--- a/Engine/Source/Runtime/Engine/Private/Streaming/StreamingManagerTexture.cpp
+++ b/Engine/Source/Runtime/Engine/Private/Streaming/StreamingManagerTexture.cpp
@@ -520,7 +520,7 @@ void FStreamingManagerTexture::AddLevel( ULevel* Level )
 
 	// If the level was not already there, add a new entry.
 	TextureInstanceAsyncWork->EnsureCompletion();
-	new (LevelTextureManagers) FLevelTextureManager(Level, TextureInstanceAsyncWork->GetTask());
+	LevelTextureManagers.Add(new FLevelTextureManager(Level, TextureInstanceAsyncWork->GetTask()));
 }
 
 /** Removes a ULevel from the streaming manager. */
diff --git a/Engine/Source/Runtime/Engine/Private/Texture2D.cpp b/Engine/Source/Runtime/Engine/Private/Texture2D.cpp
--- a/Engine/Source/Runtime/Engine/Private/Texture2D.cpp
+++ b/Engine/Source/Runtime/Engine/Private/Texture2D.cpp
@@ -1005,7 +1005,8 @@ UTexture2D* UTexture2D::CreateTransient(int32 InSizeX, int32 InSizeY, EPixelForm
 		// Allocate first mipmap.
 		int32 NumBlocksX = InSizeX / GPixelFormats[InFormat].BlockSizeX;
 		int32 NumBlocksY = InSizeY / GPixelFormats[InFormat].BlockSizeY;
-		FTexture2DMipMap* Mip = new(NewTexture->PlatformData->Mips) FTexture2DMipMap();
+		FTexture2DMipMap* Mip = new FTexture2DMipMap();
+		NewTexture->PlatformData->Mips.Add(Mip);
 		Mip->SizeX = InSizeX;
 		Mip->SizeY = InSizeY;
 		Mip->BulkData.Lock(LOCK_READ_WRITE);
diff --git a/Engine/Source/Runtime/Engine/Private/TextureDerivedData.cpp b/Engine/Source/Runtime/Engine/Private/TextureDerivedData.cpp
--- a/Engine/Source/Runtime/Engine/Private/TextureDerivedData.cpp
+++ b/Engine/Source/Runtime/Engine/Private/TextureDerivedData.cpp
@@ -1066,7 +1066,7 @@ static void SerializePlatformData(
 		PlatformData->Mips.Empty(NumMips);
 		for (int32 MipIndex = 0; MipIndex < NumMips; ++MipIndex)
 		{
-			new(PlatformData->Mips) FTexture2DMipMap();
+			PlatformData->Mips.Add(new FTexture2DMipMap());
 		}
 	}
 	for (int32 MipIndex = 0; MipIndex < NumMips; ++MipIndex)
diff --git a/Engine/Source/Runtime/Engine/Private/TextureDerivedDataTask.cpp b/Engine/Source/Runtime/Engine/Private/TextureDerivedDataTask.cpp
--- a/Engine/Source/Runtime/Engine/Private/TextureDerivedDataTask.cpp
+++ b/Engine/Source/Runtime/Engine/Private/TextureDerivedDataTask.cpp
@@ -154,7 +154,8 @@ void FTextureCacheDerivedDataWorker::BuildTexture()
 			for (int32 MipIndex = 0; MipIndex < MipCount; ++MipIndex)
 			{
 				const FCompressedImage2D& CompressedImage = CompressedMips[MipIndex];
-				FTexture2DMipMap* NewMip = new(DerivedData->Mips) FTexture2DMipMap();
+				FTexture2DMipMap* NewMip = new FTexture2DMipMap();
+				DerivedData->Mips.Add(NewMip);
 				NewMip->SizeX = CompressedImage.SizeX;
 				NewMip->SizeY = CompressedImage.SizeY;
 				NewMip->SizeZ = CompressedImage.SizeZ;
diff --git a/Engine/Source/Runtime/Engine/Private/UnrealEngine.cpp b/Engine/Source/Runtime/Engine/Private/UnrealEngine.cpp
--- a/Engine/Source/Runtime/Engine/Private/UnrealEngine.cpp
+++ b/Engine/Source/Runtime/Engine/Private/UnrealEngine.cpp
@@ -12584,7 +12584,8 @@ void UEngine::UpdateTransitionType(UWorld *CurrentWorld)
 
 FWorldContext& UEngine::CreateNewWorldContext(EWorldType::Type WorldType)
 {
-	FWorldContext *NewWorldContext = (new (WorldList) FWorldContext);
+	FWorldContext* NewWorldContext = new FWorldContext;
+	WorldList.Add(NewWorldContext);
 	NewWorldContext->WorldType = WorldType;
 	NewWorldContext->ContextHandle = FName(*FString::Printf(TEXT("Context_%d"), NextWorldContextHandle++));
 
@@ -13626,7 +13627,8 @@ void UEngine::CopyPropertiesForUnrelatedObjects(UObject* OldObject, UObject* New
 
 		for (UObject* OldInstance : Components)
 		{
-			FInstancedObjectRecord* pRecord = new(SavedInstances) FInstancedObjectRecord();
+			FInstancedObjectRecord* pRecord = new FInstancedObjectRecord();
+			SavedInstances.Add(pRecord);
 			pRecord->OldInstance = OldInstance;
 			OldInstanceMap.Add(OldInstance->GetPathName(OldObject), SavedInstances.Num() - 1);
 			const uint32 AdditionalPortFlags = Params.bCopyDeprecatedProperties ? PPF_UseDeprecatedProperties : PPF_None;
diff --git a/Engine/Source/Runtime/Engine/Public/ComponentReregisterContext.h b/Engine/Source/Runtime/Engine/Public/ComponentReregisterContext.h
--- a/Engine/Source/Runtime/Engine/Public/ComponentReregisterContext.h
+++ b/Engine/Source/Runtime/Engine/Public/ComponentReregisterContext.h
@@ -208,7 +208,7 @@ public:
 		// Reregister all components of the templated type.
 		for(TObjectIterator<ComponentType> ComponentIt;ComponentIt;++ComponentIt)
 		{
-			new(ComponentContexts) FComponentReregisterContext(*ComponentIt);
+			ComponentContexts.Add(new FComponentReregisterContext(*ComponentIt));
 		}
 	}
 
diff --git a/Engine/Source/Runtime/RenderCore/Private/RenderResource.cpp b/Engine/Source/Runtime/RenderCore/Private/RenderResource.cpp
--- a/Engine/Source/Runtime/RenderCore/Private/RenderResource.cpp
+++ b/Engine/Source/Runtime/RenderCore/Private/RenderResource.cpp
@@ -427,7 +427,8 @@ FGlobalDynamicVertexBuffer::FAllocation FGlobalDynamicVertexBuffer::Allocate(uin
 		// Create a new vertex buffer if needed.
 		if (VertexBuffer == NULL)
 		{
-			VertexBuffer = new(Pool->VertexBuffers) FDynamicVertexBuffer(SizeInBytes);
+			VertexBuffer = new FDynamicVertexBuffer(SizeInBytes);
+			Pool->VertexBuffers.Add(VertexBuffer);
 			VertexBuffer->InitResource();
 		}
 
@@ -629,7 +630,8 @@ FGlobalDynamicIndexBuffer::FAllocation FGlobalDynamicIndexBuffer::Allocate(uint3
 		// Create a new index buffer if needed.
 		if (IndexBuffer == NULL)
 		{
-			IndexBuffer = new(Pool->IndexBuffers) FDynamicIndexBuffer(SizeInBytes, Pool->BufferStride);
+			IndexBuffer = new FDynamicIndexBuffer(SizeInBytes, Pool->BufferStride);
+			Pool->IndexBuffers.Add(IndexBuffer);
 			IndexBuffer->InitResource();
 		}
 
diff --git a/Engine/Source/Runtime/Renderer/Private/DynamicPrimitiveDrawing.inl b/Engine/Source/Runtime/Renderer/Private/DynamicPrimitiveDrawing.inl
--- a/Engine/Source/Runtime/Renderer/Private/DynamicPrimitiveDrawing.inl
+++ b/Engine/Source/Runtime/Renderer/Private/DynamicPrimitiveDrawing.inl
@@ -337,7 +337,8 @@ inline int32 FViewElementPDI::DrawMesh(const FMeshBatch& Mesh)
 		// Translucent view mesh elements in the foreground dpg are not supported yet
 		TIndirectArray<FMeshBatch>& ViewMeshElementList = ( ( DPGIndex == SDPG_Foreground  ) ? ViewInfo->TopViewMeshElements : ViewInfo->ViewMeshElements );
 
-		FMeshBatch* NewMesh = new(ViewMeshElementList) FMeshBatch(Mesh);
+		FMeshBatch* NewMesh = new FMeshBatch(Mesh);
+		ViewMeshElementList.Add(NewMesh);
 		if( CurrentHitProxy != nullptr )
 		{
 			NewMesh->BatchHitProxyId = CurrentHitProxy->Id;
diff --git a/Engine/Source/Runtime/Renderer/Private/PrimitiveSceneInfo.cpp b/Engine/Source/Runtime/Renderer/Private/PrimitiveSceneInfo.cpp
--- a/Engine/Source/Runtime/Renderer/Private/PrimitiveSceneInfo.cpp
+++ b/Engine/Source/Runtime/Renderer/Private/PrimitiveSceneInfo.cpp
@@ -50,12 +50,13 @@ public:
 #endif
 			PrimitiveSceneInfo->Proxy->VerifyUsedMaterial(Mesh.MaterialRenderProxy);
 
-			FStaticMesh* StaticMesh = new(PrimitiveSceneInfo->StaticMeshes) FStaticMesh(
+			FStaticMesh* StaticMesh = new FStaticMesh(
 				PrimitiveSceneInfo,
 				Mesh,
 				ScreenSize,
 				CurrentHitProxy ? CurrentHitProxy->Id : FHitProxyId()
 				);
+			PrimitiveSceneInfo->StaticMeshes.Add(StaticMesh);
 		}
 	}
 
diff --git a/Engine/Source/Runtime/RuntimeAssetCache/Private/RuntimeAssetCacheBuilders.cpp b/Engine/Source/Runtime/RuntimeAssetCache/Private/RuntimeAssetCacheBuilders.cpp
--- a/Engine/Source/Runtime/RuntimeAssetCache/Private/RuntimeAssetCacheBuilders.cpp
+++ b/Engine/Source/Runtime/RuntimeAssetCache/Private/RuntimeAssetCacheBuilders.cpp
@@ -179,7 +179,7 @@ void UExampleTextureCacheBuilder::SerializeAsset(FArchive& Ar)
 			PlatformData->Mips.Empty(NumMips);
 			for (int32 MipIndex = 0; MipIndex < NumMips; ++MipIndex)
 			{
-				new(PlatformData->Mips) FTexture2DMipMap();
+				PlatformData->Mips.Add(new FTexture2DMipMap());
 			}
 		}
 
