Fixing cooked skeletal meshes crashing when opened in the editor

diff --git a/Engine/Source/Runtime/Core/Public/UObject/CoreObjectVersion.h b/Engine/Source/Runtime/Core/Public/UObject/CoreObjectVersion.h
--- a/Engine/Source/Runtime/Core/Public/UObject/CoreObjectVersion.h
+++ b/Engine/Source/Runtime/Core/Public/UObject/CoreObjectVersion.h
@@ -13,6 +13,7 @@ struct CORE_API FCoreObjectVersion
 		BeforeCustomVersionWasAdded = 0,
 		MaterialInputNativeSerialize,
 		EnumProperties,
+		SkeletalMaterialEditorDataStripping,
 
 		// -----<new versions can be added above this line>-------------------------------------------------
 		VersionPlusOne,
diff --git a/Engine/Source/Runtime/Engine/Private/SkeletalMesh.cpp b/Engine/Source/Runtime/Engine/Private/SkeletalMesh.cpp
--- a/Engine/Source/Runtime/Engine/Private/SkeletalMesh.cpp
+++ b/Engine/Source/Runtime/Engine/Private/SkeletalMesh.cpp
@@ -25,6 +25,7 @@
 #include "ComponentReregisterContext.h"
 #include "UObject/EditorObjectVersion.h"
 #include "UObject/RenderingObjectVersion.h"
+#include "UObject/CoreObjectVersion.h"
 #include "EngineUtils.h"
 #include "EditorSupportDelegates.h"
 #include "GPUSkinVertexFactory.h"
@@ -1509,104 +1510,105 @@ void USkeletalMesh::PostLoad()
 	Super::PostLoad();
 
 #if WITH_EDITOR
-	// If LODInfo is missing - create array of correct size.
-	if( LODInfo.Num() != ImportedModel->LODModels.Num() )
+	if (!GetOutermost()->bIsCookedForEditor)
 	{
-		LODInfo.Empty(ImportedModel->LODModels.Num());
-		LODInfo.AddZeroed(ImportedModel->LODModels.Num());
-
-		for(int32 i=0; i<LODInfo.Num(); i++)
+		// If LODInfo is missing - create array of correct size.
+		if (LODInfo.Num() != ImportedModel->LODModels.Num())
 		{
-			LODInfo[i].LODHysteresis = 0.02f;
-		}
-	}
+			LODInfo.Empty(ImportedModel->LODModels.Num());
+			LODInfo.AddZeroed(ImportedModel->LODModels.Num());
 
-	int32 TotalLODNum = LODInfo.Num();
-	for (int32 LodIndex = 0; LodIndex<TotalLODNum; LodIndex++)
-	{
-		FSkeletalMeshLODInfo& ThisLODInfo = LODInfo[LodIndex];
-		FSkeletalMeshLODModel& ThisLODModel = ImportedModel->LODModels[LodIndex];
+			for (int32 i = 0; i < LODInfo.Num(); i++)
+			{
+				LODInfo[i].LODHysteresis = 0.02f;
+			}
+		}
 
-		if (ThisLODInfo.ReductionSettings.BonesToRemove_DEPRECATED.Num() > 0)
+		int32 TotalLODNum = LODInfo.Num();
+		for (int32 LodIndex = 0; LodIndex < TotalLODNum; LodIndex++)
 		{
-			for (auto& BoneToRemove : ThisLODInfo.ReductionSettings.BonesToRemove_DEPRECATED)
+			FSkeletalMeshLODInfo& ThisLODInfo = LODInfo[LodIndex];
+			FSkeletalMeshLODModel& ThisLODModel = ImportedModel->LODModels[LodIndex];
+
+			if (ThisLODInfo.ReductionSettings.BonesToRemove_DEPRECATED.Num() > 0)
 			{
-				AddBoneToReductionSetting(LodIndex, BoneToRemove.BoneName);
+				for (auto& BoneToRemove : ThisLODInfo.ReductionSettings.BonesToRemove_DEPRECATED)
+				{
+					AddBoneToReductionSetting(LodIndex, BoneToRemove.BoneName);
+				}
+
+				// since in previous system, we always removed from previous LOD, I'm adding this 
+				// here for previous LODs
+				for (int32 CurLodIndx = LodIndex + 1; CurLodIndx < TotalLODNum; ++CurLodIndx)
+				{
+					AddBoneToReductionSetting(CurLodIndx, ThisLODInfo.RemovedBones_DEPRECATED);
+				}
+
+				// we don't apply this change here, but this will be applied when you re-gen simplygon
+				ThisLODInfo.ReductionSettings.BonesToRemove_DEPRECATED.Empty();
 			}
 
-			// since in previous system, we always removed from previous LOD, I'm adding this 
-			// here for previous LODs
-			for (int32 CurLodIndx = LodIndex + 1; CurLodIndx < TotalLODNum; ++CurLodIndx)
+			if (ThisLODInfo.ReductionSettings.BakePose_DEPRECATED != nullptr)
 			{
-				AddBoneToReductionSetting(CurLodIndx, ThisLODInfo.RemovedBones_DEPRECATED);
+				ThisLODInfo.BakePose = ThisLODInfo.ReductionSettings.BakePose_DEPRECATED;
+				ThisLODInfo.ReductionSettings.BakePose_DEPRECATED = nullptr;
 			}
-
-			// we don't apply this change here, but this will be applied when you re-gen simplygon
-			ThisLODInfo.ReductionSettings.BonesToRemove_DEPRECATED.Empty();
-		}
-
-		if (ThisLODInfo.ReductionSettings.BakePose_DEPRECATED != nullptr)
-		{
-			ThisLODInfo.BakePose = ThisLODInfo.ReductionSettings.BakePose_DEPRECATED;
-			ThisLODInfo.ReductionSettings.BakePose_DEPRECATED = nullptr;
 		}
-	}
 
-	// load LODinfo if using shared asset, it can override existing bone remove settings
-#if WITH_EDITORONLY_DATA
-	if (LODSettings != nullptr)
-	{
-		//before we copy
-		if (GetLinkerCustomVersion(FFortniteMainBranchObjectVersion::GUID) < FFortniteMainBranchObjectVersion::AddBakePoseOverrideForSkeletalMeshReductionSetting)
+		// load LODinfo if using shared asset, it can override existing bone remove settings
+		if (LODSettings != nullptr)
 		{
-			// if LODsetting doesn't have BakePose, but this does, we'll have to copy that to BakePoseOverride
-			const int32 NumSettings = FMath::Min(LODSettings->GetNumberOfSettings(), GetLODNum());
-			for (int32 Index = 0; Index < NumSettings; ++Index)
+			//before we copy
+			if (GetLinkerCustomVersion(FFortniteMainBranchObjectVersion::GUID) < FFortniteMainBranchObjectVersion::AddBakePoseOverrideForSkeletalMeshReductionSetting)
 			{
-				const FSkeletalMeshLODGroupSettings& GroupSetting = LODSettings->GetSettingsForLODLevel(Index);
-				// if lod setting doesn't have bake pose, but this lod does, that means this bakepose has to move to BakePoseOverride
-				// since we want to match what GroupSetting has
-				if (GroupSetting.BakePose == nullptr && LODInfo[Index].BakePose)
+				// if LODsetting doesn't have BakePose, but this does, we'll have to copy that to BakePoseOverride
+				const int32 NumSettings = FMath::Min(LODSettings->GetNumberOfSettings(), GetLODNum());
+				for (int32 Index = 0; Index < NumSettings; ++Index)
 				{
-					// in this case,
-					LODInfo[Index].BakePoseOverride = LODInfo[Index].BakePose;
-					LODInfo[Index].BakePose = nullptr;
+					const FSkeletalMeshLODGroupSettings& GroupSetting = LODSettings->GetSettingsForLODLevel(Index);
+					// if lod setting doesn't have bake pose, but this lod does, that means this bakepose has to move to BakePoseOverride
+					// since we want to match what GroupSetting has
+					if (GroupSetting.BakePose == nullptr && LODInfo[Index].BakePose)
+					{
+						// in this case,
+						LODInfo[Index].BakePoseOverride = LODInfo[Index].BakePose;
+						LODInfo[Index].BakePose = nullptr;
+					}
 				}
 			}
+			LODSettings->SetLODSettingsToMesh(this);
 		}
-		LODSettings->SetLODSettingsToMesh(this);
-	}
-#endif // WITH_EDITORONLY_DATA
 
-	if (GetLinkerUE4Version() < VER_UE4_SORT_ACTIVE_BONE_INDICES)
-	{
-		for (int32 LodIndex = 0; LodIndex < LODInfo.Num(); LodIndex++)
+		if (GetLinkerUE4Version() < VER_UE4_SORT_ACTIVE_BONE_INDICES)
 		{
-			FSkeletalMeshLODModel & ThisLODModel = ImportedModel->LODModels[LodIndex];
-			ThisLODModel.ActiveBoneIndices.Sort();
+			for (int32 LodIndex = 0; LodIndex < LODInfo.Num(); LodIndex++)
+			{
+				FSkeletalMeshLODModel & ThisLODModel = ImportedModel->LODModels[LodIndex];
+				ThisLODModel.ActiveBoneIndices.Sort();
+			}
 		}
-	}
 
-	// make sure older versions contain active bone indices with parents present
-	// even if they're not skinned, missing matrix calculation will mess up skinned children
-	if (GetLinkerCustomVersion(FFortniteMainBranchObjectVersion::GUID) < FFortniteMainBranchObjectVersion::EnsureActiveBoneIndicesToContainParents)
-	{
-		for (int32 LodIndex = 0; LodIndex < LODInfo.Num(); LodIndex++)
+		// make sure older versions contain active bone indices with parents present
+		// even if they're not skinned, missing matrix calculation will mess up skinned children
+		if (GetLinkerCustomVersion(FFortniteMainBranchObjectVersion::GUID) < FFortniteMainBranchObjectVersion::EnsureActiveBoneIndicesToContainParents)
 		{
-			FSkeletalMeshLODModel& ThisLODModel = ImportedModel->LODModels[LodIndex];
-			RefSkeleton.EnsureParentsExistAndSort(ThisLODModel.ActiveBoneIndices);
+			for (int32 LodIndex = 0; LodIndex < LODInfo.Num(); LodIndex++)
+			{
+				FSkeletalMeshLODModel& ThisLODModel = ImportedModel->LODModels[LodIndex];
+				RefSkeleton.EnsureParentsExistAndSort(ThisLODModel.ActiveBoneIndices);
+			}
 		}
-	}
 
 #if WITH_APEX_CLOTHING
-	UpgradeOldClothingAssets();
+		UpgradeOldClothingAssets();
 #endif // WITH_APEX_CLOTHING
 
-	RemoveLegacyClothingSections();
+		RemoveLegacyClothingSections();
 
-	if (GetResourceForRendering() == nullptr)
-	{
-		CacheDerivedData();
+		if (GetResourceForRendering() == nullptr)
+		{
+			CacheDerivedData();
+		}
 	}
 #endif // WITH_EDITOR
 
@@ -1635,7 +1637,7 @@ void USkeletalMesh::PostLoad()
 	CalculateInvRefMatrices();
 
 #if WITH_EDITORONLY_DATA
-	if (RetargetBasePose.Num() == 0)
+	if (RetargetBasePose.Num() == 0 && !GetOutermost()->bIsCookedForEditor)
 	{
 		RetargetBasePose = RefSkeleton.GetRefBonePose();
 	}
@@ -2226,6 +2228,7 @@ FArchive& operator<<(FArchive& Ar, FMeshUVChannelInfo& ChannelData)
 FArchive& operator<<(FArchive& Ar, FSkeletalMaterial& Elem)
 {
 	Ar.UsingCustomVersion(FEditorObjectVersion::GUID);
+	Ar.UsingCustomVersion(FCoreObjectVersion::GUID);
 
 	Ar << Elem.MaterialInterface;
 
@@ -2233,12 +2236,20 @@ FArchive& operator<<(FArchive& Ar, FSkeletalMaterial& Elem)
 	if (Ar.CustomVer(FEditorObjectVersion::GUID) >= FEditorObjectVersion::RefactorMeshEditorMaterials)
 	{
 		Ar << Elem.MaterialSlotName;
-#if WITH_EDITORONLY_DATA
-		if (!Ar.IsCooking() || Ar.CookingTarget()->HasEditorOnlyData())
+
+		bool bSerializeImportedMaterialSlotName = !Ar.IsCooking() || Ar.CookingTarget()->HasEditorOnlyData();
+		if (Ar.CustomVer(FCoreObjectVersion::GUID) >= FCoreObjectVersion::SkeletalMaterialEditorDataStripping)
+		{
+			Ar << bSerializeImportedMaterialSlotName;
+		}
+		else if (!FPlatformProperties::HasEditorOnlyData())
+		{
+			bSerializeImportedMaterialSlotName = false;
+		}
+		if (bSerializeImportedMaterialSlotName)
 		{
 			Ar << Elem.ImportedMaterialSlotName;
 		}
-#endif //#if WITH_EDITORONLY_DATA
 	}
 #if WITH_EDITORONLY_DATA
 	else
