Fixed 'Missing DDC Key' log spam when interacting with cooked texture in the editor

diff --git a/Engine/Source/Runtime/Engine/Private/Texture2D.cpp b/Engine/Source/Runtime/Engine/Private/Texture2D.cpp
--- a/Engine/Source/Runtime/Engine/Private/Texture2D.cpp
+++ b/Engine/Source/Runtime/Engine/Private/Texture2D.cpp
@@ -1543,7 +1543,7 @@ bool UTexture2D::StreamIn(int32 NewMipCount, bool bHighPrio)
 	if (bIsStreamable && !PendingUpdate && Texture2DResource && Texture2DResource->bReadyForStreaming && NewMipCount > GetNumResidentMips())
 	{
 #if WITH_EDITORONLY_DATA
-		if (FPlatformProperties::HasEditorOnlyData())
+		if (FPlatformProperties::HasEditorOnlyData() && !GetOutermost()->bIsCookedForEditor)
 		{
 			if (GRHISupportsAsyncTextureCreation)
 			{
