Fix for editor-only field.

diff --git a/Engine/Source/Runtime/Engine/Private/SkeletalMesh.cpp b/Engine/Source/Runtime/Engine/Private/SkeletalMesh.cpp
--- a/Engine/Source/Runtime/Engine/Private/SkeletalMesh.cpp
+++ b/Engine/Source/Runtime/Engine/Private/SkeletalMesh.cpp
@@ -2248,7 +2248,12 @@ FArchive& operator<<(FArchive& Ar, FSkeletalMaterial& Elem)
 		}
 		if (bSerializeImportedMaterialSlotName)
 		{
+#if WITH_EDITORONLY_DATA
 			Ar << Elem.ImportedMaterialSlotName;
+#else
+			FName UnusedImportedMaterialSlotName;
+			Ar << UnusedImportedMaterialSlotName;
+#endif
 		}
 	}
 #if WITH_EDITORONLY_DATA
