Fix for warning caused by failed move while slicing during return-by-value.

diff --git a/Engine/Source/Runtime/Engine/Private/TimerManager.cpp b/Engine/Source/Runtime/Engine/Private/TimerManager.cpp
--- a/Engine/Source/Runtime/Engine/Private/TimerManager.cpp
+++ b/Engine/Source/Runtime/Engine/Private/TimerManager.cpp
@@ -42,7 +42,7 @@ namespace
 	{
 		FStringOutputDevice Output;
 		DescribeFTimerDataSafely(Output, Data);
-		return Output;
+		return MoveTemp(Output);
 	}
 }
 
