Consider target platform device profiles before cooking optimized Blueprint component data.

#rnx

diff --git a/Engine/Source/Runtime/Engine/Private/Blueprint.cpp b/Engine/Source/Runtime/Engine/Private/Blueprint.cpp
--- a/Engine/Source/Runtime/Engine/Private/Blueprint.cpp
+++ b/Engine/Source/Runtime/Engine/Private/Blueprint.cpp
@@ -1282,7 +1282,8 @@ void UBlueprint::BeginCacheForCookedPlatformData(const ITargetPlatform *TargetPl
 		{
 			if (InComponentTemplate)
 			{
-				return (!TargetPlatform->IsClientOnly() || InComponentTemplate->NeedsLoadForClient())
+				return InComponentTemplate->NeedsLoadForTargetPlatform(TargetPlatform)
+					&& (!TargetPlatform->IsClientOnly() || InComponentTemplate->NeedsLoadForClient())
 					&& (!TargetPlatform->IsServerOnly() || InComponentTemplate->NeedsLoadForServer());
 			}
 
