Merging //UE4/Dev-Main to Dev-Sequencer (//UE4/Dev-Sequencer)

diff --git a/Engine/Source/Runtime/Engine/Private/Components/CharacterMovementComponent.cpp b/Engine/Source/Runtime/Engine/Private/Components/CharacterMovementComponent.cpp
--- a/Engine/Source/Runtime/Engine/Private/Components/CharacterMovementComponent.cpp
+++ b/Engine/Source/Runtime/Engine/Private/Components/CharacterMovementComponent.cpp
@@ -7672,7 +7672,11 @@ bool UCharacterMovementComponent::ForcePositionUpdate(float DeltaTime)
 	ServerData->ServerAccumulatedClientTimeStamp += DeltaTime;
 
 #if !(UE_BUILD_SHIPPING)
-	UE_LOG(LogNetPlayerMovement, Warning, TEXT("ForcePositionUpdate %s (DeltaTime %.2f)"), *CharacterOwner->GetName(), DeltaTime);
+	const bool bServerMoveHasOccurred = (ServerData->ServerTimeStampLastServerMove != 0.f);
+	if (bServerMoveHasOccurred)
+	{
+		UE_LOG(LogNetPlayerMovement, Warning, TEXT("ForcePositionUpdate %s (DeltaTime %.2f)"), *CharacterOwner->GetName(), DeltaTime);
+	}
 #endif
 
 	// Force movement update.
diff --git a/Engine/Source/Runtime/Engine/Private/VorbisAudioInfo.cpp b/Engine/Source/Runtime/Engine/Private/VorbisAudioInfo.cpp
--- a/Engine/Source/Runtime/Engine/Private/VorbisAudioInfo.cpp
+++ b/Engine/Source/Runtime/Engine/Private/VorbisAudioInfo.cpp
@@ -483,6 +483,11 @@ void FVorbisAudioInfo::EnableHalfRate( bool HalfRate )
 
 bool FVorbisAudioInfo::StreamCompressedInfo(USoundWave* Wave, struct FSoundQualityInfo* QualityInfo)
 {
+	if (!bDllLoaded)
+	{
+		return false;
+	}
+
 	SCOPE_CYCLE_COUNTER( STAT_VorbisPrepareDecompressionTime );
 
 	FScopeLock ScopeLock(&VorbisCriticalSection);
@@ -511,6 +516,11 @@ bool FVorbisAudioInfo::StreamCompressedInfo(USoundWave* Wave, struct FSoundQuali
 
 bool FVorbisAudioInfo::StreamCompressedData(uint8* InDestination, bool bLooping, uint32 BufferSize)
 {
+	if (!bDllLoaded)
+	{
+		return false;
+	}
+
 	check( VFWrapper != NULL );
 
 #if PLATFORM_ANDROID
diff --git a/Engine/Source/Runtime/PreLoadScreen/Private/BuildPatchServicesPreLoadManager.cpp b/Engine/Source/Runtime/PreLoadScreen/Private/BuildPatchServicesPreLoadManager.cpp
--- a/Engine/Source/Runtime/PreLoadScreen/Private/BuildPatchServicesPreLoadManager.cpp
+++ b/Engine/Source/Runtime/PreLoadScreen/Private/BuildPatchServicesPreLoadManager.cpp
@@ -64,6 +64,14 @@ void FBuildPatchServicesPreLoadManagerBase::ResumeBuildPatchInstall()
     }
 }
 
+void FBuildPatchServicesPreLoadManagerBase::CancelBuildPatchInstall()
+{
+	if (ContentBuildInstaller.IsValid())
+	{
+		ContentBuildInstaller->CancelInstall();
+	}
+}
+
 const FText& FBuildPatchServicesPreLoadManagerBase::GetStatusText() const
 {
 	// Static const fixed FText values so that they are not constantly constructed
diff --git a/Engine/Source/Runtime/PreLoadScreen/Public/BuildPatchServicesPreLoadManager.h b/Engine/Source/Runtime/PreLoadScreen/Public/BuildPatchServicesPreLoadManager.h
--- a/Engine/Source/Runtime/PreLoadScreen/Public/BuildPatchServicesPreLoadManager.h
+++ b/Engine/Source/Runtime/PreLoadScreen/Public/BuildPatchServicesPreLoadManager.h
@@ -31,6 +31,7 @@ public:
 
     virtual void PauseBuildPatchInstall();
     virtual void ResumeBuildPatchInstall();
+	virtual void CancelBuildPatchInstall();
 
     float GetProgressPercent() const { return ContentBuildInstaller.IsValid() ? ContentBuildInstaller->GetUpdateProgress() : 0.0f; }
     EBuildPatchDownloadHealth GetDownloadHealth() const { return ContentBuildInstaller.IsValid() ? ContentBuildInstaller->GetDownloadHealth() : EBuildPatchDownloadHealth::NUM_Values; }
