Adding ability to disable baked FFT/envelope following analysis feature via cvar

diff --git a/Engine/Source/Runtime/Engine/Private/AudioDevice.cpp b/Engine/Source/Runtime/Engine/Private/AudioDevice.cpp
--- a/Engine/Source/Runtime/Engine/Private/AudioDevice.cpp
+++ b/Engine/Source/Runtime/Engine/Private/AudioDevice.cpp
@@ -96,6 +96,13 @@ FAutoConsoleVariableRef CVarWaitForSoundWaveToLoad(
 	TEXT("0: Attempt to play back, 1: Wait for load."),
 	ECVF_Default);
 
+static int32 BakedAnalysisEnabledCVar = 1;
+FAutoConsoleVariableRef CVarBakedAnalysisEnabledCVar(
+	TEXT("au.BakedAnalysisEnabled"),
+	BakedAnalysisEnabledCVar,
+	TEXT("Enables or disables queries to baked analysis from audio component.\n"),
+	ECVF_Default);
+
 static int32 NumPrecacheFramesCvar = 0;
 FAutoConsoleVariableRef CVarNumPrecacheFrames(
 	TEXT("au.NumPrecacheFrames"),
@@ -200,6 +207,7 @@ FAudioDevice::FAudioDevice()
 	, bDisableAudioCaching(false)
 	, bIsAudioDeviceHardwareInitialized(false)
 	, bIsStoppingVoicesEnabled(false)
+	, bIsBakedAnalysisEnabled(false)
 	, bAudioMixerModuleLoaded(false)
 	, bSpatializationIsExternalSend(false)
 	, bOcclusionIsExternalSend(false)
@@ -290,6 +298,8 @@ bool FAudioDevice::Init(int32 InMaxChannels)
 
 	bIsStoppingVoicesEnabled = !DisableStoppingVoicesCvar;
 
+	bIsBakedAnalysisEnabled = (BakedAnalysisEnabledCVar == 1);
+
 	const UAudioSettings* AudioSettings = GetDefault<UAudioSettings>();
 
 	bAllowCenterChannel3DPanning = AudioSettings->bAllowCenterChannel3DPanning;
@@ -3679,6 +3689,9 @@ void FAudioDevice::Update(bool bGameTicking)
 		UpdateAudioClock();
 	}
 
+	// update if baked analysis is enabled
+	bIsBakedAnalysisEnabled = (BakedAnalysisEnabledCVar == 1);
+
 	if (bGameTicking)
 	{
 		GlobalPitchScale.Update(GetDeviceDeltaTime());
diff --git a/Engine/Source/Runtime/Engine/Private/Components/AudioComponent.cpp b/Engine/Source/Runtime/Engine/Private/Components/AudioComponent.cpp
--- a/Engine/Source/Runtime/Engine/Private/Components/AudioComponent.cpp
+++ b/Engine/Source/Runtime/Engine/Private/Components/AudioComponent.cpp
@@ -416,16 +416,19 @@ void UAudioComponent::PlayInternal(const float StartTime, const float FadeInDura
 			NewActiveSound.bUpdatePlayPercentage = OnAudioPlaybackPercentNative.IsBound() || OnAudioPlaybackPercent.IsBound();
 			NewActiveSound.bUpdateSingleEnvelopeValue = OnAudioSingleEnvelopeValue.IsBound() || OnAudioSingleEnvelopeValueNative.IsBound();
 			NewActiveSound.bUpdateMultiEnvelopeValue = OnAudioMultiEnvelopeValue.IsBound() || OnAudioMultiEnvelopeValueNative.IsBound();
-			
-			// Setup audio component cooked analysis data playback data set
-			TArray<USoundWave*> SoundWavesWithCookedData;
-			NewActiveSound.bUpdatePlaybackTime = Sound->GetSoundWavesWithCookedAnalysisData(SoundWavesWithCookedData);
 
-			// Reset the audio component's soundwave playback times
-			SoundWavePlaybackTimes.Reset();
-			for (USoundWave* SoundWave : SoundWavesWithCookedData)
+			// Setup audio component cooked analysis data playback data set
+			if (AudioDevice->IsBakedAnalaysisQueryingEnabled())
 			{
-				SoundWavePlaybackTimes.Add(SoundWave->GetUniqueID(), FSoundWavePlaybackTimeData(SoundWave));
+				TArray<USoundWave*> SoundWavesWithCookedData;
+				NewActiveSound.bUpdatePlaybackTime = Sound->GetSoundWavesWithCookedAnalysisData(SoundWavesWithCookedData);
+
+				// Reset the audio component's soundwave playback times
+				SoundWavePlaybackTimes.Reset();
+				for (USoundWave* SoundWave : SoundWavesWithCookedData)
+				{
+					SoundWavePlaybackTimes.Add(SoundWave->GetUniqueID(), FSoundWavePlaybackTimeData(SoundWave));
+				}
 			}
 
 			NewActiveSound.MaxDistance = MaxDistance;
diff --git a/Engine/Source/Runtime/Engine/Public/AudioDevice.h b/Engine/Source/Runtime/Engine/Public/AudioDevice.h
--- a/Engine/Source/Runtime/Engine/Public/AudioDevice.h
+++ b/Engine/Source/Runtime/Engine/Public/AudioDevice.h
@@ -1171,6 +1171,12 @@ public:
 		return bIsStoppingVoicesEnabled;
 	}
 
+	/** Returns if baked analysis querying is enabled. */
+	bool IsBakedAnalaysisQueryingEnabled() const
+	{
+		return bIsBakedAnalysisEnabled;
+	}
+
 	/** Updates the source effect chain. Only implemented in audio mixer. */
 	virtual void UpdateSourceEffectChain(const uint32 SourceEffectChainId, const TArray<FSourceEffectChainEntry>& SourceEffectChain, const bool bPlayEffectChainTails) {}
 
@@ -1700,6 +1706,9 @@ public:
 
 	uint8 bIsStoppingVoicesEnabled : 1;
 
+	/** If baked analysis querying is enabled. */
+	uint8 bIsBakedAnalysisEnabled : 1;
+
 	/** Whether or not the audio mixer module is being used by this device. */
 	uint8 bAudioMixerModuleLoaded : 1;
 
