Remove cpp type from serialized net field exports in replays

diff --git a/Engine/Source/Runtime/Core/Private/Misc/NetworkVersion.cpp b/Engine/Source/Runtime/Core/Private/Misc/NetworkVersion.cpp
--- a/Engine/Source/Runtime/Core/Private/Misc/NetworkVersion.cpp
+++ b/Engine/Source/Runtime/Core/Private/Misc/NetworkVersion.cpp
@@ -18,7 +18,7 @@ FString FNetworkVersion::ProjectVersion = TEXT("1.0.0");
 bool FNetworkVersion::bHasCachedNetworkChecksum			= false;
 uint32 FNetworkVersion::CachedNetworkChecksum			= 0;
 
-uint32 FNetworkVersion::EngineNetworkProtocolVersion	= HISTORY_ACKS_INCLUDED_IN_HEADER;
+uint32 FNetworkVersion::EngineNetworkProtocolVersion	= HISTORY_NETEXPORT_SERIALIZATION;
 uint32 FNetworkVersion::GameNetworkProtocolVersion		= 0;
 
 uint32 FNetworkVersion::EngineCompatibleNetworkProtocolVersion		= HISTORY_REPLAY_BACKWARDS_COMPAT;
diff --git a/Engine/Source/Runtime/Core/Public/Misc/NetworkVersion.h b/Engine/Source/Runtime/Core/Public/Misc/NetworkVersion.h
--- a/Engine/Source/Runtime/Core/Public/Misc/NetworkVersion.h
+++ b/Engine/Source/Runtime/Core/Public/Misc/NetworkVersion.h
@@ -40,6 +40,7 @@ enum EEngineNetworkVersionHistory
 	HISTORY_CHANNEL_NAMES = 6,						// Bump version since channel type is now an fname
 	HISTORY_CHANNEL_CLOSE_REASON = 7,				// Bump version to serialize a channel close reason in bunches instead of bDormant
 	HISTORY_ACKS_INCLUDED_IN_HEADER = 8,			// Bump version since acks are now sent as part of the header
+	HISTORY_NETEXPORT_SERIALIZATION = 9,			// Bump version due to serialization change to FNetFieldExport
 };
 
 struct CORE_API FNetworkVersion
diff --git a/Engine/Source/Runtime/Engine/Classes/Engine/PackageMapClient.h b/Engine/Source/Runtime/Engine/Classes/Engine/PackageMapClient.h
--- a/Engine/Source/Runtime/Engine/Classes/Engine/PackageMapClient.h
+++ b/Engine/Source/Runtime/Engine/Classes/Engine/PackageMapClient.h
@@ -24,6 +24,7 @@
 #include "UObject/ObjectMacros.h"
 #include "UObject/UObjectGlobals.h"
 #include "Misc/NetworkGuid.h"
+#include "Misc/NetworkVersion.h"
 #include "UObject/CoreNet.h"
 #include "Net/DataBunch.h"
 #include "PackageMapClient.generated.h"
@@ -31,6 +32,7 @@
 class UNetConnection;
 class UNetDriver;
 
+PRAGMA_DISABLE_DEPRECATION_WARNINGS	// 4.22, Name, Type
 class ENGINE_API FNetFieldExport
 {
 public:
@@ -38,16 +40,27 @@ public:
 	{
 	}
 
+	UE_DEPRECATED(4.22, "Type is no longer required, please use other constructor.")
 	FNetFieldExport( const uint32 InHandle, const uint32 InCompatibleChecksum, const FString InName, FString InType ) :
 		bExported( false ),
 		Handle( InHandle ),
 		CompatibleChecksum( InCompatibleChecksum ),
+		ExportName( *InName ),
 		Name( InName ),
 		Type( InType ),
 		bIncompatible( false )
 	{
 	}
 
+	FNetFieldExport( const uint32 InHandle, const uint32 InCompatibleChecksum, const FName& InName ) :
+		bExported( false ),
+		Handle( InHandle ),
+		CompatibleChecksum( InCompatibleChecksum ),
+		ExportName( InName ),
+		bIncompatible( false )
+	{
+	}
+
 	friend FArchive& operator<<( FArchive& Ar, FNetFieldExport& C )
 	{
 		uint8 Flags = C.bExported ? 1 : 0;
@@ -56,13 +69,27 @@ public:
 
 		if ( Ar.IsLoading() )
 		{
-			C.bExported = Flags == 1 ? true : false;
+			C.bExported = (Flags == 1);
 		}
 
 		if ( C.bExported )
 		{
 			Ar.SerializeIntPacked( C.Handle );
-			Ar << C.CompatibleChecksum << C.Name << C.Type;
+			Ar << C.CompatibleChecksum;
+
+			if (Ar.IsLoading() && Ar.EngineNetVer() < HISTORY_NETEXPORT_SERIALIZATION)
+			{
+				Ar << C.Name;
+				Ar << C.Type;
+
+				C.ExportName = FName(*C.Name);
+			}
+			else
+			{
+				Ar << C.ExportName;
+
+				C.Name = C.ExportName.ToString();
+			}
 		}
 
 		return Ar;
@@ -73,12 +100,16 @@ public:
 	bool			bExported;
 	uint32			Handle;
 	uint32			CompatibleChecksum;
+	FName			ExportName;
+	UE_DEPRECATED(4.22, "Name is deprecated.")
 	FString			Name;
+	UE_DEPRECATED(4.22, "Type is deprecated.")
 	FString			Type;
 
 	// Transient properties
 	mutable bool	bIncompatible;		// If true, we've already determined that this property isn't compatible. We use this to curb warning spam.
 };
+PRAGMA_ENABLE_DEPRECATION_WARNINGS
 
 class ENGINE_API FNetFieldExportGroup
 {
diff --git a/Engine/Source/Runtime/Engine/Private/DataChannel.cpp b/Engine/Source/Runtime/Engine/Private/DataChannel.cpp
--- a/Engine/Source/Runtime/Engine/Private/DataChannel.cpp
+++ b/Engine/Source/Runtime/Engine/Private/DataChannel.cpp
@@ -3482,7 +3482,7 @@ bool UActorChannel::ReadFieldHeaderAndPayload( UObject* Object, const FClassNetC
 
 		if ( !ensure( NetFieldExport.CompatibleChecksum != 0 ) )
 		{
-			UE_LOG( LogNet, Error, TEXT( "ReadFieldHeaderAndPayload: NetFieldExport.CompatibleChecksum was 0. Object: %s, Property: %s, Type: %s" ), *Object->GetFullName(), *NetFieldExport.Name, *NetFieldExport.Type );
+			UE_LOG( LogNet, Error, TEXT( "ReadFieldHeaderAndPayload: NetFieldExport.CompatibleChecksum was 0. Object: %s, Property: %s" ), *Object->GetFullName(), *NetFieldExport.ExportName.ToString() );
 			Bunch.SetError();
 			return false;
 		}
@@ -3493,7 +3493,7 @@ bool UActorChannel::ReadFieldHeaderAndPayload( UObject* Object, const FClassNetC
 		{
 			if ( !NetFieldExport.bIncompatible )
 			{
-				UE_LOG( LogNet, Warning, TEXT( "ReadFieldHeaderAndPayload: GetFromChecksum failed (NetBackwardsCompatibility). Object: %s, Property: %s, Type: %s" ), *Object->GetFullName(), *NetFieldExport.Name, *NetFieldExport.Type );
+				UE_LOG( LogNet, Warning, TEXT( "ReadFieldHeaderAndPayload: GetFromChecksum failed (NetBackwardsCompatibility). Object: %s, Property: %s" ), *Object->GetFullName(), *NetFieldExport.ExportName.ToString() );
 				NetFieldExport.bIncompatible = true;
 			}
 		}
@@ -3595,13 +3595,7 @@ FNetFieldExportGroup* UActorChannel::GetOrCreateNetFieldExportGroupForClassNetCa
 					continue;	// We only care about net fields that aren't in a rep layout
 				}
 
-				FNetFieldExport NetFieldExport(
-					CurrentHandle++,
-					Fields[i].FieldChecksum,
-					Field ? Field->GetName() : TEXT( "" ),
-					Property ? Property->GetCPPType( nullptr, 0 ) : TEXT( "" ) );
-
-				NetFieldExportGroup->NetFieldExports.Add( NetFieldExport );
+				NetFieldExportGroup->NetFieldExports.Emplace( CurrentHandle++, Fields[i].FieldChecksum, Field ? Field->GetFName() : NAME_None );
 			}
 		}
 
diff --git a/Engine/Source/Runtime/Engine/Private/DemoNetDriver.cpp b/Engine/Source/Runtime/Engine/Private/DemoNetDriver.cpp
--- a/Engine/Source/Runtime/Engine/Private/DemoNetDriver.cpp
+++ b/Engine/Source/Runtime/Engine/Private/DemoNetDriver.cpp
@@ -3667,6 +3667,12 @@ void UDemoNetDriver::TickDemoPlayback( float DeltaSeconds )
 	// Speculatively grab seconds now in case we need it to get the time it took to fast forward
 	const double FastForwardStartSeconds = FPlatformTime::Seconds();
 
+	if (FArchive* const StreamingArchive = ReplayStreamer->GetStreamingArchive())
+	{
+		StreamingArchive->SetEngineNetVer(PlaybackDemoHeader.EngineNetworkProtocolVersion);
+		StreamingArchive->SetGameNetVer(PlaybackDemoHeader.GameNetworkProtocolVersion);
+	}
+
 	// Buffer up demo frames until we have enough time built-up
 	while ( ConditionallyReadDemoFrameIntoPlaybackPackets( *ReplayStreamer->GetStreamingArchive() ) )
 	{
@@ -4290,6 +4296,9 @@ bool UDemoNetDriver::FastForwardLevels(const FGotoResult& GotoResult)
 		// First, read in the checkpoint data (if any is available);
 		if (CheckpointArchive->TotalSize() != 0)
 		{
+			CheckpointArchive->SetEngineNetVer(PlaybackDemoHeader.EngineNetworkProtocolVersion);
+			CheckpointArchive->SetGameNetVer(PlaybackDemoHeader.GameNetworkProtocolVersion);
+
 			TGuardValue<bool> LoadingCheckpointGuard(bIsLoadingCheckpoint, true);
 
 			FArchivePos PacketOffset = 0;
@@ -4306,6 +4315,10 @@ bool UDemoNetDriver::FastForwardLevels(const FGotoResult& GotoResult)
 
 		// Next, read in streaming data (if any is available)
 		FArchive* StreamingAr = ReplayStreamer->GetStreamingArchive();
+
+		StreamingAr->SetEngineNetVer(PlaybackDemoHeader.EngineNetworkProtocolVersion);
+		StreamingAr->SetGameNetVer(PlaybackDemoHeader.GameNetworkProtocolVersion);
+		
 		while (!StreamingAr->AtEnd() && ReplayStreamer->IsDataAvailable() && ReadPacketsHelper.ReadPackets(*StreamingAr));
 
 		if (ReadPacketsHelper.IsError())
@@ -4485,6 +4498,9 @@ bool UDemoNetDriver::LoadCheckpoint(const FGotoResult& GotoResult)
 	check( !bIsFastForwardingForCheckpoint );
 	check( !bIsFastForwarding );
 
+	GotoCheckpointArchive->SetEngineNetVer(PlaybackDemoHeader.EngineNetworkProtocolVersion);
+	GotoCheckpointArchive->SetGameNetVer(PlaybackDemoHeader.GameNetworkProtocolVersion);
+
 	int32 LevelForCheckpoint = 0;
 
 	if (HasLevelStreamingFixes())
diff --git a/Engine/Source/Runtime/Engine/Private/PackageMapClient.cpp b/Engine/Source/Runtime/Engine/Private/PackageMapClient.cpp
--- a/Engine/Source/Runtime/Engine/Private/PackageMapClient.cpp
+++ b/Engine/Source/Runtime/Engine/Private/PackageMapClient.cpp
@@ -3187,8 +3187,10 @@ void FNetGUIDCache::CountBytes(FArchive& Ar) const
 
 void FNetFieldExport::CountBytes(FArchive& Ar) const
 {
+	PRAGMA_DISABLE_DEPRECATION_WARNINGS
 	Ar << const_cast<FString&>(Name);
 	Ar << const_cast<FString&>(Type);
+	PRAGMA_ENABLE_DEPRECATION_WARNINGS
 }
 
 void FNetFieldExportGroup::CountBytes(FArchive& Ar) const
diff --git a/Engine/Source/Runtime/Engine/Private/RepLayout.cpp b/Engine/Source/Runtime/Engine/Private/RepLayout.cpp
--- a/Engine/Source/Runtime/Engine/Private/RepLayout.cpp
+++ b/Engine/Source/Runtime/Engine/Private/RepLayout.cpp
@@ -1695,8 +1695,7 @@ TSharedPtr<FNetFieldExportGroup> FRepLayout::CreateNetfieldExportGroup() const
 		FNetFieldExport NetFieldExport(
 			i,
 			Cmds[i].CompatibleChecksum,
-			Cmds[i].Property ? Cmds[i].Property->GetName() : TEXT(""),
-			Cmds[i].Property ? Cmds[i].Property->GetCPPType(nullptr, 0) : TEXT(""));
+			Cmds[i].Property ? Cmds[i].Property->GetFName() : NAME_None );
 
 		NetFieldExportGroup->NetFieldExports[i] = NetFieldExport;
 	}
@@ -2480,7 +2479,7 @@ bool FRepLayout::ReceiveProperties_BackwardsCompatible_r(
 
 		if (!ensure(Checksum != 0))
 		{
-			UE_LOG(LogRep, Warning, TEXT("ReceiveProperties_BackwardsCompatible_r: Checksum == 0. Owner: %s, Name: %s, Type: %s, NetFieldExportHandle: %i"), *Owner->GetName(), *NetFieldExportGroup->NetFieldExports[NetFieldExportHandle].Name, *NetFieldExportGroup->NetFieldExports[NetFieldExportHandle].Type, NetFieldExportHandle);
+			UE_LOG(LogRep, Warning, TEXT("ReceiveProperties_BackwardsCompatible_r: Checksum == 0. Owner: %s, Name: %s, NetFieldExportHandle: %i"), *Owner->GetName(), *NetFieldExportGroup->NetFieldExports[NetFieldExportHandle].ExportName.ToString(), NetFieldExportHandle);
 			return false;
 		}
 
@@ -2491,7 +2490,7 @@ bool FRepLayout::ReceiveProperties_BackwardsCompatible_r(
 
 		if (Reader.IsError())
 		{
-			UE_LOG(LogRep, Warning, TEXT("ReceiveProperties_BackwardsCompatible_r: Error reading num bits. Owner: %s, Name: %s, Type: %s, NetFieldExportHandle: %i, Checksum: %u"), *Owner->GetName(), *NetFieldExportGroup->NetFieldExports[NetFieldExportHandle].Name, *NetFieldExportGroup->NetFieldExports[NetFieldExportHandle].Type, NetFieldExportHandle, Checksum);
+			UE_LOG(LogRep, Warning, TEXT("ReceiveProperties_BackwardsCompatible_r: Error reading num bits. Owner: %s, Name: %s, NetFieldExportHandle: %i, Checksum: %u"), *Owner->GetName(), *NetFieldExportGroup->NetFieldExports[NetFieldExportHandle].ExportName.ToString(), NetFieldExportHandle, Checksum);
 			return false;
 		}
 
@@ -2502,7 +2501,7 @@ bool FRepLayout::ReceiveProperties_BackwardsCompatible_r(
 
 		if (Reader.IsError())
 		{
-			UE_LOG(LogRep, Warning, TEXT("ReceiveProperties_BackwardsCompatible_r: Error reading payload. Owner: %s, Name: %s, Type: %s, NetFieldExportHandle: %i, Checksum: %u"), *Owner->GetName(), *NetFieldExportGroup->NetFieldExports[NetFieldExportHandle].Name, *NetFieldExportGroup->NetFieldExports[NetFieldExportHandle].Type, NetFieldExportHandle, Checksum);
+			UE_LOG(LogRep, Warning, TEXT("ReceiveProperties_BackwardsCompatible_r: Error reading payload. Owner: %s, Name: %s, NetFieldExportHandle: %i, Checksum: %u"), *Owner->GetName(), *NetFieldExportGroup->NetFieldExports[NetFieldExportHandle].ExportName.ToString(), NetFieldExportHandle, Checksum);
 			return false;
 		}
 
@@ -2516,7 +2515,7 @@ bool FRepLayout::ReceiveProperties_BackwardsCompatible_r(
 
 		if (CmdIndex == -1)
 		{
-			UE_LOG(LogRep, Warning, TEXT("ReceiveProperties_BackwardsCompatible_r: Property not found. Owner: %s, Name: %s, Type: %s, NetFieldExportHandle: %i, Checksum: %u"), *Owner->GetName(), *NetFieldExportGroup->NetFieldExports[NetFieldExportHandle].Name, *NetFieldExportGroup->NetFieldExports[NetFieldExportHandle].Type, NetFieldExportHandle, Checksum);
+			UE_LOG(LogRep, Warning, TEXT("ReceiveProperties_BackwardsCompatible_r: Property not found. Owner: %s, Name: %s, NetFieldExportHandle: %i, Checksum: %u"), *Owner->GetName(), *NetFieldExportGroup->NetFieldExports[NetFieldExportHandle].ExportName.ToString(), NetFieldExportHandle, Checksum);
 
 			// Mark this property as incompatible so we don't keep spamming this warning
 			NetFieldExportGroup->NetFieldExports[NetFieldExportHandle].bIncompatible = true;
@@ -2534,7 +2533,7 @@ bool FRepLayout::ReceiveProperties_BackwardsCompatible_r(
 
 			if (TempReader.IsError())
 			{
-				UE_LOG(LogRep, Warning, TEXT("ReceiveProperties_BackwardsCompatible_r: Error reading ArrayNum. Owner: %s, Name: %s, Type: %s, NetFieldExportHandle: %i, Checksum: %u"), *Owner->GetName(), *NetFieldExportGroup->NetFieldExports[NetFieldExportHandle].Name, *NetFieldExportGroup->NetFieldExports[NetFieldExportHandle].Type, NetFieldExportHandle, Checksum);
+				UE_LOG(LogRep, Warning, TEXT("ReceiveProperties_BackwardsCompatible_r: Error reading ArrayNum. Owner: %s, Name: %s, NetFieldExportHandle: %i, Checksum: %u"), *Owner->GetName(), *NetFieldExportGroup->NetFieldExports[NetFieldExportHandle].ExportName.ToString(), NetFieldExportHandle, Checksum);
 				return false;
 			}
 
@@ -2571,7 +2570,7 @@ bool FRepLayout::ReceiveProperties_BackwardsCompatible_r(
 
 				if (TempReader.IsError())
 				{
-					UE_LOG(LogRep, Warning, TEXT("ReceiveProperties_BackwardsCompatible_r: Error reading array index. Index: %i, Owner: %s, Name: %s, Type: %s, NetFieldExportHandle: %i, Checksum: %u"), Index, *Owner->GetName(), *NetFieldExportGroup->NetFieldExports[NetFieldExportHandle].Name, *NetFieldExportGroup->NetFieldExports[NetFieldExportHandle].Type, NetFieldExportHandle, Checksum);
+					UE_LOG(LogRep, Warning, TEXT("ReceiveProperties_BackwardsCompatible_r: Error reading array index. Index: %i, Owner: %s, Name: %s, NetFieldExportHandle: %i, Checksum: %u"), Index, *Owner->GetName(), *NetFieldExportGroup->NetFieldExports[NetFieldExportHandle].ExportName.ToString(), NetFieldExportHandle, Checksum);
 					return false;
 				}
 
@@ -2604,7 +2603,7 @@ bool FRepLayout::ReceiveProperties_BackwardsCompatible_r(
 
 				if (!ensure(Index < ArrayNum))
 				{
-					UE_LOG(LogRep, Warning, TEXT("ReceiveProperties_BackwardsCompatible_r: Array index out of bounds. Index: %i, ArrayNum: %i, Owner: %s, Name: %s, Type: %s, NetFieldExportHandle: %i, Checksum: %u"), Index, ArrayNum, *Owner->GetName(), *NetFieldExportGroup->NetFieldExports[NetFieldExportHandle].Name, *NetFieldExportGroup->NetFieldExports[NetFieldExportHandle].Type, NetFieldExportHandle, Checksum);
+					UE_LOG(LogRep, Warning, TEXT("ReceiveProperties_BackwardsCompatible_r: Array index out of bounds. Index: %i, ArrayNum: %i, Owner: %s, Name: %s, NetFieldExportHandle: %i, Checksum: %u"), Index, ArrayNum, *Owner->GetName(), *NetFieldExportGroup->NetFieldExports[NetFieldExportHandle].ExportName.ToString(), NetFieldExportHandle, Checksum);
 					return false;
 				}
 
@@ -2620,14 +2619,14 @@ bool FRepLayout::ReceiveProperties_BackwardsCompatible_r(
 
 				if (TempReader.IsError())
 				{
-					UE_LOG(LogRep, Warning, TEXT("ReceiveProperties_BackwardsCompatible_r: Error reading array index element payload. Index: %i, Owner: %s, Name: %s, Type: %s, NetFieldExportHandle: %i, Checksum: %u"), Index, *Owner->GetName(), *NetFieldExportGroup->NetFieldExports[NetFieldExportHandle].Name, *NetFieldExportGroup->NetFieldExports[NetFieldExportHandle].Type, NetFieldExportHandle, Checksum);
+					UE_LOG(LogRep, Warning, TEXT("ReceiveProperties_BackwardsCompatible_r: Error reading array index element payload. Index: %i, Owner: %s, Name: %s, NetFieldExportHandle: %i, Checksum: %u"), Index, *Owner->GetName(), *NetFieldExportGroup->NetFieldExports[NetFieldExportHandle].ExportName.ToString(), NetFieldExportHandle, Checksum);
 					return false;
 				}
 			}
 
 			if (TempReader.GetBitsLeft() != 0)
 			{
-				UE_LOG(LogRep, Warning, TEXT("ReceiveProperties_BackwardsCompatible_r: Array didn't read proper number of bits. Owner: %s, Name: %s, Type: %s, NetFieldExportHandle: %i, Checksum: %u, BitsLeft:%d"), *Owner->GetName(), *NetFieldExportGroup->NetFieldExports[NetFieldExportHandle].Name, *NetFieldExportGroup->NetFieldExports[NetFieldExportHandle].Type, NetFieldExportHandle, Checksum, TempReader.GetBitsLeft());
+				UE_LOG(LogRep, Warning, TEXT("ReceiveProperties_BackwardsCompatible_r: Array didn't read proper number of bits. Owner: %s, Name: %s, NetFieldExportHandle: %i, Checksum: %u, BitsLeft:%d"), *Owner->GetName(), *NetFieldExportGroup->NetFieldExports[NetFieldExportHandle].ExportName.ToString(), NetFieldExportHandle, Checksum, TempReader.GetBitsLeft());
 				return false;
 			}
 		}
@@ -2642,7 +2641,7 @@ bool FRepLayout::ReceiveProperties_BackwardsCompatible_r(
 
 			if (TempReader.GetBitsLeft() != 0)
 			{
-				UE_LOG(LogRep, Warning, TEXT("ReceiveProperties_BackwardsCompatible_r: Property didn't read proper number of bits. Owner: %s, Name: %s, Type: %s, NetFieldExportHandle: %i, Checksum: %u, BitsLeft:%d"), *Owner->GetName(), *NetFieldExportGroup->NetFieldExports[NetFieldExportHandle].Name, *NetFieldExportGroup->NetFieldExports[NetFieldExportHandle].Type, NetFieldExportHandle, Checksum, TempReader.GetBitsLeft());
+				UE_LOG(LogRep, Warning, TEXT("ReceiveProperties_BackwardsCompatible_r: Property didn't read proper number of bits. Owner: %s, Name: %s, NetFieldExportHandle: %i, Checksum: %u, BitsLeft:%d"), *Owner->GetName(), *NetFieldExportGroup->NetFieldExports[NetFieldExportHandle].ExportName.ToString(), NetFieldExportHandle, Checksum, TempReader.GetBitsLeft());
 				return false;
 			}
 		}
diff --git a/Engine/Source/Runtime/GameplayTags/Private/GameplayTagContainer.cpp b/Engine/Source/Runtime/GameplayTags/Private/GameplayTagContainer.cpp
--- a/Engine/Source/Runtime/GameplayTags/Private/GameplayTagContainer.cpp
+++ b/Engine/Source/Runtime/GameplayTags/Private/GameplayTagContainer.cpp
@@ -1194,8 +1194,7 @@ static TSharedPtr<FNetFieldExportGroup> CreateNetfieldExportGroupForNetworkGamep
 		FNetFieldExport NetFieldExport(
 			i,
 			0,
-			NetworkGameplayTagNodeIndex[i]->GetCompleteTagString(),
-			TEXT(""));
+			NetworkGameplayTagNodeIndex[i]->GetCompleteTagName());
 
 		NetFieldExportGroup->NetFieldExports[i] = NetFieldExport;
 	}
@@ -1257,7 +1256,7 @@ bool FGameplayTag::NetSerialize_Packed(FArchive& Ar, class UPackageMap* Map, boo
 				// Get the tag name from the net field export group entry
 				if (NetIndex != INVALID_TAGNETINDEX && ensure(NetFieldExportGroup.IsValid()) && ensure(NetIndex < NetFieldExportGroup->NetFieldExports.Num()))
 				{
-					TagName = FName(*NetFieldExportGroup->NetFieldExports[NetIndex].Name);
+					TagName = NetFieldExportGroup->NetFieldExports[NetIndex].ExportName;
 
 					// Validate the tag name
 					const FGameplayTag Tag = UGameplayTagsManager::Get().RequestGameplayTag(TagName, false);
