PSC manager now has each tick group wait for task it kicks off.

diff --git a/Engine/Source/Runtime/Engine/Private/Particles/ParticleSystemManager.cpp b/Engine/Source/Runtime/Engine/Private/Particles/ParticleSystemManager.cpp
--- a/Engine/Source/Runtime/Engine/Private/Particles/ParticleSystemManager.cpp
+++ b/Engine/Source/Runtime/Engine/Private/Particles/ParticleSystemManager.cpp
@@ -647,18 +647,8 @@ void FParticleSystemWorldManager::Tick(ETickingGroup TickGroup, float DeltaTime,
 	bool bAllowTickConcurrent = !FXConsoleVariables::bFreezeParticleSimulation && FXConsoleVariables::bAllowAsyncTick && FApp::ShouldUseThreadingForPerformance() && GDistributionType != 0;
 	if (bAllowTickConcurrent)
 	{
-		//Have the final tick group wait for completion until all our async work is complete.
-		//This is required for safety as afaik there is no other mechanism that would prevent our async tasks running past the GT portion of the frame and overlapping with the EOF updates.
-		//If this is an incorrect assumption then this can be discarded.
-		if (TickFunctions.Last().IsCompletionHandleValid())
-		{
-			ProcessTickList<true>(DeltaTime, TickType, TickGroup, TickLists_Concurrent, TickFunctions.Last().GetCompletionHandle());
-		}
-		else
-		{
-			UE_LOG(LogParticles, Warning, TEXT("PSC Manager final tick function did not have a valid completion handle. Waiting on this tick group."));
-			ProcessTickList<true>(DeltaTime, TickType, TickGroup, TickLists_Concurrent, MyCompletionGraphEvent);
-		}
+		//TODO: Currently waiting on this tick group but we should be able to allow these tasks to run over the whole frame if we can implement some synchronization so they don't overrun the EOF updates.
+		ProcessTickList<true>(DeltaTime, TickType, TickGroup, TickLists_Concurrent, MyCompletionGraphEvent);
 	}
 	else
 	{
