PR #5484: Fixed an issue that delegate could not be added to OnNewPawn because it returns copy value. (Contributed by [REDACTED_PERSON])

diff --git a/Engine/Source/Runtime/Engine/Classes/GameFramework/Controller.h b/Engine/Source/Runtime/Engine/Classes/GameFramework/Controller.h
--- a/Engine/Source/Runtime/Engine/Classes/GameFramework/Controller.h
+++ b/Engine/Source/Runtime/Engine/Classes/GameFramework/Controller.h
@@ -201,7 +201,7 @@ public:
 	UFUNCTION(BlueprintCallable, Category=Pawn, meta=(DisplayName="Get Controlled Pawn", ScriptName="GetControlledPawn"))
 	APawn* K2_GetPawn() const;
 
-	FPawnChangedSignature GetOnNewPawnNotifier() { return OnNewPawn; }
+	FPawnChangedSignature& GetOnNewPawnNotifier() { return OnNewPawn; }
 
 public:
 
