Removed needless logging in the animation system for test & shipping builds

diff --git a/Engine/Source/Runtime/Engine/Private/Animation/AnimInstanceProxy.cpp b/Engine/Source/Runtime/Engine/Private/Animation/AnimInstanceProxy.cpp
--- a/Engine/Source/Runtime/Engine/Private/Animation/AnimInstanceProxy.cpp
+++ b/Engine/Source/Runtime/Engine/Private/Animation/AnimInstanceProxy.cpp
@@ -137,7 +137,7 @@ void FAnimInstanceProxy::Initialize(UAnimInstance* InAnimInstance)
 		RootNode = (FAnimNode_Base*) GetCustomRootNode();
 	}
 
-#if !NO_LOGGING
+#if ENABLE_ANIM_LOGGING
 	ActorName = GetNameSafe(InAnimInstance->GetOwningActor());
 #endif
 
@@ -251,7 +251,7 @@ FGuid MakeGuidForMessage(const FText& Message)
 
 void FAnimInstanceProxy::LogMessage(FName InLogType, EMessageSeverity::Type InSeverity, const FText& InMessage)
 {
-#if !NO_LOGGING
+#if ENABLE_ANIM_LOGGING
 	FGuid CurrentMessageGuid = MakeGuidForMessage(InMessage);
 	if(!PreviouslyLoggedMessages.Contains(CurrentMessageGuid))
 	{
@@ -308,7 +308,7 @@ void FAnimInstanceProxy::PreUpdate(UAnimInstance* InAnimInstance, float DeltaSec
 	QueuedDrawDebugItems.Reset();
 #endif
 
-#if !NO_LOGGING
+#if ENABLE_ANIM_LOGGING
 	//Reset logged update messages
 	LoggedMessagesMap.FindOrAdd(NAME_Update).Reset();
 #endif
@@ -438,7 +438,7 @@ void FAnimInstanceProxy::PostUpdate(UAnimInstance* InAnimInstance) const
 	}
 #endif
 
-#if !NO_LOGGING
+#if ENABLE_ANIM_LOGGING
 	FMessageLog MessageLog(NAME_AnimBlueprintLog);
 	const TArray<FLogMessageEntry>* Messages = LoggedMessagesMap.Find(NAME_Update);
 	if (ensureMsgf(Messages, TEXT("PreUpdate isn't called. This could potentially cause other issues.")))
@@ -455,7 +455,7 @@ void FAnimInstanceProxy::PostEvaluate(UAnimInstance* InAnimInstance)
 {
 	ClearObjects();
 
-#if !NO_LOGGING
+#if ENABLE_ANIM_LOGGING
 	FMessageLog MessageLog(NAME_AnimBlueprintLog);
 	if(const TArray<FLogMessageEntry>* Messages = LoggedMessagesMap.Find(NAME_Evaluate))
 	{
@@ -1111,7 +1111,7 @@ void FAnimInstanceProxy::UpdateAnimation()
 void FAnimInstanceProxy::PreEvaluateAnimation(UAnimInstance* InAnimInstance)
 {
 	InitializeObjects(InAnimInstance);
-#if !NO_LOGGING
+#if ENABLE_ANIM_LOGGING
 	LoggedMessagesMap.FindOrAdd(NAME_Evaluate).Reset();
 #endif
 }
@@ -1349,7 +1349,7 @@ void FAnimInstanceProxy::GetSlotWeight(const FName& SlotNodeName, float& out_Slo
 #if DEBUGMONTAGEWEIGHT			
 				TotalDesiredWeight += EvalState->DesiredWeight;
 #endif
-#if !NO_LOGGING
+#if ENABLE_ANIM_LOGGING
 			UE_LOG(LogAnimation, Verbose, TEXT("GetSlotWeight : Owner: %s, AnimMontage: %s,  (DesiredWeight:%0.2f, Weight:%0.2f)"),
 						*GetActorName(), *EvalState.Montage->GetName(), EvalState.DesiredWeight, EvalState.MontageWeight);
 #endif
diff --git a/Engine/Source/Runtime/Engine/Public/Animation/AnimInstanceProxy.h b/Engine/Source/Runtime/Engine/Public/Animation/AnimInstanceProxy.h
--- a/Engine/Source/Runtime/Engine/Public/Animation/AnimInstanceProxy.h
+++ b/Engine/Source/Runtime/Engine/Public/Animation/AnimInstanceProxy.h
@@ -32,6 +32,9 @@ struct FPoseContext;
 // Disable debugging information for shipping and test builds.
 #define ENABLE_ANIM_DRAW_DEBUG (1 && !(UE_BUILD_SHIPPING || UE_BUILD_TEST))
 
+// Disable node logging for shipping and test builds
+#define ENABLE_ANIM_LOGGING (1 && !NO_LOGGING && !(UE_BUILD_SHIPPING || UE_BUILD_TEST))
+
 UENUM()
 namespace EDrawDebugItemType
 {
@@ -384,7 +387,7 @@ public:
 	void AnimDrawDebugPlane(const FTransform& BaseTransform, float Radii, const FColor& Color, bool bPersistentLines = false, float LifeTime = -1.f, float Thickness = 0.f) {}
 #endif // ENABLE_ANIM_DRAW_DEBUG
 
-#if !NO_LOGGING
+#if ENABLE_ANIM_LOGGING
 	const FString& GetActorName() const 
 	{
 		return ActorName; 
@@ -745,7 +748,7 @@ private:
 	TArray<FAnimNodePoseWatch> PoseWatchEntriesForThisFrame;
 #endif
 
-#if !NO_LOGGING
+#if ENABLE_ANIM_LOGGING
 	/** Actor name for debug logging purposes */
 	FString ActorName;
 #endif
@@ -849,7 +852,7 @@ private:
 	/** Array of snapshots. Each entry contains a name for finding specific pose snapshots */
 	TArray<FPoseSnapshot> PoseSnapshots;
 
-#if !NO_LOGGING
+#if ENABLE_ANIM_LOGGING
 	/** Logged message queues. Allows nodes to report messages to MessageLog even though they may be running
 	 *  on a worked thread
 	 */
