Fix CIS Issue for RepLayout.cpp

diff --git a/Engine/Source/Runtime/Engine/Private/RepLayout.cpp b/Engine/Source/Runtime/Engine/Private/RepLayout.cpp
--- a/Engine/Source/Runtime/Engine/Private/RepLayout.cpp
+++ b/Engine/Source/Runtime/Engine/Private/RepLayout.cpp
@@ -2336,7 +2336,7 @@ bool FRepLayout::ReceiveProperties(
 	const bool bDoChecksum = false;
 #endif
 
-	UE_LOG(LogRepProperties, VeryVerbose, TEXT("ReceiveProperties: Owner=%s, LastChangelistIndex=%d"), *Owner->GetPathName(), RepState ? RepState->LastChangelistIndex : INDEX_NONE);
+	UE_LOG(LogRepProperties, VeryVerbose, TEXT("ReceiveProperties: Owner=%s, LastChangelistIndex=%d"), *Owner->GetPathName(), RepState->LastChangelistIndex);
 
 	bOutHasUnmapped = false;
 
