Adding cvar to disable occlusion

diff --git a/Engine/Source/Runtime/Engine/Private/ActiveSound.cpp b/Engine/Source/Runtime/Engine/Private/ActiveSound.cpp
--- a/Engine/Source/Runtime/Engine/Private/ActiveSound.cpp
+++ b/Engine/Source/Runtime/Engine/Private/ActiveSound.cpp
@@ -11,6 +11,13 @@
 #include "SubtitleManager.h"
 #include "DSP/Dsp.h"
 
+static int32 AudioOcclusionDisabledCvar = 0;
+FAutoConsoleVariableRef CVarAudioOcclusionEnabled(
+	TEXT("au.DisableOcclusion"),
+	AudioOcclusionDisabledCvar,
+	TEXT("Disables (1) or enables (0) audio occlusion.\n"),
+	ECVF_Default);
+
 FTraceDelegate FActiveSound::ActiveSoundTraceDelegate;
 TMap<FTraceHandle, FActiveSound::FAsyncTraceDetails> FActiveSound::TraceToActiveSoundMap;
 
@@ -838,59 +845,71 @@ void FActiveSound::CheckOcclusion(const FVector ListenerLocation, const FVector
 	check(AttenuationSettingsPtr);
 	check(AttenuationSettingsPtr->bEnableOcclusion);
 
-	if (!bAsyncOcclusionPending && (PlaybackTime - LastOcclusionCheckTime) > OcclusionCheckInterval)
+	float InterpolationTime = 0.0f;
+
+	// If occlusion is disabled by cvar, we're always going to be not occluded
+	if (AudioOcclusionDisabledCvar == 1)
 	{
-		LastOcclusionCheckTime = PlaybackTime;
+		bIsOccluded = false;
+	}
+	else
+	{
+		if (!bAsyncOcclusionPending && (PlaybackTime - LastOcclusionCheckTime) > OcclusionCheckInterval)
+		{
+			LastOcclusionCheckTime = PlaybackTime;
 
-		const bool bUseComplexCollisionForOcclusion = AttenuationSettingsPtr->bUseComplexCollisionForOcclusion;
-		const ECollisionChannel OcclusionTraceChannel = AttenuationSettingsPtr->OcclusionTraceChannel;
+			const bool bUseComplexCollisionForOcclusion = AttenuationSettingsPtr->bUseComplexCollisionForOcclusion;
+			const ECollisionChannel OcclusionTraceChannel = AttenuationSettingsPtr->OcclusionTraceChannel;
 
-		if (!bHasCheckedOcclusion)
-		{
-			FCollisionQueryParams Params(SCENE_QUERY_STAT(SoundOcclusion), bUseComplexCollisionForOcclusion);
-			if (OwnerID > 0)
+			if (!bHasCheckedOcclusion)
 			{
-				Params.AddIgnoredActor(OwnerID);
-			}
+				FCollisionQueryParams Params(SCENE_QUERY_STAT(SoundOcclusion), bUseComplexCollisionForOcclusion);
+				if (OwnerID > 0)
+				{
+					Params.AddIgnoredActor(OwnerID);
+				}
 
-			if (UWorld* WorldPtr = World.Get())
-			{
-				// LineTraceTestByChannel is generally threadsafe, but there is a very narrow race condition here 
-				// if World goes invalid before the scene lock and queries begin.
-				bIsOccluded = WorldPtr->LineTraceTestByChannel(SoundLocation, ListenerLocation, OcclusionTraceChannel, Params);
+				if (UWorld* WorldPtr = World.Get())
+				{
+					// LineTraceTestByChannel is generally threadsafe, but there is a very narrow race condition here 
+					// if World goes invalid before the scene lock and queries begin.
+					bIsOccluded = WorldPtr->LineTraceTestByChannel(SoundLocation, ListenerLocation, OcclusionTraceChannel, Params);
+				}
 			}
-		}
-		else
-		{
-			bAsyncOcclusionPending = true;
+			else
+			{
+				bAsyncOcclusionPending = true;
 
-			const uint32 SoundOwnerID = OwnerID;
-			TWeakObjectPtr<UWorld> SoundWorld = World;
-			FAsyncTraceDetails TraceDetails;
-			TraceDetails.AudioDeviceID = AudioDevice->DeviceHandle;
-			TraceDetails.ActiveSound = this;
+				const uint32 SoundOwnerID = OwnerID;
+				TWeakObjectPtr<UWorld> SoundWorld = World;
+				FAsyncTraceDetails TraceDetails;
+				TraceDetails.AudioDeviceID = AudioDevice->DeviceHandle;
+				TraceDetails.ActiveSound = this;
 
-			FAudioThread::RunCommandOnGameThread([SoundWorld, SoundLocation, ListenerLocation, OcclusionTraceChannel, SoundOwnerID, bUseComplexCollisionForOcclusion, TraceDetails]
-			{
-				if (UWorld* WorldPtr = SoundWorld.Get())
+				FAudioThread::RunCommandOnGameThread([SoundWorld, SoundLocation, ListenerLocation, OcclusionTraceChannel, SoundOwnerID, bUseComplexCollisionForOcclusion, TraceDetails]
 				{
-					FCollisionQueryParams Params(SCENE_QUERY_STAT(SoundOcclusion), bUseComplexCollisionForOcclusion);
-					if (SoundOwnerID > 0)
+					if (UWorld* WorldPtr = SoundWorld.Get())
 					{
-						Params.AddIgnoredActor(SoundOwnerID);
+						FCollisionQueryParams Params(SCENE_QUERY_STAT(SoundOcclusion), bUseComplexCollisionForOcclusion);
+						if (SoundOwnerID > 0)
+						{
+							Params.AddIgnoredActor(SoundOwnerID);
+						}
+
+						FTraceHandle TraceHandle = WorldPtr->AsyncLineTraceByChannel(EAsyncTraceType::Test, SoundLocation, ListenerLocation, OcclusionTraceChannel, Params, FCollisionResponseParams::DefaultResponseParam, &ActiveSoundTraceDelegate);
+						TraceToActiveSoundMap.Add(TraceHandle, TraceDetails);
 					}
+				});
+			}
+		}
 
-					FTraceHandle TraceHandle = WorldPtr->AsyncLineTraceByChannel(EAsyncTraceType::Test, SoundLocation, ListenerLocation, OcclusionTraceChannel, Params, FCollisionResponseParams::DefaultResponseParam, &ActiveSoundTraceDelegate);
-					TraceToActiveSoundMap.Add(TraceHandle, TraceDetails);
-				}
-			});
+		// Update the occlusion values
+		if (bHasCheckedOcclusion)
+		{
+			InterpolationTime = AttenuationSettingsPtr->OcclusionInterpolationTime;
 		}
 	}
 
-	// Update the occlusion values
-	const float InterpolationTime = bHasCheckedOcclusion ? AttenuationSettingsPtr->OcclusionInterpolationTime : 0.0f;
-	bHasCheckedOcclusion = true;
-
 	if (bIsOccluded)
 	{
 		if (CurrentOcclusionFilterFrequency.GetTargetValue() > AttenuationSettingsPtr->OcclusionLowPassFilterFrequency)
