Add a new ping calculation method which only uses the fastest packets. This allows us to use the packets which were processed by the server and client more quickly after reaching them which produces more accurate calculations. This is currently not used anywhere in the engine pending further testing.

diff --git a/Engine/Source/Runtime/Engine/Classes/GameFramework/PlayerState.h b/Engine/Source/Runtime/Engine/Classes/GameFramework/PlayerState.h
--- a/Engine/Source/Runtime/Engine/Classes/GameFramework/PlayerState.h
+++ b/Engine/Source/Runtime/Engine/Classes/GameFramework/PlayerState.h
@@ -31,6 +31,31 @@ struct PingAvgData
 	}
 };
 
+/**
+ * Struct keeping track of the lowest ping values over a given second.
+ */
+struct PingAvgDataV2
+{
+	/** List of ping values */
+	TArray<uint16> PingValues;
+
+	/** The maximum number of ping values we will keep track of */
+	static const uint8 MAX_PING_VALUES_SIZE = 7;
+
+	/** Default constructor */
+	PingAvgDataV2()
+		: AvgPingV2(0.0f)
+	{
+		for (int32 i = 0; i < MAX_PING_VALUES_SIZE; i++)
+		{
+			PingValues.Add(MAX_uint16);
+		}
+	}
+
+	/** The average of the values in PingValues, calculated after 1s. */
+	float AvgPingV2;
+};
+
 /**
  * A PlayerState is created for every player on a server (or in a standalone game).
  * PlayerStates are replicated to all clients, and contain network game relevant information about the player, such as playername, score, etc.
@@ -117,6 +142,7 @@ public:
 
 	/** Exact ping as float (rounded and compressed in replicated Ping) */
 	float ExactPing;
+	float ExactPingV2;
 
 	/** Used to match up InactivePlayerState with rejoining playercontroller. */
 	UPROPERTY()
@@ -146,6 +172,7 @@ private:
 	 * without using up a lot of space, while also being tolerant of changes in ping update frequency
 	 */
 	PingAvgData		PingBucket[4];
+	PingAvgDataV2	PingBucketV2[4];
 
 	/** The timestamp for when the current PingBucket began filling */
 	float			CurPingBucketTimestamp;
diff --git a/Engine/Source/Runtime/Engine/Private/PlayerState.cpp b/Engine/Source/Runtime/Engine/Private/PlayerState.cpp
--- a/Engine/Source/Runtime/Engine/Private/PlayerState.cpp
+++ b/Engine/Source/Runtime/Engine/Private/PlayerState.cpp
@@ -44,11 +44,15 @@ PRAGMA_ENABLE_DEPRECATION_WARNINGS
 
 void APlayerState::UpdatePing(float InPing)
 {
+	QUICK_SCOPE_CYCLE_COUNTER(STAT_PlayerState_UpdatePing);
+
 	// Limit the size of the ping, to avoid overflowing PingBucket values
 	InPing = FMath::Min(1.1f, InPing);
 
 	float CurTime = GetWorld()->RealTimeSeconds;
 
+	float InPingInMs = InPing * 1000.f;
+
 	if ((CurTime - CurPingBucketTimestamp) >= 1.f)
 	{
 		// Trigger ping recalculation now, while all buckets are 'full'
@@ -59,15 +63,36 @@ void APlayerState::UpdatePing(float InPing)
 		CurPingBucketTimestamp = CurTime;
 
 
-		PingBucket[CurPingBucket].PingSum = FMath::FloorToInt(InPing * 1000.f);
+		PingBucket[CurPingBucket].PingSum = FMath::FloorToInt(InPingInMs);
 		PingBucket[CurPingBucket].PingCount = 1;
+
+		PingBucketV2[CurPingBucket] = PingAvgDataV2();
 	}
 	// Limit the number of pings we accept per-bucket, to avoid overflowing PingBucket values
 	else if (PingBucket[CurPingBucket].PingCount < 7)
 	{
-		PingBucket[CurPingBucket].PingSum += FMath::FloorToInt(InPing * 1000.f);
+		PingBucket[CurPingBucket].PingSum += FMath::FloorToInt(InPingInMs);
 		PingBucket[CurPingBucket].PingCount++;
 	}
+
+	TArray<uint16>& CurrentBucketPingValues = PingBucketV2[CurPingBucket].PingValues;
+
+	// This makes sure we will actually add the ping value to the list, since much of the time the new ping value will be higher than
+	// what is already there.
+	if (InPingInMs < CurrentBucketPingValues[PingAvgDataV2::MAX_PING_VALUES_SIZE - 1])
+	{
+		for (int32 i = 0; i < PingAvgDataV2::MAX_PING_VALUES_SIZE; i++)
+		{
+			// This will keep the list of ping values we're currently using sorted - this will insert
+			// a ping value if it is less than any in current list and remove the new max value.
+			if (InPingInMs < CurrentBucketPingValues[i])
+			{
+				CurrentBucketPingValues.Insert(FMath::FloorToInt(InPingInMs), i);
+				CurrentBucketPingValues.RemoveAt(PingAvgDataV2::MAX_PING_VALUES_SIZE);
+				break;
+			}
+		}
+	}
 }
 
 void APlayerState::RecalculateAvgPing()
@@ -81,6 +106,27 @@ void APlayerState::RecalculateAvgPing()
 		Count += PingBucket[i].PingCount;
 	}
 
+	int32 SumV2 = 0;
+	int32 NumValidValues = 0;
+	for (; NumValidValues < PingAvgDataV2::MAX_PING_VALUES_SIZE; NumValidValues++)
+	{
+		if (PingBucketV2[CurPingBucket].PingValues[NumValidValues] != MAX_uint16)
+		{
+			SumV2 += PingBucketV2[CurPingBucket].PingValues[NumValidValues];
+		}
+	}
+
+	// Use NumValidValues instead of MAX_PING_VALUES_SIZE in case there are fewer valid values.
+	PingBucketV2[CurPingBucket].AvgPingV2 = NumValidValues > 0 ? SumV2 / NumValidValues : MAX_flt;
+
+	float AvgSumV2 = 0.0f;
+	for (int32 i = 0; i < ARRAY_COUNT(PingBucketV2); i++)
+	{
+		AvgSumV2 += PingBucketV2[i].AvgPingV2;
+	}
+
+	ExactPingV2 = AvgSumV2 / ARRAY_COUNT(PingBucketV2);
+
 	// Calculate the average, and divide it by 4 to optimize replication
 	ExactPing = (Count > 0 ? ((float)Sum / (float)Count) : 0.f);
 
