Fix issue with curve look up tables when curve name has been removed in editor

diff --git a/Engine/Source/Runtime/Engine/Classes/Animation/SmartName.h b/Engine/Source/Runtime/Engine/Classes/Animation/SmartName.h
--- a/Engine/Source/Runtime/Engine/Classes/Animation/SmartName.h
+++ b/Engine/Source/Runtime/Engine/Classes/Animation/SmartName.h
@@ -81,10 +81,17 @@ struct ENGINE_API FSmartNameMapping
 	// @param Array - Array to fill
 	void FillNameArray(TArray<FName>& Array) const;
 
+	// Fill an array with curve names in UID order (None will be placed in invalid UID slots)
+	void FillUIDToNameArray(TArray<FName>& Array) const;
+
 	// Fill an array with curve types for all used names
 	// @param Array - Array to fill
 	void FillCurveTypeArray(TArray<FAnimCurveType>& Array) const;
 
+
+	// Fill an array with curve types in UID order
+	void FillUIDToCurveTypeArray(TArray<FAnimCurveType>& Array) const;
+
 #if WITH_EDITOR
 	// Change a name
 	// @param Uid - SmartName::UID_Type of the name to change
diff --git a/Engine/Source/Runtime/Engine/Private/Animation/SmartName.cpp b/Engine/Source/Runtime/Engine/Private/Animation/SmartName.cpp
--- a/Engine/Source/Runtime/Engine/Private/Animation/SmartName.cpp
+++ b/Engine/Source/Runtime/Engine/Private/Animation/SmartName.cpp
@@ -152,6 +152,11 @@ void FSmartNameMapping::FillNameArray(TArray<FName>& Array) const
 #endif
 }
 
+void FSmartNameMapping::FillUIDToNameArray(TArray<FName>& Array) const
+{
+	Array = CurveNameList;
+}
+
 void FSmartNameMapping::FillCurveTypeArray(TArray<FAnimCurveType>& Array) const
 {
 	Array.Reset(CurveNameList.Num());
@@ -173,6 +178,32 @@ void FSmartNameMapping::FillCurveTypeArray(TArray<FAnimCurveType>& Array) const
 	}
 }
 
+void FSmartNameMapping::FillUIDToCurveTypeArray(TArray<FAnimCurveType>& Array) const
+{
+	Array.Reset(CurveNameList.Num());
+
+	for (const FName& Name : CurveNameList)
+	{
+		const FCurveMetaData* MetaData = CurveMetaDataMap.Find(Name);
+		FAnimCurveType AnimCurveType = MetaData ? MetaData->Type : FAnimCurveType();
+
+		//In editor names can be removed and so have to deal with empty slots
+#if WITH_EDITOR
+		if (Name != NAME_None)
+		{
+			Array.Add(AnimCurveType);
+		}
+		else
+		{
+			Array.AddDefaulted();
+		}
+#else
+		Array.Add(AnimCurveType);
+#endif
+	}
+}
+
+
 bool FSmartNameMapping::Exists(const SmartName::UID_Type& Uid) const
 {
 	return CurveNameList.IsValidIndex(Uid) && CurveNameList[Uid] != NAME_None;
diff --git a/Engine/Source/Runtime/Engine/Private/BoneContainer.cpp b/Engine/Source/Runtime/Engine/Private/BoneContainer.cpp
--- a/Engine/Source/Runtime/Engine/Private/BoneContainer.cpp
+++ b/Engine/Source/Runtime/Engine/Private/BoneContainer.cpp
@@ -215,14 +215,10 @@ void FBoneContainer::CacheRequiredAnimCurveUids(const FCurveEvaluationOption& Cu
 			}
 
 			// Get Current Names / UIDs
-			Mapping->FillNameArray(UIDToNameLUT);
+			Mapping->FillUIDToNameArray(UIDToNameLUT);
 
 			// Get curve types
-			Mapping->FillCurveTypeArray(UIDToCurveTypeLUT);
-
-			// Get UID List
-			TArray<SmartName::UID_Type> UIDList;
-			Mapping->FillUidArray(UIDList);
+			Mapping->FillUIDToCurveTypeArray(UIDToCurveTypeLUT);
 
 			// if the linked joints don't exists in RequiredBones, remove itself
 			if (UIDToNameLUT.Num() > 0)
@@ -239,7 +235,11 @@ void FBoneContainer::CacheRequiredAnimCurveUids(const FCurveEvaluationOption& Cu
 					else
 					{
 						// CurveNameIndex shouyld match to UID
-						if (CurveEvalOption.DisallowedList && CurveEvalOption.DisallowedList->Contains(CurveName))
+						if (CurveName == NAME_None)
+						{
+							bBeingUsed = false;
+						}
+						else if (CurveEvalOption.DisallowedList && CurveEvalOption.DisallowedList->Contains(CurveName))
 						{
 							//remove the UID
 							bBeingUsed = false;
@@ -280,7 +280,7 @@ void FBoneContainer::CacheRequiredAnimCurveUids(const FCurveEvaluationOption& Cu
 
 					if (bBeingUsed)
 					{
-						UIDToArrayIndexLUT[UIDList[CurveNameIndex]] = NumAvailableUIDs++;
+						UIDToArrayIndexLUT[CurveNameIndex] = NumAvailableUIDs++;
 					}
 				}
 			}
