Fix for "Calling StaticLoadObject during PostLoad of AnimSequence"

#Jiira UE-68440

diff --git a/Engine/Source/Runtime/Engine/Classes/Animation/AnimCurveCompressionCodec.h b/Engine/Source/Runtime/Engine/Classes/Animation/AnimCurveCompressionCodec.h
--- a/Engine/Source/Runtime/Engine/Classes/Animation/AnimCurveCompressionCodec.h
+++ b/Engine/Source/Runtime/Engine/Classes/Animation/AnimCurveCompressionCodec.h
@@ -52,6 +52,8 @@ class ENGINE_API UAnimCurveCompressionCodec : public UObject
 	/** Compresses the curve data from an animation sequence. */
 	virtual bool Compress(const UAnimSequence& AnimSeq, FAnimCurveCompressionResult& OutResult) PURE_VIRTUAL(UAnimCurveCompressionCodec::Compress, return false;);
 
+	/** Allow us to convert DDC serialized path back into codec object */
+	virtual UAnimCurveCompressionCodec* GetCodec(const FString& Path) { return this; }
 	/*
 	 * Called to generate a unique DDC key for this codec instance.
 	 * A suitable key should be generated from: the InstanceGuid, a codec version, and all relevant properties that drive the behavior.
diff --git a/Engine/Source/Runtime/Engine/Classes/Animation/AnimCurveCompressionSettings.h b/Engine/Source/Runtime/Engine/Classes/Animation/AnimCurveCompressionSettings.h
--- a/Engine/Source/Runtime/Engine/Classes/Animation/AnimCurveCompressionSettings.h
+++ b/Engine/Source/Runtime/Engine/Classes/Animation/AnimCurveCompressionSettings.h
@@ -39,6 +39,9 @@ class ENGINE_API UAnimCurveCompressionSettings : public UObject
 	 */
 	bool Compress(UAnimSequence& AnimSeq) const;
 
+	/** Allow us to convert DDC serialized path back into codec object */
+	UAnimCurveCompressionCodec* GetCodec(const FString& Path);
+
 	/** Generates a DDC key that takes into account the current settings and selected codec. */
 	FString MakeDDCKey() const;
 #endif
diff --git a/Engine/Source/Runtime/Engine/Private/Animation/AnimCurveCompressionSettings.cpp b/Engine/Source/Runtime/Engine/Private/Animation/AnimCurveCompressionSettings.cpp
--- a/Engine/Source/Runtime/Engine/Private/Animation/AnimCurveCompressionSettings.cpp
+++ b/Engine/Source/Runtime/Engine/Private/Animation/AnimCurveCompressionSettings.cpp
@@ -45,6 +45,11 @@ bool UAnimCurveCompressionSettings::Compress(UAnimSequence& AnimSeq) const
 	return Success;
 }
 
+UAnimCurveCompressionCodec* UAnimCurveCompressionSettings::GetCodec(const FString& Path)
+{
+	return Codec->GetCodec(Path);
+}
+
 FString UAnimCurveCompressionSettings::MakeDDCKey() const
 {
 	if (Codec == nullptr)
diff --git a/Engine/Source/Runtime/Engine/Private/Animation/AnimSequence.cpp b/Engine/Source/Runtime/Engine/Private/Animation/AnimSequence.cpp
--- a/Engine/Source/Runtime/Engine/Private/Animation/AnimSequence.cpp
+++ b/Engine/Source/Runtime/Engine/Private/Animation/AnimSequence.cpp
@@ -2623,7 +2623,7 @@ void UAnimSequence::SerializeCompressedData(FArchive& Ar, bool bDDCData)
 			FString CurveCodecPath;
 			Ar << CurveCodecPath;
 
-			CurveCompressionCodec = LoadObject<UAnimCurveCompressionCodec>(nullptr, *CurveCodecPath);
+			CurveCompressionCodec = CurveCompressionSettings->GetCodec(CurveCodecPath);
 		}
 		else
 #else
