Clamp delta time value used for updating UNetConnection::QueuedBits to prevent the case where a hitch would add a lot of extra bandwidth headroom for subsequent frames. If there are a lot of actors to replicate, this can quickly cause bandwidth and CPU usage to get out of hand.

diff --git a/Engine/Source/Runtime/Engine/Private/NetConnection.cpp b/Engine/Source/Runtime/Engine/Private/NetConnection.cpp
--- a/Engine/Source/Runtime/Engine/Private/NetConnection.cpp
+++ b/Engine/Source/Runtime/Engine/Private/NetConnection.cpp
@@ -2683,7 +2683,17 @@ void UNetConnection::Tick()
 
 	// Update queued byte count.
 	// this should be at the end so that the cap is applied *after* sending (and adjusting QueuedBytes for) any remaining data for this tick
-	float DeltaBits = CurrentNetSpeed * DeltaTime * 8.f;
+
+	// Clamp DeltaTime for bandwidth limiting so that if there is a hitch, we don't try to send
+	// a large burst on the next frame, which can cause another hitch if a lot of additional replication occurs.
+	const float DesiredTickRate = GEngine->GetMaxTickRate(0.0f, false);
+	float BandwidthDeltaTime = DeltaTime;
+	if (DesiredTickRate != 0.0f)
+	{
+		BandwidthDeltaTime = FMath::Clamp(BandwidthDeltaTime, 0.0f, 1.0f / DesiredTickRate);
+	}
+
+	float DeltaBits = CurrentNetSpeed * BandwidthDeltaTime * 8.f;
 	QueuedBits -= FMath::TruncToInt(DeltaBits);
 	float AllowedLag = 2.f * DeltaBits;
 	if (QueuedBits < -AllowedLag)
