[NO NON-PII COMMIT MESSAGE]

diff --git a/Engine/Source/Runtime/Engine/Private/SoundWave.cpp b/Engine/Source/Runtime/Engine/Private/SoundWave.cpp
--- a/Engine/Source/Runtime/Engine/Private/SoundWave.cpp
+++ b/Engine/Source/Runtime/Engine/Private/SoundWave.cpp
@@ -24,6 +24,14 @@
 #include "DSP/EnvelopeFollower.h"
 #include "DSP/BufferVectorOperations.h"
 
+static int32 BypassVirtualizeWhenSilentCVar = 0;
+FAutoConsoleVariableRef CVarBypassVirtualizeWhenSilent(
+	TEXT("au.BypassVirtualizeWhenSilent"),
+	BypassVirtualizeWhenSilentCVar,
+	TEXT("When set to 1, ignores the Play When Silent flag for non-procedural sources.\n")
+	TEXT("0: Honor the Play When Silent flag, 1: stop all silent non-procedural sources."),
+	ECVF_Default);
+
 #if ENABLE_COOK_STATS
 namespace SoundWaveCookStats
 {
@@ -1399,7 +1407,10 @@ void USoundWave::Parse( FAudioDevice* AudioDevice, const UPTRINT NodeWaveInstanc
 		// Recompute the virtualizability here even though we did it up-front in the active sound parse.
 		// This is because an active sound can generate multiple sound waves, not all of them are necessarily virtualizable.
 		bool bHasSubtitles = ActiveSound.bHandleSubtitles && (ActiveSound.bHasExternalSubtitles || (Subtitles.Num() > 0));
-		if (WaveInstance->GetVolumeWithDistanceAttenuation() > KINDA_SMALL_NUMBER || ((bVirtualizeWhenSilent || bHasSubtitles) && AudioDevice->VirtualSoundsEnabled()))
+
+		// When the BypassVirtualizeWhenSilent cvar is enabled, we should only honor bVirtualizeWhenSilent for procedural sounds:
+		const bool bShouldVirtualize = bVirtualizeWhenSilent && (!BypassVirtualizeWhenSilentCVar || bProcedural);
+		if (WaveInstance->GetVolumeWithDistanceAttenuation() > KINDA_SMALL_NUMBER || ((bShouldVirtualize || bHasSubtitles) && AudioDevice->VirtualSoundsEnabled()))
 		{
 			bAddedWaveInstance = true;
 			WaveInstances.Add(WaveInstance);
