Fixup ensure that was hitting in replays due to pulling the check for NetFieldExportGroup before reading handles.

diff --git a/Engine/Source/Runtime/Engine/Private/RepLayout.cpp b/Engine/Source/Runtime/Engine/Private/RepLayout.cpp
--- a/Engine/Source/Runtime/Engine/Private/RepLayout.cpp
+++ b/Engine/Source/Runtime/Engine/Private/RepLayout.cpp
@@ -2437,26 +2437,45 @@ bool FRepLayout::ReceiveProperties_BackwardsCompatible_r(
 	bool&					bOutHasUnmapped,
 	bool&					bOutGuidsChanged) const
 {
-	if (!ensure(NetFieldExportGroup != nullptr))
+	auto ReadHandle = [this, &Reader](uint32& Handle) -> bool
 	{
-		uint32 NetFieldExportHandle = 0;
-		Reader.SerializeIntPacked(NetFieldExportHandle);
+		Reader.SerializeIntPacked(Handle);
 
-		UE_CLOG(!FApp::IsUnattended(), LogRep, Warning, TEXT("ReceiveProperties_BackwardsCompatible_r: NetFieldExportGroup == nullptr. Owner: %s, NetFieldExportHandle: %u"), *Owner->GetName(), NetFieldExportHandle);
-		Reader.SetError();
-		return false;
+		if (Reader.IsError())
+		{
+			UE_LOG(LogRep, Warning, TEXT("ReceiveProperties_BackwardsCompatible_r: Error reading handle. Owner: %s"), *Owner->GetName());
+			return false;
+		}
+
+		UE_LOG(LogRepProperties, VeryVerbose, TEXT("ReceiveProperties_BackwardsCompatible_r: NetFieldExportHandle=%d"), Handle);
+		return true;
+	};
+
+	if (NetFieldExportGroup == nullptr)
+	{
+		uint32 NetFieldExportHandle = 0;
+		if (!ReadHandle(NetFieldExportHandle))
+		{
+			return false;
+		}
+		else if (NetFieldExportHandle != 0)
+		{
+			UE_CLOG(!FApp::IsUnattended(), LogRep, Warning, TEXT("ReceiveProperties_BackwardsCompatible_r: NetFieldExportGroup == nullptr. Owner: %s, NetFieldExportHandle: %u"), *Owner->GetName(), NetFieldExportHandle);
+			Reader.SetError();
+			ensure(false);
+			return false;
+		}
+		else
+		{
+			return true;
+		}
 	}
 
 	while (true)
 	{
 		uint32 NetFieldExportHandle = 0;
-		Reader.SerializeIntPacked(NetFieldExportHandle);
-
-		UE_LOG(LogRepProperties, VeryVerbose, TEXT("ReceiveProperties_BackwardsCompatible_r: NetFieldExportHandle=%d"), NetFieldExportHandle);
-
-		if (Reader.IsError())
+		if (!ReadHandle(NetFieldExportHandle))
 		{
-			UE_LOG(LogRep, Warning, TEXT("ReceiveProperties_BackwardsCompatible_r: Error reading handle. Owner: %s"), *Owner->GetName());
 			return false;
 		}
 
