PR #5150: Fix UAnimSequence::CopyNotifies & expose to Animation Blueprint Library (Contributed by [REDACTED_PERSON])

diff --git a/Engine/Source/Editor/AnimationModifiers/Private/AnimationBlueprintLibrary.cpp b/Engine/Source/Editor/AnimationModifiers/Private/AnimationBlueprintLibrary.cpp
--- a/Engine/Source/Editor/AnimationModifiers/Private/AnimationBlueprintLibrary.cpp
+++ b/Engine/Source/Editor/AnimationModifiers/Private/AnimationBlueprintLibrary.cpp
@@ -714,6 +714,27 @@ void UAnimationBlueprintLibrary::ReplaceAnimNotifies(UAnimSequenceBase* Animatio
 	ReplaceAnimNotifies_Helper(AnimationSequence, OldNotifyClass.Get(), NewNotifyClass.Get(), OnNotifyReplaced, FOnNotifyStateReplaced());
 }
 
+void UAnimationBlueprintLibrary::CopyAnimNotifiesFromSequence(UAnimSequence* SrcAnimSequence, UAnimSequence* DestAnimSequence)
+{
+	if (SrcAnimSequence == nullptr)
+	{
+		UE_LOG(LogAnimationBlueprintLibrary, Warning, TEXT("Invalid Source Animation Sequence for CopyAnimNotifiesFromSequence"));
+	}
+	else if (DestAnimSequence == nullptr)
+	{
+		UE_LOG(LogAnimationBlueprintLibrary, Warning, TEXT("Invalid Destination Animation Sequence for CopyAnimNotifiesFromSequence"));
+	}
+	else if (SrcAnimSequence == DestAnimSequence)
+	{
+		UE_LOG(LogAnimationBlueprintLibrary, Warning, TEXT("Source and Destination Animation Sequence are the same for CopyAnimNotifiesFromSequence"));
+	}
+	else
+	{
+		const bool bShowDialogs = false;
+		UAnimSequence::CopyNotifies(SrcAnimSequence, DestAnimSequence, bShowDialogs);
+	}
+}
+
 int32 UAnimationBlueprintLibrary::RemoveAnimationNotifyEventsByName(UAnimSequence* AnimationSequence, FName NotifyName)
 {
 	int32 NumRemovedEvents = 0;
diff --git a/Engine/Source/Editor/AnimationModifiers/Public/AnimationBlueprintLibrary.h b/Engine/Source/Editor/AnimationModifiers/Public/AnimationBlueprintLibrary.h
--- a/Engine/Source/Editor/AnimationModifiers/Public/AnimationBlueprintLibrary.h
+++ b/Engine/Source/Editor/AnimationModifiers/Public/AnimationBlueprintLibrary.h
@@ -189,6 +189,9 @@ public:
 	UFUNCTION(BlueprintCallable, Category = "AnimationBlueprintLibrary|NotifyEvents")
 	static void ReplaceAnimNotifies(UAnimSequenceBase* AnimationSequence, TSubclassOf<UAnimNotify> OldNotifyClass, TSubclassOf<UAnimNotify> NewNotifyClass, FOnNotifyReplaced OnNotifyReplaced);
 
+	/** Copies animation notifies from Src Animation Sequence to Dest. Creates anim notify tracks as necessary. Returns true on success. */
+	UFUNCTION(BlueprintCallable, Category = "AnimationBlueprintLibrary|NotifyEvents")
+	static void CopyAnimNotifiesFromSequence(UAnimSequence* SrcAnimSequence, UAnimSequence* DestAnimSequence);
 
 	// Notify Tracks
 
diff --git a/Engine/Source/Runtime/Engine/Private/Animation/AnimSequence.cpp b/Engine/Source/Runtime/Engine/Private/Animation/AnimSequence.cpp
--- a/Engine/Source/Runtime/Engine/Private/Animation/AnimSequence.cpp
+++ b/Engine/Source/Runtime/Engine/Private/Animation/AnimSequence.cpp
@@ -3185,7 +3185,7 @@ bool UAnimSequence::CopyAnimSequenceProperties(UAnimSequence* SourceAnimSeq, UAn
 }
 
 
-bool UAnimSequence::CopyNotifies(UAnimSequence* SourceAnimSeq, UAnimSequence* DestAnimSeq)
+bool UAnimSequence::CopyNotifies(UAnimSequence* SourceAnimSeq, UAnimSequence* DestAnimSeq, bool bShowDialogs /*= true */)
 {
 #if WITH_EDITOR
 	// Abort if source == destination.
@@ -3197,13 +3197,13 @@ bool UAnimSequence::CopyNotifies(UAnimSequence* SourceAnimSeq, UAnimSequence* De
 	// If the destination sequence is shorter than the source sequence, we'll be dropping notifies that
 	// occur at later times than the dest sequence is long.  Give the user a chance to abort if we
 	// find any notifies that won't be copied over.
-	if( DestAnimSeq->SequenceLength < SourceAnimSeq->SequenceLength )
+	if(bShowDialogs && DestAnimSeq->SequenceLength < SourceAnimSeq->SequenceLength)
 	{
 		for(int32 NotifyIndex=0; NotifyIndex<SourceAnimSeq->Notifies.Num(); ++NotifyIndex)
 		{
 			// If a notify is found which occurs off the end of the destination sequence, prompt the user to continue.
 			const FAnimNotifyEvent& SrcNotifyEvent = SourceAnimSeq->Notifies[NotifyIndex];
-			if( SrcNotifyEvent.DisplayTime_DEPRECATED > DestAnimSeq->SequenceLength )
+			if( SrcNotifyEvent.GetTriggerTime() > DestAnimSeq->SequenceLength )
 			{
 				const bool bProceed = EAppReturnType::Yes == FMessageDialog::Open( EAppMsgType::YesNo, NSLOCTEXT("UnrealEd", "SomeNotifiesWillNotBeCopiedQ", "Some notifies will not be copied because the destination sequence is not long enough.  Proceed?") );
 				if( !bProceed )
@@ -3220,7 +3220,7 @@ bool UAnimSequence::CopyNotifies(UAnimSequence* SourceAnimSeq, UAnimSequence* De
 
 	// If the destination sequence contains any notifies, ask the user if they'd like
 	// to delete the existing notifies before copying over from the source sequence.
-	if( DestAnimSeq->Notifies.Num() > 0 )
+	if(bShowDialogs && DestAnimSeq->Notifies.Num() > 0)
 	{
 		const bool bDeleteExistingNotifies = EAppReturnType::Yes == FMessageDialog::Open(EAppMsgType::YesNo, FText::Format(
 			NSLOCTEXT("UnrealEd", "DestSeqAlreadyContainsNotifiesMergeQ", "The destination sequence already contains {0} notifies.  Delete these before copying?"), FText::AsNumber(DestAnimSeq->Notifies.Num())) );
@@ -3232,7 +3232,6 @@ bool UAnimSequence::CopyNotifies(UAnimSequence* SourceAnimSeq, UAnimSequence* De
 	}
 
 	// Do the copy.
-	TArray<int32> NewNotifyIndices;
 	int32 NumNotifiesThatWereNotCopied = 0;
 
 	for(int32 NotifyIndex=0; NotifyIndex<SourceAnimSeq->Notifies.Num(); ++NotifyIndex)
@@ -3240,32 +3239,29 @@ bool UAnimSequence::CopyNotifies(UAnimSequence* SourceAnimSeq, UAnimSequence* De
 		const FAnimNotifyEvent& SrcNotifyEvent = SourceAnimSeq->Notifies[NotifyIndex];
 
 		// Skip notifies which occur at times later than the destination sequence is long.
-		if( SrcNotifyEvent.DisplayTime_DEPRECATED > DestAnimSeq->SequenceLength )
+		if( SrcNotifyEvent.GetTriggerTime() > DestAnimSeq->SequenceLength )
 		{
+			++NumNotifiesThatWereNotCopied;
 			continue;
 		}
 
-		// Do a linear-search through existing notifies to determine where
-		// to insert the new notify.
-		int32 NewNotifyIndex = 0;
-		while( NewNotifyIndex < DestAnimSeq->Notifies.Num()
-			&& DestAnimSeq->Notifies[NewNotifyIndex].DisplayTime_DEPRECATED <= SrcNotifyEvent.DisplayTime_DEPRECATED )
+		// Copy notify tracks from src to dest if they are missing
+		if (SrcNotifyEvent.TrackIndex >= DestAnimSeq->AnimNotifyTracks.Num())
 		{
-			++NewNotifyIndex;
+			for (int32 TrackIndex = DestAnimSeq->AnimNotifyTracks.Num(); TrackIndex <= SrcNotifyEvent.TrackIndex; ++TrackIndex)
+			{
+				DestAnimSeq->AnimNotifyTracks.Add(FAnimNotifyTrack(SourceAnimSeq->AnimNotifyTracks[TrackIndex].TrackName, SourceAnimSeq->AnimNotifyTracks[TrackIndex].TrackColor));
+			}
 		}
 
 		// Track the location of the new notify.
-		NewNotifyIndices.Add(NewNotifyIndex);
-
-		// Create a new empty on in the array.
-		DestAnimSeq->Notifies.InsertZeroed(NewNotifyIndex);
-
-		// Copy time and comment.
-		FAnimNotifyEvent& Notify = DestAnimSeq->Notifies[NewNotifyIndex];
-		Notify.DisplayTime_DEPRECATED = SrcNotifyEvent.DisplayTime_DEPRECATED;
-		Notify.TriggerTimeOffset = GetTriggerTimeOffsetForType( DestAnimSeq->CalculateOffsetForNotify(Notify.DisplayTime_DEPRECATED));
-		Notify.NotifyName = SrcNotifyEvent.NotifyName;
-		Notify.Duration = SrcNotifyEvent.Duration;
+		int32 NewNotifyIndex = DestAnimSeq->Notifies.AddDefaulted();
+		FAnimNotifyEvent& NotifyEvent = DestAnimSeq->Notifies[NewNotifyIndex];
+		
+		// Copy properties of the NotifyEvent
+		NotifyEvent.TrackIndex = SrcNotifyEvent.TrackIndex;
+		NotifyEvent.NotifyName = SrcNotifyEvent.NotifyName;
+		NotifyEvent.Duration = SrcNotifyEvent.Duration;
 
 		// Copy the notify itself, and point the new one at it.
 		if( SrcNotifyEvent.Notify )
@@ -3274,7 +3270,7 @@ bool UAnimSequence::CopyNotifies(UAnimSequence* SourceAnimSeq, UAnimSequence* De
 		}
 		else
 		{
-			DestAnimSeq->Notifies[NewNotifyIndex].Notify = NULL;
+			DestAnimSeq->Notifies[NewNotifyIndex].Notify = nullptr;
 		}
 
 		if( SrcNotifyEvent.NotifyStateClass )
@@ -3283,18 +3279,23 @@ bool UAnimSequence::CopyNotifies(UAnimSequence* SourceAnimSeq, UAnimSequence* De
 		}
 		else
 		{
-			DestAnimSeq->Notifies[NewNotifyIndex].NotifyStateClass = NULL;
+			DestAnimSeq->Notifies[NewNotifyIndex].NotifyStateClass = nullptr;
 		}
+		
+		// Copy notify timing
+		NotifyEvent.LinkSequence(DestAnimSeq, SrcNotifyEvent.GetTriggerTime());
+		NotifyEvent.TriggerTimeOffset = GetTriggerTimeOffsetForType(DestAnimSeq->CalculateOffsetForNotify(NotifyEvent.GetTriggerTime()));
 
 		// Make sure editor knows we've changed something.
 		DestAnimSeq->MarkPackageDirty();
+		DestAnimSeq->RefreshCacheData();
 	}
 
 	// Inform the user if some notifies weren't copied.
-	if( SourceAnimSeq->Notifies.Num() > NewNotifyIndices.Num() )
+	if(bShowDialogs && NumNotifiesThatWereNotCopied > 0)
 	{
 		FMessageDialog::Open( EAppMsgType::Ok, FText::Format(
-			NSLOCTEXT("UnrealEd", "SomeNotifiesWereNotCopiedF", "Because the destination sequence was shorter, {0} notifies were not copied."), FText::AsNumber(SourceAnimSeq->Notifies.Num() - NewNotifyIndices.Num())) );
+			NSLOCTEXT("UnrealEd", "SomeNotifiesWereNotCopiedF", "Because the destination sequence was shorter, {0} notifies were not copied."), FText::AsNumber(NumNotifiesThatWereNotCopied)) );
 	}
 #endif // WITH_EDITOR
 
