CIS fix - missing out file in CL

diff --git a/Engine/Source/Runtime/Engine/Classes/Animation/AnimSequence.h b/Engine/Source/Runtime/Engine/Classes/Animation/AnimSequence.h
--- a/Engine/Source/Runtime/Engine/Classes/Animation/AnimSequence.h
+++ b/Engine/Source/Runtime/Engine/Classes/Animation/AnimSequence.h
@@ -939,7 +939,7 @@ public:
 	/** 
 	 * Copy AnimNotifies from one UAnimSequence to another.
 	 */
-	static bool CopyNotifies(UAnimSequence* SourceAnimSeq, UAnimSequence* DestAnimSeq);
+	static bool CopyNotifies(UAnimSequence* SourceAnimSeq, UAnimSequence* DestAnimSeq, bool bShowDialogs = true);
 
 	/**
 	 * Flip Rotation's W For NonRoot items, and compress it again if SkelMesh exists
