Fixup null dereference when checking whether or not to send a client camera update in APlayerController::Tick.
Add some comments to that bit of code.

diff --git a/Engine/Source/Runtime/Engine/Private/PlayerController.cpp b/Engine/Source/Runtime/Engine/Private/PlayerController.cpp
--- a/Engine/Source/Runtime/Engine/Private/PlayerController.cpp
+++ b/Engine/Source/Runtime/Engine/Private/PlayerController.cpp
@@ -4503,10 +4503,13 @@ void APlayerController::TickActor( float DeltaSeconds, ELevelTick TickType, FAct
 			}
 
 			// Send a camera update if necessary.
+			// That position will be used as the base for replication
+			// (i.e., the origin that will be used when calculating NetCullDistance for other Actors / Objects).
+			// We only do this when the Pawn will move, to prevent spamming RPCs.
 			if (bIsClient && bIsLocallyControlled && GetPawn() && PlayerCameraManager->bUseClientSideCameraUpdates)
 			{
 				UPawnMovementComponent* PawnMovement = GetPawn()->GetMovementComponent();
-				if (!PawnMovement->IsMoveInputIgnored() &&
+				if (PawnMovement != nullptr && !PawnMovement->IsMoveInputIgnored() &&
 					(PawnMovement->GetLastInputVector() != FVector::ZeroVector || PawnMovement->Velocity != FVector::ZeroVector))
 				{
 					PlayerCameraManager->bShouldSendClientSideCameraUpdate = true;
