Fixup issue where moving an attached Actor on the Server wouldn't update it's position on the client.
Note, this only occurs for Actors with bReplicateMovement whose RootComponents *aren't* replicated.

https://udn.unrealengine.com/questions/462120/view.html

diff --git a/Engine/Source/Runtime/Engine/Private/Actor.cpp b/Engine/Source/Runtime/Engine/Private/Actor.cpp
--- a/Engine/Source/Runtime/Engine/Private/Actor.cpp
+++ b/Engine/Source/Runtime/Engine/Private/Actor.cpp
@@ -1569,7 +1569,21 @@ void AActor::OnRep_AttachmentReplication()
 				RootComponent->RelativeLocation = AttachmentReplication.LocationOffset;
 				RootComponent->RelativeRotation = AttachmentReplication.RotationOffset;
 				RootComponent->RelativeScale3D = AttachmentReplication.RelativeScale3D;
-				RootComponent->AttachToComponent(AttachParentComponent, FAttachmentTransformRules::KeepRelativeTransform,  AttachmentReplication.AttachSocket);
+
+				// If we're already attached to the correct Parent and Socket, then the update must be position only.
+				// AttachToComponent would early out in this case.
+				// Note, we ignore the special case for simulated bodies in AttachToComponent as AttachmentReplication shouldn't get updated
+				// if the body is simulated (see AActor::GatherMovement).
+				const bool bAlreadyAttached = (AttachParentComponent == RootComponent->GetAttachParent() && AttachmentReplication.AttachSocket == RootComponent->GetAttachSocketName() && AttachParentComponent->GetAttachChildren().Contains(RootComponent));
+				if (bAlreadyAttached)
+				{
+					// Note, this doesn't match AttachToComponent, but we're assuming it's safe to skip physics (see comment above).
+					RootComponent->UpdateComponentToWorld(EUpdateTransformFlags::SkipPhysicsUpdate, ETeleportType::None);
+				}
+				else
+				{
+					RootComponent->AttachToComponent(AttachParentComponent, FAttachmentTransformRules::KeepRelativeTransform, AttachmentReplication.AttachSocket);
+				}
 			}
 		}
 	}
