Fixed issue where the position of a replicated StaticMeshActor with bRepMovement set to false would be moved to default locataion on client after being deattached.

diff --git a/Engine/Source/Runtime/Engine/Private/Actor.cpp b/Engine/Source/Runtime/Engine/Private/Actor.cpp
--- a/Engine/Source/Runtime/Engine/Private/Actor.cpp
+++ b/Engine/Source/Runtime/Engine/Private/Actor.cpp
@@ -1594,7 +1594,10 @@ void AActor::OnRep_AttachmentReplication()
 		// Handle the case where an object was both detached and moved on the server in the same frame.
 		// Calling this extraneously does not hurt but will properly fire events if the movement state changed while attached.
 		// This is needed because client side movement is ignored when attached
-		OnRep_ReplicatedMovement();
+		if (bReplicateMovement)
+		{
+			OnRep_ReplicatedMovement();
+		}
 	}
 }
 
