PR #5208: Static meshes would not replicate when setting to none (Contributed by DSDambuster)

diff --git a/Engine/Source/Runtime/Engine/Private/Components/StaticMeshComponent.cpp b/Engine/Source/Runtime/Engine/Private/Components/StaticMeshComponent.cpp
--- a/Engine/Source/Runtime/Engine/Private/Components/StaticMeshComponent.cpp
+++ b/Engine/Source/Runtime/Engine/Private/Components/StaticMeshComponent.cpp
@@ -247,7 +247,7 @@ void UStaticMeshComponent::OnRep_StaticMesh(class UStaticMesh *OldStaticMesh)
 	{
 		// We have to force a call to SetStaticMesh with a new StaticMesh
 		UStaticMesh *NewStaticMesh = StaticMesh;
-		StaticMesh = NULL;
+		StaticMesh = OldStaticMesh;
 		
 		SetStaticMesh(NewStaticMesh);
 	}
