Remove redundant loop from a bad merge.

diff --git a/Engine/Source/Runtime/Engine/Private/NetworkDriver.cpp b/Engine/Source/Runtime/Engine/Private/NetworkDriver.cpp
--- a/Engine/Source/Runtime/Engine/Private/NetworkDriver.cpp
+++ b/Engine/Source/Runtime/Engine/Private/NetworkDriver.cpp
@@ -2897,22 +2897,6 @@ void UNetDriver::AddReferencedObjects(UObject* InThis, FReferenceCollector& Coll
 		}
 	}
 
-	for (auto It = This->ReplicationChangeListMap.CreateIterator(); It; ++It)
-	{
-		if (It.Value().IsValid())
-		{
-			FRepChangelistState* const ChangelistState = It.Value()->GetRepChangelistState();
-			if (ChangelistState && ChangelistState->RepLayout.IsValid())
-			{
-				ChangelistState->RepLayout->AddReferencedObjects(Collector);
-			}
-		}
-		else
-		{
-			It.RemoveCurrent();
-		}
-	}
-	
 	for (FObjectReplicator* Replicator : This->AllOwnedReplicators)
 	{
 		Collector.AddReferencedObject(Replicator->ObjectPtr, This);
