Implement BrianB's suggestions to clean up UActorChannel::ReplicateSubobject: avoid a weak object pointer creation and remove an unused return value.

diff --git a/Engine/Source/Runtime/Engine/Private/DataChannel.cpp b/Engine/Source/Runtime/Engine/Private/DataChannel.cpp
--- a/Engine/Source/Runtime/Engine/Private/DataChannel.cpp
+++ b/Engine/Source/Runtime/Engine/Private/DataChannel.cpp
@@ -3623,19 +3623,19 @@ bool UActorChannel::ReplicateSubobject(UObject *Obj, FOutBunch &Bunch, const FRe
 {
 	SCOPE_CYCLE_UOBJECT(ActorChannelRepSubObj, Obj);
 
+	TWeakObjectPtr<UObject> WeakObj(Obj);
+
 	// Hack for now: subobjects are SupportsObject==false until they are replicated via ::ReplicateSUbobject, and then we make them supported
 	// here, by forcing the packagemap to give them a NetGUID.
 	//
 	// Once we can lazily handle unmapped references on the client side, this can be simplified.
-	if ( !Connection->Driver->GuidCache->SupportsObject( Obj ) )
+	if ( !Connection->Driver->GuidCache->SupportsObject( Obj, &WeakObj ) )
 	{
-		FNetworkGUID NetGUID = Connection->Driver->GuidCache->AssignNewNetGUID_Server( Obj );	//Make sure he gets a NetGUID so that he is now 'supported'
+		Connection->Driver->GuidCache->AssignNewNetGUID_Server( Obj );	//Make sure he gets a NetGUID so that he is now 'supported'
 	}
 
 	bool NewSubobject = false;
 
-	TWeakObjectPtr<UObject> WeakObj(Obj);
-
 	if (!ObjectHasReplicator(WeakObj))
 	{
 		// This is the first time replicating this subobject
