Speculative fix for rare server crash - protect against potential null values in UNetConnection::OpenChannels.

diff --git a/Engine/Source/Runtime/Engine/Private/NetConnection.cpp b/Engine/Source/Runtime/Engine/Private/NetConnection.cpp
--- a/Engine/Source/Runtime/Engine/Private/NetConnection.cpp
+++ b/Engine/Source/Runtime/Engine/Private/NetConnection.cpp
@@ -1224,12 +1224,21 @@ void UNetConnection::ReceivedNak( int32 NakPacketId )
 	PackageMap->ReceivedNak(NakPacketId);
 
 	// Tell channels about Nak
-	for( int32 i=OpenChannels.Num()-1; i>=0; i-- )
+	for (int32 i = OpenChannels.Num() - 1; i >= 0; i--)
 	{
-		UChannel* Channel = OpenChannels[i];
-		Channel->ReceivedNak( NakPacketId );
-		if( Channel->OpenPacketId.InRange(NakPacketId) )
-			Channel->ReceivedAcks(); //warning: May destroy Channel.
+		UChannel* const Channel = OpenChannels[i];
+		if (Channel)
+		{
+			Channel->ReceivedNak(NakPacketId);
+			if (Channel->OpenPacketId.InRange(NakPacketId))
+			{
+				Channel->ReceivedAcks(); //warning: May destroy Channel.
+			}
+		}
+		else
+		{
+			UE_LOG(LogNet, Warning, TEXT("UNetConnection::ReceivedNak: null channel in OpenChannels array. %s"), *Describe());
+		}
 	}
 }
 
@@ -1406,30 +1415,37 @@ void UNetConnection::ReceivedPacket( FBitReader& Reader )
 			// Forward the ack to the channel.
 			UE_LOG(LogNetTraffic, Verbose, TEXT("   Received ack %i (%.1f)"), AckPacketId, (Reader.GetPosBits()-StartPos)/8.f );
 
-			for( int32 i=OpenChannels.Num()-1; i>=0; i-- )
+			for (int32 i = OpenChannels.Num() - 1; i >= 0; i--)
 			{
-				UChannel* Channel = OpenChannels[i];
-				
-				if( Channel->OpenPacketId.Last==AckPacketId ) // Necessary for unreliable "bNetTemporary" channels.
-				{
-					Channel->OpenAcked = 1;
-				}
+				UChannel* const Channel = OpenChannels[i];
 				
-				for( FOutBunch* OutBunch=Channel->OutRec; OutBunch; OutBunch=OutBunch->Next )
+				if (Channel)
 				{
-					if (OutBunch->bOpen)
+					if (Channel->OpenPacketId.Last == AckPacketId) // Necessary for unreliable "bNetTemporary" channels.
 					{
-						UE_LOG(LogNet, VeryVerbose, TEXT("Channel %i reset Ackd because open is reliable. "), Channel->ChIndex );
-						Channel->OpenAcked  = 0; // We have a reliable open bunch, don't let the above code set the OpenAcked state,
-												 // it must be set in UChannel::ReceivedAcks to verify all open bunches were received.
+						Channel->OpenAcked = 1;
 					}
-
-					if( OutBunch->PacketId==AckPacketId )
+				
+					for (FOutBunch* OutBunch = Channel->OutRec; OutBunch; OutBunch = OutBunch->Next)
 					{
-						OutBunch->ReceivedAck = 1;
+						if (OutBunch->bOpen)
+						{
+							UE_LOG(LogNet, VeryVerbose, TEXT("Channel %i reset Ackd because open is reliable. "), Channel->ChIndex );
+							Channel->OpenAcked  = 0; // We have a reliable open bunch, don't let the above code set the OpenAcked state,
+													 // it must be set in UChannel::ReceivedAcks to verify all open bunches were received.
+						}
+
+						if (OutBunch->PacketId == AckPacketId)
+						{
+							OutBunch->ReceivedAck = 1;
+						}
 					}
-				}				
-				Channel->ReceivedAcks(); //warning: May destroy Channel.
+					Channel->ReceivedAcks(); //warning: May destroy Channel.
+				}
+				else
+				{
+					UE_LOG(LogNet, Warning, TEXT("UNetConnection::ReceivedPacket: null channel in OpenChannels array while processing ack. %s"), *Describe());
+				}
 			}
 		}
 		else
@@ -2284,9 +2300,16 @@ void UNetConnection::Tick()
 		}
 		else
 		{
-			for( int32 i=OpenChannels.Num()-1; i>=0; i-- )
+			for (int32 i = OpenChannels.Num() - 1; i >= 0; i--)
 			{
-				OpenChannels[i]->Tick();
+				if (OpenChannels[i])
+				{
+					OpenChannels[i]->Tick();
+				}
+				else
+				{
+					UE_LOG(LogNet, Warning, TEXT("UNetConnection::Tick: null channel in OpenChannels array. %s"), *Describe());
+				}
 			}
 		}
 
diff --git a/Engine/Source/Runtime/Engine/Private/NetworkDriver.cpp b/Engine/Source/Runtime/Engine/Private/NetworkDriver.cpp
--- a/Engine/Source/Runtime/Engine/Private/NetworkDriver.cpp
+++ b/Engine/Source/Runtime/Engine/Private/NetworkDriver.cpp
@@ -2182,7 +2182,11 @@ bool UNetDriver::HandleSocketsCommand( const TCHAR* Cmd, FOutputDevice& Ar )
 		Ar.Logf( TEXT("   Server %s"), *ServerConnection->LowLevelDescribe() );
 		for( int32 i=0; i<ServerConnection->OpenChannels.Num(); i++ )
 		{
-			Ar.Logf( TEXT("      Channel %i: %s"), ServerConnection->OpenChannels[i]->ChIndex, *ServerConnection->OpenChannels[i]->Describe() );
+			UChannel* const Channel = ServerConnection->OpenChannels[i];
+			if (Channel)
+			{
+				Ar.Logf( TEXT("      Channel %i: %s"), Channel->ChIndex, *Channel->Describe() );
+			}
 		}
 	}
 #if WITH_SERVER_CODE
@@ -2192,7 +2196,11 @@ bool UNetDriver::HandleSocketsCommand( const TCHAR* Cmd, FOutputDevice& Ar )
 		Ar.Logf( TEXT("   Client %s"), *Connection->LowLevelDescribe() );
 		for( int32 j=0; j<Connection->OpenChannels.Num(); j++ )
 		{
-			Ar.Logf( TEXT("      Channel %i: %s"), Connection->OpenChannels[j]->ChIndex, *Connection->OpenChannels[j]->Describe() );
+			UChannel* const Channel = Connection->OpenChannels[j];
+			if (Channel)
+			{
+				Ar.Logf( TEXT("      Channel %i: %s"), Channel->ChIndex, *Channel->Describe() );
+			}
 		}
 	}
 #endif
