Changed FObjectReplicator::UpdateUnmappedObjects to call CallRepNotifies() after PostNetReceived to be consistent with order in FObjectReplicator::ProcessBunch().

diff --git a/Engine/Source/Runtime/Engine/Private/DataReplication.cpp b/Engine/Source/Runtime/Engine/Private/DataReplication.cpp
--- a/Engine/Source/Runtime/Engine/Private/DataReplication.cpp
+++ b/Engine/Source/Runtime/Engine/Private/DataReplication.cpp
@@ -1082,14 +1082,7 @@ void FObjectReplicator::PostReceivedBunch()
 	{
 		PostNetReceive();
 		bHasReplicatedProperties = false;
-	}
-
-	// Check if PostNetReceive() destroyed Object
-	UObject *Object = GetObject();
-	if (Object == NULL || Object->IsPendingKill())
-	{
-		return;
-	}
+	}	
 
 	// Call RepNotifies
 	CallRepNotifies(true);
@@ -1688,11 +1681,6 @@ void FObjectReplicator::UpdateUnmappedObjects( bool & bOutHasMoreUnmapped )
 		}
 	}
 
-	// Call any rep notifies that need to happen when object pointers change
-	// Pass in false to override the check for queued bunches. Otherwise, if the owning channel has queued bunches,
-	// the RepNotifies will remain in the list and the check for 0 RepNotifies above will fail next time.
-	CallRepNotifies(false);
-
 	if ( bSomeObjectsWereMapped )
 	{
 		// If we mapped some objects, make sure to call PostNetReceive (some game code will need to think this was actually replicated to work)
@@ -1701,6 +1689,11 @@ void FObjectReplicator::UpdateUnmappedObjects( bool & bOutHasMoreUnmapped )
 		UpdateGuidToReplicatorMap();
 	}
 
+	// Call any rep notifies that need to happen when object pointers change
+	// Pass in false to override the check for queued bunches. Otherwise, if the owning channel has queued bunches,
+	// the RepNotifies will remain in the list and the check for 0 RepNotifies above will fail next time.
+	CallRepNotifies(false);
+
 	UPackageMapClient * PackageMapClient = Cast< UPackageMapClient >(Connection->PackageMap);
 
 	if (PackageMapClient && OwningChannel)
