Comment large swathes of RepLayout.h. Make general format more consistent.

diff --git a/Engine/Source/Runtime/Engine/Private/RepLayout.cpp b/Engine/Source/Runtime/Engine/Private/RepLayout.cpp
--- a/Engine/Source/Runtime/Engine/Private/RepLayout.cpp
+++ b/Engine/Source/Runtime/Engine/Private/RepLayout.cpp
@@ -22,19 +22,19 @@ DECLARE_CYCLE_STAT(TEXT("RepLayout InitFromObjectClass"), STAT_RepLayout_InitFro
 
 
 int32 GDoPropertyChecksum = 0;
-static FAutoConsoleVariableRef CVarDoPropertyChecksum( TEXT( "net.DoPropertyChecksum" ), GDoPropertyChecksum, TEXT( "" ) );
+static FAutoConsoleVariableRef CVarDoPropertyChecksum(TEXT("net.DoPropertyChecksum"), GDoPropertyChecksum, TEXT(""));
 
 int32 GDoReplicationContextString = 0;
-static FAutoConsoleVariableRef CVarDoReplicationContextString( TEXT( "net.ContextDebug" ), GDoReplicationContextString, TEXT( "" ) );
+static FAutoConsoleVariableRef CVarDoReplicationContextString(TEXT("net.ContextDebug"), GDoReplicationContextString, TEXT(""));
 
 int32 GNetSharedSerializedData = 1;
-static FAutoConsoleVariableRef CVarNetShareSerializedData( TEXT( "net.ShareSerializedData" ), GNetSharedSerializedData, TEXT( "" ) );
+static FAutoConsoleVariableRef CVarNetShareSerializedData(TEXT("net.ShareSerializedData"), GNetSharedSerializedData, TEXT(""));
 
 int32 GNetVerifyShareSerializedData = 0;
-static FAutoConsoleVariableRef CVarNetVerifyShareSerializedData( TEXT( "net.VerifyShareSerializedData" ), GNetVerifyShareSerializedData, TEXT( "" ) );
+static FAutoConsoleVariableRef CVarNetVerifyShareSerializedData(TEXT("net.VerifyShareSerializedData"), GNetVerifyShareSerializedData, TEXT(""));
 
 int32 LogSkippedRepNotifies = 0;
-static FAutoConsoleVariable CVarLogSkippedRepNotifies(TEXT("Net.LogSkippedRepNotifies"), LogSkippedRepNotifies, TEXT("Log when the networking code skips calling a repnotify clientside due to the property value not changing."), ECVF_Default );
+static FAutoConsoleVariable CVarLogSkippedRepNotifies(TEXT("Net.LogSkippedRepNotifies"), LogSkippedRepNotifies, TEXT("Log when the networking code skips calling a repnotify clientside due to the property value not changing."), ECVF_Default);
 
 int32 MaxRepArraySize = UNetworkSettings::DefaultMaxRepArraySize;
 int32 MaxRepArrayMemory = UNetworkSettings::DefaultMaxRepArrayMemory;
@@ -109,123 +109,139 @@ FConsoleVariableSinkHandle MaxRepArrayMemorySink = CreateMaxArrayMemoryCVarAndRe
 //#define ENABLE_SUPER_CHECKSUMS
 
 #ifdef USE_CUSTOM_COMPARE
-static FORCEINLINE bool CompareBool( const FRepLayoutCmd& Cmd, const void* A, const void* B )
+static FORCEINLINE bool CompareBool(
+	const FRepLayoutCmd&	Cmd,
+	const void* 			A,
+	const void* 			B)
 {
-	return Cmd.Property->Identical( A, B );
+	return Cmd.Property->Identical(A, B);
 }
 
-static FORCEINLINE bool CompareObject( const FRepLayoutCmd& Cmd, const void* A, const void* B )
+static FORCEINLINE bool CompareObject(
+	const FRepLayoutCmd&	Cmd,
+	const void* 			A,
+	const void* 			B)
 {
 #if 1
 	// Until UObjectPropertyBase::Identical is made safe for GC'd objects, we need to do it manually
 	// This saves us from having to add referenced objects during GC
-	UObjectPropertyBase * ObjProperty = CastChecked< UObjectPropertyBase>( Cmd.Property );
+	UObjectPropertyBase* ObjProperty = CastChecked<UObjectPropertyBase>(Cmd.Property);
 
-	UObject* ObjectA = ObjProperty->GetObjectPropertyValue( A );
-	UObject* ObjectB = ObjProperty->GetObjectPropertyValue( B );
+	UObject* ObjectA = ObjProperty->GetObjectPropertyValue(A);
+	UObject* ObjectB = ObjProperty->GetObjectPropertyValue(B);
 
 	return ObjectA == ObjectB;
 #else
-	return Cmd.Property->Identical( A, B );
+	return Cmd.Property->Identical(A, B);
 #endif
 }
 
-template< typename T >
-bool CompareValue( const T * A, const T * B )
+template<typename T>
+bool CompareValue(const T * A, const T * B)
 {
 	return *A == *B;
 }
 
-template< typename T >
-bool CompareValue( const void* A, const void* B )
+template<typename T>
+bool CompareValue(const void* A, const void* B)
 {
-	return CompareValue( (T*)A, (T*)B);
+	return CompareValue((T*)A, (T*)B);
 }
 
-static FORCEINLINE bool PropertiesAreIdenticalNative( const FRepLayoutCmd& Cmd, const void* A, const void* B )
+static FORCEINLINE bool PropertiesAreIdenticalNative(
+	const FRepLayoutCmd&	Cmd,
+	const void*				A,
+	const void*				B)
 {
-	switch ( Cmd.Type )
-	{
-		case ERepLayoutCmdType::PropertyBool:			return CompareBool( Cmd, A, B );
-		case ERepLayoutCmdType::PropertyByte:			return CompareValue<uint8>( A, B );
-		case ERepLayoutCmdType::PropertyFloat:			return CompareValue<float>( A, B );
-		case ERepLayoutCmdType::PropertyInt:			return CompareValue<int32>( A, B );
-		case ERepLayoutCmdType::PropertyName:			return CompareValue<FName>( A, B );
-		case ERepLayoutCmdType::PropertyObject:			return CompareObject( Cmd, A, B );
-		case ERepLayoutCmdType::PropertyUInt32:			return CompareValue<uint32>( A, B );
-		case ERepLayoutCmdType::PropertyUInt64:			return CompareValue<uint64>( A, B );
-		case ERepLayoutCmdType::PropertyVector:			return CompareValue<FVector>( A, B );
-		case ERepLayoutCmdType::PropertyVector100:		return CompareValue<FVector_NetQuantize100>( A, B );
-		case ERepLayoutCmdType::PropertyVectorQ:		return CompareValue<FVector_NetQuantize>( A, B );
-		case ERepLayoutCmdType::PropertyVectorNormal:	return CompareValue<FVector_NetQuantizeNormal>( A, B );
-		case ERepLayoutCmdType::PropertyVector10:		return CompareValue<FVector_NetQuantize10>( A, B );
-		case ERepLayoutCmdType::PropertyPlane:			return CompareValue<FPlane>( A, B );
-		case ERepLayoutCmdType::PropertyRotator:		return CompareValue<FRotator>( A, B );
-		case ERepLayoutCmdType::PropertyNetId:			return CompareValue<FUniqueNetIdRepl>( A, B );
-		case ERepLayoutCmdType::RepMovement:			return CompareValue<FRepMovement>( A, B );
-		case ERepLayoutCmdType::PropertyString:			return CompareValue<FString>( A, B );
-		case ERepLayoutCmdType::Property:				return Cmd.Property->Identical( A, B );
+	switch (Cmd.Type)
+	{
+		case ERepLayoutCmdType::PropertyBool:			return CompareBool(Cmd, A, B);
+		case ERepLayoutCmdType::PropertyByte:			return CompareValue<uint8>(A, B);
+		case ERepLayoutCmdType::PropertyFloat:			return CompareValue<float>(A, B);
+		case ERepLayoutCmdType::PropertyInt:			return CompareValue<int32>(A, B);
+		case ERepLayoutCmdType::PropertyName:			return CompareValue<FName>(A, B);
+		case ERepLayoutCmdType::PropertyObject:			return CompareObject(Cmd, A, B);
+		case ERepLayoutCmdType::PropertyUInt32:			return CompareValue<uint32>(A, B);
+		case ERepLayoutCmdType::PropertyUInt64:			return CompareValue<uint64>(A, B);
+		case ERepLayoutCmdType::PropertyVector:			return CompareValue<FVector>(A, B);
+		case ERepLayoutCmdType::PropertyVector100:		return CompareValue<FVector_NetQuantize100>(A, B);
+		case ERepLayoutCmdType::PropertyVectorQ:		return CompareValue<FVector_NetQuantize>(A, B);
+		case ERepLayoutCmdType::PropertyVectorNormal:	return CompareValue<FVector_NetQuantizeNormal>(A, B);
+		case ERepLayoutCmdType::PropertyVector10:		return CompareValue<FVector_NetQuantize10>(A, B);
+		case ERepLayoutCmdType::PropertyPlane:			return CompareValue<FPlane>(A, B);
+		case ERepLayoutCmdType::PropertyRotator:		return CompareValue<FRotator>(A, B);
+		case ERepLayoutCmdType::PropertyNetId:			return CompareValue<FUniqueNetIdRepl>(A, B);
+		case ERepLayoutCmdType::RepMovement:			return CompareValue<FRepMovement>(A, B);
+		case ERepLayoutCmdType::PropertyString:			return CompareValue<FString>(A, B);
+		case ERepLayoutCmdType::Property:				return Cmd.Property->Identical(A, B);
 		default: 
-			UE_LOG( LogRep, Fatal, TEXT( "PropertiesAreIdentical: Unsupported type! %i (%s)" ), (uint8)Cmd.Type, *Cmd.Property->GetName() );
+			UE_LOG(LogRep, Fatal, TEXT("PropertiesAreIdentical: Unsupported type! %i (%s)"), (uint8)Cmd.Type, *Cmd.Property->GetName());
 	}
 
 	return false;
 }
 
-static FORCEINLINE bool PropertiesAreIdentical( const FRepLayoutCmd& Cmd, const void* A, const void* B )
+static FORCEINLINE bool PropertiesAreIdentical(
+	const FRepLayoutCmd&	Cmd,
+	const void*				A,
+	const void*				B)
 {
-	const bool bIsIdentical = PropertiesAreIdenticalNative( Cmd, A, B );
+	const bool bIsIdentical = PropertiesAreIdenticalNative(Cmd, A, B);
 #if 0
 	// Sanity check result
-	if ( bIsIdentical != Cmd.Property->Identical( A, B ) )
+	if (bIsIdentical != Cmd.Property->Identical(A, B))
 	{
-		UE_LOG( LogRep, Fatal, TEXT( "PropertiesAreIdentical: Result mismatch! (%s)" ), *Cmd.Property->GetFullName() );
+		UE_LOG(LogRep, Fatal, TEXT("PropertiesAreIdentical: Result mismatch! (%s)"), *Cmd.Property->GetFullName());
 	}
 #endif
 	return bIsIdentical;
 }
 #else
-static FORCEINLINE bool PropertiesAreIdentical( const FRepLayoutCmd& Cmd, const void* A, const void* B )
+static FORCEINLINE bool PropertiesAreIdentical(const FRepLayoutCmd& Cmd, const void* A, const void* B)
 {
-	return Cmd.Property->Identical( A, B );
+	return Cmd.Property->Identical(A, B);
 }
 #endif
 
-static FORCEINLINE void StoreProperty( const FRepLayoutCmd& Cmd, void* A, const void* B )
+static FORCEINLINE void StoreProperty(const FRepLayoutCmd& Cmd, void* A, const void* B)
 {
-	Cmd.Property->CopySingleValue( A, B );
+	Cmd.Property->CopySingleValue(A, B);
 }
 
-static FORCEINLINE void SerializeGenericChecksum( FBitArchive & Ar )
+static FORCEINLINE void SerializeGenericChecksum(FBitArchive& Ar)
 {
 	uint32 Checksum = 0xABADF00D;
 	Ar << Checksum;
-	check( Checksum == 0xABADF00D );
+	check(Checksum == 0xABADF00D);
 }
 
-static void SerializeReadWritePropertyChecksum( const FRepLayoutCmd& Cmd, const int32 CurCmdIndex, const uint8* Data, FBitArchive & Ar )
+static void SerializeReadWritePropertyChecksum(
+	const FRepLayoutCmd&	Cmd,
+	const int32				CurCmdIndex,
+	const uint8*			Data,
+	FBitArchive&			Ar)
 {
 	// Serialize various attributes that will mostly ensure we are working on the same property
-	const uint32 NameHash = GetTypeHash( Cmd.Property->GetName() );
+	const uint32 NameHash = GetTypeHash(Cmd.Property->GetName());
 
 	uint32 MarkerChecksum = 0;
 
 	// Evolve the checksum over several values that will uniquely identity where we are and should be
-	MarkerChecksum = FCrc::MemCrc_DEPRECATED( &NameHash,		sizeof( NameHash ),		MarkerChecksum );
-	MarkerChecksum = FCrc::MemCrc_DEPRECATED( &Cmd.Offset,		sizeof( Cmd.Offset ),	MarkerChecksum );
-	MarkerChecksum = FCrc::MemCrc_DEPRECATED( &CurCmdIndex,	sizeof( CurCmdIndex ),	MarkerChecksum );
+	MarkerChecksum = FCrc::MemCrc_DEPRECATED(&NameHash, sizeof(NameHash), MarkerChecksum);
+	MarkerChecksum = FCrc::MemCrc_DEPRECATED(&Cmd.Offset, sizeof(Cmd.Offset), MarkerChecksum);
+	MarkerChecksum = FCrc::MemCrc_DEPRECATED(&CurCmdIndex, sizeof(CurCmdIndex), MarkerChecksum);
 
 	const uint32 OriginalMarkerChecksum = MarkerChecksum;
 
 	Ar << MarkerChecksum;
 
-	if ( MarkerChecksum != OriginalMarkerChecksum )
+	if (MarkerChecksum != OriginalMarkerChecksum)
 	{
 		// This is fatal, as it means we are out of sync to the point we can't recover
-		UE_LOG( LogRep, Fatal, TEXT( "SerializeReadWritePropertyChecksum: Property checksum marker failed! [%s]" ), *Cmd.Property->GetFullName() );
+		UE_LOG(LogRep, Fatal, TEXT("SerializeReadWritePropertyChecksum: Property checksum marker failed! [%s]"), *Cmd.Property->GetFullName());
 	}
 
-	if ( Cmd.Property->IsA( UObjectPropertyBase::StaticClass() ) )
+	if (Cmd.Property->IsA(UObjectPropertyBase::StaticClass()))
 	{
 		// Can't handle checksums for objects right now
 		// Need to resolve how to handle unmapped objects
@@ -236,11 +252,11 @@ static void SerializeReadWritePropertyChecksum( const FRepLayoutCmd& Cmd, const
 	// This will require NetSerializeItem to be deterministic, in and out
 	// i.e, not only does NetSerializeItem need to write the same blob on the same input data, but
 	//	it also needs to write the same blob it just read as well.
-	FBitWriter Writer( 0, true );
+	FBitWriter Writer(0, true);
 
-	Cmd.Property->NetSerializeItem( Writer, NULL, const_cast< uint8* >( Data ) );
+	Cmd.Property->NetSerializeItem(Writer, NULL, const_cast<uint8*>(Data));
 
-	if ( Ar.IsSaving() )
+	if (Ar.IsSaving())
 	{
 		// If this is the server, do a read, and then another write so that we do exactly what the client will do, which will better ensure determinism 
 
@@ -248,49 +264,53 @@ static void SerializeReadWritePropertyChecksum( const FRepLayoutCmd& Cmd, const
 		const int32 OriginalDim = Cmd.Property->ArrayDim;
 		Cmd.Property->ArrayDim = 1;
 
-		TArray< uint8 > TempPropMemory;
-		TempPropMemory.AddZeroed( Cmd.Property->ElementSize + 4 );
+		TArray<uint8> TempPropMemory;
+		TempPropMemory.AddZeroed(Cmd.Property->ElementSize + 4);
 		uint32* Guard = (uint32*)&TempPropMemory[TempPropMemory.Num() - 4];
 		const uint32 TAG_VALUE = 0xABADF00D;
 		*Guard = TAG_VALUE;
-		Cmd.Property->InitializeValue( TempPropMemory.GetData() );
-		check( *Guard == TAG_VALUE );
+		Cmd.Property->InitializeValue(TempPropMemory.GetData());
+		check(*Guard == TAG_VALUE);
 
 		// Read it back in and then write it out to produce what the client will produce
-		FBitReader Reader( Writer.GetData(), Writer.GetNumBits() );
-		Cmd.Property->NetSerializeItem( Reader, NULL, TempPropMemory.GetData() );
-		check( Reader.AtEnd() && !Reader.IsError() );
-		check( *Guard == TAG_VALUE );
+		FBitReader Reader(Writer.GetData(), Writer.GetNumBits());
+		Cmd.Property->NetSerializeItem(Reader, NULL, TempPropMemory.GetData());
+		check(Reader.AtEnd() && !Reader.IsError());
+		check(*Guard == TAG_VALUE);
 
 		// Write it back out for a final time
 		Writer.Reset();
 
-		Cmd.Property->NetSerializeItem( Writer, NULL, TempPropMemory.GetData() );
-		check( *Guard == TAG_VALUE );
+		Cmd.Property->NetSerializeItem(Writer, NULL, TempPropMemory.GetData());
+		check(*Guard == TAG_VALUE);
 
 		// Destroy temp memory
-		Cmd.Property->DestroyValue( TempPropMemory.GetData() );
+		Cmd.Property->DestroyValue(TempPropMemory.GetData());
 
 		// Restore the static array size
 		Cmd.Property->ArrayDim = OriginalDim;
 
-		check( *Guard == TAG_VALUE );
+		check(*Guard == TAG_VALUE);
 	}
 
-	uint32 PropertyChecksum = FCrc::MemCrc_DEPRECATED( Writer.GetData(), Writer.GetNumBytes() );
+	uint32 PropertyChecksum = FCrc::MemCrc_DEPRECATED(Writer.GetData(), Writer.GetNumBytes());
 
 	const uint32 OriginalPropertyChecksum = PropertyChecksum;
 
 	Ar << PropertyChecksum;
 
-	if ( PropertyChecksum != OriginalPropertyChecksum )
+	if (PropertyChecksum != OriginalPropertyChecksum)
 	{
 		// This is a warning, because for some reason, float rounding issues in the quantization functions cause this to return false positives
-		UE_LOG( LogRep, Warning, TEXT( "Property checksum failed! [%s]" ), *Cmd.Property->GetFullName() );
+		UE_LOG(LogRep, Warning, TEXT("Property checksum failed! [%s]"), *Cmd.Property->GetFullName());
 	}
 }
 
-static uint32 GetRepLayoutCmdCompatibleChecksum( const UProperty* Property, const UNetConnection* ServerConnection, const uint32 StaticArrayIndex, const uint32 InChecksum )
+static uint32 GetRepLayoutCmdCompatibleChecksum(
+	const UProperty*		Property,
+	const UNetConnection*	ServerConnection,
+	const uint32			StaticArrayIndex,
+	const uint32			InChecksum)
 {
 	// Compatible checksums are only used for InternalAck connections
 	if (ServerConnection && !ServerConnection->InternalAck)
@@ -298,162 +318,186 @@ static uint32 GetRepLayoutCmdCompatibleChecksum( const UProperty* Property, cons
 		return 0;
 	}
 
-	uint32 CompatibleChecksum = FCrc::StrCrc32( *Property->GetName().ToLower(), InChecksum );						// Evolve checksum on name
-	CompatibleChecksum = FCrc::StrCrc32( *Property->GetCPPType( nullptr, 0 ).ToLower(), CompatibleChecksum );		// Evolve by property type
+	// Evolve checksum on name
+	uint32 CompatibleChecksum = FCrc::StrCrc32(*Property->GetName().ToLower(), InChecksum);	
+	
+	// Evolve by property type			
+	CompatibleChecksum = FCrc::StrCrc32(*Property->GetCPPType(nullptr, 0).ToLower(), CompatibleChecksum);
 	
 	// Evolve by StaticArrayIndex (to make all unrolled static array elements unique)
 	if ((ServerConnection == nullptr) ||
 		(ServerConnection->EngineNetworkProtocolVersion >= EEngineNetworkVersionHistory::HISTORY_REPCMD_CHECKSUM_REMOVE_PRINTF))
 	{
-		CompatibleChecksum = FCrc::MemCrc32( &StaticArrayIndex, sizeof(StaticArrayIndex), CompatibleChecksum );
+		CompatibleChecksum = FCrc::MemCrc32(&StaticArrayIndex, sizeof(StaticArrayIndex), CompatibleChecksum);
 	}
 	else
 	{
-		CompatibleChecksum = FCrc::StrCrc32( *FString::Printf( TEXT( "%i" ), StaticArrayIndex ), CompatibleChecksum );
+		CompatibleChecksum = FCrc::StrCrc32(*FString::Printf(TEXT("%i"), StaticArrayIndex), CompatibleChecksum);
 	}
 
 	return CompatibleChecksum;
 }
 
-#define INIT_STACK( TStack )							\
-	void InitStack( TStack& StackState )				\
+#define INIT_STACK(TStack)								\
+	void InitStack(TStack& StackState)					\
 
 #define SHOULD_PROCESS_NEXT_CMD()						\
 	bool ShouldProcessNextCmd()							\
 
-#define PROCESS_ARRAY_CMD( TStack )						\
+#define PROCESS_ARRAY_CMD(TStack)						\
 	void ProcessArrayCmd_r(								\
 	TStack&							PrevStackState,		\
 	TStack&							StackState,			\
 	const FRepLayoutCmd&			Cmd,				\
 	const int32						CmdIndex,			\
 	uint8* RESTRICT					ShadowData,			\
-	uint8* RESTRICT					Data )				\
+	uint8* RESTRICT					Data)				\
 
 
-#define PROCESS_CMD( TStack )							\
+#define PROCESS_CMD(TStack)								\
 	void ProcessCmd(									\
 	TStack&							StackState,			\
 	const FRepLayoutCmd&			Cmd,				\
 	const int32						CmdIndex,			\
 	uint8* RESTRICT					ShadowData,			\
-	uint8* RESTRICT					Data )				\
+	uint8* RESTRICT					Data)				\
 
 class FCmdIteratorBaseStackState
 {
 public:
-	FCmdIteratorBaseStackState( const int32 InCmdStart, const int32 InCmdEnd, FScriptArray*	InShadowArray, FScriptArray* InDataArray, uint8* RESTRICT InShadowBaseData, uint8* RESTRICT	InBaseData ) : 
-		CmdStart( InCmdStart ),
-		CmdEnd( InCmdEnd ),
-		ShadowArray( InShadowArray ),
-		DataArray( InDataArray ),
-		ShadowBaseData( InShadowBaseData ),
-		BaseData( InBaseData )
-	{
-	}
-
-	const int32			CmdStart; 
-	const int32			CmdEnd;
-
-	FScriptArray*		ShadowArray;
-	FScriptArray*		DataArray;
-
-	uint8* RESTRICT		ShadowBaseData;
-	uint8* RESTRICT		BaseData;
+	FCmdIteratorBaseStackState(
+		const int32		InCmdStart,
+		const int32		InCmdEnd,
+		FScriptArray*	InShadowArray,
+		FScriptArray*	InDataArray,
+		uint8* RESTRICT	InShadowBaseData,
+		uint8* RESTRICT	InBaseData
+	):
+		CmdStart(InCmdStart),
+		CmdEnd(InCmdEnd),
+		ShadowArray(InShadowArray),
+		DataArray(InDataArray),
+		ShadowBaseData(InShadowBaseData),
+		BaseData(InBaseData)
+	{
+	}
+
+	const int32 CmdStart; 
+	const int32 CmdEnd;
+
+	FScriptArray* ShadowArray;
+	FScriptArray* DataArray;
+
+	uint8* RESTRICT ShadowBaseData;
+	uint8* RESTRICT BaseData;
 };
 
 // This uses the "Curiously recurring template pattern" (CRTP) ideas
-template< typename TImpl, typename TStackState >
+template<typename TImpl, typename TStackState>
 class FRepLayoutCmdIterator
 {
 public:
-	FRepLayoutCmdIterator( const TArray< FRepParentCmd >& InParents, const TArray< FRepLayoutCmd >& InCmds ) : Parents( InParents ), Cmds( InCmds ) {}
+	FRepLayoutCmdIterator(const TArray<FRepParentCmd>& InParents, const TArray<FRepLayoutCmd>& InCmds):
+		Parents(InParents),
+		Cmds(InCmds)
+	{}
 
-	void ProcessDataArrayElements_r( TStackState& StackState, const FRepLayoutCmd& Cmd )
+	void ProcessDataArrayElements_r(TStackState& StackState, const FRepLayoutCmd& Cmd)
 	{
-		const int32 NumDataArrayElements	= StackState.DataArray		? StackState.DataArray->Num()	: 0;
-		const int32 NumShadowArrayElements	= StackState.ShadowArray	? StackState.ShadowArray->Num() : 0;
+		const int32 NumDataArrayElements = StackState.DataArray ? StackState.DataArray->Num() : 0;
+		const int32 NumShadowArrayElements = StackState.ShadowArray ? StackState.ShadowArray->Num() : 0;
 
 		// Loop using the number of elements in data array
-		for ( int32 i = 0; i < NumDataArrayElements; i++ )
+		for (int32 i = 0; i < NumDataArrayElements; i++)
 		{
 			const int32 ElementOffset = i * Cmd.ElementSize;
 
-			uint8* Data			= StackState.BaseData + ElementOffset;
-			uint8* ShadowData	= i < NumShadowArrayElements ? ( StackState.ShadowBaseData + ElementOffset ) : NULL;	// ShadowArray might be smaller than DataArray
+			uint8* Data = StackState.BaseData + ElementOffset;
+			
+			// ShadowArray might be smaller than DataArray
+			uint8* ShadowData = i < NumShadowArrayElements ? (StackState.ShadowBaseData + ElementOffset) : NULL;
 
-			ProcessCmds_r( StackState, ShadowData, Data );
+			ProcessCmds_r(StackState, ShadowData, Data);
 		}
 	}
 
-	void ProcessShadowArrayElements_r( TStackState& StackState, const FRepLayoutCmd& Cmd )
+	void ProcessShadowArrayElements_r(TStackState& StackState, const FRepLayoutCmd& Cmd)
 	{
-		const int32 NumDataArrayElements	= StackState.DataArray		? StackState.DataArray->Num()	: 0;
-		const int32 NumShadowArrayElements	= StackState.ShadowArray	? StackState.ShadowArray->Num() : 0;
+		const int32 NumDataArrayElements = StackState.DataArray ? StackState.DataArray->Num()	: 0;
+		const int32 NumShadowArrayElements = StackState.ShadowArray ? StackState.ShadowArray->Num() : 0;
 
 		// Loop using the number of elements in shadow array
-		for ( int32 i = 0; i < NumShadowArrayElements; i++ )
+		for (int32 i = 0; i < NumShadowArrayElements; i++)
 		{
 			const int32 ElementOffset = i * Cmd.ElementSize;
 
-			uint8* Data			= i < NumDataArrayElements ? ( StackState.BaseData + ElementOffset ) : NULL;	// DataArray might be smaller than ShadowArray
+			// DataArray might be smaller than ShadowArray
+			uint8* Data			= i < NumDataArrayElements ? (StackState.BaseData + ElementOffset) : NULL;
 			uint8* ShadowData	= StackState.ShadowBaseData + ElementOffset;
 
-			ProcessCmds_r( StackState, ShadowData, Data );
+			ProcessCmds_r(StackState, ShadowData, Data);
 		}
 	}
 
-	void ProcessArrayCmd_r( TStackState & PrevStackState, const FRepLayoutCmd& Cmd, const int32 CmdIndex, uint8* RESTRICT ShadowData, uint8* RESTRICT Data )
+	void ProcessArrayCmd_r(
+		TStackState&			PrevStackState,
+		const FRepLayoutCmd&	Cmd,
+		const int32				CmdIndex,
+		uint8* RESTRICT			ShadowData,
+		uint8* RESTRICT			Data)
 	{
-		check( ShadowData != NULL || Data != NULL );
+		check(ShadowData != NULL || Data != NULL);
 
-		FScriptArray* ShadowArray	= (FScriptArray*)ShadowData;
-		FScriptArray* DataArray		= (FScriptArray*)Data;
+		FScriptArray* ShadowArray = (FScriptArray*)ShadowData;
+		FScriptArray* DataArray = (FScriptArray*)Data;
 
-		TStackState StackState( CmdIndex + 1, Cmd.EndCmd - 1, ShadowArray, DataArray, ShadowArray ? (uint8*)ShadowArray->GetData() : NULL, DataArray ? (uint8*)DataArray->GetData() : NULL );
+		TStackState StackState(CmdIndex + 1, Cmd.EndCmd - 1, ShadowArray, DataArray, ShadowArray ? (uint8*)ShadowArray->GetData() : NULL, DataArray ? (uint8*)DataArray->GetData() : NULL);
 
-		static_cast< TImpl* >( this )->ProcessArrayCmd_r( PrevStackState, StackState, Cmd, CmdIndex, ShadowData, Data );
+		static_cast<TImpl*>(this)->ProcessArrayCmd_r(PrevStackState, StackState, Cmd, CmdIndex, ShadowData, Data);
 	}
 
-	void ProcessCmds_r( TStackState& StackState, uint8* RESTRICT ShadowData, uint8* RESTRICT Data )
+	void ProcessCmds_r(
+		TStackState&	StackState,
+		uint8* RESTRICT	ShadowData,
+		uint8* RESTRICT	Data)
 	{
-		check( ShadowData != NULL || Data != NULL );
+		check(ShadowData != NULL || Data != NULL);
 
-		for ( int32 CmdIndex = StackState.CmdStart; CmdIndex < StackState.CmdEnd; CmdIndex++ )
+		for (int32 CmdIndex = StackState.CmdStart; CmdIndex < StackState.CmdEnd; CmdIndex++)
 		{
 			const FRepLayoutCmd& Cmd = Cmds[ CmdIndex ];
 
-			check( Cmd.Type != ERepLayoutCmdType::Return );
+			check(Cmd.Type != ERepLayoutCmdType::Return);
 
-			if ( Cmd.Type == ERepLayoutCmdType::DynamicArray )
+			if (Cmd.Type == ERepLayoutCmdType::DynamicArray)
 			{
-				if ( static_cast< TImpl* >( this )->ShouldProcessNextCmd() )
+				if (static_cast<TImpl*>(this)->ShouldProcessNextCmd())
 				{
-					ProcessArrayCmd_r( StackState, Cmd, CmdIndex, ShadowData ? ( ShadowData + Cmd.Offset ) : NULL, Data ? ( Data + Cmd.Offset ) : NULL );
+					ProcessArrayCmd_r(StackState, Cmd, CmdIndex, ShadowData ? (ShadowData + Cmd.Offset) : NULL, Data ? (Data + Cmd.Offset) : NULL);
 				}
 				CmdIndex = Cmd.EndCmd - 1;	// Jump past children of this array (-1 for ++ in for loop)
 			}
 			else
 			{
-				if ( static_cast< TImpl* >( this )->ShouldProcessNextCmd() )
+				if (static_cast<TImpl*>(this)->ShouldProcessNextCmd())
 				{
-					static_cast< TImpl* >( this )->ProcessCmd( StackState, Cmd, CmdIndex, ShadowData, Data );
+					static_cast<TImpl*>(this)->ProcessCmd(StackState, Cmd, CmdIndex, ShadowData, Data);
 				}
 			}
 		}
 	}
 
-	void ProcessCmds( uint8* RESTRICT Data, uint8* RESTRICT ShadowData )
+	void ProcessCmds(uint8* RESTRICT Data, uint8* RESTRICT ShadowData)
 	{
-		TStackState StackState( 0, Cmds.Num() - 1, NULL, NULL, ShadowData, Data );
+		TStackState StackState(0, Cmds.Num() - 1, NULL, NULL, ShadowData, Data);
 
-		static_cast< TImpl* >( this )->InitStack( StackState );
+		static_cast<TImpl*>(this)->InitStack(StackState);
 
-		ProcessCmds_r( StackState, ShadowData, Data );
+		ProcessCmds_r(StackState, ShadowData, Data);
 	}
 
-	const TArray< FRepParentCmd >&	Parents;
-	const TArray< FRepLayoutCmd >&	Cmds;
+	const TArray<FRepParentCmd>& Parents;
+	const TArray<FRepLayoutCmd>& Cmds;
 };
 
 uint16 FRepLayout::CompareProperties_r(
@@ -461,48 +505,48 @@ uint16 FRepLayout::CompareProperties_r(
 	const int32				CmdEnd,
 	const uint8* RESTRICT	CompareData,
 	const uint8* RESTRICT	Data,
-	TArray< uint16 > &		Changed,
+	TArray<uint16>&			Changed,
 	uint16					Handle,
 	const bool				bIsInitial,
-	const bool				bForceFail ) const
+	const bool				bForceFail) const
 {
-	check( CompareData );
+	check(CompareData);
 
-	for ( int32 CmdIndex = CmdStart; CmdIndex < CmdEnd; CmdIndex++ )
+	for (int32 CmdIndex = CmdStart; CmdIndex < CmdEnd; CmdIndex++)
 	{
 		const FRepLayoutCmd& Cmd = Cmds[CmdIndex];
 		const FRepParentCmd& ParentCmd = Parents[Cmd.ParentIndex];
 
-		check( Cmd.Type != ERepLayoutCmdType::Return );
+		check(Cmd.Type != ERepLayoutCmdType::Return);
 
 		Handle++;
 
 		const bool bIsLifetime = ((ParentCmd.Flags & ERepParentFlags::IsLifetime) != ERepParentFlags::None);
-		const bool bShouldSkip = !bIsLifetime || ( ParentCmd.Condition == COND_InitialOnly && !bIsInitial );
+		const bool bShouldSkip = !bIsLifetime || (ParentCmd.Condition == COND_InitialOnly && !bIsInitial);
 
-		if ( Cmd.Type == ERepLayoutCmdType::DynamicArray )
+		if (Cmd.Type == ERepLayoutCmdType::DynamicArray)
 		{
-			if ( bShouldSkip )
+			if (bShouldSkip)
 			{
 				CmdIndex = Cmd.EndCmd - 1;		// The -1 to handle the ++ in the for loop
 				continue;
 			}
 
 			// Once we hit an array, start using a stack based approach
-			CompareProperties_Array_r( CompareData + Cmd.Offset, ( const uint8* )Data + Cmd.Offset, Changed, CmdIndex, Handle, bIsInitial, bForceFail );
+			CompareProperties_Array_r(CompareData + Cmd.Offset, (const uint8*)Data + Cmd.Offset, Changed, CmdIndex, Handle, bIsInitial, bForceFail);
 			CmdIndex = Cmd.EndCmd - 1;		// The -1 to handle the ++ in the for loop
 			continue;
 		}
 
-		if ( bShouldSkip )
+		if (bShouldSkip)
 		{
 			continue;
 		}
 
-		if ( bForceFail || !PropertiesAreIdentical( Cmd, ( const void* )( CompareData + Cmd.Offset ), ( const void* )( Data + Cmd.Offset ) ) )
+		if (bForceFail || !PropertiesAreIdentical(Cmd, (const void*)(CompareData + Cmd.Offset), (const void*)(Data + Cmd.Offset)))
 		{
-			StoreProperty( Cmd, ( void* )( CompareData + Cmd.Offset ), ( const void* )( Data + Cmd.Offset ) );
-			Changed.Add( Handle );
+			StoreProperty(Cmd, (void*)(CompareData + Cmd.Offset), (const void*)(Data + Cmd.Offset));
+			Changed.Add(Handle);
 		}
 	}
 
@@ -512,7 +556,7 @@ uint16 FRepLayout::CompareProperties_r(
 void FRepLayout::CompareProperties_Array_r(
 	const uint8* RESTRICT	CompareData,
 	const uint8* RESTRICT	Data,
-	TArray< uint16 > &		Changed,
+	TArray<uint16>&			Changed,
 	const uint16			CmdIndex,
 	const uint16			Handle,
 	const bool				bIsInitial,
@@ -521,61 +565,61 @@ void FRepLayout::CompareProperties_Array_r(
 {
 	const FRepLayoutCmd& Cmd = Cmds[CmdIndex];
 
-	FScriptArray * CompareArray = ( FScriptArray * )CompareData;
-	FScriptArray * Array = ( FScriptArray * )Data;
+	FScriptArray* CompareArray = (FScriptArray*)CompareData;
+	FScriptArray* Array = (FScriptArray*)Data;
 
-	const uint16 ArrayNum			= Array->Num();
-	const uint16 CompareArrayNum	= CompareArray->Num();
+	const uint16 ArrayNum = Array->Num();
+	const uint16 CompareArrayNum = CompareArray->Num();
 
 	// Make the shadow state match the actual state at the time of compare
-	FScriptArrayHelper StoredArrayHelper( ( UArrayProperty * )Cmd.Property, CompareData );
-	StoredArrayHelper.Resize( ArrayNum );
+	FScriptArrayHelper StoredArrayHelper((UArrayProperty*)Cmd.Property, CompareData);
+	StoredArrayHelper.Resize(ArrayNum);
 
-	TArray< uint16 > ChangedLocal;
+	TArray<uint16> ChangedLocal;
 
 	uint16 LocalHandle = 0;
 
-	Data = ( uint8* )Array->GetData();
-	CompareData = ( uint8* )CompareArray->GetData();
+	Data = (uint8*)Array->GetData();
+	CompareData = (uint8*)CompareArray->GetData();
 
-	for ( int32 i = 0; i < ArrayNum; i++ )
+	for (int32 i = 0; i < ArrayNum; i++)
 	{
 		const int32 ElementOffset = i * Cmd.ElementSize;
 
 		const bool bNewForceFail = bForceFail || i >= CompareArrayNum;
 
-		LocalHandle = CompareProperties_r( CmdIndex + 1, Cmd.EndCmd - 1, CompareData + ElementOffset, Data + ElementOffset, ChangedLocal, LocalHandle, bIsInitial, bNewForceFail );
+		LocalHandle = CompareProperties_r(CmdIndex + 1, Cmd.EndCmd - 1, CompareData + ElementOffset, Data + ElementOffset, ChangedLocal, LocalHandle, bIsInitial, bNewForceFail);
 	}
 
-	if ( ChangedLocal.Num() )
+	if (ChangedLocal.Num())
 	{
-		Changed.Add( Handle );
-		Changed.Add( ChangedLocal.Num() );		// This is so we can jump over the array if we need to
-		Changed.Append( ChangedLocal );
-		Changed.Add( 0 );
+		Changed.Add(Handle);
+		Changed.Add(ChangedLocal.Num());		// This is so we can jump over the array if we need to
+		Changed.Append(ChangedLocal);
+		Changed.Add(0);
 	}
-	else if ( ArrayNum != CompareArrayNum )
+	else if (ArrayNum != CompareArrayNum)
 	{
 		// If nothing below us changed, we either shrunk, or we grew and our inner was an array that didn't have any elements
-		check( ArrayNum < CompareArrayNum || Cmds[CmdIndex + 1].Type == ERepLayoutCmdType::DynamicArray );
+		check(ArrayNum < CompareArrayNum || Cmds[CmdIndex + 1].Type == ERepLayoutCmdType::DynamicArray);
 
 		// Array got smaller, send the array handle to force array size change
-		Changed.Add( Handle );
-		Changed.Add( 0 );
-		Changed.Add( 0 );
+		Changed.Add(Handle);
+		Changed.Add(0);
+		Changed.Add(0);
 	}
 }
 
 bool FRepLayout::CompareProperties(
 	FRepChangelistState* RESTRICT	RepChangelistState,
 	const uint8* RESTRICT			Data,
-	const FReplicationFlags&		RepFlags ) const
+	const FReplicationFlags&		RepFlags) const
 {
-	SCOPE_CYCLE_COUNTER( STAT_NetReplicateDynamicPropCompareTime );
+	SCOPE_CYCLE_COUNTER(STAT_NetReplicateDynamicPropCompareTime);
 
 	RepChangelistState->CompareIndex++;
 
-	check( RepChangelistState->HistoryEnd - RepChangelistState->HistoryStart < FRepChangelistState::MAX_CHANGE_HISTORY );
+	check((RepChangelistState->HistoryEnd - RepChangelistState->HistoryStart) < FRepChangelistState::MAX_CHANGE_HISTORY);
 	const int32 HistoryIndex = RepChangelistState->HistoryEnd % FRepChangelistState::MAX_CHANGE_HISTORY;
 
 	FRepChangedHistory& NewHistoryItem = RepChangelistState->ChangeHistory[HistoryIndex];
@@ -583,9 +627,9 @@ bool FRepLayout::CompareProperties(
 	TArray<uint16>& Changed = NewHistoryItem.Changed;
 	Changed.Empty();
 
-	CompareProperties_r( 0, Cmds.Num() - 1, RepChangelistState->StaticBuffer.GetData(), Data, Changed, 0, RepFlags.bNetInitial, false );
+	CompareProperties_r(0, Cmds.Num() - 1, RepChangelistState->StaticBuffer.GetData(), Data, Changed, 0, RepFlags.bNetInitial, false);
 
-	if ( Changed.Num() == 0 )
+	if (Changed.Num() == 0)
 	{
 		return false;
 	}
@@ -595,7 +639,7 @@ bool FRepLayout::CompareProperties(
 	//
 
 	// Null terminator
-	Changed.Add( 0 );
+	Changed.Add(0);
 
 	// Move end pointer
 	RepChangelistState->HistoryEnd++;
@@ -604,7 +648,7 @@ bool FRepLayout::CompareProperties(
 	RepChangelistState->SharedSerialization.Reset();
 
 	// If we're full, merge the oldest up, so we always have room for a new entry
-	if ( RepChangelistState->HistoryEnd - RepChangelistState->HistoryStart == FRepChangelistState::MAX_CHANGE_HISTORY )
+	if ((RepChangelistState->HistoryEnd - RepChangelistState->HistoryStart) == FRepChangelistState::MAX_CHANGE_HISTORY)
 	{
 		const int32 FirstHistoryIndex = RepChangelistState->HistoryStart % FRepChangelistState::MAX_CHANGE_HISTORY;
 
@@ -612,30 +656,33 @@ bool FRepLayout::CompareProperties(
 
 		const int32 SecondHistoryIndex = RepChangelistState->HistoryStart % FRepChangelistState::MAX_CHANGE_HISTORY;
 
-		TArray< uint16 >& FirstChangelistRef = RepChangelistState->ChangeHistory[FirstHistoryIndex].Changed;
-		TArray< uint16 > SecondChangelistCopy = MoveTemp(RepChangelistState->ChangeHistory[SecondHistoryIndex].Changed);
+		TArray<uint16>& FirstChangelistRef = RepChangelistState->ChangeHistory[FirstHistoryIndex].Changed;
+		TArray<uint16> SecondChangelistCopy = MoveTemp(RepChangelistState->ChangeHistory[SecondHistoryIndex].Changed);
 
-		MergeChangeList( Data, FirstChangelistRef, SecondChangelistCopy, RepChangelistState->ChangeHistory[SecondHistoryIndex].Changed );
+		MergeChangeList(Data, FirstChangelistRef, SecondChangelistCopy, RepChangelistState->ChangeHistory[SecondHistoryIndex].Changed);
 	}
 
 	return true;
 }
 
-static FORCEINLINE void WritePropertyHandle( FNetBitWriter & Writer, uint16 Handle, bool bDoChecksum )
+static FORCEINLINE void WritePropertyHandle(
+	FNetBitWriter&	Writer,
+	uint16			Handle,
+	bool			bDoChecksum)
 {
 	const int NumStartingBits = Writer.GetNumBits();
 
 	uint32 LocalHandle = Handle;
-	Writer.SerializeIntPacked( LocalHandle );
+	Writer.SerializeIntPacked(LocalHandle);
 
 #ifdef ENABLE_PROPERTY_CHECKSUMS
-	if ( bDoChecksum )
+	if (bDoChecksum)
 	{
-		SerializeGenericChecksum( Writer );
+		SerializeGenericChecksum(Writer);
 	}
 #endif
 
-	NETWORK_PROFILER( GNetworkProfiler.TrackWritePropertyHandle( Writer.GetNumBits() - NumStartingBits, nullptr ) );
+	NETWORK_PROFILER(GNetworkProfiler.TrackWritePropertyHandle(Writer.GetNumBits() - NumStartingBits, nullptr));
 }
 
 bool FRepLayout::ReplicateProperties(
@@ -645,66 +692,67 @@ bool FRepLayout::ReplicateProperties(
 	UClass*							ObjectClass,
 	UActorChannel*					OwningChannel,
 	FNetBitWriter&					Writer,
-	const FReplicationFlags &		RepFlags ) const
+	const FReplicationFlags &		RepFlags) const
 {
-	SCOPE_CYCLE_COUNTER( STAT_NetReplicateDynamicPropTime );
+	SCOPE_CYCLE_COUNTER(STAT_NetReplicateDynamicPropTime);
 
-	check( ObjectClass == Owner );
+	check(ObjectClass == Owner);
 
 	FRepChangedPropertyTracker*	ChangeTracker = RepState->RepChangedPropertyTracker.Get();
 
 	// Rebuild conditional state if needed
-	if ( RepState->RepFlags.Value != RepFlags.Value )
+	if (RepState->RepFlags.Value != RepFlags.Value)
 	{
-		RebuildConditionalProperties( RepState, *ChangeTracker, RepFlags );
+		RebuildConditionalProperties(RepState, *ChangeTracker, RepFlags);
 
 		RepState->RepFlags.Value = RepFlags.Value;
 	}
 
-	if ( OwningChannel->Connection->bResendAllDataSinceOpen )
+	if (OwningChannel->Connection->bResendAllDataSinceOpen)
 	{
-		check( OwningChannel->Connection->InternalAck );
+		check(OwningChannel->Connection->InternalAck);
+
 		// If we are resending data since open, we don't want to affect the current state of channel/replication, so just do the minimum and send the data, and return
-		if ( RepState->LifetimeChangelist.Num() > 0 )
+		if (RepState->LifetimeChangelist.Num() > 0)
 		{
 			// Use a pruned version of the list, in case arrays changed size since the last time we replicated
-			TArray< uint16 > Pruned;
-			PruneChangeList( RepState, Data, RepState->LifetimeChangelist, Pruned );
-			RepState->LifetimeChangelist = MoveTemp( Pruned );
-			SendProperties_BackwardsCompatible( RepState, ChangeTracker, Data, OwningChannel->Connection, Writer, RepState->LifetimeChangelist );
+			TArray<uint16> Pruned;
+			PruneChangeList(RepState, Data, RepState->LifetimeChangelist, Pruned);
+			RepState->LifetimeChangelist = MoveTemp(Pruned);
+			SendProperties_BackwardsCompatible(RepState, ChangeTracker, Data, OwningChannel->Connection, Writer, RepState->LifetimeChangelist);
 			return true;
 		}
 
 		return false;
 	}
 
-	check( RepState->HistoryEnd >= RepState->HistoryStart );
-	check( RepState->HistoryEnd - RepState->HistoryStart < FRepState::MAX_CHANGE_HISTORY );
+	check(RepState->HistoryEnd >= RepState->HistoryStart);
+	check((RepState->HistoryEnd - RepState->HistoryStart) < FRepState::MAX_CHANGE_HISTORY);
 
-	const bool bFlushPreOpenAckHistory = RepState->OpenAckedCalled && RepState->PreOpenAckHistory.Num() > 0;
+	const bool bFlushPreOpenAckHistory = RepState->OpenAckedCalled && RepState->PreOpenAckHistory.Num()> 0;
 
 	const bool bCompareIndexSame = RepState->LastCompareIndex == RepChangelistState->CompareIndex;
 
 	RepState->LastCompareIndex = RepChangelistState->CompareIndex;
 
 	// We can early out if we know for sure there are no new changelists to send
-	if ( bCompareIndexSame || RepState->LastChangelistIndex == RepChangelistState->HistoryEnd )
+	if (bCompareIndexSame || RepState->LastChangelistIndex == RepChangelistState->HistoryEnd)
 	{
-		if ( RepState->NumNaks == 0 && !bFlushPreOpenAckHistory )
+		if (RepState->NumNaks == 0 && !bFlushPreOpenAckHistory)
 		{
 			// Nothing changed and there are no nak's, so just do normal housekeeping and remove acked history items
-			UpdateChangelistHistory( RepState, ObjectClass, Data, OwningChannel->Connection, NULL );
+			UpdateChangelistHistory(RepState, ObjectClass, Data, OwningChannel->Connection, NULL);
 			return false;
 		}
 	}
 
 	// Clamp to the valid history range (and log if we end up sending entire history, this should only happen if we get really far behind)
 	//	NOTE - The RepState->LastChangelistIndex != 0 should handle/ignore the JIP case
-	if ( RepState->LastChangelistIndex <= RepChangelistState->HistoryStart )
+	if (RepState->LastChangelistIndex <= RepChangelistState->HistoryStart)
 	{
-		if ( RepState->LastChangelistIndex != 0 )
+		if (RepState->LastChangelistIndex != 0)
 		{
-			UE_LOG( LogRep, Verbose, TEXT( "FRepLayout::ReplicatePropertiesUsingChangelistState: Entire history sent for: %s" ), *GetNameSafe( ObjectClass ) );
+			UE_LOG(LogRep, Verbose, TEXT("FRepLayout::ReplicatePropertiesUsingChangelistState: Entire history sent for: %s"), *GetNameSafe(ObjectClass));
 		}
 
 		RepState->LastChangelistIndex = RepChangelistState->HistoryStart;
@@ -714,37 +762,38 @@ bool FRepLayout::ReplicateProperties(
 
 	FRepChangedHistory& PossibleNewHistoryItem = RepState->ChangeHistory[PossibleNewHistoryIndex];
 
-	TArray< uint16 >& Changed = PossibleNewHistoryItem.Changed;
+	TArray<uint16>& Changed = PossibleNewHistoryItem.Changed;
 
-	check( Changed.Num() == 0 );		// Make sure this history item is actually inactive
+	// Make sure this history item is actually inactive
+	check(Changed.Num() == 0);
 
 	// Gather all change lists that are new since we last looked, and merge them all together into a single CL
-	for ( int32 i = RepState->LastChangelistIndex; i < RepChangelistState->HistoryEnd; i++ )
+	for (int32 i = RepState->LastChangelistIndex; i < RepChangelistState->HistoryEnd; i++)
 	{
 		const int32 HistoryIndex = i % FRepChangelistState::MAX_CHANGE_HISTORY;
 
 		FRepChangedHistory& HistoryItem = RepChangelistState->ChangeHistory[HistoryIndex];
 
-		TArray< uint16 > Temp = MoveTemp(Changed);
-		MergeChangeList( Data, HistoryItem.Changed, Temp, Changed );
+		TArray<uint16> Temp = MoveTemp(Changed);
+		MergeChangeList(Data, HistoryItem.Changed, Temp, Changed);
 	}
 
 	// We're all caught up now
 	RepState->LastChangelistIndex = RepChangelistState->HistoryEnd;
 
-	if ( Changed.Num() > 0 || RepState->NumNaks > 0 || bFlushPreOpenAckHistory )
+	if (Changed.Num() > 0 || RepState->NumNaks > 0 || bFlushPreOpenAckHistory)
 	{
 		RepState->HistoryEnd++;
 
-		UpdateChangelistHistory( RepState, ObjectClass, Data, OwningChannel->Connection, &Changed );
+		UpdateChangelistHistory(RepState, ObjectClass, Data, OwningChannel->Connection, &Changed);
 
 		// Merge in the PreOpenAckHistory (unreliable properties sent before the bunch was initially acked)
-		if ( bFlushPreOpenAckHistory )
+		if (bFlushPreOpenAckHistory)
 		{
-			for ( int32 i = 0; i < RepState->PreOpenAckHistory.Num(); i++ )
+			for (int32 i = 0; i < RepState->PreOpenAckHistory.Num(); i++)
 			{
-				TArray< uint16 > Temp = MoveTemp(Changed);
-				MergeChangeList( Data, RepState->PreOpenAckHistory[i].Changed, Temp, Changed );
+				TArray<uint16> Temp = MoveTemp(Changed);
+				MergeChangeList(Data, RepState->PreOpenAckHistory[i].Changed, Temp, Changed);
 			}
 			RepState->PreOpenAckHistory.Empty();
 		}
@@ -752,12 +801,12 @@ bool FRepLayout::ReplicateProperties(
 	else
 	{
 		// Nothing changed and there are no nak's, so just do normal housekeeping and remove acked history items
-		UpdateChangelistHistory( RepState, ObjectClass, Data, OwningChannel->Connection, nullptr );
-		return false;		// Nothing to send
+		UpdateChangelistHistory(RepState, ObjectClass, Data, OwningChannel->Connection, nullptr);
+		return false;
 	}
 
 	// At this point we should have a non empty change list
-	check( Changed.Num() > 0 );
+	check(Changed.Num() > 0);
 
 	// do not build shared state for InternalAck (demo) connections
 	if (!OwningChannel->Connection->InternalAck && (GNetSharedSerializedData != 0))
@@ -765,30 +814,30 @@ bool FRepLayout::ReplicateProperties(
 		// if no shared serialization info exists, build it
 		if (!RepChangelistState->SharedSerialization.IsValid())
 		{
-			BuildSharedSerialization( Data, Changed, true, RepChangelistState->SharedSerialization );
+			BuildSharedSerialization(Data, Changed, true, RepChangelistState->SharedSerialization);
 		}
 	}
 
 	const int32 NumBits = Writer.GetNumBits();
 
 	// Send the final merged change list
-	if ( OwningChannel->Connection->InternalAck )
+	if (OwningChannel->Connection->InternalAck)
 	{
 		// Remember all properties that have changed since this channel was first opened in case we need it (for bResendAllDataSinceOpen)
-		TArray< uint16 > Temp = MoveTemp(RepState->LifetimeChangelist);
-		MergeChangeList( Data, Changed, Temp, RepState->LifetimeChangelist );
+		TArray<uint16> Temp = MoveTemp(RepState->LifetimeChangelist);
+		MergeChangeList(Data, Changed, Temp, RepState->LifetimeChangelist);
 
-		SendProperties_BackwardsCompatible( RepState, ChangeTracker, Data, OwningChannel->Connection, Writer, Changed );
+		SendProperties_BackwardsCompatible(RepState, ChangeTracker, Data, OwningChannel->Connection, Writer, Changed);
 	}
 	else
 	{
-		SendProperties( RepState, ChangeTracker, Data, ObjectClass, Writer, Changed, RepChangelistState->SharedSerialization );
+		SendProperties(RepState, ChangeTracker, Data, ObjectClass, Writer, Changed, RepChangelistState->SharedSerialization);
 	}
 
 	// See if something actually sent (this may be false due to conditional checks inside the send properties function
 	const bool bSomethingSent = NumBits != Writer.GetNumBits();
 
-	if ( !bSomethingSent )
+	if (!bSomethingSent)
 	{
 		// We need to revert the change list in the history if nothing really sent (can happen due to condition checks)
 		Changed.Empty();
@@ -798,48 +847,55 @@ bool FRepLayout::ReplicateProperties(
 	return bSomethingSent;
 }
 
-void FRepLayout::UpdateChangelistHistory( FRepState * RepState, UClass * ObjectClass, const uint8* RESTRICT Data, UNetConnection* Connection, TArray< uint16 > * OutMerged ) const
+void FRepLayout::UpdateChangelistHistory(
+	FRepState*				RepState,
+	UClass*					ObjectClass,
+	const uint8* RESTRICT	Data,
+	UNetConnection*			Connection,
+	TArray<uint16>*			OutMerged) const
 {
-	check( RepState->HistoryEnd >= RepState->HistoryStart );
+	check(RepState->HistoryEnd>= RepState->HistoryStart);
 
-	const int32 HistoryCount	= RepState->HistoryEnd - RepState->HistoryStart;
-	const bool DumpHistory		= HistoryCount == FRepState::MAX_CHANGE_HISTORY;
-	const int32 AckPacketId		= Connection->OutAckPacketId;
+	const int32 HistoryCount = RepState->HistoryEnd - RepState->HistoryStart;
+	const bool DumpHistory = HistoryCount == FRepState::MAX_CHANGE_HISTORY;
+	const int32 AckPacketId = Connection->OutAckPacketId;
 
 	// If our buffer is currently full, forcibly send the entire history
-	if ( DumpHistory )
+	if (DumpHistory)
 	{
-		UE_LOG( LogRep, Verbose, TEXT( "FRepLayout::UpdateChangelistHistory: History overflow, forcing history dump %s, %s" ), *ObjectClass->GetName(), *Connection->Describe());
+		UE_LOG(LogRep, Verbose, TEXT("FRepLayout::UpdateChangelistHistory: History overflow, forcing history dump %s, %s"), *ObjectClass->GetName(), *Connection->Describe());
 	}
 
-	for ( int32 i = RepState->HistoryStart; i < RepState->HistoryEnd; i++ )
+	for (int32 i = RepState->HistoryStart; i < RepState->HistoryEnd; i++)
 	{
 		const int32 HistoryIndex = i % FRepState::MAX_CHANGE_HISTORY;
 
 		FRepChangedHistory & HistoryItem = RepState->ChangeHistory[ HistoryIndex ];
 
-		if ( HistoryItem.OutPacketIdRange.First == INDEX_NONE )
+		if (HistoryItem.OutPacketIdRange.First == INDEX_NONE)
 		{
-			continue;		//  Hasn't been initialized in PostReplicate yet
+			//  Hasn't been initialized in PostReplicate yet
+			continue;
 		}
 
-		check( HistoryItem.Changed.Num() > 0 );		// All active history items should contain a change list
+		// All active history items should contain a change list
+		check(HistoryItem.Changed.Num() > 0);
 
-		if ( AckPacketId >= HistoryItem.OutPacketIdRange.Last || HistoryItem.Resend || DumpHistory )
+		if (AckPacketId >= HistoryItem.OutPacketIdRange.Last || HistoryItem.Resend || DumpHistory)
 		{
-			if ( HistoryItem.Resend || DumpHistory )
+			if (HistoryItem.Resend || DumpHistory)
 			{
 				// Merge in nak'd change lists
-				check( OutMerged != NULL );
-				TArray< uint16 > Temp = MoveTemp(*OutMerged);
-				MergeChangeList( Data, HistoryItem.Changed, Temp, *OutMerged );
+				check(OutMerged != NULL);
+				TArray<uint16> Temp = MoveTemp(*OutMerged);
+				MergeChangeList(Data, HistoryItem.Changed, Temp, *OutMerged);
 				HistoryItem.Changed.Empty();
 
 #ifdef SANITY_CHECK_MERGES
-				SanityCheckChangeList( Data, *OutMerged );
+				SanityCheckChangeList(Data, *OutMerged);
 #endif
 
-				if ( HistoryItem.Resend )
+				if (HistoryItem.Resend)
 				{
 					HistoryItem.Resend = false;
 					RepState->NumNaks--;
@@ -855,84 +911,88 @@ void FRepLayout::UpdateChangelistHistory( FRepState * RepState, UClass * ObjectC
 	// Remove any tiling in the history markers to keep them from wrapping over time
 	const int32 NewHistoryCount	= RepState->HistoryEnd - RepState->HistoryStart;
 
-	check( NewHistoryCount <= FRepState::MAX_CHANGE_HISTORY );
+	check(NewHistoryCount <= FRepState::MAX_CHANGE_HISTORY);
 
-	RepState->HistoryStart	= RepState->HistoryStart % FRepState::MAX_CHANGE_HISTORY;
-	RepState->HistoryEnd	= RepState->HistoryStart + NewHistoryCount;
+	RepState->HistoryStart = RepState->HistoryStart % FRepState::MAX_CHANGE_HISTORY;
+	RepState->HistoryEnd = RepState->HistoryStart + NewHistoryCount;
 
-	check( RepState->NumNaks == 0 );	// Make sure we processed all the naks properly
+	// Make sure we processed all the naks properly
+	check(RepState->NumNaks == 0);
 }
 
-void FRepLayout::OpenAcked( FRepState * RepState ) const
+void FRepLayout::OpenAcked(FRepState * RepState) const
 {
-	check( RepState != NULL );
+	check(RepState != NULL);
 	RepState->OpenAckedCalled = true;
 }
 
-void FRepLayout::PostReplicate( FRepState * RepState, FPacketIdRange & PacketRange, bool bReliable ) const
+void FRepLayout::PostReplicate(
+	FRepState* RepState,
+	FPacketIdRange& PacketRange,
+	bool bReliable) const
 {
-	for ( int32 i = RepState->HistoryStart; i < RepState->HistoryEnd; i++ )
+	for (int32 i = RepState->HistoryStart; i < RepState->HistoryEnd; i++)
 	{
 		const int32 HistoryIndex = i % FRepState::MAX_CHANGE_HISTORY;
 
 		FRepChangedHistory & HistoryItem = RepState->ChangeHistory[ HistoryIndex ];
 
-		if ( HistoryItem.OutPacketIdRange.First == INDEX_NONE )
+		if (HistoryItem.OutPacketIdRange.First == INDEX_NONE)
 		{
-			check( HistoryItem.Changed.Num() > 0 );
-			check( !HistoryItem.Resend );
+			check(HistoryItem.Changed.Num()> 0);
+			check(!HistoryItem.Resend);
 
 			HistoryItem.OutPacketIdRange = PacketRange;
 
-			if ( !bReliable && !RepState->OpenAckedCalled )
+			if (!bReliable && !RepState->OpenAckedCalled)
 			{
-				RepState->PreOpenAckHistory.Add( HistoryItem );
+				RepState->PreOpenAckHistory.Add(HistoryItem);
 			}
 		}
 	}
 }
 
-void FRepLayout::ReceivedNak( FRepState * RepState, int32 NakPacketId ) const
+void FRepLayout::ReceivedNak(FRepState* RepState, int32 NakPacketId) const
 {
-	if ( RepState == NULL )
+	if (RepState == NULL)
 	{
 		return;		// I'm not 100% certain why this happens, the only think I can think of is this is a bNetTemporary?
 	}
 
-	for ( int32 i = RepState->HistoryStart; i < RepState->HistoryEnd; i++ )
+	for (int32 i = RepState->HistoryStart; i < RepState->HistoryEnd; i++)
 	{
 		const int32 HistoryIndex = i % FRepState::MAX_CHANGE_HISTORY;
 
 		FRepChangedHistory & HistoryItem = RepState->ChangeHistory[ HistoryIndex ];
 
-		if ( !HistoryItem.Resend && HistoryItem.OutPacketIdRange.InRange( NakPacketId ) )
+		if (!HistoryItem.Resend && HistoryItem.OutPacketIdRange.InRange(NakPacketId))
 		{
-			check( HistoryItem.Changed.Num() > 0 );
+			check(HistoryItem.Changed.Num() > 0);
 			HistoryItem.Resend = true;
 			RepState->NumNaks++;
 		}
 	}
 }
 
-bool FRepLayout::AllAcked( FRepState * RepState ) const
+bool FRepLayout::AllAcked(FRepState* RepState) const
 {	
-	if ( RepState->HistoryStart != RepState->HistoryEnd )
+	if (RepState->HistoryStart != RepState->HistoryEnd)
 	{
 		// We have change lists that haven't been acked
 		return false;
 	}
 
-	if ( RepState->NumNaks > 0 )
+	if (RepState->NumNaks> 0)
 	{
 		return false;
 	}
 
-	if ( !RepState->OpenAckedCalled )
+	if (!RepState->OpenAckedCalled)
 	{
 		return false;
 	}
 
-	if ( RepState->PreOpenAckHistory.Num() > 0 )
+	if (RepState->PreOpenAckHistory.Num()> 0)
 	{
 		return false;
 	}
@@ -940,14 +1000,14 @@ bool FRepLayout::AllAcked( FRepState * RepState ) const
 	return true;
 }
 
-bool FRepLayout::ReadyForDormancy( FRepState * RepState ) const
+bool FRepLayout::ReadyForDormancy(FRepState * RepState) const
 {
-	if ( RepState == NULL )
+	if (RepState == NULL)
 	{
 		return false;
 	}
 
-	return AllAcked( RepState );
+	return AllAcked(RepState);
 }
 
 void FRepLayout::SerializeObjectReplicatedProperties(UObject* Object, FBitArchive & Ar) const
@@ -974,23 +1034,23 @@ bool FRepHandleIterator::NextHandle()
 
 	Handle = ChangelistIterator.Changed[ChangelistIterator.ChangedIndex];
 
-	if ( Handle == 0 )
+	if (Handle == 0)
 	{
 		return false;		// Done
 	}
 
 	ChangelistIterator.ChangedIndex++;
 
-	if ( !ensure( ChangelistIterator.ChangedIndex < ChangelistIterator.Changed.Num() ) )
+	if (!ensure(ChangelistIterator.ChangedIndex <ChangelistIterator.Changed.Num()))
 	{
 		return false;
 	}
 
 	const int32 HandleMinusOne = Handle - 1;
 
-	ArrayIndex = ( ArrayElementSize > 0 && NumHandlesPerElement > 0 ) ? HandleMinusOne / NumHandlesPerElement : 0;
+	ArrayIndex = (ArrayElementSize> 0 && NumHandlesPerElement> 0) ? HandleMinusOne / NumHandlesPerElement : 0;
 
-	if ( ArrayIndex >= MaxArrayIndex )
+	if (ArrayIndex>= MaxArrayIndex)
 	{
 		return false;
 	}
@@ -1001,19 +1061,19 @@ bool FRepHandleIterator::NextHandle()
 
 	CmdIndex = HandleToCmdIndex[RelativeHandle].CmdIndex;
 
-	if ( !ensure( CmdIndex >= MinCmdIndex && CmdIndex < MaxCmdIndex ) )
+	if (!ensure(CmdIndex >= MinCmdIndex && CmdIndex < MaxCmdIndex))
 	{
 		return false;
 	}
 
 	const FRepLayoutCmd& Cmd = Cmds[CmdIndex];
 
-	if ( !ensure( Cmd.RelativeHandle - 1 == RelativeHandle ) )
+	if (!ensure(Cmd.RelativeHandle - 1 == RelativeHandle))
 	{
 		return false;
 	}
 
-	if ( !ensure( Cmd.Type != ERepLayoutCmdType::Return ) )
+	if (!ensure(Cmd.Type != ERepLayoutCmdType::Return))
 	{
 		return false;
 	}
@@ -1026,7 +1086,7 @@ bool FRepHandleIterator::JumpOverArray()
 	const int32 ArrayChangedCount = ChangelistIterator.Changed[ChangelistIterator.ChangedIndex++];
 	ChangelistIterator.ChangedIndex += ArrayChangedCount;
 
-	if ( !ensure( ChangelistIterator.Changed[ChangelistIterator.ChangedIndex] == 0 ) )
+	if (!ensure(ChangelistIterator.Changed[ChangelistIterator.ChangedIndex] == 0))
 	{
 		return false;
 	}
@@ -1044,113 +1104,118 @@ int32 FRepHandleIterator::PeekNextHandle() const
 class FScopedIteratorArrayTracker
 {
 public:
-	FScopedIteratorArrayTracker( FRepHandleIterator* InCmdIndexIterator )
+	FScopedIteratorArrayTracker(FRepHandleIterator* InCmdIndexIterator)
 	{
 		CmdIndexIterator = InCmdIndexIterator;
 
-		if ( CmdIndexIterator )
+		if (CmdIndexIterator)
 		{
-			ArrayChangedCount	= CmdIndexIterator->ChangelistIterator.Changed[CmdIndexIterator->ChangelistIterator.ChangedIndex++];
-			OldChangedIndex		= CmdIndexIterator->ChangelistIterator.ChangedIndex;
+			ArrayChangedCount = CmdIndexIterator->ChangelistIterator.Changed[CmdIndexIterator->ChangelistIterator.ChangedIndex++];
+			OldChangedIndex = CmdIndexIterator->ChangelistIterator.ChangedIndex;
 		}
 	}
 
 	~FScopedIteratorArrayTracker()
 	{
-		if ( CmdIndexIterator )
+		if (CmdIndexIterator)
 		{
-			check( CmdIndexIterator->ChangelistIterator.ChangedIndex - OldChangedIndex <= ArrayChangedCount );
+			check(CmdIndexIterator->ChangelistIterator.ChangedIndex - OldChangedIndex <= ArrayChangedCount);
 			CmdIndexIterator->ChangelistIterator.ChangedIndex = OldChangedIndex + ArrayChangedCount;
-			check( CmdIndexIterator->PeekNextHandle() == 0 );
+			check(CmdIndexIterator->PeekNextHandle() == 0);
 			CmdIndexIterator->ChangelistIterator.ChangedIndex++;
 		}
 	}
 
-	FRepHandleIterator*	CmdIndexIterator;
-	int32				ArrayChangedCount;
-	int32				OldChangedIndex;
+	FRepHandleIterator* CmdIndexIterator;
+	int32 ArrayChangedCount;
+	int32 OldChangedIndex;
 };
 
 void FRepLayout::MergeChangeList_r(
 	FRepHandleIterator&		RepHandleIterator1,
 	FRepHandleIterator&		RepHandleIterator2,
 	const uint8* RESTRICT	SourceData,
-	TArray< uint16 >&		OutChanged ) const
+	TArray<uint16>&			OutChanged) const
 {
-	while ( true )
+	while (true)
 	{
 		const int32 NextHandle1 = RepHandleIterator1.PeekNextHandle();
 		const int32 NextHandle2 = RepHandleIterator2.PeekNextHandle();
 
-		if ( NextHandle1 == 0 && NextHandle2 == 0 )
+		if (NextHandle1 == 0 && NextHandle2 == 0)
 		{
-			break;		// Done
+			// Done
+			break;
 		}
 
-		if ( NextHandle2 == 0 )
+		if (NextHandle2 == 0)
 		{
-			PruneChangeList_r( RepHandleIterator1, SourceData, OutChanged );
+			PruneChangeList_r(RepHandleIterator1, SourceData, OutChanged);
 			return;
 		}
-		else if ( NextHandle1 == 0 )
+		else if (NextHandle1 == 0)
 		{
-			PruneChangeList_r( RepHandleIterator2, SourceData, OutChanged );
+			PruneChangeList_r(RepHandleIterator2, SourceData, OutChanged);
 			return;
 		}
 
 		FRepHandleIterator* ActiveIterator1 = nullptr;
 		FRepHandleIterator* ActiveIterator2 = nullptr;
 
-		int32 CmdIndex		= INDEX_NONE;
-		int32 ArrayOffset	= INDEX_NONE;
+		int32 CmdIndex = INDEX_NONE;
+		int32 ArrayOffset = INDEX_NONE;
 
-		if ( NextHandle1 < NextHandle2 )
+		if (NextHandle1 < NextHandle2)
 		{
-			if ( !RepHandleIterator1.NextHandle() )
+			if (!RepHandleIterator1.NextHandle())
 			{
-				break;		// Array overflow
+				// Array overflow
+				break;
 			}
 
-			OutChanged.Add( NextHandle1 );
+			OutChanged.Add(NextHandle1);
 
-			CmdIndex	= RepHandleIterator1.CmdIndex;
+			CmdIndex = RepHandleIterator1.CmdIndex;
 			ArrayOffset = RepHandleIterator1.ArrayOffset;
 
 			ActiveIterator1 = &RepHandleIterator1;
 		}
-		else if ( NextHandle2 < NextHandle1 )
+		else if (NextHandle2 < NextHandle1)
 		{
-			if ( !RepHandleIterator2.NextHandle() )
+			if (!RepHandleIterator2.NextHandle())
 			{
-				break;		// Array overflow
+				// Array overflow
+				break;
 			}
 
-			OutChanged.Add( NextHandle2 );
+			OutChanged.Add(NextHandle2);
 
-			CmdIndex	= RepHandleIterator2.CmdIndex;
+			CmdIndex = RepHandleIterator2.CmdIndex;
 			ArrayOffset = RepHandleIterator2.ArrayOffset;
 
 			ActiveIterator2 = &RepHandleIterator2;
 		}
 		else
 		{
-			check( NextHandle1 == NextHandle2 );
+			check(NextHandle1 == NextHandle2);
 
-			if ( !RepHandleIterator1.NextHandle() )
+			if (!RepHandleIterator1.NextHandle())
 			{
-				break;		// Array overflow
+				// Array overflow
+				break;
 			}
 
-			if ( !ensure( RepHandleIterator2.NextHandle() ) )
+			if (!ensure(RepHandleIterator2.NextHandle()))
 			{
-				break;		// Array overflow
+				// Array overflow
+				break;
 			}
 
-			check( RepHandleIterator1.CmdIndex == RepHandleIterator2.CmdIndex );
+			check(RepHandleIterator1.CmdIndex == RepHandleIterator2.CmdIndex);
 
-			OutChanged.Add( NextHandle1 );
+			OutChanged.Add(NextHandle1);
 
-			CmdIndex	= RepHandleIterator1.CmdIndex;
+			CmdIndex = RepHandleIterator1.CmdIndex;
 			ArrayOffset = RepHandleIterator1.ArrayOffset;
 
 			ActiveIterator1 = &RepHandleIterator1;
@@ -1159,44 +1224,44 @@ void FRepLayout::MergeChangeList_r(
 
 		const FRepLayoutCmd& Cmd = Cmds[CmdIndex];
 
-		if ( Cmd.Type == ERepLayoutCmdType::DynamicArray )
+		if (Cmd.Type == ERepLayoutCmdType::DynamicArray)
 		{
 			const uint8* Data = SourceData + ArrayOffset + Cmd.Offset;
 
-			const FScriptArray* Array = ( FScriptArray * )Data;
+			const FScriptArray* Array = (FScriptArray *)Data;
 
-			FScopedIteratorArrayTracker ArrayTracker1( ActiveIterator1 );
-			FScopedIteratorArrayTracker ArrayTracker2( ActiveIterator2 );
+			FScopedIteratorArrayTracker ArrayTracker1(ActiveIterator1);
+			FScopedIteratorArrayTracker ArrayTracker2(ActiveIterator2);
 
 			const int32 OriginalChangedNum	= OutChanged.AddUninitialized();
 
-			const uint8* NewData = ( uint8* )Array->GetData();
+			const uint8* NewData = (uint8*)Array->GetData();
 
-			TArray< FHandleToCmdIndex >& ArrayHandleToCmdIndex = ActiveIterator1 ? *ActiveIterator1->HandleToCmdIndex[Cmd.RelativeHandle - 1].HandleToCmdIndex : *ActiveIterator2->HandleToCmdIndex[Cmd.RelativeHandle - 1].HandleToCmdIndex; //-V595
+			TArray<FHandleToCmdIndex>& ArrayHandleToCmdIndex = ActiveIterator1 ? *ActiveIterator1->HandleToCmdIndex[Cmd.RelativeHandle - 1].HandleToCmdIndex : *ActiveIterator2->HandleToCmdIndex[Cmd.RelativeHandle - 1].HandleToCmdIndex; //-V595
 
-			if ( !ActiveIterator1 )
+			if (!ActiveIterator1)
 			{
-				FRepHandleIterator ArrayIterator2( ActiveIterator2->ChangelistIterator, Cmds, ArrayHandleToCmdIndex, Cmd.ElementSize, Array->Num(), CmdIndex + 1, Cmd.EndCmd - 1 );
-				PruneChangeList_r( ArrayIterator2, NewData, OutChanged );
+				FRepHandleIterator ArrayIterator2(ActiveIterator2->ChangelistIterator, Cmds, ArrayHandleToCmdIndex, Cmd.ElementSize, Array->Num(), CmdIndex + 1, Cmd.EndCmd - 1);
+				PruneChangeList_r(ArrayIterator2, NewData, OutChanged);
 			}
-			else if ( !ActiveIterator2 )
+			else if (!ActiveIterator2)
 			{
-				FRepHandleIterator ArrayIterator1( ActiveIterator1->ChangelistIterator, Cmds, ArrayHandleToCmdIndex, Cmd.ElementSize, Array->Num(), CmdIndex + 1, Cmd.EndCmd - 1 );
-				PruneChangeList_r( ArrayIterator1, NewData, OutChanged );
+				FRepHandleIterator ArrayIterator1(ActiveIterator1->ChangelistIterator, Cmds, ArrayHandleToCmdIndex, Cmd.ElementSize, Array->Num(), CmdIndex + 1, Cmd.EndCmd - 1);
+				PruneChangeList_r(ArrayIterator1, NewData, OutChanged);
 			}
 			else
 			{
-				FRepHandleIterator ArrayIterator1( ActiveIterator1->ChangelistIterator, Cmds, ArrayHandleToCmdIndex, Cmd.ElementSize, Array->Num(), CmdIndex + 1, Cmd.EndCmd - 1 );
-				FRepHandleIterator ArrayIterator2( ActiveIterator2->ChangelistIterator, Cmds, ArrayHandleToCmdIndex, Cmd.ElementSize, Array->Num(), CmdIndex + 1, Cmd.EndCmd - 1 );
+				FRepHandleIterator ArrayIterator1(ActiveIterator1->ChangelistIterator, Cmds, ArrayHandleToCmdIndex, Cmd.ElementSize, Array->Num(), CmdIndex + 1, Cmd.EndCmd - 1);
+				FRepHandleIterator ArrayIterator2(ActiveIterator2->ChangelistIterator, Cmds, ArrayHandleToCmdIndex, Cmd.ElementSize, Array->Num(), CmdIndex + 1, Cmd.EndCmd - 1);
 
-				MergeChangeList_r( ArrayIterator1, ArrayIterator2, NewData, OutChanged );
+				MergeChangeList_r(ArrayIterator1, ArrayIterator2, NewData, OutChanged);
 			}
 
 			// Patch in the jump offset
-			OutChanged[OriginalChangedNum] = OutChanged.Num() - ( OriginalChangedNum + 1 );
+			OutChanged[OriginalChangedNum] = OutChanged.Num() - (OriginalChangedNum + 1);
 
 			// Add the array terminator
-			OutChanged.Add( 0 );
+			OutChanged.Add(0);
 		}
 	}
 }
@@ -1204,39 +1269,39 @@ void FRepLayout::MergeChangeList_r(
 void FRepLayout::PruneChangeList_r(
 	FRepHandleIterator&		RepHandleIterator,
 	const uint8* RESTRICT	SourceData,
-	TArray< uint16 >&		OutChanged ) const
+	TArray<uint16>&			OutChanged) const
 {
-	while ( RepHandleIterator.NextHandle() )
+	while (RepHandleIterator.NextHandle())
 	{
-		OutChanged.Add( RepHandleIterator.Handle );
+		OutChanged.Add(RepHandleIterator.Handle);
 
-		const int32 CmdIndex	= RepHandleIterator.CmdIndex;
+		const int32 CmdIndex = RepHandleIterator.CmdIndex;
 		const int32 ArrayOffset = RepHandleIterator.ArrayOffset;
 
 		const FRepLayoutCmd& Cmd = Cmds[CmdIndex];
 
-		if ( Cmd.Type == ERepLayoutCmdType::DynamicArray )
+		if (Cmd.Type == ERepLayoutCmdType::DynamicArray)
 		{
 			const uint8* Data = SourceData + ArrayOffset + Cmd.Offset;
 
-			const FScriptArray* Array = ( FScriptArray * )Data;
+			const FScriptArray* Array = (FScriptArray *)Data;
 
-			FScopedIteratorArrayTracker ArrayTracker( &RepHandleIterator );
+			FScopedIteratorArrayTracker ArrayTracker(&RepHandleIterator);
 
 			const int32 OriginalChangedNum = OutChanged.AddUninitialized();
 
-			const uint8* NewData = ( uint8* )Array->GetData();
+			const uint8* NewData = (uint8*)Array->GetData();
 
-			TArray< FHandleToCmdIndex >& ArrayHandleToCmdIndex = *RepHandleIterator.HandleToCmdIndex[Cmd.RelativeHandle - 1].HandleToCmdIndex;
+			TArray<FHandleToCmdIndex>& ArrayHandleToCmdIndex = *RepHandleIterator.HandleToCmdIndex[Cmd.RelativeHandle - 1].HandleToCmdIndex;
 
-			FRepHandleIterator ArrayIterator( RepHandleIterator.ChangelistIterator, Cmds, ArrayHandleToCmdIndex, Cmd.ElementSize, Array->Num(), CmdIndex + 1, Cmd.EndCmd - 1 );
-			PruneChangeList_r( ArrayIterator, NewData, OutChanged );
+			FRepHandleIterator ArrayIterator(RepHandleIterator.ChangelistIterator, Cmds, ArrayHandleToCmdIndex, Cmd.ElementSize, Array->Num(), CmdIndex + 1, Cmd.EndCmd - 1);
+			PruneChangeList_r(ArrayIterator, NewData, OutChanged);
 
 			// Patch in the jump offset
-			OutChanged[OriginalChangedNum] = OutChanged.Num() - ( OriginalChangedNum + 1 );
+			OutChanged[OriginalChangedNum] = OutChanged.Num() - (OriginalChangedNum + 1);
 
 			// Add the array terminator
-			OutChanged.Add( 0 );
+			OutChanged.Add(0);
 		}
 	}
 }
@@ -1265,22 +1330,22 @@ const FRepSerializedPropertyInfo* FRepSerializationSharedInfo::WriteSharedProper
 
 	if (bWriteHandle)
 	{
-		WritePropertyHandle( *SerializedProperties, Handle, bDoChecksum );
+		WritePropertyHandle(*SerializedProperties, Handle, bDoChecksum);
 	}
 
 	SharedPropInfo.PropBitOffset = SerializedProperties->GetNumBits();
 
 	// This property changed, so send it
-	Cmd.Property->NetSerializeItem( *SerializedProperties, nullptr, ( void* )Data );
+	Cmd.Property->NetSerializeItem(*SerializedProperties, nullptr, (void*)Data);
 
 	const int64 NumPropEndBits = SerializedProperties->GetNumBits();
 
 	SharedPropInfo.PropBitLength = NumPropEndBits - SharedPropInfo.PropBitOffset;
 
 #ifdef ENABLE_PROPERTY_CHECKSUMS
-	if ( bDoChecksum )
+	if (bDoChecksum)
 	{
-		SerializeReadWritePropertyChecksum( Cmd, CmdIndex, Data, *SerializedProperties );
+		SerializeReadWritePropertyChecksum(Cmd, CmdIndex, Data, *SerializedProperties);
 	}
 #endif
 
@@ -1290,26 +1355,26 @@ const FRepSerializedPropertyInfo* FRepSerializationSharedInfo::WriteSharedProper
 }
 
 void FRepLayout::SendProperties_r(
-	FRepState*	RESTRICT			RepState,
-	FRepChangedPropertyTracker*		ChangedTracker,
-	FNetBitWriter&					Writer,
-	const bool						bDoChecksum,
-	FRepHandleIterator&				HandleIterator,
-	const uint8* RESTRICT			SourceData,
-	const int32						ArrayDepth,
-	const FRepSerializationSharedInfo&	SharedInfo ) const
+	FRepState*	RESTRICT				RepState,
+	FRepChangedPropertyTracker*			ChangedTracker,
+	FNetBitWriter&						Writer,
+	const bool							bDoChecksum,
+	FRepHandleIterator&					HandleIterator,
+	const uint8* RESTRICT				SourceData,
+	const int32							ArrayDepth,
+	const FRepSerializationSharedInfo&	SharedInfo) const
 {
-	while ( HandleIterator.NextHandle() )
+	while (HandleIterator.NextHandle())
 	{
 		const FRepLayoutCmd& Cmd = Cmds[HandleIterator.CmdIndex];
 
 		const FRepParentCmd& ParentCmd = Parents[Cmd.ParentIndex];
 
-		if ( !RepState->ConditionMap[ParentCmd.Condition] || !ChangedTracker->Parents[Cmd.ParentIndex].Active )
+		if (!RepState->ConditionMap[ParentCmd.Condition] || !ChangedTracker->Parents[Cmd.ParentIndex].Active)
 		{
-			if ( Cmd.Type == ERepLayoutCmdType::DynamicArray )
+			if (Cmd.Type == ERepLayoutCmdType::DynamicArray)
 			{
-				if ( !HandleIterator.JumpOverArray() )
+				if (!HandleIterator.JumpOverArray())
 				{
 					break;
 				}
@@ -1320,11 +1385,11 @@ void FRepLayout::SendProperties_r(
 
 		const uint8* Data = SourceData + HandleIterator.ArrayOffset + Cmd.Offset;
 
-		if ( Cmd.Type == ERepLayoutCmdType::DynamicArray )
+		if (Cmd.Type == ERepLayoutCmdType::DynamicArray)
 		{
-			WritePropertyHandle( Writer, HandleIterator.Handle, bDoChecksum );
+			WritePropertyHandle(Writer, HandleIterator.Handle, bDoChecksum);
 
-			const FScriptArray* Array = ( FScriptArray * )Data;
+			const FScriptArray* Array = (FScriptArray *)Data;
 
 			// Write array num
 			uint16 ArrayNum = Array->Num();
@@ -1337,30 +1402,30 @@ void FRepLayout::SendProperties_r(
 
 			const int32 OldChangedIndex = HandleIterator.ChangelistIterator.ChangedIndex;
 
-			const uint8* NewData = ( uint8* )Array->GetData();
+			const uint8* NewData = (uint8*)Array->GetData();
 
-			TArray< FHandleToCmdIndex >& ArrayHandleToCmdIndex = *HandleIterator.HandleToCmdIndex[Cmd.RelativeHandle - 1].HandleToCmdIndex;
+			TArray<FHandleToCmdIndex>& ArrayHandleToCmdIndex = *HandleIterator.HandleToCmdIndex[Cmd.RelativeHandle - 1].HandleToCmdIndex;
 
-			FRepHandleIterator ArrayHandleIterator( HandleIterator.ChangelistIterator, Cmds, ArrayHandleToCmdIndex, Cmd.ElementSize, ArrayNum, HandleIterator.CmdIndex + 1, Cmd.EndCmd - 1 );
+			FRepHandleIterator ArrayHandleIterator(HandleIterator.ChangelistIterator, Cmds, ArrayHandleToCmdIndex, Cmd.ElementSize, ArrayNum, HandleIterator.CmdIndex + 1, Cmd.EndCmd - 1);
 
-			check( ArrayHandleIterator.ArrayElementSize > 0 );
-			check( ArrayHandleIterator.NumHandlesPerElement > 0 );
+			check(ArrayHandleIterator.ArrayElementSize> 0);
+			check(ArrayHandleIterator.NumHandlesPerElement> 0);
 
-			SendProperties_r( RepState, ChangedTracker, Writer, bDoChecksum, ArrayHandleIterator, NewData, ArrayDepth + 1, SharedInfo );
+			SendProperties_r(RepState, ChangedTracker, Writer, bDoChecksum, ArrayHandleIterator, NewData, ArrayDepth + 1, SharedInfo);
 
-			check( HandleIterator.ChangelistIterator.ChangedIndex - OldChangedIndex == ArrayChangedCount );				// Make sure we read correct amount
-			check( HandleIterator.ChangelistIterator.Changed[HandleIterator.ChangelistIterator.ChangedIndex] == 0 );	// Make sure we are at the end
+			check(HandleIterator.ChangelistIterator.ChangedIndex - OldChangedIndex == ArrayChangedCount);				// Make sure we read correct amount
+			check(HandleIterator.ChangelistIterator.Changed[HandleIterator.ChangelistIterator.ChangedIndex] == 0);	// Make sure we are at the end
 
 			HandleIterator.ChangelistIterator.ChangedIndex++;
 
-			WritePropertyHandle( Writer, 0, bDoChecksum );		// Signify end of dynamic array
+			WritePropertyHandle(Writer, 0, bDoChecksum);		// Signify end of dynamic array
 			continue;
 		}
 
 #if !(UE_BUILD_SHIPPING || UE_BUILD_TEST)
-		if ( GDoReplicationContextString > 0 )
+		if (GDoReplicationContextString> 0)
 		{
-			Writer.PackageMap->SetDebugContextString( FString::Printf( TEXT( "%s - %s" ), *Owner->GetPathName(), *Cmd.Property->GetPathName() ) );
+			Writer.PackageMap->SetDebugContextString(FString::Printf(TEXT("%s - %s"), *Owner->GetPathName(), *Cmd.Property->GetPathName()));
 		}
 #endif
 
@@ -1385,12 +1450,12 @@ void FRepLayout::SendProperties_r(
 			{
 				FBitWriterMark BitWriterMark(Writer);
 			
-				WritePropertyHandle( Writer, HandleIterator.Handle, bDoChecksum );
-				Cmd.Property->NetSerializeItem( Writer, Writer.PackageMap, ( void* )Data );
+				WritePropertyHandle(Writer, HandleIterator.Handle, bDoChecksum);
+				Cmd.Property->NetSerializeItem(Writer, Writer.PackageMap, (void*)Data);
 #ifdef ENABLE_PROPERTY_CHECKSUMS
-				if ( bDoChecksum )
+				if (bDoChecksum)
 				{
-					SerializeReadWritePropertyChecksum( Cmd, HandleIterator.CmdIndex, Data, Writer );
+					SerializeReadWritePropertyChecksum(Cmd, HandleIterator.CmdIndex, Data, Writer);
 				}
 #endif
 				TArray<uint8> StandardBuffer;
@@ -1413,32 +1478,32 @@ void FRepLayout::SendProperties_r(
 				Writer.SerializeBitsWithOffset(SharedInfo.SerializedProperties->GetData(), SharedPropInfo->BitOffset, SharedPropInfo->BitLength);
 			}
 
-			NETWORK_PROFILER( GNetworkProfiler.TrackReplicateProperty( ParentCmd.Property, SharedPropInfo->PropBitLength, nullptr ) );
+			NETWORK_PROFILER(GNetworkProfiler.TrackReplicateProperty(ParentCmd.Property, SharedPropInfo->PropBitLength, nullptr));
 		}
 		else
 		{
 			GNumSharedSerializationMiss++;
-			WritePropertyHandle( Writer, HandleIterator.Handle, bDoChecksum );
+			WritePropertyHandle(Writer, HandleIterator.Handle, bDoChecksum);
 
 			const int32 NumStartBits = Writer.GetNumBits();
 
 			// This property changed, so send it
-			Cmd.Property->NetSerializeItem( Writer, Writer.PackageMap, ( void* )Data );
+			Cmd.Property->NetSerializeItem(Writer, Writer.PackageMap, (void*)Data);
 
 			const int32 NumEndBits = Writer.GetNumBits();
 
-			NETWORK_PROFILER( GNetworkProfiler.TrackReplicateProperty( ParentCmd.Property, NumEndBits - NumStartBits, nullptr ) );
+			NETWORK_PROFILER(GNetworkProfiler.TrackReplicateProperty(ParentCmd.Property, NumEndBits - NumStartBits, nullptr));
 
 #ifdef ENABLE_PROPERTY_CHECKSUMS
-			if ( bDoChecksum )
+			if (bDoChecksum)
 			{
-				SerializeReadWritePropertyChecksum( Cmd, HandleIterator.CmdIndex, Data, Writer );
+				SerializeReadWritePropertyChecksum(Cmd, HandleIterator.CmdIndex, Data, Writer);
 			}
 #endif
 		}
 
 #if !(UE_BUILD_SHIPPING || UE_BUILD_TEST)
-		if ( GDoReplicationContextString > 0 )
+		if (GDoReplicationContextString> 0)
 		{
 			Writer.PackageMap->ClearDebugContextString();
 		}
@@ -1447,13 +1512,13 @@ void FRepLayout::SendProperties_r(
 }
 
 void FRepLayout::SendProperties(
-	FRepState *	RESTRICT		RepState,
-	FRepChangedPropertyTracker* ChangedTracker,
-	const uint8* RESTRICT		Data,
-	UClass *					ObjectClass,
-	FNetBitWriter&				Writer,
-	TArray< uint16 > &			Changed,
-	const FRepSerializationSharedInfo& SharedInfo) const
+	FRepState *	RESTRICT				RepState,
+	FRepChangedPropertyTracker*			ChangedTracker,
+	const uint8* RESTRICT				Data,
+	UClass *							ObjectClass,
+	FNetBitWriter&						Writer,
+	TArray<uint16> &					Changed,
+	const FRepSerializationSharedInfo&	SharedInfo) const
 {
 	SCOPE_CYCLE_COUNTER(STAT_NetReplicateDynamicPropSendTime);
 
@@ -1463,60 +1528,63 @@ void FRepLayout::SendProperties(
 	const bool bDoChecksum = false;
 #endif
 
-	FBitWriterMark Mark( Writer );
+	FBitWriterMark Mark(Writer);
 
 #ifdef ENABLE_PROPERTY_CHECKSUMS
-	Writer.WriteBit( bDoChecksum ? 1 : 0 );
+	Writer.WriteBit(bDoChecksum ? 1 : 0);
 #endif
 
 	const int32 NumBits = Writer.GetNumBits();
 
-	FChangelistIterator ChangelistIterator( Changed, 0 );
-	FRepHandleIterator HandleIterator( ChangelistIterator, Cmds, BaseHandleToCmdIndex, 0, 1, 0, Cmds.Num() - 1 );
+	FChangelistIterator ChangelistIterator(Changed, 0);
+	FRepHandleIterator HandleIterator(ChangelistIterator, Cmds, BaseHandleToCmdIndex, 0, 1, 0, Cmds.Num() - 1);
 
-	SendProperties_r( RepState, ChangedTracker, Writer, bDoChecksum, HandleIterator, Data, 0, SharedInfo );
+	SendProperties_r(RepState, ChangedTracker, Writer, bDoChecksum, HandleIterator, Data, 0, SharedInfo);
 
-	if ( NumBits != Writer.GetNumBits() )
+	if (NumBits != Writer.GetNumBits())
 	{
 		// We actually wrote stuff
-		WritePropertyHandle( Writer, 0, bDoChecksum );
+		WritePropertyHandle(Writer, 0, bDoChecksum);
 	}
 	else
 	{
-		Mark.Pop( Writer );
+		Mark.Pop(Writer);
 	}
 }
 
-static FORCEINLINE void WritePropertyHandle_BackwardsCompatible( FNetBitWriter & Writer, uint32 NetFieldExportHandle, bool bDoChecksum )
+static FORCEINLINE void WritePropertyHandle_BackwardsCompatible(
+	FNetBitWriter&	Writer,
+	uint32			NetFieldExportHandle,
+	bool			bDoChecksum)
 {
 	const int NumStartingBits = Writer.GetNumBits();
 
-	Writer.SerializeIntPacked( NetFieldExportHandle );
+	Writer.SerializeIntPacked(NetFieldExportHandle);
 
 #ifdef ENABLE_PROPERTY_CHECKSUMS
-	if ( bDoChecksum )
+	if (bDoChecksum)
 	{
-		SerializeGenericChecksum( Writer );
+		SerializeGenericChecksum(Writer);
 	}
 #endif
 
-	NETWORK_PROFILER( GNetworkProfiler.TrackWritePropertyHandle( Writer.GetNumBits() - NumStartingBits, nullptr ) );
+	NETWORK_PROFILER(GNetworkProfiler.TrackWritePropertyHandle(Writer.GetNumBits() - NumStartingBits, nullptr));
 }
 
-TSharedPtr< FNetFieldExportGroup > FRepLayout::CreateNetfieldExportGroup() const
+TSharedPtr<FNetFieldExportGroup> FRepLayout::CreateNetfieldExportGroup() const
 {
-	TSharedPtr< FNetFieldExportGroup > NetFieldExportGroup = TSharedPtr< FNetFieldExportGroup >( new FNetFieldExportGroup() );
+	TSharedPtr<FNetFieldExportGroup> NetFieldExportGroup = TSharedPtr<FNetFieldExportGroup>(new FNetFieldExportGroup());
 
 	NetFieldExportGroup->PathName = Owner->GetPathName();
-	NetFieldExportGroup->NetFieldExports.SetNum( Cmds.Num() );
+	NetFieldExportGroup->NetFieldExports.SetNum(Cmds.Num());
 
-	for ( int32 i = 0; i < Cmds.Num(); i++ )
+	for (int32 i = 0; i < Cmds.Num(); i++)
 	{
 		FNetFieldExport NetFieldExport(
 			i,
 			Cmds[i].CompatibleChecksum,
-			Cmds[i].Property ? Cmds[i].Property->GetName() : TEXT( "" ),
-			Cmds[i].Property ? Cmds[i].Property->GetCPPType( nullptr, 0 ) : TEXT( "" ) );
+			Cmds[i].Property ? Cmds[i].Property->GetName() : TEXT(""),
+			Cmds[i].Property ? Cmds[i].Property->GetCPPType(nullptr, 0) : TEXT(""));
 
 		NetFieldExportGroup->NetFieldExports[i] = NetFieldExport;
 	}
@@ -1524,41 +1592,47 @@ TSharedPtr< FNetFieldExportGroup > FRepLayout::CreateNetfieldExportGroup() const
 	return NetFieldExportGroup;
 }
 
-static FORCEINLINE void WriteProperty_BackwardsCompatible( FNetBitWriter& Writer, const FRepLayoutCmd& Cmd, const int32 CmdIndex, const UObject* Owner, const uint8* Data, const bool bDoChecksum )
+static FORCEINLINE void WriteProperty_BackwardsCompatible(
+	FNetBitWriter&			Writer,
+	const FRepLayoutCmd&	Cmd,
+	const int32				CmdIndex,
+	const UObject*			Owner,
+	const uint8*			Data,
+	const bool				bDoChecksum)
 {
 #if !(UE_BUILD_SHIPPING || UE_BUILD_TEST)
-	if ( GDoReplicationContextString > 0 )
+	if (GDoReplicationContextString> 0)
 	{
-		Writer.PackageMap->SetDebugContextString( FString::Printf( TEXT( "%s - %s" ), *Owner->GetPathName(), *Cmd.Property->GetPathName() ) );
+		Writer.PackageMap->SetDebugContextString(FString::Printf(TEXT("%s - %s"), *Owner->GetPathName(), *Cmd.Property->GetPathName()));
 	}
 #endif
 
 	const int32 NumStartBits = Writer.GetNumBits();
 
-	FNetBitWriter TempWriter( Writer.PackageMap, 0 );
+	FNetBitWriter TempWriter(Writer.PackageMap, 0);
 
 	// This property changed, so send it
-	Cmd.Property->NetSerializeItem( TempWriter, TempWriter.PackageMap, ( void* )Data );
+	Cmd.Property->NetSerializeItem(TempWriter, TempWriter.PackageMap, (void*)Data);
 
 	uint32 NumBits = TempWriter.GetNumBits();
-	Writer.SerializeIntPacked( NumBits );
-	Writer.SerializeBits( TempWriter.GetData(), NumBits );
+	Writer.SerializeIntPacked(NumBits);
+	Writer.SerializeBits(TempWriter.GetData(), NumBits);
 
 	const int32 NumEndBits = Writer.GetNumBits();
 
-	NETWORK_PROFILER( GNetworkProfiler.TrackReplicateProperty( Cmd.Property, NumEndBits - NumStartBits, nullptr ) );
+	NETWORK_PROFILER(GNetworkProfiler.TrackReplicateProperty(Cmd.Property, NumEndBits - NumStartBits, nullptr));
 
 #if !(UE_BUILD_SHIPPING || UE_BUILD_TEST)
-	if ( GDoReplicationContextString > 0 )
+	if (GDoReplicationContextString> 0)
 	{
 		Writer.PackageMap->ClearDebugContextString();
 	}
 #endif
 
 #ifdef ENABLE_PROPERTY_CHECKSUMS
-	if ( bDoChecksum )
+	if (bDoChecksum)
 	{
-		SerializeReadWritePropertyChecksum( Cmd, CmdIndex, Data, Writer );
+		SerializeReadWritePropertyChecksum(Cmd, CmdIndex, Data, Writer);
 	}
 #endif
 }
@@ -1571,25 +1645,25 @@ void FRepLayout::SendProperties_BackwardsCompatible_r(
 	FNetBitWriter&						Writer,
 	const bool							bDoChecksum,
 	FRepHandleIterator&					HandleIterator,
-	const uint8* RESTRICT				SourceData ) const
+	const uint8* RESTRICT				SourceData) const
 {
 	int32 OldIndex = -1;
 
-	FNetBitWriter TempWriter( Writer.PackageMap, 0 );
+	FNetBitWriter TempWriter(Writer.PackageMap, 0);
 
-	while ( HandleIterator.NextHandle() )
+	while (HandleIterator.NextHandle())
 	{
 		const FRepLayoutCmd& Cmd = Cmds[HandleIterator.CmdIndex];
 
 		const FRepParentCmd& ParentCmd = Parents[Cmd.ParentIndex];
 
-		const bool bConditionMatches = ChangedTracker == nullptr || ( RepState->ConditionMap[ParentCmd.Condition] && ChangedTracker->Parents[Cmd.ParentIndex].Active );
+		const bool bConditionMatches = ChangedTracker == nullptr || (RepState->ConditionMap[ParentCmd.Condition] && ChangedTracker->Parents[Cmd.ParentIndex].Active);
 
-		if ( !bConditionMatches )
+		if (!bConditionMatches)
 		{
-			if ( Cmd.Type == ERepLayoutCmdType::DynamicArray )
+			if (Cmd.Type == ERepLayoutCmdType::DynamicArray)
 			{
-				if ( !HandleIterator.JumpOverArray() )
+				if (!HandleIterator.JumpOverArray())
 				{
 					break;
 				}
@@ -1600,25 +1674,25 @@ void FRepLayout::SendProperties_BackwardsCompatible_r(
 
 		const uint8* Data = SourceData + HandleIterator.ArrayOffset + Cmd.Offset;
 
-		PackageMapClient->TrackNetFieldExport( NetFieldExportGroup, HandleIterator.CmdIndex );
+		PackageMapClient->TrackNetFieldExport(NetFieldExportGroup, HandleIterator.CmdIndex);
 
-		if ( HandleIterator.ArrayElementSize > 0 && HandleIterator.ArrayIndex != OldIndex )
+		if (HandleIterator.ArrayElementSize> 0 && HandleIterator.ArrayIndex != OldIndex)
 		{
-			if ( OldIndex != -1 )
+			if (OldIndex != -1)
 			{
-				WritePropertyHandle_BackwardsCompatible( Writer, 0, bDoChecksum );
+				WritePropertyHandle_BackwardsCompatible(Writer, 0, bDoChecksum);
 			}
 
 			uint32 Index = HandleIterator.ArrayIndex + 1;
-			Writer.SerializeIntPacked( Index );
+			Writer.SerializeIntPacked(Index);
 			OldIndex = HandleIterator.ArrayIndex;
 		}
 
-		WritePropertyHandle_BackwardsCompatible( Writer, HandleIterator.CmdIndex + 1, bDoChecksum );
+		WritePropertyHandle_BackwardsCompatible(Writer, HandleIterator.CmdIndex + 1, bDoChecksum);
 
-		if ( Cmd.Type == ERepLayoutCmdType::DynamicArray )
+		if (Cmd.Type == ERepLayoutCmdType::DynamicArray)
 		{
-			const FScriptArray* Array = ( FScriptArray * )Data;
+			const FScriptArray* Array = (FScriptArray *)Data;
 
 			uint32 ArrayNum = Array->Num();
 
@@ -1629,101 +1703,101 @@ void FRepLayout::SendProperties_BackwardsCompatible_r(
 
 			const int32 OldChangedIndex = HandleIterator.ChangelistIterator.ChangedIndex;
 
-			const uint8* NewData = ( uint8* )Array->GetData();
+			const uint8* NewData = (uint8*)Array->GetData();
 
-			TArray< FHandleToCmdIndex >& ArrayHandleToCmdIndex = *HandleIterator.HandleToCmdIndex[Cmd.RelativeHandle - 1].HandleToCmdIndex;
+			TArray<FHandleToCmdIndex>& ArrayHandleToCmdIndex = *HandleIterator.HandleToCmdIndex[Cmd.RelativeHandle - 1].HandleToCmdIndex;
 
-			FRepHandleIterator ArrayHandleIterator( HandleIterator.ChangelistIterator, Cmds, ArrayHandleToCmdIndex, Cmd.ElementSize, ArrayNum, HandleIterator.CmdIndex + 1, Cmd.EndCmd - 1 );
+			FRepHandleIterator ArrayHandleIterator(HandleIterator.ChangelistIterator, Cmds, ArrayHandleToCmdIndex, Cmd.ElementSize, ArrayNum, HandleIterator.CmdIndex + 1, Cmd.EndCmd - 1);
 
-			check( ArrayHandleIterator.ArrayElementSize > 0 );
-			check( ArrayHandleIterator.NumHandlesPerElement > 0 );
+			check(ArrayHandleIterator.ArrayElementSize> 0);
+			check(ArrayHandleIterator.NumHandlesPerElement> 0);
 
 			TempWriter.Reset();
 
 			// Write array num
-			TempWriter.SerializeIntPacked( ArrayNum );
+			TempWriter.SerializeIntPacked(ArrayNum);
 
-			if ( ArrayNum > 0 )
+			if (ArrayNum> 0)
 			{
-				SendProperties_BackwardsCompatible_r( RepState, PackageMapClient, NetFieldExportGroup, ChangedTracker, TempWriter, bDoChecksum, ArrayHandleIterator, NewData );
+				SendProperties_BackwardsCompatible_r(RepState, PackageMapClient, NetFieldExportGroup, ChangedTracker, TempWriter, bDoChecksum, ArrayHandleIterator, NewData);
 			}
 
 			uint32 EndArrayIndex = 0;
-			TempWriter.SerializeIntPacked( EndArrayIndex );
+			TempWriter.SerializeIntPacked(EndArrayIndex);
 
 			uint32 NumBits = TempWriter.GetNumBits();
-			Writer.SerializeIntPacked( NumBits );
-			Writer.SerializeBits( TempWriter.GetData(), NumBits );
+			Writer.SerializeIntPacked(NumBits);
+			Writer.SerializeBits(TempWriter.GetData(), NumBits);
 
-			check( HandleIterator.ChangelistIterator.ChangedIndex - OldChangedIndex == ArrayChangedCount );				// Make sure we read correct amount
-			check( HandleIterator.ChangelistIterator.Changed[HandleIterator.ChangelistIterator.ChangedIndex] == 0 );	// Make sure we are at the end
+			check(HandleIterator.ChangelistIterator.ChangedIndex - OldChangedIndex == ArrayChangedCount);				// Make sure we read correct amount
+			check(HandleIterator.ChangelistIterator.Changed[HandleIterator.ChangelistIterator.ChangedIndex] == 0);	// Make sure we are at the end
 
 			HandleIterator.ChangelistIterator.ChangedIndex++;
 			continue;
 		}
 
-		WriteProperty_BackwardsCompatible( Writer, Cmd, HandleIterator.CmdIndex, Owner, Data, bDoChecksum );
+		WriteProperty_BackwardsCompatible(Writer, Cmd, HandleIterator.CmdIndex, Owner, Data, bDoChecksum);
 	}
 
-	WritePropertyHandle_BackwardsCompatible( Writer, 0, bDoChecksum );
+	WritePropertyHandle_BackwardsCompatible(Writer, 0, bDoChecksum);
 }
 
-void FRepLayout::SendAllProperties_BackwardsCompatible_r( 
+void FRepLayout::SendAllProperties_BackwardsCompatible_r(
 	FNetBitWriter&						Writer,
 	const bool							bDoChecksum,
 	UPackageMapClient*					PackageMapClient,
 	FNetFieldExportGroup*				NetFieldExportGroup,
 	const int32							CmdStart,
 	const int32							CmdEnd, 
-	const uint8*						SourceData ) const
+	const uint8*						SourceData) const
 {
-	FNetBitWriter TempWriter( Writer.PackageMap, 0 );
+	FNetBitWriter TempWriter(Writer.PackageMap, 0);
 
-	for ( int32 CmdIndex = CmdStart; CmdIndex < CmdEnd; CmdIndex++ )
+	for (int32 CmdIndex = CmdStart; CmdIndex < CmdEnd; CmdIndex++)
 	{
 		const FRepLayoutCmd& Cmd = Cmds[CmdIndex];
 
-		check( Cmd.Type != ERepLayoutCmdType::Return );
+		check(Cmd.Type != ERepLayoutCmdType::Return);
 
-		PackageMapClient->TrackNetFieldExport( NetFieldExportGroup, CmdIndex );
+		PackageMapClient->TrackNetFieldExport(NetFieldExportGroup, CmdIndex);
 
-		WritePropertyHandle_BackwardsCompatible( Writer, CmdIndex + 1, bDoChecksum );
+		WritePropertyHandle_BackwardsCompatible(Writer, CmdIndex + 1, bDoChecksum);
 
 		const uint8* Data = SourceData + Cmd.Offset;
 
-		if ( Cmd.Type == ERepLayoutCmdType::DynamicArray )
+		if (Cmd.Type == ERepLayoutCmdType::DynamicArray)
 		{			
-			const FScriptArray* Array = ( FScriptArray * )Data;
+			const FScriptArray* Array = (FScriptArray *)Data;
 
 			TempWriter.Reset();
 
 			// Write array num
 			uint32 ArrayNum = Array->Num();
-			TempWriter.SerializeIntPacked( ArrayNum );
+			TempWriter.SerializeIntPacked(ArrayNum);
 
-			for ( int32 i = 0; i < Array->Num(); i++ )
+			for (int32 i = 0; i < Array->Num(); i++)
 			{
 				uint32 ArrayIndex = i + 1;
-				TempWriter.SerializeIntPacked( ArrayIndex );
+				TempWriter.SerializeIntPacked(ArrayIndex);
 
-				SendAllProperties_BackwardsCompatible_r( TempWriter, bDoChecksum, PackageMapClient, NetFieldExportGroup, CmdIndex + 1, Cmd.EndCmd - 1, ((const uint8*)Array->GetData()) + Cmd.ElementSize * i );
+				SendAllProperties_BackwardsCompatible_r(TempWriter, bDoChecksum, PackageMapClient, NetFieldExportGroup, CmdIndex + 1, Cmd.EndCmd - 1, ((const uint8*)Array->GetData()) + Cmd.ElementSize * i);
 			}
 
 			uint32 EndArrayIndex = 0;
-			TempWriter.SerializeIntPacked( EndArrayIndex );
+			TempWriter.SerializeIntPacked(EndArrayIndex);
 
 			uint32 NumBits = TempWriter.GetNumBits();
-			Writer.SerializeIntPacked( NumBits );
-			Writer.SerializeBits( TempWriter.GetData(), NumBits );
+			Writer.SerializeIntPacked(NumBits);
+			Writer.SerializeBits(TempWriter.GetData(), NumBits);
 
 			CmdIndex = Cmd.EndCmd - 1;		// The -1 to handle the ++ in the for loop
 			continue;
 		}
 
-		WriteProperty_BackwardsCompatible( Writer, Cmd, CmdIndex, Owner, Data, bDoChecksum );
+		WriteProperty_BackwardsCompatible(Writer, Cmd, CmdIndex, Owner, Data, bDoChecksum);
 	}
 
-	WritePropertyHandle_BackwardsCompatible( Writer, 0, bDoChecksum );
+	WritePropertyHandle_BackwardsCompatible(Writer, 0, bDoChecksum);
 }
 
 void FRepLayout::SendProperties_BackwardsCompatible(
@@ -1732,73 +1806,80 @@ void FRepLayout::SendProperties_BackwardsCompatible(
 	const uint8* RESTRICT		Data,
 	UNetConnection*				Connection,
 	FNetBitWriter&				Writer,
-	TArray< uint16 >&			Changed ) const
+	TArray<uint16>&				Changed) const
 {
 	SCOPE_CYCLE_COUNTER(STAT_NetReplicateDynamicPropSendBackCompatTime);
 
-	FBitWriterMark Mark( Writer );
+	FBitWriterMark Mark(Writer);
 
 #ifdef ENABLE_PROPERTY_CHECKSUMS
 	const bool bDoChecksum = (GDoPropertyChecksum == 1);
-	Writer.WriteBit( bDoChecksum ? 1 : 0 );
+	Writer.WriteBit(bDoChecksum ? 1 : 0);
 #else
 	const bool bDoChecksum = false;
 #endif
 
-	UPackageMapClient* PackageMapClient = ( UPackageMapClient* )Connection->PackageMap;
+	UPackageMapClient* PackageMapClient = (UPackageMapClient*)Connection->PackageMap;
 
-	TSharedPtr< FNetFieldExportGroup > NetFieldExportGroup = PackageMapClient->GetNetFieldExportGroup( Owner->GetPathName() );
+	TSharedPtr<FNetFieldExportGroup> NetFieldExportGroup = PackageMapClient->GetNetFieldExportGroup(Owner->GetPathName());
 
-	if ( !NetFieldExportGroup.IsValid() )
+	if (!NetFieldExportGroup.IsValid())
 	{
 		NetFieldExportGroup = CreateNetfieldExportGroup();
 
-		PackageMapClient->AddNetFieldExportGroup( Owner->GetPathName(), NetFieldExportGroup );
+		PackageMapClient->AddNetFieldExportGroup(Owner->GetPathName(), NetFieldExportGroup);
 	}
 
 	const int32 NumBits = Writer.GetNumBits();
 
-	if ( Changed.Num() == 0 )
+	if (Changed.Num() == 0)
 	{
-		SendAllProperties_BackwardsCompatible_r( Writer, bDoChecksum, PackageMapClient, NetFieldExportGroup.Get(), 0, Cmds.Num() - 1, Data );
+		SendAllProperties_BackwardsCompatible_r(Writer, bDoChecksum, PackageMapClient, NetFieldExportGroup.Get(), 0, Cmds.Num() - 1, Data);
 	}
 	else
 	{
-		FChangelistIterator ChangelistIterator( Changed, 0 );
-		FRepHandleIterator HandleIterator( ChangelistIterator, Cmds, BaseHandleToCmdIndex, 0, 1, 0, Cmds.Num() - 1 );
+		FChangelistIterator ChangelistIterator(Changed, 0);
+		FRepHandleIterator HandleIterator(ChangelistIterator, Cmds, BaseHandleToCmdIndex, 0, 1, 0, Cmds.Num() - 1);
 
-		SendProperties_BackwardsCompatible_r( RepState, PackageMapClient, NetFieldExportGroup.Get(), ChangedTracker, Writer, bDoChecksum, HandleIterator, Data );
+		SendProperties_BackwardsCompatible_r(RepState, PackageMapClient, NetFieldExportGroup.Get(), ChangedTracker, Writer, bDoChecksum, HandleIterator, Data);
 	}
 
-	if ( NumBits == Writer.GetNumBits() )
+	if (NumBits == Writer.GetNumBits())
 	{
-		Mark.Pop( Writer );
+		Mark.Pop(Writer);
 	}
 }
 
 class FReceivedPropertiesStackState : public FCmdIteratorBaseStackState
 {
 public:
-	FReceivedPropertiesStackState( const int32 InCmdStart, const int32 InCmdEnd, FScriptArray*	InShadowArray, FScriptArray* InDataArray, uint8* RESTRICT InShadowBaseData, uint8* RESTRICT	InBaseData ) : 
-		FCmdIteratorBaseStackState( InCmdStart, InCmdEnd, InShadowArray, InDataArray, InShadowBaseData, InBaseData ),
-		GuidReferencesMap( NULL )
+	FReceivedPropertiesStackState(
+		const int32		InCmdStart,
+		const int32		InCmdEnd,
+		FScriptArray*	InShadowArray,
+		FScriptArray*	InDataArray,
+		uint8* RESTRICT	InShadowBaseData,
+		uint8* RESTRICT	InBaseData
+	): 
+		FCmdIteratorBaseStackState(InCmdStart, InCmdEnd, InShadowArray, InDataArray, InShadowBaseData, InBaseData),
+		GuidReferencesMap(NULL)
 	{}
 
 	FGuidReferencesMap* GuidReferencesMap;
 };
 
-static bool ReceivePropertyHelper( 
+static bool ReceivePropertyHelper(
 	FNetBitReader&					Bunch, 
 	FGuidReferencesMap*				GuidReferencesMap,
 	const int32						ElementOffset, 
 	uint8* RESTRICT					ShadowData,
 	uint8* RESTRICT					Data,
-	TArray< UProperty * >*			RepNotifies,
-	const TArray< FRepParentCmd >&	Parents,
-	const TArray< FRepLayoutCmd >&	Cmds,
+	TArray<UProperty*>*				RepNotifies,
+	const TArray<FRepParentCmd>&	Parents,
+	const TArray<FRepLayoutCmd>&	Cmds,
 	const int32						CmdIndex,
 	const bool						bDoChecksum,
-	bool&							bOutGuidsChanged )
+	bool&							bOutGuidsChanged)
 {
 	const FRepLayoutCmd& Cmd	= Cmds[CmdIndex];
 	const FRepParentCmd& Parent = Parents[Cmd.ParentIndex];
@@ -1806,94 +1887,94 @@ static bool ReceivePropertyHelper(
 	// This swaps Role/RemoteRole as we write it
 	const FRepLayoutCmd& SwappedCmd = Parent.RoleSwapIndex != -1 ? Cmds[Parents[Parent.RoleSwapIndex].CmdStart] : Cmd;
 
-	if ( GuidReferencesMap )		// Don't reset unmapped guids here if we are told not to (assuming calling code is handling this)
+	if (GuidReferencesMap)		// Don't reset unmapped guids here if we are told not to (assuming calling code is handling this)
 	{
 		// Let package map know we want to track and know about any guids that are unmapped during the serialize call
-		Bunch.PackageMap->ResetTrackedGuids( true );
+		Bunch.PackageMap->ResetTrackedGuids(true);
 	}
 
 	// Remember where we started reading from, so that if we have unmapped properties, we can re-deserialize from this data later
-	FBitReaderMark Mark( Bunch );
+	FBitReaderMark Mark(Bunch);
 
-	if ( RepNotifies != nullptr && Parent.Property->HasAnyPropertyFlags( CPF_RepNotify ) )
+	if (RepNotifies != nullptr && Parent.Property->HasAnyPropertyFlags(CPF_RepNotify))
 	{
 		// Copy current value over so we can check to see if it changed
-		StoreProperty( Cmd, ShadowData + Cmd.Offset, Data + SwappedCmd.Offset );
+		StoreProperty(Cmd, ShadowData + Cmd.Offset, Data + SwappedCmd.Offset);
 
 		// Read the property
-		Cmd.Property->NetSerializeItem( Bunch, Bunch.PackageMap, Data + SwappedCmd.Offset );
+		Cmd.Property->NetSerializeItem(Bunch, Bunch.PackageMap, Data + SwappedCmd.Offset);
 
 		// Check to see if this property changed
-		if ( Parent.RepNotifyCondition == REPNOTIFY_Always || !PropertiesAreIdentical( Cmd, ShadowData + Cmd.Offset, Data + SwappedCmd.Offset ) )
+		if (Parent.RepNotifyCondition == REPNOTIFY_Always || !PropertiesAreIdentical(Cmd, ShadowData + Cmd.Offset, Data + SwappedCmd.Offset))
 		{
-			(*RepNotifies).AddUnique( Parent.Property );
+			(*RepNotifies).AddUnique(Parent.Property);
 		}
 		else
 		{
-			UE_CLOG( LogSkippedRepNotifies > 0, LogRep, Display, TEXT( "2 FReceivedPropertiesStackState Skipping RepNotify for property %s because local value has not changed." ), *Cmd.Property->GetName() );
+			UE_CLOG(LogSkippedRepNotifies> 0, LogRep, Display, TEXT("2 FReceivedPropertiesStackState Skipping RepNotify for property %s because local value has not changed."), *Cmd.Property->GetName());
 		}
 	}
 	else
 	{
-		Cmd.Property->NetSerializeItem( Bunch, Bunch.PackageMap, Data + SwappedCmd.Offset );
+		Cmd.Property->NetSerializeItem(Bunch, Bunch.PackageMap, Data + SwappedCmd.Offset);
 	}
 
 #ifdef ENABLE_PROPERTY_CHECKSUMS
-	if ( bDoChecksum )
+	if (bDoChecksum)
 	{
-		SerializeReadWritePropertyChecksum( Cmd, CmdIndex, Data + SwappedCmd.Offset, Bunch );
+		SerializeReadWritePropertyChecksum(Cmd, CmdIndex, Data + SwappedCmd.Offset, Bunch);
 	}
 #endif
 
-	if ( GuidReferencesMap )
+	if (GuidReferencesMap)
 	{
 		const int32 AbsOffset = ElementOffset + SwappedCmd.Offset;
 
 		// Loop over all de-serialized network guids and track them so we can manage their pointers as their replicated reference goes in/out of relevancy
-		const TSet< FNetworkGUID >& TrackedUnmappedGuids = Bunch.PackageMap->GetTrackedUnmappedGuids();
-		const TSet< FNetworkGUID >& TrackedDynamicMappedGuids = Bunch.PackageMap->GetTrackedDynamicMappedGuids();
+		const TSet<FNetworkGUID>& TrackedUnmappedGuids = Bunch.PackageMap->GetTrackedUnmappedGuids();
+		const TSet<FNetworkGUID>& TrackedDynamicMappedGuids = Bunch.PackageMap->GetTrackedDynamicMappedGuids();
 
-		const bool bHasUnmapped = TrackedUnmappedGuids.Num() > 0;
+		const bool bHasUnmapped = TrackedUnmappedGuids.Num()> 0;
 
-		FGuidReferences* GuidReferences = GuidReferencesMap->Find( AbsOffset );
+		FGuidReferences* GuidReferences = GuidReferencesMap->Find(AbsOffset);
 
-		if ( TrackedUnmappedGuids.Num() > 0 || TrackedDynamicMappedGuids.Num() > 0 )
+		if (TrackedUnmappedGuids.Num()> 0 || TrackedDynamicMappedGuids.Num()> 0)
 		{
-			if ( GuidReferences != nullptr )
+			if (GuidReferences != nullptr)
 			{
-				check( GuidReferences->CmdIndex == CmdIndex );
-				check( GuidReferences->ParentIndex == Cmd.ParentIndex );
+				check(GuidReferences->CmdIndex == CmdIndex);
+				check(GuidReferences->ParentIndex == Cmd.ParentIndex);
 
 				// If we're already tracking the guids, re-copy lists only if they've changed
-				if ( !NetworkGuidSetsAreSame( GuidReferences->UnmappedGUIDs, TrackedUnmappedGuids ) )
+				if (!NetworkGuidSetsAreSame(GuidReferences->UnmappedGUIDs, TrackedUnmappedGuids))
 				{
 					bOutGuidsChanged = true;
 				}
-				else if ( !NetworkGuidSetsAreSame( GuidReferences->MappedDynamicGUIDs, TrackedDynamicMappedGuids ) )
+				else if (!NetworkGuidSetsAreSame(GuidReferences->MappedDynamicGUIDs, TrackedDynamicMappedGuids))
 				{
 					bOutGuidsChanged = true;
 				}
 			}
 			
-			if ( GuidReferences == nullptr || bOutGuidsChanged )
+			if (GuidReferences == nullptr || bOutGuidsChanged)
 			{
 				// First time tracking these guids (or guids changed), so add (or replace) new entry
-				GuidReferencesMap->Add( AbsOffset, FGuidReferences( Bunch, Mark, TrackedUnmappedGuids, TrackedDynamicMappedGuids, Cmd.ParentIndex, CmdIndex ) );
+				GuidReferencesMap->Add(AbsOffset, FGuidReferences(Bunch, Mark, TrackedUnmappedGuids, TrackedDynamicMappedGuids, Cmd.ParentIndex, CmdIndex));
 				bOutGuidsChanged = true;
 			}
 		}
 		else
 		{
 			// If we don't have any unmapped guids, then make sure to remove the entry so we don't serialize old data when we update unmapped objects
-			if ( GuidReferences != nullptr )
+			if (GuidReferences != nullptr)
 			{
-				GuidReferencesMap->Remove( AbsOffset );
+				GuidReferencesMap->Remove(AbsOffset);
 				bOutGuidsChanged = true;
 			}
 		}
 
 		// Stop tracking unmapped objects
-		Bunch.PackageMap->ResetTrackedGuids( false );
+		Bunch.PackageMap->ResetTrackedGuids(false);
 
 		return bHasUnmapped;
 	}
@@ -1912,92 +1993,99 @@ static FGuidReferencesMap* PrepReceivedArray(
 	const int32				CmdIndex,
 	uint8* RESTRICT*		OutShadowBaseData,
 	uint8* RESTRICT*		OutBaseData,
-	TArray< UProperty * >*	RepNotifies )
+	TArray<UProperty*>*		RepNotifies)
 {
 	FGuidReferences* NewGuidReferencesArray = nullptr;
 
-	if ( ParentGuidReferences != nullptr )
+	if (ParentGuidReferences != nullptr)
 	{
 		// Since we don't know yet if something under us could be unmapped, go ahead and allocate an array container now
-		NewGuidReferencesArray = ParentGuidReferences->Find( AbsOffset );
+		NewGuidReferencesArray = ParentGuidReferences->Find(AbsOffset);
 
-		if ( NewGuidReferencesArray == NULL )
+		if (NewGuidReferencesArray == NULL)
 		{
-			NewGuidReferencesArray = &ParentGuidReferences->FindOrAdd( AbsOffset );
+			NewGuidReferencesArray = &ParentGuidReferences->FindOrAdd(AbsOffset);
 
-			NewGuidReferencesArray->Array		= new FGuidReferencesMap;
-			NewGuidReferencesArray->ParentIndex	= Cmd.ParentIndex;
-			NewGuidReferencesArray->CmdIndex	= CmdIndex;
+			NewGuidReferencesArray->Array = new FGuidReferencesMap;
+			NewGuidReferencesArray->ParentIndex = Cmd.ParentIndex;
+			NewGuidReferencesArray->CmdIndex = CmdIndex;
 		}
 
-		check( NewGuidReferencesArray != NULL );
-		check( NewGuidReferencesArray->ParentIndex == Cmd.ParentIndex );
-		check( NewGuidReferencesArray->CmdIndex == CmdIndex );
+		check(NewGuidReferencesArray != NULL);
+		check(NewGuidReferencesArray->ParentIndex == Cmd.ParentIndex);
+		check(NewGuidReferencesArray->CmdIndex == CmdIndex);
 	}
 
-	if ( RepNotifies != nullptr )
+	if (RepNotifies != nullptr)
 	{
-		if ( ( DataArray->Num() != ArrayNum || Parent.RepNotifyCondition == REPNOTIFY_Always ) && Parent.Property->HasAnyPropertyFlags( CPF_RepNotify ) )
+		if ((DataArray->Num() != ArrayNum || Parent.RepNotifyCondition == REPNOTIFY_Always) && Parent.Property->HasAnyPropertyFlags(CPF_RepNotify))
 		{
-			( *RepNotifies ).AddUnique( Parent.Property );
+			(*RepNotifies).AddUnique(Parent.Property);
 		}
 		else
 		{
-			UE_CLOG( LogSkippedRepNotifies > 0, LogRep, Display, TEXT( "1 FReceivedPropertiesStackState Skipping RepNotify for property %s because local value has not changed." ), *Cmd.Property->GetName() );
+			UE_CLOG(LogSkippedRepNotifies> 0, LogRep, Display, TEXT("1 FReceivedPropertiesStackState Skipping RepNotify for property %s because local value has not changed."), *Cmd.Property->GetName());
 		}
 	}
 
-	check( CastChecked< UArrayProperty >( Cmd.Property ) != NULL );
+	check(CastChecked<UArrayProperty>(Cmd.Property) != NULL);
 
 	// Resize arrays if needed
-	FScriptArrayHelper ArrayHelper( ( UArrayProperty * )Cmd.Property, DataArray );
-	ArrayHelper.Resize( ArrayNum );
+	FScriptArrayHelper ArrayHelper((UArrayProperty *)Cmd.Property, DataArray);
+	ArrayHelper.Resize(ArrayNum);
 
 	// Re-compute the base data values since they could have changed after the resize above
-	*OutBaseData		= ( uint8* )DataArray->GetData();
+	*OutBaseData		= (uint8*)DataArray->GetData();
 	*OutShadowBaseData	= nullptr;
 
 	// Only resize the shadow data array if we're actually tracking RepNotifies
-	if ( RepNotifies != nullptr )
+	if (RepNotifies != nullptr)
 	{
-		check( ShadowArray != nullptr );
+		check(ShadowArray != nullptr);
 
-		FScriptArrayHelper ShadowArrayHelper( ( UArrayProperty* )Cmd.Property, ShadowArray );
-		ShadowArrayHelper.Resize( ArrayNum );
+		FScriptArrayHelper ShadowArrayHelper((UArrayProperty*)Cmd.Property, ShadowArray);
+		ShadowArrayHelper.Resize(ArrayNum);
 
-		*OutShadowBaseData = ( uint8* )ShadowArray->GetData();
+		*OutShadowBaseData = (uint8*)ShadowArray->GetData();
 	}
 
 	return NewGuidReferencesArray ? NewGuidReferencesArray->Array : nullptr;
 }
 
-class FReceivePropertiesImpl : public FRepLayoutCmdIterator< FReceivePropertiesImpl, FReceivedPropertiesStackState >
+class FReceivePropertiesImpl : public FRepLayoutCmdIterator<FReceivePropertiesImpl, FReceivedPropertiesStackState>
 {
 public:
-	FReceivePropertiesImpl( FNetBitReader & InBunch, FRepState * InRepState, bool bInDoChecksum, const TArray< FRepParentCmd >& InParents, const TArray< FRepLayoutCmd >& InCmds, const bool bInDoRepNotify ) :
-        FRepLayoutCmdIterator( InParents, InCmds ),
-		WaitingHandle( 0 ),
-		CurrentHandle( 0 ), 
-		Bunch( InBunch ),
-		RepState( InRepState ),
-		bDoChecksum( bInDoChecksum ),
-		bHasUnmapped( false ),
-		bDoRepNotify( bInDoRepNotify )
+	FReceivePropertiesImpl(
+		FNetBitReader&					InBunch,
+		FRepState*						InRepState,
+		bool							bInDoChecksum,
+		const TArray<FRepParentCmd>&	InParents,
+		const TArray<FRepLayoutCmd>&	InCmds,
+		const bool						bInDoRepNotify
+	):
+        FRepLayoutCmdIterator(InParents, InCmds),
+		WaitingHandle(0),
+		CurrentHandle(0), 
+		Bunch(InBunch),
+		RepState(InRepState),
+		bDoChecksum(bInDoChecksum),
+		bHasUnmapped(false),
+		bDoRepNotify(bInDoRepNotify)
 	{}
 
 	void ReadNextHandle()
 	{
-		Bunch.SerializeIntPacked( WaitingHandle );
+		Bunch.SerializeIntPacked(WaitingHandle);
 
 #ifdef ENABLE_PROPERTY_CHECKSUMS
-		if ( bDoChecksum )
+		if (bDoChecksum)
 		{
-			SerializeGenericChecksum( Bunch );
+			SerializeGenericChecksum(Bunch);
 		}
 #endif
 	}
 
-	INIT_STACK( FReceivedPropertiesStackState )
+	INIT_STACK(FReceivedPropertiesStackState)
 	{
 		StackState.GuidReferencesMap = &RepState->GuidReferencesMap;
 	}
@@ -2006,16 +2094,16 @@ public:
 	{
 		CurrentHandle++;
 
-		if ( CurrentHandle == WaitingHandle )
+		if (CurrentHandle == WaitingHandle)
 		{
-			check( WaitingHandle != 0 );
+			check(WaitingHandle != 0);
 			return true;
 		}
 
 		return false;
 	}
 
-	PROCESS_ARRAY_CMD( FReceivedPropertiesStackState )
+	PROCESS_ARRAY_CMD(FReceivedPropertiesStackState)
 	{
 		// Read array size
 		uint16 ArrayNum = 0;
@@ -2039,7 +2127,7 @@ public:
 			CmdIndex,
 			&StackState.ShadowBaseData,
 			&StackState.BaseData,
-			bDoRepNotify ? &RepState->RepNotifies : nullptr );
+			bDoRepNotify ? &RepState->RepNotifies : nullptr);
 
 		// Save the old handle so we can restore it when we pop out of the array
 		const uint16 OldHandle = CurrentHandle;
@@ -2048,23 +2136,23 @@ public:
 		CurrentHandle = 0;
 
 		// Loop over array
-		ProcessDataArrayElements_r( StackState, Cmd );
+		ProcessDataArrayElements_r(StackState, Cmd);
 
 		// Restore the current handle to what it was before we processed this array
 		CurrentHandle = OldHandle;
 
 		// We should be waiting on the NULL terminator handle at this point
-		check( WaitingHandle == 0 );
+		check(WaitingHandle == 0);
 		ReadNextHandle();
 	}
 
-	PROCESS_CMD( FReceivedPropertiesStackState )
+	PROCESS_CMD(FReceivedPropertiesStackState)
 	{
-		check( StackState.GuidReferencesMap != NULL );
+		check(StackState.GuidReferencesMap != NULL);
 
-		const int32 ElementOffset = ( Data - StackState.BaseData );
+		const int32 ElementOffset = (Data - StackState.BaseData);
 
-		if ( ReceivePropertyHelper( Bunch, StackState.GuidReferencesMap, ElementOffset, ShadowData, Data, bDoRepNotify ? &RepState->RepNotifies : nullptr, Parents, Cmds, CmdIndex, bDoChecksum, bGuidsChanged ) )
+		if (ReceivePropertyHelper(Bunch, StackState.GuidReferencesMap, ElementOffset, ShadowData, Data, bDoRepNotify ? &RepState->RepNotifies : nullptr, Parents, Cmds, CmdIndex, bDoChecksum, bGuidsChanged))
 		{
 			bHasUnmapped = true;
 		}
@@ -2083,9 +2171,17 @@ public:
 	bool					bGuidsChanged;
 };
 
-bool FRepLayout::ReceiveProperties( UActorChannel* OwningChannel, UClass * InObjectClass, FRepState * RESTRICT RepState, void* RESTRICT Data, FNetBitReader & InBunch, bool & bOutHasUnmapped, const bool bEnableRepNotifies, bool& bOutGuidsChanged ) const
+bool FRepLayout::ReceiveProperties(
+	UActorChannel*		OwningChannel,
+	UClass*				InObjectClass,
+	FRepState* RESTRICT	RepState,
+	void* RESTRICT		Data,
+	FNetBitReader&		InBunch,
+	bool&				bOutHasUnmapped,
+	const bool			bEnableRepNotifies,
+	bool&				bOutGuidsChanged) const
 {
-	check( InObjectClass == Owner );
+	check(InObjectClass == Owner);
 
 #ifdef ENABLE_PROPERTY_CHECKSUMS
 	const bool bDoChecksum = InBunch.ReadBit() ? true : false;
@@ -2095,32 +2191,32 @@ bool FRepLayout::ReceiveProperties( UActorChannel* OwningChannel, UClass * InObj
 
 	bOutHasUnmapped = false;
 
-	if ( OwningChannel->Connection->InternalAck )
+	if (OwningChannel->Connection->InternalAck)
 	{
-		TSharedPtr< FNetFieldExportGroup > NetFieldExportGroup = ( ( UPackageMapClient* )OwningChannel->Connection->PackageMap )->GetNetFieldExportGroup( Owner->GetPathName() );
+		TSharedPtr<FNetFieldExportGroup> NetFieldExportGroup = ((UPackageMapClient*)OwningChannel->Connection->PackageMap)->GetNetFieldExportGroup(Owner->GetPathName());
 
-		return ReceiveProperties_BackwardsCompatible_r( RepState, NetFieldExportGroup.Get(), InBunch, 0, Cmds.Num() - 1, bEnableRepNotifies ? RepState->StaticBuffer.GetData() : nullptr, ( uint8* )Data, ( uint8* )Data, &RepState->GuidReferencesMap, bOutHasUnmapped, bOutGuidsChanged );
+		return ReceiveProperties_BackwardsCompatible_r(RepState, NetFieldExportGroup.Get(), InBunch, 0, Cmds.Num() - 1, bEnableRepNotifies ? RepState->StaticBuffer.GetData() : nullptr, (uint8*)Data, (uint8*)Data, &RepState->GuidReferencesMap, bOutHasUnmapped, bOutGuidsChanged);
 	}
 
-	FReceivePropertiesImpl ReceivePropertiesImpl( InBunch, RepState, bDoChecksum, Parents, Cmds, bEnableRepNotifies );
+	FReceivePropertiesImpl ReceivePropertiesImpl(InBunch, RepState, bDoChecksum, Parents, Cmds, bEnableRepNotifies);
 
 	// Read first handle
 	ReceivePropertiesImpl.ReadNextHandle();
 
 	// Read all properties
-	ReceivePropertiesImpl.ProcessCmds( (uint8*)Data, RepState->StaticBuffer.GetData() );
+	ReceivePropertiesImpl.ProcessCmds((uint8*)Data, RepState->StaticBuffer.GetData());
 
 	// Make sure we're waiting on the last NULL terminator
-	if ( ReceivePropertiesImpl.WaitingHandle != 0 )
+	if (ReceivePropertiesImpl.WaitingHandle != 0)
 	{
-		UE_LOG( LogRep, Warning, TEXT( "Read out of sync." ) );
+		UE_LOG(LogRep, Warning, TEXT("Read out of sync."));
 		return false;
 	}
 
 #ifdef ENABLE_SUPER_CHECKSUMS
-	if ( InBunch.ReadBit() == 1 )
+	if (InBunch.ReadBit() == 1)
 	{
-		ValidateWithChecksum( RepState->StaticBuffer.GetData(), InBunch );
+		ValidateWithChecksum(RepState->StaticBuffer.GetData(), InBunch);
 	}
 #endif
 
@@ -2130,7 +2226,14 @@ bool FRepLayout::ReceiveProperties( UActorChannel* OwningChannel, UClass * InObj
 	return true;
 }
 
-bool FRepLayout::ReceiveProperties_BackwardsCompatible( UNetConnection* Connection, FRepState* RESTRICT RepState, void* RESTRICT Data, FNetBitReader& InBunch, bool& bOutHasUnmapped, const bool bEnableRepNotifies, bool& bOutGuidsChanged ) const
+bool FRepLayout::ReceiveProperties_BackwardsCompatible(
+	UNetConnection*			Connection,
+	FRepState* RESTRICT		RepState,
+	void* RESTRICT			Data,
+	FNetBitReader&			InBunch,
+	bool&					bOutHasUnmapped,
+	const bool				bEnableRepNotifies,
+	bool&					bOutGuidsChanged) const
 {
 #ifdef ENABLE_PROPERTY_CHECKSUMS
 	const bool bDoChecksum = InBunch.ReadBit() ? true : false;
@@ -2140,26 +2243,29 @@ bool FRepLayout::ReceiveProperties_BackwardsCompatible( UNetConnection* Connecti
 
 	bOutHasUnmapped = false;
 
-	TSharedPtr< FNetFieldExportGroup > NetFieldExportGroup = ( ( UPackageMapClient* )Connection->PackageMap )->GetNetFieldExportGroup( Owner->GetPathName() );
+	TSharedPtr<FNetFieldExportGroup> NetFieldExportGroup = ((UPackageMapClient*)Connection->PackageMap)->GetNetFieldExportGroup(Owner->GetPathName());
 
-	return ReceiveProperties_BackwardsCompatible_r( RepState, NetFieldExportGroup.Get(), InBunch, 0, Cmds.Num() - 1, (bEnableRepNotifies && RepState) ? RepState->StaticBuffer.GetData() : nullptr, ( uint8* )Data, ( uint8* )Data, RepState ? &RepState->GuidReferencesMap : nullptr, bOutHasUnmapped, bOutGuidsChanged );
+	return ReceiveProperties_BackwardsCompatible_r(RepState, NetFieldExportGroup.Get(), InBunch, 0, Cmds.Num() - 1, (bEnableRepNotifies && RepState) ? RepState->StaticBuffer.GetData() : nullptr, (uint8*)Data, (uint8*)Data, RepState ? &RepState->GuidReferencesMap : nullptr, bOutHasUnmapped, bOutGuidsChanged);
 }
 
-int32 FRepLayout::FindCompatibleProperty( const int32 CmdStart, const int32 CmdEnd, const uint32 Checksum ) const
+int32 FRepLayout::FindCompatibleProperty(
+	const int32		CmdStart,
+	const int32		CmdEnd,
+	const uint32	Checksum) const
 {
-	for ( int32 CmdIndex = CmdStart; CmdIndex < CmdEnd; CmdIndex++ )
+	for (int32 CmdIndex = CmdStart; CmdIndex < CmdEnd; CmdIndex++)
 	{
 		const FRepLayoutCmd& Cmd = Cmds[CmdIndex];
 
-		check( Cmd.Type != ERepLayoutCmdType::Return );
+		check(Cmd.Type != ERepLayoutCmdType::Return);
 
-		if ( Cmd.CompatibleChecksum == Checksum )
+		if (Cmd.CompatibleChecksum == Checksum)
 		{
 			return CmdIndex;
 		}
 
 		// Jump over entire array and inner properties if checksum didn't match
-		if ( Cmd.Type == ERepLayoutCmdType::DynamicArray )
+		if (Cmd.Type == ERepLayoutCmdType::DynamicArray)
 		{
 			CmdIndex = Cmd.EndCmd - 1;
 		}
@@ -2168,7 +2274,7 @@ int32 FRepLayout::FindCompatibleProperty( const int32 CmdStart, const int32 CmdE
 	return -1;
 }
 
-bool FRepLayout::ReceiveProperties_BackwardsCompatible_r( 
+bool FRepLayout::ReceiveProperties_BackwardsCompatible_r(
 	FRepState * RESTRICT	RepState,
 	FNetFieldExportGroup*	NetFieldExportGroup,
 	FNetBitReader &			Reader,
@@ -2179,26 +2285,27 @@ bool FRepLayout::ReceiveProperties_BackwardsCompatible_r(
 	uint8* RESTRICT			Data,
 	FGuidReferencesMap*		GuidReferencesMap,
 	bool&					bOutHasUnmapped,
-	bool&					bOutGuidsChanged ) const
+	bool&					bOutGuidsChanged) const
 {
-	while ( true )
+
+	while (true)
 	{
 		uint32 NetFieldExportHandle = 0;
-		Reader.SerializeIntPacked( NetFieldExportHandle );
+		Reader.SerializeIntPacked(NetFieldExportHandle);
 
-		if ( Reader.IsError() )
+		if (Reader.IsError())
 		{
-			UE_LOG( LogRep, Warning, TEXT( "ReceiveProperties_BackwardsCompatible_r: Error reading handle. Owner: %s" ), *Owner->GetName() );
+			UE_LOG(LogRep, Warning, TEXT("ReceiveProperties_BackwardsCompatible_r: Error reading handle. Owner: %s"), *Owner->GetName());
 			return false;
 		}
 
-		if ( NetFieldExportHandle == 0 )
+		if (NetFieldExportHandle == 0)
 		{
 			// We're done
 			break;
 		}
 
-		if ( !ensure( NetFieldExportGroup != nullptr ) )
+		if (!ensure(NetFieldExportGroup != nullptr))
 		{
 			UE_LOG( LogRep, Warning, TEXT( "ReceiveProperties_BackwardsCompatible_r: NetFieldExportGroup == nullptr. Owner: %s, NetFieldExportHandle: %u" ), *Owner->GetName(), NetFieldExportHandle );
 			Reader.SetError();
@@ -2208,51 +2315,51 @@ bool FRepLayout::ReceiveProperties_BackwardsCompatible_r(
 		// We purposely add 1 on save, so we can reserve 0 for "done"
 		NetFieldExportHandle--;
 
-		if ( !ensure( NetFieldExportHandle < ( uint32 )NetFieldExportGroup->NetFieldExports.Num() ) )
+		if (!ensure(NetFieldExportHandle < (uint32)NetFieldExportGroup->NetFieldExports.Num()))
 		{
-			UE_LOG( LogRep, Warning, TEXT( "ReceiveProperties_BackwardsCompatible_r: NetFieldExportHandle > NetFieldExportGroup->NetFieldExports.Num(). Owner: %s, NetFieldExportHandle: %u" ), *Owner->GetName(), NetFieldExportHandle );
+			UE_LOG(LogRep, Warning, TEXT("ReceiveProperties_BackwardsCompatible_r: NetFieldExportHandle> NetFieldExportGroup->NetFieldExports.Num(). Owner: %s, NetFieldExportHandle: %u"), *Owner->GetName(), NetFieldExportHandle);
 			return false;
 		}
 
 		const uint32 Checksum = NetFieldExportGroup->NetFieldExports[NetFieldExportHandle].CompatibleChecksum;
 
-		if ( !ensure( Checksum != 0 ) )
+		if (!ensure(Checksum != 0))
 		{
-			UE_LOG( LogRep, Warning, TEXT( "ReceiveProperties_BackwardsCompatible_r: Checksum == 0. Owner: %s, Name: %s, Type: %s, NetFieldExportHandle: %i" ), *Owner->GetName(), *NetFieldExportGroup->NetFieldExports[NetFieldExportHandle].Name, *NetFieldExportGroup->NetFieldExports[NetFieldExportHandle].Type, NetFieldExportHandle );
+			UE_LOG(LogRep, Warning, TEXT("ReceiveProperties_BackwardsCompatible_r: Checksum == 0. Owner: %s, Name: %s, Type: %s, NetFieldExportHandle: %i"), *Owner->GetName(), *NetFieldExportGroup->NetFieldExports[NetFieldExportHandle].Name, *NetFieldExportGroup->NetFieldExports[NetFieldExportHandle].Type, NetFieldExportHandle);
 			return false;
 		}
 
 		uint32 NumBits = 0;
-		Reader.SerializeIntPacked( NumBits );
+		Reader.SerializeIntPacked(NumBits);
 
-		if ( Reader.IsError() )
+		if (Reader.IsError())
 		{
-			UE_LOG( LogRep, Warning, TEXT( "ReceiveProperties_BackwardsCompatible_r: Error reading num bits. Owner: %s, Name: %s, Type: %s, NetFieldExportHandle: %i, Checksum: %u" ), *Owner->GetName(), *NetFieldExportGroup->NetFieldExports[NetFieldExportHandle].Name, *NetFieldExportGroup->NetFieldExports[NetFieldExportHandle].Type, NetFieldExportHandle, Checksum );
+			UE_LOG(LogRep, Warning, TEXT("ReceiveProperties_BackwardsCompatible_r: Error reading num bits. Owner: %s, Name: %s, Type: %s, NetFieldExportHandle: %i, Checksum: %u"), *Owner->GetName(), *NetFieldExportGroup->NetFieldExports[NetFieldExportHandle].Name, *NetFieldExportGroup->NetFieldExports[NetFieldExportHandle].Type, NetFieldExportHandle, Checksum);
 			return false;
 		}
 
 		FNetBitReader TempReader;
 		
 		TempReader.PackageMap = Reader.PackageMap;
-		TempReader.SetData( Reader, NumBits );
+		TempReader.SetData(Reader, NumBits);
 
-		if ( Reader.IsError() )
+		if (Reader.IsError())
 		{
-			UE_LOG( LogRep, Warning, TEXT( "ReceiveProperties_BackwardsCompatible_r: Error reading payload. Owner: %s, Name: %s, Type: %s, NetFieldExportHandle: %i, Checksum: %u" ), *Owner->GetName(), *NetFieldExportGroup->NetFieldExports[NetFieldExportHandle].Name, *NetFieldExportGroup->NetFieldExports[NetFieldExportHandle].Type, NetFieldExportHandle, Checksum );
+			UE_LOG(LogRep, Warning, TEXT("ReceiveProperties_BackwardsCompatible_r: Error reading payload. Owner: %s, Name: %s, Type: %s, NetFieldExportHandle: %i, Checksum: %u"), *Owner->GetName(), *NetFieldExportGroup->NetFieldExports[NetFieldExportHandle].Name, *NetFieldExportGroup->NetFieldExports[NetFieldExportHandle].Type, NetFieldExportHandle, Checksum);
 			return false;
 		}
 
-		if ( NetFieldExportGroup->NetFieldExports[NetFieldExportHandle].bIncompatible )
+		if (NetFieldExportGroup->NetFieldExports[NetFieldExportHandle].bIncompatible)
 		{
 			continue;		// We've already warned that this property doesn't load anymore
 		}
 
 		// Find this property
-		const int32 CmdIndex = FindCompatibleProperty( CmdStart, CmdEnd, Checksum );
+		const int32 CmdIndex = FindCompatibleProperty(CmdStart, CmdEnd, Checksum);
 
-		if ( CmdIndex == -1 )
+		if (CmdIndex == -1)
 		{
-			UE_LOG( LogRep, Warning, TEXT( "ReceiveProperties_BackwardsCompatible_r: Property not found. Owner: %s, Name: %s, Type: %s, NetFieldExportHandle: %i, Checksum: %u" ), *Owner->GetName(), *NetFieldExportGroup->NetFieldExports[NetFieldExportHandle].Name, *NetFieldExportGroup->NetFieldExports[NetFieldExportHandle].Type, NetFieldExportHandle, Checksum );
+			UE_LOG(LogRep, Warning, TEXT("ReceiveProperties_BackwardsCompatible_r: Property not found. Owner: %s, Name: %s, Type: %s, NetFieldExportHandle: %i, Checksum: %u"), *Owner->GetName(), *NetFieldExportGroup->NetFieldExports[NetFieldExportHandle].Name, *NetFieldExportGroup->NetFieldExports[NetFieldExportHandle].Type, NetFieldExportHandle, Checksum);
 
 			// Mark this property as incompatible so we don't keep spamming this warning
 			NetFieldExportGroup->NetFieldExports[NetFieldExportHandle].bIncompatible = true;
@@ -2261,20 +2368,20 @@ bool FRepLayout::ReceiveProperties_BackwardsCompatible_r(
 
 		const FRepLayoutCmd& Cmd = Cmds[CmdIndex];
 
-		if ( Cmd.Type == ERepLayoutCmdType::DynamicArray )
+		if (Cmd.Type == ERepLayoutCmdType::DynamicArray)
 		{
 			uint32 ArrayNum = 0;
-			TempReader.SerializeIntPacked( ArrayNum );
+			TempReader.SerializeIntPacked(ArrayNum);
 
-			if ( TempReader.IsError() )
+			if (TempReader.IsError())
 			{
 				return false;
 			}
 
-			const int32 AbsOffset = ( Data - OldData ) + Cmd.Offset;
+			const int32 AbsOffset = (Data - OldData) + Cmd.Offset;
 
-			FScriptArray* DataArray		= ( FScriptArray* )( Data + Cmd.Offset );
-			FScriptArray* ShadowArray	= ShadowData ? ( FScriptArray* )( ShadowData + Cmd.Offset ) : nullptr;
+			FScriptArray* DataArray		= (FScriptArray*)(Data + Cmd.Offset);
+			FScriptArray* ShadowArray	= ShadowData ? (FScriptArray*)(ShadowData + Cmd.Offset) : nullptr;
 
 			uint8* LocalData			= Data;
 			uint8* LocalShadowData		= ShadowData;
@@ -2290,21 +2397,21 @@ bool FRepLayout::ReceiveProperties_BackwardsCompatible_r(
 				CmdIndex,
 				&LocalShadowData,
 				&LocalData,
-				ShadowData != nullptr ? &RepState->RepNotifies : nullptr );
+				ShadowData != nullptr ? &RepState->RepNotifies : nullptr);
 
 			// Read until we read all array elements
-			while ( true )
+			while (true)
 			{
 				uint32 Index = 0;
-				TempReader.SerializeIntPacked( Index );
+				TempReader.SerializeIntPacked(Index);
 
-				if ( TempReader.IsError() )
+				if (TempReader.IsError())
 				{
-					UE_LOG( LogRep, Warning, TEXT( "ReceiveProperties_BackwardsCompatible_r: Error reading array index. Index: %i, Owner: %s, Name: %s, Type: %s, NetFieldExportHandle: %i, Checksum: %u" ), Index, *Owner->GetName(), *NetFieldExportGroup->NetFieldExports[NetFieldExportHandle].Name, *NetFieldExportGroup->NetFieldExports[NetFieldExportHandle].Type, NetFieldExportHandle, Checksum );
+					UE_LOG(LogRep, Warning, TEXT("ReceiveProperties_BackwardsCompatible_r: Error reading array index. Index: %i, Owner: %s, Name: %s, Type: %s, NetFieldExportHandle: %i, Checksum: %u"), Index, *Owner->GetName(), *NetFieldExportGroup->NetFieldExports[NetFieldExportHandle].Name, *NetFieldExportGroup->NetFieldExports[NetFieldExportHandle].Type, NetFieldExportHandle, Checksum);
 					return false;
 				}
 
-				if ( Index == 0 )
+				if (Index == 0)
 				{
 					// We're done
 					break;
@@ -2313,9 +2420,9 @@ bool FRepLayout::ReceiveProperties_BackwardsCompatible_r(
 				// Shift all indexes down since 0 represents null handle
 				Index--;
 
-				if ( !ensure( Index < ArrayNum ) )
+				if (!ensure(Index < ArrayNum))
 				{
-					UE_LOG( LogRep, Warning, TEXT( "ReceiveProperties_BackwardsCompatible_r: Array index out of bounds. Index: %i, ArrayNum: %i, Owner: %s, Name: %s, Type: %s, NetFieldExportHandle: %i, Checksum: %u" ), Index, ArrayNum, *Owner->GetName(), *NetFieldExportGroup->NetFieldExports[NetFieldExportHandle].Name, *NetFieldExportGroup->NetFieldExports[NetFieldExportHandle].Type, NetFieldExportHandle, Checksum );
+					UE_LOG(LogRep, Warning, TEXT("ReceiveProperties_BackwardsCompatible_r: Array index out of bounds. Index: %i, ArrayNum: %i, Owner: %s, Name: %s, Type: %s, NetFieldExportHandle: %i, Checksum: %u"), Index, ArrayNum, *Owner->GetName(), *NetFieldExportGroup->NetFieldExports[NetFieldExportHandle].Name, *NetFieldExportGroup->NetFieldExports[NetFieldExportHandle].Type, NetFieldExportHandle, Checksum);
 					return false;
 				}
 
@@ -2324,36 +2431,36 @@ bool FRepLayout::ReceiveProperties_BackwardsCompatible_r(
 				uint8* ElementData			= LocalData + ElementOffset;
 				uint8* ElementShadowData	= LocalShadowData ? LocalShadowData + ElementOffset : nullptr;
 
-				if ( !ReceiveProperties_BackwardsCompatible_r( RepState, NetFieldExportGroup, TempReader, CmdIndex + 1, Cmd.EndCmd - 1, ElementShadowData, LocalData, ElementData, NewGuidReferencesArray, bOutHasUnmapped, bOutGuidsChanged ) )
+				if (!ReceiveProperties_BackwardsCompatible_r(RepState, NetFieldExportGroup, TempReader, CmdIndex + 1, Cmd.EndCmd - 1, ElementShadowData, LocalData, ElementData, NewGuidReferencesArray, bOutHasUnmapped, bOutGuidsChanged))
 				{
 					return false;
 				}
 
-				if ( TempReader.IsError() )
+				if (TempReader.IsError())
 				{
-					UE_LOG( LogRep, Warning, TEXT( "ReceiveProperties_BackwardsCompatible_r: Error reading array index element payload. Index: %i, Owner: %s, Name: %s, Type: %s, NetFieldExportHandle: %i, Checksum: %u" ), Index, *Owner->GetName(), *NetFieldExportGroup->NetFieldExports[NetFieldExportHandle].Name, *NetFieldExportGroup->NetFieldExports[NetFieldExportHandle].Type, NetFieldExportHandle, Checksum );
+					UE_LOG(LogRep, Warning, TEXT("ReceiveProperties_BackwardsCompatible_r: Error reading array index element payload. Index: %i, Owner: %s, Name: %s, Type: %s, NetFieldExportHandle: %i, Checksum: %u"), Index, *Owner->GetName(), *NetFieldExportGroup->NetFieldExports[NetFieldExportHandle].Name, *NetFieldExportGroup->NetFieldExports[NetFieldExportHandle].Type, NetFieldExportHandle, Checksum);
 					return false;
 				}
 			}
 
-			if ( TempReader.GetBitsLeft() != 0 )
+			if (TempReader.GetBitsLeft() != 0)
 			{
-				UE_LOG( LogRep, Warning, TEXT( "ReceiveProperties_BackwardsCompatible_r: Array didn't read proper number of bits. Owner: %s, Name: %s, Type: %s, NetFieldExportHandle: %i, Checksum: %u" ), *Owner->GetName(), *NetFieldExportGroup->NetFieldExports[NetFieldExportHandle].Name, *NetFieldExportGroup->NetFieldExports[NetFieldExportHandle].Type, NetFieldExportHandle, Checksum );
+				UE_LOG(LogRep, Warning, TEXT("ReceiveProperties_BackwardsCompatible_r: Array didn't read proper number of bits. Owner: %s, Name: %s, Type: %s, NetFieldExportHandle: %i, Checksum: %u"), *Owner->GetName(), *NetFieldExportGroup->NetFieldExports[NetFieldExportHandle].Name, *NetFieldExportGroup->NetFieldExports[NetFieldExportHandle].Type, NetFieldExportHandle, Checksum);
 				return false;
 			}
 		}
 		else
 		{
-			const int32 ElementOffset = ( Data - OldData );
+			const int32 ElementOffset = (Data - OldData);
 
-			if ( ReceivePropertyHelper( TempReader, GuidReferencesMap, ElementOffset, ShadowData, Data, ShadowData != nullptr ? &RepState->RepNotifies : nullptr, Parents, Cmds, CmdIndex, false, bOutGuidsChanged ) )
+			if (ReceivePropertyHelper(TempReader, GuidReferencesMap, ElementOffset, ShadowData, Data, ShadowData != nullptr ? &RepState->RepNotifies : nullptr, Parents, Cmds, CmdIndex, false, bOutGuidsChanged))
 			{
 				bOutHasUnmapped = true;
 			}
 
-			if ( TempReader.GetBitsLeft() != 0 )
+			if (TempReader.GetBitsLeft() != 0)
 			{
-				UE_LOG( LogRep, Warning, TEXT( "ReceiveProperties_BackwardsCompatible_r: Property didn't read proper number of bits. Owner: %s, Name: %s, Type: %s, NetFieldExportHandle: %i, Checksum: %u" ), *Owner->GetName(), *NetFieldExportGroup->NetFieldExports[NetFieldExportHandle].Name, *NetFieldExportGroup->NetFieldExports[NetFieldExportHandle].Type, NetFieldExportHandle, Checksum );
+				UE_LOG(LogRep, Warning, TEXT("ReceiveProperties_BackwardsCompatible_r: Property didn't read proper number of bits. Owner: %s, Name: %s, Type: %s, NetFieldExportHandle: %i, Checksum: %u"), *Owner->GetName(), *NetFieldExportGroup->NetFieldExports[NetFieldExportHandle].Name, *NetFieldExportGroup->NetFieldExports[NetFieldExportHandle].Type, NetFieldExportHandle, Checksum);
 				return false;
 			}
 		}
@@ -2364,62 +2471,68 @@ bool FRepLayout::ReceiveProperties_BackwardsCompatible_r(
 
 FGuidReferences::~FGuidReferences()
 {
-	if ( Array != NULL )
+	if (Array != NULL)
 	{
 		delete Array;
 		Array = NULL;
 	}
 }
 
-void FRepLayout::GatherGuidReferences_r( FGuidReferencesMap* GuidReferencesMap, TSet< FNetworkGUID >& OutReferencedGuids, int32& OutTrackedGuidMemoryBytes ) const
+void FRepLayout::GatherGuidReferences_r(
+	FGuidReferencesMap*	GuidReferencesMap,
+	TSet<FNetworkGUID>&	OutReferencedGuids,
+	int32&				OutTrackedGuidMemoryBytes) const
 {
-	for ( const auto& GuidReferencePair : *GuidReferencesMap )
+	for (const auto& GuidReferencePair : *GuidReferencesMap)
 	{
 		const FGuidReferences& GuidReferences = GuidReferencePair.Value;
 
-		if ( GuidReferences.Array != NULL )
+		if (GuidReferences.Array != NULL)
 		{
-			check( Cmds[GuidReferences.CmdIndex].Type == ERepLayoutCmdType::DynamicArray );
+			check(Cmds[GuidReferences.CmdIndex].Type == ERepLayoutCmdType::DynamicArray);
 
-			GatherGuidReferences_r( GuidReferences.Array, OutReferencedGuids, OutTrackedGuidMemoryBytes );
+			GatherGuidReferences_r(GuidReferences.Array, OutReferencedGuids, OutTrackedGuidMemoryBytes);
 			continue;
 		}
 
 		OutTrackedGuidMemoryBytes += GuidReferences.Buffer.Num();
 
-		OutReferencedGuids.Append( GuidReferences.UnmappedGUIDs );
-		OutReferencedGuids.Append( GuidReferences.MappedDynamicGUIDs );
+		OutReferencedGuids.Append(GuidReferences.UnmappedGUIDs);
+		OutReferencedGuids.Append(GuidReferences.MappedDynamicGUIDs);
 	}
 }
 
-void FRepLayout::GatherGuidReferences( FRepState* RepState, TSet< FNetworkGUID >& OutReferencedGuids, int32& OutTrackedGuidMemoryBytes ) const
+void FRepLayout::GatherGuidReferences(
+	FRepState*			RepState,
+	TSet<FNetworkGUID>&	OutReferencedGuids,
+	int32&				OutTrackedGuidMemoryBytes) const
 {
-	GatherGuidReferences_r( &RepState->GuidReferencesMap, OutReferencedGuids, OutTrackedGuidMemoryBytes );
+	GatherGuidReferences_r(&RepState->GuidReferencesMap, OutReferencedGuids, OutTrackedGuidMemoryBytes);
 }
 
-bool FRepLayout::MoveMappedObjectToUnmapped_r( FGuidReferencesMap* GuidReferencesMap, const FNetworkGUID& GUID ) const
+bool FRepLayout::MoveMappedObjectToUnmapped_r(FGuidReferencesMap* GuidReferencesMap, const FNetworkGUID& GUID) const
 {
 	bool bFoundGUID = false;
 
-	for ( auto& GuidReferencePair : *GuidReferencesMap )
+	for (auto& GuidReferencePair : *GuidReferencesMap)
 	{
 		FGuidReferences& GuidReferences = GuidReferencePair.Value;
 
-		if ( GuidReferences.Array != NULL )
+		if (GuidReferences.Array != NULL)
 		{
-			check( Cmds[GuidReferences.CmdIndex].Type == ERepLayoutCmdType::DynamicArray );
+			check(Cmds[GuidReferences.CmdIndex].Type == ERepLayoutCmdType::DynamicArray);
 
-			if ( MoveMappedObjectToUnmapped_r( GuidReferences.Array, GUID ) )
+			if (MoveMappedObjectToUnmapped_r(GuidReferences.Array, GUID))
 			{
 				bFoundGUID = true;
 			}
 			continue;
 		}
 
-		if ( GuidReferences.MappedDynamicGUIDs.Contains( GUID ) )
+		if (GuidReferences.MappedDynamicGUIDs.Contains(GUID))
 		{
-			GuidReferences.MappedDynamicGUIDs.Remove( GUID );
-			GuidReferences.UnmappedGUIDs.Add( GUID );
+			GuidReferences.MappedDynamicGUIDs.Remove(GUID);
+			GuidReferences.UnmappedGUIDs.Add(GUID);
 			bFoundGUID = true;
 		}
 	}
@@ -2427,12 +2540,12 @@ bool FRepLayout::MoveMappedObjectToUnmapped_r( FGuidReferencesMap* GuidReference
 	return bFoundGUID;
 }
 
-bool FRepLayout::MoveMappedObjectToUnmapped( FRepState* RepState, const FNetworkGUID& GUID ) const
+bool FRepLayout::MoveMappedObjectToUnmapped(FRepState* RepState, const FNetworkGUID& GUID) const
 {
-	return MoveMappedObjectToUnmapped_r( &RepState->GuidReferencesMap, GUID );
+	return MoveMappedObjectToUnmapped_r(&RepState->GuidReferencesMap, GUID);
 }
 
-void FRepLayout::UpdateUnmappedObjects_r( 
+void FRepLayout::UpdateUnmappedObjects_r(
 	FRepState*				RepState, 
 	FGuidReferencesMap*		GuidReferencesMap,
 	UObject*				OriginalObject,
@@ -2441,16 +2554,16 @@ void FRepLayout::UpdateUnmappedObjects_r(
 	uint8* RESTRICT			Data, 
 	const int32				MaxAbsOffset,
 	bool&					bOutSomeObjectsWereMapped,
-	bool&					bOutHasMoreUnmapped ) const
+	bool&					bOutHasMoreUnmapped) const
 {
-	for ( auto It = GuidReferencesMap->CreateIterator(); It; ++It )
+	for (auto It = GuidReferencesMap->CreateIterator(); It; ++It)
 	{
 		const int32 AbsOffset = It.Key();
 
-		if ( AbsOffset >= MaxAbsOffset )
+		if (AbsOffset >= MaxAbsOffset)
 		{
 			// Array must have shrunk, we can remove this item
-			UE_LOG( LogRep, VeryVerbose, TEXT( "UpdateUnmappedObjects_r: REMOVED unmapped property: AbsOffset >= MaxAbsOffset. Offset: %i" ), AbsOffset );
+			UE_LOG(LogRep, VeryVerbose, TEXT("UpdateUnmappedObjects_r: REMOVED unmapped property: AbsOffset>= MaxAbsOffset. Offset: %i"), AbsOffset);
 			It.RemoveCurrent();
 			continue;
 		}
@@ -2459,42 +2572,42 @@ void FRepLayout::UpdateUnmappedObjects_r(
 		const FRepLayoutCmd&		Cmd = Cmds[GuidReferences.CmdIndex];
 		const FRepParentCmd&		Parent = Parents[GuidReferences.ParentIndex];
 
-		if ( GuidReferences.Array != NULL )
+		if (GuidReferences.Array != NULL)
 		{
-			check( Cmd.Type == ERepLayoutCmdType::DynamicArray );
+			check(Cmd.Type == ERepLayoutCmdType::DynamicArray);
 			
-			FScriptArray* StoredArray = (FScriptArray *)( StoredData + AbsOffset );
-			FScriptArray* Array = (FScriptArray *)( Data + AbsOffset );
+			FScriptArray* StoredArray = (FScriptArray *)(StoredData + AbsOffset);
+			FScriptArray* Array = (FScriptArray *)(Data + AbsOffset);
 			
-			const int32 NewMaxOffset = FMath::Min( StoredArray->Num() * Cmd.ElementSize, Array->Num() * Cmd.ElementSize );
+			const int32 NewMaxOffset = FMath::Min(StoredArray->Num() * Cmd.ElementSize, Array->Num() * Cmd.ElementSize);
 
-			UpdateUnmappedObjects_r( RepState, GuidReferences.Array, OriginalObject, PackageMap, (uint8*)StoredArray->GetData(), (uint8*)Array->GetData(), NewMaxOffset, bOutSomeObjectsWereMapped, bOutHasMoreUnmapped );
+			UpdateUnmappedObjects_r(RepState, GuidReferences.Array, OriginalObject, PackageMap, (uint8*)StoredArray->GetData(), (uint8*)Array->GetData(), NewMaxOffset, bOutSomeObjectsWereMapped, bOutHasMoreUnmapped);
 			continue;
 		}
 
 		bool bMappedSomeGUIDs = false;
 
-		for ( auto UnmappedIt = GuidReferences.UnmappedGUIDs.CreateIterator(); UnmappedIt; ++UnmappedIt )
+		for (auto UnmappedIt = GuidReferences.UnmappedGUIDs.CreateIterator(); UnmappedIt; ++UnmappedIt)
 		{			
 			const FNetworkGUID& GUID = *UnmappedIt;
 
-			if ( PackageMap->IsGUIDBroken( GUID, false ) )
+			if (PackageMap->IsGUIDBroken(GUID, false))
 			{
-				UE_LOG( LogRep, Warning, TEXT( "UpdateUnmappedObjects_r: Broken GUID. NetGuid: %s" ), *GUID.ToString() );
+				UE_LOG(LogRep, Warning, TEXT("UpdateUnmappedObjects_r: Broken GUID. NetGuid: %s"), *GUID.ToString());
 				UnmappedIt.RemoveCurrent();
 				continue;
 			}
 
-			UObject* Object = PackageMap->GetObjectFromNetGUID( GUID, false );
+			UObject* Object = PackageMap->GetObjectFromNetGUID(GUID, false);
 
-			if ( Object != NULL )
+			if (Object != NULL)
 			{
-				UE_LOG( LogRep, VeryVerbose, TEXT( "UpdateUnmappedObjects_r: REMOVED unmapped property: Offset: %i, Guid: %s, PropName: %s, ObjName: %s" ), AbsOffset, *GUID.ToString(), *Cmd.Property->GetName(), *Object->GetName() );
+				UE_LOG(LogRep, VeryVerbose, TEXT("UpdateUnmappedObjects_r: REMOVED unmapped property: Offset: %i, Guid: %s, PropName: %s, ObjName: %s"), AbsOffset, *GUID.ToString(), *Cmd.Property->GetName(), *Object->GetName());
 
-				if ( GUID.IsDynamic() )
+				if (GUID.IsDynamic())
 				{
 					// If this guid is dynamic, move it to the dynamic guids list
-					GuidReferences.MappedDynamicGUIDs.Add( GUID );
+					GuidReferences.MappedDynamicGUIDs.Add(GUID);
 				}
 
 				// Remove from unmapped guids list
@@ -2504,9 +2617,9 @@ void FRepLayout::UpdateUnmappedObjects_r(
 		}
 
 		// If we resolved some guids, re-deserialize the data which will hook up the object pointer with the property
-		if ( bMappedSomeGUIDs )
+		if (bMappedSomeGUIDs)
 		{
-			if ( !bOutSomeObjectsWereMapped )
+			if (!bOutSomeObjectsWereMapped)
 			{
 				// Call PreNetReceive if we are going to change a value (some game code will need to think this is an actual replicated value)
 				OriginalObject->PreNetReceive();
@@ -2514,64 +2627,69 @@ void FRepLayout::UpdateUnmappedObjects_r(
 			}
 
 			// Copy current value over so we can check to see if it changed
-			if ( Parent.Property->HasAnyPropertyFlags( CPF_RepNotify ) )
+			if (Parent.Property->HasAnyPropertyFlags(CPF_RepNotify))
 			{
-				StoreProperty( Cmd, StoredData + AbsOffset, Data + AbsOffset );
+				StoreProperty(Cmd, StoredData + AbsOffset, Data + AbsOffset);
 			}
 
 			// Initialize the reader with the stored buffer that we need to read from
-			FNetBitReader Reader( PackageMap, GuidReferences.Buffer.GetData(), GuidReferences.NumBufferBits );
+			FNetBitReader Reader(PackageMap, GuidReferences.Buffer.GetData(), GuidReferences.NumBufferBits);
 
 			// Read the property
-			Cmd.Property->NetSerializeItem( Reader, PackageMap, Data + AbsOffset );
+			Cmd.Property->NetSerializeItem(Reader, PackageMap, Data + AbsOffset);
 
 			// Check to see if this property changed
-			if ( Parent.Property->HasAnyPropertyFlags( CPF_RepNotify ) )
+			if (Parent.Property->HasAnyPropertyFlags(CPF_RepNotify))
 			{
-				if ( Parent.RepNotifyCondition == REPNOTIFY_Always || !PropertiesAreIdentical( Cmd, StoredData + AbsOffset, Data + AbsOffset ) )
+				if (Parent.RepNotifyCondition == REPNOTIFY_Always || !PropertiesAreIdentical(Cmd, StoredData + AbsOffset, Data + AbsOffset))
 				{
 					// If this properties needs an OnRep, queue that up to be handled later
-					RepState->RepNotifies.AddUnique( Parent.Property );
+					RepState->RepNotifies.AddUnique(Parent.Property);
 				}
 				else
 				{
-					UE_CLOG( LogSkippedRepNotifies, LogRep, Display, TEXT( "UpdateUnmappedObjects_r: Skipping RepNotify because Property did not change. %s" ), *Cmd.Property->GetName() );
+					UE_CLOG(LogSkippedRepNotifies, LogRep, Display, TEXT("UpdateUnmappedObjects_r: Skipping RepNotify because Property did not change. %s"), *Cmd.Property->GetName());
 				}
 			}
 		}
 
 		// If we still have more unmapped guids, we need to keep processing this entry
-		if ( GuidReferences.UnmappedGUIDs.Num() > 0 )
+		if (GuidReferences.UnmappedGUIDs.Num() > 0)
 		{
 			bOutHasMoreUnmapped = true;
 		}
-		else if ( GuidReferences.UnmappedGUIDs.Num() == 0 && GuidReferences.MappedDynamicGUIDs.Num() == 0 )
+		else if (GuidReferences.UnmappedGUIDs.Num() == 0 && GuidReferences.MappedDynamicGUIDs.Num() == 0)
 		{
 			It.RemoveCurrent();
 		}
 	}
 }
 
-void FRepLayout::UpdateUnmappedObjects( FRepState *	RepState, UPackageMap * PackageMap, UObject* OriginalObject, bool & bOutSomeObjectsWereMapped, bool & bOutHasMoreUnmapped ) const
+void FRepLayout::UpdateUnmappedObjects(
+	FRepState*		RepState,
+	UPackageMap*	PackageMap,
+	UObject*		OriginalObject,
+	bool&			bOutSomeObjectsWereMapped,
+	bool&			bOutHasMoreUnmapped) const
 {
-	bOutSomeObjectsWereMapped	= false;
-	bOutHasMoreUnmapped			= false;
+	bOutSomeObjectsWereMapped = false;
+	bOutHasMoreUnmapped = false;
 
-	UpdateUnmappedObjects_r( RepState, &RepState->GuidReferencesMap, OriginalObject, PackageMap, (uint8*)RepState->StaticBuffer.GetData(), (uint8*)OriginalObject, RepState->StaticBuffer.Num(), bOutSomeObjectsWereMapped, bOutHasMoreUnmapped );
+	UpdateUnmappedObjects_r(RepState, &RepState->GuidReferencesMap, OriginalObject, PackageMap, (uint8*)RepState->StaticBuffer.GetData(), (uint8*)OriginalObject, RepState->StaticBuffer.Num(), bOutSomeObjectsWereMapped, bOutHasMoreUnmapped);
 }
 
-void FRepLayout::CallRepNotifies( FRepState * RepState, UObject* Object ) const
+void FRepLayout::CallRepNotifies(FRepState* RepState, UObject* Object) const
 {
-	if ( RepState->RepNotifies.Num() == 0 )
+	if (RepState->RepNotifies.Num() == 0)
 	{
 		return;
 	}
 
-	for ( int32 i = 0; i < RepState->RepNotifies.Num(); i++ )
+	for (int32 i = 0; i < RepState->RepNotifies.Num(); i++)
 	{
 		UProperty * RepProperty = RepState->RepNotifies[i];
 
-		UFunction * RepNotifyFunc = Object->FindFunction( RepProperty->RepNotifyFunc );
+		UFunction * RepNotifyFunc = Object->FindFunction(RepProperty->RepNotifyFunc);
 
 		if (RepNotifyFunc == nullptr)
 		{
@@ -2580,132 +2698,144 @@ void FRepLayout::CallRepNotifies( FRepState * RepState, UObject* Object ) const
 			continue;
 		}
 
-		check( RepNotifyFunc->NumParms <= 1 );	// 2 parms not supported yet
+		check(RepNotifyFunc->NumParms <= 1);	// 2 parms not supported yet
 
-		if ( RepNotifyFunc->NumParms == 0 )
+		if (RepNotifyFunc->NumParms == 0)
 		{
-			Object->ProcessEvent( RepNotifyFunc, NULL );
+			Object->ProcessEvent(RepNotifyFunc, NULL);
 		}
-		else if (RepNotifyFunc->NumParms == 1 )
+		else if (RepNotifyFunc->NumParms == 1)
 		{
-			Object->ProcessEvent( RepNotifyFunc, RepProperty->ContainerPtrToValuePtr<uint8>( RepState->StaticBuffer.GetData() ) );
+			Object->ProcessEvent(RepNotifyFunc, RepProperty->ContainerPtrToValuePtr<uint8>(RepState->StaticBuffer.GetData()));
 		}
 
 		// Store the property we just received
-		//StoreProperty( Cmd, StoredData + Cmd.Offset, Data + SwappedCmd.Offset );
+		//StoreProperty(Cmd, StoredData + Cmd.Offset, Data + SwappedCmd.Offset);
 	}
 
 	RepState->RepNotifies.Empty();
 }
 
-void FRepLayout::ValidateWithChecksum_DynamicArray_r( const FRepLayoutCmd& Cmd, const int32 CmdIndex, const uint8* RESTRICT Data, FBitArchive & Ar ) const
+void FRepLayout::ValidateWithChecksum_DynamicArray_r(
+	const FRepLayoutCmd&	Cmd,
+	const int32				CmdIndex,
+	const uint8* RESTRICT	Data,
+	FBitArchive &			Ar) const
 {
 	FScriptArray * Array = (FScriptArray *)Data;
 
-	uint16 ArrayNum		= Array->Num();
-	uint16 ElementSize	= Cmd.ElementSize;
+	uint16 ArrayNum = Array->Num();
+	uint16 ElementSize = Cmd.ElementSize;
 
 	Ar << ArrayNum;
 	Ar << ElementSize;
 
-	if ( ArrayNum != Array->Num() )
+	if (ArrayNum != Array->Num())
 	{
-		UE_LOG( LogRep, Fatal, TEXT( "ValidateWithChecksum_AnyArray_r: Array sizes different! %s %i / %i" ), *Cmd.Property->GetFullName(), ArrayNum, Array->Num() );
+		UE_LOG(LogRep, Fatal, TEXT("ValidateWithChecksum_AnyArray_r: Array sizes different! %s %i / %i"), *Cmd.Property->GetFullName(), ArrayNum, Array->Num());
 	}
 
-	if ( ElementSize != Cmd.ElementSize )
+	if (ElementSize != Cmd.ElementSize)
 	{
-		UE_LOG( LogRep, Fatal, TEXT( "ValidateWithChecksum_AnyArray_r: Array element sizes different! %s %i / %i" ), *Cmd.Property->GetFullName(), ElementSize, Cmd.ElementSize );
+		UE_LOG(LogRep, Fatal, TEXT("ValidateWithChecksum_AnyArray_r: Array element sizes different! %s %i / %i"), *Cmd.Property->GetFullName(), ElementSize, Cmd.ElementSize);
 	}
 
 	uint8* LocalData = (uint8*)Array->GetData();
 
-	for ( int32 i = 0; i < ArrayNum; i++ )
+	for (int32 i = 0; i < ArrayNum; i++)
 	{
-		ValidateWithChecksum_r( CmdIndex + 1, Cmd.EndCmd - 1, LocalData + i * ElementSize, Ar );
+		ValidateWithChecksum_r(CmdIndex + 1, Cmd.EndCmd - 1, LocalData + i * ElementSize, Ar);
 	}
 }
 
-void FRepLayout::ValidateWithChecksum_r( 
+void FRepLayout::ValidateWithChecksum_r(
 	const int32				CmdStart, 
 	const int32				CmdEnd, 
 	const uint8* RESTRICT	Data, 
-	FBitArchive &			Ar ) const
+	FBitArchive &			Ar) const
 {
-	for ( int32 CmdIndex = CmdStart; CmdIndex < CmdEnd; CmdIndex++ )
+	for (int32 CmdIndex = CmdStart; CmdIndex < CmdEnd; CmdIndex++)
 	{
 		const FRepLayoutCmd& Cmd = Cmds[ CmdIndex ];
 
-		check( Cmd.Type != ERepLayoutCmdType::Return );
+		check(Cmd.Type != ERepLayoutCmdType::Return);
 
-		if ( Cmd.Type == ERepLayoutCmdType::DynamicArray )
+		if (Cmd.Type == ERepLayoutCmdType::DynamicArray)
 		{
-			ValidateWithChecksum_DynamicArray_r( Cmd, CmdIndex, Data + Cmd.Offset, Ar );
+			ValidateWithChecksum_DynamicArray_r(Cmd, CmdIndex, Data + Cmd.Offset, Ar);
 			CmdIndex = Cmd.EndCmd - 1;	// Jump past children of this array (-1 for ++ in for loop)
 			continue;
 		}
 
-		SerializeReadWritePropertyChecksum( Cmd, CmdIndex - 1, Data + Cmd.Offset, Ar );
+		SerializeReadWritePropertyChecksum(Cmd, CmdIndex - 1, Data + Cmd.Offset, Ar);
 	}
 }
 
-void FRepLayout::ValidateWithChecksum( const void* RESTRICT Data, FBitArchive & Ar ) const
+void FRepLayout::ValidateWithChecksum(const void* RESTRICT Data, FBitArchive & Ar) const
 {
-	ValidateWithChecksum_r( 0, Cmds.Num() - 1, (const uint8*)Data, Ar );
+	ValidateWithChecksum_r(0, Cmds.Num() - 1, (const uint8*)Data, Ar);
 }
 
-uint32 FRepLayout::GenerateChecksum( const FRepState* RepState ) const
+uint32 FRepLayout::GenerateChecksum(const FRepState* RepState) const
 {
-	FBitWriter Writer( 1024, true );
-	ValidateWithChecksum_r( 0, Cmds.Num() - 1, (const uint8*)RepState->StaticBuffer.GetData(), Writer );
+	FBitWriter Writer(1024, true);
+	ValidateWithChecksum_r(0, Cmds.Num() - 1, (const uint8*)RepState->StaticBuffer.GetData(), Writer);
 
-	return FCrc::MemCrc32( Writer.GetData(), Writer.GetNumBytes(), 0 );
+	return FCrc::MemCrc32(Writer.GetData(), Writer.GetNumBytes(), 0);
 }
 
-void FRepLayout::PruneChangeList( FRepState* RepState, const void* RESTRICT Data, const TArray< uint16 >& Changed, TArray< uint16 >& PrunedChanged ) const
+void FRepLayout::PruneChangeList(
+	FRepState*				RepState,
+	const void* RESTRICT	Data,
+	const TArray<uint16>&	Changed,
+	TArray<uint16>& 		PrunedChanged) const
 {
-	check( Changed.Num() > 0 );
+	check(Changed.Num() > 0);
 
 	PrunedChanged.Empty();
 
-	FChangelistIterator ChangelistIterator( Changed, 0 );
-	FRepHandleIterator HandleIterator( ChangelistIterator, Cmds, BaseHandleToCmdIndex, 0, 1, 0, Cmds.Num() - 1 );
-	PruneChangeList_r( HandleIterator, ( uint8* )Data, PrunedChanged );
+	FChangelistIterator ChangelistIterator(Changed, 0);
+	FRepHandleIterator HandleIterator(ChangelistIterator, Cmds, BaseHandleToCmdIndex, 0, 1, 0, Cmds.Num() - 1);
+	PruneChangeList_r(HandleIterator, (uint8*)Data, PrunedChanged);
 
-	PrunedChanged.Add( 0 );
+	PrunedChanged.Add(0);
 }
 
-void FRepLayout::MergeChangeList( const uint8* RESTRICT Data, const TArray< uint16 > & Dirty1, const TArray< uint16 > & Dirty2, TArray< uint16 > & MergedDirty ) const
+void FRepLayout::MergeChangeList(
+	const uint8* RESTRICT	Data,
+	const TArray<uint16>&	Dirty1,
+	const TArray<uint16>&	Dirty2,
+	TArray<uint16>&			MergedDirty) const
 {
-	check( Dirty1.Num() > 0 );
+	check(Dirty1.Num() > 0);
 
 	MergedDirty.Empty();
 
-	if ( Dirty2.Num() == 0 )
+	if (Dirty2.Num() == 0)
 	{
-		FChangelistIterator ChangelistIterator( Dirty1, 0 );
-		FRepHandleIterator HandleIterator( ChangelistIterator, Cmds, BaseHandleToCmdIndex, 0, 1, 0, Cmds.Num() - 1 );
-		PruneChangeList_r( HandleIterator, ( uint8* )Data, MergedDirty );
+		FChangelistIterator ChangelistIterator(Dirty1, 0);
+		FRepHandleIterator HandleIterator(ChangelistIterator, Cmds, BaseHandleToCmdIndex, 0, 1, 0, Cmds.Num() - 1);
+		PruneChangeList_r(HandleIterator, (uint8*)Data, MergedDirty);
 	}
 	else
 	{
-		FChangelistIterator ChangelistIterator1( Dirty1, 0 );
-		FRepHandleIterator HandleIterator1( ChangelistIterator1, Cmds, BaseHandleToCmdIndex, 0, 1, 0, Cmds.Num() - 1 );
+		FChangelistIterator ChangelistIterator1(Dirty1, 0);
+		FRepHandleIterator HandleIterator1(ChangelistIterator1, Cmds, BaseHandleToCmdIndex, 0, 1, 0, Cmds.Num() - 1);
 
-		FChangelistIterator ChangelistIterator2( Dirty2, 0 );
-		FRepHandleIterator HandleIterator2( ChangelistIterator2, Cmds, BaseHandleToCmdIndex, 0, 1, 0, Cmds.Num() - 1 );
+		FChangelistIterator ChangelistIterator2(Dirty2, 0);
+		FRepHandleIterator HandleIterator2(ChangelistIterator2, Cmds, BaseHandleToCmdIndex, 0, 1, 0, Cmds.Num() - 1);
 
-		MergeChangeList_r( HandleIterator1, HandleIterator2, ( uint8* )Data, MergedDirty );
+		MergeChangeList_r(HandleIterator1, HandleIterator2, (uint8*)Data, MergedDirty);
 	}
 
-	MergedDirty.Add( 0 );
+	MergedDirty.Add(0);
 }
 
-void FRepLayout::SanityCheckChangeList_DynamicArray_r( 
+void FRepLayout::SanityCheckChangeList_DynamicArray_r(
 	const int32				CmdIndex, 
 	const uint8* RESTRICT	Data, 
-	TArray< uint16 > &		Changed,
-	int32 &					ChangedIndex ) const
+	TArray<uint16> &		Changed,
+	int32 &					ChangedIndex) const
 {
 	const FRepLayoutCmd& Cmd = Cmds[ CmdIndex ];
 
@@ -2722,226 +2852,245 @@ void FRepLayout::SanityCheckChangeList_DynamicArray_r(
 
 	uint16 LocalHandle = 0;
 
-	for ( int32 i = 0; i < Array->Num(); i++ )
+	for (int32 i = 0; i < Array->Num(); i++)
 	{
-		LocalHandle = SanityCheckChangeList_r( CmdIndex + 1, Cmd.EndCmd - 1, Data + i * Cmd.ElementSize, Changed, ChangedIndex, LocalHandle );
+		LocalHandle = SanityCheckChangeList_r(CmdIndex + 1, Cmd.EndCmd - 1, Data + i * Cmd.ElementSize, Changed, ChangedIndex, LocalHandle);
 	}
 
-	check( ChangedIndex - OldChangedIndex == ArrayChangedCount );	// Make sure we read correct amount
-	check( Changed[ChangedIndex] == 0 );							// Make sure we are at the end
+	check(ChangedIndex - OldChangedIndex == ArrayChangedCount);	// Make sure we read correct amount
+	check(Changed[ChangedIndex] == 0);							// Make sure we are at the end
 
 	ChangedIndex++;
 }
 
-uint16 FRepLayout::SanityCheckChangeList_r( 
+uint16 FRepLayout::SanityCheckChangeList_r(
 	const int32				CmdStart, 
 	const int32				CmdEnd, 
 	const uint8* RESTRICT	Data, 
-	TArray< uint16 > &		Changed,
+	TArray<uint16> &		Changed,
 	int32 &					ChangedIndex,
 	uint16					Handle 
 	) const
 {
-	for ( int32 CmdIndex = CmdStart; CmdIndex < CmdEnd; CmdIndex++ )
+	for (int32 CmdIndex = CmdStart; CmdIndex < CmdEnd; CmdIndex++)
 	{
 		const FRepLayoutCmd& Cmd = Cmds[ CmdIndex ];
 
-		check( Cmd.Type != ERepLayoutCmdType::Return );
+		check(Cmd.Type != ERepLayoutCmdType::Return);
 
 		Handle++;
 
-		if ( Cmd.Type == ERepLayoutCmdType::DynamicArray )
+		if (Cmd.Type == ERepLayoutCmdType::DynamicArray)
 		{
-			if ( Handle == Changed[ChangedIndex] )
+			if (Handle == Changed[ChangedIndex])
 			{
 				const int32 LastChangedArrayHandle = Changed[ChangedIndex];
 				ChangedIndex++;
-				SanityCheckChangeList_DynamicArray_r( CmdIndex, Data + Cmd.Offset, Changed, ChangedIndex );
-				check( Changed[ChangedIndex] == 0 || Changed[ChangedIndex] > LastChangedArrayHandle );
+				SanityCheckChangeList_DynamicArray_r(CmdIndex, Data + Cmd.Offset, Changed, ChangedIndex);
+				check(Changed[ChangedIndex] == 0 || Changed[ChangedIndex] > LastChangedArrayHandle);
 			}
 			CmdIndex = Cmd.EndCmd - 1;	// Jump past children of this array (the -1 because of the ++ in the for loop)
 			continue;
 		}
 
-		if ( Handle == Changed[ChangedIndex] )
+		if (Handle == Changed[ChangedIndex])
 		{
 			const int32 LastChangedArrayHandle = Changed[ChangedIndex];
 			ChangedIndex++;
-			check( Changed[ChangedIndex] == 0 || Changed[ChangedIndex] > LastChangedArrayHandle );
+			check(Changed[ChangedIndex] == 0 || Changed[ChangedIndex] > LastChangedArrayHandle);
 		}
 	}
 
 	return Handle;
 }
 
-void FRepLayout::SanityCheckChangeList( const uint8* RESTRICT Data, TArray< uint16 > & Changed ) const
+void FRepLayout::SanityCheckChangeList(const uint8* RESTRICT Data, TArray<uint16> & Changed) const
 {
 	int32 ChangedIndex = 0;
-	SanityCheckChangeList_r( 0, Cmds.Num() - 1, Data, Changed, ChangedIndex, 0 );
-	check( Changed[ChangedIndex] == 0 );
+	SanityCheckChangeList_r(0, Cmds.Num() - 1, Data, Changed, ChangedIndex, 0);
+	check(Changed[ChangedIndex] == 0);
 }
 
-class FDiffPropertiesImpl : public FRepLayoutCmdIterator< FDiffPropertiesImpl, FCmdIteratorBaseStackState >
+class FDiffPropertiesImpl : public FRepLayoutCmdIterator<FDiffPropertiesImpl, FCmdIteratorBaseStackState>
 {
 public:
-	FDiffPropertiesImpl( const EDiffPropertiesFlags InFlags, TArray< UProperty * >*	InRepNotifies, const TArray< FRepParentCmd >& InParents, const TArray< FRepLayoutCmd >& InCmds ) : 
-		FRepLayoutCmdIterator( InParents, InCmds ),
-		Flags( InFlags ),
-		ParentPropertyFlags( ERepParentFlags::IsLifetime ),
-		RepNotifies( InRepNotifies ),
-		bDifferent( false )
+
+	FDiffPropertiesImpl(
+		const EDiffPropertiesFlags InFlags,
+		TArray<UProperty*>*	InRepNotifies,
+		const TArray<FRepParentCmd>& InParents,
+		const TArray<FRepLayoutCmd>& InCmds
+	): 
+		FRepLayoutCmdIterator(InParents, InCmds),
+		Flags(InFlags),
+		ParentPropertyFlags(ERepParentFlags::IsLifetime),
+		RepNotifies(InRepNotifies),
+		bDifferent(false)
 	{
 		// Currently, only lifetime properties init from their defaults, so default to that,
 		// but also diff conditional properties if requested.
-		if ( (Flags & EDiffPropertiesFlags::IncludeConditionalProperties) != EDiffPropertiesFlags::None )
+		if ((Flags & EDiffPropertiesFlags::IncludeConditionalProperties) != EDiffPropertiesFlags::None)
 		{
 			ParentPropertyFlags |= ERepParentFlags::IsConditional;
 		}
 	}
 
-	INIT_STACK( FCmdIteratorBaseStackState ) { }
+	INIT_STACK(FCmdIteratorBaseStackState) { }
 
 	SHOULD_PROCESS_NEXT_CMD() 
 	{ 
 		return true;
 	}
 
-	PROCESS_ARRAY_CMD( FCmdIteratorBaseStackState ) 
+	PROCESS_ARRAY_CMD(FCmdIteratorBaseStackState) 
 	{
-		if ( StackState.DataArray->Num() != StackState.ShadowArray->Num() )
+		if (StackState.DataArray->Num() != StackState.ShadowArray->Num())
 		{
 			bDifferent = true;
 
-			if ( (Flags & EDiffPropertiesFlags::Sync) == EDiffPropertiesFlags::None )
+			if ((Flags & EDiffPropertiesFlags::Sync) == EDiffPropertiesFlags::None)
 			{
-				UE_LOG( LogRep, Warning, TEXT( "FDiffPropertiesImpl: Array sizes different: %s %i / %i" ), *Cmd.Property->GetFullName(), StackState.DataArray->Num(), StackState.ShadowArray->Num() );
+				UE_LOG(LogRep, Warning, TEXT("FDiffPropertiesImpl: Array sizes different: %s %i / %i"), *Cmd.Property->GetFullName(), StackState.DataArray->Num(), StackState.ShadowArray->Num());
 				return;
 			}
 
-			if ( (Parents[Cmd.ParentIndex].Flags & ParentPropertyFlags) == ERepParentFlags::None )
+			if ((Parents[Cmd.ParentIndex].Flags & ParentPropertyFlags) == ERepParentFlags::None)
 			{
 				return;
 			}
 
 			// Make the shadow state match the actual state
-			FScriptArrayHelper ShadowArrayHelper( (UArrayProperty *)Cmd.Property, ShadowData );
-			ShadowArrayHelper.Resize( StackState.DataArray->Num() );
+			FScriptArrayHelper ShadowArrayHelper((UArrayProperty *)Cmd.Property, ShadowData);
+			ShadowArrayHelper.Resize(StackState.DataArray->Num());
 		}
 
-		StackState.BaseData			= (uint8*)StackState.DataArray->GetData();
-		StackState.ShadowBaseData	= (uint8*)StackState.ShadowArray->GetData();
+		StackState.BaseData = (uint8*)StackState.DataArray->GetData();
+		StackState.ShadowBaseData = (uint8*)StackState.ShadowArray->GetData();
 
 		// Loop over array
-		ProcessDataArrayElements_r( StackState, Cmd );
+		ProcessDataArrayElements_r(StackState, Cmd);
 	}
 
-	PROCESS_CMD( FCmdIteratorBaseStackState ) 
+	PROCESS_CMD(FCmdIteratorBaseStackState) 
 	{
 		const FRepParentCmd& Parent = Parents[Cmd.ParentIndex];
 
 		// Make the shadow state match the actual state at the time of send
-		if ( (RepNotifies && Parent.RepNotifyCondition == REPNOTIFY_Always) || !PropertiesAreIdentical( Cmd, (const void*)( Data + Cmd.Offset ), (const void*)( ShadowData + Cmd.Offset ) ) )
+		if ((RepNotifies && Parent.RepNotifyCondition == REPNOTIFY_Always) || !PropertiesAreIdentical(Cmd, (const void*)(Data + Cmd.Offset), (const void*)(ShadowData + Cmd.Offset)))
 		{
 			bDifferent = true;
 
-			if ( (Flags & EDiffPropertiesFlags::Sync) == EDiffPropertiesFlags::None )
+			if ((Flags & EDiffPropertiesFlags::Sync) == EDiffPropertiesFlags::None)
 			{			
-				UE_LOG( LogRep, Warning, TEXT( "FDiffPropertiesImpl: Property different: %s" ), *Cmd.Property->GetFullName() );
+				UE_LOG(LogRep, Warning, TEXT("FDiffPropertiesImpl: Property different: %s"), *Cmd.Property->GetFullName());
 				return;
 			}
 
-			if ( (Parents[Cmd.ParentIndex].Flags & ParentPropertyFlags) == ERepParentFlags::None )
+			if ((Parents[Cmd.ParentIndex].Flags & ParentPropertyFlags) == ERepParentFlags::None)
 			{
 				return;
 			}
 
-			StoreProperty( Cmd, (void*)( Data + Cmd.Offset ), (const void*)( ShadowData + Cmd.Offset ) );
+			StoreProperty(Cmd, (void*)(Data + Cmd.Offset), (const void*)(ShadowData + Cmd.Offset));
 
-			if ( RepNotifies && Parent.Property->HasAnyPropertyFlags( CPF_RepNotify ) )
+			if (RepNotifies && Parent.Property->HasAnyPropertyFlags(CPF_RepNotify))
 			{
-				RepNotifies->AddUnique( Parent.Property );
+				RepNotifies->AddUnique(Parent.Property);
 			}
 		}
 		else
 		{
-			UE_CLOG( LogSkippedRepNotifies > 0, LogRep, Display, TEXT( "FDiffPropertiesImpl: Skipping RepNotify because values are the same: %s" ), *Cmd.Property->GetFullName() );
+			UE_CLOG(LogSkippedRepNotifies> 0, LogRep, Display, TEXT("FDiffPropertiesImpl: Skipping RepNotify because values are the same: %s"), *Cmd.Property->GetFullName());
 		}
 	}
 
 	bool IsDifferent() const { return bDifferent; }
 
 private:
-	EDiffPropertiesFlags	Flags;
-	ERepParentFlags			ParentPropertyFlags;
-	TArray< UProperty * >*	RepNotifies;
-	bool					bDifferent;
+	EDiffPropertiesFlags Flags;
+	ERepParentFlags ParentPropertyFlags;
+	TArray<UProperty*>* RepNotifies;
+	bool bDifferent;
 };
 
-bool FRepLayout::DiffProperties( TArray<UProperty*>* RepNotifies, void* RESTRICT Destination, const void* RESTRICT Source, const bool bSync ) const
+bool FRepLayout::DiffProperties(
+	TArray<UProperty*>*		RepNotifies,
+	void* RESTRICT			Destination,
+	const void* RESTRICT	Source,
+	const bool bSync)		const
 {
 	return DiffProperties(RepNotifies, Destination, Source, bSync ? EDiffPropertiesFlags::Sync : EDiffPropertiesFlags::None);
 }
 
-bool FRepLayout::DiffProperties( TArray<UProperty*>* RepNotifies, void* RESTRICT Destination, const void* RESTRICT Source, const EDiffPropertiesFlags Flags ) const
+bool FRepLayout::DiffProperties(
+	TArray<UProperty*>*			RepNotifies,
+	void* RESTRICT				Destination,
+	const void* RESTRICT		Source,
+	const EDiffPropertiesFlags	Flags) const
 {
-	FDiffPropertiesImpl DiffPropertiesImpl( Flags, RepNotifies, Parents, Cmds );
+	FDiffPropertiesImpl DiffPropertiesImpl(Flags, RepNotifies, Parents, Cmds);
 
-	DiffPropertiesImpl.ProcessCmds( (uint8*)Destination, (uint8*)Source );
+	DiffPropertiesImpl.ProcessCmds((uint8*)Destination, (uint8*)Source);
 
 	return DiffPropertiesImpl.IsDifferent();
 }
 
-class FDiffStablePropertiesImpl : public FRepLayoutCmdIterator< FDiffStablePropertiesImpl, FCmdIteratorBaseStackState >
+class FDiffStablePropertiesImpl : public FRepLayoutCmdIterator<FDiffStablePropertiesImpl, FCmdIteratorBaseStackState>
 {
 public:
-	FDiffStablePropertiesImpl( TArray<UProperty*>* InRepNotifies, TArray<UObject*>* InObjReferences, const TArray< FRepParentCmd >& InParents, const TArray< FRepLayoutCmd >& InCmds ) : 
-		FRepLayoutCmdIterator( InParents, InCmds ),
+	FDiffStablePropertiesImpl(
+		TArray<UProperty*>*				InRepNotifies,
+		TArray<UObject*>*				InObjReferences,
+		const TArray<FRepParentCmd>&	InParents,
+		const TArray<FRepLayoutCmd>&	InCmds
+	): 
+		FRepLayoutCmdIterator(InParents, InCmds),
 		RepNotifies(InRepNotifies),
 		ObjReferences(InObjReferences),
-		bDifferent( false )
+		bDifferent(false)
 	{}
 
-	INIT_STACK( FCmdIteratorBaseStackState ) { }
+	INIT_STACK(FCmdIteratorBaseStackState) { }
 
 	SHOULD_PROCESS_NEXT_CMD() 
 	{ 
 		return true;
 	}
 
-	PROCESS_ARRAY_CMD( FCmdIteratorBaseStackState ) 
+	PROCESS_ARRAY_CMD(FCmdIteratorBaseStackState) 
 	{
-		if ( StackState.DataArray->Num() != StackState.ShadowArray->Num() )
+		if (StackState.DataArray->Num() != StackState.ShadowArray->Num())
 		{
 			bDifferent = true;
 
-			if ( (Parents[Cmd.ParentIndex].Flags & ERepParentFlags::IsLifetime) == ERepParentFlags::None )
+			if ((Parents[Cmd.ParentIndex].Flags & ERepParentFlags::IsLifetime) == ERepParentFlags::None)
 			{
 				// Currently, only lifetime properties init from their defaults
 				return;
 			}
 
 			// Do not adjust source data, only the destination
-			FScriptArrayHelper ArrayHelper( (UArrayProperty *)Cmd.Property, Data );
-			ArrayHelper.Resize( StackState.ShadowArray->Num() );
+			FScriptArrayHelper ArrayHelper((UArrayProperty *)Cmd.Property, Data);
+			ArrayHelper.Resize(StackState.ShadowArray->Num());
 		}
 
-		StackState.BaseData			= (uint8*)StackState.DataArray->GetData();
-		StackState.ShadowBaseData	= (uint8*)StackState.ShadowArray->GetData();
+		StackState.BaseData = (uint8*)StackState.DataArray->GetData();
+		StackState.ShadowBaseData = (uint8*)StackState.ShadowArray->GetData();
 
 		// Loop over array
-		ProcessDataArrayElements_r( StackState, Cmd );
+		ProcessDataArrayElements_r(StackState, Cmd);
 	}
 
-	PROCESS_CMD( FCmdIteratorBaseStackState ) 
+	PROCESS_CMD(FCmdIteratorBaseStackState) 
 	{
 		const FRepParentCmd& Parent = Parents[Cmd.ParentIndex];
 
 		// Make the shadow state match the actual state at the time of send
-		if ( !PropertiesAreIdentical( Cmd, (const void*)( Data + Cmd.Offset ), (const void*)( ShadowData + Cmd.Offset ) ) )
+		if (!PropertiesAreIdentical(Cmd, (const void*)(Data + Cmd.Offset), (const void*)(ShadowData + Cmd.Offset)))
 		{
 			bDifferent = true;
 
-			if ( (Parent.Flags & ERepParentFlags::IsLifetime) == ERepParentFlags::None )
+			if ((Parent.Flags & ERepParentFlags::IsLifetime) == ERepParentFlags::None)
 			{
 				// Currently, only lifetime properties init from their defaults
 				return;
@@ -2955,7 +3104,7 @@ public:
 
 			if (Cmd.Type == ERepLayoutCmdType::PropertyObject)
 			{
-				UObjectPropertyBase * ObjProperty = CastChecked<UObjectPropertyBase>( Cmd.Property );
+				UObjectPropertyBase* ObjProperty = CastChecked<UObjectPropertyBase>(Cmd.Property);
 				if (ObjProperty)
 				{
 					if (ObjProperty->PropertyClass && (ObjProperty->PropertyClass->IsChildOf(AActor::StaticClass()) || ObjProperty->PropertyClass->IsChildOf(UActorComponent::StaticClass())))
@@ -2964,7 +3113,7 @@ public:
 						return;
 					}
 
-					UObject* ObjValue = ObjProperty->GetObjectPropertyValue( (const void*)( ShadowData + Cmd.Offset ) );
+					UObject* ObjValue = ObjProperty->GetObjectPropertyValue((const void*)(ShadowData + Cmd.Offset));
 					if (ObjValue)
 					{
 						const bool bStableForNetworking = (ObjValue->HasAnyFlags(RF_WasLoaded | RF_DefaultSubObject) || ObjValue->IsNative() || ObjValue->IsDefaultSubobject());
@@ -2982,11 +3131,11 @@ public:
 				}
 			}
 
-			StoreProperty( Cmd, (void*)( Data + Cmd.Offset ), (const void*)( ShadowData + Cmd.Offset ) );
+			StoreProperty(Cmd, (void*)(Data + Cmd.Offset), (const void*)(ShadowData + Cmd.Offset));
 
-			if ( RepNotifies && Parent.Property->HasAnyPropertyFlags( CPF_RepNotify ) )
+			if (RepNotifies && Parent.Property->HasAnyPropertyFlags(CPF_RepNotify))
 			{
-				RepNotifies->AddUnique( Parent.Property );
+				RepNotifies->AddUnique(Parent.Property);
 			}
 		}
 	}
@@ -2996,11 +3145,15 @@ public:
 	bool bDifferent;
 };
 
-bool FRepLayout::DiffStableProperties( TArray<UProperty*>* RepNotifies, TArray<UObject*>* ObjReferences, void* RESTRICT Destination, const void* RESTRICT Source ) const
+bool FRepLayout::DiffStableProperties(
+	TArray<UProperty*>*		RepNotifies,
+	TArray<UObject*>*		ObjReferences,
+	void* RESTRICT			Destination,
+	const void* RESTRICT	Source) const
 {
-	FDiffStablePropertiesImpl DiffPropertiesImpl( RepNotifies, ObjReferences, Parents, Cmds );
+	FDiffStablePropertiesImpl DiffPropertiesImpl(RepNotifies, ObjReferences, Parents, Cmds);
 
-	DiffPropertiesImpl.ProcessCmds( (uint8*)Destination, (uint8*)Source );
+	DiffPropertiesImpl.ProcessCmds((uint8*)Destination, (uint8*)Source);
 
 	return DiffPropertiesImpl.bDifferent;
 }
@@ -3012,7 +3165,14 @@ static FName NAME_Vector_NetQuantize(TEXT("Vector_NetQuantize"));
 static FName NAME_UniqueNetIdRepl(TEXT("UniqueNetIdRepl"));
 static FName NAME_RepMovement(TEXT("RepMovement"));
 
-uint32 FRepLayout::AddPropertyCmd( UProperty * Property, int32 Offset, int32 RelativeHandle, int32 ParentIndex, uint32 ParentChecksum, int32 StaticArrayIndex, const UNetConnection* ServerConnection )
+uint32 FRepLayout::AddPropertyCmd(
+	UProperty*				Property,
+	int32					Offset,
+	int32					RelativeHandle,
+	int32					ParentIndex,
+	uint32					ParentChecksum,
+	int32					StaticArrayIndex,
+	const UNetConnection*	ServerConnection)
 {
 	SCOPE_CYCLE_COUNTER(STAT_RepLayout_AddPropertyCmd);
 
@@ -3020,105 +3180,105 @@ uint32 FRepLayout::AddPropertyCmd( UProperty * Property, int32 Offset, int32 Rel
 
 	FRepLayoutCmd & Cmd = Cmds[Index];
 
-	Cmd.Property			= Property;
-	Cmd.Type				= ERepLayoutCmdType::Property;		// Initially set to generic type
-	Cmd.Offset				= Offset;
-	Cmd.ElementSize			= Property->ElementSize;
-	Cmd.RelativeHandle		= RelativeHandle;
-	Cmd.ParentIndex			= ParentIndex;
-	Cmd.CompatibleChecksum	= GetRepLayoutCmdCompatibleChecksum(Property, ServerConnection, StaticArrayIndex, ParentChecksum);
+	Cmd.Property = Property;
+	Cmd.Type = ERepLayoutCmdType::Property;		// Initially set to generic type
+	Cmd.Offset = Offset;
+	Cmd.ElementSize = Property->ElementSize;
+	Cmd.RelativeHandle = RelativeHandle;
+	Cmd.ParentIndex = ParentIndex;
+	Cmd.CompatibleChecksum = GetRepLayoutCmdCompatibleChecksum(Property, ServerConnection, StaticArrayIndex, ParentChecksum);
 
-	UProperty * UnderlyingProperty = Property;
-	if ( UEnumProperty * EnumProperty = Cast< UEnumProperty >( Property ) )
+	UProperty* UnderlyingProperty = Property;
+	if (UEnumProperty* EnumProperty = Cast<UEnumProperty>(Property))
 	{
 		UnderlyingProperty = EnumProperty->GetUnderlyingProperty();
 	}
 
 	// Try to special case to custom types we know about
-	if ( UnderlyingProperty->IsA( UStructProperty::StaticClass() ) )
+	if (UnderlyingProperty->IsA(UStructProperty::StaticClass()))
 	{
-		UStructProperty * StructProp = Cast< UStructProperty >( UnderlyingProperty );
-		UScriptStruct * Struct = StructProp->Struct;
-		if ( Struct->GetFName() == NAME_Vector )
+		UStructProperty* StructProp = Cast<UStructProperty>(UnderlyingProperty);
+		UScriptStruct* Struct = StructProp->Struct;
+		if (Struct->GetFName() == NAME_Vector)
 		{
 			Cmd.Type = ERepLayoutCmdType::PropertyVector;
 		}
-		else if ( Struct->GetFName() == NAME_Rotator )
+		else if (Struct->GetFName() == NAME_Rotator)
 		{
 			Cmd.Type = ERepLayoutCmdType::PropertyRotator;
 		}
-		else if ( Struct->GetFName() == NAME_Plane )
+		else if (Struct->GetFName() == NAME_Plane)
 		{
 			Cmd.Type = ERepLayoutCmdType::PropertyPlane;
 		}
-		else if ( Struct->GetFName() == NAME_Vector_NetQuantize100 )
+		else if (Struct->GetFName() == NAME_Vector_NetQuantize100)
 		{
 			Cmd.Type = ERepLayoutCmdType::PropertyVector100;
 		}
-		else if ( Struct->GetFName() == NAME_Vector_NetQuantize10 )
+		else if (Struct->GetFName() == NAME_Vector_NetQuantize10)
 		{
 			Cmd.Type = ERepLayoutCmdType::PropertyVector10;
 		}
-		else if ( Struct->GetFName() == NAME_Vector_NetQuantizeNormal )
+		else if (Struct->GetFName() == NAME_Vector_NetQuantizeNormal)
 		{
 			Cmd.Type = ERepLayoutCmdType::PropertyVectorNormal;
 		}
-		else if ( Struct->GetFName() == NAME_Vector_NetQuantize )
+		else if (Struct->GetFName() == NAME_Vector_NetQuantize)
 		{
 			Cmd.Type = ERepLayoutCmdType::PropertyVectorQ;
 		}
-		else if ( Struct->GetFName() == NAME_UniqueNetIdRepl )
+		else if (Struct->GetFName() == NAME_UniqueNetIdRepl)
 		{
 			Cmd.Type = ERepLayoutCmdType::PropertyNetId;
 		}
-		else if ( Struct->GetFName() == NAME_RepMovement )
+		else if (Struct->GetFName() == NAME_RepMovement)
 		{
 			Cmd.Type = ERepLayoutCmdType::RepMovement;
 		}
 		else
 		{
-			UE_LOG( LogRep, VeryVerbose, TEXT( "AddPropertyCmd: Falling back to default type for property [%s]" ), *Cmd.Property->GetFullName() );
+			UE_LOG(LogRep, VeryVerbose, TEXT("AddPropertyCmd: Falling back to default type for property [%s]"), *Cmd.Property->GetFullName());
 		}
 	}
-	else if ( UnderlyingProperty->IsA( UBoolProperty::StaticClass() ) )
+	else if (UnderlyingProperty->IsA(UBoolProperty::StaticClass()))
 	{
 		Cmd.Type = ERepLayoutCmdType::PropertyBool;
 	}
-	else if ( UnderlyingProperty->IsA( UFloatProperty::StaticClass() ) )
+	else if (UnderlyingProperty->IsA(UFloatProperty::StaticClass()))
 	{
 		Cmd.Type = ERepLayoutCmdType::PropertyFloat;
 	}
-	else if ( UnderlyingProperty->IsA( UIntProperty::StaticClass() ) )
+	else if (UnderlyingProperty->IsA(UIntProperty::StaticClass()))
 	{
 		Cmd.Type = ERepLayoutCmdType::PropertyInt;
 	}
-	else if ( UnderlyingProperty->IsA( UByteProperty::StaticClass() ) )
+	else if (UnderlyingProperty->IsA(UByteProperty::StaticClass()))
 	{
 		Cmd.Type = ERepLayoutCmdType::PropertyByte;
 	}
-	else if ( UnderlyingProperty->IsA( UObjectPropertyBase::StaticClass() ) )
+	else if (UnderlyingProperty->IsA(UObjectPropertyBase::StaticClass()))
 	{
 		Cmd.Type = ERepLayoutCmdType::PropertyObject;
 	}
-	else if ( UnderlyingProperty->IsA( UNameProperty::StaticClass() ) )
+	else if (UnderlyingProperty->IsA(UNameProperty::StaticClass()))
 	{
 		Cmd.Type = ERepLayoutCmdType::PropertyName;
 	}
-	else if ( UnderlyingProperty->IsA( UUInt32Property::StaticClass() ) )
+	else if (UnderlyingProperty->IsA(UUInt32Property::StaticClass()))
 	{
 		Cmd.Type = ERepLayoutCmdType::PropertyUInt32;
 	}
-	else if ( UnderlyingProperty->IsA( UUInt64Property::StaticClass() ) )
+	else if (UnderlyingProperty->IsA(UUInt64Property::StaticClass()))
 	{
 		Cmd.Type = ERepLayoutCmdType::PropertyUInt64;
 	}
-	else if ( UnderlyingProperty->IsA( UStrProperty::StaticClass() ) )
+	else if (UnderlyingProperty->IsA(UStrProperty::StaticClass()))
 	{
 		Cmd.Type = ERepLayoutCmdType::PropertyString;
 	}
 	else
 	{
-		UE_LOG( LogRep, VeryVerbose, TEXT( "AddPropertyCmd: Falling back to default type for property [%s]" ), *Cmd.Property->GetFullName() );
+		UE_LOG(LogRep, VeryVerbose, TEXT("AddPropertyCmd: Falling back to default type for property [%s]"), *Cmd.Property->GetFullName());
 	}
 
 	// Cannot write a shared version of a property that depends on per-connection data (the PackageMap).
@@ -3132,19 +3292,26 @@ uint32 FRepLayout::AddPropertyCmd( UProperty * Property, int32 Offset, int32 Rel
 	return Cmd.CompatibleChecksum;
 }
 
-uint32 FRepLayout::AddArrayCmd( UArrayProperty * Property, int32 Offset, int32 RelativeHandle, int32 ParentIndex, uint32 ParentChecksum, int32 StaticArrayIndex, const UNetConnection* ServerConnection )
+uint32 FRepLayout::AddArrayCmd(
+	UArrayProperty*			Property,
+	int32					Offset,
+	int32					RelativeHandle,
+	int32					ParentIndex,
+	uint32					ParentChecksum,
+	int32					StaticArrayIndex,
+	const UNetConnection*	ServerConnection)
 {
 	const int32 Index = Cmds.AddZeroed();
 
 	FRepLayoutCmd & Cmd = Cmds[Index];
 
-	Cmd.Type				= ERepLayoutCmdType::DynamicArray;
-	Cmd.Property			= Property;
-	Cmd.Offset				= Offset;
-	Cmd.ElementSize			= Property->Inner->ElementSize;
-	Cmd.RelativeHandle		= RelativeHandle;
-	Cmd.ParentIndex			= ParentIndex;
-	Cmd.CompatibleChecksum	= GetRepLayoutCmdCompatibleChecksum(Property, ServerConnection, StaticArrayIndex, ParentChecksum);
+	Cmd.Type = ERepLayoutCmdType::DynamicArray;
+	Cmd.Property = Property;
+	Cmd.Offset = Offset;
+	Cmd.ElementSize = Property->Inner->ElementSize;
+	Cmd.RelativeHandle = RelativeHandle;
+	Cmd.ParentIndex = ParentIndex;
+	Cmd.CompatibleChecksum = GetRepLayoutCmdCompatibleChecksum(Property, ServerConnection, StaticArrayIndex, ParentChecksum);
 
 	return Cmd.CompatibleChecksum;
 }
@@ -3158,19 +3325,26 @@ void FRepLayout::AddReturnCmd()
 	Cmd.Type = ERepLayoutCmdType::Return;
 }
 
-int32 FRepLayout::InitFromProperty_r( UProperty * Property, int32 Offset, int32 RelativeHandle, int32 ParentIndex, uint32 ParentChecksum, int32 StaticArrayIndex, const UNetConnection* ServerConnection )
+int32 FRepLayout::InitFromProperty_r(
+	UProperty*				Property,
+	int32					Offset,
+	int32					RelativeHandle,
+	int32					ParentIndex,
+	uint32					ParentChecksum,
+	int32					StaticArrayIndex,
+	const UNetConnection*	ServerConnection)
 {
-	UArrayProperty * ArrayProp = Cast< UArrayProperty >( Property );
+	UArrayProperty * ArrayProp = Cast<UArrayProperty>(Property);
 
-	if ( ArrayProp != NULL )
+	if (ArrayProp != NULL)
 	{
 		const int32 CmdStart = Cmds.Num();
 
 		RelativeHandle++;
 
-		const uint32 ArrayChecksum = AddArrayCmd( ArrayProp, Offset + ArrayProp->GetOffset_ForGC(), RelativeHandle, ParentIndex, ParentChecksum, StaticArrayIndex, ServerConnection );
+		const uint32 ArrayChecksum = AddArrayCmd(ArrayProp, Offset + ArrayProp->GetOffset_ForGC(), RelativeHandle, ParentIndex, ParentChecksum, StaticArrayIndex, ServerConnection);
 
-		InitFromProperty_r( ArrayProp->Inner, 0, 0, ParentIndex, ArrayChecksum, 0, ServerConnection );
+		InitFromProperty_r(ArrayProp->Inner, 0, 0, ParentIndex, ArrayChecksum, 0, ServerConnection);
 		
 		AddReturnCmd();
 
@@ -3179,44 +3353,45 @@ int32 FRepLayout::InitFromProperty_r( UProperty * Property, int32 Offset, int32
 		return RelativeHandle;
 	}
 
-	UStructProperty * StructProp = Cast< UStructProperty >( Property );
+	UStructProperty* StructProp = Cast<UStructProperty>(Property);
 
-	if ( StructProp != NULL )
+	if (StructProp != NULL)
 	{
-		UScriptStruct * Struct = StructProp->Struct;
+		UScriptStruct* Struct = StructProp->Struct;
 
-		if ( Struct->StructFlags & STRUCT_NetDeltaSerializeNative )
+		if (Struct->StructFlags & STRUCT_NetDeltaSerializeNative)
 		{
 			// Custom delta serializers handles outside of FRepLayout
 			return RelativeHandle;
 		}
 
-		if ( Struct->StructFlags & STRUCT_NetSerializeNative )
+		if (Struct->StructFlags & STRUCT_NetSerializeNative)
 		{
 			RelativeHandle++;
-			AddPropertyCmd( Property, Offset + Property->GetOffset_ForGC(), RelativeHandle, ParentIndex, ParentChecksum, StaticArrayIndex, ServerConnection );
+			AddPropertyCmd(Property, Offset + Property->GetOffset_ForGC(), RelativeHandle, ParentIndex, ParentChecksum, StaticArrayIndex, ServerConnection);
 			return RelativeHandle;
 		}
 
-		TArray< UProperty * > NetProperties;		// Track properties so me can ensure they are sorted by offsets at the end
+		// Track properties so me can ensure they are sorted by offsets at the end
+		TArray<UProperty*> NetProperties;
 
-		for ( TFieldIterator<UProperty> It( Struct ); It; ++It )
+		for (TFieldIterator<UProperty> It(Struct); It; ++It)
 		{
-			if ( ( It->PropertyFlags & CPF_RepSkip ) )
+			if ((It->PropertyFlags & CPF_RepSkip))
 			{
 				continue;
 			}
 
-			NetProperties.Add( *It );
+			NetProperties.Add(*It);
 		}
 
 		// Sort NetProperties by memory offset
 		struct FCompareUFieldOffsets
 		{
-			FORCEINLINE bool operator()( UProperty & A, UProperty & B ) const
+			FORCEINLINE bool operator()(UProperty & A, UProperty & B) const
 			{
 				// Ensure stable sort
-				if ( A.GetOffset_ForGC() == B.GetOffset_ForGC() )
+				if (A.GetOffset_ForGC() == B.GetOffset_ForGC())
 				{
 					return A.GetName() < B.GetName();
 				}
@@ -3225,15 +3400,15 @@ int32 FRepLayout::InitFromProperty_r( UProperty * Property, int32 Offset, int32
 			}
 		};
 
-		Sort( NetProperties.GetData(), NetProperties.Num(), FCompareUFieldOffsets() );
+		Sort(NetProperties.GetData(), NetProperties.Num(), FCompareUFieldOffsets());
 
 		const uint32 StructChecksum = GetRepLayoutCmdCompatibleChecksum(Property, ServerConnection, StaticArrayIndex, ParentChecksum);
 
-		for ( int32 i = 0; i < NetProperties.Num(); i++ )
+		for (int32 i = 0; i < NetProperties.Num(); i++)
 		{
-			for ( int32 j = 0; j < NetProperties[i]->ArrayDim; j++ )
+			for (int32 j = 0; j < NetProperties[i]->ArrayDim; j++)
 			{
-				RelativeHandle = InitFromProperty_r( NetProperties[i], Offset + StructProp->GetOffset_ForGC() + j * NetProperties[i]->ElementSize, RelativeHandle, ParentIndex, StructChecksum, j, ServerConnection );
+				RelativeHandle = InitFromProperty_r(NetProperties[i], Offset + StructProp->GetOffset_ForGC() + j * NetProperties[i]->ElementSize, RelativeHandle, ParentIndex, StructChecksum, j, ServerConnection);
 			}
 		}
 		return RelativeHandle;
@@ -3242,17 +3417,17 @@ int32 FRepLayout::InitFromProperty_r( UProperty * Property, int32 Offset, int32
 	// Add actual property
 	RelativeHandle++;
 
-	AddPropertyCmd( Property, Offset + Property->GetOffset_ForGC(), RelativeHandle, ParentIndex, ParentChecksum, StaticArrayIndex, ServerConnection );
+	AddPropertyCmd(Property, Offset + Property->GetOffset_ForGC(), RelativeHandle, ParentIndex, ParentChecksum, StaticArrayIndex, ServerConnection);
 
 	return RelativeHandle;
 }
 
-uint16 FRepLayout::AddParentProperty( UProperty * Property, int32 ArrayIndex )
+uint16 FRepLayout::AddParentProperty(UProperty* Property, int32 ArrayIndex)
 {
-	return Parents.Add( FRepParentCmd( Property, ArrayIndex ) );
+	return Parents.Add(FRepParentCmd(Property, ArrayIndex));
 }
 
-void FRepLayout::InitFromObjectClass( UClass * InObjectClass, const UNetConnection* ServerConnection )
+void FRepLayout::InitFromObjectClass(UClass* InObjectClass, const UNetConnection* ServerConnection)
 {
 	SCOPE_CYCLE_COUNTER(STAT_RepLayout_InitFromObjectClass);
 	SCOPE_CYCLE_UOBJECT(ObjectClass, InObjectClass);
@@ -3268,67 +3443,67 @@ void FRepLayout::InitFromObjectClass( UClass * InObjectClass, const UNetConnecti
 
 	InObjectClass->SetUpRuntimeReplicationData();
 
-	for ( int32 i = 0; i < InObjectClass->ClassReps.Num(); i++ )
+	for (int32 i = 0; i < InObjectClass->ClassReps.Num(); i++)
 	{
-		UProperty * Property	= InObjectClass->ClassReps[i].Property;
-		const int32 ArrayIdx	= InObjectClass->ClassReps[i].Index;
+		UProperty * Property = InObjectClass->ClassReps[i].Property;
+		const int32 ArrayIdx = InObjectClass->ClassReps[i].Index;
 
-		check( Property->PropertyFlags & CPF_Net );
+		check(Property->PropertyFlags & CPF_Net);
 
-		const int32 ParentHandle = AddParentProperty( Property, ArrayIdx );
+		const int32 ParentHandle = AddParentProperty(Property, ArrayIdx);
 
-		check( ParentHandle == i );
-		check( Parents[i].Property->RepIndex + Parents[i].ArrayIndex == i );
+		check(ParentHandle == i);
+		check(Parents[i].Property->RepIndex + Parents[i].ArrayIndex == i);
 
 		Parents[ParentHandle].CmdStart = Cmds.Num();
-		RelativeHandle = InitFromProperty_r( Property, Property->ElementSize * ArrayIdx, RelativeHandle, ParentHandle, 0, ArrayIdx, ServerConnection );
+		RelativeHandle = InitFromProperty_r(Property, Property->ElementSize * ArrayIdx, RelativeHandle, ParentHandle, 0, ArrayIdx, ServerConnection);
 		Parents[ParentHandle].CmdEnd = Cmds.Num();
 		Parents[ParentHandle].Flags |= ERepParentFlags::IsConditional;
 
-		if ( Parents[i].CmdEnd > Parents[i].CmdStart )
+		if (Parents[i].CmdEnd > Parents[i].CmdStart)
 		{
-			check( Cmds[Parents[i].CmdStart].Offset >= LastOffset );		// >= since bool's can be combined
+			check(Cmds[Parents[i].CmdStart].Offset >= LastOffset);		//>= since bool's can be combined
 			LastOffset = Cmds[Parents[i].CmdStart].Offset;
 		}
 
 		// Setup flags
-		if ( IsCustomDeltaProperty( Property ) )
+		if (IsCustomDeltaProperty(Property))
 		{
 			Parents[ParentHandle].Flags |= ERepParentFlags::IsCustomDelta;
 		}
 
-		if ( Property->GetPropertyFlags() & CPF_Config )
+		if (Property->GetPropertyFlags() & CPF_Config)
 		{
 			Parents[ParentHandle].Flags |= ERepParentFlags::IsConfig;
 		}
 
 		// Hijack the first non custom property for identifying this as a rep layout block
-		if ( FirstNonCustomParent == -1 && Property->ArrayDim == 1 && ((Parents[ParentHandle].Flags & ERepParentFlags::IsCustomDelta) == ERepParentFlags::None) )
+		if (FirstNonCustomParent == -1 && Property->ArrayDim == 1 && ((Parents[ParentHandle].Flags & ERepParentFlags::IsCustomDelta) == ERepParentFlags::None))
 		{
 			FirstNonCustomParent = ParentHandle;
 		}
 
 		// Find Role/RemoteRole property indexes so we can swap them on the client
-		if ( Property->GetFName() == NAME_Role )
+		if (Property->GetFName() == NAME_Role)
 		{
-			check( RoleIndex == -1 );
-			check( Parents[ParentHandle].CmdEnd == Parents[ParentHandle].CmdStart + 1 );
+			check(RoleIndex == -1);
+			check(Parents[ParentHandle].CmdEnd == Parents[ParentHandle].CmdStart + 1);
 			RoleIndex = ParentHandle;
 		}
 
-		if ( Property->GetFName() == NAME_RemoteRole )
+		if (Property->GetFName() == NAME_RemoteRole)
 		{
-			check( RemoteRoleIndex == -1 );
-			check( Parents[ParentHandle].CmdEnd == Parents[ParentHandle].CmdStart + 1 );
+			check(RemoteRoleIndex == -1);
+			check(Parents[ParentHandle].CmdEnd == Parents[ParentHandle].CmdStart + 1);
 			RemoteRoleIndex = ParentHandle;
 		}
 	}
 
 	// Make sure it either found both, or didn't find either
-	check( ( RoleIndex == -1 ) == ( RemoteRoleIndex == -1 ) );
+	check((RoleIndex == -1) == (RemoteRoleIndex == -1));
 
 	// This is so the receiving side can swap these as it receives them
-	if ( RoleIndex != -1 )
+	if (RoleIndex != -1)
 	{
 		Parents[RoleIndex].RoleSwapIndex = RemoteRoleIndex;
 		Parents[RemoteRoleIndex].RoleSwapIndex = RoleIndex;
@@ -3337,14 +3512,15 @@ void FRepLayout::InitFromObjectClass( UClass * InObjectClass, const UNetConnecti
 	AddReturnCmd();
 
 	// Initialize lifetime props
-	TArray< FLifetimeProperty >	LifetimeProps;			// Properties that replicate for the lifetime of the channel
+	// Properties that replicate for the lifetime of the channel
+	TArray<FLifetimeProperty> LifetimeProps;
 
 	UObject* Object = InObjectClass->GetDefaultObject();
 
-	Object->GetLifetimeReplicatedProps( LifetimeProps );
+	Object->GetLifetimeReplicatedProps(LifetimeProps);
 
 	// Setup lifetime replicated properties
-	for ( int32 i = 0; i < LifetimeProps.Num(); i++ )
+	for (int32 i = 0; i < LifetimeProps.Num(); i++)
 	{
 		const int32 ParentIndex = LifetimeProps[i].RepIndex;
 
@@ -3355,97 +3531,98 @@ void FRepLayout::InitFromObjectClass( UClass * InObjectClass, const UNetConnecti
 		}
 
 		// Store the condition on the parent in case we need it
-		Parents[ParentIndex].Condition			= LifetimeProps[i].Condition;
-		Parents[ParentIndex].RepNotifyCondition	= LifetimeProps[i].RepNotifyCondition;
+		Parents[ParentIndex].Condition = LifetimeProps[i].Condition;
+		Parents[ParentIndex].RepNotifyCondition = LifetimeProps[i].RepNotifyCondition;
 
-		if ( (Parents[ParentIndex].Flags & ERepParentFlags::IsCustomDelta) != ERepParentFlags::None )
+		if ((Parents[ParentIndex].Flags & ERepParentFlags::IsCustomDelta) != ERepParentFlags::None)
 		{
-			continue;		// We don't handle custom properties in the FRepLayout class
+			// We don't handle custom properties in the FRepLayout class
+			continue;
 		}
 
 		Parents[ParentIndex].Flags |= ERepParentFlags::IsLifetime;
 
-		if ( ParentIndex == RemoteRoleIndex )
+		if (ParentIndex == RemoteRoleIndex)
 		{
 			// We handle remote role specially, since it can change between connections when downgraded
 			// So we force it on the conditional list
-			check( LifetimeProps[i].Condition == COND_None );
+			check(LifetimeProps[i].Condition == COND_None);
 			LifetimeProps[i].Condition = COND_Custom;
-			continue;		
+			continue;
 		}
 
-		if ( LifetimeProps[i].Condition == COND_None )
+		if (LifetimeProps[i].Condition == COND_None)
 		{
 			Parents[ParentIndex].Flags &= ~ERepParentFlags::IsConditional;
-		}			
+		}
 	}
 
-	BuildHandleToCmdIndexTable_r( 0, Cmds.Num() - 1, BaseHandleToCmdIndex );
+	BuildHandleToCmdIndexTable_r(0, Cmds.Num() - 1, BaseHandleToCmdIndex);
 
 	Owner = InObjectClass;
 }
 
-void FRepLayout::InitFromFunction( UFunction * InFunction, const UNetConnection* ServerConnection )
+void FRepLayout::InitFromFunction(UFunction* InFunction, const UNetConnection* ServerConnection)
 {
 	int32 RelativeHandle = 0;
 
-	for ( TFieldIterator<UProperty> It( InFunction ); It && ( It->PropertyFlags & ( CPF_Parm | CPF_ReturnParm ) ) == CPF_Parm; ++It )
+	for (TFieldIterator<UProperty> It(InFunction); It && (It->PropertyFlags & (CPF_Parm | CPF_ReturnParm)) == CPF_Parm; ++It)
 	{
-		for ( int32 ArrayIdx = 0; ArrayIdx < It->ArrayDim; ++ArrayIdx )
+		for (int32 ArrayIdx = 0; ArrayIdx < It->ArrayDim; ++ArrayIdx)
 		{
-			const int32 ParentHandle = AddParentProperty( *It, ArrayIdx );
+			const int32 ParentHandle = AddParentProperty(*It, ArrayIdx);
 			Parents[ParentHandle].CmdStart = Cmds.Num();
-			RelativeHandle = InitFromProperty_r( *It, It->ElementSize * ArrayIdx, RelativeHandle, ParentHandle, 0, ArrayIdx, ServerConnection );
+			RelativeHandle = InitFromProperty_r(*It, It->ElementSize * ArrayIdx, RelativeHandle, ParentHandle, 0, ArrayIdx, ServerConnection);
 			Parents[ParentHandle].CmdEnd = Cmds.Num();
 		}
 	}
 
 	AddReturnCmd();
 
-	BuildHandleToCmdIndexTable_r( 0, Cmds.Num() - 1, BaseHandleToCmdIndex );
+	BuildHandleToCmdIndexTable_r(0, Cmds.Num() - 1, BaseHandleToCmdIndex);
 
 	Owner = InFunction;
 }
 
-void FRepLayout::InitFromStruct( UStruct * InStruct, const UNetConnection* ServerConnection )
+void FRepLayout::InitFromStruct(UStruct* InStruct, const UNetConnection* ServerConnection)
 {
 	int32 RelativeHandle = 0;
 
-	for ( TFieldIterator<UProperty> It( InStruct ); It; ++It )
+	for (TFieldIterator<UProperty> It(InStruct); It; ++It)
 	{
-		if ( It->PropertyFlags & CPF_RepSkip )
+		if (It->PropertyFlags & CPF_RepSkip)
 		{
 			continue;
 		}
 			
-		for ( int32 ArrayIdx = 0; ArrayIdx < It->ArrayDim; ++ArrayIdx )
+		for (int32 ArrayIdx = 0; ArrayIdx < It->ArrayDim; ++ArrayIdx)
 		{
-			const int32 ParentHandle = AddParentProperty( *It, ArrayIdx );
+			const int32 ParentHandle = AddParentProperty(*It, ArrayIdx);
 			Parents[ParentHandle].CmdStart = Cmds.Num();
-			RelativeHandle = InitFromProperty_r( *It, It->ElementSize * ArrayIdx, RelativeHandle, ParentHandle, 0, ArrayIdx, ServerConnection );
+			RelativeHandle = InitFromProperty_r(*It, It->ElementSize * ArrayIdx, RelativeHandle, ParentHandle, 0, ArrayIdx, ServerConnection);
 			Parents[ParentHandle].CmdEnd = Cmds.Num();
 		}
 	}
 
 	AddReturnCmd();
 
-	BuildHandleToCmdIndexTable_r( 0, Cmds.Num() - 1, BaseHandleToCmdIndex );
+	BuildHandleToCmdIndexTable_r(0, Cmds.Num() - 1, BaseHandleToCmdIndex);
 
 	Owner = InStruct;
 }
 
-void FRepLayout::SerializeProperties_DynamicArray_r( 
-	FBitArchive &		Ar, 
-	UPackageMap	*		Map,
-	const int32			CmdIndex,
-	uint8 *				Data,
-	bool &				bHasUnmapped,
-	const int32			ArrayDepth,
-	const FRepSerializationSharedInfo&	SharedInfo ) const
+void FRepLayout::SerializeProperties_DynamicArray_r(
+	FBitArchive&						Ar, 
+	UPackageMap*						Map,
+	const int32							CmdIndex,
+	uint8*								Data,
+	bool&								bHasUnmapped,
+	const int32							ArrayDepth,
+	const FRepSerializationSharedInfo&	SharedInfo) const
 {
 	const FRepLayoutCmd& Cmd = Cmds[ CmdIndex ];
 
-	FScriptArray * Array = (FScriptArray *)Data;	
+	FScriptArray * Array = (FScriptArray *)Data;
 
 	uint16 OutArrayNum = Array->Num();
 	Ar << OutArrayNum;
@@ -3455,7 +3632,7 @@ void FRepLayout::SerializeProperties_DynamicArray_r(
 	const int32 ArrayNum = Ar.IsLoading() ? (int32)OutArrayNum : Array->Num();
 
 	// Validate the maximum number of elements.
-	if (ArrayNum > MaxRepArraySize)
+	if (ArrayNum> MaxRepArraySize)
 	{
 		UE_LOG(LogRepTraffic, Error, TEXT("SerializeProperties_DynamicArray_r: ArraySize (%d) > net.MaxRepArraySize(%d) (%s). net.MaxRepArraySize can be updated in Project Settings under Network Settings."),
 			ArrayNum, MaxRepArraySize, *Cmd.Property->GetName());
@@ -3490,26 +3667,26 @@ void FRepLayout::SerializeProperties_DynamicArray_r(
 	}	
 }
 
-void FRepLayout::SerializeProperties_r( 
-	FBitArchive &		Ar, 
-	UPackageMap	*		Map,
-	const int32			CmdStart, 
-	const int32			CmdEnd,
-	void *				Data,
-	bool &				bHasUnmapped,
-	const int32			ArrayIndex,
-	const int32			ArrayDepth,
-	const FRepSerializationSharedInfo&	SharedInfo ) const
+void FRepLayout::SerializeProperties_r(
+	FBitArchive&						Ar, 
+	UPackageMap*						Map,
+	const int32							CmdStart, 
+	const int32							CmdEnd,
+	void*								Data,
+	bool&								bHasUnmapped,
+	const int32							ArrayIndex,
+	const int32							ArrayDepth,
+	const FRepSerializationSharedInfo&	SharedInfo) const
 {
-	for ( int32 CmdIndex = CmdStart; CmdIndex < CmdEnd && !Ar.IsError(); CmdIndex++ )
+	for (int32 CmdIndex = CmdStart; CmdIndex < CmdEnd && !Ar.IsError(); CmdIndex++)
 	{
 		const FRepLayoutCmd& Cmd = Cmds[CmdIndex];
 
-		check( Cmd.Type != ERepLayoutCmdType::Return );
+		check(Cmd.Type != ERepLayoutCmdType::Return);
 
-		if ( Cmd.Type == ERepLayoutCmdType::DynamicArray )
+		if (Cmd.Type == ERepLayoutCmdType::DynamicArray)
 		{
-			SerializeProperties_DynamicArray_r( Ar, Map, CmdIndex, (uint8*)Data + Cmd.Offset, bHasUnmapped, ArrayDepth + 1, SharedInfo );
+			SerializeProperties_DynamicArray_r(Ar, Map, CmdIndex, (uint8*)Data + Cmd.Offset, bHasUnmapped, ArrayDepth + 1, SharedInfo);
 			CmdIndex = Cmd.EndCmd - 1;		// The -1 to handle the ++ in the for loop
 			continue;
 		}
@@ -3517,7 +3694,7 @@ void FRepLayout::SerializeProperties_r(
 #if !(UE_BUILD_SHIPPING || UE_BUILD_TEST)
 		if (GDoReplicationContextString > 0)
 		{
-			Map->SetDebugContextString( FString::Printf(TEXT("%s - %s"), *Owner->GetPathName(), *Cmd.Property->GetPathName() ) );
+			Map->SetDebugContextString(FString::Printf(TEXT("%s - %s"), *Owner->GetPathName(), *Cmd.Property->GetPathName()));
 		}
 #endif
 
@@ -3545,7 +3722,7 @@ void FRepLayout::SerializeProperties_r(
 
 				FBitWriterMark BitWriterMark(Writer);
 
-				Cmd.Property->NetSerializeItem( Writer, Map, (void*)( (uint8*)Data + Cmd.Offset ) );
+				Cmd.Property->NetSerializeItem(Writer, Map, (void*)((uint8*)Data + Cmd.Offset));
 
 				TArray<uint8> StandardBuffer;
 				BitWriterMark.Copy(Writer, StandardBuffer);
@@ -3570,7 +3747,7 @@ void FRepLayout::SerializeProperties_r(
 		else
 		{
 			GNumSharedSerializationMiss++;
-			if ( !Cmd.Property->NetSerializeItem( Ar, Map, (void*)( (uint8*)Data + Cmd.Offset ) ) )
+			if (!Cmd.Property->NetSerializeItem(Ar, Map, (void*)((uint8*)Data + Cmd.Offset)))
 			{
 				bHasUnmapped = true;
 			}
@@ -3585,55 +3762,61 @@ void FRepLayout::SerializeProperties_r(
 	}
 }
 
-void FRepLayout::BuildChangeList_r( const TArray< FHandleToCmdIndex >& HandleToCmdIndex, const int32 CmdStart, const int32 CmdEnd, uint8* Data, const int32 HandleOffset, TArray< uint16 >& Changed ) const
+void FRepLayout::BuildChangeList_r(
+	const TArray<FHandleToCmdIndex>&	HandleToCmdIndex,
+	const int32							CmdStart,
+	const int32							CmdEnd,
+	uint8*								Data,
+	const int32							HandleOffset,
+	TArray<uint16>&						Changed) const
 {
-	for ( int32 CmdIndex = CmdStart; CmdIndex < CmdEnd; CmdIndex++ )
+	for (int32 CmdIndex = CmdStart; CmdIndex < CmdEnd; CmdIndex++)
 	{
 		const FRepLayoutCmd& Cmd = Cmds[CmdIndex];
 
-		check( Cmd.Type != ERepLayoutCmdType::Return );
+		check(Cmd.Type != ERepLayoutCmdType::Return);
 
-		if ( Cmd.Type == ERepLayoutCmdType::DynamicArray )
+		if (Cmd.Type == ERepLayoutCmdType::DynamicArray)
 		{			
-			FScriptArray* Array = ( FScriptArray * )( Data + Cmd.Offset );
+			FScriptArray* Array = (FScriptArray *)(Data + Cmd.Offset);
 
-			TArray< uint16 > ChangedLocal;
+			TArray<uint16> ChangedLocal;
 
-			TArray< FHandleToCmdIndex >& ArrayHandleToCmdIndex = *HandleToCmdIndex[Cmd.RelativeHandle - 1].HandleToCmdIndex;
+			TArray<FHandleToCmdIndex>& ArrayHandleToCmdIndex = *HandleToCmdIndex[Cmd.RelativeHandle - 1].HandleToCmdIndex;
 
 			const int32 ArrayCmdStart			= CmdIndex + 1;
 			const int32 ArrayCmdEnd				= Cmd.EndCmd - 1;
 			const int32 NumHandlesPerElement	= ArrayHandleToCmdIndex.Num();
 
-			check( NumHandlesPerElement > 0 );
+			check(NumHandlesPerElement > 0);
 
-			for ( int32 i = 0; i < Array->Num(); i++ )
+			for (int32 i = 0; i < Array->Num(); i++)
 			{
-				BuildChangeList_r( ArrayHandleToCmdIndex, ArrayCmdStart, ArrayCmdEnd, ((uint8*)Array->GetData()) + Cmd.ElementSize * i, i * NumHandlesPerElement, ChangedLocal );
+				BuildChangeList_r(ArrayHandleToCmdIndex, ArrayCmdStart, ArrayCmdEnd, ((uint8*)Array->GetData()) + Cmd.ElementSize * i, i * NumHandlesPerElement, ChangedLocal);
 			}
 
-			if ( ChangedLocal.Num() )
+			if (ChangedLocal.Num())
 			{
-				Changed.Add( Cmd.RelativeHandle + HandleOffset );	// Identify the array cmd handle
-				Changed.Add( ChangedLocal.Num() );					// This is so we can jump over the array if we need to
-				Changed.Append( ChangedLocal );						// Append the change list under the array
-				Changed.Add( 0 );									// Null terminator
+				Changed.Add(Cmd.RelativeHandle + HandleOffset);	// Identify the array cmd handle
+				Changed.Add(ChangedLocal.Num());					// This is so we can jump over the array if we need to
+				Changed.Append(ChangedLocal);						// Append the change list under the array
+				Changed.Add(0);									// Null terminator
 			}
 
 			CmdIndex = Cmd.EndCmd - 1;		// The -1 to handle the ++ in the for loop
 			continue;
 		}
 
-		Changed.Add( Cmd.RelativeHandle + HandleOffset );
+		Changed.Add(Cmd.RelativeHandle + HandleOffset);
 	}
 }
 
 
 void FRepLayout::BuildSharedSerialization(
-	const uint8* RESTRICT		Data,
-	TArray< uint16 >&			Changed,
-	const bool					bWriteHandle,
-	FRepSerializationSharedInfo& SharedInfo) const
+	const uint8* RESTRICT			Data,
+	TArray<uint16>&					Changed,
+	const bool						bWriteHandle,
+	FRepSerializationSharedInfo&	SharedInfo) const
 {
 #ifdef ENABLE_PROPERTY_CHECKSUMS
 	const bool bDoChecksum = (GDoPropertyChecksum == 1);
@@ -3641,23 +3824,23 @@ void FRepLayout::BuildSharedSerialization(
 	const bool bDoChecksum = false;
 #endif
 
-	FChangelistIterator ChangelistIterator( Changed, 0 );
-	FRepHandleIterator HandleIterator( ChangelistIterator, Cmds, BaseHandleToCmdIndex, 0, 1, 0, Cmds.Num() - 1 );
+	FChangelistIterator ChangelistIterator(Changed, 0);
+	FRepHandleIterator HandleIterator(ChangelistIterator, Cmds, BaseHandleToCmdIndex, 0, 1, 0, Cmds.Num() - 1);
 
-	BuildSharedSerialization_r( HandleIterator, ( uint8* )Data, bWriteHandle, bDoChecksum, 0, SharedInfo );
+	BuildSharedSerialization_r(HandleIterator, (uint8*)Data, bWriteHandle, bDoChecksum, 0, SharedInfo);
 
 	SharedInfo.SetValid();
 }
 
 void FRepLayout::BuildSharedSerialization_r(
-	FRepHandleIterator&		HandleIterator,
-	const uint8* RESTRICT	SourceData,
-	const bool				bWriteHandle,
-	const bool				bDoChecksum,
-	const int32				ArrayDepth,
-	FRepSerializationSharedInfo& SharedInfo ) const
+	FRepHandleIterator&				HandleIterator,
+	const uint8* RESTRICT			SourceData,
+	const bool						bWriteHandle,
+	const bool						bDoChecksum,
+	const int32						ArrayDepth,
+	FRepSerializationSharedInfo&	SharedInfo) const
 {
-	while ( HandleIterator.NextHandle() )
+	while (HandleIterator.NextHandle())
 	{
 		const int32 CmdIndex	= HandleIterator.CmdIndex;
 		const int32 ArrayOffset = HandleIterator.ArrayOffset;
@@ -3667,17 +3850,17 @@ void FRepLayout::BuildSharedSerialization_r(
 
 		const uint8* Data = SourceData + ArrayOffset + Cmd.Offset;
 
-		if ( Cmd.Type == ERepLayoutCmdType::DynamicArray )
+		if (Cmd.Type == ERepLayoutCmdType::DynamicArray)
 		{
-			const FScriptArray* Array = ( FScriptArray * )Data;
-			const uint8* NewData = ( uint8* )Array->GetData();
+			const FScriptArray* Array = (FScriptArray *)Data;
+			const uint8* NewData = (uint8*)Array->GetData();
 
-			FScopedIteratorArrayTracker ArrayTracker( &HandleIterator );
+			FScopedIteratorArrayTracker ArrayTracker(&HandleIterator);
 
-			TArray< FHandleToCmdIndex >& ArrayHandleToCmdIndex = *HandleIterator.HandleToCmdIndex[Cmd.RelativeHandle - 1].HandleToCmdIndex;
+			TArray<FHandleToCmdIndex>& ArrayHandleToCmdIndex = *HandleIterator.HandleToCmdIndex[Cmd.RelativeHandle - 1].HandleToCmdIndex;
 
-			FRepHandleIterator ArrayIterator( HandleIterator.ChangelistIterator, Cmds, ArrayHandleToCmdIndex, Cmd.ElementSize, Array->Num(), CmdIndex + 1, Cmd.EndCmd - 1 );
-			BuildSharedSerialization_r( ArrayIterator, NewData, bWriteHandle, bDoChecksum, ArrayDepth + 1, SharedInfo );
+			FRepHandleIterator ArrayIterator(HandleIterator.ChangelistIterator, Cmds, ArrayHandleToCmdIndex, Cmd.ElementSize, Array->Num(), CmdIndex + 1, Cmd.EndCmd - 1);
+			BuildSharedSerialization_r(ArrayIterator, NewData, bWriteHandle, bDoChecksum, ArrayDepth + 1, SharedInfo);
 			continue;
 		}
 
@@ -3688,7 +3871,11 @@ void FRepLayout::BuildSharedSerialization_r(
 	}
 }
 
-void FRepLayout::BuildSharedSerializationForRPC_DynamicArray_r(const int32 CmdIndex, uint8 * Data, int32 ArrayDepth, FRepSerializationSharedInfo& SharedInfo)
+void FRepLayout::BuildSharedSerializationForRPC_DynamicArray_r(
+	const int32						CmdIndex,
+	uint8*							Data,
+	int32							ArrayDepth,
+	FRepSerializationSharedInfo&	SharedInfo)
 {
 	const FRepLayoutCmd& Cmd = Cmds[ CmdIndex ];
 
@@ -3714,17 +3901,23 @@ void FRepLayout::BuildSharedSerializationForRPC_DynamicArray_r(const int32 CmdIn
 	}
 }
 
-void FRepLayout::BuildSharedSerializationForRPC_r(const int32 CmdStart, const int32 CmdEnd, void * Data, int32 ArrayIndex, int32 ArrayDepth, FRepSerializationSharedInfo& SharedInfo)
+void FRepLayout::BuildSharedSerializationForRPC_r(
+	const int32						CmdStart,
+	const int32						CmdEnd,
+	void*							Data,
+	int32							ArrayIndex,
+	int32							ArrayDepth,
+	FRepSerializationSharedInfo&	SharedInfo)
 {
-	for ( int32 CmdIndex = CmdStart; CmdIndex < CmdEnd; CmdIndex++ )
+	for (int32 CmdIndex = CmdStart; CmdIndex < CmdEnd; CmdIndex++)
 	{
 		const FRepLayoutCmd& Cmd = Cmds[CmdIndex];
 
-		check( Cmd.Type != ERepLayoutCmdType::Return );
+		check(Cmd.Type != ERepLayoutCmdType::Return);
 
-		if ( Cmd.Type == ERepLayoutCmdType::DynamicArray )
+		if (Cmd.Type == ERepLayoutCmdType::DynamicArray)
 		{
-			BuildSharedSerializationForRPC_DynamicArray_r( CmdIndex, (uint8*)Data + Cmd.Offset, ArrayDepth + 1, SharedInfo );
+			BuildSharedSerializationForRPC_DynamicArray_r(CmdIndex, (uint8*)Data + Cmd.Offset, ArrayDepth + 1, SharedInfo);
 			CmdIndex = Cmd.EndCmd - 1;		// The -1 to handle the ++ in the for loop
 			continue;
 		}
@@ -3738,13 +3931,13 @@ void FRepLayout::BuildSharedSerializationForRPC_r(const int32 CmdStart, const in
 	}
 }
 
-void FRepLayout::BuildSharedSerializationForRPC( void* Data )
+void FRepLayout::BuildSharedSerializationForRPC(void* Data)
 {
 	if ((GNetSharedSerializedData != 0) && !SharedInfoRPC.IsValid())
 	{
 		SharedInfoRPCParentsChanged.Init(false, Parents.Num());
 
-		for ( int32 i = 0; i < Parents.Num(); i++ )
+		for (int32 i = 0; i < Parents.Num(); i++)
 		{
 			if (Parents[i].Property->HasAnyPropertyFlags(CPF_OutParm))
 			{
@@ -3753,15 +3946,15 @@ void FRepLayout::BuildSharedSerializationForRPC( void* Data )
 
 			bool bSend = true;
 
-			if ( !Cast<UBoolProperty>( Parents[i].Property ) )
+			if (!Cast<UBoolProperty>(Parents[i].Property))
 			{
 				// check for a complete match, including arrays
 				// (we're comparing against zero data here, since 
 				// that's the default.)
-				bSend = !Parents[i].Property->Identical_InContainer( Data, nullptr, Parents[i].ArrayIndex );
+				bSend = !Parents[i].Property->Identical_InContainer(Data, nullptr, Parents[i].ArrayIndex);
 			}
 
-			if ( bSend )
+			if (bSend)
 			{
 				// Cache result of property comparison to default so we only have to do it once
 				SharedInfoRPCParentsChanged[i] = true;
@@ -3780,34 +3973,38 @@ void FRepLayout::ClearSharedSerializationForRPC()
 	SharedInfoRPCParentsChanged.Reset();
 }
 
-void FRepLayout::SendPropertiesForRPC( UFunction * Function, UActorChannel * Channel, FNetBitWriter & Writer, void* Data ) const
+void FRepLayout::SendPropertiesForRPC(
+	UFunction*		Function,
+	UActorChannel*	Channel,
+	FNetBitWriter&	Writer,
+	void*			Data) const
 {
-	check( Function == Owner );
+	check(Function == Owner);
 
-	if ( Channel->Connection->InternalAck )
+	if (Channel->Connection->InternalAck)
 	{
-		TArray< uint16 > Changed;
+		TArray<uint16> Changed;
 
-		for ( int32 i = 0; i < Parents.Num(); i++ )
+		for (int32 i = 0; i < Parents.Num(); i++)
 		{
-			if ( !Parents[i].Property->Identical_InContainer( Data, NULL, Parents[i].ArrayIndex ) )
+			if (!Parents[i].Property->Identical_InContainer(Data, NULL, Parents[i].ArrayIndex))
 			{
-				BuildChangeList_r( BaseHandleToCmdIndex, Parents[i].CmdStart, Parents[i].CmdEnd, ( uint8* )Data, 0, Changed );
+				BuildChangeList_r(BaseHandleToCmdIndex, Parents[i].CmdStart, Parents[i].CmdEnd, (uint8*)Data, 0, Changed);
 			}
 		}
 
-		Changed.Add( 0 ); // Null terminator
+		Changed.Add(0); // Null terminator
 
-		SendProperties_BackwardsCompatible( nullptr, nullptr, ( uint8* )Data, Channel->Connection, Writer, Changed );
+		SendProperties_BackwardsCompatible(nullptr, nullptr, (uint8*)Data, Channel->Connection, Writer, Changed);
 
 		return;
 	}
 
-	for ( int32 i = 0; i < Parents.Num(); i++ )
+	for (int32 i = 0; i < Parents.Num(); i++)
 	{
 		bool Send = true;
 
-		if ( !Cast<UBoolProperty>( Parents[i].Property ) )
+		if (!Cast<UBoolProperty>(Parents[i].Property))
 		{
 			// Used cached comparison result if possible
 			if ((GNetSharedSerializedData != 0) && SharedInfoRPC.IsValid() && !Parents[i].Property->HasAnyPropertyFlags(CPF_OutParm))
@@ -3819,56 +4016,62 @@ void FRepLayout::SendPropertiesForRPC( UFunction * Function, UActorChannel * Cha
 				// check for a complete match, including arrays
 				// (we're comparing against zero data here, since 
 				// that's the default.)
-				Send = !Parents[i].Property->Identical_InContainer( Data, NULL, Parents[i].ArrayIndex );
+				Send = !Parents[i].Property->Identical_InContainer(Data, NULL, Parents[i].ArrayIndex);
 			}
 
-			Writer.WriteBit( Send ? 1 : 0 );
+			Writer.WriteBit(Send ? 1 : 0);
 		}
 
-		if ( Send )
+		if (Send)
 		{
 			bool bHasUnmapped = false;
-			SerializeProperties_r( Writer, Writer.PackageMap, Parents[i].CmdStart, Parents[i].CmdEnd, Data, bHasUnmapped, 0, 0, SharedInfoRPC );
+			SerializeProperties_r(Writer, Writer.PackageMap, Parents[i].CmdStart, Parents[i].CmdEnd, Data, bHasUnmapped, 0, 0, SharedInfoRPC);
 		}
 	}
 }
 
-void FRepLayout::ReceivePropertiesForRPC( UObject* Object, UFunction * Function, UActorChannel * Channel, FNetBitReader & Reader, void* Data, TSet<FNetworkGUID>& UnmappedGuids) const
+void FRepLayout::ReceivePropertiesForRPC(
+	UObject*			Object,
+	UFunction*			Function,
+	UActorChannel*		Channel,
+	FNetBitReader&		Reader,
+	void*				Data,
+	TSet<FNetworkGUID>&	UnmappedGuids) const
 {
-	check( Function == Owner );
+	check(Function == Owner);
 
-	for ( int32 i = 0; i < Parents.Num(); i++ )
+	for (int32 i = 0; i < Parents.Num(); i++)
 	{
-		if ( Parents[i].ArrayIndex == 0 && ( Parents[i].Property->PropertyFlags & CPF_ZeroConstructor ) == 0 )
+		if (Parents[i].ArrayIndex == 0 && (Parents[i].Property->PropertyFlags & CPF_ZeroConstructor) == 0)
 		{
 			// If this property needs to be constructed, make sure we do that
 			Parents[i].Property->InitializeValue((uint8*)Data + Parents[i].Property->GetOffset_ForUFunction());
 		}
 	}
 
-	if ( Channel->Connection->InternalAck )
+	if (Channel->Connection->InternalAck)
 	{
 		bool bHasUnmapped = false;
 		bool bGuidsChanged = false;
 
 		// Let package map know we want to track and know about any guids that are unmapped during the serialize call
 		// We have to do this manually since we aren't passing in any unmapped info
-		Reader.PackageMap->ResetTrackedGuids( true );
+		Reader.PackageMap->ResetTrackedGuids(true);
 
-		ReceiveProperties_BackwardsCompatible( Channel->Connection, nullptr, Data, Reader, bHasUnmapped, false, bGuidsChanged );
+		ReceiveProperties_BackwardsCompatible(Channel->Connection, nullptr, Data, Reader, bHasUnmapped, false, bGuidsChanged);
 
-		if ( Reader.PackageMap->GetTrackedUnmappedGuids().Num() > 0 )
+		if (Reader.PackageMap->GetTrackedUnmappedGuids().Num() > 0)
 		{
 			bHasUnmapped = true;
 			UnmappedGuids = Reader.PackageMap->GetTrackedUnmappedGuids();
 		}
 
-		Reader.PackageMap->ResetTrackedGuids( false );
+		Reader.PackageMap->ResetTrackedGuids(false);
 
-		if ( bHasUnmapped )
+		if (bHasUnmapped)
 		{
-			UE_LOG( LogRepTraffic, Log, TEXT( "Unable to resolve RPC parameter to do being unmapped. Object[%d] %s. Function %s." ),
-					Channel->ChIndex, *Object->GetName(), *Function->GetName() );
+			UE_LOG(LogRepTraffic, Log, TEXT("Unable to resolve RPC parameter to do being unmapped. Object[%d] %s. Function %s."),
+					Channel->ChIndex, *Object->GetName(), *Function->GetName());
 		}
 	}
 	else
@@ -3877,28 +4080,28 @@ void FRepLayout::ReceivePropertiesForRPC( UObject* Object, UFunction * Function,
 
 		static FRepSerializationSharedInfo Empty;
 
-		for ( int32 i = 0; i < Parents.Num(); i++ )
+		for (int32 i = 0; i < Parents.Num(); i++)
 		{
-			if ( Cast<UBoolProperty>( Parents[i].Property ) || Reader.ReadBit() )
+			if (Cast<UBoolProperty>(Parents[i].Property) || Reader.ReadBit())
 			{
 				bool bHasUnmapped = false;
 
-				SerializeProperties_r( Reader, Reader.PackageMap, Parents[i].CmdStart, Parents[i].CmdEnd, Data, bHasUnmapped, 0, 0, Empty );
+				SerializeProperties_r(Reader, Reader.PackageMap, Parents[i].CmdStart, Parents[i].CmdEnd, Data, bHasUnmapped, 0, 0, Empty);
 
-				if ( Reader.IsError() )
+				if (Reader.IsError())
 				{
 					return;
 				}
 
-				if ( bHasUnmapped )
+				if (bHasUnmapped)
 				{
-					UE_LOG( LogRepTraffic, Log, TEXT( "Unable to resolve RPC parameter. Object[%d] %s. Function %s. Parameter %s." ),
-							Channel->ChIndex, *Object->GetName(), *Function->GetName(), *Parents[i].Property->GetName() );
+					UE_LOG(LogRepTraffic, Log, TEXT("Unable to resolve RPC parameter. Object[%d] %s. Function %s. Parameter %s."),
+							Channel->ChIndex, *Object->GetName(), *Function->GetName(), *Parents[i].Property->GetName());
 				}
 			}
 		}
 
-		if (Reader.PackageMap->GetTrackedUnmappedGuids().Num() > 0)
+		if (Reader.PackageMap->GetTrackedUnmappedGuids().Num()> 0)
 		{
 			UnmappedGuids = Reader.PackageMap->GetTrackedUnmappedGuids();
 		}
@@ -3907,58 +4110,69 @@ void FRepLayout::ReceivePropertiesForRPC( UObject* Object, UFunction * Function,
 	}
 }
 
-void FRepLayout::SerializePropertiesForStruct( UStruct * Struct, FBitArchive & Ar, UPackageMap	* Map, void* Data, bool & bHasUnmapped ) const
+void FRepLayout::SerializePropertiesForStruct(
+	UStruct*		Struct,
+	FBitArchive&	Ar,
+	UPackageMap*	Map,
+	void*			Data,
+	bool&			bHasUnmapped) const
 {
-	check( Struct == Owner );
+	check(Struct == Owner);
 
 	static FRepSerializationSharedInfo Empty;
 
-	for ( int32 i = 0; i < Parents.Num(); i++ )
+	for (int32 i = 0; i < Parents.Num(); i++)
 	{
-		SerializeProperties_r( Ar, Map, Parents[i].CmdStart, Parents[i].CmdEnd, Data, bHasUnmapped, 0, 0, Empty );
+		SerializeProperties_r(Ar, Map, Parents[i].CmdStart, Parents[i].CmdEnd, Data, bHasUnmapped, 0, 0, Empty);
 
-		if ( Ar.IsError() )
+		if (Ar.IsError())
 		{
 			return;
 		}
 	}
 }
 
-void FRepLayout::BuildHandleToCmdIndexTable_r( const int32 CmdStart, const int32 CmdEnd, TArray< FHandleToCmdIndex >& HandleToCmdIndex )
+void FRepLayout::BuildHandleToCmdIndexTable_r(
+	const int32					CmdStart,
+	const int32					CmdEnd,
+	TArray<FHandleToCmdIndex>&	HandleToCmdIndex)
 {
-	for ( int32 CmdIndex = CmdStart; CmdIndex < CmdEnd; CmdIndex++ )
+	for (int32 CmdIndex = CmdStart; CmdIndex < CmdEnd; CmdIndex++)
 	{
 		const FRepLayoutCmd& Cmd = Cmds[CmdIndex];
 
-		check( Cmd.Type != ERepLayoutCmdType::Return );
+		check(Cmd.Type != ERepLayoutCmdType::Return);
 
-		const int32 Index = HandleToCmdIndex.Add( FHandleToCmdIndex( CmdIndex ) );
+		const int32 Index = HandleToCmdIndex.Add(FHandleToCmdIndex(CmdIndex));
 
-		if ( Cmd.Type == ERepLayoutCmdType::DynamicArray )
+		if (Cmd.Type == ERepLayoutCmdType::DynamicArray)
 		{
-			HandleToCmdIndex[Index].HandleToCmdIndex = TUniquePtr< TArray< FHandleToCmdIndex > >( new TArray< FHandleToCmdIndex >() );
+			HandleToCmdIndex[Index].HandleToCmdIndex = TUniquePtr<TArray<FHandleToCmdIndex>>(new TArray<FHandleToCmdIndex>());
 
-			TArray< FHandleToCmdIndex >& ArrayHandleToCmdIndex = *HandleToCmdIndex[Index].HandleToCmdIndex;
+			TArray<FHandleToCmdIndex>& ArrayHandleToCmdIndex = *HandleToCmdIndex[Index].HandleToCmdIndex;
 
-			BuildHandleToCmdIndexTable_r( CmdIndex + 1, Cmd.EndCmd - 1, ArrayHandleToCmdIndex );
+			BuildHandleToCmdIndexTable_r(CmdIndex + 1, Cmd.EndCmd - 1, ArrayHandleToCmdIndex);
 			CmdIndex = Cmd.EndCmd - 1;		// The -1 to handle the ++ in the for loop
 		}
 	}
 }
 
-void FRepLayout::RebuildConditionalProperties( FRepState * RESTRICT	RepState, const FRepChangedPropertyTracker& ChangedTracker, const FReplicationFlags& RepFlags ) const
+void FRepLayout::RebuildConditionalProperties(
+	FRepState * RESTRICT				RepState,
+	const FRepChangedPropertyTracker&	ChangedTracker,
+	const FReplicationFlags&			RepFlags) const
 {
-	SCOPE_CYCLE_COUNTER( STAT_NetRebuildConditionalTime );
+	SCOPE_CYCLE_COUNTER(STAT_NetRebuildConditionalTime);
 	
 	// Setup condition map
-	const bool bIsInitial	= RepFlags.bNetInitial		? true : false;
-	const bool bIsOwner		= RepFlags.bNetOwner		? true : false;
-	const bool bIsSimulated	= RepFlags.bNetSimulated	? true : false;
-	const bool bIsPhysics	= RepFlags.bRepPhysics		? true : false;
-	const bool bIsReplay	= RepFlags.bReplay			? true : false;
+	const bool bIsInitial = RepFlags.bNetInitial ? true : false;
+	const bool bIsOwner = RepFlags.bNetOwner ? true : false;
+	const bool bIsSimulated = RepFlags.bNetSimulated ? true : false;
+	const bool bIsPhysics = RepFlags.bRepPhysics ? true : false;
+	const bool bIsReplay = RepFlags.bReplay ? true : false;
 
-	RepState->ConditionMap[COND_None]						= true;
-	RepState->ConditionMap[COND_InitialOnly]				= bIsInitial;
+	RepState->ConditionMap[COND_None] = true;
+	RepState->ConditionMap[COND_InitialOnly] = bIsInitial;
 
 	RepState->ConditionMap[COND_OwnerOnly] = bIsOwner;
 	RepState->ConditionMap[COND_SkipOwner] = !bIsOwner;
@@ -3980,7 +4194,7 @@ void FRepLayout::RebuildConditionalProperties( FRepState * RESTRICT	RepState, co
 	RepState->RepFlags = RepFlags;
 }
 
-void FRepLayout::InitChangedTracker(FRepChangedPropertyTracker * ChangedTracker) const
+void FRepLayout::InitChangedTracker(FRepChangedPropertyTracker* ChangedTracker) const
 {
 	ChangedTracker->Parents.SetNum(Parents.Num());
 
@@ -3991,105 +4205,105 @@ void FRepLayout::InitChangedTracker(FRepChangedPropertyTracker * ChangedTracker)
 }
 
 void FRepLayout::InitShadowData(
-	TArray< uint8, TAlignedHeapAllocator<16> >&	ShadowData,
-	UClass *									InObjectClass,
-	const uint8 * const							Src ) const
+	TArray<uint8, TAlignedHeapAllocator<16>>&	ShadowData,
+	UClass*										InObjectClass,
+	const uint8* const							Src) const
 {
 	ShadowData.Empty();
-	ShadowData.AddZeroed( InObjectClass->GetDefaultsCount() );
+	ShadowData.AddZeroed(InObjectClass->GetDefaultsCount());
 
 	// Construct the properties
-	ConstructProperties( ShadowData );
+	ConstructProperties(ShadowData);
 
 	// Init the properties
-	InitProperties( ShadowData, Src );
+	InitProperties(ShadowData, Src);
 }
 
 void FRepLayout::InitRepState(
 	FRepState*									RepState,
-	UClass *									InObjectClass, 
-	const uint8 * const							Src, 
-	TSharedPtr< FRepChangedPropertyTracker > &	InRepChangedPropertyTracker ) const
+	UClass*										InObjectClass, 
+	const uint8* const							Src, 
+	TSharedPtr<FRepChangedPropertyTracker>&		InRepChangedPropertyTracker) const
 {
-	InitShadowData( RepState->StaticBuffer, InObjectClass, Src );
+	InitShadowData(RepState->StaticBuffer, InObjectClass, Src);
 
 	RepState->RepChangedPropertyTracker = InRepChangedPropertyTracker;
 
-	check( RepState->RepChangedPropertyTracker->Parents.Num() == Parents.Num() );
+	check(RepState->RepChangedPropertyTracker->Parents.Num() == Parents.Num());
 
 	// Start out the conditional props based on a default RepFlags struct
 	// It will rebuild if it ever changes
-	RebuildConditionalProperties( RepState, *InRepChangedPropertyTracker.Get(), FReplicationFlags() );
+	RebuildConditionalProperties(RepState, *InRepChangedPropertyTracker.Get(), FReplicationFlags());
 }
 
-void FRepLayout::ConstructProperties( TArray< uint8, TAlignedHeapAllocator<16> >& ShadowData ) const
+void FRepLayout::ConstructProperties(TArray<uint8, TAlignedHeapAllocator<16>>& ShadowData) const
 {
 	uint8* StoredData = ShadowData.GetData();
 
 	// Construct all items
-	for ( int32 i = 0; i < Parents.Num(); i++ )
+	for (int32 i = 0; i < Parents.Num(); i++)
 	{
 		// Only construct the 0th element of static arrays (InitializeValue will handle the elements)
-		if ( Parents[i].ArrayIndex == 0 )
+		if (Parents[i].ArrayIndex == 0)
 		{
-			PTRINT Offset = Parents[i].Property->ContainerPtrToValuePtr<uint8>( StoredData ) - StoredData;
-			check( Offset >= 0 && Offset < ShadowData.Num() );
+			PTRINT Offset = Parents[i].Property->ContainerPtrToValuePtr<uint8>(StoredData) - StoredData;
+			check(Offset>= 0 && Offset < ShadowData.Num());
 
-			Parents[i].Property->InitializeValue( StoredData + Offset );
+			Parents[i].Property->InitializeValue(StoredData + Offset);
 		}
 	}
 }
 
-void FRepLayout::InitProperties( TArray< uint8, TAlignedHeapAllocator<16> >& ShadowData, const uint8* const Src ) const
+void FRepLayout::InitProperties(TArray<uint8, TAlignedHeapAllocator<16>>& ShadowData, const uint8* const Src) const
 {
 	LLM_SCOPE(ELLMTag::Networking);
 
 	uint8* StoredData = ShadowData.GetData();
 
 	// Init all items
-	for ( int32 i = 0; i < Parents.Num(); i++ )
+	for (int32 i = 0; i < Parents.Num(); i++)
 	{
 		// Only copy the 0th element of static arrays (CopyCompleteValue will handle the elements)
-		if ( Parents[i].ArrayIndex == 0 )
+		if (Parents[i].ArrayIndex == 0)
 		{
-			PTRINT Offset = Parents[i].Property->ContainerPtrToValuePtr<uint8>( StoredData ) - StoredData;
-			check( Offset >= 0 && Offset < ShadowData.Num() );
+			PTRINT Offset = Parents[i].Property->ContainerPtrToValuePtr<uint8>(StoredData) - StoredData;
+			check(Offset>= 0 && Offset < ShadowData.Num());
 
-			Parents[i].Property->CopyCompleteValue( StoredData + Offset, Src + Offset );
+			Parents[i].Property->CopyCompleteValue(StoredData + Offset, Src + Offset);
 		}
 	}
 }
 
-void FRepLayout::DestructProperties( FRepStateStaticBuffer& RepStateStaticBuffer ) const
+void FRepLayout::DestructProperties(FRepStateStaticBuffer& RepStateStaticBuffer) const
 {
 	uint8* StoredData = RepStateStaticBuffer.GetData();
 
 	// Destruct all items
-	for ( int32 i = 0; i < Parents.Num(); i++ )
+	for (int32 i = 0; i < Parents.Num(); i++)
 	{
 		// Only copy the 0th element of static arrays (DestroyValue will handle the elements)
-		if ( Parents[i].ArrayIndex == 0 )
+		if (Parents[i].ArrayIndex == 0)
 		{
-			PTRINT Offset = Parents[i].Property->ContainerPtrToValuePtr<uint8>( StoredData ) - StoredData;
-			check( Offset >= 0 && Offset < RepStateStaticBuffer.Num() );
+			PTRINT Offset = Parents[i].Property->ContainerPtrToValuePtr<uint8>(StoredData) - StoredData;
+			check(Offset>= 0 && Offset < RepStateStaticBuffer.Num());
 
-			Parents[i].Property->DestroyValue( StoredData + Offset );
+			Parents[i].Property->DestroyValue(StoredData + Offset);
 		}
 	}
 
 	RepStateStaticBuffer.Empty();
 }
 
-void FRepLayout::GetLifetimeCustomDeltaProperties(TArray< int32 > & OutCustom, TArray< ELifetimeCondition >	& OutConditions)
+void FRepLayout::GetLifetimeCustomDeltaProperties(TArray<int32>& OutCustom, TArray<ELifetimeCondition>& OutConditions)
 {
 	OutCustom.Empty();
 	OutConditions.Empty();
 
-	for ( int32 i = 0; i < Parents.Num(); i++ )
+	for (int32 i = 0; i < Parents.Num(); i++)
 	{
-		if ( (Parents[i].Flags & ERepParentFlags::IsCustomDelta) != ERepParentFlags::None )
+		if ((Parents[i].Flags & ERepParentFlags::IsCustomDelta) != ERepParentFlags::None)
 		{
-			check( Parents[i].Property->RepIndex + Parents[i].ArrayIndex == i );
+			check(Parents[i].Property->RepIndex + Parents[i].ArrayIndex == i);
 
 			OutCustom.Add(i);
 			OutConditions.Add(Parents[i].Condition);
@@ -4113,7 +4327,7 @@ FRepState::~FRepState()
 {
 	if (RepLayout.IsValid() && StaticBuffer.Num() > 0)
 	{	
-		RepLayout->DestructProperties( StaticBuffer );
+		RepLayout->DestructProperties(StaticBuffer);
 	}
 }
 
@@ -4121,6 +4335,6 @@ FRepChangelistState::~FRepChangelistState()
 {
 	if (RepLayout.IsValid() && StaticBuffer.Num() > 0)
 	{	
-		RepLayout->DestructProperties( StaticBuffer );
+		RepLayout->DestructProperties(StaticBuffer);
 	}
 }
diff --git a/Engine/Source/Runtime/Engine/Public/Net/RepLayout.h b/Engine/Source/Runtime/Engine/Public/Net/RepLayout.h
--- a/Engine/Source/Runtime/Engine/Public/Net/RepLayout.h
+++ b/Engine/Source/Runtime/Engine/Public/Net/RepLayout.h
@@ -1,9 +1,17 @@
 // Copyright 1998-2018 Epic Games, Inc. All Rights Reserved.
 
-/*=============================================================================
-	RepLayout.h:
-	FRepLayout is a helper class to quickly replicate properties that are marked for replication
-=============================================================================*/
+/**
+* This file contains the core logic and types that support Object and RPC replication.
+*
+* These types don't dictate how RPCs are triggered or when an Object should be replicated,
+* although there are some methods defined here that may be used in those determinations.
+*
+* Instead, the types here focus on how data from Objects, Structs, Containers, and Properties
+* are generically tracked and serialized on both Clients and Servers.
+*
+* The main class is FRepLayout.
+*/
+
 #pragma once
 
 #include "CoreMinimal.h"
@@ -20,104 +28,151 @@ class UNetConnection;
 class UPackageMapClient;
 
 // Properties will be copied in here so memory needs aligned to largest type
-typedef TArray< uint8, TAlignedHeapAllocator<16> > FRepStateStaticBuffer;
+typedef TArray<uint8, TAlignedHeapAllocator<16>> FRepStateStaticBuffer;
 
 enum class EDiffPropertiesFlags : uint32
 {
-	None							= 0,
-	Sync							= ( 1 << 0 ),
-	IncludeConditionalProperties	= ( 1 << 1 )
+	None = 0,
+	Sync = (1 <<0),							//! Indicates that properties should be updated (synchronized), not just diffed.
+	IncludeConditionalProperties = (1 <<1)	//! Whether or not conditional properties should be included.
 };
 
 ENUM_CLASS_FLAGS(EDiffPropertiesFlags);
 
+/** Stores meta data about a given Replicated property. */
 class FRepChangedParent
 {
 public:
-	FRepChangedParent() : Active( 1 ), OldActive( 1 ), IsConditional( 0 ) {}
+	FRepChangedParent():
+		Active(1),
+		OldActive(1),
+		IsConditional(0)
+	{}
+
+	/** Whether or not this property is currently Active (i.e., considered for replication). */
+	uint32 Active: 1;
+
+	/** The last updated state of Active, used to track when the Active state changes. */
+	uint32 OldActive: 1;
 
-	uint32				Active			: 1;
-	uint32				OldActive		: 1;
-	uint32				IsConditional	: 1;
+	/**
+	 * Whether or not this property has conditions that may exclude it from replicating to a given connection.
+	 * @see FRepState::ConditionMap.
+	 */
+	uint32 IsConditional: 1;
 };
 
-/** FRepChangedPropertyTracker
- * This class is used to store the change list for a group of properties of a particular actor/object
- * This information is shared across connections when possible
+/**
+ * This class is used to store meta data about properties that is shared between connections,
+ * including whether or not a given property is Conditional, Active, and any external data
+ * that may be needed for Replays.
+ *
+ * TODO: This class (and arguably IRepChangedPropertyTracker) should be renamed to reflect
+ *			what they actually do now.
  */
 class FRepChangedPropertyTracker : public IRepChangedPropertyTracker
 {
 public:
-	FRepChangedPropertyTracker( const bool InbIsReplay, const bool InbIsClientReplayRecording )
-		: bIsReplay( InbIsReplay )
-		, bIsClientReplayRecording( InbIsClientReplayRecording )
-		, ExternalDataNumBits( 0 ) {}
-	virtual ~FRepChangedPropertyTracker() { }
+	FRepChangedPropertyTracker(const bool InbIsReplay, const bool InbIsClientReplayRecording):
+		bIsReplay(InbIsReplay),
+		bIsClientReplayRecording(InbIsClientReplayRecording),
+		ExternalDataNumBits(0)
+	{}
 
-	virtual void SetCustomIsActiveOverride( const uint16 RepIndex, const bool bIsActive ) override
+	virtual ~FRepChangedPropertyTracker() {}
+
+	//~ Begin IRepChangedPropertyTracker Interface.
+	/**
+	 * Manually set whether or not Property should be marked inactive.
+	 * This will change the Active status for all connections.
+	 *
+	 * @see DOREPLIFETIME_ACTIVE_OVERRIDE
+	 *
+	 * @param RepIndex	Replication index for the Property.
+	 * @param bIsActive	The new Active state.
+	 */
+	virtual void SetCustomIsActiveOverride(const uint16 RepIndex, const bool bIsActive) override
 	{
 		FRepChangedParent & Parent = Parents[RepIndex];
 
-		checkSlow( Parent.IsConditional );
+		checkSlow(Parent.IsConditional);
 
 		Parent.Active = (bIsActive || bIsClientReplayRecording) ? 1 : 0;
-
 		Parent.OldActive = Parent.Active;
 	}
 
-	virtual void SetExternalData( const uint8* Src, const int32 NumBits ) override
+	/**
+	 * Sets (or resets) the External Data.
+	 * External Data is primarily used for Replays, and is used to track additional non-replicated
+	 * data or state about an object.
+	 *
+	 * @param Src		Memory containing the external data.
+	 * @param NumBits	Size of the memory, in bits.
+	 */
+	virtual void SetExternalData(const uint8* Src, const int32 NumBits) override
 	{
 		ExternalDataNumBits = NumBits;
-		const int32 NumBytes = ( NumBits + 7 ) >> 3;
-		ExternalData.Reset( NumBytes );
-		ExternalData.AddUninitialized( NumBytes );
-		FMemory::Memcpy( ExternalData.GetData(), Src, NumBytes );
+		const int32 NumBytes = (NumBits + 7) >> 3;
+		ExternalData.Reset(NumBytes);
+		ExternalData.AddUninitialized(NumBytes);
+		FMemory::Memcpy(ExternalData.GetData(), Src, NumBytes);
 	}
 
+	/** Whether or not this is being used for a replay (may be recording or playback). */
 	virtual bool IsReplay() const override
 	{
 		return bIsReplay;
 	}
+	//~ End IRepChangedPropertyTracker Interface
 
-	TArray< FRepChangedParent >	Parents;
+	/** Activation data for top level Properties on the given Actor / Object. */
+	TArray<FRepChangedParent>	Parents;
 
-	bool						bIsReplay;							// True when recording/playing replays
-	bool						bIsClientReplayRecording;			// True when recording client replays
+	/** Whether or not this is being used for a replay (may be recording or playback). */
+	bool bIsReplay;
 
-	TArray< uint8 >				ExternalData;
-	uint32						ExternalDataNumBits;
+	/** Whether or not this is being used for a client replay recording. */
+	bool bIsClientReplayRecording;
+
+	TArray<uint8> ExternalData;
+	uint32 ExternalDataNumBits;
 };
 
 class FRepLayout;
 class FRepLayoutCmd;
 
-/** FRepSerializedPropertyInfo
- * Holds the unique identifier and offsets/lengths of a net serialized property
- */
+/** Holds the unique identifier and offsets/lengths of a net serialized property used for Shared Serialization */
 struct FRepSerializedPropertyInfo
 {
-	FRepSerializedPropertyInfo()
-		: BitOffset(0)
-		, BitLength(0)
-		, PropBitOffset(0)
-		, PropBitLength(0)
+	FRepSerializedPropertyInfo():
+		BitOffset(0),
+		BitLength(0),
+		PropBitOffset(0),
+		PropBitLength(0)
 	{}
 
-	FGuid Guid;				// unique identifier for this property, may include array index and depth
-	int32 BitOffset;		// bit offset into shared buffer of the shared data
-	int32 BitLength;		// length in bits of all serialized data for this property, may include handle and checksum
-	int32 PropBitOffset;	// bit offset into shared buffer of the property data
-	int32 PropBitLength;	// length in bits of net serialized property data only
+	/** Unique identifier for this property, may include array index and depth. */
+	FGuid Guid;
+
+	/** Bit offset into shared buffer of the shared data */
+	int32 BitOffset;
+
+	/** Length in bits of all serialized data for this property, may include handle and checksum. */
+	int32 BitLength;
+
+	/** Bit offset into shared buffer of the property data. */
+	int32 PropBitOffset;
+
+	/** Length in bits of net serialized property data only */
+	int32 PropBitLength;
 };
 
-/** FRepSerializationSharedInfo
- * Holds a set of shared net serialized properties
- */
+/** Holds a set of shared net serialized properties */
 struct FRepSerializationSharedInfo
 {
-	FRepSerializationSharedInfo()
-		: SerializedProperties(MakeUnique<FNetBitWriter>(0))
-		, bIsValid(false)
+	FRepSerializationSharedInfo():
+		SerializedProperties(MakeUnique<FNetBitWriter>(0)),
+		bIsValid(false)
 	{}
 
 	void SetValid()
@@ -141,6 +196,17 @@ struct FRepSerializationSharedInfo
 		}
 	}
 
+	/**
+	 * Creates a new SharedPropertyInfo and adds it to the SharedPropertyInfo list.
+	 *
+	 * @param Cmd				The command that represents the property we want to share.
+	 * @param PropertyGuid		A guid used to identify the property.
+	 * @param CmdIndex			Index of the property command. Only used if bDoChecksum is true.
+	 * @param Handle			Relative Handle of the property command. Only used if bWriteHandle is true.
+	 * @param Data				Pointer to the raw property memory that will be serialized.
+	 * @param bWriteHandler		Whether or not we should write Command handles into the serialized data.
+	 * @param bDoChecksum		Whether or not we should do checksums. Only used if ENABLE_PROPERTY_CHECKSUMS is enabled.
+	 */
 	const FRepSerializedPropertyInfo* WriteSharedProperty(
 		const FRepLayoutCmd&	Cmd,
 		const FGuid&			PropertyGuid,
@@ -150,136 +216,238 @@ struct FRepSerializationSharedInfo
 		const bool				bWriteHandle,
 		const bool				bDoChecksum);
 
-	// Metadata for properties in the shared data blob
+	/** Metadata for properties in the shared data blob. */
 	TArray<FRepSerializedPropertyInfo> SharedPropertyInfo;
 
-	// Binary blob of net serialized data to be shared
-	TUniquePtr<FNetBitWriter>	SerializedProperties;
+	/** Binary blob of net serialized data to be shared */
+	TUniquePtr<FNetBitWriter> SerializedProperties;
 
 private:
+
+	/** Whether or not shared serialization data has been successfully built. */
 	bool bIsValid;
 };
 
+/**
+ * Represents a single changelist, tracking changed properties.
+ *
+ * Properties are tracked via Relative Property Command Handles.
+ * Valid handles are 1-based, and 0 is reserved as a terminator.
+ *
+ * Arrays are tracked as a special case inline, where the first entry is the number of array elements,
+ * followed by handles for each array element, and ending with their own 0 terminator.
+ *
+ * Arrays may be nested by continually applying that pattern.
+ */
 class FRepChangedHistory
 {
 public:
-	FRepChangedHistory() : Resend( false ) {}
+	FRepChangedHistory():
+		Resend(false)
+	{}
+
+	/** Range of the Packets that this changelist was last sent with. Used to track acknowledgments. */
+	FPacketIdRange OutPacketIdRange;
+
+	/** List of Property Command Handles that changed in this changelist. */
+	TArray<uint16> Changed;
+
+	/** Whether or not this Changelist should be resent due to a Nak. */
+	bool Resend;
+};
+
+/**
+ * Stores changelist history (that are used to know what properties have changed) for objects.
+ *
+ * Only a fixed number of history items are kept. Once that limit is reached, old entries are
+ * merged into a single monolithic changelist (this happens incrementally each time a new entry
+ * is added).
+ */
+class FRepChangelistState
+{
+public:
+
+	FRepChangelistState() :
+		HistoryStart(0),
+		HistoryEnd(0),
+		CompareIndex(0)
+	{}
+
+	~FRepChangelistState();
+
+	TSharedPtr<FRepLayout> RepLayout;
+
+	/** The maximum number of individual changelists allowed.*/
+	static const int32 MAX_CHANGE_HISTORY = 64;
+
+	/** Circular buffer of changelists. */
+	FRepChangedHistory ChangeHistory[MAX_CHANGE_HISTORY];
+
+	/** Index in the buffer where changelist history starts (i.e., the Oldest changelist). */
+	int32 HistoryStart;
 
-	FPacketIdRange		OutPacketIdRange;
-	TArray< uint16 >	Changed;
-	bool				Resend;
+	/** Index in the buffer where changelist history ends (i.e., the Newest changelist). */
+	int32 HistoryEnd;
+
+	/** Number of times that properties have been compared */
+	int32 CompareIndex;
+
+	/** Latest state of all property data. Not used on Clients, only used on Servers if Shadow State is enabled. */
+	FRepStateStaticBuffer StaticBuffer;
+
+	/** Latest state of all shared serialization data. */
+	FRepSerializationSharedInfo SharedSerialization;
 };
 
 class FGuidReferences;
 
-typedef TMap< int32, FGuidReferences > FGuidReferencesMap;
+typedef TMap<int32, FGuidReferences> FGuidReferencesMap;
 
+/**
+ * Helper class that tracks Replicated Properties that reference NetGUIDs (for Object Pointers / References).
+ */
 class FGuidReferences
 {
 public:
-	FGuidReferences() : NumBufferBits( 0 ), Array( NULL ) {}
-	FGuidReferences( FBitReader& InReader, FBitReaderMark& InMark, const TSet< FNetworkGUID >& InUnmappedGUIDs, const TSet< FNetworkGUID >& InMappedDynamicGUIDs, const int32 InParentIndex, const int32 InCmdIndex ) : UnmappedGUIDs( InUnmappedGUIDs ), MappedDynamicGUIDs( InMappedDynamicGUIDs ), Array( NULL ), ParentIndex( InParentIndex ), CmdIndex( InCmdIndex )
+	FGuidReferences() : NumBufferBits(0), Array(NULL) {}
+
+	FGuidReferences(
+		FBitReader&					InReader,
+		FBitReaderMark&				InMark,
+		const TSet<FNetworkGUID>&	InUnmappedGUIDs,
+		const TSet<FNetworkGUID>&	InMappedDynamicGUIDs,
+		const int32					InParentIndex,
+		const int32					InCmdIndex
+	):
+			UnmappedGUIDs(InUnmappedGUIDs),
+			MappedDynamicGUIDs(InMappedDynamicGUIDs),
+			Array(NULL),
+			ParentIndex(InParentIndex),
+			CmdIndex(InCmdIndex)
 	{
 		NumBufferBits = InReader.GetPosBits() - InMark.GetPos();
-		InMark.Copy( InReader, Buffer );
+		InMark.Copy(InReader, Buffer);
 	}
-	FGuidReferences( FGuidReferencesMap* InArray, const int32 InParentIndex, const int32 InCmdIndex ) : NumBufferBits( 0 ), Array( InArray ), ParentIndex( InParentIndex ), CmdIndex( InCmdIndex ) {}
 
-	~FGuidReferences();
+	FGuidReferences(
+		FGuidReferencesMap*	InArray,
+		const int32			InParentIndex,
+		const int32			InCmdIndex
+	):
 
-	TSet< FNetworkGUID >		UnmappedGUIDs;
-	TSet< FNetworkGUID >		MappedDynamicGUIDs;
-	TArray< uint8 >				Buffer;
-	int32						NumBufferBits;
-
-	FGuidReferencesMap*			Array;
-	int32						ParentIndex;
-	int32						CmdIndex;
-};
-
-/** FRepChangelistState
-*  Stores changelist history (that are used to know what properties have changed) for objects
-*/
-class FRepChangelistState
-{
-public:
-	FRepChangelistState() :
-		HistoryStart( 0 ),
-		HistoryEnd( 0 ),
-		CompareIndex( 0 )
-	{ }
+			NumBufferBits(0),
+			Array(InArray),
+			ParentIndex(InParentIndex),
+			CmdIndex(InCmdIndex)
+	{}
+		
+	~FGuidReferences();
 
-	~FRepChangelistState();
+	/** GUIDs for objects that haven't been loaded / created yet. */
+	TSet<FNetworkGUID> UnmappedGUIDs;
 
-	TSharedPtr< FRepLayout >						RepLayout;
+	/** GUIDs for dynamically spawned objects that have already been created. */
+	TSet<FNetworkGUID> MappedDynamicGUIDs;
+	TArray<uint8> Buffer;
+	int32 NumBufferBits;
 
-	static const int32 MAX_CHANGE_HISTORY = 64;
+	FGuidReferencesMap* Array;
 
-	FRepChangedHistory								ChangeHistory[MAX_CHANGE_HISTORY];
-	int32											HistoryStart;
-	int32											HistoryEnd;
-	int32											CompareIndex;
+	/** The Property Command index of the top level property that references the GUID. */
+	int32 ParentIndex;
 
-	FRepStateStaticBuffer							StaticBuffer;
-	FRepSerializationSharedInfo						SharedSerialization;
+	/** The Property Command index of the actual property that references the GUID. */
+	int32 CmdIndex;
 };
 
-/** FRepState
- *  Stores state used by the FRepLayout manager
-*/
+/** Replication State that is unique Per Object Per Net Connection. */
 class FRepState
 {
 public:
+
 	FRepState() : 
-		HistoryStart( 0 ), 
-		HistoryEnd( 0 ),
-		NumNaks( 0 ),
-		OpenAckedCalled( false ),
-		AwakeFromDormancy( false ),
-		LastChangelistIndex( 0 ),
-		LastCompareIndex( 0 )
-	{ }
+		HistoryStart(0), 
+		HistoryEnd(0),
+		NumNaks(0),
+		OpenAckedCalled(false),
+		AwakeFromDormancy(false),
+		LastChangelistIndex(0),
+		LastCompareIndex(0)
+	{}
 
 	~FRepState();
 
-	FRepStateStaticBuffer			StaticBuffer;
+	/** Latest state of all property data. Used on Clients, or on Servers if Shadow State is disabled. */
+	FRepStateStaticBuffer StaticBuffer;
 
-	FGuidReferencesMap				GuidReferencesMap;
+	FGuidReferencesMap GuidReferencesMap;
 
-	TSharedPtr< FRepLayout >		RepLayout;
-	
-	TArray< UProperty * >			RepNotifies;
+	TSharedPtr<FRepLayout> RepLayout;
+
+	/** Properties that have RepNotifies that we will need to call on Clients (and ListenServers). */
+	TArray<UProperty *> RepNotifies;
 
 	TSharedPtr<FRepChangedPropertyTracker> RepChangedPropertyTracker;
 
+	/** The maximum number of individual changelists allowed.*/
 	static const int32 MAX_CHANGE_HISTORY = 32;
 
-	FRepChangedHistory				ChangeHistory[MAX_CHANGE_HISTORY];
-	int32							HistoryStart;
-	int32							HistoryEnd;
+	/** Circular buffer of changelists. */
+	FRepChangedHistory ChangeHistory[MAX_CHANGE_HISTORY];
+
+	/** Index in the buffer where changelist history starts (i.e., the Oldest changelist). */
+	int32 HistoryStart;
+
+	/** Index in the buffer where changelist history ends (i.e., the Newest changelist). */
+	int32 HistoryEnd;
+
+	/** Number of Changelist history entries that have outstanding Naks. */
+	int32 NumNaks;
+
+	/** List of changelists that were generated before the channel was fully opened.*/
+	TArray<FRepChangedHistory> PreOpenAckHistory;
 
-	int32							NumNaks;
+	/** Whether or not FRepLayout::OpenAcked has been called with this FRepState. */
+	bool OpenAckedCalled;
 
-	TArray< FRepChangedHistory >	PreOpenAckHistory;
+	/** This property is no longer used. */
+	bool AwakeFromDormancy;
 
-	bool							OpenAckedCalled;
-	bool							AwakeFromDormancy;
+	FReplicationFlags RepFlags;
 
-	FReplicationFlags				RepFlags;
+	/** The unique list of properties that have changed since the channel was first opened */
+	TArray<uint16> LifetimeChangelist;
 
-	TArray< uint16 >				LifetimeChangelist;			// This is the unique list of properties that have changed since the channel was first opened
+	/**
+	 * The last change list history item we replicated from FRepChangelistState.
+	 * (If we are caught up to FRepChangelistState::HistoryEnd, there are no new changelists to replicate).
+	 */
+	int32 LastChangelistIndex;
 
-	int32							LastChangelistIndex;		// The last change list history item we replicated from FRepChangelistState (if we are caught up to FRepChangelistState::HistoryEnd, there are no new changelists to replicate)
-	int32							LastCompareIndex;			// If == FRepChangelistState::CompareIndex, then there is definitely no new information since the last time we checked
+	/**
+	 * Tracks the last time this RepState actually replicated data.
+	 * If this matches FRepChangelistState::CompareIndex, then there is definitely no new
+	 * information since the last time we checked.
+	 *
+	 * @see FRepChangelistState::CompareIndex.
+	 *
+	 * Note, we can't solely rely on on LastChangelistIndex, since changelists are stored in circular buffers.
+	 */
+	int32 LastCompareIndex;
 
-	bool							ConditionMap[COND_Max];
+	/**
+	 * A map tracking which replication conditions are currently active.
+	 * @see ELifetimeCondition.
+	 */
+	bool ConditionMap[COND_Max];
 };
 
+/** Various types of Properties supported for Replication. */
 enum class ERepLayoutCmdType : uint8
 {
-	DynamicArray			= 0,	// Dynamic array
-	Return					= 1,	// Return from array, or end of stream
-	Property				= 2,	// Generic property
+	DynamicArray			= 0,	//! Dynamic array
+	Return					= 1,	//! Return from array, or end of stream
+	Property				= 2,	//! Generic property
 
 	PropertyBool			= 3,
 	PropertyFloat			= 4,
@@ -301,159 +469,318 @@ enum class ERepLayoutCmdType : uint8
 	PropertyUInt64			= 20,
 };
 
+/** Various flags that describe how a Top Level Property should be handled. */
 enum class ERepParentFlags : uint32
 {
 	None				= 0,
-	IsLifetime			= ( 1 << 0 ),
-	IsConditional		= ( 1 << 1 ),		// True if this property has a secondary condition to check
-	IsConfig			= ( 1 << 2 ),		// True if this property is defaulted from a config file
-	IsCustomDelta		= ( 1 << 3 )		// True if this property uses custom delta compression
+	IsLifetime			= (1 <<0),	//! This property is valid for the lifetime of the object (almost always set).
+	IsConditional		= (1 <<1),	//! This property has a secondary condition to check
+	IsConfig			= (1 <<2),	//! This property is defaulted from a config file
+	IsCustomDelta		= (1 <<3)	//! This property uses custom delta compression
 };
 
 ENUM_CLASS_FLAGS(ERepParentFlags)
 
+/**
+ * A Top Level Property of a UClass, UStruct, or UFunction (arguments to a UFunction).
+ *
+ * @see FRepLayout
+ */
 class FRepParentCmd
 {
 public:
-	FRepParentCmd( UProperty * InProperty, int32 InArrayIndex ) : 
-		Property( InProperty ), 
-		ArrayIndex( InArrayIndex ), 
-		CmdStart( 0 ), 
-		CmdEnd( 0 ), 
-		RoleSwapIndex( -1 ), 
-		Condition( COND_None ),
+
+	FRepParentCmd(UProperty* InProperty, int32 InArrayIndex): 
+		Property(InProperty), 
+		ArrayIndex(InArrayIndex), 
+		CmdStart(0), 
+		CmdEnd(0), 
+		RoleSwapIndex(-1), 
+		Condition(COND_None),
 		RepNotifyCondition(REPNOTIFY_OnChanged),
-		Flags( ERepParentFlags::None )
+		Flags(ERepParentFlags::None)
 	{}
 
-	UProperty *			Property;
-	int32				ArrayIndex;
-	uint16				CmdStart;
-	uint16				CmdEnd;
-	int32				RoleSwapIndex;
-	ELifetimeCondition	Condition;
+	UProperty* Property;
+
+	/**
+	 * If the Property is a C-Style fixed size array, then a command will be created for every element in the array.
+	 * This is the index of the element in the array for which the command represents.
+	 *
+	 * This will always be 0 for non array properties.
+	 */
+	int32 ArrayIndex;
+
+
+	/**
+	 * CmdStart and CmdEnd define the range of FRepLayoutCommands (by index in FRepLayouts Cmd array) of commands
+	 * that are associated with this Parent Command.
+	 *
+	 * This is used to track and access nested Properties from the parent.
+	 */
+	uint16 CmdStart;
+
+	/** @see CmdStart */
+	uint16 CmdEnd;
+
+	/**
+	 * This value indicates whether or not this command needs to swapped, and what other
+	 * command it should be swapped with.
+	 *
+	 * This is used for Role and RemoteRole which have inverted values on Servers and Clients.
+	 */
+	int32 RoleSwapIndex;
+
+	ELifetimeCondition Condition;
 	ELifetimeRepNotifyCondition	RepNotifyCondition;
-
-	ERepParentFlags		Flags;
+	ERepParentFlags Flags;
 };
 
+/** Various flags that describe how a Property should be handled. */
 enum class ERepLayoutFlags : uint8
 {
-	None					= 0,
-	IsSharedSerialization	= ( 1 << 0 )
+	None					= 0,		//! No flags.
+	IsSharedSerialization	= (1 <<0)	//! Indicates the property is eligible for shared serialization.
 };
 
 ENUM_CLASS_FLAGS(ERepLayoutFlags)
 
+/**
+ * Represents a single property, which could be either a Top Level Property, a Nested Struct Property,
+ * or an element in a Dynamic Array.
+ *
+ * @see FRepLayout
+ */
 class FRepLayoutCmd
 {
 public:
-	UProperty *			Property;			// Pointer back to property, used for NetSerialize calls, etc.
-	uint16				EndCmd;				// For arrays, this is the cmd index to jump to, to skip this arrays inner elements
-	uint16				ElementSize;		// For arrays, element size of data
-	int32				Offset;				// Absolute offset of property
-	uint16				RelativeHandle;		// Handle relative to start of array, or top list
-	uint16				ParentIndex;		// Index into Parents
-	uint32				CompatibleChecksum;	// Used to determine if property is still compatible
-	ERepLayoutCmdType	Type;
-	ERepLayoutFlags		Flags;
+
+	/** Pointer back to property, used for NetSerialize calls, etc. */
+	UProperty* Property;
+
+	/** For arrays, this is the cmd index to jump to, to skip this arrays inner elements. */
+	uint16 EndCmd;
+
+	/** For arrays, element size of data. */
+	uint16 ElementSize;
+
+	/** Absolute offset of property. */
+	int32 Offset;
+
+	/** Handle relative to start of array, or top list. */
+	uint16 RelativeHandle;
+
+	/** Index into Parents. */
+	uint16 ParentIndex;
+
+	/** Used to determine if property is still compatible */
+	uint32 CompatibleChecksum;
+
+	ERepLayoutCmdType Type;
+	ERepLayoutFlags Flags;
 };
 	
-/** FHandleToCmdIndex
- *  Converts a relative handle to the appropriate index into the Cmds array
- */
+/** Converts a relative handle to the appropriate index into the Cmds array */
 class FHandleToCmdIndex
 {
 public:
-	FHandleToCmdIndex() : CmdIndex( INDEX_NONE )
+	FHandleToCmdIndex():
+		CmdIndex(INDEX_NONE)
 	{
 	}
 
-	FHandleToCmdIndex( const int32 InHandleToCmdIndex ) : CmdIndex( InHandleToCmdIndex )
+	FHandleToCmdIndex(const int32 InHandleToCmdIndex):
+		CmdIndex(InHandleToCmdIndex)
 	{
 	}
 
-	int32										CmdIndex;
-	TUniquePtr< TArray< FHandleToCmdIndex > >	HandleToCmdIndex;
-
-	// Fix VS2013 compiler issue (VS2013 doesn't synthesize move constructor or move assignment operator version of these)
-	FHandleToCmdIndex( FHandleToCmdIndex&& Other ) : CmdIndex( Other.CmdIndex ), HandleToCmdIndex( MoveTemp( Other.HandleToCmdIndex ) )
+	FHandleToCmdIndex(FHandleToCmdIndex&& Other):
+		CmdIndex(Other.CmdIndex),
+		HandleToCmdIndex(MoveTemp(Other.HandleToCmdIndex))
 	{
 	}
 
-	FHandleToCmdIndex& operator=( FHandleToCmdIndex&& Other )
+	FHandleToCmdIndex& operator=(FHandleToCmdIndex&& Other)
 	{
-		if ( this != &Other )
+		if (this != &Other)
 		{
-			CmdIndex			= Other.CmdIndex;
-			HandleToCmdIndex	= MoveTemp( Other.HandleToCmdIndex );
+			CmdIndex = Other.CmdIndex;
+			HandleToCmdIndex = MoveTemp(Other.HandleToCmdIndex);
 		}
 
 		return *this;
 	}
+
+	int32 CmdIndex;
+	TUniquePtr<TArray<FHandleToCmdIndex>> HandleToCmdIndex;
 };
 
+/**
+ * Simple helper class to track state while iterating over changelists.
+ * This class doesn't actually expose methods to do the iteration, or to retrieve
+ * the current value.
+ */
 class FChangelistIterator
 {
 public:
-	FChangelistIterator( const TArray< uint16 >& InChanged, const int32 InChangedIndex ) : Changed( InChanged ), ChangedIndex( InChangedIndex )
-	{
-	}
 
-	const TArray< uint16 >& Changed;
-	int32					ChangedIndex;
+	FChangelistIterator(const TArray<uint16>& InChanged, const int32 InChangedIndex):
+		Changed(InChanged),
+		ChangedIndex(InChangedIndex)
+	{}
+
+	/** Changelist that is being iterated. */
+	const TArray<uint16>& Changed;
+
+	/** Current index into the changelist. */
+	int32 ChangedIndex;
 };
 
-/** FRepHandleIterator
- *  Iterates over a changelist, taking each handle, and mapping to rep layout index, array index, etc
- */
+/** Iterates over a changelist, taking each handle, and mapping to rep layout index, array index, etc. */
 class FRepHandleIterator
 {
 public:
+
 	FRepHandleIterator(
-		FChangelistIterator&				InChangelistIterator,
-		const TArray< FRepLayoutCmd >&		InCmds,
-		const TArray< FHandleToCmdIndex >&	InHandleToCmdIndex,
+		FChangelistIterator&	InChangelistIterator,
+		const TArray<FRepLayoutCmd>&		InCmds,
+		const TArray<FHandleToCmdIndex>&	InHandleToCmdIndex,
 		const int32							InElementSize,
 		const int32							InMaxArrayIndex,
 		const int32							InMinCmdIndex,
 		const int32							InMaxCmdIndex
-	) :
-		ChangelistIterator( InChangelistIterator ),
-		Cmds( InCmds ),
-		HandleToCmdIndex( InHandleToCmdIndex ),
-		NumHandlesPerElement( HandleToCmdIndex.Num() ),
-		ArrayElementSize( InElementSize ),
-		MaxArrayIndex( InMaxArrayIndex ),
-		MinCmdIndex( InMinCmdIndex ),
-		MaxCmdIndex( InMaxCmdIndex )
-	{
-	}
-
-	FChangelistIterator&				ChangelistIterator;
-
-	const TArray< FRepLayoutCmd >&		Cmds;
-	const TArray< FHandleToCmdIndex >&	HandleToCmdIndex;
-	const int32							NumHandlesPerElement;
-	const int32							ArrayElementSize;
-	const int32							MaxArrayIndex;
-	const int32							MinCmdIndex;
-	const int32							MaxCmdIndex;
-
-	int32								Handle;
-	int32								CmdIndex;
-	int32								ArrayIndex;
-	int32								ArrayOffset;
+	):
+		ChangelistIterator(InChangelistIterator),
+		Cmds(InCmds),
+		HandleToCmdIndex(InHandleToCmdIndex),
+		NumHandlesPerElement(HandleToCmdIndex.Num()),
+		ArrayElementSize(InElementSize),
+		MaxArrayIndex(InMaxArrayIndex),
+		MinCmdIndex(InMinCmdIndex),
+		MaxCmdIndex(InMaxCmdIndex)
+	{}
 
-	bool	NextHandle();
-	bool	JumpOverArray();
-	int32	PeekNextHandle() const;
+	/**
+	 * Moves the iterator to the next available handle.
+	 *
+	 * @return True if the move was successful, false otherwise (e.g., end of the iteration range was reached).
+	 */
+	bool NextHandle();
+
+	/**
+	 * Skips all the handles associated with a dynamic array at the iterators current position.
+	 *
+	 * @return True if the move was successful, false otherwise (this will return True even if the next handle is the end).
+	 */
+	bool JumpOverArray();
+
+	/**
+	 * Gets the handle at the iterators current position without advancing it.
+	 *
+	 * @return The next Property Command handle.
+	 */
+	int32 PeekNextHandle() const;
+
+	/** Used to track current state of the iteration. */
+	FChangelistIterator& ChangelistIterator;
+
+	/** List of all available Layout Commands. */
+	const TArray<FRepLayoutCmd>& Cmds;
+
+	/** Used to map Relative Handles to absolute Property Command Indices. */
+	const TArray<FHandleToCmdIndex>& HandleToCmdIndex;
+
+	/**
+	 * The number of handles per Command.
+	 * This should always be 1, except for Arrays.
+	 */
+	const int32 NumHandlesPerElement;
+
+	/**
+	 * Only used for Dynamic Arrays.
+	 * @see FRepLayout ElementSize.
+	 */
+	const int32	ArrayElementSize;
+
+	/**
+	 * Number of elements in a Dynamic array.
+	 * Should be 1 when iterating Top Level Properties or non-array properties.
+	 */
+	const int32	MaxArrayIndex;
+
+	/** Lowest index in Cmds where the iterator can go. */
+	const int32	MinCmdIndex;
+
+	/** Highest index in Cmds where the iterator can go. */
+	const int32	MaxCmdIndex;
+
+	/** The current Relative Property Command handle. */
+	int32 Handle;
+
+	/** The current Property Command index. */
+	int32 CmdIndex;
+
+	/** The index of the current element in a dynamic array. */
+	int32 ArrayIndex;
+
+	/** The Byte offset of Serialized Property data for a dynamic array to the current element. */
+	int32 ArrayOffset;
 };
 
-/** FRepLayout
- *  This class holds all replicated properties for a parent property, and all its children
- *	Helpers functions exist to read/write and compare property state.
-*/
+/**
+ * This class holds all replicated properties for a given type (either a UClass, UStruct, or UFunction).
+ * Helpers functions exist to read, write, and compare property state.
+ *
+ * There is only one FRepLayout for a given type, meaning all instances of the type share the FRepState.
+ *
+ * COMMANDS:
+ *
+ * All Properties in a RepLayout are represented as Layout Commands.
+ * These commands dictate:
+ *		- What the underlying data type is.
+ *		- How the data is laid out in memory.
+ *		- How the data should be serialized.
+ *		- How the data should be compared (between instances of Objects, Structs, etc.).
+ *		- Whether or not the data should trigger notifications on change (RepNotifies).
+ *		- Whether or not the data is conditional (e.g. may be skipped when sending to some or all connections).
+ *
+ * Commands are split into 2 main types: Parent Commands (@see FRepParentCmd) and Child Commands (@see FRepLayoutCmd).
+ *
+ * A Parent Command represents a Top Level Property of the type represented by an FRepLayout.
+ * A Child Command represents any Property (even nested properties).
+ * Every Parent Command will have at least one associated Child Command.
+ *
+ * E.G.,
+ *		Imagine an Object O, with 4 Properties, CA, DA, I, and S.
+ *			CA is a fixed size C-Style array. This will generate 1 Parent Command and 1 Child Command for *each* element in the array.
+ *			DA is a dynamic array (TArray). This will generate only 1 Parent Command and 1 Child Command, both referencing the array.
+ *			S is a UStruct. This will generate 1 Parent Command and 1 Child Command for the Struct, and an additional Child Command for each property in the Struct.
+ *			I is an integer (or other supported non-Struct type). This will generate 1 Parent Command and 1 Child Command for the int.
+ *
+ * CHANGELISTS
+ *
+ * Along with Layout Commands that describe the Properties in a type, RepLayout uses changelists to know
+ * what Properties have changed between frames. @see FRepChangedHistory.
+ *
+ * Changelists are arrays of Property Handles that describe what Properties have changed, however they don't
+ * track the actual values of the Properties.
+ *
+ * Property Handles are either Layout Command indices (in FRepLayout::Cmds), or array indices for Properties in
+ * dynamic arrays.
+ * Handles are 1-based, reserving 0 as a terminal case.
+ *
+ * In order to generate Changelists, Layout Commands are sequentially applied that compare the values
+ * of an object's cached state to a object's current state. Any properties that are found to be different
+ * will have their handle written into the changelist. This means handles within a changelists are
+ * inherently ordered (with arrays inserted whose Handles are also ordered).
+ *
+ * When we want to replicate properties for an object, merge together any outstanding changelists
+ * and then iterate over it using Layout Commands that serialize the necessary property data.
+ *
+ * Receiving is very similar, except the Handles are baked into the serialized data so no
+ * explicit changelist is required. As each Handle is read, a Layout Command is applied
+ * that serializes the data from the network bunch and applies it to an object.
+ */
 class FRepLayout : public FGCObject
 {
 	friend class FRepState;
@@ -461,101 +788,276 @@ class FRepLayout : public FGCObject
 	friend class UPackageMapClient;
 
 public:
-	FRepLayout() : FirstNonCustomParent( 0 ), RoleIndex( -1 ), RemoteRoleIndex( -1 ), Owner( NULL ) {}
-
-	void OpenAcked( FRepState * RepState ) const;
+	FRepLayout():
+		FirstNonCustomParent(0),
+		RoleIndex(-1),
+		RemoteRoleIndex(-1),
+		Owner(NULL)
+	{}
 
+	/**
+	 * Used to signal that the channel that owns a given object has been opened and acknowledged
+	 * by a client.
+	 *
+	 * @param RepState	RepState for the Object whose channel was acked.
+	 */
+	void OpenAcked(FRepState * RepState) const;
+
+	/**
+	 * Used to initialize the given shadow data.
+	 *
+	 * Shadow Data / Shadow States are used to cache property data so that the Object's state can be
+	 * compared between frames to see if any properties have changed. They are also used on clients
+	 * to keep track of RepNotify state.
+	 *
+	 * This includes:
+	 *		- Allocating memory for all Properties in the class.
+	 *		- Constructing instances of each Property.
+	 *		- Copying the values of the Properties from given object.
+	 *
+	 * @param ShadowData	The buffer where shadow data will be stored.
+	 * @param Class			The class of the object represented by the input memory buffer.
+	 * @param Src			Memory buffer storing object property data.
+	 */
 	void InitShadowData(
-		TArray< uint8, TAlignedHeapAllocator<16> >&	ShadowData,
+		TArray<uint8, TAlignedHeapAllocator<16>>&	ShadowData,
 		UClass *									InObjectClass,
-		const uint8 * const							Src ) const;
-
-	void InitRepState( 
+		const uint8 * const							Src) const;
+
+	/**
+	 * Used to initialize a FRepState.
+	 *
+	 * This includes:
+	 *		- Initializing the ShadowData.
+	 *		- Associating and validating the appropriate ChangedPropertyTracker.
+	 *		- Building initial ConditionMap.
+	 *
+	 * @param RepState						The RepState to initialize.
+	 * @param Class							The class of the object represented by the input memory.
+	 * @param Src							Memory buffer storing object property data.
+	 * @param InRepChangedPropertyTracker	The PropertyTracker we want to associate with the RepState.
+	 */
+	void InitRepState(
 		FRepState *									RepState, 
 		UClass *									InObjectClass, 
 		const uint8 * const							Src, 
-		TSharedPtr< FRepChangedPropertyTracker > &	InRepChangedPropertyTracker ) const;
-
-	void InitChangedTracker( FRepChangedPropertyTracker * ChangedTracker ) const;
-
+		TSharedPtr<FRepChangedPropertyTracker> &	InRepChangedPropertyTracker) const;
+
+	void InitChangedTracker(FRepChangedPropertyTracker * ChangedTracker) const;
+
+	/**
+	 * Writes out any changed properties for an Object into the given data buffer,
+	 * and does book keeping for the RepState of the object.
+	 *
+	 * Note, this does not compare properties or send them on the wire, it's only used
+	 * to serialize properties.
+	 *
+	 * @param RepState				RepState for the object.
+	 * @param RepChangelistState	RepChangelistState for the object.
+	 * @param Data					Pointer to memory where property data is stored.
+	 * @param ObjectClass			Class of the object.
+	 * @param Writer				Writer used to store / write out the replicated properties.
+	 * @param RepFlags				Flags used for replication.
+	 */
 	bool ReplicateProperties(
 		FRepState* RESTRICT				RepState,
 		FRepChangelistState* RESTRICT	RepChangelistState,
 		const uint8* RESTRICT			Data,
-		UClass *						ObjectClass,
+		UClass*							ObjectClass,
 		UActorChannel *					OwningChannel,
 		FNetBitWriter&					Writer,
-		const FReplicationFlags &		RepFlags ) const;
-
+		const FReplicationFlags &		RepFlags) const;
+
+	/**
+	 * Writes all changed property values from the input owner data to the given buffer.
+	 * This is used primarily by ReplicateProperties.
+	 *
+	 * @param RepState			RepState for the object.
+	 * @param ChangedTracker	Used to indicate
+	 * @param Data				Pointer to the object's memory.
+	 * @param ObjectClass		Class of the object.
+	 * @param Writer			Writer used to store / write out the replicated properties.
+	 * @param Changed			Aggregate list of property handles that need to be written.
+	 * @param SharedInfo		Shared Serialization state for properties.
+	 */
 	void SendProperties(
-		FRepState*	RESTRICT		RepState,
-		FRepChangedPropertyTracker* ChangedTracker,
-		const uint8* RESTRICT		Data,
-		UClass*						ObjectClass,
-		FNetBitWriter&				Writer,
-		TArray< uint16 >&			Changed,
-		const FRepSerializationSharedInfo& SharedInfo) const;
+		FRepState*	RESTRICT				RepState,
+		FRepChangedPropertyTracker*			ChangedTracker,
+		const uint8* RESTRICT				Data,
+		UClass*								ObjectClass,
+		FNetBitWriter&						Writer,
+		TArray<uint16>&						Changed,
+		const FRepSerializationSharedInfo&	SharedInfo) const;
+
+	ENGINE_API void InitFromObjectClass(UClass * InObjectClass, const UNetConnection* ServerConnection = nullptr);
+
+	/**
+	 * Reads all property values from the received buffer, and applies them to the
+	 * property memory.
+	 *
+	 * @param OwningChannel			The channel of the Actor that owns the object whose properties we're reading.
+	 * @param InObjectClass			Class of the object.
+	 * @param RepState				RepState for the object.
+	 * @param Data					Pointer to memory where read property data should be stored.
+	 * @param InBunch				The data that should be read.
+	 * @param bOutHasUnmapped		Whether or not unmapped GUIDs were read.
+	 * @param bEnableRepNotifies	Whether or not RepNotifies will be fired due to changed properties.
+	 * @param bOutGuidsChanged		Whether or not any GUIDs were changed.
+	 */
+	bool ReceiveProperties(
+		UActorChannel*			OwningChannel,
+		UClass*					InObjectClass,
+		FRepState* RESTRICT		RepState,
+		void* RESTRICT			Data,
+		FNetBitReader&			InBunch,
+		bool&					bOutHasUnmapped,
+		const bool				bEnableRepNotifies,
+		bool&					bOutGuidsChanged) const;
 
-	ENGINE_API void InitFromObjectClass( UClass * InObjectClass, const UNetConnection* ServerConnection = nullptr );
+	void GatherGuidReferences(
+		FRepState*			RepState,
+		TSet<FNetworkGUID>&	OutReferencedGuids,
+		int32&				OutTrackedGuidMemoryBytes) const;
 
-	bool ReceiveProperties( UActorChannel* OwningChannel, UClass * InObjectClass, FRepState * RESTRICT RepState, void* RESTRICT Data, FNetBitReader & InBunch, bool & bOutHasUnmapped, const bool bEnableRepNotifies, bool& bOutGuidsChanged ) const;
+	bool MoveMappedObjectToUnmapped(FRepState* RepState, const FNetworkGUID& GUID) const;
 
-	void GatherGuidReferences( FRepState* RepState, TSet< FNetworkGUID >& OutReferencedGuids, int32& OutTrackedGuidMemoryBytes ) const;
+	void UpdateUnmappedObjects(
+		FRepState*		RepState,
+		UPackageMap*	PackageMap,
+		UObject*		Object,
+		bool&			bOutSomeObjectsWereMapped,
+		bool&			bOutHasMoreUnmapped) const;
 
-	bool MoveMappedObjectToUnmapped( FRepState* RepState, const FNetworkGUID& GUID ) const;
+	void CallRepNotifies(FRepState * RepState, UObject* Object) const;
 
-	void UpdateUnmappedObjects( FRepState *	RepState, UPackageMap * PackageMap, UObject* Object, bool & bOutSomeObjectsWereMapped, bool & bOutHasMoreUnmapped ) const;
+	void PostReplicate(
+		FRepState*		RepState,
+		FPacketIdRange&	PacketRange,
+		bool			bReliable) const;
 
-	void CallRepNotifies( FRepState * RepState, UObject* Object ) const;
-	void PostReplicate( FRepState * RepState, FPacketIdRange & PacketRange, bool bReliable ) const;
-	void ReceivedNak( FRepState * RepState, int32 NakPacketId ) const;
-	bool AllAcked( FRepState * RepState ) const;
-	bool ReadyForDormancy( FRepState * RepState ) const;
+	void ReceivedNak(FRepState * RepState, int32 NakPacketId) const;
+	bool AllAcked(FRepState * RepState) const;
+	bool ReadyForDormancy(FRepState * RepState) const;
 
-	void ValidateWithChecksum( const void* RESTRICT Data, FBitArchive & Ar ) const;
-	uint32 GenerateChecksum( const FRepState* RepState ) const;
+	void ValidateWithChecksum(const void* RESTRICT Data, FBitArchive & Ar) const;
+	uint32 GenerateChecksum(const FRepState* RepState) const;
 
 	/** Clamp the changelist so that it conforms to the current size of either the array, or arrays within structs/arrays */
-	void PruneChangeList( FRepState * RepState, const void* RESTRICT Data, const TArray< uint16 >& Changed, TArray< uint16 >& PrunedChanged ) const;
-
-	void MergeChangeList( const uint8* RESTRICT Data, const TArray< uint16 >& Dirty1, const TArray< uint16 >& Dirty2, TArray< uint16 >& MergedDirty ) const;
+	void PruneChangeList(
+		FRepState*				RepState,
+		const void* RESTRICT	Data,
+		const TArray<uint16>&	Changed,
+		TArray<uint16>&			PrunedChanged) const;
+
+	/**
+	 * Combines two changelists, ensuring that handles are in the correct order, and arrays are properly structured.
+	 *
+	 * @param Data			Property Data for the Object / RepState the changelists refer to.
+	 * @param Dirty1		First changelist to merge. Must be non-empty, valid changelist.
+	 * @param Dirty2		Second changelist to merge. May be empty, or otherwise valid changelist.
+	 * @param MergedDirty	The combined changelist.
+	 */
+	void MergeChangeList(
+		const uint8* RESTRICT	Data,
+		const TArray<uint16>&	Dirty1,
+		const TArray<uint16>&	Dirty2,
+		TArray<uint16>&			MergedDirty) const;
 
 	DEPRECATED(4.20, "Use the DiffProperties overload with the EDiffPropertiesFlags parameter")
-	bool DiffProperties( TArray<UProperty*>* RepNotifies, void* RESTRICT Destination, const void* RESTRICT Source, const bool bSync ) const;
-
-	bool DiffProperties( TArray<UProperty*>* RepNotifies, void* RESTRICT Destination, const void* RESTRICT Source, const EDiffPropertiesFlags Flags ) const;
-	bool DiffStableProperties( TArray<UProperty*>* RepNotifies, TArray<UObject*>* ObjReferences, void* RESTRICT Destination, const void* RESTRICT Source ) const;
-
-	void GetLifetimeCustomDeltaProperties(TArray< int32 > & OutCustom, TArray< ELifetimeCondition >	& OutConditions);
+	bool DiffProperties(TArray<UProperty*>* RepNotifies, void* RESTRICT Destination, const void* RESTRICT Source, const bool bSync) const;
+
+	/**
+	 * Compare all properties between source and destination buffer, and optionally update the destination
+	 * buffer to match the state of the source buffer if they don't match.
+	 *
+	 * @param RepNotifies	RepNotifies that should be fired if we're changing properties.
+	 * @param Destination	Destination buffer that will be changed if we're changing properties.
+	 * @param Source		Source buffer containing desired property values.
+	 * @param Flags			Controls how DiffProperties behaves.
+	 *
+	 * @return True if there were any properties with different values.
+	 */
+	bool DiffProperties(
+		TArray<UProperty*>*			RepNotifies,
+		void* RESTRICT				Destination,
+		const void* RESTRICT		Source,
+		const EDiffPropertiesFlags	Flags) const;
+
+	/**
+	 * @see DiffProperties
+	 *
+	 * The main difference between this method and DiffProperties is that this method will skip
+	 * any properties that are:
+	 *
+	 *	- Transient
+	 *	- Point to Actors or ActorComponents
+	 *	- Point to Objects that are non-stably named for networking.
+	 *
+	 * @param RepNotifies	RepNotifies that should be fired if we're changing properties.
+	 * @param Destination	Destination buffer that will be changed if we're changing properties.
+	 * @param Source		Source buffer containing desired property values.
+	 * @param Flags			Controls how DiffProperties behaves.
+	 *
+	 * @return True if there were any properties with different values.
+	 */
+	bool DiffStableProperties(
+		TArray<UProperty*>*		RepNotifies,
+		TArray<UObject*>*		ObjReferences,
+		void* RESTRICT			Destination,
+		const void* RESTRICT	Source) const;
+
+	void GetLifetimeCustomDeltaProperties(TArray<int32>& OutCustom, TArray<ELifetimeCondition>& OutConditions);
 
 	// RPC support
-	void InitFromFunction( UFunction * InFunction, const UNetConnection* ServerConnection = nullptr );
-	void ENGINE_API SendPropertiesForRPC( UFunction * Function, UActorChannel * Channel, FNetBitWriter & Writer, void* Data ) const;
-	void ReceivePropertiesForRPC( UObject* Object, UFunction * Function, UActorChannel * Channel, FNetBitReader & Reader, void* Data, TSet<FNetworkGUID>& UnmappedGuids) const;
-
-	// Builds shared serialization state for a multicast rpc
-	void ENGINE_API BuildSharedSerializationForRPC( void* Data );
-
-	// Clears shared serialization state for a multicast rpc
+	void InitFromFunction(UFunction* InFunction, const UNetConnection* ServerConnection = nullptr);
+
+	/** @see SendProperties. */
+	void ENGINE_API SendPropertiesForRPC(
+		UFunction*		Function,
+		UActorChannel*	Channel,
+		FNetBitWriter&	Writer,
+		void*			Data) const;
+
+	/** @see ReceiveProperties. */
+	void ReceivePropertiesForRPC(
+		UObject*			Object,
+		UFunction*			Function,
+		UActorChannel*		Channel,
+		FNetBitReader&		Reader,
+		void*				Data,
+		TSet<FNetworkGUID>&	UnmappedGuids) const;
+
+	/** Builds shared serialization state for a multicast rpc */
+	void ENGINE_API BuildSharedSerializationForRPC(void* Data);
+
+	/** Clears shared serialization state for a multicast rpc */
 	void ENGINE_API ClearSharedSerializationForRPC();
 
 	// Struct support
-	ENGINE_API void SerializePropertiesForStruct( UStruct * Struct, FBitArchive & Ar, UPackageMap	* Map, void* Data, bool & bHasUnmapped ) const;	
-	ENGINE_API void InitFromStruct( UStruct * InStruct, const UNetConnection* ServerConnection = nullptr );
+	ENGINE_API void SerializePropertiesForStruct(
+		UStruct*		Struct,
+		FBitArchive&	Ar,
+		UPackageMap*	Map,
+		void*			Data,
+		bool&			bHasUnmapped) const;
 
-	// Serializes all replicated properties of a UObject in or out of an archive (depending on what type of archive it is)
+	ENGINE_API void InitFromStruct(UStruct * InStruct, const UNetConnection* ServerConnection = nullptr);
+
+	/** Serializes all replicated properties of a UObject in or out of an archive (depending on what type of archive it is). */
 	ENGINE_API void SerializeObjectReplicatedProperties(UObject* Object, FBitArchive & Ar) const;
 
 	UObject* GetOwner() const { return Owner; }
 
+	/** Currently only used for Replays / with the UDemoNetDriver. */
 	void SendProperties_BackwardsCompatible(
 		FRepState* RESTRICT			RepState,
 		FRepChangedPropertyTracker* ChangedTracker,
 		const uint8* RESTRICT		Data,
 		UNetConnection*				Connection,
 		FNetBitWriter&				Writer,
-		TArray< uint16 >&			Changed ) const;
+		TArray<uint16>&				Changed) const;
 
+	/** Currently only used for Replays / with the UDemoNetDriver. */
 	bool ReceiveProperties_BackwardsCompatible(
 		UNetConnection*				Connection,
 		FRepState* RESTRICT			RepState,
@@ -563,19 +1065,37 @@ public:
 		FNetBitReader&				InBunch,
 		bool&						bOutHasUnmapped,
 		const bool					bEnableRepNotifies,
-		bool&						bOutGuidsChanged ) const;
-
+		bool&						bOutGuidsChanged) const;
+
+	/**
+	 * Compare Property Values currently stored in the Changelist State to the Property Values
+	 * in the passed in data, generating a new changelist if necessary.
+	 *
+	 * @param RepState	The FRepChangelistState that contains the last cached values and changelists.
+	 * @param Data		The newest Property Data available.
+	 * @param RepFlags	Flags that will be used if the object is replicated.
+	 */
 	bool CompareProperties(
 		FRepChangelistState* RESTRICT	RepState,
 		const uint8* RESTRICT			Data,
-		const FReplicationFlags&		RepFlags ) const;
+		const FReplicationFlags&		RepFlags) const;
 
+	//~ Begin FGCObject Interface
 	ENGINE_API virtual void AddReferencedObjects(FReferenceCollector& Collector) override;
+	//~ End FGCObject Interface
 
 private:
-	void RebuildConditionalProperties( FRepState * RESTRICT	RepState, const FRepChangedPropertyTracker& ChangedTracker, const FReplicationFlags& RepFlags ) const;
+	void RebuildConditionalProperties(
+		FRepState* RESTRICT					RepState,
+		const FRepChangedPropertyTracker&	ChangedTracker,
+		const FReplicationFlags&			RepFlags) const;
 
-	void UpdateChangelistHistory( FRepState * RepState, UClass * ObjectClass, const uint8* RESTRICT Data, UNetConnection* Connection, TArray< uint16 > * OutMerged ) const;
+	void UpdateChangelistHistory(
+		FRepState*				RepState,
+		UClass*					ObjectClass,
+		const uint8* RESTRICT	Data,
+		UNetConnection*			Connection,
+		TArray<uint16>*			OutMerged) const;
 
 	void SendProperties_BackwardsCompatible_r(
 		FRepState* RESTRICT					RepState,
@@ -585,7 +1105,7 @@ private:
 		FNetBitWriter&						Writer,
 		const bool							bDoChecksum,
 		FRepHandleIterator&					HandleIterator,
-		const uint8* RESTRICT				SourceData ) const;
+		const uint8* RESTRICT				SourceData) const;
 
 	void SendAllProperties_BackwardsCompatible_r(
 		FNetBitWriter&						Writer,
@@ -594,7 +1114,7 @@ private:
 		FNetFieldExportGroup*				NetFieldExportGroup,
 		const int32							CmdStart,
 		const int32							CmdEnd,
-		const uint8*						SourceData ) const;
+		const uint8*						SourceData) const;
 
 	void SendProperties_r(
 		FRepState*	RESTRICT				RepState,
@@ -604,47 +1124,61 @@ private:
 		FRepHandleIterator&					HandleIterator,
 		const uint8* RESTRICT				SourceData,
 		const int32							ArrayDepth,
-		const FRepSerializationSharedInfo&	SharedInfo ) const;
+		const FRepSerializationSharedInfo&	SharedInfo) const;
 
 	uint16 CompareProperties_r(
 		const int32				CmdStart,
 		const int32				CmdEnd,
 		const uint8* RESTRICT	CompareData,
 		const uint8* RESTRICT	Data,
-		TArray< uint16 >&		Changed,
+		TArray<uint16>&		Changed,
 		uint16					Handle,
 		const bool				bIsInitial,
-		const bool				bForceFail ) const;
+		const bool				bForceFail) const;
 
 	void CompareProperties_Array_r(
 		const uint8* RESTRICT	CompareData,
 		const uint8* RESTRICT	Data,
-		TArray< uint16 >&		Changed,
+		TArray<uint16>&			Changed,
 		const uint16			CmdIndex,
 		const uint16			Handle,
 		const bool				bIsInitial,
-		const bool				bForceFail ) const;
+		const bool				bForceFail) const;
 
 	void BuildSharedSerialization(
-		const uint8* RESTRICT		Data,
-		TArray< uint16 >&			Changed,
-		const bool					bWriteHandle,
-		FRepSerializationSharedInfo& SharedInfo ) const;
+		const uint8* RESTRICT			Data,
+		TArray<uint16>&					Changed,
+		const bool						bWriteHandle,
+		FRepSerializationSharedInfo&	SharedInfo) const;
 
 	void BuildSharedSerialization_r(
-		FRepHandleIterator&		RepHandleIterator,
-		const uint8* RESTRICT	SourceData,
-		const bool				bWriteHandle,
-		const bool				bDoChecksum,
-		const int32				ArrayDepth,
-		FRepSerializationSharedInfo& SharedInfo ) const;
-
-	void BuildSharedSerializationForRPC_DynamicArray_r( const int32 CmdIndex, uint8 * Data, int32 AarayDepth, FRepSerializationSharedInfo& SharedInfo );
-	void BuildSharedSerializationForRPC_r( const int32 CmdStart, const int32 CmdEnd, void * Data, int32 ArrayIndex, int32 ArrayDepth, FRepSerializationSharedInfo& SharedInfo );
-
-	TSharedPtr< FNetFieldExportGroup > CreateNetfieldExportGroup() const;
-
-	int32 FindCompatibleProperty( const int32 CmdStart, const int32 CmdEnd, const uint32 Checksum ) const;
+		FRepHandleIterator&				RepHandleIterator,
+		const uint8* RESTRICT			SourceData,
+		const bool						bWriteHandle,
+		const bool						bDoChecksum,
+		const int32						ArrayDepth,
+		FRepSerializationSharedInfo&	SharedInfo) const;
+
+	void BuildSharedSerializationForRPC_DynamicArray_r(
+		const int32						CmdIndex,
+		uint8*							Data,
+		int32							AarayDepth,
+		FRepSerializationSharedInfo&	SharedInfo);
+
+	void BuildSharedSerializationForRPC_r(
+		const int32						CmdStart,
+		const int32						CmdEnd,
+		void*							Data,
+		int32							ArrayIndex,
+		int32							ArrayDepth,
+		FRepSerializationSharedInfo&	SharedInfo);
+
+	TSharedPtr<FNetFieldExportGroup> CreateNetfieldExportGroup() const;
+
+	int32 FindCompatibleProperty(
+		const int32		CmdStart,
+		const int32		CmdEnd,
+		const uint32	Checksum) const;
 
 	bool ReceiveProperties_BackwardsCompatible_r(
 		FRepState * RESTRICT	RepState,
@@ -657,11 +1191,14 @@ private:
 		uint8* RESTRICT			Data,
 		FGuidReferencesMap*		GuidReferencesMap,
 		bool&					bOutHasUnmapped,
-		bool&					bOutGuidsChanged ) const;
+		bool&					bOutGuidsChanged) const;
 
-	void GatherGuidReferences_r( FGuidReferencesMap* GuidReferencesMap, TSet< FNetworkGUID >& OutReferencedGuids, int32& OutTrackedGuidMemoryBytes ) const;
+	void GatherGuidReferences_r(
+		FGuidReferencesMap*	GuidReferencesMap,
+		TSet<FNetworkGUID>&	OutReferencedGuids,
+		int32&				OutTrackedGuidMemoryBytes) const;
 
-	bool MoveMappedObjectToUnmapped_r( FGuidReferencesMap* GuidReferencesMap, const FNetworkGUID& GUID ) const;
+	bool MoveMappedObjectToUnmapped_r(FGuidReferencesMap* GuidReferencesMap, const FNetworkGUID& GUID) const;
 
 	void UpdateUnmappedObjects_r(
 		FRepState*				RepState, 
@@ -672,85 +1209,138 @@ private:
 		uint8* RESTRICT			Data, 
 		const int32				MaxAbsOffset,
 		bool&					bOutSomeObjectsWereMapped,
-		bool&					bOutHasMoreUnmapped ) const;
+		bool&					bOutHasMoreUnmapped) const;
 
-	void ValidateWithChecksum_DynamicArray_r( const FRepLayoutCmd& Cmd, const int32 CmdIndex, const uint8* RESTRICT Data, FBitArchive & Ar ) const;
-	void ValidateWithChecksum_r( const int32 CmdStart, const int32 CmdEnd, const uint8* RESTRICT Data, FBitArchive & Ar ) const;
+	void ValidateWithChecksum_DynamicArray_r(
+		const FRepLayoutCmd&	Cmd,
+		const int32				CmdIndex,
+		const uint8* RESTRICT	Data,
+		FBitArchive&			Ar) const;
 
-	void SanityCheckChangeList_DynamicArray_r( 
+	void ValidateWithChecksum_r(
+		const int32				CmdStart,
+		const int32				CmdEnd,
+		const uint8* RESTRICT	Data,
+		FBitArchive&			Ar) const;
+
+	void SanityCheckChangeList_DynamicArray_r(
 		const int32				CmdIndex, 
 		const uint8* RESTRICT	Data, 
-		TArray< uint16 > &		Changed,
-		int32 &					ChangedIndex ) const;
+		TArray<uint16> &		Changed,
+		int32 &					ChangedIndex) const;
 
-	uint16 SanityCheckChangeList_r( 
+	uint16 SanityCheckChangeList_r(
 		const int32				CmdStart, 
 		const int32				CmdEnd, 
 		const uint8* RESTRICT	Data, 
-		TArray< uint16 > &		Changed,
+		TArray<uint16> &		Changed,
 		int32 &					ChangedIndex,
-		uint16					Handle ) const;
-
-	void SanityCheckChangeList( const uint8* RESTRICT Data, TArray< uint16 > & Changed ) const;
-
-	uint16 AddParentProperty( UProperty * Property, int32 ArrayIndex );
-
-	int32 InitFromProperty_r( UProperty * Property, int32 Offset, int32 RelativeHandle, int32 ParentIndex, uint32 ParentChecksum, int32 StaticArrayIndex, const UNetConnection* ServerConnection );
+		uint16					Handle) const;
+
+	void SanityCheckChangeList(const uint8* RESTRICT Data, TArray<uint16>& Changed) const;
+
+	uint16 AddParentProperty(UProperty* Property, int32 ArrayIndex);
+
+	int32 InitFromProperty_r(
+		UProperty* Property,
+		int32					Offset,
+		int32					RelativeHandle,
+		int32					ParentIndex,
+		uint32					ParentChecksum,
+		int32					StaticArrayIndex,
+		const UNetConnection*	ServerConnection);
+
+	uint32 AddPropertyCmd(
+		UProperty* Property,
+		int32					Offset,
+		int32					RelativeHandle,
+		int32					ParentIndex,
+		uint32					ParentChecksum,
+		int32					StaticArrayIndex,
+		const UNetConnection*	ServerConnection);
+
+	uint32 AddArrayCmd(
+		UArrayProperty*			Property,
+		int32					Offset,
+		int32					RelativeHandle,
+		int32					ParentIndex,
+		uint32					ParentChecksum,
+		int32					StaticArrayIndex,
+		const UNetConnection*	ServerConnection);
 
-	uint32 AddPropertyCmd( UProperty * Property, int32 Offset, int32 RelativeHandle, int32 ParentIndex, uint32 ParentChecksum, int32 StaticArrayIndex, const UNetConnection* ServerConnection );
-	uint32 AddArrayCmd( UArrayProperty * Property, int32 Offset, int32 RelativeHandle, int32 ParentIndex, uint32 ParentChecksum, int32 StaticArrayIndex, const UNetConnection* ServerConnection );
 	void AddReturnCmd();
 
-	void SerializeProperties_DynamicArray_r( 
-		FBitArchive &		Ar, 
-		UPackageMap	*		Map,
-		const int32			CmdIndex,
-		uint8 *				Data,
-		bool &				bHasUnmapped,
-		const int32			ArrayDepth,
-		const FRepSerializationSharedInfo&	SharedInfo ) const;
-
-	void SerializeProperties_r( 
-		FBitArchive &		Ar, 
-		UPackageMap	*		Map,
-		const int32			CmdStart, 
-		const int32			CmdEnd, 
-		void *				Data,
-		bool &				bHasUnmapped,
-		const int32			ArrayIndex,
-		const int32			ArrayDepth,
-		const FRepSerializationSharedInfo&	SharedInfo ) const;
+	void SerializeProperties_DynamicArray_r(
+		FBitArchive &						Ar, 
+		UPackageMap*						Map,
+		const int32							CmdIndex,
+		uint8*								Data,
+		bool &								bHasUnmapped,
+		const int32							ArrayDepth,
+		const FRepSerializationSharedInfo&	SharedInfo) const;
+
+	void SerializeProperties_r(
+		FBitArchive&						Ar, 
+		UPackageMap*						Map,
+		const int32							CmdStart, 
+		const int32							CmdEnd, 
+		void*								Data,
+		bool&								bHasUnmapped,
+		const int32							ArrayIndex,
+		const int32							ArrayDepth,
+		const FRepSerializationSharedInfo&	SharedInfo) const;
 
 	void MergeChangeList_r(
 		FRepHandleIterator&		RepHandleIterator1,
 		FRepHandleIterator&		RepHandleIterator2,
 		const uint8* RESTRICT	SourceData,
-		TArray< uint16 >&		OutChanged ) const;
+		TArray<uint16>&			OutChanged) const;
 
 	void PruneChangeList_r(
 		FRepHandleIterator&		RepHandleIterator,
 		const uint8* RESTRICT	SourceData,
-		TArray< uint16 >&		OutChanged ) const;
+		TArray<uint16>&			OutChanged) const;
 		
-	void BuildChangeList_r( const TArray< FHandleToCmdIndex >& HandleToCmdIndex, const int32 CmdStart, const int32 CmdEnd, uint8* Data, const int32 HandleOffset, TArray< uint16 >& Changed ) const;
+	void BuildChangeList_r(
+		const TArray<FHandleToCmdIndex>&	HandleToCmdIndex,
+		const int32							CmdStart,
+		const int32							CmdEnd,
+		uint8*								Data,
+		const int32							HandleOffset,
+		TArray<uint16>&						Changed) const;
+
+	void BuildHandleToCmdIndexTable_r(
+		const int32					CmdStart,
+		const int32					CmdEnd,
+		TArray<FHandleToCmdIndex>&	HandleToCmdIndex);
+
+	void ConstructProperties(TArray<uint8, TAlignedHeapAllocator<16>>& ShadowData) const;
+	void InitProperties(TArray<uint8, TAlignedHeapAllocator<16>>& ShadowData, const uint8* const Src) const;
+	void DestructProperties(FRepStateStaticBuffer& RepStateStaticBuffer) const;
+
+	/** Top level Layout Commands. */
+	TArray<FRepParentCmd> Parents;
+
+	/** All Layout Commands. */
+	TArray<FRepLayoutCmd> Cmds;
 
-	void BuildHandleToCmdIndexTable_r( const int32 CmdStart, const int32 CmdEnd, TArray< FHandleToCmdIndex >& HandleToCmdIndex );
+	/** Converts a relative handle to the appropriate index into the Cmds array */
+	TArray<FHandleToCmdIndex> BaseHandleToCmdIndex;
 
-	void ConstructProperties( TArray< uint8, TAlignedHeapAllocator<16> >& ShadowData ) const;
-	void InitProperties( TArray< uint8, TAlignedHeapAllocator<16> >& ShadowData, const uint8* const Src ) const;
-	void DestructProperties( FRepStateStaticBuffer& RepStateStaticBuffer ) const;
+	int32 FirstNonCustomParent;
 
-	TArray< FRepParentCmd >		Parents;
-	TArray< FRepLayoutCmd >		Cmds;
+	/** Index of the Role property in the Parents list. May be INDEX_NONE if Owner doesn't have the property. */
+	int32 RoleIndex;
 
-	TArray< FHandleToCmdIndex >	BaseHandleToCmdIndex;		// Converts a relative handle to the appropriate index into the Cmds array
+	/** Index of the RemoteRole property in the Parents list. May be INDEX_NONE if Owner doesn't have the property. */
+	int32 RemoteRoleIndex;
 
-	int32						FirstNonCustomParent;
-	int32						RoleIndex;
-	int32						RemoteRoleIndex;
+	/** UClass, UStruct, or UFunction that this FRepLayout represents.*/
+	UObject* Owner;
 
-	UObject *					Owner;						// Either a UCkass or UFunction
+	/** Shared serialization state for a multicast rpc */
+	FRepSerializationSharedInfo SharedInfoRPC;
 
-	FRepSerializationSharedInfo	SharedInfoRPC;					// Shared serialization state for a multicast rpc
-	TBitArray<>					SharedInfoRPCParentsChanged;	// Shared comparison to default state for multicast rpc
+	/** Shared comparison to default state for multicast rpc */
+	TBitArray<> SharedInfoRPCParentsChanged;
 };
