Fix for attachement issue with actors using replicated rootcomponent where the actor sometime would get invalid offset on clients.

diff --git a/Engine/Source/Runtime/Engine/Classes/Components/SceneComponent.h b/Engine/Source/Runtime/Engine/Classes/Components/SceneComponent.h
--- a/Engine/Source/Runtime/Engine/Classes/Components/SceneComponent.h
+++ b/Engine/Source/Runtime/Engine/Classes/Components/SceneComponent.h
@@ -148,11 +148,11 @@ public:
 	FBoxSphereBounds Bounds;
 
 	/** Location of the component relative to its parent */
-	UPROPERTY(EditAnywhere, BlueprintReadOnly, ReplicatedUsing=OnRep_Transform, Category = Transform)
+	UPROPERTY(EditAnywhere, BlueprintReadOnly, ReplicatedUsing=OnRep_RelativeLocation, Category = Transform)
 	FVector RelativeLocation;
 
 	/** Rotation of the component relative to its parent */
-	UPROPERTY(EditAnywhere, BlueprintReadOnly, ReplicatedUsing=OnRep_Transform, Category=Transform)
+	UPROPERTY(EditAnywhere, BlueprintReadOnly, ReplicatedUsing=OnRep_RelativeRotation, Category=Transform)
 	FRotator RelativeRotation;
 
 	/**
@@ -252,6 +252,8 @@ private:
 
 	uint8 bNetUpdateTransform : 1;
 	uint8 bNetUpdateAttachment : 1;
+	uint8 bNetHasReceivedRelativeLocation : 1;
+	uint8 bNetHasReceivedRelativeRotation : 1;
 
 	// DEPRECATED
 	UPROPERTY()
@@ -320,6 +322,12 @@ private:
 	UFUNCTION()
 	void OnRep_Transform();
 
+	UFUNCTION()
+	void OnRep_RelativeLocation();
+
+	UFUNCTION()
+	void OnRep_RelativeRotation();
+
 	UFUNCTION()
 	void OnRep_AttachParent();
 
diff --git a/Engine/Source/Runtime/Engine/Private/Components/SceneComponent.cpp b/Engine/Source/Runtime/Engine/Private/Components/SceneComponent.cpp
--- a/Engine/Source/Runtime/Engine/Private/Components/SceneComponent.cpp
+++ b/Engine/Source/Runtime/Engine/Private/Components/SceneComponent.cpp
@@ -75,6 +75,9 @@ USceneComponent::USceneComponent(const FObjectInitializer& ObjectInitializer /*=
 	// default behavior is visible
 	bVisible = true;
 	bAutoActivate = false;
+
+	bNetHasReceivedRelativeLocation = false;
+	bNetHasReceivedRelativeRotation = false;
 }
 
 #if WITH_EDITORONLY_DATA
@@ -3119,6 +3122,18 @@ void USceneComponent::OnRep_Transform()
 	bNetUpdateTransform = true;
 }
 
+void USceneComponent::OnRep_RelativeLocation()
+{
+	bNetHasReceivedRelativeLocation = true;
+	OnRep_Transform();
+}
+
+void USceneComponent::OnRep_RelativeRotation()
+{
+	bNetHasReceivedRelativeRotation = true;
+	OnRep_Transform();
+}
+
 void USceneComponent::OnRep_AttachParent()
 {
 	bNetUpdateAttachment = true;
@@ -3211,6 +3226,19 @@ void USceneComponent::PostRepNotifies()
 	{
 		Exchange(NetOldAttachParent, AttachParent);
 		Exchange(NetOldAttachSocketName, AttachSocketName);
+		
+		// Note: This is a local fix for JIRA UE-43355. If we receive bNetUpdateAttachment without having received a updated transform, assume that we intend to snap and that the relative location/rotation should defaulted
+		if (!bNetHasReceivedRelativeLocation)
+		{
+			RelativeLocation = FVector::ZeroVector;
+			bNetHasReceivedRelativeLocation = true;
+		}
+		if (!bNetHasReceivedRelativeRotation)
+		{
+			RelativeRotation = FRotator::ZeroRotator;
+			bNetHasReceivedRelativeRotation = true;
+		}
+
 		AttachToComponent(NetOldAttachParent, FAttachmentTransformRules::KeepRelativeTransform, NetOldAttachSocketName);
 		bNetUpdateAttachment = false;
 	}
