Fix Role and RemoteRole getting stuck on clients.

diff --git a/Engine/Source/Runtime/Engine/Private/Actor.cpp b/Engine/Source/Runtime/Engine/Private/Actor.cpp
--- a/Engine/Source/Runtime/Engine/Private/Actor.cpp
+++ b/Engine/Source/Runtime/Engine/Private/Actor.cpp
@@ -3153,19 +3153,24 @@ void AActor::SetReplicates(bool bInReplicates)
 { 
 	if (Role == ROLE_Authority)
 	{
-		const bool ChangedReplicates = (bReplicates == false && bInReplicates == true);
-
-		// Update our settings before calling into net driver
-		RemoteRole = (bInReplicates ? ROLE_SimulatedProxy : ROLE_None);
-		bReplicates = bInReplicates;
-
 		// Only call into net driver if we actually changed
-		if (ChangedReplicates)
+		if (bReplicates != bInReplicates)
 		{
-			if (UWorld* MyWorld = GetWorld())		// GetWorld will return nullptr on CDO, FYI
+			// Update our settings before calling into net driver
+			RemoteRole = bInReplicates ? ROLE_SimulatedProxy : ROLE_None;
+			bReplicates = bInReplicates;
+
+			// This actor should already be in the Network Actors List if it was already replicating.
+			if (bReplicates)
 			{
-				MyWorld->AddNetworkActor(this);
+				// GetWorld will return nullptr on CDO, FYI
+				if (UWorld* MyWorld = GetWorld())		
+				{
+					MyWorld->AddNetworkActor(this);
+					ForcePropertyCompare();
+				}
 			}
+			
 		}
 	}
 	else
@@ -3206,6 +3211,7 @@ void AActor::CopyRemoteRoleFrom(const AActor* CopyFromActor)
 	if (RemoteRole != ROLE_None)
 	{
 		GetWorld()->AddNetworkActor(this);
+		ForcePropertyCompare();
 	}
 }
 
@@ -3238,6 +3244,7 @@ void AActor::ExchangeNetRoles(bool bRemoteOwned)
 	{
 		if (bRemoteOwned)
 		{
+			// Don't worry about calling SetRemoteRoleInternal here, as this should only be hit during initialization.
 			Exchange( Role, RemoteRole );
 		}
 		bExchangedRoles = true;
@@ -3247,6 +3254,7 @@ void AActor::ExchangeNetRoles(bool bRemoteOwned)
 void AActor::SwapRoles()
 {
 	Swap(Role, RemoteRole);
+	ForcePropertyCompare();
 }
 
 void AActor::DispatchBeginPlay()
diff --git a/Engine/Source/Runtime/Engine/Private/DataReplication.cpp b/Engine/Source/Runtime/Engine/Private/DataReplication.cpp
--- a/Engine/Source/Runtime/Engine/Private/DataReplication.cpp
+++ b/Engine/Source/Runtime/Engine/Private/DataReplication.cpp
@@ -411,7 +411,8 @@ FReplicationChangelistMgr::~FReplicationChangelistMgr()
 {
 }
 
-void FReplicationChangelistMgr::Update( const UObject* InObject, const uint32 ReplicationFrame, const int32 LastCompareIndex, const FReplicationFlags& RepFlags, const bool bForceCompare )
+PRAGMA_DISABLE_DEPRECATION_WARNINGS
+void FReplicationChangelistMgr::Update(const UObject* InObject, const uint32 ReplicationFrame, const int32 LastCompareIndex, const FReplicationFlags& RepFlags, const bool bForceCompare)
 {
 	// See if we can re-use the work already done on a previous connection
 	// Rules:
@@ -419,13 +420,33 @@ void FReplicationChangelistMgr::Update( const UObject* InObject, const uint32 Re
 	//	2. We check LastCompareIndex > 1 so we can do at least one pass per connection to compare all properties
 	//		This is necessary due to how RemoteRole is manipulated per connection, so we need to give all connections a chance to see if it changed
 	//	3. We ALWAYS compare on bNetInitial to make sure we have a fresh changelist of net initial properties in this case
-	if ( !bForceCompare && GShareShadowState && !RepFlags.bNetInitial && LastCompareIndex > 1 && LastReplicationFrame == ReplicationFrame )
+	if (!bForceCompare && GShareShadowState && !RepFlags.bNetInitial && LastCompareIndex > 1 && LastReplicationFrame == ReplicationFrame)
+	{
+		INC_DWORD_STAT_BY(STAT_NetSkippedDynamicProps, 1);
+		return;
+	}
+
+	RepLayout->CompareProperties(nullptr, RepChangelistState.Get(), (const uint8*)InObject, RepFlags);
+
+	LastReplicationFrame = ReplicationFrame;
+}
+PRAGMA_ENABLE_DEPRECATION_WARNINGS
+
+void FReplicationChangelistMgr::Update(FRepState* RESTRICT RepState, const UObject* InObject, const uint32 ReplicationFrame, const FReplicationFlags& RepFlags, const bool bForceCompare)
+{
+	// See if we can re-use the work already done on a previous connection
+	// Rules:
+	//	1. We always compare once per frame (i.e. check LastReplicationFrame == ReplicationFrame)
+	//	2. We check LastCompareIndex > 1 so we can do at least one pass per connection to compare all properties
+	//		This is necessary due to how RemoteRole is manipulated per connection, so we need to give all connections a chance to see if it changed
+	//	3. We ALWAYS compare on bNetInitial to make sure we have a fresh changelist of net initial properties in this case
+	if (!bForceCompare && GShareShadowState && !RepFlags.bNetInitial && RepState->LastCompareIndex > 1 && LastReplicationFrame == ReplicationFrame)
 	{
 		INC_DWORD_STAT_BY( STAT_NetSkippedDynamicProps, 1 );
 		return;
 	}
 
-	RepLayout->CompareProperties( RepChangelistState.Get(), (const uint8*)InObject, RepFlags );
+	RepLayout->CompareProperties(RepState, RepChangelistState.Get(), (const uint8*)InObject, RepFlags);
 
 	LastReplicationFrame = ReplicationFrame;
 }
@@ -1296,7 +1317,7 @@ bool FObjectReplicator::ReplicateProperties( FOutBunch & Bunch, FReplicationFlag
 	FNetBitWriter Writer( Bunch.PackageMap, 8192 );
 
 	// Update change list (this will re-use work done by previous connections)
-	ChangelistMgr->Update( Object, Connection->Driver->ReplicationFrame, RepState->LastCompareIndex, RepFlags, OwningChannel->bForceCompareProperties );
+	ChangelistMgr->Update(RepState.Get(), Object, Connection->Driver->ReplicationFrame, RepFlags, OwningChannel->bForceCompareProperties);
 
 	// Replicate properties in the layout
 	const bool bHasRepLayout = RepLayout->ReplicateProperties( RepState.Get(), ChangelistMgr->GetRepChangelistState(), ( uint8* )Object, ObjectClass, OwningChannel, Writer, RepFlags );
diff --git a/Engine/Source/Runtime/Engine/Private/RepLayout.cpp b/Engine/Source/Runtime/Engine/Private/RepLayout.cpp
--- a/Engine/Source/Runtime/Engine/Private/RepLayout.cpp
+++ b/Engine/Source/Runtime/Engine/Private/RepLayout.cpp
@@ -503,6 +503,7 @@ public:
 };
 
 uint16 FRepLayout::CompareProperties_r(
+	FRepState* RESTRICT		RepState,
 	const int32				CmdStart,
 	const int32				CmdEnd,
 	const uint8* RESTRICT	CompareData,
@@ -535,7 +536,7 @@ uint16 FRepLayout::CompareProperties_r(
 			}
 
 			// Once we hit an array, start using a stack based approach
-			CompareProperties_Array_r(CompareData + Cmd.Offset, (const uint8*)Data + Cmd.Offset, Changed, CmdIndex, Handle, bIsInitial, bForceFail);
+			CompareProperties_Array_r(RepState, CompareData + Cmd.Offset, ( const uint8* )Data + Cmd.Offset, Changed, CmdIndex, Handle, bIsInitial, bForceFail);
 			CmdIndex = Cmd.EndCmd - 1;		// The -1 to handle the ++ in the for loop
 			continue;
 		}
@@ -545,7 +546,27 @@ uint16 FRepLayout::CompareProperties_r(
 			continue;
 		}
 
-		if (bForceFail || !PropertiesAreIdentical(Cmd, (const void*)(CompareData + Cmd.Offset), (const void*)(Data + Cmd.Offset)))
+		// RepState may be null in the case where a deprecated version of this function is called.
+		// In that case, just allow this to fail and perform the old logic.
+		if (RepState && Cmd.ParentIndex == RoleIndex)
+		{
+			const ENetRole ObjectRole = *(const ENetRole*)(Data + Cmd.Offset);
+			if (bForceFail || RepState->SavedRole != ObjectRole)
+			{
+				RepState->SavedRole = ObjectRole;
+				Changed.Add(Handle);
+			}
+		}
+		else if (RepState && Cmd.ParentIndex == RemoteRoleIndex)
+		{
+			const ENetRole ObjectRemoteRole = *(const ENetRole*)(Data + Cmd.Offset);
+			if (bForceFail || RepState->SavedRemoteRole != ObjectRemoteRole)
+			{
+				RepState->SavedRemoteRole = ObjectRemoteRole;
+				Changed.Add(Handle);
+			}
+		}
+		else if (bForceFail || !PropertiesAreIdentical(Cmd, (const void*)(CompareData + Cmd.Offset), (const void*)(Data + Cmd.Offset)))
 		{
 			StoreProperty(Cmd, (void*)(CompareData + Cmd.Offset), (const void*)(Data + Cmd.Offset));
 			Changed.Add(Handle);
@@ -556,6 +577,7 @@ uint16 FRepLayout::CompareProperties_r(
 }
 
 void FRepLayout::CompareProperties_Array_r(
+	FRepState* RESTRICT		RepState,
 	const uint8* RESTRICT	CompareData,
 	const uint8* RESTRICT	Data,
 	TArray<uint16>&			Changed,
@@ -590,7 +612,7 @@ void FRepLayout::CompareProperties_Array_r(
 
 		const bool bNewForceFail = bForceFail || i >= CompareArrayNum;
 
-		LocalHandle = CompareProperties_r(CmdIndex + 1, Cmd.EndCmd - 1, CompareData + ElementOffset, Data + ElementOffset, ChangedLocal, LocalHandle, bIsInitial, bNewForceFail);
+		LocalHandle = CompareProperties_r(RepState, CmdIndex + 1, Cmd.EndCmd - 1, CompareData + ElementOffset, Data + ElementOffset, ChangedLocal, LocalHandle, bIsInitial, bNewForceFail);
 	}
 
 	if (ChangedLocal.Num())
@@ -613,6 +635,7 @@ void FRepLayout::CompareProperties_Array_r(
 }
 
 bool FRepLayout::CompareProperties(
+	FRepState* RESTRICT				RepState,
 	FRepChangelistState* RESTRICT	RepChangelistState,
 	const uint8* RESTRICT			Data,
 	const FReplicationFlags&		RepFlags) const
@@ -629,7 +652,7 @@ bool FRepLayout::CompareProperties(
 	TArray<uint16>& Changed = NewHistoryItem.Changed;
 	Changed.Empty();
 
-	CompareProperties_r(0, Cmds.Num() - 1, RepChangelistState->StaticBuffer.GetData(), Data, Changed, 0, RepFlags.bNetInitial, false);
+	CompareProperties_r(RepState, 0, Cmds.Num() - 1, RepChangelistState->StaticBuffer.GetData(), Data, Changed, 0, RepFlags.bNetInitial, false);
 
 	if (Changed.Num() == 0)
 	{
@@ -1423,6 +1446,17 @@ void FRepLayout::SendProperties_r(
 			WritePropertyHandle(Writer, 0, bDoChecksum);		// Signify end of dynamic array
 			continue;
 		}
+		else
+		{
+			if (Cmd.ParentIndex == RoleIndex)
+			{
+				Data = reinterpret_cast<const uint8*>(&(RepState->SavedRole));
+			}
+			else if (Cmd.ParentIndex == RemoteRoleIndex)
+			{
+				Data = reinterpret_cast<const uint8*>(&(RepState->SavedRemoteRole));
+			}
+		}	
 
 #if !(UE_BUILD_SHIPPING || UE_BUILD_TEST)
 		if (GDoReplicationContextString> 0)
@@ -1737,6 +1771,17 @@ void FRepLayout::SendProperties_BackwardsCompatible_r(
 			HandleIterator.ChangelistIterator.ChangedIndex++;
 			continue;
 		}
+		else
+		{
+			if (Cmd.ParentIndex == RoleIndex)
+			{
+				Data = reinterpret_cast<const uint8*>(&(RepState->SavedRole));
+			}
+			else if (Cmd.ParentIndex == RemoteRoleIndex)
+			{
+				Data = reinterpret_cast<const uint8*>(&(RepState->SavedRemoteRole));
+			}
+		}
 
 		WriteProperty_BackwardsCompatible(Writer, Cmd, HandleIterator.CmdIndex, Owner, Data, bDoChecksum);
 	}
@@ -1745,6 +1790,7 @@ void FRepLayout::SendProperties_BackwardsCompatible_r(
 }
 
 void FRepLayout::SendAllProperties_BackwardsCompatible_r(
+	FRepState* RESTRICT					RepState,
 	FNetBitWriter&						Writer,
 	const bool							bDoChecksum,
 	UPackageMapClient*					PackageMapClient,
@@ -1782,7 +1828,7 @@ void FRepLayout::SendAllProperties_BackwardsCompatible_r(
 				uint32 ArrayIndex = i + 1;
 				TempWriter.SerializeIntPacked(ArrayIndex);
 
-				SendAllProperties_BackwardsCompatible_r(TempWriter, bDoChecksum, PackageMapClient, NetFieldExportGroup, CmdIndex + 1, Cmd.EndCmd - 1, ((const uint8*)Array->GetData()) + Cmd.ElementSize * i);
+				SendAllProperties_BackwardsCompatible_r(RepState, TempWriter, bDoChecksum, PackageMapClient, NetFieldExportGroup, CmdIndex + 1, Cmd.EndCmd - 1, ((const uint8*)Array->GetData()) + Cmd.ElementSize * i);
 			}
 
 			uint32 EndArrayIndex = 0;
@@ -1795,6 +1841,17 @@ void FRepLayout::SendAllProperties_BackwardsCompatible_r(
 			CmdIndex = Cmd.EndCmd - 1;		// The -1 to handle the ++ in the for loop
 			continue;
 		}
+		else
+		{
+			if (Cmd.ParentIndex == RoleIndex)
+			{
+				Data = reinterpret_cast<const uint8*>(&(RepState->SavedRole));
+			}
+			else if (Cmd.ParentIndex == RemoteRoleIndex)
+			{
+				Data = reinterpret_cast<const uint8*>(&(RepState->SavedRemoteRole));
+			}
+		}
 
 		WriteProperty_BackwardsCompatible(Writer, Cmd, CmdIndex, Owner, Data, bDoChecksum);
 	}
@@ -1836,7 +1893,7 @@ void FRepLayout::SendProperties_BackwardsCompatible(
 
 	if (Changed.Num() == 0)
 	{
-		SendAllProperties_BackwardsCompatible_r(Writer, bDoChecksum, PackageMapClient, NetFieldExportGroup.Get(), 0, Cmds.Num() - 1, Data);
+		SendAllProperties_BackwardsCompatible_r(RepState, Writer, bDoChecksum, PackageMapClient, NetFieldExportGroup.Get(), 0, Cmds.Num() - 1, Data);
 	}
 	else
 	{
diff --git a/Engine/Source/Runtime/Engine/Public/Net/DataReplication.h b/Engine/Source/Runtime/Engine/Public/Net/DataReplication.h
--- a/Engine/Source/Runtime/Engine/Public/Net/DataReplication.h
+++ b/Engine/Source/Runtime/Engine/Public/Net/DataReplication.h
@@ -61,7 +61,19 @@ public:
 
 	~FReplicationChangelistMgr();
 
-	void Update( const UObject* InObject, const uint32 ReplicationFrame, const int32 LastCompareIndex, const FReplicationFlags& RepFlags, const bool bForceCompare );
+	DEPRECATED(4.22, "Please use the version of Update that accepts a FRepState pointer.")
+	void Update(const UObject* InObject, const uint32 ReplicationFrame, const int32 LastCompareIndex, const FReplicationFlags& RepFlags, const bool bForceCompare);
+
+	/**
+	 * Updates the shared RepChangelistState for the given object, potentially skipping unnecessary updates.
+	 * See FRepLayout::CompareProperties.
+	 *
+	 * @param RepState		The connection specific RepState for this object.
+	 * @param InObject		The Object associated with the Changelist / RepLayout.
+	 * @param RepFlags		Replication Flags that will be used if the object needs to be replicated.
+	 * @param bForceCompare	Force the comparison, even if other connections have already done it this frame.
+	 */
+	void Update(FRepState* RESTRICT RepState, const UObject* InObject, const uint32 ReplicationFrame, const FReplicationFlags& RepFlags, const bool bForceCompare);
 
 	FRepChangelistState* GetRepChangelistState() const { return RepChangelistState.Get(); }
 
diff --git a/Engine/Source/Runtime/Engine/Public/Net/RepLayout.h b/Engine/Source/Runtime/Engine/Public/Net/RepLayout.h
--- a/Engine/Source/Runtime/Engine/Public/Net/RepLayout.h
+++ b/Engine/Source/Runtime/Engine/Public/Net/RepLayout.h
@@ -456,6 +456,12 @@ public:
 	 * @see ELifetimeCondition.
 	 */
 	bool ConditionMap[COND_Max];
+
+	// Cache off the RemoteRole and Role per connection to avoid issues with
+	// FScopedRoleDowngrade. See UE-66313 (among others).
+
+	ENetRole SavedRemoteRole = ROLE_MAX;
+	ENetRole SavedRole = ROLE_MAX;
 };
 
 /** Various types of Properties supported for Replication. */
@@ -1084,16 +1090,27 @@ public:
 		const bool					bEnableRepNotifies,
 		bool&						bOutGuidsChanged) const;
 
+	DEPRECATED(4.22, "Please use the version of CompareProperties that accepts a FRepState pointer.")
+	bool CompareProperties(
+		FRepChangelistState* RESTRICT	RepState,
+		const uint8* RESTRICT			Data,
+		const FReplicationFlags&		RepFlags) const
+	{
+		CompareProperties(nullptr, RepState, Data, RepFlags);
+	}
+
 	/**
 	 * Compare Property Values currently stored in the Changelist State to the Property Values
 	 * in the passed in data, generating a new changelist if necessary.
 	 *
-	 * @param RepState	The FRepChangelistState that contains the last cached values and changelists.
-	 * @param Data		The newest Property Data available.
-	 * @param RepFlags	Flags that will be used if the object is replicated.
+	 * @param RepState				RepState for the object.
+	 * @param RepChangelistState	The FRepChangelistState that contains the last cached values and changelists.
+	 * @param Data					The newest Property Data available.
+	 * @param RepFlags				Flags that will be used if the object is replicated.
 	 */
 	bool CompareProperties(
-		FRepChangelistState* RESTRICT	RepState,
+		FRepState* RESTRICT				RepState,
+		FRepChangelistState* RESTRICT	RepChangelistState,
 		const uint8* RESTRICT			Data,
 		const FReplicationFlags&		RepFlags) const;
 
@@ -1125,6 +1142,7 @@ private:
 		const uint8* RESTRICT				SourceData) const;
 
 	void SendAllProperties_BackwardsCompatible_r(
+		FRepState* RESTRICT					RepState,
 		FNetBitWriter&						Writer,
 		const bool							bDoChecksum,
 		UPackageMapClient*					PackageMapClient,
@@ -1144,6 +1162,7 @@ private:
 		const FRepSerializationSharedInfo&	SharedInfo) const;
 
 	uint16 CompareProperties_r(
+		FRepState* RESTRICT		RepState,
 		const int32				CmdStart,
 		const int32				CmdEnd,
 		const uint8* RESTRICT	CompareData,
@@ -1154,6 +1173,7 @@ private:
 		const bool				bForceFail) const;
 
 	void CompareProperties_Array_r(
+		FRepState* RESTRICT		RepState,
 		const uint8* RESTRICT	CompareData,
 		const uint8* RESTRICT	Data,
 		TArray<uint16>&			Changed,
