Fixup Deprecated macros in CL-4596859

diff --git a/Engine/Source/Runtime/Engine/Public/Net/DataReplication.h b/Engine/Source/Runtime/Engine/Public/Net/DataReplication.h
--- a/Engine/Source/Runtime/Engine/Public/Net/DataReplication.h
+++ b/Engine/Source/Runtime/Engine/Public/Net/DataReplication.h
@@ -61,7 +61,7 @@ public:
 
 	~FReplicationChangelistMgr();
 
-	DEPRECATED(4.22, "Please use the version of Update that accepts a FRepState pointer.")
+	UE_DEPRECATED(4.22, "Please use the version of Update that accepts a FRepState pointer.")
 	void Update(const UObject* InObject, const uint32 ReplicationFrame, const int32 LastCompareIndex, const FReplicationFlags& RepFlags, const bool bForceCompare);
 
 	/**
diff --git a/Engine/Source/Runtime/Engine/Public/Net/RepLayout.h b/Engine/Source/Runtime/Engine/Public/Net/RepLayout.h
--- a/Engine/Source/Runtime/Engine/Public/Net/RepLayout.h
+++ b/Engine/Source/Runtime/Engine/Public/Net/RepLayout.h
@@ -1090,7 +1090,7 @@ public:
 		const bool					bEnableRepNotifies,
 		bool&						bOutGuidsChanged) const;
 
-	DEPRECATED(4.22, "Please use the version of CompareProperties that accepts a FRepState pointer.")
+	UE_DEPRECATED(4.22, "Please use the version of CompareProperties that accepts a FRepState pointer.")
 	bool CompareProperties(
 		FRepChangelistState* RESTRICT	RepState,
 		const uint8* RESTRICT			Data,
