Fix missing return

diff --git a/Engine/Source/Runtime/Engine/Public/Net/RepLayout.h b/Engine/Source/Runtime/Engine/Public/Net/RepLayout.h
--- a/Engine/Source/Runtime/Engine/Public/Net/RepLayout.h
+++ b/Engine/Source/Runtime/Engine/Public/Net/RepLayout.h
@@ -1096,7 +1096,7 @@ public:
 		const uint8* RESTRICT			Data,
 		const FReplicationFlags&		RepFlags) const
 	{
-		CompareProperties(nullptr, RepState, Data, RepFlags);
+		return CompareProperties(nullptr, RepState, Data, RepFlags);
 	}
 
 	/**
