Increase the max bunch header size to 256 bits from 64 bits to handle cases where a high channel index could cause the old limit to be exceeded.

diff --git a/Engine/Source/Runtime/Engine/Classes/Engine/NetConnection.h b/Engine/Source/Runtime/Engine/Classes/Engine/NetConnection.h
--- a/Engine/Source/Runtime/Engine/Classes/Engine/NetConnection.h
+++ b/Engine/Source/Runtime/Engine/Classes/Engine/NetConnection.h
@@ -40,7 +40,7 @@ typedef TMap<TWeakObjectPtr<AActor>, UActorChannel*, FDefaultSetAllocator, TWeak
 enum { RELIABLE_BUFFER = 256 }; // Power of 2 >= 1.
 enum { MAX_PACKETID = 16384 };  // Power of 2 >= 1, covering guaranteed loss/misorder time.
 enum { MAX_CHSEQUENCE = 1024 }; // Power of 2 >RELIABLE_BUFFER, covering loss/misorder time.
-enum { MAX_BUNCH_HEADER_BITS = 64 };
+enum { MAX_BUNCH_HEADER_BITS = 256 };
 enum { MAX_PACKET_HEADER_BITS = 15 }; // = FMath::CeilLogTwo(MAX_PACKETID) + 1 (IsAck)
 enum { MAX_PACKET_TRAILER_BITS = 1 };
 
