Acks are now sent as part of the packet header including a history of the last 32 last received packets.
Fixes issue with dropped acks being interpreted as packetloss even though packets are delivered.

diff --git a/Engine/Source/Runtime/Core/Private/Misc/NetworkVersion.cpp b/Engine/Source/Runtime/Core/Private/Misc/NetworkVersion.cpp
--- a/Engine/Source/Runtime/Core/Private/Misc/NetworkVersion.cpp
+++ b/Engine/Source/Runtime/Core/Private/Misc/NetworkVersion.cpp
@@ -18,7 +18,7 @@ FString FNetworkVersion::ProjectVersion = TEXT("1.0.0");
 bool FNetworkVersion::bHasCachedNetworkChecksum			= false;
 uint32 FNetworkVersion::CachedNetworkChecksum			= 0;
 
-uint32 FNetworkVersion::EngineNetworkProtocolVersion	= HISTORY_CHANNEL_NAMES;
+uint32 FNetworkVersion::EngineNetworkProtocolVersion	= HISTORY_ACKS_INCLUDED_IN_HEADER;
 uint32 FNetworkVersion::GameNetworkProtocolVersion		= 0;
 
 uint32 FNetworkVersion::EngineCompatibleNetworkProtocolVersion		= HISTORY_REPLAY_BACKWARDS_COMPAT;
diff --git a/Engine/Source/Runtime/Core/Public/Misc/NetworkVersion.h b/Engine/Source/Runtime/Core/Public/Misc/NetworkVersion.h
--- a/Engine/Source/Runtime/Core/Public/Misc/NetworkVersion.h
+++ b/Engine/Source/Runtime/Core/Public/Misc/NetworkVersion.h
@@ -38,6 +38,7 @@ enum EEngineNetworkVersionHistory
 	HISTORY_REPCMD_CHECKSUM_REMOVE_PRINTF = 4,		// Bump version since the way FRepLayoutCmd::CompatibleChecksum was calculated changed due to an optimization
 	HISTORY_NEW_ACTOR_OVERRIDE_LEVEL = 5,			// Bump version since a level reference was added to the new actor information
 	HISTORY_CHANNEL_NAMES = 6,						// Bump version since channel type is now an fname
+	HISTORY_ACKS_INCLUDED_IN_HEADER = 7,			// Bump version since acks are now sent as part of the header
 };
 
 struct CORE_API FNetworkVersion
diff --git a/Engine/Source/Runtime/Engine/Classes/Engine/ChildConnection.h b/Engine/Source/Runtime/Engine/Classes/Engine/ChildConnection.h
--- a/Engine/Source/Runtime/Engine/Classes/Engine/ChildConnection.h
+++ b/Engine/Source/Runtime/Engine/Classes/Engine/ChildConnection.h
@@ -52,10 +52,6 @@ public:
 		Parent->AssertValid();
 	}
 
-	virtual void SendAck(int32 PacketId, bool FirstTime=1) override
-	{
-	}
-
 	virtual void FlushNet(bool bIgnoreSimulation = false) override
 	{
 		Parent->FlushNet(bIgnoreSimulation);
diff --git a/Engine/Source/Runtime/Engine/Classes/Engine/NetConnection.h b/Engine/Source/Runtime/Engine/Classes/Engine/NetConnection.h
--- a/Engine/Source/Runtime/Engine/Classes/Engine/NetConnection.h
+++ b/Engine/Source/Runtime/Engine/Classes/Engine/NetConnection.h
@@ -14,6 +14,7 @@
 #include "GameFramework/OnlineReplStructs.h"
 #include "Engine/NetDriver.h"
 #include "Net/DataBunch.h"
+#include "Net/NetPacketNotify.h"
 #include "Engine/Player.h"
 #include "Engine/Channel.h"
 #include "ProfilingDebugging/Histogram.h"
@@ -38,10 +39,11 @@ typedef TMap<TWeakObjectPtr<AActor>, UActorChannel*, FDefaultSetAllocator, TWeak
 	Types.
 -----------------------------------------------------------------------------*/
 enum { RELIABLE_BUFFER = 256 }; // Power of 2 >= 1.
-enum { MAX_PACKETID = 16384 };  // Power of 2 >= 1, covering guaranteed loss/misorder time.
+enum { MAX_PACKETID = FNetPacketNotify::SequenceNumberT::SeqNumberCount };  // Power of 2 >= 1, covering guaranteed loss/misorder time.
 enum { MAX_CHSEQUENCE = 1024 }; // Power of 2 >RELIABLE_BUFFER, covering loss/misorder time.
 enum { MAX_BUNCH_HEADER_BITS = 256 };
-enum { MAX_PACKET_HEADER_BITS = 15 }; // = FMath::CeilLogTwo(MAX_PACKETID) + 1 (IsAck)
+enum { UE_PACKET_SEQUENCE_HISTORY_BITS = FNetPacketNotify::SequenceHistoryLength };
+enum { MAX_PACKET_HEADER_BITS = 81 }; // = [(Seq)16|(AckSeq)16|(SeqHistory)32|1|?(ServerTime)8|(ReceviedRateByte)8]
 enum { MAX_PACKET_TRAILER_BITS = 1 };
 
 // 
@@ -384,7 +386,7 @@ public:
 	FBitWriter		SendBuffer;						// Queued up bits waiting to send
 	double			OutLagTime[256];				// For lag measuring.
 	int32			OutLagPacketId[256];			// For lag measuring.
-	int32			OutBytesPerSecondHistory[256];	// For saturation measuring.
+	uint8			OutBytesPerSecondHistory[256];	// For saturation measuring.
 	float			RemoteSaturation;
 	int32			InPacketId;						// Full incoming packet index.
 	int32			OutPacketId;					// Most recently sent packet.
@@ -400,6 +402,7 @@ public:
 	TArray<int32>		OutReliable;
 	TArray<int32>		InReliable;
 	TArray<int32>		PendingOutRec;	// Outgoing reliable unacked data from previous (now destroyed) channel in this slot.  This contains the first chsequence not acked
+	UE_DEPRECATED(4.22, "QueuedAcks, ResendAcks variables will be removed.")
 	TArray<int32> QueuedAcks, ResendAcks;
 
 	int32				InitOutReliable;
@@ -688,6 +691,7 @@ public:
 	ENGINE_API virtual void AssertValid();
 
 	/** Send an acknowledgment. */
+	UE_DEPRECATED(4.22, "This method will be removed")
 	ENGINE_API virtual void SendAck( int32 PacketId, bool FirstTime=1);
 
 	/**
@@ -852,6 +856,7 @@ public:
 	// Functions.
 
 	/** Resend any pending acks. */
+	UE_DEPRECATED(4.22, "This method will be removed.")
 	void PurgeAcks();
 
 	/** Send package map to the remote. */
@@ -1049,6 +1054,18 @@ private:
 	/** Returns true if an outgoing packet should be dropped due to packet simulation settings, including loss burst simulation. */
 	bool ShouldDropOutgoingPacketForLossSimulation(int64 NumBits) const;
 
+	/** Write packetHeader */
+	void WritePacketHeader(FBitWriter& Writer) const;
+
+	/** Write extended packet header information (ServerFrameTime, SaturationData) */
+	void WritePacketInfo(FBitWriter& Writer) const;
+	
+	/** Read extended packet header information (ServerFrameTime, SaturationData) */
+	bool ReadPacketInfo(FBitReader& Reader);
+
+	/** Packet was acknowledged as delivered */
+	void ReceivedAck(int32 AckPacketId);
+
 	/**
 	 * on the server, the world the client has told us it has loaded
 	 * used to make sure the client has traveled correctly, prevent replicating actors before level transitions are done, etc
@@ -1067,6 +1084,15 @@ private:
 
 	/** Bookkeeping for connections using internal ack to avoid acking all open channels */
 	TArray<int32> ChannelsWaitingForInternalAck;
+
+	/** Sequence data used to implement reliability */
+	FNetPacketNotify PacketNotify;
+
+	/** Full PacketId  of last sent packet that we have received notification for (i.e. we know if it was delivered or not). Related to OutAckPacketId which is tha last successfully delivered PacketId */
+	int32 LastNotifiedPacketId;
+
+	/** Keep old behavior where we send a packet with only acks even if we have no other outgoing data if we got incoming data */
+	bool bHasDirtyAcks;
 };
 
 
diff --git a/Engine/Source/Runtime/Engine/Private/Net/NetPacketNotify.cpp b/Engine/Source/Runtime/Engine/Private/Net/NetPacketNotify.cpp
new file mode 100644
--- /dev/null
+++ b/Engine/Source/Runtime/Engine/Private/Net/NetPacketNotify.cpp
@@ -0,0 +1,62 @@
+// Copyright 1998-2018 Epic Games, Inc. All Rights Reserved.
+
+#include "Net/NetPacketNotify.h"
+
+void FNetPacketNotify::Init(SequenceNumberT InitialInSeq, SequenceNumberT InitialOutSeq)
+{
+	InSeqHistory.Reset();
+	InSeq = InitialInSeq;
+	InAckSeq = InitialInSeq;
+	OutSeq = InitialOutSeq;
+	OutAckSeq = SequenceNumberT(InitialOutSeq.Get() - 1);
+}
+
+void FNetPacketNotify::GetHeader(FNotificationHeader& DataOut) const
+{
+	DataOut.Seq = OutSeq;
+	DataOut.AckedSeq = InAckSeq;
+	DataOut.History = InSeqHistory;
+}
+
+void FNetPacketNotify::AckSeq(SequenceNumberT AckedSeq, bool IsAck)
+{
+	check( AckedSeq == InSeq);
+
+	while (AckedSeq > InAckSeq)
+	{
+		++InAckSeq;
+
+		const bool bReportAcked = InAckSeq == AckedSeq ? IsAck : false;
+
+		UE_LOG_PACKET_NOTIFY(TEXT("FNetPacketNotify::AckSeq - AckedSeq: %u, IsAck %u"), InAckSeq.Get(), bReportAcked ? 1u : 0u);
+
+		InSeqHistory.AddDeliveryStatus(bReportAcked);		
+	}
+}
+
+// These methods must always write and read the exact same number of bits, that is the reason for not using WriteInt/WrittedWrappedInt
+void FNetPacketNotify::WriteHeader(FBitWriter& Writer, const FNotificationHeader& Data)
+{
+	SequenceNumberT::SequenceT Seq = Data.Seq.Get();
+	SequenceNumberT::SequenceT AckedSeq = Data.AckedSeq.Get();
+	
+	Writer << Seq;
+	Writer << AckedSeq;
+	Data.History.Write(Writer);
+}
+
+bool FNetPacketNotify::ReadHeader(FNotificationHeader& Data, FBitReader& Reader)
+{
+	SequenceNumberT::SequenceT RcvdSeq;
+	SequenceNumberT::SequenceT RcvdAckedSeq;
+
+	Reader << RcvdSeq;
+	Reader << RcvdAckedSeq;
+	Data.History.Read(Reader);
+
+	Data.Seq = SequenceNumberT(RcvdSeq);
+	Data.AckedSeq = SequenceNumberT(RcvdAckedSeq);
+
+	return Reader.IsError() == false;
+}
+
diff --git a/Engine/Source/Runtime/Engine/Private/Net/Tests/NetPacketNotifyTests.cpp b/Engine/Source/Runtime/Engine/Private/Net/Tests/NetPacketNotifyTests.cpp
new file mode 100644
--- /dev/null
+++ b/Engine/Source/Runtime/Engine/Private/Net/Tests/NetPacketNotifyTests.cpp
@@ -0,0 +1,296 @@
+#include "Misc/AutomationTest.h"
+#include "Net/NetPacketNotify.h"
+
+#if WITH_DEV_AUTOMATION_TESTS
+
+IMPLEMENT_SIMPLE_AUTOMATION_TEST(FNetPacketNotifyTest, "Network.PacketNotifyTest", EAutomationTestFlags::ApplicationContextMask | EAutomationTestFlags::EngineFilter)
+
+struct FNetPacketNotifyTestUtil
+{
+	FNetPacketNotify DefaultNotify;
+
+	FNetPacketNotifyTestUtil()
+	{
+		DefaultNotify.Init(FNetPacketNotify::SequenceNumberT(-1),  FNetPacketNotify::SequenceNumberT(0));
+	}
+
+	// Pretend to receive incoming packet to generate ackdata
+	static int32 PretendReceiveSeq(FNetPacketNotify& PacketNotify, FNetPacketNotify::SequenceNumberT Seq, bool Ack = true)
+	{
+		FNetPacketNotify::FNotificationHeader Data;
+		Data.Seq = Seq;
+		Data.AckedSeq = PacketNotify.GetOutAckSeq();
+		Data.History = FNetPacketNotify::SequenceHistoryT(0);
+		
+		FNetPacketNotify::SequenceNumberT::DifferenceT SeqDelta = PacketNotify.Update(Data, [](FNetPacketNotify::SequenceNumberT AckedSequence, bool delivered) {});
+		if (SeqDelta > 0)
+		{
+			if (Ack)
+				PacketNotify.AckSeq(Seq);
+		}
+
+		return SeqDelta;
+	}
+
+	static bool VerifyNotificaitonState(const FNetPacketNotify& A, const FNetPacketNotify& B)
+	{
+		bool bEquals = 
+			A.GetInSeq() == B.GetInSeq() &&
+			A.GetInSeqHistory() == B.GetInSeqHistory() &&
+			A.GetOutSeq() == B.GetOutSeq() &&
+			A.GetOutAckSeq() == B.GetOutAckSeq();
+
+		return bEquals;
+	}
+
+	// pretend to ack array of sequence numbers
+	template<typename T>
+	static void PretendAckSequenceNumbers(FNetPacketNotify& PacketNotify, const T& InSequenceNumbers)
+	{
+		size_t SequenceNumberCount = sizeof(InSequenceNumbers) / sizeof(InSequenceNumbers[0]);
+
+		for (int32 I=0; I<SequenceNumberCount; ++I)
+		{
+			FNetPacketNotifyTestUtil::PretendReceiveSeq(PacketNotify, InSequenceNumbers[I]);
+		}
+	}
+	
+	// pretend to deliver notifications for incoming seq header
+	template<typename T>
+	static size_t PretendDeliverNotifications(FNetPacketNotify& PacketNotify, const FNetPacketNotify::FNotificationHeader Data, T& OutSequenceNumbers)
+	{
+		size_t NotificationCount = 0;
+
+		auto HandleAck = [&OutSequenceNumbers, &NotificationCount](FNetPacketNotify::SequenceNumberT Seq, bool delivered)
+		{
+			const size_t MaxSequenceNumberCount = sizeof(OutSequenceNumbers) / sizeof(OutSequenceNumbers[0]);
+
+			if (delivered)
+			{
+				if (NotificationCount < MaxSequenceNumberCount)
+				{
+					OutSequenceNumbers[NotificationCount] = Seq;
+				}
+				++NotificationCount;		
+			}
+		};
+		return PacketNotify.Update(Data, HandleAck);
+	}
+
+	template<typename T>
+	static bool TestNotificationSequence(const T& InSequenceNumbers, FNetPacketNotify::SequenceNumberT FirstSequence = 0)
+	{
+		T NotifiedSequenceNumbers = { 0 };
+		//FPlatformMemory::Memset(&NotifiedSequenceNumbers[0], 0, sizeof(NotifiedSequenceNumbers));
+
+		FNetPacketNotify Acked;
+		Acked.Init(FNetPacketNotify::SequenceNumberT(FirstSequence.Get() - 1), FirstSequence);
+		PretendAckSequenceNumbers(Acked, InSequenceNumbers);
+
+		// In order to be able to accepts the acks we must pretend that we have sent a packet which we will get an ack for.
+		FNetPacketNotify Notified;
+		Notified.Init(FNetPacketNotify::SequenceNumberT(FirstSequence.Get() - 1), FirstSequence);
+		FNetPacketNotify::FNotificationHeader Data;
+		Acked.GetHeader(Data);
+
+		PretendDeliverNotifications(Notified, Data, NotifiedSequenceNumbers);
+
+		return FPlatformMemory::Memcmp(&InSequenceNumbers[0], &NotifiedSequenceNumbers[0], sizeof(InSequenceNumbers)) == 0;
+	}
+};
+
+
+bool FNetPacketNotifyTest::RunTest(const FString& Parameters)
+{
+	FNetPacketNotifyTestUtil Util;
+	// Test fill
+	{
+		FNetPacketNotify::SequenceNumberT ExpectedInSeq(31);
+		FNetPacketNotify::SequenceHistoryT ExpectedInSeqHistory(0xffffffffu);
+
+		FNetPacketNotify Acks = Util.DefaultNotify;
+				
+		for (int32 I=0; I<32; ++I)
+		{
+			FNetPacketNotifyTestUtil::PretendReceiveSeq(Acks, I, true);
+		}
+
+		TestEqual(TEXT("Test fill - InSeq"), Acks.GetInSeq(), ExpectedInSeq);
+		TestEqual(TEXT("Test fill - History"), Acks.GetInSeqHistory(), ExpectedInSeqHistory);
+	}
+
+	// Test drop every other
+	{
+		FNetPacketNotify::SequenceNumberT ExpectedInSeq(30);
+		FNetPacketNotify::SequenceHistoryT ExpectedInSeqHistory(0x55555555u);
+
+		FNetPacketNotify Acks = Util.DefaultNotify;
+				
+		for (int32 I=0; I<16; ++I)
+		{
+			FNetPacketNotifyTestUtil::PretendReceiveSeq(Acks, I*2);
+		}
+		
+		TestEqual(TEXT("Test drop every other - InSeq"), Acks.GetInSeq(), ExpectedInSeq);
+		TestEqual(TEXT("Test drop every other - History"), Acks.GetInSeqHistory(), ExpectedInSeqHistory);
+	}
+
+	// Test burst drop
+	{
+		FNetPacketNotify::SequenceNumberT ExpectedInSeq(128);
+		FNetPacketNotify::SequenceHistoryT ExpectedInSeqHistory(0x1u);
+
+		FNetPacketNotify Acks = Util.DefaultNotify;
+
+		// Drop early
+		FNetPacketNotifyTestUtil::PretendReceiveSeq(Acks, 3);
+
+		// Large gap until next seq
+		FNetPacketNotifyTestUtil::PretendReceiveSeq(Acks, 128);
+
+		TestEqual(TEXT("Test burst drop - InSeq"), Acks.GetInSeq(), ExpectedInSeq);
+		TestEqual(TEXT("Test burst drop - History"), Acks.GetInSeqHistory(), ExpectedInSeqHistory);
+	}
+
+	// Test window overflow
+	{
+		FNetPacketNotify Acks = Util.DefaultNotify;
+
+		const FNetPacketNotify::SequenceNumberT ExpectedInSeq(0);
+		const FNetPacketNotify::SequenceNumberT MaxWindowSeq(FNetPacketNotify::SequenceNumberT::SeqNumberHalf);
+
+		FNetPacketNotifyTestUtil::PretendReceiveSeq(Acks, ExpectedInSeq, true);
+		TestEqual(TEXT("Test window overflow - Expect InSeq 0"), Acks.GetInSeq(), ExpectedInSeq);
+		TestEqual(TEXT("Test window overflow - Expect InAckSeq 0"), Acks.GetInAckSeq(), ExpectedInSeq);
+
+		FNetPacketNotifyTestUtil::PretendReceiveSeq(Acks, MaxWindowSeq, true);
+		TestEqual(TEXT("Test window overflow - Expect Seq reject"), Acks.GetInSeq(), ExpectedInSeq);
+		TestEqual(TEXT("Test window overflow - Expect Seq reject"), Acks.GetInAckSeq(), ExpectedInSeq);
+
+		const FNetPacketNotify::SequenceNumberT NextExpectedInSeq(1);
+		FNetPacketNotifyTestUtil::PretendReceiveSeq(Acks, NextExpectedInSeq, true);
+		TestEqual(TEXT("Test window overflow - Expect InSeq 1"), Acks.GetInSeq(), NextExpectedInSeq);
+		TestEqual(TEXT("Test window overflow - Expect InAckSeq 1"), Acks.GetInAckSeq(), NextExpectedInSeq);
+
+		FNetPacketNotifyTestUtil::PretendReceiveSeq(Acks, MaxWindowSeq, true);
+		TestEqual(TEXT("Test window overflow - Expect InSeq MaxWindowSeq"), Acks.GetInSeq(), MaxWindowSeq);
+		TestEqual(TEXT("Test window overflow - Expect InSeq MaxWindowSeq"), Acks.GetInAckSeq(), MaxWindowSeq);
+	}
+
+	// create history
+	{
+		const FNetPacketNotify::SequenceNumberT ExpectedInSeq(18);
+		const FNetPacketNotify::SequenceHistoryT ExpectedInSeqHistory(0x8853u);
+
+		const FNetPacketNotify::SequenceNumberT AckdPacketIds[] = {3, 7, 12, 14, 17, 18};
+		const size_t ExpectedCount = sizeof(AckdPacketIds)/sizeof((AckdPacketIds)[0]);		
+
+		FNetPacketNotify Acks = Util.DefaultNotify;
+		FNetPacketNotifyTestUtil::PretendAckSequenceNumbers(Acks, AckdPacketIds);
+ 
+		TestEqual(TEXT("Create history - InSeq"), Acks.GetInSeq(), ExpectedInSeq);
+		TestEqual(TEXT("Create history - InSeqHistory"), Acks.GetInSeqHistory(), ExpectedInSeqHistory);
+	}
+
+	
+	// test notifications
+	{
+		static const FNetPacketNotify::SequenceNumberT ExpectedAckdPacketIds[] = {3, 7, 12, 14, 17, 18};
+		static const size_t ExpectedCount = sizeof(ExpectedAckdPacketIds)/sizeof((ExpectedAckdPacketIds)[0]);		
+
+		FNetPacketNotify::SequenceNumberT RcvdAcks[ExpectedCount] = { 0 };
+		size_t RcvdCount = 0;
+
+		// Create src data
+		FNetPacketNotify Acks = Util.DefaultNotify;
+
+		// Fill in some data
+		FNetPacketNotify::FNotificationHeader Data;
+		Data.Seq = FNetPacketNotify::SequenceNumberT(0);
+		Data.AckedSeq = FNetPacketNotify::SequenceNumberT(18);
+		Data.History = FNetPacketNotify::SequenceHistoryT(0x8853u);
+
+		size_t DeltaSeq = FNetPacketNotifyTestUtil::PretendDeliverNotifications(Acks, Data, RcvdAcks);
+
+		TestEqual(TEXT("Notifications - Create sequence delta"), DeltaSeq, 1);
+		TestEqual(TEXT("Notifications - Create sequence"), FPlatformMemory::Memcmp(ExpectedAckdPacketIds, RcvdAcks, sizeof(ExpectedAckdPacketIds)), 0u);
+	}
+
+	// test various sequence
+	{
+		static const FNetPacketNotify::SequenceNumberT TestSeqs[] = {3, 7, 12, 14, 17, 18};	
+		TestTrue(TEXT("Test Seq {3, 7, 12, 14, 17, 18}"), FNetPacketNotifyTestUtil::TestNotificationSequence(TestSeqs));
+	}
+	{
+		static const FNetPacketNotify::SequenceNumberT TestSeqs[] = {0, 1, 2, 3, 4, 5, 6, 7, 8, 9, 10, 11, 12, 13, 14, 15, 16, 17, 18, 19, 20, 21, 22, 23, 24, 25, 26, 27, 28, 29, 30, 31};	
+		TestTrue(TEXT("Test Seq {0, 1, 2, 3, 4, 5, 6, 7, 8, 9, 10, 11, 12, 13, 14, 15, 16, 17, 18, 19, 20, 21, 22, 23, 24, 25, 26, 27, 28, 29, 30, 31};"), FNetPacketNotifyTestUtil::TestNotificationSequence(TestSeqs));
+	}
+	{
+		static const FNetPacketNotify::SequenceNumberT TestSeqs[] = {2, 3, 4, 5, 6, 7, 8, 9, 10, 12, 13, 14, 16, 17, 18, 19, 20, 21, 22, 23, 24, 25, 26, 27, 28, 29, 30, 31};	
+		TestTrue(TEXT("{2, 3, 4, 5, 6, 7, 8, 9, 10, 12, 13, 14, 16, 17, 18, 19, 20, 21, 22, 23, 24, 25, 26, 27, 28, 29, 30, 31}"), FNetPacketNotifyTestUtil::TestNotificationSequence(TestSeqs));
+	}
+	{
+		static const FNetPacketNotify::SequenceNumberT TestSeqs[] = {0, 31};	
+		TestTrue(TEXT("Test Seq {0, 31}"), FNetPacketNotifyTestUtil::TestNotificationSequence(TestSeqs));
+	}
+	{
+		static const FNetPacketNotify::SequenceNumberT TestSeqs[] = {0, 32};	
+		TestFalse(TEXT("Test Seq {0, 32}"), FNetPacketNotifyTestUtil::TestNotificationSequence(TestSeqs));
+	}
+	{
+		static const FNetPacketNotify::SequenceNumberT TestSeqs[] = {0, FNetPacketNotify::SequenceNumberT::SeqNumberHalf };
+		TestFalse(TEXT("Test Seq {FNetPacketNotify::SequenceNumberT::SeqNumberMax, 0};"), FNetPacketNotifyTestUtil::TestNotificationSequence(TestSeqs));
+	}
+	{
+		static const FNetPacketNotify::SequenceNumberT TestSeqs[] = {0, FNetPacketNotify::SequenceNumberT::SeqNumberHalf - 1};
+		TestFalse(TEXT("Test Seq {FNetPacketNotify::SequenceNumberT::SeqNumberMax, 0};"), FNetPacketNotifyTestUtil::TestNotificationSequence(TestSeqs));
+	}
+	{
+		static const FNetPacketNotify::SequenceNumberT TestSeqs[] = {FNetPacketNotify::SequenceNumberT::SeqNumberMax, 0};
+		TestFalse(TEXT("Test Seq {FNetPacketNotify::SequenceNumberT::SeqNumberMax, 0} From 0;"), FNetPacketNotifyTestUtil::TestNotificationSequence(TestSeqs));
+	}
+	{
+		static const FNetPacketNotify::SequenceNumberT TestSeqs[] = {FNetPacketNotify::SequenceNumberT::SeqNumberMax, 0};
+		TestTrue(TEXT("Test Seq {FNetPacketNotify::SequenceNumberT::SeqNumberMax, 0} From SeqNumberHalf + 1;"), FNetPacketNotifyTestUtil::TestNotificationSequence(TestSeqs, FNetPacketNotify::SequenceNumberT::SeqNumberHalf + 2));
+	}
+
+	// Test sequence numbers
+	{
+		// Valid sequence = 0-7, max distance between sequence number in order to determine order is half the sequence space (0-3)
+		typedef TSequenceNumber<3, uint16> FSequence3;
+
+		for (int32 I = 0; I < FSequence3::SeqNumberCount; ++I)
+		{
+			FSequence3 Seq(I);
+			FSequence3 Ref(I);
+			
+			for (int32 U = 0; U < FSequence3::SeqNumberCount; ++U)
+			{	
+				FSequence3::DifferenceT Diff = FSequence3::Diff(Seq, Ref);
+
+				if (U < FSequence3::SeqNumberHalf)
+				{
+					if (!(Diff == U))
+					{
+						TestTrue(TEXT("SequenceNumbers - Expected Diff "), Diff == U);
+					}
+				}
+				else
+				{
+					if (!(Diff == (U - FSequence3::SeqNumberCount)))
+					{
+						TestTrue(TEXT("SequenceNumbers - Expected Diff"), Diff == (U - FSequence3::SeqNumberCount));
+					}
+				}
+				++Seq;
+			}
+		}
+
+		check(!HasAnyErrors());
+
+	}
+
+	return true;
+}
+
+#endif // #if WITH_DEV_AUTOMATION_TESTS
\ No newline at end of file
diff --git a/Engine/Source/Runtime/Engine/Private/NetConnection.cpp b/Engine/Source/Runtime/Engine/Private/NetConnection.cpp
--- a/Engine/Source/Runtime/Engine/Private/NetConnection.cpp
+++ b/Engine/Source/Runtime/Engine/Private/NetConnection.cpp
@@ -19,6 +19,7 @@
 #include "Net/UnrealNetwork.h"
 #include "Net/NetworkProfiler.h"
 #include "Net/DataReplication.h"
+#include "Net/NetPacketNotify.h"
 #include "Engine/ActorChannel.h"
 #include "Engine/ChildConnection.h"
 #include "Engine/VoiceChannel.h"
@@ -26,7 +27,6 @@
 #include "Engine/PackageMapClient.h"
 #include "Engine/NetworkObjectList.h"
 #include "EncryptionComponent.h"
-
 #include "Net/PerfCountersHelpers.h"
 #include "GameDelegates.h"
 #include "Misc/PackageName.h"
@@ -133,6 +133,8 @@ UNetConnection::UNetConnection(const FObjectInitializer& ObjectInitializer)
 #endif
 ,	PlayerOnlinePlatformName( NAME_None )
 ,	ClientWorldPackageName( NAME_None )
+,	LastNotifiedPacketId( -1 )
+,	bHasDirtyAcks(false)
 {
 	MaxChannelSize = CVarMaxChannelSize.GetValueOnAnyThread();
 	if (MaxChannelSize <= 0)
@@ -145,6 +147,8 @@ UNetConnection::UNetConnection(const FObjectInitializer& ObjectInitializer)
 	OutReliable.AddDefaulted(MaxChannelSize);
 	InReliable.AddDefaulted(MaxChannelSize);
 	PendingOutRec.AddDefaulted(MaxChannelSize);
+
+	PacketNotify.Init(InPacketId, OutPacketId);
 }
 
 /**
@@ -338,6 +342,7 @@ void UNetConnection::InitSequence(int32 IncomingSequence, int32 OutgoingSequence
 		InPacketId = IncomingSequence - 1;
 		OutPacketId = OutgoingSequence;
 		OutAckPacketId = OutgoingSequence - 1;
+		LastNotifiedPacketId = OutAckPacketId;
 
 		// Initialize the reliable packet sequence (more useful/effective at preventing attacks)
 		InitInReliable = IncomingSequence & (MAX_CHSEQUENCE - 1);
@@ -346,6 +351,8 @@ void UNetConnection::InitSequence(int32 IncomingSequence, int32 OutgoingSequence
 		InReliable.Init(InitInReliable, InReliable.Num());
 		OutReliable.Init(InitOutReliable, OutReliable.Num());
 
+		PacketNotify.Init(InPacketId, OutPacketId);
+
 		UE_LOG(LogNet, Verbose, TEXT("InitSequence: IncomingSequence: %i, OutgoingSequence: %i, InitInReliable: %i, InitOutReliable: %i"), IncomingSequence, OutgoingSequence, InitInReliable, InitOutReliable);
 	}
 }
@@ -466,8 +473,6 @@ void UNetConnection::Serialize( FArchive& Ar )
 	{
 		Children.CountBytes(Ar);
 		ClientVisibleLevelNames.CountBytes(Ar);
-		QueuedAcks.CountBytes(Ar);
-		ResendAcks.CountBytes(Ar);
 		OpenChannels.CountBytes(Ar);
 		SentTemporaries.CountBytes(Ar);
 		ActorChannels.CountBytes(Ar);
@@ -1041,7 +1046,7 @@ void UNetConnection::FlushNet(bool bIgnoreSimulation)
 	TimeSensitive = 0;
 
 	// If there is any pending data to send, send it.
-	if (SendBuffer.GetNumBits() || ( Driver->Time-LastSendTime > Driver->KeepAliveTime && !InternalAck && State != USOCK_Closed))
+	if (SendBuffer.GetNumBits() || bHasDirtyAcks || ( Driver->Time-LastSendTime > Driver->KeepAliveTime && !InternalAck && State != USOCK_Closed))
 	{
 		// Due to the PacketHandler handshake code, servers must never send the client data,
 		// before first receiving a client control packet (which is taken as an indication of a complete handshake).
@@ -1056,10 +1061,10 @@ void UNetConnection::FlushNet(bool bIgnoreSimulation)
 
 		FOutPacketTraits Traits;
 
-		// If sending keepalive packet, still write the packet id
+		// If sending keepalive packet or just acks, still write the packet header
 		if (SendBuffer.GetNumBits() == 0)
 		{
-			WriteBitsToSendBuffer( NULL, 0 );		// This will force the packet id to be written
+			WriteBitsToSendBuffer( NULL, 0 );		// This will force the packet header to be written
 
 			Traits.bIsKeepAlive = true;
 			AnalyticsVars.OutKeepAliveCount++;
@@ -1076,6 +1081,12 @@ void UNetConnection::FlushNet(bool bIgnoreSimulation)
 		// Write the UNetConnection-level termination bit
 		SendBuffer.WriteBit(1);
 
+		// Refresh outgoing header with latest data
+		if ( !InternalAck )
+		{
+			WritePacketHeader(SendBuffer);
+		}
+
 		ValidateSendBuffer();
 
 		const int32 NumStrayBits = SendBuffer.GetNumBits();
@@ -1155,9 +1166,14 @@ void UNetConnection::FlushNet(bool bIgnoreSimulation)
 		// Remember the actual time this packet was sent out, so we can compute ping when the ack comes back
 		OutLagPacketId[Index]			= OutPacketId;
 		OutLagTime[Index]				= FPlatformTime::Seconds();
-		OutBytesPerSecondHistory[Index]	= OutBytesPerSecond / 1024;
+	
+		OutBytesPerSecondHistory[Index]	= FMath::Min(OutBytesPerSecond / 1024, 255);
+
+		// Increase outgoing sequence number
+		PacketNotify.IncrementOutSeq();
+		bHasDirtyAcks = false;
+		++OutPacketId; 
 
-		OutPacketId++;
 		++OutPackets;
 		++OutTotalPackets;
 		Driver->OutPackets++;
@@ -1183,13 +1199,6 @@ void UNetConnection::FlushNet(bool bIgnoreSimulation)
 
 		InitSendBuffer();
 	}
-
-	// Move acks around.
-	for (int32 i=0; i<QueuedAcks.Num(); i++)
-	{
-		ResendAcks.Add(QueuedAcks[i]);
-	}
-	QueuedAcks.Empty(32);
 }
 
 bool UNetConnection::ShouldDropOutgoingPacketForLossSimulation(int64 NumBits) const
@@ -1217,8 +1226,64 @@ int32 UNetConnection::IsNetReady( bool Saturate )
 void UNetConnection::ReadInput( float DeltaSeconds )
 {}
 
+void UNetConnection::ReceivedAck(int32 AckPacketId)
+{
+	UE_LOG(LogNetTraffic, Verbose, TEXT("   Received ack %i"), AckPacketId);
+
+	// Advance OutAckPacketId
+	OutAckPacketId = AckPacketId;
+
+	// Process the bunch.
+	LastRecvAckTime = Driver->Time;
+
+	if ( PackageMap != NULL )
+	{
+		PackageMap->ReceivedAck( AckPacketId );
+	}
+
+	// Forward the ack to the channel.
+	{
+		SCOPE_CYCLE_COUNTER(Stat_NetConnectionReceivedAcks);
+	
+		for (int32 i = OpenChannels.Num() - 1; i >= 0; i--)
+		{
+			UChannel* const Channel = OpenChannels[i];
+				
+			if (Channel)
+			{
+				if (Channel->OpenPacketId.Last == AckPacketId) // Necessary for unreliable "bNetTemporary" channels.
+				{
+					Channel->OpenAcked = 1;
+				}
+				
+				for (FOutBunch* OutBunch = Channel->OutRec; OutBunch; OutBunch = OutBunch->Next)
+				{
+					if (OutBunch->bOpen)
+					{
+						UE_LOG(LogNet, VeryVerbose, TEXT("Channel %i reset Ackd because open is reliable. "), Channel->ChIndex );
+						Channel->OpenAcked  = 0; // We have a reliable open bunch, don't let the above code set the OpenAcked state,
+												// it must be set in UChannel::ReceivedAcks to verify all open bunches were received.
+					}
+
+					if (OutBunch->PacketId == AckPacketId)
+					{
+						OutBunch->ReceivedAck = 1;
+					}
+				}
+				Channel->ReceivedAcks(); //warning: May destroy Channel.
+			}
+			else
+			{
+				UE_LOG(LogNet, Warning, TEXT("UNetConnection::ReceivedPacket: null channel in OpenChannels array while processing ack. %s"), *Describe());
+			}
+		}
+	}
+}
+
 void UNetConnection::ReceivedNak( int32 NakPacketId )
 {
+	UE_LOG(LogNetTraffic, Verbose, TEXT("   Received nak %i"), NakPacketId);
+
 	SCOPE_CYCLE_COUNTER(Stat_NetConnectionReceivedNak);
 
 	// Update pending NetGUIDs
@@ -1241,6 +1306,137 @@ void UNetConnection::ReceivedNak( int32 NakPacketId )
 			UE_LOG(LogNet, Warning, TEXT("UNetConnection::ReceivedNak: null channel in OpenChannels array. %s"), *Describe());
 		}
 	}
+
+	// Stats
+	++OutPacketsLost;
+	++OutTotalPacketsLost;
+	++Driver->OutTotalPacketsLost;
+}
+
+// IMPORTANT:
+// WritePacketHeader must ALWAYS write the exact same number of bits as we go back and rewrite the header
+// right before we put the packet on the wire.
+void UNetConnection::WritePacketHeader(FBitWriter& Writer) const
+{
+	// If this is a header refresh, we only serialize the updated serial number information
+	const bool bIsHeaderUpdate = Writer.GetNumBits() > 0u;
+
+	// Header is always written first in the packet
+	FBitWriterMark Reset;
+	FBitWriterMark Restore(Writer);
+	Reset.PopWithoutClear(Writer);
+	
+	FNetPacketNotify::FNotificationHeader Header;
+	PacketNotify.GetHeader(Header);
+	FNetPacketNotify::WriteHeader(Writer, Header);
+
+	// Jump back to where we came from.
+	if (bIsHeaderUpdate)
+	{
+		Restore.PopWithoutClear(Writer);
+	}
+}
+
+namespace
+{
+
+static bool ReadPacketHeader(FNetPacketNotify::FNotificationHeader& Header, FBitReader& Reader)
+{
+	FNetPacketNotify::ReadHeader(Header, Reader);
+		
+	return Reader.IsError() == false;
+}
+
+}
+
+void UNetConnection::WritePacketInfo(FBitWriter& Writer) const
+{
+	const uint8 bHasServerFrameTime = Driver->IsServer() ? bLastHasServerFrameTime : ( CVarPingExcludeFrameTime.GetValueOnGameThread() > 0 ? 1u : 0u );
+
+	// Write data used to calculate link latency
+	Writer.WriteBit(bHasServerFrameTime);
+	if (bHasServerFrameTime && Driver->IsServer())
+	{
+		uint8 FrameTimeByte = FMath::Min( FMath::FloorToInt( FrameTime * 1000 ), 255 );
+		Writer << FrameTimeByte;
+	}
+
+	// Notify server of our current rate per second at this time
+	uint8 InKBytesPerSecondByte = FMath::Min(InBytesPerSecond / 1024, 255);
+	Writer << InKBytesPerSecondByte;
+}
+
+bool UNetConnection::ReadPacketInfo(FBitReader& Reader)
+{
+	const bool bHasServerFrameTime = Reader.ReadBit();
+	double ServerFrameTime = 0;
+
+	if ( !Driver->IsServer() )
+	{
+		if ( bHasServerFrameTime )
+		{
+			uint8 FrameTimeByte	= 0;
+			Reader << FrameTimeByte;
+			// As a client, our request was granted, read the frame time
+			ServerFrameTime = ( double )FrameTimeByte / 1000;
+		}
+	}
+	else
+	{
+		bLastHasServerFrameTime = bHasServerFrameTime;
+	}
+
+	// limit to known size to know the size of the packet header
+	uint8 RemoteInKBytesPerSecondByte = 0;
+	Reader << RemoteInKBytesPerSecondByte;
+
+	if ( Reader.IsError() )
+	{
+		return false;
+	}
+
+	// Update ping
+	// At this time we have updated OutAckPacketId to the latest received ack.
+	const int32 Index = OutAckPacketId & (ARRAY_COUNT(OutLagPacketId)-1);
+
+	if ( OutLagPacketId[Index] == OutAckPacketId )
+	{
+		OutLagPacketId[Index] = -1;		// Only use the ack once
+
+#if !UE_BUILD_SHIPPING
+		if ( CVarPingDisplayServerTime.GetValueOnAnyThread() > 0 )
+		{
+			UE_LOG( LogNetTraffic, Warning, TEXT( "ServerFrameTime: %2.2f" ), ServerFrameTime * 1000.0f );
+		}
+#endif
+
+		// use FApp's time because it is set closer to the beginning of the frame - we don't care about the time so far of the current frame to process the packet
+		const double CurrentTime = FApp::GetCurrentTime();
+		const double GameTime	 = ServerFrameTime;
+		const double RTT		 = (CurrentTime - OutLagTime[Index] ) - ( CVarPingExcludeFrameTime.GetValueOnAnyThread() ? GameTime : 0.0 );
+		const double NewLag		 = FMath::Max( RTT, 0.0 );
+
+		if ( OutBytesPerSecondHistory[Index] > 0 )
+		{
+			RemoteSaturation = ( 1.0f - FMath::Min( ( float )RemoteInKBytesPerSecondByte / ( float )OutBytesPerSecondHistory[Index], 1.0f ) ) * 100.0f;
+		}
+		else
+		{
+			RemoteSaturation = 0.0f;
+		}
+
+		//UE_LOG( LogNet, Warning, TEXT( "Out: %i, InRemote: %i, Saturation: %f" ), OutBytesPerSecondHistory[Index], RemoteInKBytesPerSecond, RemoteSaturation );
+
+		LagAcc += NewLag;
+		LagCount++;
+
+		if (PlayerController != NULL)
+		{
+			PlayerController->UpdatePing(NewLag);
+		}
+	}
+
+	return true;
 }
 
 void UNetConnection::ReceivedPacket( FBitReader& Reader )
@@ -1265,201 +1461,107 @@ void UNetConnection::ReceivedPacket( FBitReader& Reader )
 	LastReceiveTime		= Driver->Time;
 	LastReceiveRealtime = FPlatformTime::Seconds();
 
-	// Check packet ordering.
-	const int32 PacketId = InternalAck ? InPacketId + 1 : MakeRelative(Reader.ReadInt(MAX_PACKETID),InPacketId,MAX_PACKETID);
-	if( PacketId > InPacketId )
+	if (InternalAck)
 	{
-		const int32 PacketsLost = PacketId - InPacketId - 1;
-		
-		if ( PacketsLost > 10 )
-		{
-			UE_LOG( LogNetTraffic, Log, TEXT( "High single frame packet loss. PacketsLost: %i %s" ), PacketsLost, *Describe() );
-		}
-
-		InPacketsLost += PacketsLost;
-		InTotalPacketsLost += PacketsLost;
-		Driver->InPacketsLost += PacketsLost;
-		Driver->InTotalPacketsLost += PacketsLost;
-		InPacketId = PacketId;
-	}
+		++InPacketId;
+	}	
 	else
 	{
-		Driver->InOutOfOrderPackets++;
-		// Protect against replay attacks
-		// We already protect against this for reliable bunches, and unreliable properties
-		// The only bunch we would process would be unreliable RPC's, which could allow for replay attacks
-		// So rather than add individual protection for unreliable RPC's as well, just kill it at the source, 
-		// which protects everything in one fell swoop
-		return;
-	}
-
-	const bool bIgnoreRPCs = Driver->ShouldIgnoreRPCs();
-
-	bool bSkipAck = false;
-
-	// Track channels that were rejected while processing this packet - used to avoid sending multiple close-channel bunches,
-	// which would cause a disconnect serverside
-	TArray<int32> RejectedChans;
-
-	// Disassemble and dispatch all bunches in the packet.
-	while( !Reader.AtEnd() && State!=USOCK_Closed )
-	{
-		// Parse the bunch.
-		int32 StartPos = Reader.GetPosBits();
-		bool IsAck = !!Reader.ReadBit();
-		if ( Reader.IsError() )
+		// Read packet header
+		FNetPacketNotify::FNotificationHeader Header;
+		if (!ReadPacketHeader(Header, Reader))
 		{
-			CLOSE_CONNECTION_DUE_TO_SECURITY_VIOLATION(this, ESecurityEvent::Malformed_Packet, TEXT("Bunch missing ack flag"));
+			CLOSE_CONNECTION_DUE_TO_SECURITY_VIOLATION(this, ESecurityEvent::Malformed_Packet, TEXT("Failed to read PacketHeader"));
 			return;
 		}
 
-		// Process the bunch.
-		if( IsAck )
+		// Lambda to dispatch delivery notifications, 
+		auto HandlePacketNotification = [&Header, this](FNetPacketNotify::SequenceNumberT AckedSequence, bool bDelivered)
 		{
-			LastRecvAckTime = Driver->Time;
-
-			// This is an acknowledgment.
-			const int32 AckPacketId = MakeRelative(Reader.ReadInt(MAX_PACKETID),OutAckPacketId,MAX_PACKETID);
+			// Increase LastNotifiedPacketId, this is a full packet Id
+			++LastNotifiedPacketId;
 
-			if( Reader.IsError() )
+			// Sanity check
+			if (FNetPacketNotify::SequenceNumberT(LastNotifiedPacketId) != AckedSequence)
 			{
-				CLOSE_CONNECTION_DUE_TO_SECURITY_VIOLATION(this, ESecurityEvent::Malformed_Packet, TEXT("Bunch missing ack"));
+				CLOSE_CONNECTION_DUE_TO_SECURITY_VIOLATION(this, ESecurityEvent::Malformed_Packet, TEXT("LastNotifiedPacketId != AckedSequence"));
 				return;
 			}
 
-			double ServerFrameTime = 0;
-
-			// If this is the server, we're reading in the request to send them our frame time
-			// If this is the client, we're reading in confirmation that our request to get frame time from server is granted
-			const bool bHasServerFrameTime = !!Reader.ReadBit();
-
-			if ( !Driver->IsServer() )
+			if (bDelivered)
 			{
-				if ( bHasServerFrameTime )
-				{
-					// As a client, our request was granted, read the frame time
-					uint8 FrameTimeByte	= 0;
-					Reader << FrameTimeByte;
-					ServerFrameTime = ( double )FrameTimeByte / 1000;
-				}
+				ReceivedAck(LastNotifiedPacketId);
 			}
 			else
 			{
-				// Server remembers so he can use during SendAck to notify to client of his frame time
-				bLastHasServerFrameTime = bHasServerFrameTime;
-			}
-
-			uint32 RemoteInKBytesPerSecond = 0;
-			Reader.SerializeIntPacked( RemoteInKBytesPerSecond );
+				ReceivedNak(LastNotifiedPacketId);
+			};
+		};
 
-			// Resend any old reliable packets that the receiver hasn't acknowledged.
-			if( AckPacketId>OutAckPacketId )
-			{
-				for (int32 NakPacketId = OutAckPacketId + 1; NakPacketId<AckPacketId; NakPacketId++, OutPacketsLost++, OutTotalPacketsLost++, Driver->OutTotalPacketsLost++)
-				{
-					UE_LOG(LogNetTraffic, Verbose, TEXT("   Received virtual nak %i (%.1f)"), NakPacketId, (Reader.GetPosBits()-StartPos)/8.f );
-					ReceivedNak( NakPacketId );
-				}
-				OutAckPacketId = AckPacketId;
-			}
-			else if( AckPacketId<OutAckPacketId )
+		// Update incoming sequence data and deliver packet notifications
+		// Packet is only accepted if both the incoming sequence number and incoming ack data are valid
+		int PacketSequenceDelta = PacketNotify.Update(Header, HandlePacketNotification);
+		if (PacketSequenceDelta > 0)
+		{
+			const int32 PacketsLost = PacketSequenceDelta - 1;
+		
+			if ( PacketsLost > 10 )
 			{
-				//warning: Double-ack logic makes this unmeasurable.
-				//OutOrdAcc++;
+				UE_LOG( LogNetTraffic, Verbose, TEXT( "High single frame packet loss. PacketsLost: %i %s" ), PacketsLost, *Describe() );
 			}
 
-			// Update ping
-			const int32 Index = AckPacketId & (ARRAY_COUNT(OutLagPacketId)-1);
+			InPacketsLost += PacketsLost;
+			InTotalPacketsLost += PacketsLost;
+			Driver->InPacketsLost += PacketsLost;
+			Driver->InTotalPacketsLost += PacketsLost;
+			InPacketId += PacketSequenceDelta;
 
-			if ( OutLagPacketId[Index] == AckPacketId )
+			// Extra information associated with the header
+			if (!ReadPacketInfo(Reader))
 			{
-				OutLagPacketId[Index] = -1;		// Only use the ack once
-
-#if !UE_BUILD_SHIPPING
-				if ( CVarPingDisplayServerTime.GetValueOnAnyThread() > 0 )
-				{
-					UE_LOG( LogNetTraffic, Warning, TEXT( "ServerFrameTime: %2.2f" ), ServerFrameTime * 1000.0f );
-				}
-#endif
-
-				// use FApp's time because it is set closer to the beginning of the frame - we don't care about the time so far of the current frame to process the packet
-				const double CurrentTime = FApp::GetCurrentTime();
-				const double GameTime	 = ServerFrameTime;
-				const double RTT		 = (CurrentTime - OutLagTime[Index] ) - ( CVarPingExcludeFrameTime.GetValueOnAnyThread() ? GameTime : 0.0 );
-				const double NewLag		 = FMath::Max( RTT, 0.0 );
-
-				if ( OutBytesPerSecondHistory[Index] > 0 )
-				{
-					RemoteSaturation = ( 1.0f - FMath::Min( ( float )RemoteInKBytesPerSecond / ( float )OutBytesPerSecondHistory[Index], 1.0f ) ) * 100.0f;
-				}
-				else
-				{
-					RemoteSaturation = 0.0f;
-				}
-
-				//UE_LOG( LogNet, Warning, TEXT( "Out: %i, InRemote: %i, Saturation: %f" ), OutBytesPerSecondHistory[Index], RemoteInKBytesPerSecond, RemoteSaturation );
-
-				LagAcc += NewLag;
-				LagCount++;
-
-				if (PlayerController != NULL)
-				{
-					PlayerController->UpdatePing(NewLag);
-				}
+				CLOSE_CONNECTION_DUE_TO_SECURITY_VIOLATION(this, ESecurityEvent::Malformed_Packet, TEXT("Failed to read PacketHeader"));
+				return;
 			}
+		}
+		else
+		{
+			Driver->InOutOfOrderPackets++;
+			// Protect against replay attacks
+			// We already protect against this for reliable bunches, and unreliable properties
+			// The only bunch we would process would be unreliable RPC's, which could allow for replay attacks
+			// So rather than add individual protection for unreliable RPC's as well, just kill it at the source, 
+			// which protects everything in one fell swoop
+			return;
+		}
+	}
 
-			if ( PackageMap != NULL )
-			{
-				PackageMap->ReceivedAck( AckPacketId );
-			}
+	const bool bIgnoreRPCs = Driver->ShouldIgnoreRPCs();
 
-			// Forward the ack to the channel.
-			UE_LOG(LogNetTraffic, Verbose, TEXT("   Received ack %i (%.1f)"), AckPacketId, (Reader.GetPosBits()-StartPos)/8.f );
+	bool bSkipAck = false;
 
-			{
-				SCOPE_CYCLE_COUNTER(Stat_NetConnectionReceivedAcks);
-	
-				for (int32 i = OpenChannels.Num() - 1; i >= 0; i--)
-				{
-					UChannel* const Channel = OpenChannels[i];
-				
-					if (Channel)
-					{
-						if (Channel->OpenPacketId.Last == AckPacketId) // Necessary for unreliable "bNetTemporary" channels.
-						{
-							Channel->OpenAcked = 1;
-						}
-				
-						for (FOutBunch* OutBunch = Channel->OutRec; OutBunch; OutBunch = OutBunch->Next)
-						{
-							if (OutBunch->bOpen)
-							{
-								UE_LOG(LogNet, VeryVerbose, TEXT("Channel %i reset Ackd because open is reliable. "), Channel->ChIndex );
-								Channel->OpenAcked  = 0; // We have a reliable open bunch, don't let the above code set the OpenAcked state,
-													 // it must be set in UChannel::ReceivedAcks to verify all open bunches were received.
-							}
+	// Track channels that were rejected while processing this packet - used to avoid sending multiple close-channel bunches,
+	// which would cause a disconnect serverside
+	TArray<int32> RejectedChans;
 
-							if (OutBunch->PacketId == AckPacketId)
-							{
-								OutBunch->ReceivedAck = 1;
-							}
-						}
-						Channel->ReceivedAcks(); //warning: May destroy Channel.
-					}
-					else
-					{
-						UE_LOG(LogNet, Warning, TEXT("UNetConnection::ReceivedPacket: null channel in OpenChannels array while processing ack. %s"), *Describe());
-					}
-				}
-			}
+	// Disassemble and dispatch all bunches in the packet.
+	while( !Reader.AtEnd() && State!=USOCK_Closed )
+	{
+		// For demo backwards compatibility, old replays still have this bit
+		if (InternalAck && EngineNetworkProtocolVersion < EEngineNetworkVersionHistory::HISTORY_ACKS_INCLUDED_IN_HEADER)
+		{
+			const bool IsAckDummy = Reader.ReadBit();
 		}
-		else
+
+		// Parse the bunch.
+		int32 StartPos = Reader.GetPosBits();
+		
+		// Process Received data
 		{
 			// Parse the incoming data.
 			FInBunch Bunch( this );
 			int32 IncomingStartPos		= Reader.GetPosBits();
 			uint8 bControl				= Reader.ReadBit();
-			Bunch.PacketId				= PacketId;
+			Bunch.PacketId				= InPacketId;
 			Bunch.bOpen					= bControl ? Reader.ReadBit() : 0;
 			Bunch.bClose				= bControl ? Reader.ReadBit() : 0;
 			Bunch.bDormant				= Bunch.bClose ? Reader.ReadBit() : 0;
@@ -1505,7 +1607,7 @@ void UNetConnection::ReceivedPacket( FBitReader& Reader )
 			else if ( Bunch.bPartial )
 			{
 				// If this is an unreliable partial bunch, we simply use packet sequence since we already have it
-				Bunch.ChSequence = PacketId;
+				Bunch.ChSequence = InPacketId;
 			}
 			else
 			{
@@ -1833,8 +1935,24 @@ void UNetConnection::ReceivedPacket( FBitReader& Reader )
 	if ( !bSkipAck )
 	{
 		LastGoodPacketRealtime = FPlatformTime::Seconds();
+	}
 
-		SendAck(PacketId, true);
+	if( !InternalAck )
+	{
+		// We always call AckSequence even if we are explicitly rejecting the packet as this updates the expected InSeq used to drive future acks.
+		if ( bSkipAck )
+		{
+			// Explicit Nak, we treat this packet as dropped but we still report it to the sending side as quickly as possible
+			PacketNotify.NakSeq( InPacketId );
+		}
+		else
+		{
+			PacketNotify.AckSeq( InPacketId );
+		}
+
+		// We do want to let the other side know about the ack, so even if there are no other outgoing data when we tick the connection we will send an ackpacket.
+		TimeSensitive = 1;
+		bHasDirtyAcks = 1;
 	}
 }
 
@@ -1878,10 +1996,18 @@ int32 UNetConnection::WriteBitsToSendBuffer(
 	// If this is the start of the queue, make sure to add the packet id
 	if ( SendBuffer.GetNumBits() == 0 && !InternalAck )
 	{
-		SendBuffer.WriteIntWrapped( OutPacketId, MAX_PACKETID );
-		ValidateSendBuffer();
+		// Write Packet Header, before sending the packet we will go back and rewrite the data
+		WritePacketHeader(SendBuffer);
 
-		NumPacketIdBits += SendBuffer.GetNumBits();
+		// Also write server RTT and Recived rate
+		WritePacketInfo(SendBuffer);
+	
+		// Update stats for PacketIdBits and ackdata (also including the data used for packet RTT and saturation calculations)
+		int64 BitsWritten = SendBuffer.GetNumBits();
+		NumPacketIdBits += FNetPacketNotify::SequenceNumberT::SeqNumberBits;
+		NumAckBits += BitsWritten - FNetPacketNotify::SequenceNumberT::SeqNumberBits;
+
+		ValidateSendBuffer();
 	}
 
 	// Add the bits to the queue
@@ -1905,9 +2031,6 @@ int32 UNetConnection::WriteBitsToSendBuffer(
 		case EWriteBitsDataType::Bunch:
 			NumBunchBits += SizeInBits + ExtraSizeInBits;
 			break;
-		case EWriteBitsDataType::Ack:
-			NumAckBits += SizeInBits + ExtraSizeInBits;
-			break;
 		default:
 			break;
 	}
@@ -1955,64 +2078,12 @@ TSharedPtr<FObjectReplicator> UNetConnection::CreateReplicatorForNewActorChannel
 
 void UNetConnection::PurgeAcks()
 {
-	for ( int32 i = 0; i < ResendAcks.Num(); i++ )
-	{
-		SendAck( ResendAcks[i], 0 );
-	}
-
-	ResendAcks.Empty(32);
 }
 
-
 void UNetConnection::SendAck(int32 AckPacketId, bool FirstTime/*=1*/)
 {
-	SCOPE_CYCLE_COUNTER(Stat_NetConnectionSendAck);
-
-	ValidateSendBuffer();
-
-	if( !InternalAck )
-	{
-		if( FirstTime )
-		{
-			PurgeAcks();
-			QueuedAcks.Add(AckPacketId);
-		}
-
-		FBitWriter AckData( 32, true );
-
-		AckData.WriteBit( 1 );
-		AckData.WriteIntWrapped(AckPacketId, MAX_PACKETID);
-		
-		const bool bHasServerFrameTime = Driver->IsServer() ? bLastHasServerFrameTime : ( CVarPingExcludeFrameTime.GetValueOnGameThread() > 0 ? true : false );
-
-		AckData.WriteBit( bHasServerFrameTime ? 1 : 0 );
-
-		if ( Driver->IsServer() && bHasServerFrameTime )
-		{
-			uint8 FrameTimeByte = FMath::Min( FMath::FloorToInt( FrameTime * 1000 ), 255 );
-			AckData << FrameTimeByte;
-		}
-
-		// Notify server of our current rate per second at this time
-		uint32 InKBytesPerSecond = InBytesPerSecond / 1024;
-		AckData.SerializeIntPacked( InKBytesPerSecond );
-
-		NETWORK_PROFILER( GNetworkProfiler.TrackSendAck( AckData.GetNumBits(), this ) );
-
-		WriteBitsToSendBuffer( AckData.GetData(), AckData.GetNumBits(), nullptr, 0, EWriteBitsDataType::Ack );
-
-		AllowMerge = false;
-
-		TimeSensitive = 1;
-
-		UE_LOG(LogNetTraffic, Log, TEXT("   Send ack %i"), AckPacketId);
-	}
-
-	++OutTotalAcks;
-	++Driver->OutTotalAcks;
 }
 
-
 int32 UNetConnection::SendRawBunch( FOutBunch& Bunch, bool InAllowMerge )
 {
 	ValidateSendBuffer();
@@ -2024,7 +2095,7 @@ int32 UNetConnection::SendRawBunch( FOutBunch& Bunch, bool InAllowMerge )
 
 	// Build header.
 	SendBunchHeader.Reset();
-	SendBunchHeader.WriteBit( 0 );
+
 	SendBunchHeader.WriteBit( Bunch.bOpen || Bunch.bClose );
 	if( Bunch.bOpen || Bunch.bClose )
 	{
@@ -2464,8 +2535,6 @@ void UNetConnection::Tick()
 	}
 
 	// Flush.
-	PurgeAcks();
-
 	if ( TimeSensitive || (Driver->Time - LastSendTime) > Driver->KeepAliveTime)
 	{
 		bool bHandlerHandshakeComplete = !Handler.IsValid() || Handler->IsFullyInitialized();
diff --git a/Engine/Source/Runtime/Engine/Public/Net/NetPacketNotify.h b/Engine/Source/Runtime/Engine/Public/Net/NetPacketNotify.h
new file mode 100644
--- /dev/null
+++ b/Engine/Source/Runtime/Engine/Public/Net/NetPacketNotify.h
@@ -0,0 +1,350 @@
+// Copyright 1998-2018 Epic Games, Inc. All Rights Reserved.
+
+#pragma once
+
+#include "HAL/UnrealMemory.h"
+#include "Templates/IsSigned.h"
+#include "Serialization/BitWriter.h"
+#include "Serialization/BitReader.h"
+
+#if (UE_BUILD_SHIPPING || UE_BUILD_TEST)
+#	define UE_NET_ENABLE_PACKET_NOTIFY_LOG 0
+#else
+#	define UE_NET_ENABLE_PACKET_NOTIFY_LOG 1
+#endif 
+
+#if UE_NET_ENABLE_PACKET_NOTIFY_LOG
+#	include "Logging/LogMacros.h"
+#	include "EngineLogs.h"
+#	define UE_LOG_PACKET_NOTIFY(Format, ...)  UE_LOG(LogNetTraffic, Verbose, Format, ##__VA_ARGS__)
+#else
+#	define UE_LOG_PACKET_NOTIFY(...)
+#endif
+
+struct FBitWriter;
+struct FBitReader;
+
+/** Util class to manage history of received sequence numbers */
+template <size_t HistorySize>
+class TSequenceHistory
+{
+	typedef uint32 WordT;
+	
+	static const size_t BitsPerWord = sizeof(WordT) * 8;
+	static const size_t WordCount = HistorySize / BitsPerWord;
+	
+	static_assert(HistorySize > 0, "Invalid History");
+	static_assert(HistorySize % BitsPerWord == 0, "HistorySize must be a modulo of the wordsize");
+
+public:	
+	TSequenceHistory();
+
+#if WITH_DEV_AUTOMATION_TESTS
+	explicit TSequenceHistory(WordT Value);
+#endif
+	/** Reset */
+	void Reset();
+
+	/** Store delivery status, oldest will be dropped */
+	void AddDeliveryStatus(bool Delivered);
+
+	/** Query the status of a specific index, index 0 is last stored status */
+	bool IsDelivered(size_t Index) const;
+
+	/** Return the size of the history buffer */
+	size_t Size() const { return HistorySize; }	
+	
+	bool operator==(const TSequenceHistory& Other) const
+	{
+		return FMemory::Memcmp(Data(), Other.Data(), WordCount * sizeof (WordT)) == 0;
+	}
+
+	bool operator!=(const TSequenceHistory& Other) const
+	{
+		return FMemory::Memcmp(Data(), Other.Data(), WordCount * sizeof (WordT)) != 0;
+	}
+
+	/** Write history to BitStream */
+	void Write(FBitWriter& Writer) const;
+
+	/** Read history from BitStream */
+	void Read(FBitReader& Reader);
+
+private:
+	const WordT* Data() const { return &Storage[0]; }
+	size_t NumWords() const { return WordCount; }
+
+	WordT Storage[WordCount];
+};
+
+/** Helper class to work with sequence numbers */
+template <size_t NumBits, typename SequenceType>
+class TSequenceNumber
+{
+	static_assert(TIsSigned<SequenceType>().Value == false, "The base type for sequence numbers must be unsigned");
+
+public:
+	typedef SequenceType SequenceT;
+	typedef int32 DifferenceT;
+
+	// Constants
+	enum { SeqNumberBits = NumBits };
+	enum { SeqNumberCount = SequenceT(1) << NumBits };
+	enum { SeqNumberHalf = SequenceT(1) << (NumBits - 1) };
+	enum { SeqNumberMax = SeqNumberCount - 1u };
+	enum { SeqNumberMask = SeqNumberMax };
+
+	/** Default constructor */
+	TSequenceNumber() : Value(0u) {}
+
+	/** Constructor with given value */
+	TSequenceNumber(SequenceT ValueIn) : Value(ValueIn & SeqNumberMask) {}
+	
+	/** Get Current Value */	
+	SequenceT Get() const { return Value; }
+
+	/** Diff between sequence numbers (A - B) only valid if (A - B) < SeqNumberHalf */
+	static DifferenceT Diff(TSequenceNumber A, TSequenceNumber B);
+	
+	/** return true if this is > Other, this is only considered to be the case if (A - B) < SeqNumberHalf since we have to be able to detect wraparounds */
+	bool operator>(const TSequenceNumber& Other) const { return (Value != Other.Value) && (((Value - Other.Value) & SeqNumberMask) < SeqNumberHalf); }
+
+	/** Check if this is >= Other, See above */
+	bool operator>=(const TSequenceNumber& Other) const { return ((Value - Other.Value) & SeqNumberMask) < SeqNumberHalf; }
+
+	/** Equals, NOTE that serial numbers wrap around so 0 == 0 + SequenceNumberCount */
+	bool operator==(const TSequenceNumber& Other) const { return Value == Other.Value; }
+
+	bool operator!=(const TSequenceNumber& Other) const { return Value != Other.Value; }
+
+	/** Pre-increment and wrap around */
+	TSequenceNumber& operator++() { Increment(1u); return *this; }
+	
+	/** Post-increment and wrap around */
+	TSequenceNumber operator++(int) { TSequenceNumber Tmp(*this); Increment(1u); return Tmp; }
+
+private:
+	void Increment(SequenceT InValue) { *this = TSequenceNumber(Value + InValue); }
+	SequenceT Value;
+};
+
+/** 
+	FNetPacketNotify - Drives delivery of sequence numbers, acknowledgments and notifications of delivery sequence numbers
+*/
+class FNetPacketNotify
+{
+public:
+	enum { SequenceNumberBits = 14 };
+	enum { SequenceHistoryLength = 32 };
+
+	typedef TSequenceNumber<SequenceNumberBits, uint16> SequenceNumberT;
+	typedef TSequenceHistory<SequenceHistoryLength> SequenceHistoryT;
+
+	struct FNotificationHeader
+	{
+		SequenceHistoryT History;
+		SequenceNumberT Seq;
+		SequenceNumberT AckedSeq;
+	};
+
+	/** Init notification with expected initial sequence numbers */
+	void Init(SequenceNumberT InitialInSeq, SequenceNumberT InitialOutSeq);
+
+	/** Mark Seq as received and update current InSeq, missing sequence numbers will be marked as lost */
+	void AckSeq(SequenceNumberT Seq) { AckSeq(Seq, true); }
+
+	/** Explicitly mark Seq as not received and update current InSeq, additional missing sequence numbers will be marked as lost */
+	void NakSeq(SequenceNumberT Seq) { AckSeq(Seq, false); }
+
+	/** Increment outgoing seq number */
+	SequenceNumberT IncrementOutSeq() { return ++OutSeq; }	
+
+	/** Get header that should be written to stream */
+	void GetHeader(FNotificationHeader& DataOut) const;
+
+	/** Write header to stream, note this must always write the exact same number of bits as we go back and update the header */
+	static void WriteHeader(FBitWriter& Writer, const FNotificationHeader& Data);
+
+	/** Read header from stream */
+	static bool ReadHeader(FNotificationHeader& Data, FBitReader& Reader);
+
+	/** Update state of PacketNotification based on received header and invoke packet notifications for received acks.
+		InFunc is a function in the format void)(FNetPacketNotify::SequenceNumberT AckedSequence, bool delivered) to handle packetNotifications.
+
+		Returns the positive delta of the incoming seq number if it is within half the sequence number space
+		Returns 0 if the received sequence number is out of the current window or if the acknowledged seq received by the remote is invalid.
+	*/
+	template<class Functor>
+	SequenceNumberT::DifferenceT Update(const FNotificationHeader& NotificationData, Functor&& InFunc)
+	{		
+		if (NotificationData.Seq > InSeq && NotificationData.AckedSeq >= OutAckSeq)
+		{
+			UE_LOG_PACKET_NOTIFY(TEXT("FNetPacketNotify::Update - Seq %u, InSeq %u"), NotificationData.Seq.Get(), InSeq.Get());
+
+			const SequenceNumberT::DifferenceT InSeqDelta = SequenceNumberT::Diff(NotificationData.Seq, InSeq);
+			
+			ProcessReceivedAcks(NotificationData, InFunc);
+
+			// accept sequence
+			InSeq = NotificationData.Seq;
+
+			return InSeqDelta;
+		}
+		else
+		{
+			return 0;
+		}
+	}
+
+	/** Get the current SequenceHistory */
+	const SequenceHistoryT& GetInSeqHistory() const { return InSeqHistory; }
+
+	/** Get the last received in sequence number */
+	SequenceNumberT GetInSeq() const { return InSeq; }
+
+	/** Get the last received sequence number that we have accepted, InAckSeq cannot be larger than InSeq */
+	SequenceNumberT GetInAckSeq() const { return InAckSeq; }
+
+	/** Get the current outgoing sequence number */
+	SequenceNumberT GetOutSeq() const { return OutSeq; }
+
+	/** Get the last outgoing sequence number acknowledged by remote */
+	SequenceNumberT GetOutAckSeq() const { return OutAckSeq; }
+
+	/** If we do have more unacknowledged sequence numbers in-flight than our maximum sendwindow we should not send more as the receiving end will not be able to detect if the sequence number has wrapped around */
+	bool CanSend() const { SequenceNumberT NextOutSeq = OutSeq; ++NextOutSeq; return NextOutSeq >= OutAckSeq; }
+
+private:
+	SequenceHistoryT InSeqHistory;	// BitBuffer containing a bitfield describing the history of received packets
+	SequenceNumberT InSeq;			// Last sequence number received from remote
+	SequenceNumberT InAckSeq;		// Last sequence number received from remote that we have acknowledged, this is needed since we support accepting a packet but explicitly not acknowledge it as received.
+	SequenceNumberT OutSeq;			// Outgoing sequence number
+	SequenceNumberT OutAckSeq;		// Last sequence number that we know that the remote side have received.
+
+	template<class Functor>
+	void ProcessReceivedAcks(const FNotificationHeader& NotificationData, Functor&& InFunc)
+	{
+		if (NotificationData.AckedSeq > OutAckSeq)
+		{
+			UE_LOG_PACKET_NOTIFY(TEXT("Notification::ProcessReceivedAcks - AckedSeq: %u, OutAckSeq: %u"), NotificationData.AckedSeq.Get(), OutAckSeq.Get());
+
+			SequenceNumberT::DifferenceT AckCount = SequenceNumberT::Diff(NotificationData.AckedSeq, OutAckSeq);
+
+			// ExpectedAck = OutAckSeq + 1
+			SequenceNumberT CurrentAck(OutAckSeq);
+			++CurrentAck;
+
+			// Everything not found in the history buffer is treated as lost
+			while (AckCount > NotificationData.History.Size())
+			{
+				--AckCount;
+				UE_LOG_PACKET_NOTIFY(TEXT("Notification::ProcessReceivedAcks Seq: %u - IsAck: 0 HistoryIndex: N/A"), CurrentAck.Get());
+				InFunc(CurrentAck, false);
+				++CurrentAck;
+			}
+
+			// For sequence numbers contained in the history we lookup the delivery status
+			while (AckCount > 0)
+			{
+				--AckCount;
+				UE_LOG_PACKET_NOTIFY(TEXT("Notification::ProcessReceivedAcks Seq: %u - IsAck: %u HistoryIndex: %u"), CurrentAck.Get(), NotificationData.History.IsDelivered(AckCount) ? 1u : 0u, AckCount);
+				InFunc(CurrentAck, NotificationData.History.IsDelivered(AckCount));
+				++CurrentAck;
+			}
+			OutAckSeq = NotificationData.AckedSeq;	
+		}
+	}
+
+	void AckSeq(SequenceNumberT AckedSeq, bool IsAck);
+
+#if WITH_DEV_AUTOMATION_TESTS
+	friend struct FNetPacketNotifyTestUtil;
+#endif
+};
+
+
+//////////////////////////////////////////////////////////////////////////
+// Implementation of TSequenceHistory
+//////////////////////////////////////////////////////////////////////////
+template <size_t HistorySize>
+TSequenceHistory<HistorySize>::TSequenceHistory()
+{
+	Reset();
+}
+
+#if WITH_DEV_AUTOMATION_TESTS
+template <size_t HistorySize>
+TSequenceHistory<HistorySize>::TSequenceHistory(WordT Value)
+{
+	for (size_t CurrentWordIt = 0; CurrentWordIt < WordCount; ++CurrentWordIt)
+	{
+		Storage[CurrentWordIt] = Value;
+	}	
+}
+#endif 
+
+template <size_t HistorySize>
+void TSequenceHistory<HistorySize>::Reset()
+{
+	memset(&Storage[0], 0, WordCount * sizeof(WordT));
+}
+
+template <size_t HistorySize>
+void TSequenceHistory<HistorySize>::AddDeliveryStatus(bool Delivered)
+{
+	WordT Carry = Delivered ? 1u : 0u;
+	const WordT ValueMask = 1u << (BitsPerWord - 1);
+	
+	for (size_t CurrentWordIt = 0; CurrentWordIt < WordCount; ++CurrentWordIt)
+	{
+		const WordT OldValue = Carry;
+		
+		// carry over highest bit in each word to the next word
+		Carry = (Storage[CurrentWordIt] & ValueMask) >> (BitsPerWord - 1);
+		Storage[CurrentWordIt] = (Storage[CurrentWordIt] << 1u) | OldValue;
+	}
+}
+
+template <size_t HistorySize>
+bool TSequenceHistory<HistorySize>::IsDelivered(size_t Index) const
+{
+	check(Index < Size());
+
+	const size_t WordIndex = Index / BitsPerWord;
+	const WordT WordMask = (WordT(1) << (Index & (BitsPerWord - 1)));
+	
+	return (Storage[WordIndex] & WordMask) != 0u;
+}
+
+template <size_t HistorySize>
+void TSequenceHistory<HistorySize>::Write(FBitWriter& Writer) const
+{
+	for (size_t CurrentWordIt = 0; CurrentWordIt < WordCount; ++CurrentWordIt)
+	{
+		WordT temp = Storage[CurrentWordIt];
+		Writer << temp;
+	}
+}
+
+template <size_t HistorySize>
+void TSequenceHistory<HistorySize>::Read(FBitReader& Reader)
+{
+	for (size_t CurrentWordIt = 0; CurrentWordIt < WordCount; ++CurrentWordIt)
+	{
+		Reader << Storage[CurrentWordIt];
+	}
+}
+
+//////////////////////////////////////////////////////////////////////////
+// Implementation of TSequenceNumber
+//////////////////////////////////////////////////////////////////////////
+template <size_t NumBits, typename SequenceType>
+typename TSequenceNumber<NumBits, SequenceType>::DifferenceT TSequenceNumber<NumBits, SequenceType>::Diff(TSequenceNumber A, TSequenceNumber B) 
+{ 
+	const size_t ShiftValue = sizeof(DifferenceT)*8 - NumBits;
+
+	const SequenceT ValueA = A.Value;
+	const SequenceT ValueB = B.Value;
+
+	return (DifferenceT)((ValueA - ValueB) << ShiftValue) >> ShiftValue;
+};
