NetworkProfiler updates to reading addresses. Can now read and display addresses of any type, regardless of protocol. SteamNet, IPv6, IPv4, Oculus, etc.

Backwards compatible with older profiles as well.

diff --git a/Engine/Source/Programs/NetworkProfiler/NetworkProfiler/MainWindow.cs b/Engine/Source/Programs/NetworkProfiler/NetworkProfiler/MainWindow.cs
--- a/Engine/Source/Programs/NetworkProfiler/NetworkProfiler/MainWindow.cs
+++ b/Engine/Source/Programs/NetworkProfiler/NetworkProfiler/MainWindow.cs
@@ -419,11 +419,15 @@ namespace NetworkProfiler
 
 			ConnectionListBox.Items.Clear();
 
-			for ( int i = 0; i < CurrentNetworkStream.AddressArray.Count; i++ )
+			int NumberOfAddresses = ( CurrentNetworkStream.GetVersion() < 12 ) ? 
+				CurrentNetworkStream.AddressArray.Count : CurrentNetworkStream.StringAddressArray.Count;
+
+			for (int i = 0; i < NumberOfAddresses; i++)
 			{
-				ConnectionListBox.Items.Add( CurrentNetworkStream.GetIpString( i ) );
+				ConnectionListBox.Items.Add( CurrentNetworkStream.GetIpString( i, CurrentNetworkStream.GetVersion() ) );
 			}
 
+
 			for ( int i = 0; i < ConnectionListBox.Items.Count; ++i )
 			{
 				ConnectionListBox.SetItemChecked( i, true );
diff --git a/Engine/Source/Programs/NetworkProfiler/NetworkProfiler/NetworkStream.cs b/Engine/Source/Programs/NetworkProfiler/NetworkProfiler/NetworkStream.cs
--- a/Engine/Source/Programs/NetworkProfiler/NetworkProfiler/NetworkStream.cs
+++ b/Engine/Source/Programs/NetworkProfiler/NetworkProfiler/NetworkStream.cs
@@ -24,6 +24,8 @@ namespace NetworkProfiler
 
 		/** Array of unique addresses. Code has fixed indexes into it.				*/
 		public List<UInt64> AddressArray = new List<UInt64>();
+		// Used for new storage method.
+		public List<string> StringAddressArray = new List<string>();
 
 		/** Last address index parsed from token stream								*/
 		public int CurrentConnectionIndex = 0;
@@ -61,22 +63,27 @@ namespace NetworkProfiler
 		 * Returns the ip address string associated with the passed in connection index.
 		 * 
 		 * @param	ConnectionIndex	Index in address table
-		 * @return	Ip string associated with adress table index
+		 * @param	NetworkVersion the version of this network profiler, determines if we need to do extra parsing with the string
+		 * @return	Ip string associated with address table index
 		 */
-		public string GetIpString( int ConnectionIndex )
+		public string GetIpString( int ConnectionIndex, uint NetworkVersion )
 		{
-			UInt64 Addr = AddressArray[ConnectionIndex];
-			UInt32 IP	= ( UInt32 )( Addr >> 32 );
-			UInt32 Port = ( UInt32 )( Addr & ( ( ( UInt64 )1 << 32 ) - 1 ) );
-
-			byte ip0 = ( byte )( ( IP >> 24 ) & 255 );
-			byte ip1 = ( byte )( ( IP >> 16 ) & 255 );
-			byte ip2 = ( byte )( ( IP >> 8 ) & 255 );
-			byte ip3 = ( byte )( ( IP >> 0 ) & 255 );
-
-			//return string.Format( "{0,3:000}.{1,3:000}.{2,3:000}.{3,3:000}: {4,-5}", ip0, ip1, ip2, ip3, Port );
-			//return string.Format( "{0,3}.{1,3}.{2,3}.{3,3}: {4,-5}", ip0, ip1, ip2, ip3, Port );
-			return string.Format( "{0}.{1}.{2}.{3}: {4}", ip0, ip1, ip2, ip3, Port );
+			if( NetworkVersion < 12 )
+			{
+				UInt64 Addr = AddressArray[ConnectionIndex];
+				UInt32 IP = (UInt32)(Addr >> 32);
+				UInt32 Port = (UInt32)(Addr & (((UInt64)1 << 32) - 1));
+
+				byte ip0 = (byte)((IP >> 24) & 255);
+				byte ip1 = (byte)((IP >> 16) & 255);
+				byte ip2 = (byte)((IP >> 8) & 255);
+				byte ip3 = (byte)((IP >> 0) & 255);
+				return string.Format("{0}.{1}.{2}.{3}: {4}", ip0, ip1, ip2, ip3, Port);
+			}
+			else
+			{
+				return StringAddressArray[ConnectionIndex];
+			}
 		}
 
 		/**
diff --git a/Engine/Source/Programs/NetworkProfiler/NetworkProfiler/StreamParser.cs b/Engine/Source/Programs/NetworkProfiler/NetworkProfiler/StreamParser.cs
--- a/Engine/Source/Programs/NetworkProfiler/NetworkProfiler/StreamParser.cs
+++ b/Engine/Source/Programs/NetworkProfiler/NetworkProfiler/StreamParser.cs
@@ -135,7 +135,15 @@ namespace NetworkProfiler
 
 				if ( Token.TokenType == ETokenTypes.ConnectionReference )
 				{
-					NetworkStream.AddressArray.Add( ( Token as TokenConnectionReference ).Address );
+					if( NetworkStream.GetVersion() < 12 )
+					{
+						NetworkStream.AddressArray.Add((Token as TokenConnectionReference).Address);
+					}
+					else
+					{
+						NetworkStream.StringAddressArray.Add((Token as TokenConnectionStringReference).Address);
+					}
+					
 					continue;
 				}
 
diff --git a/Engine/Source/Programs/NetworkProfiler/NetworkProfiler/Tokens.cs b/Engine/Source/Programs/NetworkProfiler/NetworkProfiler/Tokens.cs
--- a/Engine/Source/Programs/NetworkProfiler/NetworkProfiler/Tokens.cs
+++ b/Engine/Source/Programs/NetworkProfiler/NetworkProfiler/Tokens.cs
@@ -127,7 +127,16 @@ namespace NetworkProfiler
 					SerializedToken = new TokenNameReference( BinaryStream );
 					break;
 				case ETokenTypes.ConnectionReference:
-					SerializedToken = new TokenConnectionReference( BinaryStream );
+					{
+						if (InNetworkStream.GetVersion() < 12)
+						{
+							SerializedToken = new TokenConnectionReference(BinaryStream);
+						}
+						else
+						{
+							SerializedToken = new TokenConnectionStringReference(BinaryStream);
+						}
+					}
 					break;
 				case ETokenTypes.ConnectionChange:
 					SerializedToken = new TokenConnectionChanged( BinaryStream );
@@ -228,7 +237,7 @@ namespace NetworkProfiler
 		{
 			TreeNode Child = TokenHelper.AddNode( Tree, "Socket SendTo" );
 
-			Child = Child.Nodes.Add( "Destination : " + StreamParser.NetworkStream.GetIpString( ConnectionIndex ) );
+			Child = Child.Nodes.Add( "Destination : " + StreamParser.NetworkStream.GetIpString( ConnectionIndex, StreamParser.NetworkStream.GetVersion() ) );
 
 			Child.Nodes.Add( "SocketName          : " + StreamParser.NetworkStream.GetName( SocketNameIndex ) );
 			Child.Nodes.Add( "DesiredBytesSent    : " + ( NumPacketIdBits + NumBunchBits + NumAckBits + NumPaddingBits ) / 8.0f );
@@ -751,6 +760,24 @@ namespace NetworkProfiler
 		}
 	}
 
+	/**
+	 * Token for connection reference as a string.
+	 * This allows for support for different address formats without having to do any additional work.
+	 * Addresses are pushed in via the ToString(true) call on an FInternetAddr
+	 */
+	class TokenConnectionStringReference : TokenBase
+	{
+		/** Address of connection */
+		public string Address = null;
+
+		/** Constructor, serializing members from passed in stream. */
+		public TokenConnectionStringReference(BinaryReader BinaryStream)
+		{
+			int StrLength = Math.Abs(BinaryStream.ReadInt32());
+			Address = new string(BinaryStream.ReadChars(StrLength));
+		}
+	}
+
 	/**
 	 * Token for events.
 	 */
diff --git a/Engine/Source/Runtime/Engine/Private/NetworkProfiler.cpp b/Engine/Source/Runtime/Engine/Private/NetworkProfiler.cpp
--- a/Engine/Source/Runtime/Engine/Private/NetworkProfiler.cpp
+++ b/Engine/Source/Runtime/Engine/Private/NetworkProfiler.cpp
@@ -31,7 +31,7 @@ FNetworkProfiler GNetworkProfiler;
 /** Magic value, determining that file is a network profiler file.				*/
 #define NETWORK_PROFILER_MAGIC						0x1DBF348C
 /** Version of memory profiler. Incremented on serialization changes.			*/
-#define NETWORK_PROFILER_VERSION					11
+#define NETWORK_PROFILER_VERSION					12
 
 static const FString UnknownName("UnknownName");
 
@@ -96,10 +96,10 @@ FArchive& operator << ( FArchive& Ar, FNetworkProfilerHeader& Header )
  * Constructor, initializing member variables.
  */
 FNetworkProfiler::FNetworkProfiler()
-:	FileWriter(NULL)
+:	FileWriter( nullptr )
 ,	bHasNoticeableNetworkTrafficOccured(false)
 ,	bIsTrackingEnabled(false)
-,	LastAddress( 0xFFFFFFFFFFFFFFFF )
+,	LastAddress( nullptr )
 {
 }
 
@@ -140,10 +140,10 @@ int32 FNetworkProfiler::GetNameTableIndex( const FString& Name )
 /**
 * Returns index of passed in name into name array. If not found, adds it.
 *
-* @param	Name	Name to find index for
+* @param	Address	Address string to find index for
 * @return	Index of passed in name
 */
-int32 FNetworkProfiler::GetAddressTableIndex( uint64 Address )
+int32 FNetworkProfiler::GetAddressTableIndex( const FString& Address )
 {
 	// Index of name in name table.
 	int32 Index = INDEX_NONE;
@@ -157,14 +157,13 @@ int32 FNetworkProfiler::GetAddressTableIndex( uint64 Address )
 	// Encountered new name, add to array and set index mapping.
 	else
 	{
-		Index = AddressArray.Num();
-		AddressArray.Add( Address );
+		Index = AddressTableIndexMap.Num();
 		AddressTableIndexMap.Add( Address, Index );
 
 		// Write out the name reference token
 		uint8 Type = NPTYPE_ConnectionReference;
 		( *FileWriter ) << Type;
-		( *FileWriter ) << Address;
+		Address.SerializeAsANSICharArray( *FileWriter );
 	}
 
 	check( Index != INDEX_NONE );
@@ -204,7 +203,7 @@ void FNetworkProfiler::TrackFrameBegin()
 		(*FileWriter) << Type;
 		float RelativeTime=  (float)(FPlatformTime::Seconds() - GStartTime);
 		(*FileWriter) << RelativeTime;
-		LastAddress = 0xFFFFFFFFFFFFFFFF;
+		LastAddress = nullptr;
 	}
 }
 
@@ -217,20 +216,19 @@ void FNetworkProfiler::SetCurrentConnection( UNetConnection* Connection )
 {
 	if ( bIsTrackingEnabled && Connection != nullptr )
 	{
-		const uint32 NetworkByteOrderIP = Connection->GetAddrAsInt();
-		const uint32 Port				= Connection->GetAddrPort();
-
-		const uint64 Address = ( ( ( uint64 )NetworkByteOrderIP ) << 32 ) | Port;
-
-		if ( Address != LastAddress )
+		const TSharedPtr<FInternetAddr> ConnectionAddr = Connection->GetInternetAddr();
+		if ( ConnectionAddr.IsValid() )
 		{
-			uint32 Index = GetAddressTableIndex( Address );
+			if ( LastAddress != ConnectionAddr )
+			{
+				uint32 Index = GetAddressTableIndex(ConnectionAddr->ToString(true));
 
-			uint8 Type = NPTYPE_ConnectionChanged;
-			( *FileWriter ) << Type;
-			( *FileWriter ).SerializeIntPacked( Index );
+				uint8 Type = NPTYPE_ConnectionChanged;
+				(*FileWriter) << Type;
+				(*FileWriter).SerializeIntPacked(Index);
 
-			LastAddress = Address;
+				LastAddress = ConnectionAddr;
+			}
 		}
 	}
 }
@@ -596,14 +594,14 @@ void FNetworkProfiler::TrackSessionChange( bool bShouldContinueTracking, const F
 
 			// Clean up.
 			delete FileWriter;
-			FileWriter = NULL;
+			FileWriter = nullptr;
 			bHasNoticeableNetworkTrafficOccured = false;
 		}
 
 		if( bShouldContinueTracking )
 		{
 			// Start a new tracking session.
-			check( FileWriter == NULL );
+			check( FileWriter == nullptr );
 
 			static int32 Salt = 0;
 			Salt++;		// Use a salt to solve the issue where this function is called so fast it produces the same time (seems to happen during seamless travel)
@@ -617,7 +615,6 @@ void FNetworkProfiler::TrackSessionChange( bool bShouldContinueTracking, const F
 			NameToNameTableIndexMap.Reset();
 			NameArray.Reset();
 			AddressTableIndexMap.Reset();
-			AddressArray.Reset();
 
 			CurrentHeader.Reset(InURL);
 
@@ -746,10 +743,10 @@ bool FNetworkProfiler::Exec( UWorld * InWorld, const TCHAR* Cmd, FOutputDevice &
 	}
 
 	// If we are tracking, and we don't have a file writer, force one now 
-	if ( bIsTrackingEnabled && FileWriter == NULL ) 
+	if ( bIsTrackingEnabled && FileWriter == nullptr ) 
 	{
 		TrackSessionChange( true, InWorld != nullptr ? InWorld->URL : FURL() );
-		if ( FileWriter == NULL )
+		if ( FileWriter == nullptr )
 		{
 			UE_LOG(LogNet, Warning, TEXT("FNetworkProfiler::Exec: FAILED to create file writer!"));
 			EnableTracking( false );
diff --git a/Engine/Source/Runtime/Engine/Public/Net/NetworkProfiler.h b/Engine/Source/Runtime/Engine/Public/Net/NetworkProfiler.h
--- a/Engine/Source/Runtime/Engine/Public/Net/NetworkProfiler.h
+++ b/Engine/Source/Runtime/Engine/Public/Net/NetworkProfiler.h
@@ -8,6 +8,7 @@
 
 #include "CoreMinimal.h"
 #include "Engine/EngineTypes.h"
+#include "IPAddress.h"
 
 class AActor;
 class FOutBunch;
@@ -80,10 +81,7 @@ private:
 	TArray<FString>							NameArray;
 
 	/** Mapping from address to index in address array.												*/
-	TMap<uint64, int32>						AddressTableIndexMap;
-
-	/** Array of unique addresses																	*/
-	TArray<uint64>							AddressArray;
+	TMap<FString, int32>					AddressTableIndexMap;
 
 	/** Whether noticeable network traffic has occured in this session. Used to discard it.			*/
 	bool									bHasNoticeableNetworkTrafficOccured;
@@ -94,7 +92,7 @@ private:
 	FNetworkProfilerHeader					CurrentHeader;
 
 	/** Last known address																			*/
-	uint64									LastAddress;
+	TSharedPtr<FInternetAddr>				LastAddress;
 
 	/** All the data required for writing sent bunches to the profiler stream						*/
 	struct FSendBunchInfo
@@ -157,7 +155,7 @@ private:
 	* @param	Address	Address to find index for
 	* @return	Index of passed in name
 	*/
-	int32 GetAddressTableIndex( uint64 Address );
+	int32 GetAddressTableIndex( const FString& Address );
 
 public:
 	/**
