Fixed signed/unsigned comparison warning in NetPacketNotify and NetPacketNotifyTests

diff --git a/Engine/Source/Runtime/Engine/Private/Net/Tests/NetPacketNotifyTests.cpp b/Engine/Source/Runtime/Engine/Private/Net/Tests/NetPacketNotifyTests.cpp
--- a/Engine/Source/Runtime/Engine/Private/Net/Tests/NetPacketNotifyTests.cpp
+++ b/Engine/Source/Runtime/Engine/Private/Net/Tests/NetPacketNotifyTests.cpp
@@ -49,7 +49,7 @@ struct FNetPacketNotifyTestUtil
 	{
 		size_t SequenceNumberCount = sizeof(InSequenceNumbers) / sizeof(InSequenceNumbers[0]);
 
-		for (int32 I=0; I<SequenceNumberCount; ++I)
+		for (size_t I=0; I<SequenceNumberCount; ++I)
 		{
 			FNetPacketNotifyTestUtil::PretendReceiveSeq(PacketNotify, InSequenceNumbers[I]);
 		}
diff --git a/Engine/Source/Runtime/Engine/Public/Net/NetPacketNotify.h b/Engine/Source/Runtime/Engine/Public/Net/NetPacketNotify.h
--- a/Engine/Source/Runtime/Engine/Public/Net/NetPacketNotify.h
+++ b/Engine/Source/Runtime/Engine/Public/Net/NetPacketNotify.h
@@ -235,7 +235,7 @@ private:
 			++CurrentAck;
 
 			// Everything not found in the history buffer is treated as lost
-			while (AckCount > NotificationData.History.Size())
+			while (AckCount > (SequenceNumberT::DifferenceT)(NotificationData.History.Size()))
 			{
 				--AckCount;
 				UE_LOG_PACKET_NOTIFY(TEXT("Notification::ProcessReceivedAcks Seq: %u - IsAck: 0 HistoryIndex: N/A"), CurrentAck.Get());
