Partial backout of 4647801 until I have resolved issue.

diff --git a/Engine/Source/Runtime/Engine/Classes/Engine/ChildConnection.h b/Engine/Source/Runtime/Engine/Classes/Engine/ChildConnection.h
--- a/Engine/Source/Runtime/Engine/Classes/Engine/ChildConnection.h
+++ b/Engine/Source/Runtime/Engine/Classes/Engine/ChildConnection.h
@@ -52,6 +52,10 @@ public:
 		Parent->AssertValid();
 	}
 
+	virtual void SendAck(int32 PacketId, bool FirstTime=1) override
+	{
+	}
+
 	virtual void FlushNet(bool bIgnoreSimulation = false) override
 	{
 		Parent->FlushNet(bIgnoreSimulation);
diff --git a/Engine/Source/Runtime/Engine/Classes/Engine/NetConnection.h b/Engine/Source/Runtime/Engine/Classes/Engine/NetConnection.h
--- a/Engine/Source/Runtime/Engine/Classes/Engine/NetConnection.h
+++ b/Engine/Source/Runtime/Engine/Classes/Engine/NetConnection.h
@@ -14,7 +14,6 @@
 #include "GameFramework/OnlineReplStructs.h"
 #include "Engine/NetDriver.h"
 #include "Net/DataBunch.h"
-#include "Net/NetPacketNotify.h"
 #include "Engine/Player.h"
 #include "Engine/Channel.h"
 #include "ProfilingDebugging/Histogram.h"
@@ -39,11 +38,10 @@ typedef TMap<TWeakObjectPtr<AActor>, UActorChannel*, FDefaultSetAllocator, TWeak
 	Types.
 -----------------------------------------------------------------------------*/
 enum { RELIABLE_BUFFER = 256 }; // Power of 2 >= 1.
-enum { MAX_PACKETID = FNetPacketNotify::SequenceNumberT::SeqNumberCount };  // Power of 2 >= 1, covering guaranteed loss/misorder time.
+enum { MAX_PACKETID = 16384 };  // Power of 2 >= 1, covering guaranteed loss/misorder time.
 enum { MAX_CHSEQUENCE = 1024 }; // Power of 2 >RELIABLE_BUFFER, covering loss/misorder time.
 enum { MAX_BUNCH_HEADER_BITS = 256 };
-enum { UE_PACKET_SEQUENCE_HISTORY_BITS = FNetPacketNotify::SequenceHistoryLength };
-enum { MAX_PACKET_HEADER_BITS = 81 }; // = [(Seq)16|(AckSeq)16|(SeqHistory)32|1|?(ServerTime)8|(ReceviedRateByte)8]
+enum { MAX_PACKET_HEADER_BITS = 15 }; // = FMath::CeilLogTwo(MAX_PACKETID) + 1 (IsAck)
 enum { MAX_PACKET_TRAILER_BITS = 1 };
 
 // 
@@ -386,7 +384,7 @@ public:
 	FBitWriter		SendBuffer;						// Queued up bits waiting to send
 	double			OutLagTime[256];				// For lag measuring.
 	int32			OutLagPacketId[256];			// For lag measuring.
-	uint8			OutBytesPerSecondHistory[256];	// For saturation measuring.
+	int32			OutBytesPerSecondHistory[256];	// For saturation measuring.
 	float			RemoteSaturation;
 	int32			InPacketId;						// Full incoming packet index.
 	int32			OutPacketId;					// Most recently sent packet.
@@ -402,7 +400,6 @@ public:
 	TArray<int32>		OutReliable;
 	TArray<int32>		InReliable;
 	TArray<int32>		PendingOutRec;	// Outgoing reliable unacked data from previous (now destroyed) channel in this slot.  This contains the first chsequence not acked
-	UE_DEPRECATED(4.22, "QueuedAcks, ResendAcks variables will be removed.")
 	TArray<int32> QueuedAcks, ResendAcks;
 
 	int32				InitOutReliable;
@@ -691,7 +688,6 @@ public:
 	ENGINE_API virtual void AssertValid();
 
 	/** Send an acknowledgment. */
-	UE_DEPRECATED(4.22, "This method will be removed")
 	ENGINE_API virtual void SendAck( int32 PacketId, bool FirstTime=1);
 
 	/**
@@ -856,7 +852,6 @@ public:
 	// Functions.
 
 	/** Resend any pending acks. */
-	UE_DEPRECATED(4.22, "This method will be removed.")
 	void PurgeAcks();
 
 	/** Send package map to the remote. */
@@ -1054,18 +1049,6 @@ private:
 	/** Returns true if an outgoing packet should be dropped due to packet simulation settings, including loss burst simulation. */
 	bool ShouldDropOutgoingPacketForLossSimulation(int64 NumBits) const;
 
-	/** Write packetHeader */
-	void WritePacketHeader(FBitWriter& Writer) const;
-
-	/** Write extended packet header information (ServerFrameTime, SaturationData) */
-	void WritePacketInfo(FBitWriter& Writer) const;
-	
-	/** Read extended packet header information (ServerFrameTime, SaturationData) */
-	bool ReadPacketInfo(FBitReader& Reader);
-
-	/** Packet was acknowledged as delivered */
-	void ReceivedAck(int32 AckPacketId);
-
 	/**
 	 * on the server, the world the client has told us it has loaded
 	 * used to make sure the client has traveled correctly, prevent replicating actors before level transitions are done, etc
@@ -1084,15 +1067,6 @@ private:
 
 	/** Bookkeeping for connections using internal ack to avoid acking all open channels */
 	TArray<int32> ChannelsWaitingForInternalAck;
-
-	/** Sequence data used to implement reliability */
-	FNetPacketNotify PacketNotify;
-
-	/** Full PacketId  of last sent packet that we have received notification for (i.e. we know if it was delivered or not). Related to OutAckPacketId which is tha last successfully delivered PacketId */
-	int32 LastNotifiedPacketId;
-
-	/** Keep old behavior where we send a packet with only acks even if we have no other outgoing data if we got incoming data */
-	bool bHasDirtyAcks;
 };
 
 
diff --git a/Engine/Source/Runtime/Engine/Private/NetConnection.cpp b/Engine/Source/Runtime/Engine/Private/NetConnection.cpp
--- a/Engine/Source/Runtime/Engine/Private/NetConnection.cpp
+++ b/Engine/Source/Runtime/Engine/Private/NetConnection.cpp
@@ -19,7 +19,6 @@
 #include "Net/UnrealNetwork.h"
 #include "Net/NetworkProfiler.h"
 #include "Net/DataReplication.h"
-#include "Net/NetPacketNotify.h"
 #include "Engine/ActorChannel.h"
 #include "Engine/ChildConnection.h"
 #include "Engine/VoiceChannel.h"
@@ -27,6 +26,7 @@
 #include "Engine/PackageMapClient.h"
 #include "Engine/NetworkObjectList.h"
 #include "EncryptionComponent.h"
+
 #include "Net/PerfCountersHelpers.h"
 #include "GameDelegates.h"
 #include "Misc/PackageName.h"
@@ -133,8 +133,6 @@ UNetConnection::UNetConnection(const FObjectInitializer& ObjectInitializer)
 #endif
 ,	PlayerOnlinePlatformName( NAME_None )
 ,	ClientWorldPackageName( NAME_None )
-,	LastNotifiedPacketId( -1 )
-,	bHasDirtyAcks(false)
 {
 	MaxChannelSize = CVarMaxChannelSize.GetValueOnAnyThread();
 	if (MaxChannelSize <= 0)
@@ -147,8 +145,6 @@ UNetConnection::UNetConnection(const FObjectInitializer& ObjectInitializer)
 	OutReliable.AddDefaulted(MaxChannelSize);
 	InReliable.AddDefaulted(MaxChannelSize);
 	PendingOutRec.AddDefaulted(MaxChannelSize);
-
-	PacketNotify.Init(InPacketId, OutPacketId);
 }
 
 /**
@@ -342,7 +338,6 @@ void UNetConnection::InitSequence(int32 IncomingSequence, int32 OutgoingSequence
 		InPacketId = IncomingSequence - 1;
 		OutPacketId = OutgoingSequence;
 		OutAckPacketId = OutgoingSequence - 1;
-		LastNotifiedPacketId = OutAckPacketId;
 
 		// Initialize the reliable packet sequence (more useful/effective at preventing attacks)
 		InitInReliable = IncomingSequence & (MAX_CHSEQUENCE - 1);
@@ -351,8 +346,6 @@ void UNetConnection::InitSequence(int32 IncomingSequence, int32 OutgoingSequence
 		InReliable.Init(InitInReliable, InReliable.Num());
 		OutReliable.Init(InitOutReliable, OutReliable.Num());
 
-		PacketNotify.Init(InPacketId, OutPacketId);
-
 		UE_LOG(LogNet, Verbose, TEXT("InitSequence: IncomingSequence: %i, OutgoingSequence: %i, InitInReliable: %i, InitOutReliable: %i"), IncomingSequence, OutgoingSequence, InitInReliable, InitOutReliable);
 	}
 }
@@ -473,6 +466,8 @@ void UNetConnection::Serialize( FArchive& Ar )
 	{
 		Children.CountBytes(Ar);
 		ClientVisibleLevelNames.CountBytes(Ar);
+		QueuedAcks.CountBytes(Ar);
+		ResendAcks.CountBytes(Ar);
 		OpenChannels.CountBytes(Ar);
 		SentTemporaries.CountBytes(Ar);
 		ActorChannels.CountBytes(Ar);
@@ -1046,7 +1041,7 @@ void UNetConnection::FlushNet(bool bIgnoreSimulation)
 	TimeSensitive = 0;
 
 	// If there is any pending data to send, send it.
-	if (SendBuffer.GetNumBits() || bHasDirtyAcks || ( Driver->Time-LastSendTime > Driver->KeepAliveTime && !InternalAck && State != USOCK_Closed))
+	if (SendBuffer.GetNumBits() || ( Driver->Time-LastSendTime > Driver->KeepAliveTime && !InternalAck && State != USOCK_Closed))
 	{
 		// Due to the PacketHandler handshake code, servers must never send the client data,
 		// before first receiving a client control packet (which is taken as an indication of a complete handshake).
@@ -1061,10 +1056,10 @@ void UNetConnection::FlushNet(bool bIgnoreSimulation)
 
 		FOutPacketTraits Traits;
 
-		// If sending keepalive packet or just acks, still write the packet header
+		// If sending keepalive packet, still write the packet id
 		if (SendBuffer.GetNumBits() == 0)
 		{
-			WriteBitsToSendBuffer( NULL, 0 );		// This will force the packet header to be written
+			WriteBitsToSendBuffer( NULL, 0 );		// This will force the packet id to be written
 
 			Traits.bIsKeepAlive = true;
 			AnalyticsVars.OutKeepAliveCount++;
@@ -1081,12 +1076,6 @@ void UNetConnection::FlushNet(bool bIgnoreSimulation)
 		// Write the UNetConnection-level termination bit
 		SendBuffer.WriteBit(1);
 
-		// Refresh outgoing header with latest data
-		if ( !InternalAck )
-		{
-			WritePacketHeader(SendBuffer);
-		}
-
 		ValidateSendBuffer();
 
 		const int32 NumStrayBits = SendBuffer.GetNumBits();
@@ -1166,14 +1155,9 @@ void UNetConnection::FlushNet(bool bIgnoreSimulation)
 		// Remember the actual time this packet was sent out, so we can compute ping when the ack comes back
 		OutLagPacketId[Index]			= OutPacketId;
 		OutLagTime[Index]				= FPlatformTime::Seconds();
-	
-		OutBytesPerSecondHistory[Index]	= FMath::Min(OutBytesPerSecond / 1024, 255);
-
-		// Increase outgoing sequence number
-		PacketNotify.IncrementOutSeq();
-		bHasDirtyAcks = false;
-		++OutPacketId; 
+		OutBytesPerSecondHistory[Index]	= OutBytesPerSecond / 1024;
 
+		OutPacketId++;
 		++OutPackets;
 		++OutTotalPackets;
 		Driver->OutPackets++;
@@ -1199,6 +1183,13 @@ void UNetConnection::FlushNet(bool bIgnoreSimulation)
 
 		InitSendBuffer();
 	}
+
+	// Move acks around.
+	for (int32 i=0; i<QueuedAcks.Num(); i++)
+	{
+		ResendAcks.Add(QueuedAcks[i]);
+	}
+	QueuedAcks.Empty(32);
 }
 
 bool UNetConnection::ShouldDropOutgoingPacketForLossSimulation(int64 NumBits) const
@@ -1226,64 +1217,8 @@ int32 UNetConnection::IsNetReady( bool Saturate )
 void UNetConnection::ReadInput( float DeltaSeconds )
 {}
 
-void UNetConnection::ReceivedAck(int32 AckPacketId)
-{
-	UE_LOG(LogNetTraffic, Verbose, TEXT("   Received ack %i"), AckPacketId);
-
-	// Advance OutAckPacketId
-	OutAckPacketId = AckPacketId;
-
-	// Process the bunch.
-	LastRecvAckTime = Driver->Time;
-
-	if ( PackageMap != NULL )
-	{
-		PackageMap->ReceivedAck( AckPacketId );
-	}
-
-	// Forward the ack to the channel.
-	{
-		SCOPE_CYCLE_COUNTER(Stat_NetConnectionReceivedAcks);
-	
-		for (int32 i = OpenChannels.Num() - 1; i >= 0; i--)
-		{
-			UChannel* const Channel = OpenChannels[i];
-				
-			if (Channel)
-			{
-				if (Channel->OpenPacketId.Last == AckPacketId) // Necessary for unreliable "bNetTemporary" channels.
-				{
-					Channel->OpenAcked = 1;
-				}
-				
-				for (FOutBunch* OutBunch = Channel->OutRec; OutBunch; OutBunch = OutBunch->Next)
-				{
-					if (OutBunch->bOpen)
-					{
-						UE_LOG(LogNet, VeryVerbose, TEXT("Channel %i reset Ackd because open is reliable. "), Channel->ChIndex );
-						Channel->OpenAcked  = 0; // We have a reliable open bunch, don't let the above code set the OpenAcked state,
-												// it must be set in UChannel::ReceivedAcks to verify all open bunches were received.
-					}
-
-					if (OutBunch->PacketId == AckPacketId)
-					{
-						OutBunch->ReceivedAck = 1;
-					}
-				}
-				Channel->ReceivedAcks(); //warning: May destroy Channel.
-			}
-			else
-			{
-				UE_LOG(LogNet, Warning, TEXT("UNetConnection::ReceivedPacket: null channel in OpenChannels array while processing ack. %s"), *Describe());
-			}
-		}
-	}
-}
-
 void UNetConnection::ReceivedNak( int32 NakPacketId )
 {
-	UE_LOG(LogNetTraffic, Verbose, TEXT("   Received nak %i"), NakPacketId);
-
 	SCOPE_CYCLE_COUNTER(Stat_NetConnectionReceivedNak);
 
 	// Update pending NetGUIDs
@@ -1306,137 +1241,6 @@ void UNetConnection::ReceivedNak( int32 NakPacketId )
 			UE_LOG(LogNet, Warning, TEXT("UNetConnection::ReceivedNak: null channel in OpenChannels array. %s"), *Describe());
 		}
 	}
-
-	// Stats
-	++OutPacketsLost;
-	++OutTotalPacketsLost;
-	++Driver->OutTotalPacketsLost;
-}
-
-// IMPORTANT:
-// WritePacketHeader must ALWAYS write the exact same number of bits as we go back and rewrite the header
-// right before we put the packet on the wire.
-void UNetConnection::WritePacketHeader(FBitWriter& Writer) const
-{
-	// If this is a header refresh, we only serialize the updated serial number information
-	const bool bIsHeaderUpdate = Writer.GetNumBits() > 0u;
-
-	// Header is always written first in the packet
-	FBitWriterMark Reset;
-	FBitWriterMark Restore(Writer);
-	Reset.PopWithoutClear(Writer);
-	
-	FNetPacketNotify::FNotificationHeader Header;
-	PacketNotify.GetHeader(Header);
-	FNetPacketNotify::WriteHeader(Writer, Header);
-
-	// Jump back to where we came from.
-	if (bIsHeaderUpdate)
-	{
-		Restore.PopWithoutClear(Writer);
-	}
-}
-
-namespace
-{
-
-static bool ReadPacketHeader(FNetPacketNotify::FNotificationHeader& Header, FBitReader& Reader)
-{
-	FNetPacketNotify::ReadHeader(Header, Reader);
-		
-	return Reader.IsError() == false;
-}
-
-}
-
-void UNetConnection::WritePacketInfo(FBitWriter& Writer) const
-{
-	const uint8 bHasServerFrameTime = Driver->IsServer() ? bLastHasServerFrameTime : ( CVarPingExcludeFrameTime.GetValueOnGameThread() > 0 ? 1u : 0u );
-
-	// Write data used to calculate link latency
-	Writer.WriteBit(bHasServerFrameTime);
-	if (bHasServerFrameTime && Driver->IsServer())
-	{
-		uint8 FrameTimeByte = FMath::Min( FMath::FloorToInt( FrameTime * 1000 ), 255 );
-		Writer << FrameTimeByte;
-	}
-
-	// Notify server of our current rate per second at this time
-	uint8 InKBytesPerSecondByte = FMath::Min(InBytesPerSecond / 1024, 255);
-	Writer << InKBytesPerSecondByte;
-}
-
-bool UNetConnection::ReadPacketInfo(FBitReader& Reader)
-{
-	const bool bHasServerFrameTime = Reader.ReadBit();
-	double ServerFrameTime = 0;
-
-	if ( !Driver->IsServer() )
-	{
-		if ( bHasServerFrameTime )
-		{
-			uint8 FrameTimeByte	= 0;
-			Reader << FrameTimeByte;
-			// As a client, our request was granted, read the frame time
-			ServerFrameTime = ( double )FrameTimeByte / 1000;
-		}
-	}
-	else
-	{
-		bLastHasServerFrameTime = bHasServerFrameTime;
-	}
-
-	// limit to known size to know the size of the packet header
-	uint8 RemoteInKBytesPerSecondByte = 0;
-	Reader << RemoteInKBytesPerSecondByte;
-
-	if ( Reader.IsError() )
-	{
-		return false;
-	}
-
-	// Update ping
-	// At this time we have updated OutAckPacketId to the latest received ack.
-	const int32 Index = OutAckPacketId & (ARRAY_COUNT(OutLagPacketId)-1);
-
-	if ( OutLagPacketId[Index] == OutAckPacketId )
-	{
-		OutLagPacketId[Index] = -1;		// Only use the ack once
-
-#if !UE_BUILD_SHIPPING
-		if ( CVarPingDisplayServerTime.GetValueOnAnyThread() > 0 )
-		{
-			UE_LOG( LogNetTraffic, Warning, TEXT( "ServerFrameTime: %2.2f" ), ServerFrameTime * 1000.0f );
-		}
-#endif
-
-		// use FApp's time because it is set closer to the beginning of the frame - we don't care about the time so far of the current frame to process the packet
-		const double CurrentTime = FApp::GetCurrentTime();
-		const double GameTime	 = ServerFrameTime;
-		const double RTT		 = (CurrentTime - OutLagTime[Index] ) - ( CVarPingExcludeFrameTime.GetValueOnAnyThread() ? GameTime : 0.0 );
-		const double NewLag		 = FMath::Max( RTT, 0.0 );
-
-		if ( OutBytesPerSecondHistory[Index] > 0 )
-		{
-			RemoteSaturation = ( 1.0f - FMath::Min( ( float )RemoteInKBytesPerSecondByte / ( float )OutBytesPerSecondHistory[Index], 1.0f ) ) * 100.0f;
-		}
-		else
-		{
-			RemoteSaturation = 0.0f;
-		}
-
-		//UE_LOG( LogNet, Warning, TEXT( "Out: %i, InRemote: %i, Saturation: %f" ), OutBytesPerSecondHistory[Index], RemoteInKBytesPerSecond, RemoteSaturation );
-
-		LagAcc += NewLag;
-		LagCount++;
-
-		if (PlayerController != NULL)
-		{
-			PlayerController->UpdatePing(NewLag);
-		}
-	}
-
-	return true;
 }
 
 void UNetConnection::ReceivedPacket( FBitReader& Reader )
@@ -1461,107 +1265,201 @@ void UNetConnection::ReceivedPacket( FBitReader& Reader )
 	LastReceiveTime		= Driver->Time;
 	LastReceiveRealtime = FPlatformTime::Seconds();
 
-	if (InternalAck)
+	// Check packet ordering.
+	const int32 PacketId = InternalAck ? InPacketId + 1 : MakeRelative(Reader.ReadInt(MAX_PACKETID),InPacketId,MAX_PACKETID);
+	if( PacketId > InPacketId )
 	{
-		++InPacketId;
-	}	
+		const int32 PacketsLost = PacketId - InPacketId - 1;
+		
+		if ( PacketsLost > 10 )
+		{
+			UE_LOG( LogNetTraffic, Log, TEXT( "High single frame packet loss. PacketsLost: %i %s" ), PacketsLost, *Describe() );
+		}
+
+		InPacketsLost += PacketsLost;
+		InTotalPacketsLost += PacketsLost;
+		Driver->InPacketsLost += PacketsLost;
+		Driver->InTotalPacketsLost += PacketsLost;
+		InPacketId = PacketId;
+	}
 	else
 	{
-		// Read packet header
-		FNetPacketNotify::FNotificationHeader Header;
-		if (!ReadPacketHeader(Header, Reader))
+		Driver->InOutOfOrderPackets++;
+		// Protect against replay attacks
+		// We already protect against this for reliable bunches, and unreliable properties
+		// The only bunch we would process would be unreliable RPC's, which could allow for replay attacks
+		// So rather than add individual protection for unreliable RPC's as well, just kill it at the source, 
+		// which protects everything in one fell swoop
+		return;
+	}
+
+	const bool bIgnoreRPCs = Driver->ShouldIgnoreRPCs();
+
+	bool bSkipAck = false;
+
+	// Track channels that were rejected while processing this packet - used to avoid sending multiple close-channel bunches,
+	// which would cause a disconnect serverside
+	TArray<int32> RejectedChans;
+
+	// Disassemble and dispatch all bunches in the packet.
+	while( !Reader.AtEnd() && State!=USOCK_Closed )
+	{
+		// Parse the bunch.
+		int32 StartPos = Reader.GetPosBits();
+		bool IsAck = !!Reader.ReadBit();
+		if ( Reader.IsError() )
 		{
-			CLOSE_CONNECTION_DUE_TO_SECURITY_VIOLATION(this, ESecurityEvent::Malformed_Packet, TEXT("Failed to read PacketHeader"));
+			CLOSE_CONNECTION_DUE_TO_SECURITY_VIOLATION(this, ESecurityEvent::Malformed_Packet, TEXT("Bunch missing ack flag"));
 			return;
 		}
 
-		// Lambda to dispatch delivery notifications, 
-		auto HandlePacketNotification = [&Header, this](FNetPacketNotify::SequenceNumberT AckedSequence, bool bDelivered)
+		// Process the bunch.
+		if( IsAck )
 		{
-			// Increase LastNotifiedPacketId, this is a full packet Id
-			++LastNotifiedPacketId;
+			LastRecvAckTime = Driver->Time;
+
+			// This is an acknowledgment.
+			const int32 AckPacketId = MakeRelative(Reader.ReadInt(MAX_PACKETID),OutAckPacketId,MAX_PACKETID);
 
-			// Sanity check
-			if (FNetPacketNotify::SequenceNumberT(LastNotifiedPacketId) != AckedSequence)
+			if( Reader.IsError() )
 			{
-				CLOSE_CONNECTION_DUE_TO_SECURITY_VIOLATION(this, ESecurityEvent::Malformed_Packet, TEXT("LastNotifiedPacketId != AckedSequence"));
+				CLOSE_CONNECTION_DUE_TO_SECURITY_VIOLATION(this, ESecurityEvent::Malformed_Packet, TEXT("Bunch missing ack"));
 				return;
 			}
 
-			if (bDelivered)
+			double ServerFrameTime = 0;
+
+			// If this is the server, we're reading in the request to send them our frame time
+			// If this is the client, we're reading in confirmation that our request to get frame time from server is granted
+			const bool bHasServerFrameTime = !!Reader.ReadBit();
+
+			if ( !Driver->IsServer() )
 			{
-				ReceivedAck(LastNotifiedPacketId);
+				if ( bHasServerFrameTime )
+				{
+					// As a client, our request was granted, read the frame time
+					uint8 FrameTimeByte	= 0;
+					Reader << FrameTimeByte;
+					ServerFrameTime = ( double )FrameTimeByte / 1000;
+				}
 			}
 			else
 			{
-				ReceivedNak(LastNotifiedPacketId);
-			};
-		};
+				// Server remembers so he can use during SendAck to notify to client of his frame time
+				bLastHasServerFrameTime = bHasServerFrameTime;
+			}
 
-		// Update incoming sequence data and deliver packet notifications
-		// Packet is only accepted if both the incoming sequence number and incoming ack data are valid
-		int PacketSequenceDelta = PacketNotify.Update(Header, HandlePacketNotification);
-		if (PacketSequenceDelta > 0)
-		{
-			const int32 PacketsLost = PacketSequenceDelta - 1;
-		
-			if ( PacketsLost > 10 )
+			uint32 RemoteInKBytesPerSecond = 0;
+			Reader.SerializeIntPacked( RemoteInKBytesPerSecond );
+
+			// Resend any old reliable packets that the receiver hasn't acknowledged.
+			if( AckPacketId>OutAckPacketId )
 			{
-				UE_LOG( LogNetTraffic, Verbose, TEXT( "High single frame packet loss. PacketsLost: %i %s" ), PacketsLost, *Describe() );
+				for (int32 NakPacketId = OutAckPacketId + 1; NakPacketId<AckPacketId; NakPacketId++, OutPacketsLost++, OutTotalPacketsLost++, Driver->OutTotalPacketsLost++)
+				{
+					UE_LOG(LogNetTraffic, Verbose, TEXT("   Received virtual nak %i (%.1f)"), NakPacketId, (Reader.GetPosBits()-StartPos)/8.f );
+					ReceivedNak( NakPacketId );
+				}
+				OutAckPacketId = AckPacketId;
+			}
+			else if( AckPacketId<OutAckPacketId )
+			{
+				//warning: Double-ack logic makes this unmeasurable.
+				//OutOrdAcc++;
 			}
 
-			InPacketsLost += PacketsLost;
-			InTotalPacketsLost += PacketsLost;
-			Driver->InPacketsLost += PacketsLost;
-			Driver->InTotalPacketsLost += PacketsLost;
-			InPacketId += PacketSequenceDelta;
+			// Update ping
+			const int32 Index = AckPacketId & (ARRAY_COUNT(OutLagPacketId)-1);
 
-			// Extra information associated with the header
-			if (!ReadPacketInfo(Reader))
+			if ( OutLagPacketId[Index] == AckPacketId )
 			{
-				CLOSE_CONNECTION_DUE_TO_SECURITY_VIOLATION(this, ESecurityEvent::Malformed_Packet, TEXT("Failed to read PacketHeader"));
-				return;
+				OutLagPacketId[Index] = -1;		// Only use the ack once
+
+#if !UE_BUILD_SHIPPING
+				if ( CVarPingDisplayServerTime.GetValueOnAnyThread() > 0 )
+				{
+					UE_LOG( LogNetTraffic, Warning, TEXT( "ServerFrameTime: %2.2f" ), ServerFrameTime * 1000.0f );
+				}
+#endif
+
+				// use FApp's time because it is set closer to the beginning of the frame - we don't care about the time so far of the current frame to process the packet
+				const double CurrentTime = FApp::GetCurrentTime();
+				const double GameTime	 = ServerFrameTime;
+				const double RTT		 = (CurrentTime - OutLagTime[Index] ) - ( CVarPingExcludeFrameTime.GetValueOnAnyThread() ? GameTime : 0.0 );
+				const double NewLag		 = FMath::Max( RTT, 0.0 );
+
+				if ( OutBytesPerSecondHistory[Index] > 0 )
+				{
+					RemoteSaturation = ( 1.0f - FMath::Min( ( float )RemoteInKBytesPerSecond / ( float )OutBytesPerSecondHistory[Index], 1.0f ) ) * 100.0f;
+				}
+				else
+				{
+					RemoteSaturation = 0.0f;
+				}
+
+				//UE_LOG( LogNet, Warning, TEXT( "Out: %i, InRemote: %i, Saturation: %f" ), OutBytesPerSecondHistory[Index], RemoteInKBytesPerSecond, RemoteSaturation );
+
+				LagAcc += NewLag;
+				LagCount++;
+
+				if (PlayerController != NULL)
+				{
+					PlayerController->UpdatePing(NewLag);
+				}
 			}
-		}
-		else
-		{
-			Driver->InOutOfOrderPackets++;
-			// Protect against replay attacks
-			// We already protect against this for reliable bunches, and unreliable properties
-			// The only bunch we would process would be unreliable RPC's, which could allow for replay attacks
-			// So rather than add individual protection for unreliable RPC's as well, just kill it at the source, 
-			// which protects everything in one fell swoop
-			return;
-		}
-	}
 
-	const bool bIgnoreRPCs = Driver->ShouldIgnoreRPCs();
+			if ( PackageMap != NULL )
+			{
+				PackageMap->ReceivedAck( AckPacketId );
+			}
 
-	bool bSkipAck = false;
+			// Forward the ack to the channel.
+			UE_LOG(LogNetTraffic, Verbose, TEXT("   Received ack %i (%.1f)"), AckPacketId, (Reader.GetPosBits()-StartPos)/8.f );
 
-	// Track channels that were rejected while processing this packet - used to avoid sending multiple close-channel bunches,
-	// which would cause a disconnect serverside
-	TArray<int32> RejectedChans;
+			{
+				SCOPE_CYCLE_COUNTER(Stat_NetConnectionReceivedAcks);
+	
+				for (int32 i = OpenChannels.Num() - 1; i >= 0; i--)
+				{
+					UChannel* const Channel = OpenChannels[i];
+				
+					if (Channel)
+					{
+						if (Channel->OpenPacketId.Last == AckPacketId) // Necessary for unreliable "bNetTemporary" channels.
+						{
+							Channel->OpenAcked = 1;
+						}
+				
+						for (FOutBunch* OutBunch = Channel->OutRec; OutBunch; OutBunch = OutBunch->Next)
+						{
+							if (OutBunch->bOpen)
+							{
+								UE_LOG(LogNet, VeryVerbose, TEXT("Channel %i reset Ackd because open is reliable. "), Channel->ChIndex );
+								Channel->OpenAcked  = 0; // We have a reliable open bunch, don't let the above code set the OpenAcked state,
+													 // it must be set in UChannel::ReceivedAcks to verify all open bunches were received.
+							}
 
-	// Disassemble and dispatch all bunches in the packet.
-	while( !Reader.AtEnd() && State!=USOCK_Closed )
-	{
-		// For demo backwards compatibility, old replays still have this bit
-		if (InternalAck && EngineNetworkProtocolVersion < EEngineNetworkVersionHistory::HISTORY_ACKS_INCLUDED_IN_HEADER)
-		{
-			const bool IsAckDummy = Reader.ReadBit();
+							if (OutBunch->PacketId == AckPacketId)
+							{
+								OutBunch->ReceivedAck = 1;
+							}
+						}
+						Channel->ReceivedAcks(); //warning: May destroy Channel.
+					}
+					else
+					{
+						UE_LOG(LogNet, Warning, TEXT("UNetConnection::ReceivedPacket: null channel in OpenChannels array while processing ack. %s"), *Describe());
+					}
+				}
+			}
 		}
-
-		// Parse the bunch.
-		int32 StartPos = Reader.GetPosBits();
-		
-		// Process Received data
+		else
 		{
 			// Parse the incoming data.
 			FInBunch Bunch( this );
 			int32 IncomingStartPos		= Reader.GetPosBits();
 			uint8 bControl				= Reader.ReadBit();
-			Bunch.PacketId				= InPacketId;
+			Bunch.PacketId				= PacketId;
 			Bunch.bOpen					= bControl ? Reader.ReadBit() : 0;
 			Bunch.bClose				= bControl ? Reader.ReadBit() : 0;
 			Bunch.bDormant				= Bunch.bClose ? Reader.ReadBit() : 0;
@@ -1607,7 +1505,7 @@ void UNetConnection::ReceivedPacket( FBitReader& Reader )
 			else if ( Bunch.bPartial )
 			{
 				// If this is an unreliable partial bunch, we simply use packet sequence since we already have it
-				Bunch.ChSequence = InPacketId;
+				Bunch.ChSequence = PacketId;
 			}
 			else
 			{
@@ -1935,24 +1833,8 @@ void UNetConnection::ReceivedPacket( FBitReader& Reader )
 	if ( !bSkipAck )
 	{
 		LastGoodPacketRealtime = FPlatformTime::Seconds();
-	}
 
-	if( !InternalAck )
-	{
-		// We always call AckSequence even if we are explicitly rejecting the packet as this updates the expected InSeq used to drive future acks.
-		if ( bSkipAck )
-		{
-			// Explicit Nak, we treat this packet as dropped but we still report it to the sending side as quickly as possible
-			PacketNotify.NakSeq( InPacketId );
-		}
-		else
-		{
-			PacketNotify.AckSeq( InPacketId );
-		}
-
-		// We do want to let the other side know about the ack, so even if there are no other outgoing data when we tick the connection we will send an ackpacket.
-		TimeSensitive = 1;
-		bHasDirtyAcks = 1;
+		SendAck(PacketId, true);
 	}
 }
 
@@ -1996,18 +1878,10 @@ int32 UNetConnection::WriteBitsToSendBuffer(
 	// If this is the start of the queue, make sure to add the packet id
 	if ( SendBuffer.GetNumBits() == 0 && !InternalAck )
 	{
-		// Write Packet Header, before sending the packet we will go back and rewrite the data
-		WritePacketHeader(SendBuffer);
-
-		// Also write server RTT and Recived rate
-		WritePacketInfo(SendBuffer);
-	
-		// Update stats for PacketIdBits and ackdata (also including the data used for packet RTT and saturation calculations)
-		int64 BitsWritten = SendBuffer.GetNumBits();
-		NumPacketIdBits += FNetPacketNotify::SequenceNumberT::SeqNumberBits;
-		NumAckBits += BitsWritten - FNetPacketNotify::SequenceNumberT::SeqNumberBits;
-
+		SendBuffer.WriteIntWrapped( OutPacketId, MAX_PACKETID );
 		ValidateSendBuffer();
+
+		NumPacketIdBits += SendBuffer.GetNumBits();
 	}
 
 	// Add the bits to the queue
@@ -2031,6 +1905,9 @@ int32 UNetConnection::WriteBitsToSendBuffer(
 		case EWriteBitsDataType::Bunch:
 			NumBunchBits += SizeInBits + ExtraSizeInBits;
 			break;
+		case EWriteBitsDataType::Ack:
+			NumAckBits += SizeInBits + ExtraSizeInBits;
+			break;
 		default:
 			break;
 	}
@@ -2078,12 +1955,64 @@ TSharedPtr<FObjectReplicator> UNetConnection::CreateReplicatorForNewActorChannel
 
 void UNetConnection::PurgeAcks()
 {
+	for ( int32 i = 0; i < ResendAcks.Num(); i++ )
+	{
+		SendAck( ResendAcks[i], 0 );
+	}
+
+	ResendAcks.Empty(32);
 }
 
+
 void UNetConnection::SendAck(int32 AckPacketId, bool FirstTime/*=1*/)
 {
+	SCOPE_CYCLE_COUNTER(Stat_NetConnectionSendAck);
+
+	ValidateSendBuffer();
+
+	if( !InternalAck )
+	{
+		if( FirstTime )
+		{
+			PurgeAcks();
+			QueuedAcks.Add(AckPacketId);
+		}
+
+		FBitWriter AckData( 32, true );
+
+		AckData.WriteBit( 1 );
+		AckData.WriteIntWrapped(AckPacketId, MAX_PACKETID);
+		
+		const bool bHasServerFrameTime = Driver->IsServer() ? bLastHasServerFrameTime : ( CVarPingExcludeFrameTime.GetValueOnGameThread() > 0 ? true : false );
+
+		AckData.WriteBit( bHasServerFrameTime ? 1 : 0 );
+
+		if ( Driver->IsServer() && bHasServerFrameTime )
+		{
+			uint8 FrameTimeByte = FMath::Min( FMath::FloorToInt( FrameTime * 1000 ), 255 );
+			AckData << FrameTimeByte;
+		}
+
+		// Notify server of our current rate per second at this time
+		uint32 InKBytesPerSecond = InBytesPerSecond / 1024;
+		AckData.SerializeIntPacked( InKBytesPerSecond );
+
+		NETWORK_PROFILER( GNetworkProfiler.TrackSendAck( AckData.GetNumBits(), this ) );
+
+		WriteBitsToSendBuffer( AckData.GetData(), AckData.GetNumBits(), nullptr, 0, EWriteBitsDataType::Ack );
+
+		AllowMerge = false;
+
+		TimeSensitive = 1;
+
+		UE_LOG(LogNetTraffic, Log, TEXT("   Send ack %i"), AckPacketId);
+	}
+
+	++OutTotalAcks;
+	++Driver->OutTotalAcks;
 }
 
+
 int32 UNetConnection::SendRawBunch( FOutBunch& Bunch, bool InAllowMerge )
 {
 	ValidateSendBuffer();
@@ -2095,7 +2024,7 @@ int32 UNetConnection::SendRawBunch( FOutBunch& Bunch, bool InAllowMerge )
 
 	// Build header.
 	SendBunchHeader.Reset();
-
+	SendBunchHeader.WriteBit( 0 );
 	SendBunchHeader.WriteBit( Bunch.bOpen || Bunch.bClose );
 	if( Bunch.bOpen || Bunch.bClose )
 	{
@@ -2535,6 +2464,8 @@ void UNetConnection::Tick()
 	}
 
 	// Flush.
+	PurgeAcks();
+
 	if ( TimeSensitive || (Driver->Time - LastSendTime) > Driver->KeepAliveTime)
 	{
 		bool bHandlerHandshakeComplete = !Handler.IsValid() || Handler->IsFullyInitialized();
