Clear out ObjectReplicator instances without active object to prevent memory growth and performance degradation. This is done for channels still receiving data.

diff --git a/Engine/Source/Runtime/Engine/Private/DataChannel.cpp b/Engine/Source/Runtime/Engine/Private/DataChannel.cpp
--- a/Engine/Source/Runtime/Engine/Private/DataChannel.cpp
+++ b/Engine/Source/Runtime/Engine/Private/DataChannel.cpp
@@ -2559,7 +2559,14 @@ void UActorChannel::ProcessBunch( FInBunch & Bunch )
 
 	for (auto RepComp = ReplicationMap.CreateIterator(); RepComp; ++RepComp)
 	{
-		RepComp.Value()->PostReceivedBunch();
+		TSharedRef<FObjectReplicator>& ObjectReplicator = RepComp.Value();
+		if (ObjectReplicator->GetObject() == nullptr)
+		{
+			RepComp.RemoveCurrent();
+			continue;
+		}
+
+		ObjectReplicator->PostReceivedBunch();
 	}
 
 	// After all properties have been initialized, call PostNetInit. This should call BeginPlay() so initialization can be done with proper starting values.
