Fixed CIS fail.

diff --git a/Engine/Source/Runtime/Engine/Public/Net/NetPacketNotify.h b/Engine/Source/Runtime/Engine/Public/Net/NetPacketNotify.h
--- a/Engine/Source/Runtime/Engine/Public/Net/NetPacketNotify.h
+++ b/Engine/Source/Runtime/Engine/Public/Net/NetPacketNotify.h
@@ -81,7 +81,7 @@ private:
 template <size_t NumBits, typename SequenceType>
 class TSequenceNumber
 {
-	static_assert(TIsSigned<SequenceType>().Value == false, "The base type for sequence numbers must be unsigned");
+	static_assert(TIsSigned<SequenceType>::Value == false, "The base type for sequence numbers must be unsigned");
 
 public:
 	typedef SequenceType SequenceT;
