Clear out AttachmentReplication.AttachComponent when clearing out AttachmentReplication.AttachParent to ensure that it stays in sync on server and client

diff --git a/Engine/Source/Runtime/Engine/Private/Actor.cpp b/Engine/Source/Runtime/Engine/Private/Actor.cpp
--- a/Engine/Source/Runtime/Engine/Private/Actor.cpp
+++ b/Engine/Source/Runtime/Engine/Private/Actor.cpp
@@ -982,6 +982,7 @@ void AActor::PreReplication( IRepChangedPropertyTracker & ChangedPropertyTracker
 {
 	// Attachment replication gets filled in by GatherCurrentMovement(), but in the case of a detached root we need to trigger remote detachment.
 	AttachmentReplication.AttachParent = nullptr;
+	AttachmentReplication.AttachComponent = nullptr;
 
 	GatherCurrentMovement();
 
diff --git a/Engine/Source/Runtime/Engine/Private/ActorReplication.cpp b/Engine/Source/Runtime/Engine/Private/ActorReplication.cpp
--- a/Engine/Source/Runtime/Engine/Private/ActorReplication.cpp
+++ b/Engine/Source/Runtime/Engine/Private/ActorReplication.cpp
@@ -328,6 +328,7 @@ void AActor::GatherCurrentMovement()
 	if (bReplicateMovement || (RootComponent && RootComponent->GetAttachParent()))
 	{
 		AttachmentReplication.AttachParent = nullptr;
+		AttachmentReplication.AttachComponent = nullptr;
 
 		UPrimitiveComponent* RootPrimComp = Cast<UPrimitiveComponent>(GetRootComponent());
 		if (RootPrimComp && RootPrimComp->IsSimulatingPhysics())
