Simple low-overhead implementation of a Circular Queue that can grow for trivially copyable data.

diff --git a/Engine/Source/Runtime/Engine/Private/Net/Tests/Util/ResizableCircularQueueTests.cpp b/Engine/Source/Runtime/Engine/Private/Net/Tests/Util/ResizableCircularQueueTests.cpp
new file mode 100644
--- /dev/null
+++ b/Engine/Source/Runtime/Engine/Private/Net/Tests/Util/ResizableCircularQueueTests.cpp
@@ -0,0 +1,218 @@
+// Copyright 1998-2019 Epic Games, Inc. All Rights Reserved.
+
+#include "Misc/AutomationTest.h"
+#include "Net/Util/ResizableCircularQueue.h"
+
+#if WITH_DEV_AUTOMATION_TESTS
+
+IMPLEMENT_SIMPLE_AUTOMATION_TEST(FResizableCircularQueueTest, "Network.ResizableCircularQueueTest", EAutomationTestFlags::ApplicationContextMask | EAutomationTestFlags::EngineFilter)
+
+struct FResizableCircularQueueTestUtil
+{
+	typedef TResizableCircularQueue<uint32> QueueT;
+
+	static bool VerifyQueueIntegrity(const QueueT& Queue, uint32 ExpectedValueAtFront, uint32 Increment)
+	{
+		bool bSuccess = true;
+		SIZE_T Offset = Queue.Count() - 1;
+	
+		uint32 ExpectedValue = ExpectedValueAtFront;
+
+		// Peek elements in queue at given offset, peek from back to front
+		for (SIZE_T It = 0; It < Queue.Count(); ++It)
+		{
+			bSuccess = bSuccess && (ExpectedValue == Queue.PeekAtOffset(It));
+			ExpectedValue += Increment;
+		}
+
+		return bSuccess;
+	}
+
+	static void OverideHeadAndTail(QueueT& Queue, uint32 Head, uint32 Tail)
+	{
+		Queue.Head = Head;
+		Queue.Tail = Tail;
+	}
+};
+
+bool FResizableCircularQueueTest::RunTest(const FString& Parameters)
+{
+	// Test empty
+	{
+		FResizableCircularQueueTestUtil::QueueT Q(0);
+
+		TestEqual(TEXT("Test empty - Size"), Q.Count(), 0);
+		TestTrue(TEXT("Test empty - IsEmpty"), Q.IsEmpty());
+		TestEqual(TEXT("Test empty - Capacity"), Q.AllocatedCapacity(), 0);
+	}
+
+	// Test Push to Capacity
+	{
+		const SIZE_T ElementsToPush = 8;
+
+		FResizableCircularQueueTestUtil::QueueT Q(ElementsToPush);
+
+		for (SIZE_T It=0; It < ElementsToPush; ++It)
+		{
+			Q.Enqueue(It);
+		}
+		
+		TestEqual(TEXT("Test Push to Capacity - Size"), Q.Count(), ElementsToPush);
+		TestEqual(TEXT("Test Push to Capacity - Capacity"), Q.AllocatedCapacity(), ElementsToPush);
+ 		TestTrue(TEXT("Test Push to Capacity - Expected"), FResizableCircularQueueTestUtil::VerifyQueueIntegrity(Q, 0, 1));
+	}
+
+	// Test Push over Capacity
+	{
+		const SIZE_T ElementsToPush = 32;
+
+		FResizableCircularQueueTestUtil::QueueT Q(ElementsToPush);
+
+		for (SIZE_T It=0; It < ElementsToPush; ++It)
+		{
+			Q.Enqueue(It);
+		}
+		
+		TestEqual(TEXT("Test Push over Capacity - Size"), Q.Count(), ElementsToPush);
+		TestEqual(TEXT("Test Push over Capacity - Capacity"), Q.AllocatedCapacity(), ElementsToPush);
+ 		TestTrue(TEXT("Test Push over Capacity - Expected"), FResizableCircularQueueTestUtil::VerifyQueueIntegrity(Q, 0, 1));
+	}
+
+	// Test Push and Pop
+	{
+		const SIZE_T ElementsToPush = 256;
+		const SIZE_T ElementPopMod = 16;
+		const SIZE_T ExpectedSize = 256 - ElementPopMod;
+		const SIZE_T ExpectedCapacity = 256;
+
+		FResizableCircularQueueTestUtil::QueueT Q(4);
+
+		uint32 ExpectedPoppedValue = 0;
+		for (uint32 It = 0; It < 256; ++It)
+		{
+			Q.Enqueue(It);
+			TestEqual(TEXT("Test Push and pop - Push"), It, Q.PeekAtOffset(Q.Count() - 1));
+
+			if (It % ElementPopMod == 0)
+			{
+				const uint32 PoppedValue = Q.PeekAtOffset(0);
+				TestEqual(TEXT("Test Push and pop - Pop"), ExpectedPoppedValue, PoppedValue);
+				++ExpectedPoppedValue;
+				Q.Pop();
+			}
+		}
+
+		TestEqual(TEXT("Test Push and pop - Size"), Q.Count(), ExpectedSize);
+		TestEqual(TEXT("Test Push and pop - Capacity"), Q.AllocatedCapacity(), ExpectedCapacity);
+ 		TestTrue (TEXT("Test Push and pop - Expected"), FResizableCircularQueueTestUtil::VerifyQueueIntegrity(Q, ElementPopMod, 1));	
+	}
+
+	// Test Push and pop all
+	{
+		const SIZE_T ElementsToPush = 256;
+
+		FResizableCircularQueueTestUtil::QueueT Q(ElementsToPush);
+
+		TestTrue(TEXT ("Test Push and pop all - IsEmpty before"), Q.IsEmpty());
+		TestEqual(TEXT("Test Push and pop all - Size before"), Q.Count(), 0);
+
+		for (SIZE_T It=0; It < ElementsToPush; ++It)
+		{
+			Q.Enqueue(It);
+		}
+		
+		TestEqual(TEXT("Test Push and pop all - Size"), Q.Count(), ElementsToPush);
+		TestEqual(TEXT("Test Push and pop all - Capacity"), Q.AllocatedCapacity(), ElementsToPush);
+ 		TestTrue (TEXT("Test Push and pop all - Expected"), FResizableCircularQueueTestUtil::VerifyQueueIntegrity(Q, 0, 1));
+
+		for (SIZE_T It=0; It < ElementsToPush; ++It)
+		{
+			Q.Pop();
+		}
+		
+		TestTrue( TEXT("Test Push and pop all - IsEmpty after"), Q.IsEmpty());
+		TestEqual(TEXT("Test Push and pop all - Size after"), Q.Count(), 0);
+		TestEqual(TEXT("Test Push and pop all - Capacity after"), Q.AllocatedCapacity(), ElementsToPush);
+	}
+
+	// Test index wrap
+	{
+		const SIZE_T ElementsToPush = 256;
+		const SIZE_T ElementPopMod = 16;
+		const SIZE_T ExpectedSize = 256 - ElementPopMod;
+		const SIZE_T ExpectedCapacity = 256;
+
+		FResizableCircularQueueTestUtil::QueueT Q(4);
+
+		// Set head and tail at the end of the space to test index wraparound
+		FResizableCircularQueueTestUtil::OverideHeadAndTail(Q, uint32(-2), uint32(-2));
+
+		TestTrue(TEXT ("Test index wrap - IsEmpty before"), Q.IsEmpty());
+		TestEqual(TEXT("Test index wrap - Size before"), Q.Count(), 0);
+
+		for (SIZE_T It=0; It < ElementsToPush; ++It)
+		{
+			Q.Enqueue(It);
+		}
+		
+		TestEqual(TEXT("Test index wrap - Size"), Q.Count(), ElementsToPush);
+		TestEqual(TEXT("Test index wrap - Capacity"), Q.AllocatedCapacity(), ElementsToPush);
+ 		TestTrue (TEXT("Test index wrap - Expected"), FResizableCircularQueueTestUtil::VerifyQueueIntegrity(Q, 0, 1));
+
+		for (SIZE_T It=0; It < ElementsToPush; ++It)
+		{
+			Q.Pop();
+		}
+		
+		TestTrue( TEXT("Test index wrap - IsEmpty after"), Q.IsEmpty());
+		TestEqual(TEXT("Test index wrap - Size after"), Q.Count(), 0);
+		TestEqual(TEXT("Test index wrap - Capacity after"), Q.AllocatedCapacity(), ElementsToPush);
+	}
+
+	// Test Trim
+	{
+		const SIZE_T ElementsToPush = 9;
+		const SIZE_T ElementsToPop = 5;
+		const SIZE_T ExpectedCapacity = 16;
+		const SIZE_T ExpectedCapacityAfterTrim = 4;
+
+		FResizableCircularQueueTestUtil::QueueT Q(0);
+
+		for (SIZE_T It=0; It < ElementsToPush; ++It)
+		{
+			Q.Enqueue(It);
+		}
+		
+		TestEqual(TEXT("Test Trim - Size"), Q.Count(), ElementsToPush);
+		TestEqual(TEXT("Test Trim - Capacity"), Q.AllocatedCapacity(), ExpectedCapacity);
+ 		TestTrue(TEXT("Test Trim - Expected"), FResizableCircularQueueTestUtil::VerifyQueueIntegrity(Q, 0, 1));
+
+		for (SIZE_T It=0; It < ElementsToPop; ++It)
+		{
+			Q.Pop();
+		}
+
+		Q.Trim();
+
+		TestEqual(TEXT("Test Trim - Size"), Q.Count(), ElementsToPush - ElementsToPop);
+		TestEqual(TEXT("Test Trim - Capacity"), Q.AllocatedCapacity(), ExpectedCapacityAfterTrim);
+ 		TestTrue(TEXT("Test Trim - Expected"), FResizableCircularQueueTestUtil::VerifyQueueIntegrity(Q, ElementsToPop, 1));
+	}
+
+	// Test Trim empty
+	{
+		FResizableCircularQueueTestUtil::QueueT Q(0);
+
+		Q.Trim();
+
+		TestEqual(TEXT("Test trim empty - Size"), Q.Count(), 0);
+		TestEqual(TEXT("Test trim empty - Capacity"), Q.AllocatedCapacity(), 0);
+	}
+
+
+	check(!HasAnyErrors());
+
+	return true;
+}
+
+#endif // #if WITH_DEV_AUTOMATION_TESTS
\ No newline at end of file
diff --git a/Engine/Source/Runtime/Engine/Public/Net/Util/ResizableCircularQueue.h b/Engine/Source/Runtime/Engine/Public/Net/Util/ResizableCircularQueue.h
new file mode 100644
--- /dev/null
+++ b/Engine/Source/Runtime/Engine/Public/Net/Util/ResizableCircularQueue.h
@@ -0,0 +1,207 @@
+// Copyright 1998-2019 Epic Games, Inc. All Rights Reserved.
+
+#pragma once
+
+#include "CoreTypes.h"
+#include "Misc/AssertionMacros.h"
+#include "Containers/Array.h"
+
+/**
+ * Simple ResisableCircularQueue for trivially copyable types
+ * Relies on unsigned arithmetics and ever increasing head and tail indices to avoid having to store 
+ * an extra element or maintain explicit empty state.
+ * Capacity most be power of two and it supports growing
+ */
+class FDefaultAllocator;
+
+template<typename T, typename AllocatorT = FDefaultAllocator>
+class TResizableCircularQueue
+{
+public:
+	typedef T ElementT;
+
+	/** Construct Empty Queue with the given initial Capacity, Capacity must be a power of two since we rely on unsigned arithmetics for wraparound */
+	explicit TResizableCircularQueue(SIZE_T InitialCapacity);
+
+	/** Returns true if the queue is empty */
+	bool IsEmpty() const { return Head == Tail; }
+
+	/** Gets the number of elements in the queue. */
+	SIZE_T Count() const { return Head - Tail; }
+
+	/** Current allocated Capacity */
+	SIZE_T AllocatedCapacity() const { return Storage.Num(); }
+
+	/** Push single element to the back of the Queue */
+	void Enqueue(const ElementT& SrcData);
+
+	/** Pop elements from the front of the queue */
+	void Pop();
+
+	/** Pop Count elements from the front of the queue */
+	void Pop(SIZE_T Count);
+
+	/**  Unchecked version, Pop a single element from the front of the queue */
+	void PopNoCheck() { ++Tail; }
+
+	/** Unchecked version, Pop Count elements from the front of the queue */
+	void PopNoCheck(SIZE_T Count) { Tail += Count; }
+
+	/** Peek with the given offset from the front of the queue */
+	const ElementT& PeekAtOffset(SIZE_T Offset = 0) const { check(Offset < Count()); return PeekAtOffsetNoCheck(Offset); }
+
+	/** Peek at front element */
+	const ElementT& Peek() const { return PeekAtOffset(0); }
+
+	/**  Unchecked version, Peek with the given offset from the front of the queue */
+	const ElementT& PeekAtOffsetNoCheck(SIZE_T Offset = 0) const { return Storage.GetData()[(Tail + Offset) & IndexMask]; }
+
+	/** Peek at front element with no check */
+	const ElementT& PeekNoCheck() const { return PeekAtOffsetNoCheck(0); }
+
+	/**  Trim memory usage to the next power of two for the current size */
+	void Trim();
+
+	/** Empty queue without releasing memory */
+	void Reset() { Head = Tail = 0u; }
+
+	/** Empty queue and release memory */
+	void Empty() { Reset(); Storage.Empty(); IndexMask = IndexT(-1); }
+	
+private:
+
+#if WITH_DEV_AUTOMATION_TESTS
+	friend struct FResizableCircularQueueTestUtil;
+#endif
+
+	// Resize buffer maintaining validity of stored data.
+	void SetCapacity(SIZE_T NewCapacity);
+
+	// TIsTriviallyCopyable (a.k.a. std::is_trivially_copyable) should be used but
+	// TIsTriviallyCopyable is not provided by the core module at the moment. Core's
+	// implementation of TIsTrivial is actually equivalent to std::is_trivially_copyable
+	// since std::is_trivial<T> requires T to have a trivial constructor but TIsTrivial
+	// doesn't
+	static_assert(TIsTrivial<ElementT>::Value, "ResizableCircularQueue only supports trivially copyable elements");
+
+	typedef uint32 IndexT;
+	typedef TArray<ElementT, AllocatorT> StorageT;
+
+	IndexT Head;
+	IndexT Tail;
+
+	IndexT IndexMask;
+	StorageT Storage;
+};
+
+
+template<typename T, typename AllocatorT>
+TResizableCircularQueue<T, AllocatorT>::TResizableCircularQueue(SIZE_T InitialCapacity)
+: Head(0u)
+, Tail(0u)
+, IndexMask(0u)
+{
+	// Capacity should be power of two, it will be rounded up but lets warn the user anyway debug builds.
+	checkSlow(InitialCapacity & (InitialCapacity - 1) = 0);
+	if (InitialCapacity > 0)
+	{
+		SetCapacity(InitialCapacity);
+	}
+}
+
+template<typename T, typename AllocatorT>
+void TResizableCircularQueue<T, AllocatorT>::Enqueue(const ElementT& SrcData)
+{ 
+	const SIZE_T RequiredCapacity = Count() + 1;
+	if (RequiredCapacity > AllocatedCapacity())
+	{
+		// Capacity must be power of two
+		SetCapacity(RequiredCapacity);
+	}
+
+	T* DstData = Storage.GetData();
+	const IndexT MaskedIndex = Head & IndexMask;
+	DstData[MaskedIndex] = SrcData;
+	++Head;		
+}
+
+template<typename T, typename AllocatorT>
+void
+TResizableCircularQueue<T, AllocatorT>::Pop(SIZE_T PopCount)
+{
+	if (ensure(Count() > PopCount))
+	{
+		PopNoCheck(PopCount);
+	}
+}
+
+template<typename T, typename AllocatorT>
+void TResizableCircularQueue<T, AllocatorT>::Pop()
+{
+	if (ensure(Count() > 0))
+	{
+		PopNoCheck();
+	}
+}
+
+template<typename T, typename AllocatorT>
+void TResizableCircularQueue<T, AllocatorT>::SetCapacity(SIZE_T RequiredCapacity)
+{
+	SIZE_T NewCapacity = FMath::RoundUpToPowerOfTwo(RequiredCapacity);
+
+	if ((NewCapacity == Storage.Num()) || (NewCapacity < Count()))
+	{
+		return;
+	}
+
+	if (Storage.Num() > 0)
+	{
+		StorageT NewStorage;
+		NewStorage.AddUninitialized(NewCapacity);
+		
+		// copy data to new storage
+		const IndexT MaskedTail = Tail & IndexMask;
+		const IndexT MaskedHead = Head & IndexMask;
+
+		const ElementT* SrcData = Storage.GetData();
+		const SIZE_T SrcCapacity = Storage.Num();
+		const SIZE_T SrcSize = Count();
+		ElementT* DstData = NewStorage.GetData();
+
+		if (MaskedTail >= MaskedHead)
+		{
+			const SIZE_T CopyCount = (SrcCapacity - MaskedTail);
+			FPlatformMemory::Memcpy(DstData, SrcData + MaskedTail, CopyCount * sizeof(ElementT));
+			FPlatformMemory::Memcpy(DstData + CopyCount, SrcData, MaskedHead * sizeof(ElementT));
+		}
+		else
+		{
+			FPlatformMemory::Memcpy(DstData, SrcData + MaskedTail, SrcSize * sizeof(ElementT));
+		}
+
+		this->Storage = MoveTemp(NewStorage);
+		IndexMask = NewCapacity - 1;
+		Tail = 0u;
+		Head = SrcSize;
+	}
+	else
+	{
+		IndexMask = NewCapacity - 1;
+		Storage.AddUninitialized(NewCapacity);
+	}
+}
+
+template<typename T, typename AllocatorT>
+void TResizableCircularQueue<T, AllocatorT>::Trim()
+{
+	if (IsEmpty())
+	{
+		Empty();
+	}
+	else
+	{
+		SetCapacity(Count());
+	}
+}
+
+
