NetPacketNotify Utility used by NetConnection to provide delivery status information for sent packets which is the foundation for implementing reliability

diff --git a/Engine/Source/Runtime/Engine/Private/Net/NetPacketNotify.cpp b/Engine/Source/Runtime/Engine/Private/Net/NetPacketNotify.cpp
--- a/Engine/Source/Runtime/Engine/Private/Net/NetPacketNotify.cpp
+++ b/Engine/Source/Runtime/Engine/Private/Net/NetPacketNotify.cpp
@@ -1,23 +1,32 @@
-// Copyright 1998-2018 Epic Games, Inc. All Rights Reserved.
+// Copyright 1998-2019 Epic Games, Inc. All Rights Reserved.
 
 #include "Net/NetPacketNotify.h"
+#include "Serialization/BitWriter.h"
+#include "Serialization/BitReader.h"
+
+SIZE_T FNetPacketNotify::GetCurrentSequenceHistoryLength() const
+{
+	if (InAckSeq >= InAckSeqAck)
+	{
+		return (SIZE_T)SequenceNumberT::Diff(InAckSeq, InAckSeqAck);
+	}
+	else
+	{
+		// Worst case send full history
+		return SequenceHistoryT::Size;
+	}
+}
 
 void FNetPacketNotify::Init(SequenceNumberT InitialInSeq, SequenceNumberT InitialOutSeq)
 {
 	InSeqHistory.Reset();
 	InSeq = InitialInSeq;
 	InAckSeq = InitialInSeq;
+	InAckSeqAck = InitialInSeq;
 	OutSeq = InitialOutSeq;
 	OutAckSeq = SequenceNumberT(InitialOutSeq.Get() - 1);
 }
 
-void FNetPacketNotify::GetHeader(FNotificationHeader& DataOut) const
-{
-	DataOut.Seq = OutSeq;
-	DataOut.AckedSeq = InAckSeq;
-	DataOut.History = InSeqHistory;
-}
-
 void FNetPacketNotify::AckSeq(SequenceNumberT AckedSeq, bool IsAck)
 {
 	check( AckedSeq == InSeq);
@@ -34,29 +43,97 @@ void FNetPacketNotify::AckSeq(SequenceNumberT AckedSeq, bool IsAck)
 	}
 }
 
+namespace 
+{
+	struct FPackedHeader
+	{
+		using SequenceNumberT = FNetPacketNotify::SequenceNumberT;
+
+		static_assert(FNetPacketNotify::SequenceNumberBits <= 14, "SequenceNumbers must be smaller than 14 bits to fit history word count");
+
+		enum { HistoryWordCountBits = 4 };
+		enum { SeqMask				= (1 << FNetPacketNotify::SequenceNumberBits) - 1 };
+		enum { HistoryWordCountMask	= (1 << HistoryWordCountBits) - 1 };
+		enum { AckSeqShift			= HistoryWordCountBits };
+		enum { SeqShift				= AckSeqShift + FNetPacketNotify::SequenceNumberBits };
+		
+		static uint32 Pack(SequenceNumberT Seq, SequenceNumberT AckedSeq, SIZE_T HistoryWordCount)
+		{
+			uint32 Packed = 0u;
+
+			Packed |= Seq.Get() << SeqShift;
+			Packed |= AckedSeq.Get() << AckSeqShift;
+			Packed |= HistoryWordCount & HistoryWordCountMask;
+
+			return Packed;
+		}
+
+		static SequenceNumberT GetSeq(uint32 Packed) { return SequenceNumberT(Packed >> SeqShift & SeqMask); }
+		static SequenceNumberT GetAckedSeq(uint32 Packed) { return SequenceNumberT(Packed >> AckSeqShift & SeqMask); }
+		static SIZE_T GetHistoryWordCount(uint32 Packed) { return (Packed & HistoryWordCountMask); }
+	};
+}
+
 // These methods must always write and read the exact same number of bits, that is the reason for not using WriteInt/WrittedWrappedInt
-void FNetPacketNotify::WriteHeader(FBitWriter& Writer, const FNotificationHeader& Data)
+void FNetPacketNotify::WriteHeader(FBitWriter& Writer, bool bRefresh)
 {
-	SequenceNumberT::SequenceT Seq = Data.Seq.Get();
-	SequenceNumberT::SequenceT AckedSeq = Data.AckedSeq.Get();
-	
-	Writer << Seq;
-	Writer << AckedSeq;
-	Data.History.Write(Writer);
+	// we always write at least 1 word
+	SIZE_T CurrentHistoryWordCount = FMath::Clamp<SIZE_T>((GetCurrentSequenceHistoryLength() + SequenceHistoryT::BitsPerWord - 1u) / SequenceHistoryT::BitsPerWord, 1u, SequenceHistoryT::WordCount);
+
+	// We can only do a refresh if we do not need more space for the history
+	if (bRefresh && (CurrentHistoryWordCount > WrittenHistoryWordCount))
+	{
+		return;
+	}
+
+	// How many words of ack data should we write? If this is a refresh we must write the same size as the original header
+	WrittenHistoryWordCount = bRefresh ? WrittenHistoryWordCount : CurrentHistoryWordCount;
+	// This is the last InAck we have acknowledged at this time
+	WrittenInAckSeq = InAckSeq;
+
+	SequenceNumberT::SequenceT Seq = OutSeq.Get();
+	SequenceNumberT::SequenceT AckedSeq = InAckSeq.Get();
+
+	// Pack data into a uint
+	uint32 PackedHeader = FPackedHeader::Pack(Seq, AckedSeq, WrittenHistoryWordCount - 1);
+
+	// Write packed header
+	Writer << PackedHeader;
+
+	// Write ack history
+	InSeqHistory.Write(Writer, WrittenHistoryWordCount);
+
+	UE_LOG_PACKET_NOTIFY(TEXT("FNetPacketNotify::WriteHeader - Seq %u, AckedSeq %u bReFresh %u HistorySizeInWords %u"), Seq, AckedSeq, bRefresh ? 1u : 0u, WrittenHistoryWordCount);
 }
 
-bool FNetPacketNotify::ReadHeader(FNotificationHeader& Data, FBitReader& Reader)
+bool FNetPacketNotify::ReadHeader(FNotificationHeader& Data, FBitReader& Reader) const
 {
-	SequenceNumberT::SequenceT RcvdSeq;
-	SequenceNumberT::SequenceT RcvdAckedSeq;
+	// Read packed header
+	uint32 PackedHeader = 0;	
+	Reader << PackedHeader;
 
-	Reader << RcvdSeq;
-	Reader << RcvdAckedSeq;
-	Data.History.Read(Reader);
+	// unpack
+	Data.Seq = FPackedHeader::GetSeq(PackedHeader);
+	Data.AckedSeq = FPackedHeader::GetAckedSeq(PackedHeader);
+	Data.HistoryWordCount = FPackedHeader::GetHistoryWordCount(PackedHeader) + 1;
 
-	Data.Seq = SequenceNumberT(RcvdSeq);
-	Data.AckedSeq = SequenceNumberT(RcvdAckedSeq);
+	// Read ack history
+	Data.History.Read(Reader, Data.HistoryWordCount);
+
+	UE_LOG_PACKET_NOTIFY(TEXT("FNetPacketNotify::ReadHeader - Seq %u, AckedSeq %u HistorySizeInWords %u"), Data.Seq.Get(), Data.AckedSeq.Get(), Data.HistoryWordCount);
 
 	return Reader.IsError() == false;
 }
 
+FNetPacketNotify::SequenceNumberT FNetPacketNotify::CommitAndIncrementOutSeq()
+{
+	// we have not written a header...this is a fail.
+	check(WrittenHistoryWordCount != 0);
+
+	// Add entry to the ack-record so that we can update the InAckSeqAck when we received the ack for this OutSeq.
+	AckRecord.Enqueue( {OutSeq, WrittenInAckSeq} );
+	WrittenHistoryWordCount = 0u;
+	
+	return ++OutSeq;
+}
+
diff --git a/Engine/Source/Runtime/Engine/Private/Net/Tests/NetPacketNotifyTests.cpp b/Engine/Source/Runtime/Engine/Private/Net/Tests/NetPacketNotifyTests.cpp
--- a/Engine/Source/Runtime/Engine/Private/Net/Tests/NetPacketNotifyTests.cpp
+++ b/Engine/Source/Runtime/Engine/Private/Net/Tests/NetPacketNotifyTests.cpp
@@ -1,3 +1,5 @@
+// Copyright 1998-2019 Epic Games, Inc. All Rights Reserved.
+
 #include "Misc/AutomationTest.h"
 #include "Net/NetPacketNotify.h"
 
@@ -8,19 +10,32 @@ IMPLEMENT_SIMPLE_AUTOMATION_TEST(FNetPacketNotifyTest, "Network.PacketNotifyTest
 struct FNetPacketNotifyTestUtil
 {
 	FNetPacketNotify DefaultNotify;
-
 	FNetPacketNotifyTestUtil()
 	{
 		DefaultNotify.Init(FNetPacketNotify::SequenceNumberT(-1),  FNetPacketNotify::SequenceNumberT(0));
 	}
 
-	// Pretend to receive incoming packet to generate ackdata
+	// Helper to fill in SequenceHistory with expected result
+	template <typename T>
+	static void InitHistory(FNetPacketNotify::SequenceHistoryT& History, const T& DataToSet)
+	{
+		const SIZE_T Count = sizeof(T) / sizeof(DataToSet[0]);
+		static_assert(Count < FNetPacketNotify::SequenceHistoryT::WordCount, "DataToSet must be smaller than HistoryBuffer");
+
+		for (SIZE_T It=0; It < Count; ++It)
+		{
+			History.Data()[It] = DataToSet[It];
+		}
+	}
+
+	// Pretend to receive and acknowledge incoming packet to generate ackdata
 	static int32 PretendReceiveSeq(FNetPacketNotify& PacketNotify, FNetPacketNotify::SequenceNumberT Seq, bool Ack = true)
 	{
 		FNetPacketNotify::FNotificationHeader Data;
 		Data.Seq = Seq;
 		Data.AckedSeq = PacketNotify.GetOutAckSeq();
 		Data.History = FNetPacketNotify::SequenceHistoryT(0);
+		Data.HistoryWordCount = 1;
 		
 		FNetPacketNotify::SequenceNumberT::DifferenceT SeqDelta = PacketNotify.Update(Data, [](FNetPacketNotify::SequenceNumberT AckedSequence, bool delivered) {});
 		if (SeqDelta > 0)
@@ -32,38 +47,38 @@ struct FNetPacketNotifyTestUtil
 		return SeqDelta;
 	}
 
-	static bool VerifyNotificaitonState(const FNetPacketNotify& A, const FNetPacketNotify& B)
+	// Pretend to send packet
+	static void PretendSendSeq(FNetPacketNotify& PacketNotify, FNetPacketNotify::SequenceNumberT LastAckSeq)
 	{
-		bool bEquals = 
-			A.GetInSeq() == B.GetInSeq() &&
-			A.GetInSeqHistory() == B.GetInSeqHistory() &&
-			A.GetOutSeq() == B.GetOutSeq() &&
-			A.GetOutAckSeq() == B.GetOutAckSeq();
+		// set last InAcqSeq that we know that the remote end knows that we know (AckAck)
+		PacketNotify.WrittenHistoryWordCount = 1;
+		PacketNotify.WrittenInAckSeq = LastAckSeq;
 
-		return bEquals;
+		// Store data
+		PacketNotify.CommitAndIncrementOutSeq();
 	}
 
 	// pretend to ack array of sequence numbers
 	template<typename T>
 	static void PretendAckSequenceNumbers(FNetPacketNotify& PacketNotify, const T& InSequenceNumbers)
 	{
-		size_t SequenceNumberCount = sizeof(InSequenceNumbers) / sizeof(InSequenceNumbers[0]);
+		SIZE_T SequenceNumberCount = sizeof(InSequenceNumbers) / sizeof(InSequenceNumbers[0]);
 
-		for (size_t I=0; I<SequenceNumberCount; ++I)
+		for (SIZE_T I=0; I<SequenceNumberCount; ++I)
 		{
 			FNetPacketNotifyTestUtil::PretendReceiveSeq(PacketNotify, InSequenceNumbers[I]);
 		}
 	}
 	
-	// pretend to deliver notifications for incoming seq header
+	// Pretend that we received a packet
 	template<typename T>
-	static size_t PretendDeliverNotifications(FNetPacketNotify& PacketNotify, const FNetPacketNotify::FNotificationHeader Data, T& OutSequenceNumbers)
+	static SIZE_T PretendReceivedPacket(FNetPacketNotify& PacketNotify, const FNetPacketNotify::FNotificationHeader Data, T& OutSequenceNumbers)
 	{
-		size_t NotificationCount = 0;
+		SIZE_T NotificationCount = 0;
 
 		auto HandleAck = [&OutSequenceNumbers, &NotificationCount](FNetPacketNotify::SequenceNumberT Seq, bool delivered)
 		{
-			const size_t MaxSequenceNumberCount = sizeof(OutSequenceNumbers) / sizeof(OutSequenceNumbers[0]);
+			const SIZE_T MaxSequenceNumberCount = sizeof(OutSequenceNumbers) / sizeof(OutSequenceNumbers[0]);
 
 			if (delivered)
 			{
@@ -77,36 +92,54 @@ struct FNetPacketNotifyTestUtil
 		return PacketNotify.Update(Data, HandleAck);
 	}
 
+	// Test to fake sending and receiving an array of sequence numbers and test if we get the expected notifications back
 	template<typename T>
 	static bool TestNotificationSequence(const T& InSequenceNumbers, FNetPacketNotify::SequenceNumberT FirstSequence = 0)
 	{
 		T NotifiedSequenceNumbers = { 0 };
-		//FPlatformMemory::Memset(&NotifiedSequenceNumbers[0], 0, sizeof(NotifiedSequenceNumbers));
 
-		FNetPacketNotify Acked;
-		Acked.Init(FNetPacketNotify::SequenceNumberT(FirstSequence.Get() - 1), FirstSequence);
-		PretendAckSequenceNumbers(Acked, InSequenceNumbers);
+		// Sender, which we will also be the receiver of the acks
+		FNetPacketNotify Sender;
+		Sender.Init(FNetPacketNotify::SequenceNumberT(FirstSequence.Get() - 1), FirstSequence);
+	
+		// pretend that we have sent the InSequenceNumbers
+		SIZE_T Count = sizeof(T) / sizeof(InSequenceNumbers[0]);
+		for (SIZE_T It=0; It < Count; ++It)
+		{			
+			do
+			{
+				FNetPacketNotifyTestUtil::PretendSendSeq(Sender, 0);
+			}
+			while (InSequenceNumbers[It] >= Sender.GetOutSeq());
+		}
 
-		// In order to be able to accepts the acks we must pretend that we have sent a packet which we will get an ack for.
-		FNetPacketNotify Notified;
-		Notified.Init(FNetPacketNotify::SequenceNumberT(FirstSequence.Get() - 1), FirstSequence);
+		// Receiver which we fake have received the packets sent from sender
+		FNetPacketNotify Receiver;
+		Receiver.Init(FNetPacketNotify::SequenceNumberT(FirstSequence.Get() - 1), FirstSequence);
+		PretendAckSequenceNumbers(Receiver, InSequenceNumbers);
+
+		// Fake Header with acks sent from receiver back to sender
 		FNetPacketNotify::FNotificationHeader Data;
-		Acked.GetHeader(Data);
+		Data.Seq = Receiver.GetOutSeq();
+		Data.AckedSeq = Receiver.GetInAckSeq();
+		Data.HistoryWordCount = FNetPacketNotify::SequenceHistoryT::WordCount;
+		Data.History = Receiver.GetInSeqHistory();
 
-		PretendDeliverNotifications(Notified, Data, NotifiedSequenceNumbers);
+		// Process the received ack information
+		PretendReceivedPacket(Sender, Data, NotifiedSequenceNumbers);
 
+		// Check that it matches the expected result
 		return FPlatformMemory::Memcmp(&InSequenceNumbers[0], &NotifiedSequenceNumbers[0], sizeof(InSequenceNumbers)) == 0;
 	}
 };
 
-
 bool FNetPacketNotifyTest::RunTest(const FString& Parameters)
 {
 	FNetPacketNotifyTestUtil Util;
 	// Test fill
 	{
 		FNetPacketNotify::SequenceNumberT ExpectedInSeq(31);
-		FNetPacketNotify::SequenceHistoryT ExpectedInSeqHistory(0xffffffffu);
+		FNetPacketNotify::SequenceHistoryT ExpectedInSeqHistory(0xffffffffu, 1);
 
 		FNetPacketNotify Acks = Util.DefaultNotify;
 				
@@ -122,7 +155,7 @@ bool FNetPacketNotifyTest::RunTest(const FString& Parameters)
 	// Test drop every other
 	{
 		FNetPacketNotify::SequenceNumberT ExpectedInSeq(30);
-		FNetPacketNotify::SequenceHistoryT ExpectedInSeqHistory(0x55555555u);
+		FNetPacketNotify::SequenceHistoryT ExpectedInSeqHistory(0x55555555u, 1);
 
 		FNetPacketNotify Acks = Util.DefaultNotify;
 				
@@ -138,7 +171,9 @@ bool FNetPacketNotifyTest::RunTest(const FString& Parameters)
 	// Test burst drop
 	{
 		FNetPacketNotify::SequenceNumberT ExpectedInSeq(128);
-		FNetPacketNotify::SequenceHistoryT ExpectedInSeqHistory(0x1u);
+		FNetPacketNotify::SequenceHistoryT ExpectedInSeqHistory;
+		uint32 ExpectedArray[] = {0x1, 0, 0, 0x20000000 };
+		FNetPacketNotifyTestUtil::InitHistory(ExpectedInSeqHistory, ExpectedArray );
 
 		FNetPacketNotify Acks = Util.DefaultNotify;
 
@@ -183,7 +218,7 @@ bool FNetPacketNotifyTest::RunTest(const FString& Parameters)
 		const FNetPacketNotify::SequenceHistoryT ExpectedInSeqHistory(0x8853u);
 
 		const FNetPacketNotify::SequenceNumberT AckdPacketIds[] = {3, 7, 12, 14, 17, 18};
-		const size_t ExpectedCount = sizeof(AckdPacketIds)/sizeof((AckdPacketIds)[0]);		
+		const SIZE_T ExpectedCount = sizeof(AckdPacketIds)/sizeof((AckdPacketIds)[0]);		
 
 		FNetPacketNotify Acks = Util.DefaultNotify;
 		FNetPacketNotifyTestUtil::PretendAckSequenceNumbers(Acks, AckdPacketIds);
@@ -191,15 +226,14 @@ bool FNetPacketNotifyTest::RunTest(const FString& Parameters)
 		TestEqual(TEXT("Create history - InSeq"), Acks.GetInSeq(), ExpectedInSeq);
 		TestEqual(TEXT("Create history - InSeqHistory"), Acks.GetInSeqHistory(), ExpectedInSeqHistory);
 	}
-
 	
 	// test notifications
 	{
 		static const FNetPacketNotify::SequenceNumberT ExpectedAckdPacketIds[] = {3, 7, 12, 14, 17, 18};
-		static const size_t ExpectedCount = sizeof(ExpectedAckdPacketIds)/sizeof((ExpectedAckdPacketIds)[0]);		
+		static const SIZE_T ExpectedCount = sizeof(ExpectedAckdPacketIds)/sizeof((ExpectedAckdPacketIds)[0]);		
 
 		FNetPacketNotify::SequenceNumberT RcvdAcks[ExpectedCount] = { 0 };
-		size_t RcvdCount = 0;
+		SIZE_T RcvdCount = 0;
 
 		// Create src data
 		FNetPacketNotify Acks = Util.DefaultNotify;
@@ -209,8 +243,15 @@ bool FNetPacketNotifyTest::RunTest(const FString& Parameters)
 		Data.Seq = FNetPacketNotify::SequenceNumberT(0);
 		Data.AckedSeq = FNetPacketNotify::SequenceNumberT(18);
 		Data.History = FNetPacketNotify::SequenceHistoryT(0x8853u);
+		Data.HistoryWordCount = 1;
 
-		size_t DeltaSeq = FNetPacketNotifyTestUtil::PretendDeliverNotifications(Acks, Data, RcvdAcks);
+		// Need to fake ack record as well.
+		for (SIZE_T It=0; It <= 18; ++It)
+		{
+			FNetPacketNotifyTestUtil::PretendSendSeq(Acks, 0);
+		}
+	
+		SIZE_T DeltaSeq = FNetPacketNotifyTestUtil::PretendReceivedPacket(Acks, Data, RcvdAcks);
 
 		TestEqual(TEXT("Notifications - Create sequence delta"), DeltaSeq, 1);
 		TestEqual(TEXT("Notifications - Create sequence"), FPlatformMemory::Memcmp(ExpectedAckdPacketIds, RcvdAcks, sizeof(ExpectedAckdPacketIds)), 0u);
@@ -230,12 +271,12 @@ bool FNetPacketNotifyTest::RunTest(const FString& Parameters)
 		TestTrue(TEXT("{2, 3, 4, 5, 6, 7, 8, 9, 10, 12, 13, 14, 16, 17, 18, 19, 20, 21, 22, 23, 24, 25, 26, 27, 28, 29, 30, 31}"), FNetPacketNotifyTestUtil::TestNotificationSequence(TestSeqs));
 	}
 	{
-		static const FNetPacketNotify::SequenceNumberT TestSeqs[] = {0, 31};	
-		TestTrue(TEXT("Test Seq {0, 31}"), FNetPacketNotifyTestUtil::TestNotificationSequence(TestSeqs));
+		static const FNetPacketNotify::SequenceNumberT TestSeqs[] = {0, FNetPacketNotify::MaxSequenceHistoryLength - 1};	
+		TestTrue(TEXT("Test Seq {0, FNetPacketNotify::MaxSequenceHistoryLength - 1}"), FNetPacketNotifyTestUtil::TestNotificationSequence(TestSeqs));
 	}
 	{
-		static const FNetPacketNotify::SequenceNumberT TestSeqs[] = {0, 32};	
-		TestFalse(TEXT("Test Seq {0, 32}"), FNetPacketNotifyTestUtil::TestNotificationSequence(TestSeqs));
+		static const FNetPacketNotify::SequenceNumberT TestSeqs[] = {0, FNetPacketNotify::MaxSequenceHistoryLength};	
+		TestFalse(TEXT("Test Seq {0, FNetPacketNotify::MaxSequenceHistoryLength}"), FNetPacketNotifyTestUtil::TestNotificationSequence(TestSeqs));
 	}
 	{
 		static const FNetPacketNotify::SequenceNumberT TestSeqs[] = {0, FNetPacketNotify::SequenceNumberT::SeqNumberHalf };
@@ -254,41 +295,7 @@ bool FNetPacketNotifyTest::RunTest(const FString& Parameters)
 		TestTrue(TEXT("Test Seq {FNetPacketNotify::SequenceNumberT::SeqNumberMax, 0} From SeqNumberHalf + 1;"), FNetPacketNotifyTestUtil::TestNotificationSequence(TestSeqs, FNetPacketNotify::SequenceNumberT::SeqNumberHalf + 2));
 	}
 
-	// Test sequence numbers
-	{
-		// Valid sequence = 0-7, max distance between sequence number in order to determine order is half the sequence space (0-3)
-		typedef TSequenceNumber<3, uint16> FSequence3;
-
-		for (int32 I = 0; I < FSequence3::SeqNumberCount; ++I)
-		{
-			FSequence3 Seq(I);
-			FSequence3 Ref(I);
-			
-			for (int32 U = 0; U < FSequence3::SeqNumberCount; ++U)
-			{	
-				FSequence3::DifferenceT Diff = FSequence3::Diff(Seq, Ref);
-
-				if (U < FSequence3::SeqNumberHalf)
-				{
-					if (!(Diff == U))
-					{
-						TestTrue(TEXT("SequenceNumbers - Expected Diff "), Diff == U);
-					}
-				}
-				else
-				{
-					if (!(Diff == (U - FSequence3::SeqNumberCount)))
-					{
-						TestTrue(TEXT("SequenceNumbers - Expected Diff"), Diff == (U - FSequence3::SeqNumberCount));
-					}
-				}
-				++Seq;
-			}
-		}
-
-		check(!HasAnyErrors());
-
-	}
+	check(!HasAnyErrors());
 
 	return true;
 }
diff --git a/Engine/Source/Runtime/Engine/Private/Net/Tests/Util/SequenceNumberTests.cpp b/Engine/Source/Runtime/Engine/Private/Net/Tests/Util/SequenceNumberTests.cpp
new file mode 100644
--- /dev/null
+++ b/Engine/Source/Runtime/Engine/Private/Net/Tests/Util/SequenceNumberTests.cpp
@@ -0,0 +1,112 @@
+// Copyright 1998-2019 Epic Games, Inc. All Rights Reserved.
+
+#include "Misc/AutomationTest.h"
+#include "Net/Util/SequenceNumber.h"
+
+#if WITH_DEV_AUTOMATION_TESTS
+
+IMPLEMENT_SIMPLE_AUTOMATION_TEST(FNetSequenceNumberTest, "Network.SequenceNumberTest", EAutomationTestFlags::ApplicationContextMask | EAutomationTestFlags::EngineFilter)
+
+bool FNetSequenceNumberTest::RunTest(const FString& Parameters)
+{
+	typedef TSequenceNumber<3, uint16> FSequence3;
+
+	// Test construct
+	{
+		FSequence3 Seq;
+		TestEqual(TEXT("SequenceNumbers - Seq() == 0"), 0u, Seq.Get());		
+	}
+
+	// Test construct with value
+	{
+		FSequence3 Seq(8);
+		TestEqual(TEXT("SequenceNumbers - Seq(8) == 0"), 0u, Seq.Get());		
+	}
+
+	// Test postIncrement
+	{
+		FSequence3 Seq;
+		Seq++;
+		TestEqual("SequenceNumbers - Seq()++ == 1", 1, Seq.Get());
+	}
+
+	// Test preIncrement
+	{
+		FSequence3 Seq;
+		++Seq;
+		TestEqual("SequenceNumbers - ++Seq == 1", 1, Seq.Get());
+	}
+
+	// Test wraparound
+	{
+		FSequence3 Seq(7);
+		Seq++;
+		TestEqual("SequenceNumbers - Seq(7) + 1 == 0", 0, Seq.Get());
+	}
+
+	// Test equal and not equal
+	{
+		FSequence3 A(2);
+		FSequence3 B(2);
+		FSequence3 C(1);
+
+		TestTrue("SequenceNumbers - operator ==", A == B);
+		TestTrue("SequenceNumbers - operator ==", B == A);
+		TestTrue("SequenceNumbers - operator !=", A != C);
+		TestTrue("SequenceNumbers - operator !=", C != A);
+		TestFalse("SequenceNumbers - operator !=", A != B);
+	}
+
+	// Test operator >
+	{
+		FSequence3 Zero(0);
+		FSequence3 HalfMinusOne(FSequence3::SeqNumberHalf-1);
+		FSequence3 Half(FSequence3::SeqNumberHalf);
+		FSequence3 MaxSeq(FSequence3::SeqNumberCount - 1);
+
+		TestTrue("SequenceNumbers - operator> HalfMinusOne > Zero", HalfMinusOne > Zero);
+		TestTrue("SequenceNumbers - operator> Half > HalfMinusOne", Half > HalfMinusOne);
+		TestFalse("SequenceNumbers - operator> Half > Zero", Half > Zero);
+
+		TestFalse("SequenceNumbers - operator> MaxSeq > Zero", MaxSeq > Zero);
+		TestTrue("SequenceNumbers - operator> Zero > MaxSeq", Zero > MaxSeq);
+		TestFalse("SequenceNumbers - operator> Zero > Zero", Zero > Zero);
+	}
+	
+	// Test diff sequence numbers
+	{		
+		// Valid sequence = 0-7, max distance between sequence number in order to determine order is half the sequence space (0-3)
+		for (int32 I = 0; I < FSequence3::SeqNumberCount; ++I)
+		{
+			FSequence3 Seq(I);
+			FSequence3 Ref(I);
+			
+			for (int32 U = 0; U < FSequence3::SeqNumberCount; ++U)
+			{	
+				FSequence3::DifferenceT Diff = FSequence3::Diff(Seq, Ref);
+
+				if (U < FSequence3::SeqNumberHalf)
+				{
+					if (!(Diff == U))
+					{
+						TestTrue(TEXT("SequenceNumbers - Expected Diff "), Diff == U);
+					}
+				}
+				else
+				{
+					if (!(Diff == (U - FSequence3::SeqNumberCount)))
+					{
+						TestTrue(TEXT("SequenceNumbers - Expected Diff"), Diff == (U - FSequence3::SeqNumberCount));
+					}
+				}
+				++Seq;
+			}
+		}
+
+		check(!HasAnyErrors());
+	}
+
+	return true;
+}
+
+#endif // #if WITH_DEV_AUTOMATION_TESTS
\ No newline at end of file
diff --git a/Engine/Source/Runtime/Engine/Public/Net/NetPacketNotify.h b/Engine/Source/Runtime/Engine/Public/Net/NetPacketNotify.h
--- a/Engine/Source/Runtime/Engine/Public/Net/NetPacketNotify.h
+++ b/Engine/Source/Runtime/Engine/Public/Net/NetPacketNotify.h
@@ -1,11 +1,14 @@
-// Copyright 1998-2018 Epic Games, Inc. All Rights Reserved.
+// Copyright 1998-2019 Epic Games, Inc. All Rights Reserved.
 
 #pragma once
 
+#include "EngineLogs.h"
 #include "HAL/UnrealMemory.h"
+#include "Logging/LogMacros.h"
 #include "Templates/IsSigned.h"
-#include "Serialization/BitWriter.h"
-#include "Serialization/BitReader.h"
+#include "Util/ResizableCircularQueue.h"
+#include "Util/SequenceNumber.h"
+#include "Util/SequenceHistory.h"
 
 #if (UE_BUILD_SHIPPING || UE_BUILD_TEST)
 #	define UE_NET_ENABLE_PACKET_NOTIFY_LOG 0
@@ -14,120 +17,16 @@
 #endif 
 
 #if UE_NET_ENABLE_PACKET_NOTIFY_LOG
-#	include "Logging/LogMacros.h"
-#	include "EngineLogs.h"
 #	define UE_LOG_PACKET_NOTIFY(Format, ...)  UE_LOG(LogNetTraffic, Verbose, Format, ##__VA_ARGS__)
 #else
 #	define UE_LOG_PACKET_NOTIFY(...)
 #endif
 
+#define UE_LOG_PACKET_NOTIFY_WARNING(Format, ...)  UE_LOG(LogNetTraffic, Warning, Format, ##__VA_ARGS__)
+
 struct FBitWriter;
 struct FBitReader;
 
-/** Util class to manage history of received sequence numbers */
-template <size_t HistorySize>
-class TSequenceHistory
-{
-	typedef uint32 WordT;
-	
-	static const size_t BitsPerWord = sizeof(WordT) * 8;
-	static const size_t WordCount = HistorySize / BitsPerWord;
-	
-	static_assert(HistorySize > 0, "Invalid History");
-	static_assert(HistorySize % BitsPerWord == 0, "HistorySize must be a modulo of the wordsize");
-
-public:	
-	TSequenceHistory();
-
-#if WITH_DEV_AUTOMATION_TESTS
-	explicit TSequenceHistory(WordT Value);
-#endif
-	/** Reset */
-	void Reset();
-
-	/** Store delivery status, oldest will be dropped */
-	void AddDeliveryStatus(bool Delivered);
-
-	/** Query the status of a specific index, index 0 is last stored status */
-	bool IsDelivered(size_t Index) const;
-
-	/** Return the size of the history buffer */
-	size_t Size() const { return HistorySize; }	
-	
-	bool operator==(const TSequenceHistory& Other) const
-	{
-		return FMemory::Memcmp(Data(), Other.Data(), WordCount * sizeof (WordT)) == 0;
-	}
-
-	bool operator!=(const TSequenceHistory& Other) const
-	{
-		return FMemory::Memcmp(Data(), Other.Data(), WordCount * sizeof (WordT)) != 0;
-	}
-
-	/** Write history to BitStream */
-	void Write(FBitWriter& Writer) const;
-
-	/** Read history from BitStream */
-	void Read(FBitReader& Reader);
-
-private:
-	const WordT* Data() const { return &Storage[0]; }
-	size_t NumWords() const { return WordCount; }
-
-	WordT Storage[WordCount];
-};
-
-/** Helper class to work with sequence numbers */
-template <size_t NumBits, typename SequenceType>
-class TSequenceNumber
-{
-	static_assert(TIsSigned<SequenceType>::Value == false, "The base type for sequence numbers must be unsigned");
-
-public:
-	typedef SequenceType SequenceT;
-	typedef int32 DifferenceT;
-
-	// Constants
-	enum { SeqNumberBits = NumBits };
-	enum { SeqNumberCount = SequenceT(1) << NumBits };
-	enum { SeqNumberHalf = SequenceT(1) << (NumBits - 1) };
-	enum { SeqNumberMax = SeqNumberCount - 1u };
-	enum { SeqNumberMask = SeqNumberMax };
-
-	/** Default constructor */
-	TSequenceNumber() : Value(0u) {}
-
-	/** Constructor with given value */
-	TSequenceNumber(SequenceT ValueIn) : Value(ValueIn & SeqNumberMask) {}
-	
-	/** Get Current Value */	
-	SequenceT Get() const { return Value; }
-
-	/** Diff between sequence numbers (A - B) only valid if (A - B) < SeqNumberHalf */
-	static DifferenceT Diff(TSequenceNumber A, TSequenceNumber B);
-	
-	/** return true if this is > Other, this is only considered to be the case if (A - B) < SeqNumberHalf since we have to be able to detect wraparounds */
-	bool operator>(const TSequenceNumber& Other) const { return (Value != Other.Value) && (((Value - Other.Value) & SeqNumberMask) < SeqNumberHalf); }
-
-	/** Check if this is >= Other, See above */
-	bool operator>=(const TSequenceNumber& Other) const { return ((Value - Other.Value) & SeqNumberMask) < SeqNumberHalf; }
-
-	/** Equals, NOTE that serial numbers wrap around so 0 == 0 + SequenceNumberCount */
-	bool operator==(const TSequenceNumber& Other) const { return Value == Other.Value; }
-
-	bool operator!=(const TSequenceNumber& Other) const { return Value != Other.Value; }
-
-	/** Pre-increment and wrap around */
-	TSequenceNumber& operator++() { Increment(1u); return *this; }
-	
-	/** Post-increment and wrap around */
-	TSequenceNumber operator++(int) { TSequenceNumber Tmp(*this); Increment(1u); return Tmp; }
-
-private:
-	void Increment(SequenceT InValue) { *this = TSequenceNumber(Value + InValue); }
-	SequenceT Value;
-};
-
 /** 
 	FNetPacketNotify - Drives delivery of sequence numbers, acknowledgments and notifications of delivery sequence numbers
 */
@@ -135,18 +34,25 @@ class FNetPacketNotify
 {
 public:
 	enum { SequenceNumberBits = 14 };
-	enum { SequenceHistoryLength = 32 };
+	enum { MaxSequenceHistoryLength = 256 };
 
 	typedef TSequenceNumber<SequenceNumberBits, uint16> SequenceNumberT;
-	typedef TSequenceHistory<SequenceHistoryLength> SequenceHistoryT;
+	typedef TSequenceHistory<MaxSequenceHistoryLength> SequenceHistoryT;
 
 	struct FNotificationHeader
 	{
 		SequenceHistoryT History;
+		SIZE_T HistoryWordCount;
 		SequenceNumberT Seq;
 		SequenceNumberT AckedSeq;
 	};
 
+	/** Constructor */
+	FNetPacketNotify()
+	: AckRecord(64)
+	, WrittenHistoryWordCount(0)
+	{}
+
 	/** Init notification with expected initial sequence numbers */
 	void Init(SequenceNumberT InitialInSeq, SequenceNumberT InitialOutSeq);
 
@@ -156,17 +62,16 @@ public:
 	/** Explicitly mark Seq as not received and update current InSeq, additional missing sequence numbers will be marked as lost */
 	void NakSeq(SequenceNumberT Seq) { AckSeq(Seq, false); }
 
-	/** Increment outgoing seq number */
-	SequenceNumberT IncrementOutSeq() { return ++OutSeq; }	
-
-	/** Get header that should be written to stream */
-	void GetHeader(FNotificationHeader& DataOut) const;
-
-	/** Write header to stream, note this must always write the exact same number of bits as we go back and update the header */
-	static void WriteHeader(FBitWriter& Writer, const FNotificationHeader& Data);
+	/** Increment outgoing seq number and commit data*/
+	SequenceNumberT CommitAndIncrementOutSeq();
 
+	/** Write NotificationHeader, and update outgoing ack record 
+	 if bRefresh is true we will attempt to refresh a previously written header 
+	if the size remains the same */
+	void WriteHeader(FBitWriter& Writer, bool bRefresh = false);
+	
 	/** Read header from stream */
-	static bool ReadHeader(FNotificationHeader& Data, FBitReader& Reader);
+	bool ReadHeader(FNotificationHeader& Data, FBitReader& Reader) const;
 
 	/** Update state of PacketNotification based on received header and invoke packet notifications for received acks.
 		InFunc is a function in the format void)(FNetPacketNotify::SequenceNumberT AckedSequence, bool delivered) to handle packetNotifications.
@@ -214,13 +119,58 @@ public:
 	/** If we do have more unacknowledged sequence numbers in-flight than our maximum sendwindow we should not send more as the receiving end will not be able to detect if the sequence number has wrapped around */
 	bool CanSend() const { SequenceNumberT NextOutSeq = OutSeq; ++NextOutSeq; return NextOutSeq >= OutAckSeq; }
 
+	/** Get the current sequenceHistory length in bits */
+	SIZE_T GetCurrentSequenceHistoryLength() const;
+
+	SequenceNumberT UpdateInAckSeqAck(SequenceNumberT::DifferenceT AckCount, SequenceNumberT AckedSeq)
+	{
+		check(AckCount <= AckRecord.Count());
+
+		if (AckCount > 1)
+		{			
+			AckRecord.PopNoCheck(AckCount - 1);
+		}
+		
+		FSentAckData AckData = AckRecord.PeekNoCheck();
+		AckRecord.PopNoCheck();
+
+		// verify that we have a matching sequence number
+		if (AckData.OutSeq == AckedSeq)
+		{
+			return AckData.InAckSeq;
+		}
+		else
+		{
+			UE_LOG_PACKET_NOTIFY_WARNING(TEXT("FNetPacketNotify::UpdateInAckSeqAck - Failed to find matching AckRecord for %u, (Found %u)"), AckedSeq.Get(), AckData.OutSeq.Get());
+
+			// Pessimistic view, should never occur
+			return SequenceNumberT(AckedSeq.Get() - MaxSequenceHistoryLength);
+		}
+	}
+
 private:
-	SequenceHistoryT InSeqHistory;	// BitBuffer containing a bitfield describing the history of received packets
-	SequenceNumberT InSeq;			// Last sequence number received from remote
-	SequenceNumberT InAckSeq;		// Last sequence number received from remote that we have acknowledged, this is needed since we support accepting a packet but explicitly not acknowledge it as received.
-	SequenceNumberT OutSeq;			// Outgoing sequence number
-	SequenceNumberT OutAckSeq;		// Last sequence number that we know that the remote side have received.
+	struct FSentAckData
+	{
+		SequenceNumberT OutSeq;	// Not needed... just to verify that things work as expected
+		SequenceNumberT InAckSeq;
+	};
+	typedef TResizableCircularQueue<FSentAckData, TInlineAllocator<128>> AckRecordT;
+
+	AckRecordT AckRecord;				// Track acked seq for each sent packet to track size of ack history
+	SIZE_T WrittenHistoryWordCount;		// Bookkeeping to track if we can update data
+	SequenceNumberT WrittenInAckSeq;	// When we call CommitAndIncrementOutSequence this will be committed along with the current outgoing sequence number for bookkeeping
 
+	// Track incoming sequence data
+	SequenceHistoryT InSeqHistory;		// BitBuffer containing a bitfield describing the history of received packets
+	SequenceNumberT InSeq;				// Last sequence number received and accepted from remote
+	SequenceNumberT InAckSeq;			// Last sequence number received from remote that we have acknowledged, this is needed since we support accepting a packet but explicitly not acknowledge it as received.
+	SequenceNumberT InAckSeqAck;		// Last sequence number received from remote that we have acknowledged and also knows that the remote has received the ack, used to calculate how big our history must be
+
+	// Track outgoing sequence data
+	SequenceNumberT OutSeq;				// Outgoing sequence number
+	SequenceNumberT OutAckSeq;			// Last sequence number that we know that the remote side have received.
+
+private:
 	template<class Functor>
 	void ProcessReceivedAcks(const FNotificationHeader& NotificationData, Functor&& InFunc)
 	{
@@ -230,20 +180,23 @@ private:
 
 			SequenceNumberT::DifferenceT AckCount = SequenceNumberT::Diff(NotificationData.AckedSeq, OutAckSeq);
 
+			// Update InAckSeqAck used to track the needed number of bits to transmit our ack history
+			InAckSeqAck = UpdateInAckSeqAck(AckCount, NotificationData.AckedSeq);
+
 			// ExpectedAck = OutAckSeq + 1
 			SequenceNumberT CurrentAck(OutAckSeq);
 			++CurrentAck;
 
 			// Everything not found in the history buffer is treated as lost
-			while (AckCount > (SequenceNumberT::DifferenceT)(NotificationData.History.Size()))
+			while (AckCount > (SequenceNumberT::DifferenceT)(SequenceHistoryT::Size))
 			{
 				--AckCount;
-				UE_LOG_PACKET_NOTIFY(TEXT("Notification::ProcessReceivedAcks Seq: %u - IsAck: 0 HistoryIndex: N/A"), CurrentAck.Get());
+				UE_LOG_PACKET_NOTIFY_WARNING(TEXT("Notification::ProcessReceivedAcks Seq: %u - IsAck: 0 HistoryIndex: N/A"), CurrentAck.Get());
 				InFunc(CurrentAck, false);
 				++CurrentAck;
 			}
 
-			// For sequence numbers contained in the history we lookup the delivery status
+			// For sequence numbers contained in the history we lookup the delivery status from the history
 			while (AckCount > 0)
 			{
 				--AckCount;
@@ -251,7 +204,7 @@ private:
 				InFunc(CurrentAck, NotificationData.History.IsDelivered(AckCount));
 				++CurrentAck;
 			}
-			OutAckSeq = NotificationData.AckedSeq;	
+			OutAckSeq = NotificationData.AckedSeq;
 		}
 	}
 
@@ -261,90 +214,3 @@ private:
 	friend struct FNetPacketNotifyTestUtil;
 #endif
 };
-
-
-//////////////////////////////////////////////////////////////////////////
-// Implementation of TSequenceHistory
-//////////////////////////////////////////////////////////////////////////
-template <size_t HistorySize>
-TSequenceHistory<HistorySize>::TSequenceHistory()
-{
-	Reset();
-}
-
-#if WITH_DEV_AUTOMATION_TESTS
-template <size_t HistorySize>
-TSequenceHistory<HistorySize>::TSequenceHistory(WordT Value)
-{
-	for (size_t CurrentWordIt = 0; CurrentWordIt < WordCount; ++CurrentWordIt)
-	{
-		Storage[CurrentWordIt] = Value;
-	}	
-}
-#endif 
-
-template <size_t HistorySize>
-void TSequenceHistory<HistorySize>::Reset()
-{
-	memset(&Storage[0], 0, WordCount * sizeof(WordT));
-}
-
-template <size_t HistorySize>
-void TSequenceHistory<HistorySize>::AddDeliveryStatus(bool Delivered)
-{
-	WordT Carry = Delivered ? 1u : 0u;
-	const WordT ValueMask = 1u << (BitsPerWord - 1);
-	
-	for (size_t CurrentWordIt = 0; CurrentWordIt < WordCount; ++CurrentWordIt)
-	{
-		const WordT OldValue = Carry;
-		
-		// carry over highest bit in each word to the next word
-		Carry = (Storage[CurrentWordIt] & ValueMask) >> (BitsPerWord - 1);
-		Storage[CurrentWordIt] = (Storage[CurrentWordIt] << 1u) | OldValue;
-	}
-}
-
-template <size_t HistorySize>
-bool TSequenceHistory<HistorySize>::IsDelivered(size_t Index) const
-{
-	check(Index < Size());
-
-	const size_t WordIndex = Index / BitsPerWord;
-	const WordT WordMask = (WordT(1) << (Index & (BitsPerWord - 1)));
-	
-	return (Storage[WordIndex] & WordMask) != 0u;
-}
-
-template <size_t HistorySize>
-void TSequenceHistory<HistorySize>::Write(FBitWriter& Writer) const
-{
-	for (size_t CurrentWordIt = 0; CurrentWordIt < WordCount; ++CurrentWordIt)
-	{
-		WordT temp = Storage[CurrentWordIt];
-		Writer << temp;
-	}
-}
-
-template <size_t HistorySize>
-void TSequenceHistory<HistorySize>::Read(FBitReader& Reader)
-{
-	for (size_t CurrentWordIt = 0; CurrentWordIt < WordCount; ++CurrentWordIt)
-	{
-		Reader << Storage[CurrentWordIt];
-	}
-}
-
-//////////////////////////////////////////////////////////////////////////
-// Implementation of TSequenceNumber
-//////////////////////////////////////////////////////////////////////////
-template <size_t NumBits, typename SequenceType>
-typename TSequenceNumber<NumBits, SequenceType>::DifferenceT TSequenceNumber<NumBits, SequenceType>::Diff(TSequenceNumber A, TSequenceNumber B) 
-{ 
-	const size_t ShiftValue = sizeof(DifferenceT)*8 - NumBits;
-
-	const SequenceT ValueA = A.Value;
-	const SequenceT ValueB = B.Value;
-
-	return (DifferenceT)((ValueA - ValueB) << ShiftValue) >> ShiftValue;
-};
diff --git a/Engine/Source/Runtime/Engine/Public/Net/Util/SequenceHistory.h b/Engine/Source/Runtime/Engine/Public/Net/Util/SequenceHistory.h
new file mode 100644
--- /dev/null
+++ b/Engine/Source/Runtime/Engine/Public/Net/Util/SequenceHistory.h
@@ -0,0 +1,126 @@
+// Copyright 1998-2019 Epic Games, Inc. All Rights Reserved.
+
+#pragma once
+
+#include "CoreTypes.h"
+#include "HAL/UnrealMemory.h"
+#include "Misc/AssertionMacros.h"
+#include "Templates/IsSigned.h"
+#include "Serialization/BitWriter.h"
+#include "Serialization/BitReader.h"
+
+/** Util class to manage history of received sequence numbers */
+template <SIZE_T HistorySize>
+class TSequenceHistory
+{
+public:	
+	typedef uint32 WordT;
+	
+	static const SIZE_T BitsPerWord = sizeof(WordT) * 8;
+	static const SIZE_T WordCount = HistorySize / BitsPerWord;
+	static const SIZE_T Size = HistorySize;
+
+	static_assert(HistorySize > 0, "HistorySize must be > 0");
+	static_assert(HistorySize % BitsPerWord == 0, "InMaxHistorySize must be a modulo of the wordsize");
+
+public:	
+	TSequenceHistory();
+
+#if WITH_DEV_AUTOMATION_TESTS
+	explicit TSequenceHistory(WordT Value, SIZE_T Count = 1);
+	WordT* Data() { return &Storage[0]; }
+#endif
+	/** Reset */
+	void Reset();
+
+	/** Store delivery status, oldest will be dropped */
+	void AddDeliveryStatus(bool Delivered);
+
+	/** Query the status of a specific index, index 0 is last stored status */
+	bool IsDelivered(SIZE_T Index) const;
+	
+	bool operator==(const TSequenceHistory& Other) const { return FMemory::Memcmp(Storage, Other.Storage, WordCount * sizeof (WordT)) == 0; }
+
+	bool operator!=(const TSequenceHistory& Other) const { return FMemory::Memcmp(Storage, Other.Storage, WordCount * sizeof (WordT)) != 0; }
+
+	/** Write history to BitStream */
+	void Write(FBitWriter& Writer, SIZE_T NumWords) const;
+
+	/** Read history from BitStream */
+	void Read(FBitReader& Reader, SIZE_T NumWords);
+
+private:
+	WordT Storage[WordCount];
+};
+
+template <SIZE_T HistorySize>
+TSequenceHistory<HistorySize>::TSequenceHistory()
+{
+	Reset();
+}
+
+#if WITH_DEV_AUTOMATION_TESTS
+template <SIZE_T HistorySize>
+TSequenceHistory<HistorySize>::TSequenceHistory(WordT Value, SIZE_T Count)
+{
+	Reset();
+	for (SIZE_T CurrentWordIt = 0; CurrentWordIt < Count; ++CurrentWordIt)
+	{
+		Storage[CurrentWordIt] = Value;
+	}	
+}
+#endif 
+
+template <SIZE_T HistorySize>
+void TSequenceHistory<HistorySize>::Reset()
+{
+	FPlatformMemory::Memset(&Storage[0], 0, WordCount * sizeof(WordT));
+}
+
+template <SIZE_T HistorySize>
+void TSequenceHistory<HistorySize>::AddDeliveryStatus(bool Delivered)
+{
+	WordT Carry = Delivered ? 1u : 0u;
+	const WordT ValueMask = 1u << (BitsPerWord - 1);
+	
+	for (SIZE_T CurrentWordIt = 0; CurrentWordIt < WordCount; ++CurrentWordIt)
+	{
+		const WordT OldValue = Carry;
+		
+		// carry over highest bit in each word to the next word
+		Carry = (Storage[CurrentWordIt] & ValueMask) >> (BitsPerWord - 1);
+		Storage[CurrentWordIt] = (Storage[CurrentWordIt] << 1u) | OldValue;
+	}
+}
+
+template <SIZE_T HistorySize>
+bool TSequenceHistory<HistorySize>::IsDelivered(SIZE_T Index) const
+{
+	check(Index < Size);
+
+	const SIZE_T WordIndex = Index / BitsPerWord;
+	const WordT WordMask = (WordT(1) << (Index & (BitsPerWord - 1)));
+	
+	return (Storage[WordIndex] & WordMask) != 0u;
+}
+
+template <SIZE_T HistorySize>
+void TSequenceHistory<HistorySize>::Write(FBitWriter& Writer, SIZE_T NumWords) const
+{
+	NumWords = FPlatformMath::Min(NumWords, WordCount);
+	for (SIZE_T CurrentWordIt = 0; CurrentWordIt < NumWords; ++CurrentWordIt)
+	{
+		WordT temp = Storage[CurrentWordIt];
+		Writer << temp;
+	}
+}
+
+template <SIZE_T HistorySize>
+void TSequenceHistory<HistorySize>::Read(FBitReader& Reader, SIZE_T NumWords)
+{
+	NumWords = FPlatformMath::Min(NumWords, WordCount);
+	for (SIZE_T CurrentWordIt = 0; CurrentWordIt < NumWords; ++CurrentWordIt)
+	{
+		Reader << Storage[CurrentWordIt];
+	}
+}
diff --git a/Engine/Source/Runtime/Engine/Public/Net/Util/SequenceNumber.h b/Engine/Source/Runtime/Engine/Public/Net/Util/SequenceNumber.h
new file mode 100644
--- /dev/null
+++ b/Engine/Source/Runtime/Engine/Public/Net/Util/SequenceNumber.h
@@ -0,0 +1,70 @@
+// Copyright 1998-2019 Epic Games, Inc. All Rights Reserved.
+
+#pragma once
+
+#include "CoreTypes.h"
+#include "HAL/UnrealMemory.h"
+#include "Misc/AssertionMacros.h"
+#include "Templates/IsSigned.h"
+
+/** Helper class to work with sequence numbers */
+template <SIZE_T NumBits, typename SequenceType>
+class TSequenceNumber
+{
+	static_assert(TIsSigned<SequenceType>::Value == false, "The base type for sequence numbers must be unsigned");
+
+public:
+	typedef SequenceType SequenceT;
+	typedef int32 DifferenceT;
+
+	// Constants
+	enum { SeqNumberBits = NumBits };
+	enum { SeqNumberCount = SequenceT(1) << NumBits };
+	enum { SeqNumberHalf = SequenceT(1) << (NumBits - 1) };
+	enum { SeqNumberMax = SeqNumberCount - 1u };
+	enum { SeqNumberMask = SeqNumberMax };
+
+	/** Default constructor */
+	TSequenceNumber() : Value(0u) {}
+
+	/** Constructor with given value */
+	TSequenceNumber(SequenceT ValueIn) : Value(ValueIn & SeqNumberMask) {}
+	
+	/** Get Current Value */	
+	SequenceT Get() const { return Value; }
+
+	/** Diff between sequence numbers (A - B) only valid if (A - B) < SeqNumberHalf */
+	static DifferenceT Diff(TSequenceNumber A, TSequenceNumber B);
+	
+	/** return true if this is > Other, this is only considered to be the case if (A - B) < SeqNumberHalf since we have to be able to detect wraparounds */
+	bool operator>(const TSequenceNumber& Other) const { return (Value != Other.Value) && (((Value - Other.Value) & SeqNumberMask) < SeqNumberHalf); }
+
+	/** Check if this is >= Other, See above */
+	bool operator>=(const TSequenceNumber& Other) const { return ((Value - Other.Value) & SeqNumberMask) < SeqNumberHalf; }
+
+	/** Equals, NOTE that sequence numbers wrap around so 0 == 0 + SequenceNumberCount */
+	bool operator==(const TSequenceNumber& Other) const { return Value == Other.Value; }
+
+	bool operator!=(const TSequenceNumber& Other) const { return Value != Other.Value; }
+
+	/** Pre-increment and wrap around */
+	TSequenceNumber& operator++() { Increment(1u); return *this; }
+	
+	/** Post-increment and wrap around */
+	TSequenceNumber operator++(int) { TSequenceNumber Tmp(*this); Increment(1u); return Tmp; }
+
+private:
+	void Increment(SequenceT InValue) { *this = TSequenceNumber(Value + InValue); }
+	SequenceT Value;
+};
+
+template <SIZE_T NumBits, typename SequenceType>
+typename TSequenceNumber<NumBits, SequenceType>::DifferenceT TSequenceNumber<NumBits, SequenceType>::Diff(TSequenceNumber A, TSequenceNumber B) 
+{ 
+	const SIZE_T ShiftValue = sizeof(DifferenceT)*8 - NumBits;
+
+	const SequenceT ValueA = A.Value;
+	const SequenceT ValueB = B.Value;
+
+	return (DifferenceT)((ValueA - ValueB) << ShiftValue) >> ShiftValue;
+};
