Fixed deprecation warning in ChildConnection.h

diff --git a/Engine/Source/Runtime/Engine/Classes/Engine/ChildConnection.h b/Engine/Source/Runtime/Engine/Classes/Engine/ChildConnection.h
--- a/Engine/Source/Runtime/Engine/Classes/Engine/ChildConnection.h
+++ b/Engine/Source/Runtime/Engine/Classes/Engine/ChildConnection.h
@@ -52,10 +52,6 @@ public:
 		Parent->AssertValid();
 	}
 
-	virtual void SendAck(int32 PacketId, bool FirstTime=1) override
-	{
-	}
-
 	virtual void FlushNet(bool bIgnoreSimulation = false) override
 	{
 		Parent->FlushNet(bIgnoreSimulation);
