Fixed Nightly CIS errors

diff --git a/Engine/Source/Runtime/Engine/Private/Net/Tests/Util/SequenceNumberTests.cpp b/Engine/Source/Runtime/Engine/Private/Net/Tests/Util/SequenceNumberTests.cpp
--- a/Engine/Source/Runtime/Engine/Private/Net/Tests/Util/SequenceNumberTests.cpp
+++ b/Engine/Source/Runtime/Engine/Private/Net/Tests/Util/SequenceNumberTests.cpp
@@ -70,7 +70,6 @@ bool FNetSequenceNumberTest::RunTest(const FString& Parameters)
 
 		TestFalse("SequenceNumbers - operator> MaxSeq > Zero", MaxSeq > Zero);
 		TestTrue("SequenceNumbers - operator> Zero > MaxSeq", Zero > MaxSeq);
-		TestFalse("SequenceNumbers - operator> Zero > Zero", Zero > Zero);
 	}
 	
 	// Test diff sequence numbers
diff --git a/Engine/Source/Runtime/Engine/Private/NetConnection.cpp b/Engine/Source/Runtime/Engine/Private/NetConnection.cpp
--- a/Engine/Source/Runtime/Engine/Private/NetConnection.cpp
+++ b/Engine/Source/Runtime/Engine/Private/NetConnection.cpp
@@ -1360,7 +1360,7 @@ void UNetConnection::WritePacketInfo(FBitWriter& Writer) const
 
 bool UNetConnection::ReadPacketInfo(FBitReader& Reader)
 {
-	const bool bHasServerFrameTime = Reader.ReadBit();
+	const bool bHasServerFrameTime = Reader.ReadBit() == 1u;
 	double ServerFrameTime = 0;
 
 	if ( !Driver->IsServer() )
@@ -1541,7 +1541,7 @@ void UNetConnection::ReceivedPacket( FBitReader& Reader )
 		// For demo backwards compatibility, old replays still have this bit
 		if (InternalAck && EngineNetworkProtocolVersion < EEngineNetworkVersionHistory::HISTORY_ACKS_INCLUDED_IN_HEADER)
 		{
-			const bool IsAckDummy = Reader.ReadBit();
+			const bool IsAckDummy = Reader.ReadBit() == 1u;
 		}
 
 		// Parse the bunch.
diff --git a/Engine/Source/Runtime/Engine/Public/Net/NetPacketNotify.h b/Engine/Source/Runtime/Engine/Public/Net/NetPacketNotify.h
--- a/Engine/Source/Runtime/Engine/Public/Net/NetPacketNotify.h
+++ b/Engine/Source/Runtime/Engine/Public/Net/NetPacketNotify.h
@@ -124,7 +124,7 @@ public:
 
 	SequenceNumberT UpdateInAckSeqAck(SequenceNumberT::DifferenceT AckCount, SequenceNumberT AckedSeq)
 	{
-		check(AckCount <= AckRecord.Count());
+		check((SIZE_T)AckCount <= AckRecord.Count());
 
 		if (AckCount > 1)
 		{			
diff --git a/Engine/Source/Runtime/Engine/Public/Net/Util/ResizableCircularQueue.h b/Engine/Source/Runtime/Engine/Public/Net/Util/ResizableCircularQueue.h
--- a/Engine/Source/Runtime/Engine/Public/Net/Util/ResizableCircularQueue.h
+++ b/Engine/Source/Runtime/Engine/Public/Net/Util/ResizableCircularQueue.h
@@ -101,8 +101,8 @@ TResizableCircularQueue<T, AllocatorT>::TResizableCircularQueue(SIZE_T InitialCa
 , Tail(0u)
 , IndexMask(0u)
 {
-	// Capacity should be power of two, it will be rounded up but lets warn the user anyway debug builds.
-	checkSlow(InitialCapacity & (InitialCapacity - 1) = 0);
+	// Capacity should be power of two, it will be rounded up but lets warn the user anyway in debug builds.
+	checkSlow((InitialCapacity & (InitialCapacity - 1)) == 0);
 	if (InitialCapacity > 0)
 	{
 		SetCapacity(InitialCapacity);
