Fixed CIS warning for VS 2015 by removing deprecated variables that was initialized by default constuctor

diff --git a/Engine/Source/Runtime/Engine/Classes/Engine/NetConnection.h b/Engine/Source/Runtime/Engine/Classes/Engine/NetConnection.h
--- a/Engine/Source/Runtime/Engine/Classes/Engine/NetConnection.h
+++ b/Engine/Source/Runtime/Engine/Classes/Engine/NetConnection.h
@@ -402,9 +402,6 @@ public:
 	TArray<int32>		OutReliable;
 	TArray<int32>		InReliable;
 	TArray<int32>		PendingOutRec;	// Outgoing reliable unacked data from previous (now destroyed) channel in this slot.  This contains the first chsequence not acked
-	UE_DEPRECATED(4.22, "QueuedAcks, ResendAcks variables will be removed.")
-	TArray<int32> QueuedAcks, ResendAcks;
-
 	int32				InitOutReliable;
 	int32				InitInReliable;
 
