Changed UpdateInAckSeqAck function to private as it should be

diff --git a/Engine/Source/Runtime/Engine/Public/Net/NetPacketNotify.h b/Engine/Source/Runtime/Engine/Public/Net/NetPacketNotify.h
--- a/Engine/Source/Runtime/Engine/Public/Net/NetPacketNotify.h
+++ b/Engine/Source/Runtime/Engine/Public/Net/NetPacketNotify.h
@@ -122,12 +122,35 @@ public:
 	/** Get the current sequenceHistory length in bits */
 	SIZE_T GetCurrentSequenceHistoryLength() const;
 
+private:
+	struct FSentAckData
+	{
+		SequenceNumberT OutSeq;	// Not needed... just to verify that things work as expected
+		SequenceNumberT InAckSeq;
+	};
+	typedef TResizableCircularQueue<FSentAckData, TInlineAllocator<128>> AckRecordT;
+
+	AckRecordT AckRecord;				// Track acked seq for each sent packet to track size of ack history
+	SIZE_T WrittenHistoryWordCount;		// Bookkeeping to track if we can update data
+	SequenceNumberT WrittenInAckSeq;	// When we call CommitAndIncrementOutSequence this will be committed along with the current outgoing sequence number for bookkeeping
+
+	// Track incoming sequence data
+	SequenceHistoryT InSeqHistory;		// BitBuffer containing a bitfield describing the history of received packets
+	SequenceNumberT InSeq;				// Last sequence number received and accepted from remote
+	SequenceNumberT InAckSeq;			// Last sequence number received from remote that we have acknowledged, this is needed since we support accepting a packet but explicitly not acknowledge it as received.
+	SequenceNumberT InAckSeqAck;		// Last sequence number received from remote that we have acknowledged and also knows that the remote has received the ack, used to calculate how big our history must be
+
+	// Track outgoing sequence data
+	SequenceNumberT OutSeq;				// Outgoing sequence number
+	SequenceNumberT OutAckSeq;			// Last sequence number that we know that the remote side have received.
+
+private:
 	SequenceNumberT UpdateInAckSeqAck(SequenceNumberT::DifferenceT AckCount, SequenceNumberT AckedSeq)
 	{
 		check((SIZE_T)AckCount <= AckRecord.Count());
 
 		if (AckCount > 1)
-		{			
+		{
 			AckRecord.PopNoCheck(AckCount - 1);
 		}
 		
@@ -148,29 +171,6 @@ public:
 		}
 	}
 
-private:
-	struct FSentAckData
-	{
-		SequenceNumberT OutSeq;	// Not needed... just to verify that things work as expected
-		SequenceNumberT InAckSeq;
-	};
-	typedef TResizableCircularQueue<FSentAckData, TInlineAllocator<128>> AckRecordT;
-
-	AckRecordT AckRecord;				// Track acked seq for each sent packet to track size of ack history
-	SIZE_T WrittenHistoryWordCount;		// Bookkeeping to track if we can update data
-	SequenceNumberT WrittenInAckSeq;	// When we call CommitAndIncrementOutSequence this will be committed along with the current outgoing sequence number for bookkeeping
-
-	// Track incoming sequence data
-	SequenceHistoryT InSeqHistory;		// BitBuffer containing a bitfield describing the history of received packets
-	SequenceNumberT InSeq;				// Last sequence number received and accepted from remote
-	SequenceNumberT InAckSeq;			// Last sequence number received from remote that we have acknowledged, this is needed since we support accepting a packet but explicitly not acknowledge it as received.
-	SequenceNumberT InAckSeqAck;		// Last sequence number received from remote that we have acknowledged and also knows that the remote has received the ack, used to calculate how big our history must be
-
-	// Track outgoing sequence data
-	SequenceNumberT OutSeq;				// Outgoing sequence number
-	SequenceNumberT OutAckSeq;			// Last sequence number that we know that the remote side have received.
-
-private:
 	template<class Functor>
 	void ProcessReceivedAcks(const FNotificationHeader& NotificationData, Functor&& InFunc)
 	{
