Fixed replication of DetachFromComponent (Server setting AttachParant to nullptr) by adding explicit replication whether the component should be attached or not.
Disabled ensure in detachFromComponent if not authority and when both scenecomponent and parent is replicated as different replication order of attachments might result in the ensure triggering.

diff --git a/Engine/Source/Runtime/Engine/Classes/Components/SceneComponent.h b/Engine/Source/Runtime/Engine/Classes/Components/SceneComponent.h
--- a/Engine/Source/Runtime/Engine/Classes/Components/SceneComponent.h
+++ b/Engine/Source/Runtime/Engine/Classes/Components/SceneComponent.h
@@ -207,6 +207,10 @@ public:
 	uint8 bHiddenInGame:1;
 
 private:
+	/** Whether or not we should be attached. */
+	UPROPERTY(Transient, Replicated)
+	uint8 bShouldBeAttached : 1;
+
 	/**
 	 * Whether or not the cached PhysicsVolume this component overlaps should be updated when the component is moved.
 	 * @see GetPhysicsVolume()
diff --git a/Engine/Source/Runtime/Engine/Private/Components/SceneComponent.cpp b/Engine/Source/Runtime/Engine/Private/Components/SceneComponent.cpp
--- a/Engine/Source/Runtime/Engine/Private/Components/SceneComponent.cpp
+++ b/Engine/Source/Runtime/Engine/Private/Components/SceneComponent.cpp
@@ -78,6 +78,7 @@ USceneComponent::USceneComponent(const FObjectInitializer& ObjectInitializer /*=
 
 	bNetHasReceivedRelativeLocation = false;
 	bNetHasReceivedRelativeRotation = false;
+	bShouldBeAttached = AttachParent != nullptr;
 }
 
 #if WITH_EDITORONLY_DATA
@@ -652,6 +653,7 @@ void USceneComponent::OnRegister()
 			// Failed to attach, we need to clear AttachParent so we don't think we're actually attached when we're not.
 			AttachParent = nullptr;
 			AttachSocketName = NAME_None;
+			bShouldBeAttached = false;
 		}
 	}
 	
@@ -1707,6 +1709,7 @@ void USceneComponent::SetupAttachment(class USceneComponent* InParent, FName InS
 				{
 					AttachParent = InParent;
 					AttachSocketName = InSocketName;
+					bShouldBeAttached = AttachParent != nullptr;
 				}
 			}
 		}
@@ -1936,6 +1939,7 @@ bool USceneComponent::AttachToComponent(USceneComponent* Parent, const FAttachme
 		// Save pointer from child to parent
 		AttachParent = Parent;
 		AttachSocketName = SocketName;
+		bShouldBeAttached = AttachParent != nullptr;
 
 		OnAttachmentChanged();
 
@@ -2105,9 +2109,13 @@ void USceneComponent::DetachFromComponent(const FDetachmentTransformRules& Detac
 		{
 			PrimComp->UnWeldFromParent();
 		}
-
-		// Make sure parent points to us if we're registered
-		ensureMsgf(!bRegistered || GetAttachParent()->GetAttachChildren().Contains(this), TEXT("Attempt to detach SceneComponent '%s' owned by '%s' from AttachParent '%s' while not attached."), *GetName(), (Owner ? *Owner->GetName() : TEXT("Unowned")), *GetAttachParent()->GetName());
+		
+		// Due to replication order the ensure below is only valid on server OR if not both parent and child are replicated
+		if ((Owner && Owner->GetLocalRole() == ROLE_Authority) || !(GetIsReplicated() && GetAttachParent()->GetIsReplicated()))
+		{
+			// Make sure parent points to us if we're registered
+			ensureMsgf(!bRegistered || GetAttachParent()->GetAttachChildren().Contains(this), TEXT("Attempt to detach SceneComponent '%s' owned by '%s' from AttachParent '%s' while not attached."), *GetName(), (Owner ? *Owner->GetName() : TEXT("Unowned")), *GetAttachParent()->GetName());
+		}
 
 		if (DetachmentRules.bCallModify && !HasAnyFlags(RF_Transient))
 		{
@@ -2132,6 +2140,7 @@ void USceneComponent::DetachFromComponent(const FDetachmentTransformRules& Detac
 #endif
 		AttachParent = nullptr;
 		AttachSocketName = NAME_None;
+		bShouldBeAttached = 0;
 
 		OnAttachmentChanged();
 
@@ -3208,8 +3217,8 @@ void USceneComponent::PostNetReceive()
 {
 	Super::PostNetReceive();
 
-	// If we have no attach parent, attach to parent's root component.
-	if (GetAttachParent() == nullptr)
+	// If we have no attach parent even though the server told us that we should have one, keep attach to parent's root component until we get the correct pointer.
+	if (bShouldBeAttached && GetAttachParent() == nullptr)
 	{
 		USceneComponent * ParentRoot = GetOwner()->GetRootComponent();
 		if (ParentRoot != this)
@@ -3239,7 +3248,22 @@ void USceneComponent::PostRepNotifies()
 			bNetHasReceivedRelativeRotation = true;
 		}
 
-		AttachToComponent(NetOldAttachParent, FAttachmentTransformRules::KeepRelativeTransform, NetOldAttachSocketName);
+		// Check if this is a detach
+		if (AttachParent && !bShouldBeAttached)
+		{
+			ensureMsgf(NetOldAttachParent == nullptr, TEXT("Local modification of AttachParent detected for replicated component %s, disable replication or execute detachment on host."), *GetFullName());
+			DetachFromComponent(FDetachmentTransformRules::KeepWorldTransform);
+		}
+		else
+		{
+			const uint8 bOldShouldBeAttached = bShouldBeAttached;
+
+			AttachToComponent(NetOldAttachParent, FAttachmentTransformRules::KeepRelativeTransform, NetOldAttachSocketName);
+
+			// restore bShouldBeAttached to what we have received
+			bShouldBeAttached = bOldShouldBeAttached;
+		}
+
 		bNetUpdateAttachment = false;
 	}
 
@@ -3258,6 +3282,7 @@ void USceneComponent::GetLifetimeReplicatedProps(TArray< FLifetimeProperty > & O
 	DOREPLIFETIME(USceneComponent, bAbsoluteRotation);
 	DOREPLIFETIME(USceneComponent, bAbsoluteScale);
 	DOREPLIFETIME(USceneComponent, bVisible);
+	DOREPLIFETIME(USceneComponent, bShouldBeAttached);
 	DOREPLIFETIME(USceneComponent, AttachParent);
 	DOREPLIFETIME(USceneComponent, AttachChildren);
 	DOREPLIFETIME(USceneComponent, AttachSocketName);
