Allow Actors that are not replicated to be deleted by clients.

diff --git a/Engine/Source/Runtime/Engine/Private/LevelActor.cpp b/Engine/Source/Runtime/Engine/Private/LevelActor.cpp
--- a/Engine/Source/Runtime/Engine/Private/LevelActor.cpp
+++ b/Engine/Source/Runtime/Engine/Private/LevelActor.cpp
@@ -54,6 +54,12 @@ static TAutoConsoleVariable<float> CVarEncroachEpsilon(
 	TEXT("0: use full sized shape. > 0: shrink shape size by this amount (world units)"),
 	ECVF_Default);
 
+static TAutoConsoleVariable<int32> CVarAllowDestroyNonNetworkActors(
+	TEXT("p.AllowDestroyNonNetworkActors"),
+	1,
+	TEXT("When enabled, allows Clients in Networked Games to destroy non-networked actors (AActor::Role == ROLE_None). Does not change behavior on Servers or Standalone games.")
+);
+
 #define LINE_CHECK_TRACING 0
 
 #if LINE_CHECK_TRACING
@@ -561,8 +567,19 @@ bool UWorld::DestroyActor( AActor* ThisActor, bool bNetForce, bool bShouldModify
 			return false;
 		}
 
+		// Note, for Standalone games, Actors should have Authority == ROLE_Authority.
+		// In that sense, they'll be treated as Network Actors here.
+		const bool bIsNetworkedActor = ThisActor->Role != ROLE_None;
+
 		// Can't kill if wrong role.
-		if( ThisActor->Role!=ROLE_Authority && !bNetForce && !ThisActor->bNetTemporary )
+		const bool bCanDestroyNetworkActor = ThisActor->Role == ROLE_Authority || bNetForce || ThisActor->bNetTemporary;
+		if (bIsNetworkedActor && !bCanDestroyNetworkActor)
+		{
+			return false;
+		}
+
+		const bool bCanDestroyNonNetworkActor = !!CVarAllowDestroyNonNetworkActors.GetValueOnAnyThread();
+		if (!bIsNetworkedActor && !bCanDestroyNonNetworkActor)
 		{
 			return false;
 		}
