Prevent GameplayAbilities from sending Remote Functions if they aren't supported for networking.
Previously, Non-Replicated Gameplay Abilities could still try to call remote functions if they're outer Replicated.

diff --git a/Engine/Plugins/Runtime/GameplayAbilities/Source/GameplayAbilities/Private/Abilities/GameplayAbility.cpp b/Engine/Plugins/Runtime/GameplayAbilities/Source/GameplayAbilities/Private/Abilities/GameplayAbility.cpp
--- a/Engine/Plugins/Runtime/GameplayAbilities/Source/GameplayAbilities/Private/Abilities/GameplayAbility.cpp
+++ b/Engine/Plugins/Runtime/GameplayAbilities/Source/GameplayAbilities/Private/Abilities/GameplayAbility.cpp
@@ -89,7 +89,7 @@ UWorld* UGameplayAbility::GetWorld() const
 
 int32 UGameplayAbility::GetFunctionCallspace(UFunction* Function, void* Parameters, FFrame* Stack)
 {
-	if (HasAnyFlags(RF_ClassDefaultObject))
+	if (HasAnyFlags(RF_ClassDefaultObject) || !IsSupportedForNetworking())
 	{
 		return FunctionCallspace::Local;
 	}
diff --git a/Engine/Source/Runtime/Engine/Private/NetworkDriver.cpp b/Engine/Source/Runtime/Engine/Private/NetworkDriver.cpp
--- a/Engine/Source/Runtime/Engine/Private/NetworkDriver.cpp
+++ b/Engine/Source/Runtime/Engine/Private/NetworkDriver.cpp
@@ -5294,6 +5294,11 @@ void UNetDriver::ProcessRemoteFunction(class AActor* Actor, UFunction* Function,
 	SCOPE_CYCLE_COUNTER(STAT_NetProcessRemoteFunc);
 	SCOPE_CYCLE_UOBJECT(Function, Function);
 
+	{
+		UObject* TestObject = (SubObject == nullptr) ? Actor : SubObject;
+		ensureMsgf(TestObject->IsSupportedForNetworking() || TestObject->IsNameStableForNetworking(), TEXT("Attempted to call ProcessRemoteFunction with object that is not supported for networking. Object: %s Function: %s"), *TestObject->GetPathName(), *Function->GetName());
+	}
+
 	bool bBlockSendRPC = false;
 
 	SendRPCDel.ExecuteIfBound(Actor, Function, Parameters, OutParms, Stack, SubObject, bBlockSendRPC);
