PR #5410: Fix for Reliable multicast dropping. (Contributed by [REDACTED_PERSON])

Updated the comments in the PR, and gated the value behind a CVar in case some licensees are relying on it.

diff --git a/Engine/Source/Runtime/Engine/Private/NetworkDriver.cpp b/Engine/Source/Runtime/Engine/Private/NetworkDriver.cpp
--- a/Engine/Source/Runtime/Engine/Private/NetworkDriver.cpp
+++ b/Engine/Source/Runtime/Engine/Private/NetworkDriver.cpp
@@ -203,6 +203,12 @@ static TAutoConsoleVariable<int32> CVarActorChannelPool(
 	1,
 	TEXT("If nonzero, actor channels will be pooled to save memory and object creation cost."));
 
+static TAutoConsoleVariable<int32> CVarAllowReliableMulticastToNonRelevantChannels(
+	TEXT("net.AllowReliableMulticastToNonRelevantChannels"),
+	1,
+	TEXT("Allow Reliable Server Multicasts to be sent to non-Relevant Actors, as long as their is an existing ActorChannel.")
+);
+
 /*-----------------------------------------------------------------------------
 	UNetDriver implementation.
 -----------------------------------------------------------------------------*/
@@ -5332,13 +5338,18 @@ void UNetDriver::ProcessRemoteFunction(class AActor* Actor, UFunction* Function,
 				{
 					// Only send or queue multicasts if the actor is relevant to the connection
 					FNetViewer Viewer(Connection, 0.f);
-					if (Actor->IsNetRelevantFor(Viewer.InViewer, Viewer.ViewTarget, Viewer.ViewLocation))
+
+					if (Connection->GetUChildConnection() != nullptr)
 					{
-						if (Connection->GetUChildConnection() != nullptr)
-						{
-							Connection = ((UChildConnection*)Connection)->Parent;
-						}
+						Connection = ((UChildConnection*)Connection)->Parent;
+					}
 
+					// It's possible that an actor is not relevant to a specific connection, but the channel is still alive (due to hysteresis).
+					// However, it's also possible that the Actor could become relevant again before the channel ever closed, and in that case we
+					// don't want to lose Reliable RPCs.
+					if (Actor->IsNetRelevantFor(Viewer.InViewer, Viewer.ViewTarget, Viewer.ViewLocation) ||
+						((Function->FunctionFlags & FUNC_NetReliable) && !!CVarAllowReliableMulticastToNonRelevantChannels.GetValueOnGameThread() && Connection->FindActorChannelRef(Actor)))
+					{
 						// We don't want to call this unless necessary, and it will internally handle being called multiple times before a clear
 						// Builds any shared serialization state for this rpc
 						RepLayout->BuildSharedSerializationForRPC(Parameters);
