PR #5339: Fix default scale when serializing new actor (Contributed by [REDACTED_PERSON])

diff --git a/Engine/Source/Runtime/Engine/Private/PackageMapClient.cpp b/Engine/Source/Runtime/Engine/Private/PackageMapClient.cpp
--- a/Engine/Source/Runtime/Engine/Private/PackageMapClient.cpp
+++ b/Engine/Source/Runtime/Engine/Private/PackageMapClient.cpp
@@ -371,7 +371,7 @@ bool UPackageMapClient::SerializeNewActor(FArchive& Ar, class UActorChannel *Cha
 				Location = LocalLocation = FVector::ZeroVector;
 			}
 			Rotation = RootComponent ? Actor->GetActorRotation() : FRotator::ZeroRotator;
-			Scale = RootComponent ? Actor->GetActorScale() : FVector::ZeroVector;
+			Scale = RootComponent ? Actor->GetActorScale() : FVector::OneVector;
 			Velocity = RootComponent ? Actor->GetVelocity() : FVector::ZeroVector;
 		}
 
