Fix the NetworkProfiler to save const FInternetAddrs, as it does not need to touch the addresses it saves and serializes

#rnx

diff --git a/Engine/Source/Runtime/Engine/Private/NetworkProfiler.cpp b/Engine/Source/Runtime/Engine/Private/NetworkProfiler.cpp
--- a/Engine/Source/Runtime/Engine/Private/NetworkProfiler.cpp
+++ b/Engine/Source/Runtime/Engine/Private/NetworkProfiler.cpp
@@ -216,7 +216,7 @@ void FNetworkProfiler::SetCurrentConnection( UNetConnection* Connection )
 {
 	if ( bIsTrackingEnabled && Connection != nullptr )
 	{
-		const TSharedPtr<FInternetAddr> ConnectionAddr = Connection->GetInternetAddr();
+		const TSharedPtr<const FInternetAddr> ConnectionAddr = Connection->GetInternetAddr();
 		if ( ConnectionAddr.IsValid() )
 		{
 			if ( LastAddress != ConnectionAddr )
diff --git a/Engine/Source/Runtime/Engine/Public/Net/NetworkProfiler.h b/Engine/Source/Runtime/Engine/Public/Net/NetworkProfiler.h
--- a/Engine/Source/Runtime/Engine/Public/Net/NetworkProfiler.h
+++ b/Engine/Source/Runtime/Engine/Public/Net/NetworkProfiler.h
@@ -92,7 +92,7 @@ private:
 	FNetworkProfilerHeader					CurrentHeader;
 
 	/** Last known address																			*/
-	TSharedPtr<FInternetAddr>				LastAddress;
+	TSharedPtr<const FInternetAddr>			LastAddress;
 
 	/** All the data required for writing sent bunches to the profiler stream						*/
 	struct FSendBunchInfo
