Make sure that if an Object's RepFlags change, or if a properties become Active, we send those to clients.

PR #5353: UE-43135 (Contributed by [REDACTED_PERSON])

diff --git a/Engine/Source/Runtime/Engine/Private/DataReplication.cpp b/Engine/Source/Runtime/Engine/Private/DataReplication.cpp
--- a/Engine/Source/Runtime/Engine/Private/DataReplication.cpp
+++ b/Engine/Source/Runtime/Engine/Private/DataReplication.cpp
@@ -1221,28 +1221,7 @@ void FObjectReplicator::ReplicateCustomDeltaProperties( FNetBitWriter & Bunch, F
 	UNetConnection * OwningChannelConnection = OwningChannel->Connection;
 
 	// Initialize a map of which conditions are valid
-
-	TStaticBitArray<COND_Max> ConditionMap;
-	const bool bIsInitial = RepFlags.bNetInitial ? true : false;
-	const bool bIsOwner = RepFlags.bNetOwner ? true : false;
-	const bool bIsSimulated = RepFlags.bNetSimulated ? true : false;
-	const bool bIsPhysics = RepFlags.bRepPhysics ? true : false;
-	const bool bIsReplay = RepFlags.bReplay ? true : false;
-
-	ConditionMap[COND_None] = true;
-	ConditionMap[COND_InitialOnly] = bIsInitial;
-	ConditionMap[COND_OwnerOnly] = bIsOwner;
-	ConditionMap[COND_SkipOwner] = !bIsOwner;
-	ConditionMap[COND_SimulatedOnly] = bIsSimulated;
-	ConditionMap[COND_SimulatedOnlyNoReplay] = bIsSimulated && !bIsReplay;
-	ConditionMap[COND_AutonomousOnly] = !bIsSimulated;
-	ConditionMap[COND_SimulatedOrPhysics] = bIsSimulated || bIsPhysics;
-	ConditionMap[COND_SimulatedOrPhysicsNoReplay] = ( bIsSimulated || bIsPhysics ) && !bIsReplay;
-	ConditionMap[COND_InitialOrOwner] = bIsInitial || bIsOwner;
-	ConditionMap[COND_Custom] = true;
-	ConditionMap[COND_ReplayOrOwner] = bIsReplay || bIsOwner;
-	ConditionMap[COND_ReplayOnly] = bIsReplay;
-	ConditionMap[COND_SkipReplay] = !bIsReplay;
+	const TStaticBitArray<COND_Max> ConditionMap = FRepState::BuildConditionMap(RepFlags);
 
 	// Make sure net field export group is registered
 	FNetFieldExportGroup* NetFieldExportGroup = OwningChannel->GetOrCreateNetFieldExportGroupForClassNetCache( Object );
diff --git a/Engine/Source/Runtime/Engine/Private/RepLayout.cpp b/Engine/Source/Runtime/Engine/Private/RepLayout.cpp
--- a/Engine/Source/Runtime/Engine/Private/RepLayout.cpp
+++ b/Engine/Source/Runtime/Engine/Private/RepLayout.cpp
@@ -531,6 +531,7 @@ uint16 FRepLayout::CompareProperties_r(
 {
 	check(ShadowData);
 
+	FRepChangedPropertyTracker* RepChangedPropertyTracker = (RepState ? RepState->RepChangedPropertyTracker.Get() : nullptr);
 	for (int32 CmdIndex = CmdStart; CmdIndex < CmdEnd; CmdIndex++)
 	{
 		const FRepLayoutCmd& Cmd = Cmds[CmdIndex];
@@ -541,7 +542,14 @@ uint16 FRepLayout::CompareProperties_r(
 		Handle++;
 
 		const bool bIsLifetime = ((ParentCmd.Flags & ERepParentFlags::IsLifetime) != ERepParentFlags::None);
-		const bool bShouldSkip = !bIsLifetime || (ParentCmd.Condition == COND_InitialOnly && !bIsInitial);
+
+		// Active state of a property applies to *all* connections.
+		// If the property is inactive, we can skip comparing it because we know it won't be sent.
+		// Further, this will keep the last active state of the property in the shadow buffer,
+		// meaning the next time the property becomes active it will be sent to all connections.
+		const bool bActive = (RepChangedPropertyTracker == nullptr || RepChangedPropertyTracker->Parents[Cmd.ParentIndex].Active);
+
+		const bool bShouldSkip = !bIsLifetime || !bActive || (ParentCmd.Condition == COND_InitialOnly && !bIsInitial);
 
 		if (Cmd.Type == ERepLayoutCmdType::DynamicArray)
 		{
@@ -763,12 +771,16 @@ bool FRepLayout::ReplicateProperties(
 
 	FRepChangedPropertyTracker*	ChangeTracker = RepState->RepChangedPropertyTracker.Get();
 
+	TArray<uint16> NewlyActiveChangelist;
+
 	// Rebuild conditional state if needed
 	if (RepState->RepFlags.Value != RepFlags.Value)
 	{
 		RebuildConditionalProperties(RepState, RepFlags);
 
-		RepState->RepFlags.Value = RepFlags.Value;
+		// Filter out any previously inactive changes from still inactive ones
+		TArray<uint16> InactiveChangelist = MoveTemp(RepState->InactiveChangelist);
+		FilterChangeList(InactiveChangelist, RepState->InactiveParents, RepState->InactiveChangelist, NewlyActiveChangelist);
 	}
 
 	if (OwningChannel->Connection->bResendAllDataSinceOpen)
@@ -782,8 +794,15 @@ bool FRepLayout::ReplicateProperties(
 			TArray<uint16> Pruned;
 			PruneChangeList(RepState, Data, RepState->LifetimeChangelist, Pruned);
 			RepState->LifetimeChangelist = MoveTemp(Pruned);
-			SendProperties_BackwardsCompatible(RepState, ChangeTracker, Data, OwningChannel->Connection, Writer, RepState->LifetimeChangelist);
-			return true;
+
+			// No need to merge in the newly active properties here, as the Lifetime Changelist should contain everything
+			// inactive or otherwise.
+			FilterChangeListToActive(RepState->LifetimeChangelist, RepState->InactiveParents, Pruned);
+			if (Pruned.Num() > 0)
+			{
+				SendProperties_BackwardsCompatible(RepState, ChangeTracker, Data, OwningChannel->Connection, Writer, Pruned);
+				return true;
+			}
 		}
 
 		return false;
@@ -801,7 +820,7 @@ bool FRepLayout::ReplicateProperties(
 	// We can early out if we know for sure there are no new changelists to send
 	if (bCompareIndexSame || RepState->LastChangelistIndex == RepChangelistState->HistoryEnd)
 	{
-		if (RepState->NumNaks == 0 && !bFlushPreOpenAckHistory)
+		if (RepState->NumNaks == 0 && !bFlushPreOpenAckHistory && NewlyActiveChangelist.Num() == 0)
 		{
 			// Nothing changed and there are no nak's, so just do normal housekeeping and remove acked history items
 			UpdateChangelistHistory(RepState, ObjectClass, Data, OwningChannel->Connection, NULL);
@@ -841,6 +860,13 @@ bool FRepLayout::ReplicateProperties(
 		MergeChangeList(Data, HistoryItem.Changed, Temp, Changed);
 	}
 
+	// Merge in newly active properties so they can be sent.
+	if (NewlyActiveChangelist.Num() > 0)
+	{
+		TArray<uint16> Temp = MoveTemp(Changed);
+		MergeChangeList(Data, NewlyActiveChangelist, Temp, Changed);
+	}
+
 	// We're all caught up now
 	RepState->LastChangelistIndex = RepChangelistState->HistoryEnd;
 
@@ -883,16 +909,32 @@ bool FRepLayout::ReplicateProperties(
 
 	const int32 NumBits = Writer.GetNumBits();
 
+	// Filter out the final changelist into Active and Inaction.
+	TArray<uint16> UnfilteredChanged = MoveTemp(Changed);
+	TArray<uint16> NewlyInactiveChangelist;
+	FilterChangeList(UnfilteredChanged, RepState->InactiveParents, NewlyInactiveChangelist, Changed);
+
+	// If we have any properties that are no longer active, make sure we track them.
+	if (NewlyInactiveChangelist.Num() > 1)
+	{
+		TArray<uint16> Temp = MoveTemp(RepState->InactiveChangelist);
+		MergeChangeList(Data, NewlyInactiveChangelist, Temp, RepState->InactiveChangelist);
+	}
+
 	// Send the final merged change list
 	if (OwningChannel->Connection->InternalAck)
 	{
 		// Remember all properties that have changed since this channel was first opened in case we need it (for bResendAllDataSinceOpen)
+		// We use UnfilteredChanged so LifetimeChangelist contains all properties, regardless of Active state.
 		TArray<uint16> Temp = MoveTemp(RepState->LifetimeChangelist);
-		MergeChangeList(Data, Changed, Temp, RepState->LifetimeChangelist);
+		MergeChangeList(Data, UnfilteredChanged, Temp, RepState->LifetimeChangelist);
 
-		SendProperties_BackwardsCompatible(RepState, ChangeTracker, Data, OwningChannel->Connection, Writer, Changed);
+		if (Changed.Num() > 0)
+		{
+			SendProperties_BackwardsCompatible(RepState, ChangeTracker, Data, OwningChannel->Connection, Writer, Changed);
+		}
 	}
-	else
+	else if (Changed.Num() > 0)
 	{
 		SendProperties(RepState, ChangeTracker, Data, ObjectClass, Writer, Changed, RepChangelistState->SharedSerialization);
 	}
@@ -1407,6 +1449,89 @@ void FRepLayout::PruneChangeList_r(
 	}
 }
 
+void FRepLayout::FilterChangeList(
+	const TArray<uint16>&	Changelist,
+	const TBitArray<>&		InactiveParents,
+	TArray<uint16>&			OutInactiveProperties,
+	TArray<uint16>&			OutActiveProperties) const
+{
+	FChangelistIterator ChangelistIterator(Changelist, 0);
+	FRepHandleIterator HandleIterator(ChangelistIterator, Cmds, BaseHandleToCmdIndex, 0, 1, 0, Cmds.Num() - 1);
+
+	OutInactiveProperties.Empty();
+	OutActiveProperties.Empty();
+
+	while (HandleIterator.NextHandle())
+	{
+		const FRepLayoutCmd& Cmd = Cmds[HandleIterator.CmdIndex];
+
+		TArray<uint16>& Properties = InactiveParents[Cmd.ParentIndex] ? OutInactiveProperties : OutActiveProperties;
+			
+		Properties.Add(HandleIterator.Handle);
+
+		if (Cmd.Type == ERepLayoutCmdType::DynamicArray)
+		{
+			// No need to recursively filter the change list, as handles are only enabled/disabled at the parent level
+			int32 HandleCount = ChangelistIterator.Changed[ChangelistIterator.ChangedIndex];
+			Properties.Add(HandleCount);
+					
+			for (int32 I = 0; I < HandleCount; ++I)
+			{
+				Properties.Add(ChangelistIterator.Changed[ChangelistIterator.ChangedIndex + 1 + I]);
+			}
+
+			Properties.Add(0);
+
+			HandleIterator.JumpOverArray();
+		}
+	}
+
+	OutInactiveProperties.Add(0);
+	OutActiveProperties.Add(0);
+}
+
+void FRepLayout::FilterChangeListToActive(
+	const TArray<uint16>&	Changelist,
+	const TBitArray<>&		InactiveParents,
+	TArray<uint16>&			OutProperties) const
+{
+	FChangelistIterator ChangelistIterator(Changelist, 0);
+	FRepHandleIterator HandleIterator(ChangelistIterator, Cmds, BaseHandleToCmdIndex, 0, 1, 0, Cmds.Num() - 1);
+
+	OutProperties.Empty();
+
+	while (HandleIterator.NextHandle())
+	{
+		const FRepLayoutCmd& Cmd = Cmds[HandleIterator.CmdIndex];
+		if (!InactiveParents[Cmd.ParentIndex])
+		{
+			OutProperties.Add(HandleIterator.Handle);
+
+			if (Cmd.Type == ERepLayoutCmdType::DynamicArray)
+			{
+				// No need to recursively filter the change list, as handles are only enabled/disabled at the parent level
+				int32 HandleCount = ChangelistIterator.Changed[ChangelistIterator.ChangedIndex];
+				OutProperties.Add(HandleCount);
+
+				for (int32 I = 0; I < HandleCount; ++I)
+				{
+					OutProperties.Add(ChangelistIterator.Changed[ChangelistIterator.ChangedIndex + 1 + I]);
+				}
+
+				OutProperties.Add(0);
+
+				HandleIterator.JumpOverArray();
+			}
+		}
+		else if (Cmd.Type == ERepLayoutCmdType::DynamicArray)
+		{
+			HandleIterator.JumpOverArray();
+		}
+	}
+
+	OutProperties.Add(0);
+}
+
 const FRepSerializedPropertyInfo* FRepSerializationSharedInfo::WriteSharedProperty(
 	const FRepLayoutCmd&	Cmd,
 	const FGuid&			PropertyGuid,
@@ -1468,22 +1593,8 @@ void FRepLayout::SendProperties_r(
 	while (HandleIterator.NextHandle())
 	{
 		const FRepLayoutCmd& Cmd = Cmds[HandleIterator.CmdIndex];
-
 		const FRepParentCmd& ParentCmd = Parents[Cmd.ParentIndex];
 
-		if (!RepState->ConditionMap[ParentCmd.Condition] || !ChangedTracker->Parents[Cmd.ParentIndex].Active)
-		{
-			if (Cmd.Type == ERepLayoutCmdType::DynamicArray)
-			{
-				if (!HandleIterator.JumpOverArray())
-				{
-					break;
-				}
-			}
-
-			continue;
-		}
-
 		const uint8* Data = SourceData + HandleIterator.ArrayOffset + Cmd.Offset;
 
 		if (Cmd.Type == ERepLayoutCmdType::DynamicArray)
@@ -1777,24 +1888,8 @@ void FRepLayout::SendProperties_BackwardsCompatible_r(
 	while (HandleIterator.NextHandle())
 	{
 		const FRepLayoutCmd& Cmd = Cmds[HandleIterator.CmdIndex];
-
 		const FRepParentCmd& ParentCmd = Parents[Cmd.ParentIndex];
 
-		const bool bConditionMatches = ChangedTracker == nullptr || (RepState->ConditionMap[ParentCmd.Condition] && ChangedTracker->Parents[Cmd.ParentIndex].Active);
-
-		if (!bConditionMatches)
-		{
-			if (Cmd.Type == ERepLayoutCmdType::DynamicArray)
-			{
-				if (!HandleIterator.JumpOverArray())
-				{
-					break;
-				}
-			}
-
-			continue;
-		}
-
 		const uint8* Data = SourceData + HandleIterator.ArrayOffset + Cmd.Offset;
 
 		PackageMapClient->TrackNetFieldExport(NetFieldExportGroup, HandleIterator.CmdIndex);
@@ -4704,12 +4799,10 @@ void FRepLayout::BuildHandleToCmdIndexTable_r(
 	}
 }
 
-void FRepLayout::RebuildConditionalProperties(
-	FRepState * RESTRICT				RepState,
-	const FReplicationFlags&			RepFlags) const
+TStaticBitArray<COND_Max> FRepState::BuildConditionMap(const FReplicationFlags& RepFlags)
 {
-	SCOPE_CYCLE_COUNTER(STAT_NetRebuildConditionalTime);
-	
+	TStaticBitArray<COND_Max> ConditionMap;
+
 	// Setup condition map
 	const bool bIsInitial = RepFlags.bNetInitial ? true : false;
 	const bool bIsOwner = RepFlags.bNetOwner ? true : false;
@@ -4717,27 +4810,47 @@ void FRepLayout::RebuildConditionalProperties(
 	const bool bIsPhysics = RepFlags.bRepPhysics ? true : false;
 	const bool bIsReplay = RepFlags.bReplay ? true : false;
 
-	RepState->ConditionMap[COND_None] = true;
-	RepState->ConditionMap[COND_InitialOnly] = bIsInitial;
+	ConditionMap[COND_None] = true;
+	ConditionMap[COND_InitialOnly] = bIsInitial;
 
-	RepState->ConditionMap[COND_OwnerOnly] = bIsOwner;
-	RepState->ConditionMap[COND_SkipOwner] = !bIsOwner;
+	ConditionMap[COND_OwnerOnly] = bIsOwner;
+	ConditionMap[COND_SkipOwner] = !bIsOwner;
 
-	RepState->ConditionMap[COND_SimulatedOnly] = bIsSimulated;
-	RepState->ConditionMap[COND_SimulatedOnlyNoReplay] = bIsSimulated && !bIsReplay;
-	RepState->ConditionMap[COND_AutonomousOnly] = !bIsSimulated;
+	ConditionMap[COND_SimulatedOnly] = bIsSimulated;
+	ConditionMap[COND_SimulatedOnlyNoReplay] = bIsSimulated && !bIsReplay;
+	ConditionMap[COND_AutonomousOnly] = !bIsSimulated;
 
-	RepState->ConditionMap[COND_SimulatedOrPhysics] = bIsSimulated || bIsPhysics;
-	RepState->ConditionMap[COND_SimulatedOrPhysicsNoReplay] = (bIsSimulated || bIsPhysics) && !bIsReplay;
+	ConditionMap[COND_SimulatedOrPhysics] = bIsSimulated || bIsPhysics;
+	ConditionMap[COND_SimulatedOrPhysicsNoReplay] = (bIsSimulated || bIsPhysics) && !bIsReplay;
 
-	RepState->ConditionMap[COND_InitialOrOwner] = bIsInitial || bIsOwner;
-	RepState->ConditionMap[COND_ReplayOrOwner] = bIsReplay || bIsOwner;
-	RepState->ConditionMap[COND_ReplayOnly] = bIsReplay;
-	RepState->ConditionMap[COND_SkipReplay] = !bIsReplay;
+	ConditionMap[COND_InitialOrOwner] = bIsInitial || bIsOwner;
+	ConditionMap[COND_ReplayOrOwner] = bIsReplay || bIsOwner;
+	ConditionMap[COND_ReplayOnly] = bIsReplay;
+	ConditionMap[COND_SkipReplay] = !bIsReplay;
 
-	RepState->ConditionMap[COND_Custom] = true;
+	ConditionMap[COND_Custom] = true;
+
+	return ConditionMap;
+}
+
+void FRepLayout::RebuildConditionalProperties(
+	FRepState * RESTRICT				RepState,
+	const FReplicationFlags&			RepFlags) const
+{
+	SCOPE_CYCLE_COUNTER(STAT_NetRebuildConditionalTime);
+	
+	TStaticBitArray<COND_Max> ConditionMap = FRepState::BuildConditionMap(RepFlags);
+	for (auto It = TBitArray<>::FIterator(RepState->InactiveParents); It; ++It)
+	{
+		It.GetValue() = !ConditionMap[Parents[It.GetIndex()].Condition];
+	}
 
 	RepState->RepFlags = RepFlags;
+
+PRAGMA_DISABLE_DEPRECATION_WARNINGS
+	// Keep this up to date for now, in case anyone is using it.
+	RepState->ConditionMap = MoveTemp(ConditionMap);
+PRAGMA_ENABLE_DEPRECATION_WARNINGS
 }
 
 void FRepLayout::InitChangedTracker(FRepChangedPropertyTracker* ChangedTracker) const
@@ -4797,6 +4910,7 @@ void FRepLayout::InitRepState(
 
 	// Start out the conditional props based on a default RepFlags struct
 	// It will rebuild if it ever changes
+	RepState->InactiveParents.Init(false, Parents.Num());
 	RebuildConditionalProperties(RepState, FReplicationFlags());
 }
 
@@ -4924,6 +5038,9 @@ void FRepState::CountBytes(FArchive& Ar) const
 	}
 
 	LifetimeChangelist.CountBytes(Ar);
+
+	InactiveChangelist.CountBytes(Ar);
+	InactiveParents.CountBytes(Ar);
 }
 
 FRepState::~FRepState()
diff --git a/Engine/Source/Runtime/Engine/Public/Net/RepLayout.h b/Engine/Source/Runtime/Engine/Public/Net/RepLayout.h
--- a/Engine/Source/Runtime/Engine/Public/Net/RepLayout.h
+++ b/Engine/Source/Runtime/Engine/Public/Net/RepLayout.h
@@ -433,12 +433,22 @@ public:
 		OpenAckedCalled(false),
 		AwakeFromDormancy(false),
 		LastChangelistIndex(0),
-		LastCompareIndex(0)
+		LastCompareIndex(0),
+		InactiveChangelist({0})
 	{}
 
 	~FRepState();
 
 	void CountBytes(FArchive& Ar) const;
+	
+	/**
+	 * Builds a new ConditionMap given the input RepFlags.
+	 * This can be used to determine whether or not a given property should be
+	 * considered enabled / disabled based on ELifetimeCondition.
+	 *
+	 * TODO: This doesn't have to be part of FRepState.
+	 */
+	static TStaticBitArray<COND_Max> BuildConditionMap(const FReplicationFlags& InFlags);
 
 	/** Latest state of all property data. Used on Clients, or on Servers if Shadow State is disabled. */
 	FRepStateStaticBuffer StaticBuffer;
@@ -503,6 +513,7 @@ public:
 	 * A map tracking which replication conditions are currently active.
 	 * @see ELifetimeCondition.
 	 */
+	UE_DEPRECATED(4.22, "Please use InactiveParents to determine whether or not a given ParentCommand is active.")
 	TStaticBitArray<COND_Max> ConditionMap;
 
 	// Cache off the RemoteRole and Role per connection to avoid issues with
@@ -510,6 +521,17 @@ public:
 
 	ENetRole SavedRemoteRole = ROLE_MAX;
 	ENetRole SavedRole = ROLE_MAX;
+
+	/**
+	 * Properties which are inactive through conditions have their changes stored here, so they can be 
+	 * applied if/when the property becomes active.
+	 *
+	 * This should always be a valid changelist, even if no properties are inactive.
+	 */
+	TArray<uint16> InactiveChangelist;
+
+	/** Cached set of inactive parent commands. */
+	TBitArray<> InactiveParents;
 };
 
 /** Various types of Properties supported for Replication. */
@@ -1058,6 +1080,9 @@ public:
 	 * Writes all changed property values from the input owner data to the given buffer.
 	 * This is used primarily by ReplicateProperties.
 	 *
+	 * Note, the changelist is expected to have any conditional properties whose conditions
+	 * aren't met filtered out already. See FRepState::ConditionMap and FRepLayout::FilterChangeList
+	 *
 	 * @param RepState			RepState for the object.
 	 * @param ChangedTracker	Used to indicate
 	 * @param Data				Pointer to the object's memory.
@@ -1618,7 +1643,27 @@ private:
 		FRepHandleIterator&		RepHandleIterator,
 		const uint8* RESTRICT	SourceData,
 		TArray<uint16>&			OutChanged) const;
-		
+
+	/**
+	 * Splits a given Changelist into an Inactive Change List and an Active Change List.
+	 * 
+	 * @param Changelist				The Changelist to filter.
+	 * @param InactiveParentHandles		The set of ParentCmd Indices that are not active.
+	 * @param OutInactiveProperties		The properties found to be inactive.
+	 * @param OutActiveProperties		The properties found to be active.
+	 */
+	void FilterChangeList( 
+		const TArray<uint16>&	Changelist,
+		const TBitArray<>&		InactiveParentHandles,
+		TArray<uint16>&			OutInactiveProperties,
+		TArray<uint16>&			OutActiveProperties) const;
+
+	/** Same as FilterChangeList, but only populates an Active Change List. */
+	void FilterChangeListToActive(
+		const TArray<uint16>&	Changelist,
+		const TBitArray<>&		InactiveParentHandles,
+		TArray<uint16>&			OutActiveProperties) const;
+
 	void BuildChangeList_r(
 		const TArray<FHandleToCmdIndex>&	HandleToCmdIndex,
 		const int32							CmdStart,
