Handle incorrect previous values passed to rep notifies for structs when all of the members do not change in the same frame

diff --git a/Engine/Source/Runtime/Engine/Private/RepLayout.cpp b/Engine/Source/Runtime/Engine/Private/RepLayout.cpp
--- a/Engine/Source/Runtime/Engine/Private/RepLayout.cpp
+++ b/Engine/Source/Runtime/Engine/Private/RepLayout.cpp
@@ -2149,7 +2149,7 @@ static bool ReceivePropertyHelper(
 		// Check to see if this property changed
 		if (Parent.RepNotifyCondition == REPNOTIFY_Always || !PropertiesAreIdentical(Cmd, ShadowData + Cmd, Data + SwappedCmd))
 		{
-			(*RepNotifies).AddUnique(Parent.Property);
+			RepNotifies->AddUnique(Parent.Property);
 		}
 		else
 		{
@@ -2979,11 +2979,9 @@ void FRepLayout::CallRepNotifies(FRepState* RepState, UObject* Object) const
 	check(LayoutState == ERepLayoutState::Normal);
 
 	FRepShadowDataBuffer ShadowData(RepState->StaticBuffer.GetData());
-	auto ConstParentIt = Parents.CreateConstIterator();
 
-	for (int32 i = 0; i < RepState->RepNotifies.Num(); i++)
+	for (UProperty* RepProperty : RepState->RepNotifies)
 	{
-		UProperty* RepProperty = RepState->RepNotifies[i];
 		UFunction* RepNotifyFunc = Object->FindFunction(RepProperty->RepNotifyFunc);
 
 		if (RepNotifyFunc == nullptr)
@@ -3001,7 +2999,15 @@ void FRepLayout::CallRepNotifies(FRepState* RepState, UObject* Object) const
 		}
 		else if (RepNotifyFunc->NumParms == 1)
 		{
-			Object->ProcessEvent(RepNotifyFunc, ShadowData + Parents[PropertyToParentHandle.FindChecked(RepProperty)]);
+			const FRepParentCmd& Parent = Parents[PropertyToParentHandle.FindChecked(RepProperty)];
+
+			Object->ProcessEvent(RepNotifyFunc, ShadowData + Parent);
+
+			// now store the complete value in the shadow buffer
+			if (!EnumHasAnyFlags(Parent.Flags, ERepParentFlags::IsNetSerialize | ERepParentFlags::IsCustomDelta))
+			{
+				RepProperty->CopyCompleteValue(ShadowData + Parent, RepProperty->ContainerPtrToValuePtr<uint8>(Object));
+			}
 		}
 	}
 
