Prevent destruction of level script actors by channels closing on the client

diff --git a/Engine/Source/Runtime/Engine/Private/NetworkDriver.cpp b/Engine/Source/Runtime/Engine/Private/NetworkDriver.cpp
--- a/Engine/Source/Runtime/Engine/Private/NetworkDriver.cpp
+++ b/Engine/Source/Runtime/Engine/Private/NetworkDriver.cpp
@@ -5429,7 +5429,7 @@ bool UNetDriver::ShouldCallRemoteFunction(UObject* Object, UFunction* Function,
 
 bool UNetDriver::ShouldClientDestroyActor(AActor* Actor) const
 {
-	return true;
+	return (Actor && !Actor->IsA(ALevelScriptActor::StaticClass()));
 }
 
 FAutoConsoleCommandWithWorld	DumpRelevantActorsCommand(
