Fixing non-editor build assert at TextNamespaceUtilImpl::FindOrAddPackageNamespace()

diff --git a/Engine/Source/Runtime/Engine/Private/DataTable.cpp b/Engine/Source/Runtime/Engine/Private/DataTable.cpp
--- a/Engine/Source/Runtime/Engine/Private/DataTable.cpp
+++ b/Engine/Source/Runtime/Engine/Private/DataTable.cpp
@@ -163,7 +163,10 @@ void UDataTable::OnPostDataImported(TArray<FString>& OutCollectedImportProblems)
 
 		FString DataTableTextNamespace = GetName();
 #if USE_STABLE_LOCALIZATION_KEYS
-		DataTableTextNamespace = TextNamespaceUtil::BuildFullNamespace(DataTableTextNamespace, TextNamespaceUtil::EnsurePackageNamespace(this), /*bAlwaysApplyPackageNamespace*/true);
+		if (GIsEditor)
+		{
+			DataTableTextNamespace = TextNamespaceUtil::BuildFullNamespace(DataTableTextNamespace, TextNamespaceUtil::EnsurePackageNamespace(this), /*bAlwaysApplyPackageNamespace*/true);
+		}
 #endif
 
 		for (const TPair<FName, uint8*>& TableRowPair : RowMap)
