HTML5 - stop "Resizing viewport due to setres change W x H" spam

#rnx

diff --git a/Engine/Source/Runtime/Engine/Private/GameEngine.cpp b/Engine/Source/Runtime/Engine/Private/GameEngine.cpp
--- a/Engine/Source/Runtime/Engine/Private/GameEngine.cpp
+++ b/Engine/Source/Runtime/Engine/Private/GameEngine.cpp
@@ -665,6 +665,7 @@ void UGameEngine::OnViewportResized(FViewport* Viewport, uint32 Unused)
 		{
 			GSystemResolution.ResX = ViewportSize.X;
 			GSystemResolution.ResY = ViewportSize.Y;
+			FSystemResolution::RequestResolutionChange(GSystemResolution.ResX, GSystemResolution.ResY, EWindowMode::Windowed);
 
 			UGameUserSettings* Settings = GetGameUserSettings();
 			Settings->SetScreenResolution(ViewportSize);
