Merging //UE4/Dev-Main to Dev-Mobile (//UE4/Dev-Mobile)

diff --git a/Engine/Plugins/Developer/PerforceSourceControl/Source/PerforceSourceControl/Private/PerforceSourceControlOperations.cpp b/Engine/Plugins/Developer/PerforceSourceControl/Source/PerforceSourceControl/Private/PerforceSourceControlOperations.cpp
--- a/Engine/Plugins/Developer/PerforceSourceControl/Source/PerforceSourceControl/Private/PerforceSourceControlOperations.cpp
+++ b/Engine/Plugins/Developer/PerforceSourceControl/Source/PerforceSourceControl/Private/PerforceSourceControlOperations.cpp
@@ -756,7 +756,7 @@ static void ParseUpdateStatusResults(const FP4RecordSet& InRecords, const TArray
 
 				int32 AtIndex = OtherOpenRecordValue.Find(TEXT("@"));
 				FString OtherOpenUser = AtIndex == INDEX_NONE ? FString(TEXT("")) : OtherOpenRecordValue.Left(AtIndex);
-				BranchModification.OtherUserCheckedOut += OtherOpenUser + TEXT(" @ ") + BranchModification.BranchName;
+				BranchModification.OtherUserCheckedOut += OtherOpenUser + TEXT(" @ ") + Branch;
 
 				if (OpenIdx < OtherOpenNum - 1)
 				{
diff --git a/Engine/Plugins/Media/MediaIOFramework/Source/MediaIOEditor/Public/Widgets/SMediaPermutationsSelector.inl b/Engine/Plugins/Media/MediaIOFramework/Source/MediaIOEditor/Public/Widgets/SMediaPermutationsSelector.inl
--- a/Engine/Plugins/Media/MediaIOFramework/Source/MediaIOEditor/Public/Widgets/SMediaPermutationsSelector.inl
+++ b/Engine/Plugins/Media/MediaIOFramework/Source/MediaIOEditor/Public/Widgets/SMediaPermutationsSelector.inl
@@ -22,7 +22,7 @@ ItemType SMediaPermutationsSelector<ItemType, ItemBuilder>::GetSelectedItem() co
 template <typename ItemType, typename ItemBuilder>
 void SMediaPermutationsSelector<ItemType, ItemBuilder>::Construct(const FArguments& InArgs)
 {
-	SWidget::Construct(InArgs._ToolTipText, InArgs._ToolTip, InArgs._Cursor, InArgs._IsEnabled, InArgs._Visibility, InArgs._RenderOpacity, InArgs._RenderTransform, InArgs._RenderTransformPivot, InArgs._Tag, InArgs._ForceVolatile, InArgs._Clipping, InArgs.MetaData);
+	SWidget::Construct(InArgs._ToolTipText, InArgs._ToolTip, InArgs._Cursor, InArgs._IsEnabled, InArgs._Visibility, InArgs._RenderOpacity, InArgs._RenderTransform, InArgs._RenderTransformPivot, InArgs._Tag, InArgs._ForceVolatile, InArgs._Clipping, InArgs._FlowDirectionPreference, InArgs.MetaData);
 
 	PermutationsSource = InArgs._PermutationsSource;
 	SelectedPermutationIndex = INDEX_NONE;
diff --git a/Engine/Plugins/Online/OnlineFramework/Source/Party/Private/User/SocialUser.cpp b/Engine/Plugins/Online/OnlineFramework/Source/Party/Private/User/SocialUser.cpp
--- a/Engine/Plugins/Online/OnlineFramework/Source/Party/Private/User/SocialUser.cpp
+++ b/Engine/Plugins/Online/OnlineFramework/Source/Party/Private/User/SocialUser.cpp
@@ -137,15 +137,18 @@ void USocialUser::InitLocalUser()
 		IOnlineSubsystem* OSS = OwningToolkit.GetSocialOss(SubsystemType);
 		check(OSS);
 
-		IOnlineIdentityPtr IdentityInterface = OSS->GetIdentityInterface();
-		FUniqueNetIdRepl LocalUserSubsystemId = IdentityInterface ? IdentityInterface->GetUniquePlayerId(OwningToolkit.GetLocalUserNum()) : FUniqueNetIdRepl();
-		if (ensure(LocalUserSubsystemId.IsValid()))
+		if (IOnlineIdentityPtr IdentityInterface = OSS->GetIdentityInterface())
 		{
-			SetSubsystemId(SubsystemType, LocalUserSubsystemId);
-		}
-		else
-		{
-			UE_LOG(LogParty, Error, TEXT("Local SocialUser unable to establish a valid UniqueId on subsystem [%s]"), ToString(SubsystemType));
+			const int32 LocalUserNum = OwningToolkit.GetLocalUserNum();
+			FUniqueNetIdRepl LocalUserSubsystemId = IdentityInterface->GetUniquePlayerId(LocalUserNum);
+			if (IdentityInterface->GetLoginStatus(LocalUserNum) == ELoginStatus::LoggedIn && ensure(LocalUserSubsystemId.IsValid()))
+			{
+				SetSubsystemId(SubsystemType, LocalUserSubsystemId);
+			}
+			else
+			{
+				UE_LOG(LogParty, Warning, TEXT("Local SocialUser unable to establish a valid UniqueId on subsystem [%s]"), ToString(SubsystemType));
+			}
 		}
 	}
 
diff --git a/Engine/Plugins/Runtime/GameplayAbilities/Source/GameplayAbilities/Private/AbilitySystemComponent_Abilities.cpp b/Engine/Plugins/Runtime/GameplayAbilities/Source/GameplayAbilities/Private/AbilitySystemComponent_Abilities.cpp
--- a/Engine/Plugins/Runtime/GameplayAbilities/Source/GameplayAbilities/Private/AbilitySystemComponent_Abilities.cpp
+++ b/Engine/Plugins/Runtime/GameplayAbilities/Source/GameplayAbilities/Private/AbilitySystemComponent_Abilities.cpp
@@ -30,7 +30,6 @@
 #include "Abilities/GameplayAbilityTargetActor.h"
 #include "TickableAttributeSetInterface.h"
 #include "GameplayTagResponseTable.h"
-#include "Engine/DemoNetDriver.h"
 #include "ProfilingDebugging/CsvProfiler.h"
 
 #define LOCTEXT_NAMESPACE "AbilitySystemComponent"
@@ -2519,7 +2518,7 @@ void UAbilitySystemComponent::OnRep_ReplicatedAnimMontage()
 		RepAnimMontageInfo.PlayRate = 1.f;
 	}
 
-	const bool bIsPlayingReplay = World && World->DemoNetDriver && World->DemoNetDriver->IsPlaying();
+	const bool bIsPlayingReplay = World && World->IsPlayingReplay();
 
 	const float MONTAGE_REP_POS_ERR_THRESH = bIsPlayingReplay ? CVarReplayMontageErrorThreshold.GetValueOnGameThread() : 0.1f;
 
diff --git a/Engine/Plugins/Runtime/GameplayAbilities/Source/GameplayAbilities/Private/GameplayCueManager.cpp b/Engine/Plugins/Runtime/GameplayAbilities/Source/GameplayAbilities/Private/GameplayCueManager.cpp
--- a/Engine/Plugins/Runtime/GameplayAbilities/Source/GameplayAbilities/Private/GameplayCueManager.cpp
+++ b/Engine/Plugins/Runtime/GameplayAbilities/Source/GameplayAbilities/Private/GameplayCueManager.cpp
@@ -438,6 +438,7 @@ AGameplayCueNotify_Actor* UGameplayCueManager::GetInstancedCueActor(AActor* Targ
 				ABILITY_LOG(Warning, TEXT("Spawning GameplaycueActor: %s"), *CueClass->GetName());
 			}
 
+			SpawnParams.OverrideLevel = World->PersistentLevel;
 			SpawnedCue = World->SpawnActor<AGameplayCueNotify_Actor>(CueClass, TargetActor->GetActorLocation(), TargetActor->GetActorRotation(), SpawnParams);
 		}
 
diff --git a/Engine/Plugins/Runtime/TimeSynth/Source/TimeSynth/Classes/TimeSynthComponent.h b/Engine/Plugins/Runtime/TimeSynth/Source/TimeSynth/Classes/TimeSynthComponent.h
--- a/Engine/Plugins/Runtime/TimeSynth/Source/TimeSynth/Classes/TimeSynthComponent.h
+++ b/Engine/Plugins/Runtime/TimeSynth/Source/TimeSynth/Classes/TimeSynthComponent.h
@@ -565,7 +565,7 @@ private:
 	void UpdateEnvelopeFollower();
 	void ShutdownPlayingClips();
 
-	Audio::SpectrumAnalyzerSettings::EFFTSize GetFFTSize(ETimeSynthFFTSize InSize) const;
+	Audio::FSpectrumAnalyzerSettings::EFFTSize GetFFTSize(ETimeSynthFFTSize InSize) const;
 
 	// Defines type for a volume group ID
 	typedef uint32 VolumeGroupUniqueId;
@@ -703,7 +703,7 @@ private:
 
 	// Spectum analyzer to allow BP delegates to visualize music
 	Audio::FSpectrumAnalyzer SpectrumAnalyzer;
-	Audio::SpectrumAnalyzerSettings::FSettings SpectrumAnalyzerSettings;
+	Audio::FSpectrumAnalyzerSettings SpectrumAnalyzerSettings;
 	FThreadSafeCounter SpectrumAnalysisCounter;
 
 	// Array of spectrum data, maps to FrequenciesToAnalyze UProperty
diff --git a/Engine/Plugins/Runtime/TimeSynth/Source/TimeSynth/Private/TimeSynthComponent.cpp b/Engine/Plugins/Runtime/TimeSynth/Source/TimeSynth/Private/TimeSynthComponent.cpp
--- a/Engine/Plugins/Runtime/TimeSynth/Source/TimeSynth/Private/TimeSynthComponent.cpp
+++ b/Engine/Plugins/Runtime/TimeSynth/Source/TimeSynth/Private/TimeSynthComponent.cpp
@@ -140,23 +140,23 @@ void UTimeSynthComponent::SetEnvelopeFollowerEnabled(bool bInIsEnabled)
 	});
 }
 
-Audio::SpectrumAnalyzerSettings::EFFTSize UTimeSynthComponent::GetFFTSize(ETimeSynthFFTSize InSize) const
+Audio::FSpectrumAnalyzerSettings::EFFTSize UTimeSynthComponent::GetFFTSize(ETimeSynthFFTSize InSize) const
 {
 	switch (InSize)
 	{
-		case ETimeSynthFFTSize::Min_64: return Audio::SpectrumAnalyzerSettings::EFFTSize::Min_64;
-		case ETimeSynthFFTSize::Small_256: return Audio::SpectrumAnalyzerSettings::EFFTSize::Small_256;
-		case ETimeSynthFFTSize::Medium_512: return Audio::SpectrumAnalyzerSettings::EFFTSize::Medium_512;
-		case ETimeSynthFFTSize::Large_1024: return Audio::SpectrumAnalyzerSettings::EFFTSize::Large_1024;
+		case ETimeSynthFFTSize::Min_64: return Audio::FSpectrumAnalyzerSettings::EFFTSize::Min_64;
+		case ETimeSynthFFTSize::Small_256: return Audio::FSpectrumAnalyzerSettings::EFFTSize::Small_256;
+		case ETimeSynthFFTSize::Medium_512: return Audio::FSpectrumAnalyzerSettings::EFFTSize::Medium_512;
+		case ETimeSynthFFTSize::Large_1024: return Audio::FSpectrumAnalyzerSettings::EFFTSize::Large_1024;
 		break;
 	}
 	// return default
-	return Audio::SpectrumAnalyzerSettings::EFFTSize::Medium_512;
+	return Audio::FSpectrumAnalyzerSettings::EFFTSize::Medium_512;
 }
 
 void UTimeSynthComponent::SetFFTSize(ETimeSynthFFTSize InFFTSize)
 {
-	Audio::SpectrumAnalyzerSettings::EFFTSize NewFFTSize = GetFFTSize(InFFTSize);
+	Audio::FSpectrumAnalyzerSettings::EFFTSize NewFFTSize = GetFFTSize(InFFTSize);
 
 	SynthCommand([this, NewFFTSize]
 	{
diff --git a/Engine/Source/Editor/Persona/Private/PersonaMeshDetails.cpp b/Engine/Source/Editor/Persona/Private/PersonaMeshDetails.cpp
--- a/Engine/Source/Editor/Persona/Private/PersonaMeshDetails.cpp
+++ b/Engine/Source/Editor/Persona/Private/PersonaMeshDetails.cpp
@@ -1448,6 +1448,19 @@ void FPersonaMeshDetails::RegenerateOneLOD(int32 LODIndex)
 			return;
 		}
 
+		//If we are doing inline reduction, restore the skeletalmesh LOD if it was already reduced and we want to reduced it again (user change the options)
+		//This will ensure the data is the same when starting the reduction
+		if (LODIndex == CurrentLODInfo.ReductionSettings.BaseLOD
+			&& CurrentLODInfo.bHasBeenSimplified
+			&& SkelMesh->IsReductionActive(LODIndex)
+			&& SkelMesh->GetImportedModel()->OriginalReductionSourceMeshData.IsValidIndex(LODIndex)
+			&& !SkelMesh->GetImportedModel()->OriginalReductionSourceMeshData[LODIndex]->IsEmpty())
+		{
+			//Restore the base LOD data
+			CurrentLODInfo.bHasBeenSimplified = false;
+			FLODUtilities::RestoreSkeletalMeshLODImportedData(SkelMesh, LODIndex, false);
+		}
+
 		FSkeletalMeshUpdateContext UpdateContext;
 		UpdateContext.SkeletalMesh = SkelMesh;
 		UpdateContext.AssociatedComponents.Push(GetPersonaToolkit()->GetPreviewMeshComponent());
@@ -1508,7 +1521,7 @@ FReply FPersonaMeshDetails::RegenerateLOD(int32 LODIndex)
 			else if (bIsReductionActive)
 			{
 				//Ask user a special permission when the base LOD can be reduce 
-				const FText Text(LOCTEXT("Warning_ReductionApplyingToImportedMesh_ReduceBaseLOD", "Are you sure you'd like to apply mesh reduction to the base LOD?"));
+				const FText Text(LOCTEXT("Warning_ReductionApplyingToImportedMesh_ReduceBaseLOD", "Are you sure you'd like to apply mesh reduction to the non-generated base LOD?"));
 				EAppReturnType::Type Ret = FMessageDialog::Open(EAppMsgType::YesNo, Text);
 				if (Ret == EAppReturnType::No)
 				{
diff --git a/Engine/Source/Editor/UMGEditor/Private/Customizations/UMGDetailCustomizations.cpp b/Engine/Source/Editor/UMGEditor/Private/Customizations/UMGDetailCustomizations.cpp
--- a/Engine/Source/Editor/UMGEditor/Private/Customizations/UMGDetailCustomizations.cpp
+++ b/Engine/Source/Editor/UMGEditor/Private/Customizations/UMGDetailCustomizations.cpp
@@ -281,6 +281,9 @@ void FBlueprintWidgetCustomization::CreateMulticastEventCustomization(IDetailLay
 void FBlueprintWidgetCustomization::CustomizeDetails( IDetailLayoutBuilder& DetailLayout )
 {
 	static const FName LayoutCategoryKey(TEXT("Layout"));
+	static const FName LocalizationCategoryKey(TEXT("Localization"));
+
+	DetailLayout.EditCategory(LocalizationCategoryKey, FText::GetEmpty(), ECategoryPriority::Uncommon);
 
 	TArray< TWeakObjectPtr<UObject> > OutObjects;
 	DetailLayout.GetObjectsBeingCustomized(OutObjects);
diff --git a/Engine/Source/Editor/UnrealEd/Private/Factories/SkeletalMeshImport.cpp b/Engine/Source/Editor/UnrealEd/Private/Factories/SkeletalMeshImport.cpp
--- a/Engine/Source/Editor/UnrealEd/Private/Factories/SkeletalMeshImport.cpp
+++ b/Engine/Source/Editor/UnrealEd/Private/Factories/SkeletalMeshImport.cpp
@@ -810,6 +810,10 @@ void RestoreExistingSkelMeshData(ExistingSkelMeshData* MeshData, USkeletalMesh*
 		FSkeletalMeshLODInfo& BaseLODInfo = SkeletalMesh->GetLODInfoArray()[ReimportLODIndex];
 		//Restore the reimport LOD reduction settings
 		BaseLODInfo.ReductionSettings = MeshData->ExistingReimportLODReductionSettings;
+		if (SkeletalMeshImportedModel && SkeletalMeshImportedModel->OriginalReductionSourceMeshData.IsValidIndex(ReimportLODIndex))
+		{
+			SkeletalMeshImportedModel->OriginalReductionSourceMeshData[ReimportLODIndex]->EmptyBulkData();
+		}
 		//Regenerate the reimport LOD
 		GWarn->BeginSlowTask(LOCTEXT("RegenReimportedLOD", "Generating reimported LOD"), true);
 		FSkeletalMeshUpdateContext UpdateContext;
diff --git a/Engine/Source/Editor/UnrealEd/Private/Fbx/FbxMainExport.cpp b/Engine/Source/Editor/UnrealEd/Private/Fbx/FbxMainExport.cpp
--- a/Engine/Source/Editor/UnrealEd/Private/Fbx/FbxMainExport.cpp
+++ b/Engine/Source/Editor/UnrealEd/Private/Fbx/FbxMainExport.cpp
@@ -3558,11 +3558,22 @@ void FFbxExporter::ExportObjectMetadata(const UObject* ObjectToExport, FbxNode*
 					// Remaining tag follows the format NodeName.PropertyName, so replace '.' with '_'
 					TagAsString.ReplaceInline(TEXT("."), TEXT("_"));
 
-					FbxProperty Property = FbxProperty::Create(Node, FbxStringDT, TCHAR_TO_UTF8(*TagAsString));
-					FbxString ValueString(TCHAR_TO_UTF8(*MetadataIt.Value));
+					if (MetadataIt.Value == TEXT("true") || MetadataIt.Value == TEXT("false"))
+					{
+						FbxProperty Property = FbxProperty::Create(Node, FbxBoolDT, TCHAR_TO_UTF8(*TagAsString));
+						FbxBool ValueBool = MetadataIt.Value == TEXT("true") ? true : false;
+
+						Property.Set(ValueBool);
+						Property.ModifyFlag(FbxPropertyFlags::eUserDefined, true);
+					}
+					else
+					{
+						FbxProperty Property = FbxProperty::Create(Node, FbxStringDT, TCHAR_TO_UTF8(*TagAsString));
+						FbxString ValueString(TCHAR_TO_UTF8(*MetadataIt.Value));
 
-					Property.Set(ValueString);
-					Property.ModifyFlag(FbxPropertyFlags::eUserDefined, true);
+						Property.Set(ValueString);
+						Property.ModifyFlag(FbxPropertyFlags::eUserDefined, true);
+					}
 				}
 			}
 		}
diff --git a/Engine/Source/Editor/UnrealEd/Private/Fbx/FbxMainImport.cpp b/Engine/Source/Editor/UnrealEd/Private/Fbx/FbxMainImport.cpp
--- a/Engine/Source/Editor/UnrealEd/Private/Fbx/FbxMainImport.cpp
+++ b/Engine/Source/Editor/UnrealEd/Private/Fbx/FbxMainImport.cpp
@@ -3107,7 +3107,7 @@ FString GetFbxPropertyStringValue(const FbxProperty& Property)
 	return ValueStr;
 }
 
-void FFbxImporter::ImportNodeCustomProperties(UObject* Object, FbxNode* Node)
+void FFbxImporter::ImportNodeCustomProperties(UObject* Object, FbxNode* Node, bool bPrefixTagWithNodeName)
 {
 	if (!Object || !Node)
 	{
@@ -3125,7 +3125,7 @@ void FFbxImporter::ImportNodeCustomProperties(UObject* Object, FbxNode* Node)
 			// Prefix the FBX metadata tag to make it distinguishable from other metadata
 			// so that it can be exportable through FBX export
 			FString MetadataTag = UTF8_TO_TCHAR(CurrentProperty.GetName());
-			if (!MetadataTag.StartsWith(NodeName))
+			if (bPrefixTagWithNodeName && !MetadataTag.StartsWith(NodeName))
 			{
 				// Append the node name in the tag since all the metadata will be flattened on the Object
 				MetadataTag = NodeName + TEXT(".") + MetadataTag;
@@ -3141,7 +3141,7 @@ void FFbxImporter::ImportNodeCustomProperties(UObject* Object, FbxNode* Node)
 	int NumChildren = Node->GetChildCount();
 	for (int i = 0; i < NumChildren; ++i)
 	{
-		ImportNodeCustomProperties(Object, Node->GetChild(i));
+		ImportNodeCustomProperties(Object, Node->GetChild(i), bPrefixTagWithNodeName);
 	}
 }
 
diff --git a/Engine/Source/Editor/UnrealEd/Private/Fbx/FbxSkeletalMeshExport.cpp b/Engine/Source/Editor/UnrealEd/Private/Fbx/FbxSkeletalMeshExport.cpp
--- a/Engine/Source/Editor/UnrealEd/Private/Fbx/FbxSkeletalMeshExport.cpp
+++ b/Engine/Source/Editor/UnrealEd/Private/Fbx/FbxSkeletalMeshExport.cpp
@@ -530,11 +530,22 @@ void ExportObjectMetadataToBones(const UObject* ObjectToExport, const TArray<Fbx
 				FbxNode** Node = NameToNode.Find(NodeName);
 				if (Node)
 				{
-					FbxProperty Property = FbxProperty::Create(*Node, FbxStringDT, TCHAR_TO_UTF8(*TagAsString));
-					FbxString ValueString(TCHAR_TO_UTF8(*MetadataIt.Value));
+					if (MetadataIt.Value == TEXT("true") || MetadataIt.Value == TEXT("false"))
+					{
+						FbxProperty Property = FbxProperty::Create(*Node, FbxBoolDT, TCHAR_TO_UTF8(*TagAsString));
+						FbxBool ValueBool = MetadataIt.Value == TEXT("true") ? true : false;
 
-					Property.Set(ValueString);
-					Property.ModifyFlag(FbxPropertyFlags::eUserDefined, true);
+						Property.Set(ValueBool);
+						Property.ModifyFlag(FbxPropertyFlags::eUserDefined, true);
+					}
+					else
+					{
+						FbxProperty Property = FbxProperty::Create(*Node, FbxStringDT, TCHAR_TO_UTF8(*TagAsString));
+						FbxString ValueString(TCHAR_TO_UTF8(*MetadataIt.Value));
+
+						Property.Set(ValueString);
+						Property.ModifyFlag(FbxPropertyFlags::eUserDefined, true);
+					}
 				}
 			}
 		}
diff --git a/Engine/Source/Editor/UnrealEd/Private/Fbx/FbxSkeletalMeshImport.cpp b/Engine/Source/Editor/UnrealEd/Private/Fbx/FbxSkeletalMeshImport.cpp
--- a/Engine/Source/Editor/UnrealEd/Private/Fbx/FbxSkeletalMeshImport.cpp
+++ b/Engine/Source/Editor/UnrealEd/Private/Fbx/FbxSkeletalMeshImport.cpp
@@ -1905,13 +1905,13 @@ USkeletalMesh* UnFbx::FFbxImporter::ImportSkeletalMesh(FImportSkeletalMeshArgs &
 	// Import mesh metadata to SkeletalMesh
 	for (FbxNode* MeshNode : ImportSkeletalMeshArgs.NodeArray)
 	{
-		ImportNodeCustomProperties(SkeletalMesh, MeshNode);
+		ImportNodeCustomProperties(SkeletalMesh, MeshNode, true);
 	}
 
 	// Import bone metadata to Skeleton
 	for (FbxNode* SkeletonNode : ImportSkeletalMeshArgs.BoneNodeArray)
 	{
-		ImportNodeCustomProperties(SkeletalMesh->Skeleton, SkeletonNode);
+		ImportNodeCustomProperties(SkeletalMesh->Skeleton, SkeletonNode, true);
 	}
 
 	// ComponentContexts will now go out of scope, causing components to be re-registered
diff --git a/Engine/Source/Editor/UnrealEd/Private/LODUtilities.cpp b/Engine/Source/Editor/UnrealEd/Private/LODUtilities.cpp
--- a/Engine/Source/Editor/UnrealEd/Private/LODUtilities.cpp
+++ b/Engine/Source/Editor/UnrealEd/Private/LODUtilities.cpp
@@ -422,6 +422,11 @@ void CreateLODMorphTarget(USkeletalMesh* SkeletalMesh, FReductionBaseSkeletalMes
 		ReductionBaseSkeletalMeshBulkData->LoadReductionData(TempBaseLODModel, BaseLODMorphTargetData);
 	}
 
+	FSkeletalMeshModel* SkeletalMeshModel = SkeletalMesh->GetImportedModel();
+	const FSkeletalMeshLODModel& TargetLODModel = SkeletalMeshModel->LODModels[DestinationLOD];
+
+	bool bInitializeMorphData = false;
+
 	for (UMorphTarget *MorphTarget : SkeletalMesh->MorphTargets)
 	{
 		if (!MorphTarget->HasDataForLOD(SourceLOD))
@@ -517,11 +522,17 @@ void CreateLODMorphTarget(USkeletalMesh* SkeletalMesh, FReductionBaseSkeletalMes
 			}
 		}
 		
-		FSkeletalMeshModel* SkeletalMeshModel = SkeletalMesh->GetImportedModel();
-		const FSkeletalMeshLODModel& TargetLODModel = SkeletalMeshModel->LODModels[DestinationLOD];
 		//Register the new morph target on the target LOD
 		MorphTarget->PopulateDeltas(NewMorphTargetDeltas, DestinationLOD, TargetLODModel.Sections, false, true);
-		SkeletalMesh->RegisterMorphTarget(MorphTarget);
+		if (MorphTarget->HasValidData())
+		{
+			bInitializeMorphData |= SkeletalMesh->RegisterMorphTarget(MorphTarget, false);
+		}
+	}
+
+	if (bInitializeMorphData)
+	{
+		SkeletalMesh->InitMorphTargetsAndRebuildRenderData();
 	}
 }
 
@@ -545,7 +556,7 @@ void FLODUtilities::ClearGeneratedMorphTarget(USkeletalMesh* SkeletalMesh, int32
 			continue;
 		}
 
-		if (MorphTarget->MorphLODModels[TargetLOD].bGeneratedByEngine)
+		//if (MorphTarget->MorphLODModels[TargetLOD].bGeneratedByEngine)
 		{
 			MorphTarget->MorphLODModels[TargetLOD].Reset();
 
@@ -870,6 +881,7 @@ void FLODUtilities::RestoreSkeletalMeshLODImportedData(USkeletalMesh* SkeletalMe
 		//Copy the SkeletalMeshLODModel
 		SkeletalMesh->GetImportedModel()->LODModels[LodIndex] = ImportedBaseLODModel;
 		//Copy the morph target deltas
+		bool bInitMorphTargetData = false;
 		for (UMorphTarget *MorphTarget : SkeletalMesh->MorphTargets)
 		{
 			if (!ImportedBaseLODMorphTargetData.Contains(MorphTarget->GetFullName()))
@@ -879,8 +891,9 @@ void FLODUtilities::RestoreSkeletalMeshLODImportedData(USkeletalMesh* SkeletalMe
 			TArray<FMorphTargetDelta>& ImportedDeltas = ImportedBaseLODMorphTargetData[MorphTarget->GetFullName()];
 
 			MorphTarget->PopulateDeltas(ImportedDeltas, LodIndex, SkeletalMesh->GetImportedModel()->LODModels[LodIndex].Sections, false, false);
-			SkeletalMesh->RegisterMorphTarget(MorphTarget);
+			bInitMorphTargetData |= SkeletalMesh->RegisterMorphTarget(MorphTarget, false);
 		}
+		SkeletalMesh->InitMorphTargetsAndRebuildRenderData();
 		
 		//Empty the bulkdata since we restore it
 		SkeletalMesh->GetImportedModel()->OriginalReductionSourceMeshData[LodIndex]->EmptyBulkData();
diff --git a/Engine/Source/Editor/UnrealEd/Public/FbxImporter.h b/Engine/Source/Editor/UnrealEd/Public/FbxImporter.h
--- a/Engine/Source/Editor/UnrealEd/Public/FbxImporter.h
+++ b/Engine/Source/Editor/UnrealEd/Public/FbxImporter.h
@@ -1719,7 +1719,7 @@ protected:
 	bool RetrievePoseFromBindPose(const TArray<FbxNode*>& NodeArray, FbxArray<FbxPose*> & PoseArray) const;
 
 	/** Import the user-defined properties on the node as FBX metadata on the object */
-	void ImportNodeCustomProperties(UObject* Object, FbxNode* Node);
+	void ImportNodeCustomProperties(UObject* Object, FbxNode* Node, bool bPrefixTagWithNodeName = false);
 
 public:
 	/** Import and set up animation related data from mesh **/
diff --git a/Engine/Source/Programs/UnrealBuildTool/Configuration/UEBuildModuleCPP.cs b/Engine/Source/Programs/UnrealBuildTool/Configuration/UEBuildModuleCPP.cs
--- a/Engine/Source/Programs/UnrealBuildTool/Configuration/UEBuildModuleCPP.cs
+++ b/Engine/Source/Programs/UnrealBuildTool/Configuration/UEBuildModuleCPP.cs
@@ -82,6 +82,7 @@ namespace UnrealBuildTool
 			new KeyValuePair<string, string>("Engine", "GameplayTags"),
 			new KeyValuePair<string, string>("Engine", "MaterialShaderQualitySettings"),
 			new KeyValuePair<string, string>("Engine", "UnrealEd"),
+			new KeyValuePair<string, string>("Engine", "AudioMixer"),
 			new KeyValuePair<string, string>("PacketHandler", "ReliabilityHandlerComponent"),
 			new KeyValuePair<string, string>("GameplayDebugger", "AIModule"),
 			new KeyValuePair<string, string>("GameplayDebugger", "GameplayTasks"),
diff --git a/Engine/Source/Programs/UnrealBuildTool/Platform/TVOS/UEBuildTVOS.cs b/Engine/Source/Programs/UnrealBuildTool/Platform/TVOS/UEBuildTVOS.cs
--- a/Engine/Source/Programs/UnrealBuildTool/Platform/TVOS/UEBuildTVOS.cs
+++ b/Engine/Source/Programs/UnrealBuildTool/Platform/TVOS/UEBuildTVOS.cs
@@ -17,7 +17,7 @@ namespace UnrealBuildTool
 		/// </summary>
 		public override string RuntimeVersion
 		{
-			get { return "9.0"; }
+			get { return "10.0"; }
 		}
 
 		/// <summary>
diff --git a/Engine/Source/Runtime/ApplicationCore/Private/Android/AndroidInputInterface.cpp b/Engine/Source/Runtime/ApplicationCore/Private/Android/AndroidInputInterface.cpp
--- a/Engine/Source/Runtime/ApplicationCore/Private/Android/AndroidInputInterface.cpp
+++ b/Engine/Source/Runtime/ApplicationCore/Private/Android/AndroidInputInterface.cpp
@@ -88,6 +88,16 @@ FAndroidInputInterface::FAndroidInputInterface(const TSharedRef< FGenericApplica
 	ButtonMapping[16] = AndroidKeyNames::Android_Back;  // Technically just an alias for SpecialLeft
 	ButtonMapping[17] = AndroidKeyNames::Android_Menu;  // Technically just an alias for SpecialRight
 
+	// Virtual buttons
+	ButtonMapping[MAX_NUM_PHYSICAL_CONTROLLER_BUTTONS + 0] = FGamepadKeyNames::LeftStickLeft;
+	ButtonMapping[MAX_NUM_PHYSICAL_CONTROLLER_BUTTONS + 1] = FGamepadKeyNames::LeftStickRight;
+	ButtonMapping[MAX_NUM_PHYSICAL_CONTROLLER_BUTTONS + 2] = FGamepadKeyNames::LeftStickUp;
+	ButtonMapping[MAX_NUM_PHYSICAL_CONTROLLER_BUTTONS + 3] = FGamepadKeyNames::LeftStickDown;
+	ButtonMapping[MAX_NUM_PHYSICAL_CONTROLLER_BUTTONS + 4] = FGamepadKeyNames::RightStickLeft;
+	ButtonMapping[MAX_NUM_PHYSICAL_CONTROLLER_BUTTONS + 5] = FGamepadKeyNames::RightStickRight;
+	ButtonMapping[MAX_NUM_PHYSICAL_CONTROLLER_BUTTONS + 6] = FGamepadKeyNames::RightStickUp;
+	ButtonMapping[MAX_NUM_PHYSICAL_CONTROLLER_BUTTONS + 7] = FGamepadKeyNames::RightStickDown;
+
 	InitialButtonRepeatDelay = 0.2f;
 	ButtonRepeatDelay = 0.1f;
 
@@ -955,21 +965,29 @@ void FAndroidInputInterface::SendControllerEvents()
 			if (NewControllerState.LXAnalog != OldControllerState.LXAnalog || FMath::Abs(NewControllerState.LXAnalog) >= RepeatDeadzone)
 			{
 				MessageHandler->OnControllerAnalog(FGamepadKeyNames::LeftAnalogX, NewControllerState.DeviceId, NewControllerState.LXAnalog);
+				NewControllerState.ButtonStates[MAX_NUM_PHYSICAL_CONTROLLER_BUTTONS + 1] = NewControllerState.LXAnalog >= RepeatDeadzone;
+				NewControllerState.ButtonStates[MAX_NUM_PHYSICAL_CONTROLLER_BUTTONS + 0] = NewControllerState.LXAnalog <= -RepeatDeadzone;
 			}
 			if (NewControllerState.LYAnalog != OldControllerState.LYAnalog || FMath::Abs(NewControllerState.LYAnalog) >= RepeatDeadzone)
 			{
 				//LOGD("    Sending updated LeftAnalogY value of %f", NewControllerState.LYAnalog);
 				MessageHandler->OnControllerAnalog(FGamepadKeyNames::LeftAnalogY, NewControllerState.DeviceId, NewControllerState.LYAnalog);
+				NewControllerState.ButtonStates[MAX_NUM_PHYSICAL_CONTROLLER_BUTTONS + 2] = NewControllerState.LYAnalog >= RepeatDeadzone;
+				NewControllerState.ButtonStates[MAX_NUM_PHYSICAL_CONTROLLER_BUTTONS + 3] = NewControllerState.LYAnalog <= -RepeatDeadzone;
 			}
 			if (NewControllerState.RXAnalog != OldControllerState.RXAnalog || FMath::Abs(NewControllerState.RXAnalog) >= RepeatDeadzone)
 			{
 				//LOGD("    Sending updated RightAnalogX value of %f", NewControllerState.RXAnalog);
 				MessageHandler->OnControllerAnalog(FGamepadKeyNames::RightAnalogX, NewControllerState.DeviceId, NewControllerState.RXAnalog);
+				NewControllerState.ButtonStates[MAX_NUM_PHYSICAL_CONTROLLER_BUTTONS + 5] = NewControllerState.RXAnalog >= RepeatDeadzone;
+				NewControllerState.ButtonStates[MAX_NUM_PHYSICAL_CONTROLLER_BUTTONS + 4] = NewControllerState.RXAnalog <= -RepeatDeadzone;
 			}
 			if (NewControllerState.RYAnalog != OldControllerState.RYAnalog || FMath::Abs(NewControllerState.RYAnalog) >= RepeatDeadzone)
 			{
 				//LOGD("    Sending updated RightAnalogY value of %f", NewControllerState.RYAnalog);
 				MessageHandler->OnControllerAnalog(FGamepadKeyNames::RightAnalogY, NewControllerState.DeviceId, NewControllerState.RYAnalog);
+				NewControllerState.ButtonStates[MAX_NUM_PHYSICAL_CONTROLLER_BUTTONS + 6] = NewControllerState.RYAnalog >= RepeatDeadzone;
+				NewControllerState.ButtonStates[MAX_NUM_PHYSICAL_CONTROLLER_BUTTONS + 7] = NewControllerState.RYAnalog <= -RepeatDeadzone;
 			}
 			if (NewControllerState.LTAnalog != OldControllerState.LTAnalog)
 			{
diff --git a/Engine/Source/Runtime/ApplicationCore/Private/IOS/IOSInputInterface.cpp b/Engine/Source/Runtime/ApplicationCore/Private/IOS/IOSInputInterface.cpp
--- a/Engine/Source/Runtime/ApplicationCore/Private/IOS/IOSInputInterface.cpp
+++ b/Engine/Source/Runtime/ApplicationCore/Private/IOS/IOSInputInterface.cpp
@@ -338,15 +338,15 @@ void FIOSInputInterface::SendControllerEvents()
         const double CurrentTime = FPlatformTime::Seconds();
         const float IniitialRepeatDelay = 0.2f;
         const float RepeatDelay = 0.1;
-		
-#define HANDLE_BUTTON(Gamepad, GCButton, UEButton) \
-if ((Controller.Previous##Gamepad == nil && Gamepad.GCButton.pressed) || Controller.Previous##Gamepad.GCButton.pressed != Gamepad.GCButton.pressed) \
+
+#define HANDLE_BUTTON_INTERNAL(Gamepad, bWasPressed, bIsPressed, UEButton) \
+if (bWasPressed != bIsPressed) \
 { \
-	NSLog(@"%@ button %s on controller %d", (ExtendedGamepad.GCButton.pressed) ? @"Pressed" : @"Released", TCHAR_TO_ANSI(*UEButton.ToString()), (int32)Cont.playerIndex); \
-	(Gamepad.GCButton.pressed) ? MessageHandler->OnControllerButtonPressed(UEButton, Cont.playerIndex, false) : MessageHandler->OnControllerButtonReleased(UEButton, Cont.playerIndex, false); \
+	NSLog(@"%@ button %s on controller %d", bIsPressed ? @"Pressed" : @"Released", TCHAR_TO_ANSI(*UEButton.ToString()), (int32)Cont.playerIndex); \
+	bIsPressed ? MessageHandler->OnControllerButtonPressed(UEButton, Cont.playerIndex, false) : MessageHandler->OnControllerButtonReleased(UEButton, Cont.playerIndex, false); \
     NextKeyRepeatTime.FindOrAdd(UEButton) = CurrentTime + IniitialRepeatDelay; \
 } \
-else if(Gamepad.GCButton.pressed) \
+else if(bIsPressed) \
 { \
     double* NextRepeatTime = NextKeyRepeatTime.Find(UEButton); \
     if(NextRepeatTime && *NextRepeatTime <= CurrentTime) \
@@ -359,6 +359,13 @@ else \
 { \
     NextKeyRepeatTime.Remove(UEButton); \
 }
+		
+#define HANDLE_BUTTON(Gamepad, GCButton, UEButton) \
+{ \
+	const bool bWasPressed = Controller.Previous##Gamepad != nil && Controller.Previous##Gamepad.GCButton.pressed; \
+	const bool bPressed = Gamepad.GCButton.pressed; \
+	HANDLE_BUTTON_INTERNAL(Gamepad, bWasPressed, bPressed, UEButton); \
+}
 
         // Send controller events any time we are passed the given input threshold similarly to PC/Console (see: XInputInterface.cpp)
         const float RepeatDeadzone = 0.24;
@@ -369,6 +376,16 @@ if ((Controller.Previous##Gamepad != nil && Gamepad.GCAxis.value != Controller.P
 	NSLog(@"Axis %s is %f", TCHAR_TO_ANSI(*UEAxis.ToString()), Gamepad.GCAxis.value); \
 	MessageHandler->OnControllerAnalog(UEAxis, Cont.playerIndex, Gamepad.GCAxis.value); \
 }
+
+#define HANDLE_ANALOG_VIRTUAL_BUTTONS(Gamepad, GCAxis, UEButtonNegative, UEButtonPositive) \
+{ \
+	const bool bWasNegativePressed = Controller.Previous##Gamepad != nil && Controller.Previous##Gamepad.GCAxis.value <= -RepeatDeadzone; \
+	const bool bNegativePressed = Gamepad.GCAxis.value <= -RepeatDeadzone; \
+	HANDLE_BUTTON_INTERNAL(Gamepad, bWasNegativePressed, bNegativePressed, UEButtonNegative) \
+	const bool bWasPositivePressed = Controller.Previous##Gamepad != nil && Controller.Previous##Gamepad.GCAxis.value >= RepeatDeadzone; \
+	const bool bPositivePressed = Gamepad.GCAxis.value >= RepeatDeadzone; \
+	HANDLE_BUTTON_INTERNAL(Gamepad, bWasPositivePressed, bPositivePressed, UEButtonPositive) \
+}
 		
 		if (ExtendedGamepad != nil)
 		{
@@ -392,6 +409,11 @@ if ((Controller.Previous##Gamepad != nil && Gamepad.GCAxis.value != Controller.P
 			HANDLE_ANALOG(ExtendedGamepad, leftTrigger,				FGamepadKeyNames::LeftTriggerAnalog);
 			HANDLE_ANALOG(ExtendedGamepad, rightTrigger,			FGamepadKeyNames::RightTriggerAnalog);
 
+			HANDLE_ANALOG_VIRTUAL_BUTTONS(ExtendedGamepad, leftThumbstick.xAxis, FGamepadKeyNames::LeftStickLeft, FGamepadKeyNames::LeftStickRight);
+			HANDLE_ANALOG_VIRTUAL_BUTTONS(ExtendedGamepad, leftThumbstick.yAxis, FGamepadKeyNames::LeftStickDown, FGamepadKeyNames::LeftStickUp);
+			HANDLE_ANALOG_VIRTUAL_BUTTONS(ExtendedGamepad, rightThumbstick.xAxis, FGamepadKeyNames::RightStickLeft, FGamepadKeyNames::RightStickRight);
+			HANDLE_ANALOG_VIRTUAL_BUTTONS(ExtendedGamepad, rightThumbstick.yAxis, FGamepadKeyNames::RightStickDown, FGamepadKeyNames::RightStickUp);
+
 			[Controller.PreviousExtendedGamepad release];
 			Controller.PreviousExtendedGamepad = [ExtendedGamepad saveSnapshot];
 			[Controller.PreviousExtendedGamepad retain];
diff --git a/Engine/Source/Runtime/ApplicationCore/Public/Android/AndroidInputInterface.h b/Engine/Source/Runtime/ApplicationCore/Public/Android/AndroidInputInterface.h
--- a/Engine/Source/Runtime/ApplicationCore/Public/Android/AndroidInputInterface.h
+++ b/Engine/Source/Runtime/ApplicationCore/Public/Android/AndroidInputInterface.h
@@ -153,7 +153,9 @@ struct TouchInput
 };
 
 #define MAX_NUM_CONTROLLERS					8  // reasonable limit for now
-#define MAX_NUM_CONTROLLER_BUTTONS			18
+#define MAX_NUM_PHYSICAL_CONTROLLER_BUTTONS	18
+#define MAX_NUM_VIRTUAL_CONTROLLER_BUTTONS	8
+#define MAX_NUM_CONTROLLER_BUTTONS			MAX_NUM_PHYSICAL_CONTROLLER_BUTTONS + MAX_NUM_VIRTUAL_CONTROLLER_BUTTONS
 #define MAX_DEFERRED_MESSAGE_QUEUE_SIZE		128
 
 struct FAndroidControllerData
diff --git a/Engine/Source/Runtime/AudioMixer/AudioMixer.Build.cs b/Engine/Source/Runtime/AudioMixer/AudioMixer.Build.cs
--- a/Engine/Source/Runtime/AudioMixer/AudioMixer.Build.cs
+++ b/Engine/Source/Runtime/AudioMixer/AudioMixer.Build.cs
@@ -6,8 +6,6 @@ namespace UnrealBuildTool.Rules
 	{
 		public AudioMixer(ReadOnlyTargetRules Target) : base(Target)
 		{
-			OptimizeCode = CodeOptimization.Always;
-
 			PrivateIncludePathModuleNames.Add("TargetPlatform");
 
 			PrivateIncludePaths.AddRange(
@@ -22,7 +20,6 @@ namespace UnrealBuildTool.Rules
 				{
 					"Core",
 					"CoreUObject",
-					"Engine",
 				}
 			);
 
@@ -30,6 +27,7 @@ namespace UnrealBuildTool.Rules
 				new string[]
 				{
 					"CoreUObject",
+					"Engine",
 				}
 			);
 
diff --git a/Engine/Source/Runtime/AudioMixer/Private/AudioMixerBlueprintLibrary.cpp b/Engine/Source/Runtime/AudioMixer/Private/AudioMixerBlueprintLibrary.cpp
--- a/Engine/Source/Runtime/AudioMixer/Private/AudioMixerBlueprintLibrary.cpp
+++ b/Engine/Source/Runtime/AudioMixer/Private/AudioMixerBlueprintLibrary.cpp
@@ -5,6 +5,7 @@
 #include "AudioDevice.h"
 #include "AudioMixerDevice.h"
 #include "CoreMinimal.h"
+#include "DSP/SpectrumAnalyzer.h"
 
 
 // This is our global recording task:
@@ -190,6 +191,56 @@ void UAudioMixerBlueprintLibrary::ResumeRecordingOutput(const UObject* WorldCont
 	}
 }
 
+void UAudioMixerBlueprintLibrary::StartAnalyzingOutput(const UObject* WorldContextObject, USoundSubmix* SubmixToAnalyze /*= nullptr*/, EFFTSize FFTSize /*= EFFTSize::Default*/, EFFTPeakInterpolationMethod InterpolationMethod /*= Linear*/, EFFTWindowType WindowType /*= EFFTWindowType::Hamming*/, float HopSize /*= 0*/)
+{
+	if (Audio::FMixerDevice* MixerDevice = GetAudioMixerDeviceFromWorldContext(WorldContextObject))
+	{
+		Audio::FSpectrumAnalyzerSettings Settings = Audio::FSpectrumAnalyzerSettings();
+		PopulateSpectrumAnalyzerSettings(FFTSize, InterpolationMethod, WindowType, HopSize, Settings);
+		MixerDevice->StartSpectrumAnalysis(SubmixToAnalyze, Settings);
+	}
+	else
+	{
+		UE_LOG(LogAudioMixer, Error, TEXT("Spectrum Analysis is an audio mixer only feature. Please run the game with -audiomixer to enable this feature."));
+	}
+}
+
+void UAudioMixerBlueprintLibrary::StopAnalyzingOutput(const UObject* WorldContextObject, USoundSubmix* SubmixToStopAnalyzing /*= nullptr*/)
+{
+	if (Audio::FMixerDevice* MixerDevice = GetAudioMixerDeviceFromWorldContext(WorldContextObject))
+	{
+		MixerDevice->StopSpectrumAnalysis(SubmixToStopAnalyzing);
+	}
+	else
+	{
+		UE_LOG(LogAudioMixer, Error, TEXT("Spectrum Analysis is an audio mixer only feature. Please run the game with -audiomixer to enable this feature."));
+	}
+}
+
+void UAudioMixerBlueprintLibrary::GetMagnitudeForFrequencies(const UObject* WorldContextObject, const TArray<float>& Frequencies, TArray<float>& Magnitudes, USoundSubmix* SubmixToAnalyze /*= nullptr*/)
+{
+	if (Audio::FMixerDevice* MixerDevice = GetAudioMixerDeviceFromWorldContext(WorldContextObject))
+	{
+		MixerDevice->GetMagnitudesForFrequencies(SubmixToAnalyze, Frequencies, Magnitudes);
+	}
+	else
+	{
+		UE_LOG(LogAudioMixer, Error, TEXT("Output recording is an audio mixer only feature. Please run the game with -audiomixer to enable this feature."));
+	}
+}
+
+void UAudioMixerBlueprintLibrary::GetPhaseForFrequencies(const UObject* WorldContextObject, const TArray<float>& Frequencies, TArray<float>& Phases, USoundSubmix* SubmixToAnalyze /*= nullptr*/)
+{
+	if (Audio::FMixerDevice* MixerDevice = GetAudioMixerDeviceFromWorldContext(WorldContextObject))
+	{
+		MixerDevice->GetPhasesForFrequencies(SubmixToAnalyze, Frequencies, Phases);
+	}
+	else
+	{
+		UE_LOG(LogAudioMixer, Error, TEXT("Output recording is an audio mixer only feature. Please run the game with -audiomixer to enable this feature."));
+	}
+}
+
 void UAudioMixerBlueprintLibrary::AddSourceEffectToPresetChain(const UObject* WorldContextObject, USoundEffectSourcePresetChain* PresetChain, FSourceEffectChainEntry Entry)
 {
 	if (!PresetChain)
@@ -295,3 +346,65 @@ int32 UAudioMixerBlueprintLibrary::GetNumberOfEntriesInSourceEffectChain(const U
 
 	return 0;
 }
+
+void UAudioMixerBlueprintLibrary::PopulateSpectrumAnalyzerSettings(EFFTSize FFTSize, EFFTPeakInterpolationMethod InterpolationMethod, EFFTWindowType WindowType, float HopSize, Audio::FSpectrumAnalyzerSettings &OutSettings)
+{
+	switch (FFTSize)
+	{
+	case EFFTSize::DefaultSize:
+		OutSettings.FFTSize = Audio::FSpectrumAnalyzerSettings::EFFTSize::Default;
+		break;
+	case EFFTSize::Min:
+		OutSettings.FFTSize = Audio::FSpectrumAnalyzerSettings::EFFTSize::Min_64;
+		break;
+	case EFFTSize::Small:
+		OutSettings.FFTSize = Audio::FSpectrumAnalyzerSettings::EFFTSize::Small_256;
+		break;
+	case EFFTSize::Medium:
+		OutSettings.FFTSize = Audio::FSpectrumAnalyzerSettings::EFFTSize::Medium_512;
+		break;
+	case EFFTSize::Large:
+		OutSettings.FFTSize = Audio::FSpectrumAnalyzerSettings::EFFTSize::Large_1024;
+		break;
+	case EFFTSize::Max:
+		OutSettings.FFTSize = Audio::FSpectrumAnalyzerSettings::EFFTSize::TestLarge_4096;
+		break;
+	default:
+		break;
+	}
+
+	switch (InterpolationMethod)
+	{
+	case EFFTPeakInterpolationMethod::NearestNeighbor:
+		OutSettings.InterpolationMethod = Audio::FSpectrumAnalyzerSettings::EPeakInterpolationMethod::NearestNeighbor;
+		break;
+	case EFFTPeakInterpolationMethod::Linear:
+		OutSettings.InterpolationMethod = Audio::FSpectrumAnalyzerSettings::EPeakInterpolationMethod::Linear;
+		break;
+	case EFFTPeakInterpolationMethod::Quadratic:
+		OutSettings.InterpolationMethod = Audio::FSpectrumAnalyzerSettings::EPeakInterpolationMethod::Quadratic;
+		break;
+	default:
+		break;
+	}
+
+	switch (WindowType)
+	{
+	case EFFTWindowType::None:
+		OutSettings.WindowType = Audio::EWindowType::None;
+		break;
+	case EFFTWindowType::Hamming:
+		OutSettings.WindowType = Audio::EWindowType::Hamming;
+		break;
+	case EFFTWindowType::Hann:
+		OutSettings.WindowType = Audio::EWindowType::Hann;
+		break;
+	case EFFTWindowType::Blackman:
+		OutSettings.WindowType = Audio::EWindowType::Blackman;
+		break;
+	default:
+		break;
+	}
+
+	OutSettings.HopSize = HopSize;
+}
diff --git a/Engine/Source/Runtime/AudioMixer/Private/AudioMixerDevice.cpp b/Engine/Source/Runtime/AudioMixer/Private/AudioMixerDevice.cpp
--- a/Engine/Source/Runtime/AudioMixer/Private/AudioMixerDevice.cpp
+++ b/Engine/Source/Runtime/AudioMixer/Private/AudioMixerDevice.cpp
@@ -1361,6 +1361,74 @@ namespace Audio
 	}
 
 
+	void FMixerDevice::StartSpectrumAnalysis(USoundSubmix* InSubmix, const Audio::FSpectrumAnalyzerSettings& InSettings)
+	{
+		Audio::FMixerSubmixPtr* FoundSubmix = Submixes.Find(InSubmix);
+		if (FoundSubmix)
+		{
+			(*FoundSubmix)->StartSpectrumAnalysis(InSettings);
+		}
+		else
+		{
+			FMixerSubmixWeakPtr MasterSubmix = GetMasterSubmix();
+			FMixerSubmixPtr MasterSubmixPtr = MasterSubmix.Pin();
+			check(MasterSubmixPtr.IsValid());
+
+			MasterSubmixPtr->StartSpectrumAnalysis(InSettings);
+		}
+	}
+
+	void FMixerDevice::StopSpectrumAnalysis(USoundSubmix* InSubmix)
+	{
+		Audio::FMixerSubmixPtr* FoundSubmix = Submixes.Find(InSubmix);
+		if (FoundSubmix)
+		{
+			(*FoundSubmix)->StopSpectrumAnalysis();
+		}
+		else
+		{
+			FMixerSubmixWeakPtr MasterSubmix = GetMasterSubmix();
+			FMixerSubmixPtr MasterSubmixPtr = MasterSubmix.Pin();
+			check(MasterSubmixPtr.IsValid());
+
+			MasterSubmixPtr->StopSpectrumAnalysis();
+		}
+	}
+
+	void FMixerDevice::GetMagnitudesForFrequencies(USoundSubmix* InSubmix, const TArray<float>& InFrequencies, TArray<float>& OutMagnitudes)
+	{
+		Audio::FMixerSubmixPtr* FoundSubmix = Submixes.Find(InSubmix);
+		if (FoundSubmix)
+		{
+			(*FoundSubmix)->GetMagnitudeForFrequencies(InFrequencies, OutMagnitudes);
+		}
+		else
+		{
+			FMixerSubmixWeakPtr MasterSubmix = GetMasterSubmix();
+			FMixerSubmixPtr MasterSubmixPtr = MasterSubmix.Pin();
+			check(MasterSubmixPtr.IsValid());
+
+			MasterSubmixPtr->GetMagnitudeForFrequencies(InFrequencies, OutMagnitudes);
+		}
+	}
+
+	void FMixerDevice::GetPhasesForFrequencies(USoundSubmix* InSubmix, const TArray<float>& InFrequencies, TArray<float>& OutPhases)
+	{
+		Audio::FMixerSubmixPtr* FoundSubmix = Submixes.Find(InSubmix);
+		if (FoundSubmix)
+		{
+			(*FoundSubmix)->GetPhaseForFrequencies(InFrequencies, OutPhases);
+		}
+		else
+		{
+			FMixerSubmixWeakPtr MasterSubmix = GetMasterSubmix();
+			FMixerSubmixPtr MasterSubmixPtr = MasterSubmix.Pin();
+			check(MasterSubmixPtr.IsValid());
+
+			MasterSubmixPtr->GetPhaseForFrequencies(InFrequencies, OutPhases);
+		}
+	}
+
 	void FMixerDevice::RegisterSubmixBufferListener(ISubmixBufferListener* InSubmixBufferListener, USoundSubmix* InSubmix)
 	{
 		Audio::FMixerSubmixPtr* FoundSubmix = Submixes.Find(InSubmix);
diff --git a/Engine/Source/Runtime/AudioMixer/Private/AudioMixerSubmix.cpp b/Engine/Source/Runtime/AudioMixer/Private/AudioMixerSubmix.cpp
--- a/Engine/Source/Runtime/AudioMixer/Private/AudioMixerSubmix.cpp
+++ b/Engine/Source/Runtime/AudioMixer/Private/AudioMixerSubmix.cpp
@@ -360,6 +360,29 @@ namespace Audio
 		}
 	}
 
+	void FMixerSubmix::MixBufferDownToMono(const AlignedFloatBuffer& InBuffer, int32 NumInputChannels, AlignedFloatBuffer& OutBuffer)
+	{
+		check(NumInputChannels > 0);
+
+		int32 NumFrames = InBuffer.Num() / NumInputChannels;
+		OutBuffer.Reset();
+		OutBuffer.AddZeroed(NumFrames);
+
+		const float* InData = InBuffer.GetData();
+		float* OutData = OutBuffer.GetData();
+
+		const float GainFactor = 1.0f / FMath::Sqrt((float) NumInputChannels);
+
+		for (int32 FrameIndex = 0; FrameIndex < NumFrames; FrameIndex++)
+		{
+			for (int32 ChannelIndex = 0; ChannelIndex < NumInputChannels; ChannelIndex++)
+			{
+				const int32 InputIndex = FrameIndex * NumInputChannels + ChannelIndex;
+				OutData[FrameIndex] += InData[InputIndex] * GainFactor;
+			}
+		}
+	}
+
 	void FMixerSubmix::SetUpAmbisonicsEncoder()
 	{
 		check(AmbisonicsMixer.IsValid());
@@ -665,6 +688,15 @@ namespace Audio
 			}
 		}
 
+		// If spectrum analysis is enabled for this submix, downmix the resulting audio
+		// and push it to the spectrum analyzer.
+		if (SpectrumAnalyzer.IsValid())
+		{
+			MixBufferDownToMono(InputBuffer, NumChannels, MonoMixBuffer);
+			SpectrumAnalyzer->PushAudio(MonoMixBuffer.GetData(), MonoMixBuffer.Num());
+			SpectrumAnalyzer->PerformAnalysisIfPossible(true, true);
+		}
+
 		// If the channel types match, just do a copy
 		if (ChannelFormat != ParentChannelType || SubmixAmbisonicsDecoderID != INDEX_NONE)
 		{
@@ -844,6 +876,56 @@ namespace Audio
 		OnSubmixEnvelope.AddUnique(OnSubmixEnvelopeBP);
 	}
 
+	void FMixerSubmix::StartSpectrumAnalysis(const FSpectrumAnalyzerSettings& InSettings)
+	{
+		SpectrumAnalyzer.Reset(new FSpectrumAnalyzer(InSettings, MixerDevice->GetSampleRate()));
+	}
+
+	void FMixerSubmix::StopSpectrumAnalysis()
+	{
+		SpectrumAnalyzer.Reset();
+	}
+
+	void FMixerSubmix::GetMagnitudeForFrequencies(const TArray<float>& InFrequencies, TArray<float>& OutMagnitudes)
+	{
+		if (SpectrumAnalyzer.IsValid())
+		{
+			OutMagnitudes.Reset();
+			OutMagnitudes.AddUninitialized(InFrequencies.Num());
+
+			SpectrumAnalyzer->LockOutputBuffer();
+			for (int32 Index = 0; Index < InFrequencies.Num(); Index++)
+			{
+				OutMagnitudes[Index] = SpectrumAnalyzer->GetMagnitudeForFrequency(InFrequencies[Index]);
+			}
+			SpectrumAnalyzer->UnlockOutputBuffer();
+		}
+		else
+		{
+			UE_LOG(LogAudioMixer, Warning, TEXT("Call StartSpectrumAnalysis before calling GetMagnitudeForFrequencies."));
+		}
+	}
+
+	void FMixerSubmix::GetPhaseForFrequencies(const TArray<float>& InFrequencies, TArray<float>& OutPhases)
+	{
+		if (SpectrumAnalyzer.IsValid())
+		{
+			OutPhases.Reset();
+			OutPhases.AddUninitialized(InFrequencies.Num());
+
+			SpectrumAnalyzer->LockOutputBuffer();
+			for (int32 Index = 0; Index < InFrequencies.Num(); Index++)
+			{
+				OutPhases[Index] = SpectrumAnalyzer->GetPhaseForFrequency(InFrequencies[Index]);
+			}
+			SpectrumAnalyzer->UnlockOutputBuffer();
+		}
+		else
+		{
+			UE_LOG(LogAudioMixer, Warning, TEXT("Call StartSpectrumAnalysis before calling GetMagnitudeForFrequencies."));
+		}
+	}
+
 	void FMixerSubmix::BroadcastEnvelope()
 	{
 
diff --git a/Engine/Source/Runtime/AudioMixer/Private/DSP/AudioFFT.cpp b/Engine/Source/Runtime/AudioMixer/Private/DSP/AudioFFT.cpp
--- a/Engine/Source/Runtime/AudioMixer/Private/DSP/AudioFFT.cpp
+++ b/Engine/Source/Runtime/AudioMixer/Private/DSP/AudioFFT.cpp
@@ -21,7 +21,7 @@ namespace Audio
 	{
 		const int32 N = bIsPeriodic ? NumFrames : NumFrames - 1;
 		const float PhaseDelta = 2.0f * PI / N;
-		float Phase = 0;
+		float Phase = 0.0f;
 
 		for (int32 FrameIndex = 0; FrameIndex < NumFrames; FrameIndex++)
 		{
@@ -104,59 +104,6 @@ namespace Audio
 		}
 	}
 
-	FWindow::FWindow(EWindowType InType, int32 InNumFrames, int32 InNumChannels, bool bIsPeriodic)
-		: WindowType(InType)
-		, NumSamples(InNumFrames * InNumChannels)
-	{
-		checkf(NumSamples % 4 == 0, TEXT("For performance reasons, this window's length should be a multiple of 4."));
-		Generate(InNumFrames, InNumChannels, bIsPeriodic);
-	}
-
-	FWindow::~FWindow()
-	{
-
-	}
-
-	void FWindow::Generate(int32 NumFrames, int32 NumChannels, bool bIsPeriodic)
-	{
-		if (WindowType == EWindowType::None)
-		{
-			return;
-		}
-
-		WindowBuffer.Reset();
-		WindowBuffer.AddZeroed(NumSamples);
-
-		switch (WindowType)
-		{
-			case EWindowType::Hann:
-			{
-				GenerateHannWindow(WindowBuffer.GetData(), NumFrames, NumChannels, bIsPeriodic);
-				break;
-			}
-			case EWindowType::Blackman:
-			{
-				GenerateBlackmanWindow(WindowBuffer.GetData(), NumFrames, NumChannels, bIsPeriodic);
-				break;
-			}
-			default:
-			{
-				break;
-			}
-		}
-	}
-
-	void FWindow::ApplyToBuffer(float* InBuffer)
-	{
-		if (WindowType == EWindowType::None)
-		{
-			return;
-		}
-
-		check(IsAligned<float*>(InBuffer, 4));
-		MultiplyBuffersInPlace(WindowBuffer.GetData(), InBuffer, NumSamples);
-	}
-
 	namespace FFTIntrinsics
 	{
 		// Fast bit reversal helper function. Can be used if N is a power of 2. Not well exercised.
diff --git a/Engine/Source/Runtime/AudioMixer/Private/DSP/SpectrumAnalyzer.cpp b/Engine/Source/Runtime/AudioMixer/Private/DSP/SpectrumAnalyzer.cpp
--- a/Engine/Source/Runtime/AudioMixer/Private/DSP/SpectrumAnalyzer.cpp
+++ b/Engine/Source/Runtime/AudioMixer/Private/DSP/SpectrumAnalyzer.cpp
@@ -6,8 +6,9 @@
 namespace Audio
 {
 	FSpectrumAnalyzer::FSpectrumAnalyzer()
-		: CurrentSettings(SpectrumAnalyzerSettings::FSettings())
+		: CurrentSettings(FSpectrumAnalyzerSettings())
 		, bSettingsWereUpdated(false)
+		, bIsInitialized(false)
 		, SampleRate(0.0f)
 		, Window(CurrentSettings.WindowType, (int32)CurrentSettings.FFTSize, 1, false)
 		, InputQueue(FMath::Max((int32)CurrentSettings.FFTSize * 4, 4096))
@@ -16,9 +17,10 @@ namespace Audio
 	{
 	}
 
-	FSpectrumAnalyzer::FSpectrumAnalyzer(const SpectrumAnalyzerSettings::FSettings& InSettings, float InSampleRate)
+	FSpectrumAnalyzer::FSpectrumAnalyzer(const FSpectrumAnalyzerSettings& InSettings, float InSampleRate)
 		: CurrentSettings(InSettings)
 		, bSettingsWereUpdated(false)
+		, bIsInitialized(true)
 		, SampleRate(InSampleRate)
 		, Window(InSettings.WindowType, (int32)InSettings.FFTSize, 1, false)
 		, InputQueue(FMath::Max((int32)CurrentSettings.FFTSize * 4, 4096))
@@ -29,8 +31,9 @@ namespace Audio
 	}
 
 	FSpectrumAnalyzer::FSpectrumAnalyzer(float InSampleRate)
-		: CurrentSettings(SpectrumAnalyzerSettings::FSettings())
+		: CurrentSettings(FSpectrumAnalyzerSettings())
 		, bSettingsWereUpdated(false)
+		, bIsInitialized(true)
 		, SampleRate(InSampleRate)
 		, Window(CurrentSettings.WindowType, (int32)CurrentSettings.FFTSize, 1, false)
 		, InputQueue(FMath::Max((int32)CurrentSettings.FFTSize * 4, 4096))
@@ -40,13 +43,21 @@ namespace Audio
 		ResetSettings();
 	}
 
+	FSpectrumAnalyzer::~FSpectrumAnalyzer()
+	{
+		if (AsyncAnalysisTask.IsValid())
+		{
+			AsyncAnalysisTask->EnsureCompletion(false);
+		}
+	}
+
 	void FSpectrumAnalyzer::Init(float InSampleRate)
 	{
-		SpectrumAnalyzerSettings::FSettings DefaultSettings = SpectrumAnalyzerSettings::FSettings();
+		FSpectrumAnalyzerSettings DefaultSettings = FSpectrumAnalyzerSettings();
 		Init(DefaultSettings, InSampleRate);
 	}
 
-	void FSpectrumAnalyzer::Init(const SpectrumAnalyzerSettings::FSettings& InSettings, float InSampleRate)
+	void FSpectrumAnalyzer::Init(const FSpectrumAnalyzerSettings& InSettings, float InSampleRate)
 	{
 		CurrentSettings = InSettings;
 		bSettingsWereUpdated = false;
@@ -54,6 +65,8 @@ namespace Audio
 		InputQueue.SetCapacity(FMath::Max((int32)CurrentSettings.FFTSize * 4, 4096));
 		FrequencyBuffer.Reset(CurrentSettings);
 		ResetSettings();
+
+		bIsInitialized = true;
 	}
 
 	void FSpectrumAnalyzer::ResetSettings()
@@ -67,7 +80,15 @@ namespace Audio
 
 		Window = FWindow(CurrentSettings.WindowType, (int32)CurrentSettings.FFTSize, 1, false);
 		FFTSize = (int32) CurrentSettings.FFTSize;
-		HopInSamples = GetCOLAHopSizeForWindow(CurrentSettings.WindowType, (uint32)CurrentSettings.FFTSize);
+		
+		if (FMath::IsNearlyZero(CurrentSettings.HopSize))
+		{
+			HopInSamples = GetCOLAHopSizeForWindow(CurrentSettings.WindowType, (uint32)CurrentSettings.FFTSize);
+		}
+		else
+		{
+			HopInSamples = FMath::FloorToInt((float)CurrentSettings.FFTSize * CurrentSettings.HopSize);
+		}
 		
 		AnalysisTimeDomainBuffer.Reset();
 		AnalysisTimeDomainBuffer.AddZeroed(FFTSize);
@@ -76,7 +97,7 @@ namespace Audio
 		bSettingsWereUpdated = false;
 	}
 
-	void FSpectrumAnalyzer::PerformInterpolation(const FSpectrumAnalyzerFrequencyVector* InFrequencies, SpectrumAnalyzerSettings::EPeakInterpolationMethod InMethod, const float InFreq, float& OutReal, float& OutImag)
+	void FSpectrumAnalyzer::PerformInterpolation(const FSpectrumAnalyzerFrequencyVector* InFrequencies, FSpectrumAnalyzerSettings::EPeakInterpolationMethod InMethod, const float InFreq, float& OutReal, float& OutImag)
 	{
 		const int32 VectorLength = InFrequencies->RealFrequencies.Num();
 		const float NyquistPosition = VectorLength / 2;
@@ -93,7 +114,7 @@ namespace Audio
 
 		switch (InMethod)
 		{
-		case Audio::SpectrumAnalyzerSettings::EPeakInterpolationMethod::NearestNeighbor:
+		case Audio::FSpectrumAnalyzerSettings::EPeakInterpolationMethod::NearestNeighbor:
 		{
 			int32 Index = FMath::RoundToInt(Position);
 			OutReal = InFrequencies->RealFrequencies[Index];
@@ -101,7 +122,7 @@ namespace Audio
 			break;
 		}
 			
-		case Audio::SpectrumAnalyzerSettings::EPeakInterpolationMethod::Linear:
+		case Audio::FSpectrumAnalyzerSettings::EPeakInterpolationMethod::Linear:
 		{
 			const int32 LowerIndex = FMath::FloorToInt(Position);
 			const int32 UpperIndex = FMath::CeilToInt(Position);
@@ -116,7 +137,7 @@ namespace Audio
 			OutImag = FMath::Lerp<float>(y1Imag, y2Imag, PositionFraction);
 			break;
 		}	
-		case Audio::SpectrumAnalyzerSettings::EPeakInterpolationMethod::Quadratic:
+		case Audio::FSpectrumAnalyzerSettings::EPeakInterpolationMethod::Quadratic:
 		{
 			const int32 MidIndex = FMath::RoundToInt(Position);
 			const int32 LowerIndex = FMath::Max(0, MidIndex - 1);
@@ -145,19 +166,24 @@ namespace Audio
 		}
 	}
 
-	void FSpectrumAnalyzer::SetSettings(const SpectrumAnalyzerSettings::FSettings& InSettings)
+	void FSpectrumAnalyzer::SetSettings(const FSpectrumAnalyzerSettings& InSettings)
 	{
 		CurrentSettings = InSettings;
 		bSettingsWereUpdated = true;
 	}
 
-	void FSpectrumAnalyzer::GetSettings(SpectrumAnalyzerSettings::FSettings& OutSettings)
+	void FSpectrumAnalyzer::GetSettings(FSpectrumAnalyzerSettings& OutSettings)
 	{
 		OutSettings = CurrentSettings;
 	}
 
 	float FSpectrumAnalyzer::GetMagnitudeForFrequency(float InFrequency)
 	{
+		if (!bIsInitialized)
+		{
+			return false;
+		}
+
 		const FSpectrumAnalyzerFrequencyVector* OutVector;
 		bool bShouldUnlockBuffer = true;
 
@@ -195,6 +221,11 @@ namespace Audio
 
 	float FSpectrumAnalyzer::GetPhaseForFrequency(float InFrequency)
 	{
+		if (!bIsInitialized)
+		{
+			return false;
+		}
+
 		const FSpectrumAnalyzerFrequencyVector* OutVector;
 		bool bShouldUnlockBuffer = true;
 
@@ -232,6 +263,11 @@ namespace Audio
 
 	void FSpectrumAnalyzer::LockOutputBuffer()
 	{
+		if (!bIsInitialized)
+		{
+			return;
+		}
+
 		if (LockedFrequencyVector != nullptr)
 		{
 			FrequencyBuffer.UnlockBuffer();
@@ -242,6 +278,11 @@ namespace Audio
 
 	void FSpectrumAnalyzer::UnlockOutputBuffer()
 	{
+		if (!bIsInitialized)
+		{
+			return;
+		}
+
 		if (LockedFrequencyVector != nullptr)
 		{
 			FrequencyBuffer.UnlockBuffer();
@@ -260,8 +301,29 @@ namespace Audio
 		return InputQueue.Push(InBuffer, NumSamples) > 0;
 	}
 
-	bool FSpectrumAnalyzer::PerformAnalysisIfPossible()
+	bool FSpectrumAnalyzer::PerformAnalysisIfPossible(bool bUseLatestAudio, bool bAsync)
 	{
+		if (!bIsInitialized)
+		{
+			return false;
+		}		
+
+		if (bAsync)
+		{
+			// if bAsync is true, kick off a new task if one isn't in flight already, and return.
+			if (!AsyncAnalysisTask.IsValid())
+			{
+				AsyncAnalysisTask.Reset(new FSpectrumAnalyzerTask(this, bUseLatestAudio));
+				AsyncAnalysisTask->StartBackgroundTask();
+			}
+			else if (AsyncAnalysisTask->IsDone())
+			{
+				AsyncAnalysisTask->StartBackgroundTask();
+			}
+
+			return true;
+		}
+
 		// If settings were updated, perform resizing and parameter updates here:
 		if (bSettingsWereUpdated)
 		{
@@ -276,6 +338,12 @@ namespace Audio
 		{
 			float* TimeDomainBuffer = AnalysisTimeDomainBuffer.GetData();
 
+			if (bUseLatestAudio)
+			{
+				// If we are only using the latest audio, scrap the oldest audio in the InputQueue:
+				InputQueue.SetNum((uint32)FFTSize);
+			}
+
 			// Perform pop/peek here based on FFT size and hop amount.
 			const int32 PeekAmount = FFTSize - HopInSamples;
 			InputQueue.Pop(TimeDomainBuffer, HopInSamples);
@@ -293,7 +361,6 @@ namespace Audio
 			OutputParams.OutReal = OutputVector->RealFrequencies.GetData();
 			OutputParams.OutImag = OutputVector->ImagFrequencies.GetData();
 
-
 			PerformFFT(InputParams, OutputParams);
 
 			// We're done, so unlock this vector.
@@ -307,7 +374,12 @@ namespace Audio
 		}
 	}
 
-	static const int32 SpectrumAnalyzerBufferSize = 3;
+	bool FSpectrumAnalyzer::IsInitialized()
+	{
+		return bIsInitialized;
+	}
+
+	static const int32 SpectrumAnalyzerBufferSize = 4;
 
 	FSpectrumAnalyzerBuffer::FSpectrumAnalyzerBuffer()
 		: OutputIndex(0)
@@ -315,12 +387,12 @@ namespace Audio
 	{
 	}
 
-	FSpectrumAnalyzerBuffer::FSpectrumAnalyzerBuffer(const SpectrumAnalyzerSettings::FSettings& InSettings)
+	FSpectrumAnalyzerBuffer::FSpectrumAnalyzerBuffer(const FSpectrumAnalyzerSettings& InSettings)
 	{
 		Reset(InSettings);
 	}
 
-	void FSpectrumAnalyzerBuffer::Reset(const SpectrumAnalyzerSettings::FSettings& InSettings)
+	void FSpectrumAnalyzerBuffer::Reset(const FSpectrumAnalyzerSettings& InSettings)
 	{
 		FScopeLock ScopeLock(&BufferIndicesCriticalSection);
 
@@ -391,4 +463,10 @@ namespace Audio
 		ImagFrequencies.Reset();
 		ImagFrequencies.AddZeroed(InFFTSize);
 	}
+
+	void FSpectrumAnalysisAsyncWorker::DoWork()
+	{
+		Analyzer->PerformAnalysisIfPossible(bUseLatestAudio, false);
+	}
+
 }
diff --git a/Engine/Source/Runtime/AudioMixer/Public/AudioMixerBlueprintLibrary.h b/Engine/Source/Runtime/AudioMixer/Public/AudioMixerBlueprintLibrary.h
--- a/Engine/Source/Runtime/AudioMixer/Public/AudioMixerBlueprintLibrary.h
+++ b/Engine/Source/Runtime/AudioMixer/Public/AudioMixerBlueprintLibrary.h
@@ -7,10 +7,57 @@
 #include "SubmixEffects/AudioMixerSubmixEffectDynamicsProcessor.h"
 #include "Sound/SoundEffectSource.h"
 #include "Sound/SampleBuffer.h"
+#include "DSP/SpectrumAnalyzer.h"
 #include "AudioMixerBlueprintLibrary.generated.h"
 
 class USoundSubmix;
 
+UENUM(BlueprintType)
+enum class EFFTSize : uint8
+{
+	// 512
+	DefaultSize, 
+
+	// 64
+	Min,
+
+	// 256
+	Small,
+
+	// 512
+	Medium, 
+
+	// 1024
+	Large, 
+
+	// 4096
+	Max, 
+};
+
+UENUM()
+enum class EFFTPeakInterpolationMethod : uint8
+{
+	NearestNeighbor,
+	Linear,
+	Quadratic
+};
+
+UENUM()
+enum class EFFTWindowType : uint8
+{
+	// No window is applied. Technically a boxcar window.
+	None, 
+
+	// Mainlobe width of -3 dB and sidelove attenuation of ~-40 dB. Good for COLA.
+	Hamming,
+
+	// Mainlobe width of -3 dB and sidelobe attenuation of ~-30dB. Good for COLA.
+	Hann,
+
+	// Mainlobe width of -3 dB and sidelobe attenuation of ~-60db. Tricky for COLA.
+	Blackman
+};
+
 UCLASS(meta=(ScriptName="AudioMixerLibrary"))
 class AUDIOMIXER_API UAudioMixerBlueprintLibrary : public UBlueprintFunctionLibrary
 {
@@ -42,10 +89,26 @@ public:
 	UFUNCTION(BlueprintCallable, Category = "Audio", meta = (WorldContext = "WorldContextObject", AdvancedDisplay = 1))
 	static void PauseRecordingOutput(const UObject* WorldContextObject, USoundSubmix* SubmixToPause = nullptr);
 
-	/** Resume recording audio after pausing. By leaving the Submix To Record field blank, you can record the master output of the game. */
+	/** Resume recording audio after pausing. By leaving the Submix To Pause field blank, you can record the master output of the game. */
 	UFUNCTION(BlueprintCallable, Category = "Audio", meta = (WorldContext = "WorldContextObject", AdvancedDisplay = 1))
 	static void ResumeRecordingOutput(const UObject* WorldContextObject, USoundSubmix* SubmixToPause = nullptr);
 
+	/** Start spectrum analysis of the audio output. By leaving the Submix To Analyze blank, you can analyze the master output of the game. */
+	UFUNCTION(BlueprintCallable, Category = "Audio|Analysis", meta = (WorldContext = "WorldContextObject", AdvancedDisplay = 1))
+	static void StartAnalyzingOutput(const UObject* WorldContextObject, USoundSubmix* SubmixToAnalyze = nullptr, EFFTSize FFTSize = EFFTSize::DefaultSize, EFFTPeakInterpolationMethod InterpolationMethod = EFFTPeakInterpolationMethod::Linear, EFFTWindowType WindowType = EFFTWindowType::Hann, float HopSize = 0);
+
+	/** Start spectrum analysis of the audio output. By leaving the Submix To Stop Analyzing blank, you can analyze the master output of the game. */
+	UFUNCTION(BlueprintCallable, Category = "Audio|Analysis", meta = (WorldContext = "WorldContextObject", AdvancedDisplay = 1))
+	static void StopAnalyzingOutput(const UObject* WorldContextObject, USoundSubmix* SubmixToStopAnalyzing = nullptr);
+
+	/** Start spectrum analysis of the audio output. By leaving the Submix To Analyze blank, you can analyze the master output of the game. */
+	UFUNCTION(BlueprintCallable, Category = "Audio|Analysis", meta = (WorldContext = "WorldContextObject", AdvancedDisplay = 3))
+	static void GetMagnitudeForFrequencies(const UObject* WorldContextObject, const TArray<float>& Frequencies, TArray<float>& Magnitudes, USoundSubmix* SubmixToAnalyze = nullptr);
+
+	/** Start spectrum analysis of the audio output. By leaving the Submix To Analyze blank, you can analyze the master output of the game. */
+	UFUNCTION(BlueprintCallable, Category = "Audio|Analysis", meta = (WorldContext = "WorldContextObject", AdvancedDisplay = 3))
+	static void GetPhaseForFrequencies(const UObject* WorldContextObject, const TArray<float>& Frequencies, TArray<float>& Phases, USoundSubmix* SubmixToAnalyze = nullptr);
+
 	/** Adds source effect entry to preset chain. Only effects the instance of the preset chain */
 	UFUNCTION(BlueprintCallable, Category = "Audio|Effects", meta = (WorldContext = "WorldContextObject"))
 	static void AddSourceEffectToPresetChain(const UObject* WorldContextObject, USoundEffectSourcePresetChain* PresetChain, FSourceEffectChainEntry Entry);
@@ -62,4 +125,7 @@ public:
 	UFUNCTION(BlueprintCallable, Category = "Audio|Effects", meta = (WorldContext = "WorldContextObject"))
 	static int32 GetNumberOfEntriesInSourceEffectChain(const UObject* WorldContextObject, USoundEffectSourcePresetChain* PresetChain);
 
+
+private:
+	static void PopulateSpectrumAnalyzerSettings(EFFTSize FFTSize, EFFTPeakInterpolationMethod InterpolationMethod, EFFTWindowType WindowType, float HopSize, Audio::FSpectrumAnalyzerSettings &OutSettings);
 };
\ No newline at end of file
diff --git a/Engine/Source/Runtime/AudioMixer/Public/AudioMixerDevice.h b/Engine/Source/Runtime/AudioMixer/Public/AudioMixerDevice.h
--- a/Engine/Source/Runtime/AudioMixer/Public/AudioMixerDevice.h
+++ b/Engine/Source/Runtime/AudioMixer/Public/AudioMixerDevice.h
@@ -120,6 +120,12 @@ namespace Audio
 		virtual void StopEnvelopeFollowing(USoundSubmix* InSubmix) override;
 		virtual void AddEnvelopeFollowerDelegate(USoundSubmix* InSubmix, const FOnSubmixEnvelopeBP& OnSubmixEnvelopeBP) override;
 
+		/** Submix Spectrum Analysis */
+		virtual void StartSpectrumAnalysis(USoundSubmix* InSubmix, const Audio::FSpectrumAnalyzerSettings& InSettings) override;
+		virtual void StopSpectrumAnalysis(USoundSubmix* InSubmix) override;
+		virtual void GetMagnitudesForFrequencies(USoundSubmix* InSubmix, const TArray<float>& InFrequencies, TArray<float>& OutMagnitudes);
+		virtual void GetPhasesForFrequencies(USoundSubmix* InSubmix, const TArray<float>& InFrequencies, TArray<float>& OutPhases);
+
 		// Submix buffer listener callbacks
 		virtual void RegisterSubmixBufferListener(ISubmixBufferListener* InSubmixBufferListener, USoundSubmix* InSubmix = nullptr) override;
 		virtual void UnregisterSubmixBufferListener(ISubmixBufferListener* InSubmixBufferListener, USoundSubmix* InSubmix = nullptr) override;
diff --git a/Engine/Source/Runtime/AudioMixer/Public/AudioMixerSubmix.h b/Engine/Source/Runtime/AudioMixer/Public/AudioMixerSubmix.h
--- a/Engine/Source/Runtime/AudioMixer/Public/AudioMixerSubmix.h
+++ b/Engine/Source/Runtime/AudioMixer/Public/AudioMixerSubmix.h
@@ -6,6 +6,7 @@
 #include "Sound/SoundSubmix.h"
 #include "Sound/SampleBuffer.h"
 #include "DSP/EnvelopeFollower.h"
+#include "DSP/SpectrumAnalyzer.h"
 
 class USoundEffectSubmix;
 
@@ -134,6 +135,20 @@ namespace Audio
 		// Adds an envelope follower delegate
 		void AddEnvelopeFollowerDelegate(const FOnSubmixEnvelopeBP& OnSubmixEnvelopeBP);
 
+		// Initializes a new FFT analyzer for this submix and immediately begins feeding audio to it.
+		void StartSpectrumAnalysis(const FSpectrumAnalyzerSettings& InSettings);
+
+		// Terminates whatever FFT Analyzer is being used for this submix.
+		void StopSpectrumAnalysis();
+
+		// Gets the most recent magnitude values for each corresponding value in InFrequencies (in Hz).
+		// This requires StartSpectrumAnalysis to be called first.
+		void GetMagnitudeForFrequencies(const TArray<float>& InFrequencies, TArray<float>& OutMagnitudes);
+
+		// Gets the most recent phase values for each corresponding value in InFrequencies (in Hz).
+		// This requires StartSpectrumAnalysis to be called first.
+		void GetPhaseForFrequencies(const TArray<float>& InFrequencies, TArray<float>& OutPhases);
+
 		// Broadcast the envelope value on the game thread
 		void BroadcastEnvelope();
 
@@ -141,6 +156,8 @@ namespace Audio
 		// Down mix the given buffer to the desired down mix channel count
 		void FormatChangeBuffer(const ESubmixChannelFormat NewChannelType, AlignedFloatBuffer& InBuffer, AlignedFloatBuffer& OutNewBuffer);
 
+		void MixBufferDownToMono(const AlignedFloatBuffer& InBuffer, int32 NumInputChannels, AlignedFloatBuffer& OutBuffer);
+
 		// Set up ambisonics encoder. Called when ambisonics settings are changed.
 		void SetUpAmbisonicsEncoder();
 
@@ -238,6 +255,12 @@ namespace Audio
 		int32 EnvelopeNumChannels;
 		FCriticalSection EnvelopeCriticalSection;
 
+		// Spectrum analyzer:
+		TUniquePtr<FSpectrumAnalyzer> SpectrumAnalyzer;
+		
+		// This buffer is used to downmix the submix output to mono before submitting it to the SpectrumAnalyzer.
+		AlignedFloatBuffer MonoMixBuffer;
+
 		// This buffer is encoded into for each source, then summed into the ambisonics buffer.
 		AlignedFloatBuffer InputAmbisonicsBuffer;
 
diff --git a/Engine/Source/Runtime/AudioMixer/Public/DSP/AudioFFT.h b/Engine/Source/Runtime/AudioMixer/Public/DSP/AudioFFT.h
--- a/Engine/Source/Runtime/AudioMixer/Public/DSP/AudioFFT.h
+++ b/Engine/Source/Runtime/AudioMixer/Public/DSP/AudioFFT.h
@@ -40,25 +40,75 @@ namespace Audio
 		 *                     Generally, set this to false if using this window with an STFT, but use true
 		 *                     if this window will be used on an entire, self-contained signal.
 		 */
-		FWindow(EWindowType InType, int32 InNumFrames, int32 InNumChannels, bool bIsPeriodic);
+		FWindow(EWindowType InType, int32 InNumFrames, int32 InNumChannels, bool bIsPeriodic)
+			: WindowType(InType)
+			, NumSamples(InNumFrames * InNumChannels)
+		{
+			checkf(NumSamples % 4 == 0, TEXT("For performance reasons, this window's length should be a multiple of 4."));
+			Generate(InNumFrames, InNumChannels, bIsPeriodic);
+		}
 
 		// Destructor. Releases memory used for window.
-		~FWindow();
+		~FWindow()
+		{
+		}
 
 		// Apply this window to InBuffer, which is expected to be an interleaved buffer with the same amount of frames
 		// and channels this window was constructed with.
-		void ApplyToBuffer(float* InBuffer);
+		void ApplyToBuffer(float* InBuffer)
+		{
+			if (WindowType == EWindowType::None)
+			{
+				return;
+			}
+
+			check(IsAligned<float*>(InBuffer, 4));
+			MultiplyBuffersInPlace(WindowBuffer.GetData(), InBuffer, NumSamples);
+		}
 
 	private:
+		EWindowType WindowType;
+		AlignedFloatBuffer WindowBuffer;
+		int32 NumSamples;
+
 		// Purposefully hidden constructor.
 		FWindow();
 
 		// Generate the window. Called on constructor.
-		void Generate(int32 NumFrames, int32 NumChannels, bool bIsPeriodic);
+		void Generate(int32 NumFrames, int32 NumChannels, bool bIsPeriodic)
+		{
+			if (WindowType == EWindowType::None)
+			{
+				return;
+			}
 
-		EWindowType WindowType;
-		AlignedFloatBuffer WindowBuffer;
-		int32 NumSamples;
+			WindowBuffer.Reset();
+			WindowBuffer.AddZeroed(NumSamples);
+
+			switch (WindowType)
+			{
+			case EWindowType::Hann:
+			{
+				GenerateHannWindow(WindowBuffer.GetData(), NumFrames, NumChannels, bIsPeriodic);
+				break;
+			}
+			case EWindowType::Hamming:
+			{
+				GenerateHammingWindow(WindowBuffer.GetData(), NumFrames, NumChannels, bIsPeriodic);
+				break;
+			}
+			case EWindowType::Blackman:
+			{
+				GenerateBlackmanWindow(WindowBuffer.GetData(), NumFrames, NumChannels, bIsPeriodic);
+				break;
+			}
+			default:
+			{
+				checkf(false, TEXT("Unknown window type!"));
+				break;
+			}
+			}
+		}
 	};
 
 	struct FFTTimeDomainData
diff --git a/Engine/Source/Runtime/AudioMixer/Public/DSP/Dsp.h b/Engine/Source/Runtime/AudioMixer/Public/DSP/Dsp.h
--- a/Engine/Source/Runtime/AudioMixer/Public/DSP/Dsp.h
+++ b/Engine/Source/Runtime/AudioMixer/Public/DSP/Dsp.h
@@ -553,6 +553,29 @@ namespace Audio
 			return ((int32)CurrentSlack) - NumSamples;
 		}
 
+		// When called, seeks the read or write cursor to only retain either the NumSamples latest data
+		// (if bRetainOldestSamples is false) or the NumSamples oldest data (if bRetainOldestSamples is true)
+		// in the buffer. Cannot be used to increase the capacity of this buffer.
+		void SetNum(uint32 NumSamples, bool bRetainOldestSamples = false)
+		{
+			check(NumSamples < Capacity);
+
+			if (bRetainOldestSamples)
+			{
+				WriteCounter.Set((ReadCounter.GetValue() + NumSamples) % Capacity);
+			}
+			else
+			{
+				int64 ReadCounterNum = ((int32)WriteCounter.GetValue()) - ((int32) NumSamples);
+				if (ReadCounterNum < 0)
+				{
+					ReadCounterNum = Capacity + ReadCounterNum;
+				}
+
+				ReadCounter.Set(ReadCounterNum);
+			}
+		}
+
 		// Get number of samples that can be popped off of the buffer.
 		uint32 Num()
 		{
diff --git a/Engine/Source/Runtime/AudioMixer/Public/DSP/SpectrumAnalyzer.h b/Engine/Source/Runtime/AudioMixer/Public/DSP/SpectrumAnalyzer.h
--- a/Engine/Source/Runtime/AudioMixer/Public/DSP/SpectrumAnalyzer.h
+++ b/Engine/Source/Runtime/AudioMixer/Public/DSP/SpectrumAnalyzer.h
@@ -4,13 +4,15 @@
 
 #include "CoreMinimal.h"
 #include "DSP/Dsp.h"
+#include "DSP/BufferVectorOperations.h"
 #include "DSP/AudioFFT.h"
 #include "Sound/SampleBuffer.h"
 
 namespace Audio
 {
-	
-	namespace SpectrumAnalyzerSettings
+	class FSpectrumAnalyzer;
+
+	struct AUDIOMIXER_API FSpectrumAnalyzerSettings
 	{
 		// Actual FFT size used. For FSpectrumAnalyzer, we never zero pad the input buffer.
 		enum class EFFTSize : uint16
@@ -21,6 +23,7 @@ namespace Audio
 			Small_256 = 256,
 			Medium_512 = 512,
 			Large_1024 = 1024,
+			VeryLarge_2048 = 2048,
 			TestLarge_4096 = 4096
 		};
 
@@ -33,28 +36,25 @@ namespace Audio
 			Quadratic
 		};
 
-		struct FSettings
-		{
-			EWindowType WindowType;
-			EFFTSize FFTSize;
-			EPeakInterpolationMethod InterpolationMethod;
-
-			/**
-			 * Hop size as a percentage of FFTSize.
-			 * 1.0 indicates a full hop.
-			 * Keeping this as 0.0 will use whatever hop size
-			 * can be used for WindowType to maintain COLA.
-			 */
-			float HopSize;
-
-			FSettings()
-				: WindowType(EWindowType::Hann)
-				, FFTSize(EFFTSize::Default)
-				, InterpolationMethod(EPeakInterpolationMethod::Linear)
-				, HopSize(0.0f)
-			{}
-		};
-	}
+		EWindowType WindowType;
+		EFFTSize FFTSize;
+		EPeakInterpolationMethod InterpolationMethod;
+
+		/**
+			* Hop size as a percentage of FFTSize.
+			* 1.0 indicates a full hop.
+			* Keeping this as 0.0 will use whatever hop size
+			* can be used for WindowType to maintain COLA.
+			*/
+		float HopSize;
+
+		FSpectrumAnalyzerSettings()
+			: WindowType(EWindowType::Hann)
+			, FFTSize(EFFTSize::Default)
+			, InterpolationMethod(EPeakInterpolationMethod::Linear)
+			, HopSize(0.0f)
+		{}
+	};
 
 	/**
 	 * This struct contains the output results from a singular FFT operation.
@@ -79,9 +79,9 @@ namespace Audio
 	{
 	public:
 		FSpectrumAnalyzerBuffer();
-		FSpectrumAnalyzerBuffer(const SpectrumAnalyzerSettings::FSettings& InSettings);
+		FSpectrumAnalyzerBuffer(const FSpectrumAnalyzerSettings& InSettings);
 
-		void Reset(const SpectrumAnalyzerSettings::FSettings& InSettings);
+		void Reset(const FSpectrumAnalyzerSettings& InSettings);
 
 		// Input. Used on analysis thread to lock a buffer to write to.
 		FSpectrumAnalyzerFrequencyVector* StartWorkOnBuffer();
@@ -106,6 +106,31 @@ namespace Audio
 		FCriticalSection BufferIndicesCriticalSection;
 	};
 
+	class FSpectrumAnalysisAsyncWorker : public FNonAbandonableTask
+	{
+	protected:
+		FSpectrumAnalyzer* Analyzer;
+		bool bUseLatestAudio;
+
+	public:
+		FSpectrumAnalysisAsyncWorker(FSpectrumAnalyzer* InAnalyzer, bool bInUseLatestAudio)
+			: Analyzer(InAnalyzer)
+			, bUseLatestAudio(bInUseLatestAudio)
+		{}
+
+		FORCEINLINE TStatId GetStatId() const
+		{
+			RETURN_QUICK_DECLARE_CYCLE_STAT(FSpectrumAnalysisAsyncWorker, STATGROUP_ThreadPoolAsyncTasks);
+		}
+
+		void DoWork();
+
+	private:
+		FSpectrumAnalysisAsyncWorker();
+	};
+
+	typedef FAsyncTask<FSpectrumAnalysisAsyncWorker> FSpectrumAnalyzerTask;
+
 	/**
 	 * Class built to be a rolling spectrum analyzer for arbitrary, monaural audio data.
 	 * Class is meant to scale accuracy with CPU and memory budgets.
@@ -115,20 +140,24 @@ namespace Audio
 	class AUDIOMIXER_API FSpectrumAnalyzer
 	{
 	public:
-		// Default constructor needs to call Init before using
+		// If an instance is created using the default constructor, Init() must be called before it is used.
 		FSpectrumAnalyzer();
+
+		// If an instance is created using either of these constructors, Init() is not neccessary.
 		FSpectrumAnalyzer(float InSampleRate);
-		FSpectrumAnalyzer(const SpectrumAnalyzerSettings::FSettings& InSettings, float InSampleRate);
+		FSpectrumAnalyzer(const FSpectrumAnalyzerSettings& InSettings, float InSampleRate);
+
+		~FSpectrumAnalyzer();
 
 		// Initialize sample rate of analyzer if not known at time of construction
 		void Init(float InSampleRate);
-		void Init(const SpectrumAnalyzerSettings::FSettings& InSettings, float InSampleRate);
+		void Init(const FSpectrumAnalyzerSettings& InSettings, float InSampleRate);
 
 		// Update the settings used by this Spectrum Analyzer. Safe to call on any thread, but should not be called every tick.
-		void SetSettings(const SpectrumAnalyzerSettings::FSettings& InSettings);
+		void SetSettings(const FSpectrumAnalyzerSettings& InSettings);
 
 		// Get the current settings used by this Spectrum Analyzer.
-		void GetSettings(SpectrumAnalyzerSettings::FSettings& OutSettings);
+		void GetSettings(FSpectrumAnalyzerSettings& OutSettings);
 
 		// Samples magnitude (linearly) for a given frequency, in Hz.
 		float GetMagnitudeForFrequency(float InFrequency);
@@ -148,7 +177,14 @@ namespace Audio
 		bool PushAudio(const float* InBuffer, int32 NumSamples);
 
 		// Thread safe call to perform actual FFT. Returns true if it performed the FFT, false otherwise.
-		bool PerformAnalysisIfPossible();
+		// If bAsync is true, this function will kick off an async task.
+		// If bUseLatestAudio is set to true, this function will flush the entire input buffer, potentially losing data.
+		// Otherwise it will only consume enough samples necessary to perform a single FFT.
+		bool PerformAnalysisIfPossible(bool bUseLatestAudio = false, bool bAsync = false);
+
+		// Returns false if this instance of FSpectrumAnalyzer was constructed with the default constructor 
+		// and Init() has not been called yet.
+		bool IsInitialized();
 
 	private:
 
@@ -156,12 +192,14 @@ namespace Audio
 		void ResetSettings();
 
 		// Called in GetMagnitudeForFrequency and GetPhaseForFrequency.
-		void PerformInterpolation(const FSpectrumAnalyzerFrequencyVector* InFrequencies, SpectrumAnalyzerSettings::EPeakInterpolationMethod InMethod, const float InFreq, float& OutReal, float& OutImag);
+		void PerformInterpolation(const FSpectrumAnalyzerFrequencyVector* InFrequencies, FSpectrumAnalyzerSettings::EPeakInterpolationMethod InMethod, const float InFreq, float& OutReal, float& OutImag);
 
 		// Cached current settings. Only actually used in ResetSettings().
-		SpectrumAnalyzerSettings::FSettings CurrentSettings;
+		FSpectrumAnalyzerSettings CurrentSettings;
 		volatile bool bSettingsWereUpdated;
 
+		volatile bool bIsInitialized;
+
 		float SampleRate;
 
 		// Cached window that is applied prior to running the FFT.
@@ -175,5 +213,9 @@ namespace Audio
 
 		// if non-null, owns pointer to locked frequency vector we're using.
 		const FSpectrumAnalyzerFrequencyVector* LockedFrequencyVector;
+
+		// This is used if PerformAnalysisIfPossible is called
+		// with bAsync = true.
+		TUniquePtr<FSpectrumAnalyzerTask> AsyncAnalysisTask;
 	};
 }
diff --git a/Engine/Source/Runtime/Core/Private/Logging/LogMacros.cpp b/Engine/Source/Runtime/Core/Private/Logging/LogMacros.cpp
--- a/Engine/Source/Runtime/Core/Private/Logging/LogMacros.cpp
+++ b/Engine/Source/Runtime/Core/Private/Logging/LogMacros.cpp
@@ -72,7 +72,6 @@ void FMsg::Logf_InternalImpl(const ANSICHAR* File, int32 Line, const FName& Cate
 {
 #if !NO_LOGGING
 	QUICK_SCOPE_CYCLE_COUNTER(STAT_FMsgLogf);
-	CSV_SCOPED_TIMING_STAT(FMsgLogf, FMsgLogfTime);
 	CSV_CUSTOM_STAT(FMsgLogf, FMsgLogfCount, 1, ECsvCustomStatOp::Accumulate);
 
 	if (Verbosity != ELogVerbosity::Fatal)
diff --git a/Engine/Source/Runtime/Core/Public/Math/TransformCalculus2D.h b/Engine/Source/Runtime/Core/Public/Math/TransformCalculus2D.h
--- a/Engine/Source/Runtime/Core/Public/Math/TransformCalculus2D.h
+++ b/Engine/Source/Runtime/Core/Public/Math/TransformCalculus2D.h
@@ -657,6 +657,8 @@ public:
 	/** Access to the translation */
 	const FVector2D& GetTranslation() const { return Trans; }
 
+	void SetTranslation(const FVector2D& InTrans) { Trans = InTrans; }
+
 	/**
 	 * Specialized function to determine if a transform is precisely the identity transform. Uses exact float comparison, so rounding error is not considered.
 	 */
diff --git a/Engine/Source/Runtime/Engine/Classes/Components/AudioComponent.h b/Engine/Source/Runtime/Engine/Classes/Components/AudioComponent.h
--- a/Engine/Source/Runtime/Engine/Classes/Components/AudioComponent.h
+++ b/Engine/Source/Runtime/Engine/Classes/Components/AudioComponent.h
@@ -8,6 +8,7 @@
 #include "Engine/EngineTypes.h"
 #include "Components/SceneComponent.h"
 #include "Sound/SoundAttenuation.h"
+#include "Sound/SoundWave.h"
 
 #include "AudioComponent.generated.h"
 
@@ -427,6 +428,22 @@ public:
 	UFUNCTION(BlueprintCallable, Category = "Audio|Components|Audio")
 	void SetLowPassFilterFrequency(float InLowPassFilterFrequency);
 
+	/** Queries if the sound wave playing in this audio component has cooked FFT data. */
+	UFUNCTION(BlueprintCallable, Category = "Audio|Components|Audio")
+	bool HasCookedFFTData() const;
+
+	/** Queries if the sound wave playing in this audio component has cooked amplitude analyses. */
+	UFUNCTION(BlueprintCallable, Category = "Audio|Components|Audio")
+	bool HasCookedAmplitudeEnvelopeData() const;
+
+	/** Returns the current cooked spectral data of the playing audio component. Returns true if there is data and the audio component is playing. */
+	UFUNCTION(BlueprintCallable, Category = "Audio|Components|Audio")
+	bool GetCookedFFTData(const TArray<float>& FrequenciesToGet, TArray<FSoundWaveSpectralData>& OutSoundWaveSpectralData);
+
+	/** Returns the current cooked envelope data of the playing audio component. Returns true if there is data and the audio component is playing. */
+	UFUNCTION(BlueprintCallable, Category = "Audio|Components|Audio")
+	bool GetCookedEnvelopeData(float& OutEnvelopeData);
+
 	static void PlaybackCompleted(uint64 AudioComponentID, bool bFailedToStart);
 
 private:
@@ -479,7 +496,9 @@ public:
 
 	static UAudioComponent* GetAudioComponentFromID(uint64 AudioComponentID);
 
-	void UpdateInteriorSettings(bool bFullUpdate);
+	// Sets the audio thread playback time as used by the active sound playing this audio component
+	// Will be set if the audio component is using baked FFT or envelope following data so as to be able to feed that data to BP based on playback time
+	void SetPlaybackTimes(const TMap<uint32, float>& InSoundWavePlaybackTimes);
 
 public:
 
@@ -507,6 +526,32 @@ private:
 	FRotator SavedAutoAttachRelativeRotation;
 	FVector SavedAutoAttachRelativeScale3D;
 
+	struct FSoundWavePlaybackTimeData
+	{
+		USoundWave* SoundWave;
+		float PlaybackTime;
+
+		// Cachced indices to boost searching cooked data indices
+		uint32 LastEnvelopeCookedIndex;
+		uint32 LastFFTCookedIndex;
+
+		FSoundWavePlaybackTimeData()
+			: SoundWave(nullptr)
+			, PlaybackTime(0.0f)
+			, LastEnvelopeCookedIndex(INDEX_NONE)
+			, LastFFTCookedIndex(INDEX_NONE)
+		{}
+
+		FSoundWavePlaybackTimeData(USoundWave* InSoundWave)
+			: SoundWave(InSoundWave)
+			, PlaybackTime(0.0f)
+			, LastEnvelopeCookedIndex(INDEX_NONE)
+			, LastFFTCookedIndex(INDEX_NONE)
+		{}
+	};
+	// The current playback times of sound waves in this audio component
+	TMap<uint32, FSoundWavePlaybackTimeData> SoundWavePlaybackTimes;
+
 	/** Restore relative transform from auto attachment and optionally detach from parent (regardless of whether it was an auto attachment). */
 	void CancelAutoAttachment(bool bDetachFromParent);
 
diff --git a/Engine/Source/Runtime/Engine/Classes/Sound/SoundBase.h b/Engine/Source/Runtime/Engine/Classes/Sound/SoundBase.h
--- a/Engine/Source/Runtime/Engine/Classes/Sound/SoundBase.h
+++ b/Engine/Source/Runtime/Engine/Classes/Sound/SoundBase.h
@@ -218,5 +218,8 @@ public:
 	/** Returns the sound concurrency object ID if it exists. If it doesn't exist, returns 0. */
 	uint32 GetSoundConcurrencyObjectID() const;
 
+	/** Returns whether the sound has cooked analysis data (e.g. FFT or envelope following data) and returns sound waves which have cooked data. */
+	virtual bool GetSoundWavesWithCookedAnalysisData(TArray<USoundWave*>& OutSoundWaves);
+
 };
 
diff --git a/Engine/Source/Runtime/Engine/Classes/Sound/SoundCue.h b/Engine/Source/Runtime/Engine/Classes/Sound/SoundCue.h
--- a/Engine/Source/Runtime/Engine/Classes/Sound/SoundCue.h
+++ b/Engine/Source/Runtime/Engine/Classes/Sound/SoundCue.h
@@ -149,6 +149,7 @@ public:
 	virtual float GetDuration() override;
 	virtual const FSoundAttenuationSettings* GetAttenuationSettingsToApply() const override;
 	virtual float GetSubtitlePriority() const override;
+	virtual bool GetSoundWavesWithCookedAnalysisData(TArray<USoundWave*>& OutSoundWaves) override;
 	//~ End USoundBase Interface.
 
 	/** Construct and initialize a node within this Cue */
diff --git a/Engine/Source/Runtime/Engine/Classes/Sound/SoundWave.h b/Engine/Source/Runtime/Engine/Classes/Sound/SoundWave.h
--- a/Engine/Source/Runtime/Engine/Classes/Sound/SoundWave.h
+++ b/Engine/Source/Runtime/Engine/Classes/Sound/SoundWave.h
@@ -129,6 +129,108 @@ struct FStreamedAudioPlatformData
 
 };
 
+// Struct used to retrieve the spectral magnitude of an audio signal at a given frequency to BP
+USTRUCT(BlueprintType)
+struct FSoundWaveSpectralData
+{
+	GENERATED_USTRUCT_BODY()
+
+	// The frequency hz of the spectrum value
+	UPROPERTY(EditAnywhere, BlueprintReadWrite, Category = "SpectralData|Foo")
+	float FrequencyHz;
+
+	// The magnitude of the spectrum at this frequency
+	UPROPERTY(EditAnywhere, BlueprintReadWrite, Category = "SpectralData|Bar")
+	float Magnitude;
+
+	// The normalized magnitude of the spectrum at this frequency
+	UPROPERTY(EditAnywhere, BlueprintReadWrite, Category = "SpectralData|Bar")
+	float NormalizedMagnitude;
+
+	FSoundWaveSpectralData()
+		: FrequencyHz(0.0f)
+		, Magnitude(0.0f)
+		, NormalizedMagnitude(0.0f)
+	{}
+};
+
+// Sort predicate for sorting spectral data by frequency (lowest first)
+struct FCompareSpectralDataByFrequencyHz
+{
+	FORCEINLINE bool operator()(const FSoundWaveSpectralData& A, const FSoundWaveSpectralData& B) const
+	{
+		return A.FrequencyHz < B.FrequencyHz;
+	}
+};
+
+
+// Struct used to store spectral data with time-stamps
+USTRUCT()
+struct FSoundWaveSpectralDataEntry
+{
+	GENERATED_USTRUCT_BODY()
+
+	// The magnitude of the spectrum at this frequency
+	UPROPERTY()
+	float Magnitude;
+
+	// The normalized magnitude of the spectrum at this frequency
+	UPROPERTY()
+	float NormalizedMagnitude;
+};
+
+
+// Struct used to store spectral data with time-stamps
+USTRUCT()
+struct FSoundWaveSpectralTimeData
+{
+	GENERATED_USTRUCT_BODY()
+
+	// The spectral data at the given time. The array indices correspond to the frequencies set to analyze.
+	UPROPERTY()
+	TArray<FSoundWaveSpectralDataEntry> Data;
+
+	// The timestamp associated with this spectral data
+	UPROPERTY()
+	float TimeSec;
+
+	FSoundWaveSpectralTimeData()
+		: TimeSec(0.0f)
+	{}
+};
+
+// Struct used to store time-stamped envelope data
+USTRUCT()
+struct FSoundWaveEnvelopeTimeData
+{
+	GENERATED_USTRUCT_BODY()
+
+	// The normalized linear amplitude of the audio
+	UPROPERTY()
+	float Amplitude;
+
+	// The timestamp of the audio
+	UPROPERTY()
+	float TimeSec;
+
+	FSoundWaveEnvelopeTimeData()
+		: Amplitude(0.0f)
+		, TimeSec(0.0f)
+	{}
+};
+
+// The FFT size (in audio frames) to use for baked FFT analysis
+UENUM(BlueprintType)
+enum class ESoundWaveFFTSize : uint8
+{
+	VerySmall_64,
+	Small_256,
+	Medium_512,
+	Large_1024,
+	VeryLarge_2048,
+};
+
+
 UCLASS(hidecategories=Object, editinlinenew, BlueprintType)
 class ENGINE_API USoundWave : public USoundBase
 {
@@ -196,10 +298,60 @@ class ENGINE_API USoundWave : public USoundBase
 	uint8 bIsAmbisonics : 1;
 
 	/** Whether this SoundWave was decompressed from OGG. */
-	uint8 bDecompressedFromOgg:1;
+	uint8 bDecompressedFromOgg : 1;
+
+#if WITH_EDITORONLY_DATA
+	/** Whether or not to enable cook-time baked FFT analysis. */
+	UPROPERTY(EditAnywhere, Category = "Analysis|FFT")
+	uint8 bEnableBakedFFTAnalysis : 1;
+
+	/** Whether or not to enable cook-time amplitude envelope analysis. */
+	UPROPERTY(EditAnywhere, Category = "Analysis|Envelope")
+	uint8 bEnableAmplitudeEnvelopeAnalysis : 1;
+
+	/** The FFT window size to use for fft analysis. */
+	UPROPERTY(EditAnywhere, Category = "Analysis|FFT", meta = (EditCondition = "bEnableBakedFFTAnalysis"))
+	ESoundWaveFFTSize FFTSize;
+
+	/** How many audio frames analyze at a time. */
+	UPROPERTY(EditAnywhere, Category = "Analysis|FFT", meta = (EditCondition = "bEnableAmplitudeEnvelopeAnalysis", ClampMin = "512", UIMin = "512"))
+	int32 FFTAnalysisFrameSize;
+
+	/** How many audio frames to average a new envelope value. Larger values use less memory for audio envelope data but will result in lower envelope accuracy. */
+	UPROPERTY(EditAnywhere, Category = "Analysis|Envelope", meta = (EditCondition = "bEnableAmplitudeEnvelopeAnalysis", ClampMin = "512", UIMin = "512"))
+	int32 EnvelopeFollowerFrameSize;
+
+	/** How quickly the envelope analyzer responds to increasing amplitudes. */
+	UPROPERTY(EditAnywhere, Category = "Analysis|Envelope", meta = (EditCondition = "bEnableAmplitudeEnvelopeAnalysis", ClampMin = "0", UIMin = "0"))
+	int32 EnvelopeFollowerAttackTime;
+
+	/** How quickly the envelope analyzer responds to decreasing amplitudes. */
+	UPROPERTY(EditAnywhere, Category = "Analysis|Envelope", meta = (EditCondition = "bEnableAmplitudeEnvelopeAnalysis", ClampMin = "0", UIMin = "0"))
+	int32 EnvelopeFollowerReleaseTime;
+#endif
+
+	/** The frequencies (in hz) to analyze when doing baked FFT analysis. */
+	UPROPERTY(EditAnywhere, Category = "Analysis|FFT", meta = (EditCondition = "bEnableBakedFFTAnalysis"))
+	TArray<float> FrequenciesToAnalyze;
+
+	/** The cooked spectral time data. */
+	UPROPERTY()
+	TArray<FSoundWaveSpectralTimeData> CookedSpectralTimeData;
+
+	/** The cooked cooked envelope data. */
+	UPROPERTY()
+	TArray<FSoundWaveEnvelopeTimeData> CookedEnvelopeTimeData;
+
+	/** Helper function to get interpolated cooked FFT data for a given time value. */
+	bool GetInterpolatedCookedFFTDataForTime(float InTime, uint32& InOutLastIndex, TArray<FSoundWaveSpectralData>& OutData, bool bLoop);
+	bool GetInterpolatedCookedEnvelopeDataForTime(float InTime, uint32& InOutLastIndex, float& OutAmplitude, bool bLoop);
 
 private:
 
+	/** Helper functions to search analysis data. Takes starting index to start query. Returns which data index the result was found at. Returns INDEX_NONE if not found. */
+	uint32 GetInterpolatedCookedFFTDataForTimeInternal(float InTime, uint32 StartingIndex, TArray<FSoundWaveSpectralData>& OutData, bool bLoop);
+	uint32 GetInterpolatedCookedEnvelopeDataForTimeInternal(float InTime, uint32 StartingIndex, float& OutAmplitude, bool bLoop);
+
 #if !WITH_EDITOR
 	// This is set to false on initialization, then set to true on non-editor platforms when we cache appropriate sample rate.
 	uint8 bCachedSampleRateFromPlatformSettings : 1;
@@ -387,6 +539,7 @@ public:
 	virtual float GetDuration() override;
 	virtual float GetSubtitlePriority() const override;
 	virtual bool IsAllowedVirtual() const override;
+	virtual bool GetSoundWavesWithCookedAnalysisData(TArray<USoundWave*>& OutSoundWaves) override;
 	//~ End USoundBase Interface.
 
 	// Called  when the procedural sound wave begins on the render thread. Only used in the audio mixer and when bProcedural is true.
@@ -485,9 +638,19 @@ public:
 
 	virtual bool HasCompressedData(FName Format, ITargetPlatform* TargetPlatform = GetRunningPlatform()) const;
 
+#if WITH_EDITOR
+	/** Utility which returns imported PCM data and the parsed header for the file. Returns true if there was data, false if there wasn't. */
+	bool GetImportedSoundWaveData(TArray<uint8>& OutRawPCMData, uint32& OutSampleRate, uint16& OutNumChannels);
+#endif
+
 private:
 	FName GetPlatformSpecificFormat(FName Format, const FPlatformAudioCookOverrides* CompressionOverrides);
 
+#if WITH_EDITOR
+	void BakeFFTAnalysis();
+	void BakeEnvelopeAnalysis();
+#endif
+
 public:
 	virtual void BeginGetCompressedData(FName Format, const FPlatformAudioCookOverrides* CompressionOverrides);
 
@@ -626,6 +789,7 @@ public:
 	{
 		return (ESoundWavePrecacheState)PrecacheState.GetValue();
 	}
+
 };
 
 
diff --git a/Engine/Source/Runtime/Engine/Engine.Build.cs b/Engine/Source/Runtime/Engine/Engine.Build.cs
--- a/Engine/Source/Runtime/Engine/Engine.Build.cs
+++ b/Engine/Source/Runtime/Engine/Engine.Build.cs
@@ -7,7 +7,7 @@ public class Engine : ModuleRules
 {
 	public Engine(ReadOnlyTargetRules Target) : base(Target)
 	{
-		PrivatePCHHeaderFile = "Private/EnginePrivatePCH.h";
+        PrivatePCHHeaderFile = "Private/EnginePrivatePCH.h";
 
 		SharedPCHHeaderFile = "Public/EngineSharedPCH.h";
 
@@ -96,6 +96,7 @@ public class Engine : ModuleRules
 				"CinematicCamera",
 				"Analytics",
 				"AnalyticsET",
+                "AudioMixer",
                 //"CrunchCompression"
             }
 		);
@@ -171,6 +172,7 @@ public class Engine : ModuleRules
 		CircularlyReferencedDependentModules.Add("UMG");
 		CircularlyReferencedDependentModules.Add("MaterialShaderQualitySettings");
 		CircularlyReferencedDependentModules.Add("CinematicCamera");
+        CircularlyReferencedDependentModules.Add("AudioMixer");
 
         // The AnimGraphRuntime module is not needed by Engine proper, but it is loaded in LaunchEngineLoop.cpp,
         // and needs to be listed in an always-included module in order to be compiled into standalone games
diff --git a/Engine/Source/Runtime/Engine/Private/ActiveSound.cpp b/Engine/Source/Runtime/Engine/Private/ActiveSound.cpp
--- a/Engine/Source/Runtime/Engine/Private/ActiveSound.cpp
+++ b/Engine/Source/Runtime/Engine/Private/ActiveSound.cpp
@@ -56,6 +56,7 @@ FActiveSound::FActiveSound()
 	, bUpdatePlayPercentage(false)
 	, bUpdateSingleEnvelopeValue(false)
 	, bUpdateMultiEnvelopeValue(false)
+	, bUpdatePlaybackTime(false)
 	, bIsPlayingAudio(false)
 	, bIsStopping(false)
 	, UserIndex(0)
@@ -543,39 +544,73 @@ void FActiveSound::UpdateWaveInstances( TArray<FWaveInstance*> &InWaveInstances,
 		}
 
 		// Check to see if we need to broadcast the envelope value of sounds playing with this active sound
-		if (bUpdateMultiEnvelopeValue && AudioComponentID > 0)
+		if (AudioComponentID > 0)
 		{
-			int32 NumWaveInstances = ThisSoundsWaveInstances.Num();
-
-			// Add up the envelope value for every wave instance so we get a sum of the envelope value for all sources.
-			float EnvelopeValueSum = 0.0f;
-			float MaxEnvelopeValue = 0.0f;
-			for (FWaveInstance* WaveInstance : ThisSoundsWaveInstances)
+			if (bUpdateMultiEnvelopeValue)
 			{
-				const float WaveInstanceEnvelopeValue = WaveInstance->GetEnvelopeValue();
-				EnvelopeValueSum += WaveInstanceEnvelopeValue;
-				MaxEnvelopeValue = FMath::Max(WaveInstanceEnvelopeValue, MaxEnvelopeValue);
-			}
+				int32 NumWaveInstances = ThisSoundsWaveInstances.Num();
 
-			// Now divide by the number of instances to get the average
-			float AverageEnvelopeValue = EnvelopeValueSum / NumWaveInstances;
-			uint64 AudioComponentIDCopy = AudioComponentID;
-			FAudioThread::RunCommandOnGameThread([AudioComponentIDCopy, AverageEnvelopeValue, MaxEnvelopeValue, NumWaveInstances]()
-			{
-				if (UAudioComponent* AudioComponent = UAudioComponent::GetAudioComponentFromID(AudioComponentIDCopy))
+				// Add up the envelope value for every wave instance so we get a sum of the envelope value for all sources.
+				float EnvelopeValueSum = 0.0f;
+				float MaxEnvelopeValue = 0.0f;
+				for (FWaveInstance* WaveInstance : ThisSoundsWaveInstances)
 				{
-					if (AudioComponent->OnAudioMultiEnvelopeValue.IsBound())
+					const float WaveInstanceEnvelopeValue = WaveInstance->GetEnvelopeValue();
+					EnvelopeValueSum += WaveInstanceEnvelopeValue;
+					MaxEnvelopeValue = FMath::Max(WaveInstanceEnvelopeValue, MaxEnvelopeValue);
+				}
+
+				// Now divide by the number of instances to get the average
+				float AverageEnvelopeValue = EnvelopeValueSum / NumWaveInstances;
+				uint64 AudioComponentIDCopy = AudioComponentID;
+				FAudioThread::RunCommandOnGameThread([AudioComponentIDCopy, AverageEnvelopeValue, MaxEnvelopeValue, NumWaveInstances]()
+				{
+					if (UAudioComponent* AudioComponent = UAudioComponent::GetAudioComponentFromID(AudioComponentIDCopy))
 					{
-						AudioComponent->OnAudioMultiEnvelopeValue.Broadcast(AverageEnvelopeValue, MaxEnvelopeValue, NumWaveInstances);
+						if (AudioComponent->OnAudioMultiEnvelopeValue.IsBound())
+						{
+							AudioComponent->OnAudioMultiEnvelopeValue.Broadcast(AverageEnvelopeValue, MaxEnvelopeValue, NumWaveInstances);
+						}
+
+						if (AudioComponent->OnAudioMultiEnvelopeValueNative.IsBound())
+						{
+							AudioComponent->OnAudioMultiEnvelopeValueNative.Broadcast(AudioComponent, AverageEnvelopeValue, MaxEnvelopeValue, NumWaveInstances);
+						}
 					}
+				});
+			}
+
+			if (bUpdatePlaybackTime)
+			{
+				TMap<uint32, float> WaveInstancePlaybackTimes;
+
+				// Update each of the wave instances playback time based on delta time and the wave instances pitch value
+				for (FWaveInstance* WaveInstance : ThisSoundsWaveInstances)
+				{
+					WaveInstance->PlaybackTime += DeltaTime * WaveInstance->Pitch;
 
-					if (AudioComponent->OnAudioMultiEnvelopeValueNative.IsBound())
+					// For looping sounds, we need to check the wrapping condition
+					if (WaveInstance->LoopingMode != LOOP_Never)
 					{
-						AudioComponent->OnAudioMultiEnvelopeValueNative.Broadcast(AudioComponent, AverageEnvelopeValue, MaxEnvelopeValue, NumWaveInstances);
+						float Duration = WaveInstance->WaveData->Duration;
+						if (WaveInstance->PlaybackTime > Duration)
+						{
+							WaveInstance->PlaybackTime = 0.0f;
+						}
 					}
+					WaveInstancePlaybackTimes.Add(WaveInstance->WaveData->GetUniqueID(), WaveInstance->PlaybackTime);
 				}
-			});
+				uint64 AudioComponentIDCopy = AudioComponentID;
+				FAudioThread::RunCommandOnGameThread([AudioComponentIDCopy, WaveInstancePlaybackTimes]()
+				{
+					if (UAudioComponent* AudioComponent = UAudioComponent::GetAudioComponentFromID(AudioComponentIDCopy))
+					{
+						AudioComponent->SetPlaybackTimes(WaveInstancePlaybackTimes);
+					}
+				});
+			}
 		}
+
 	}
 
 	InWaveInstances.Append(ThisSoundsWaveInstances);
diff --git a/Engine/Source/Runtime/Engine/Private/Components/AudioComponent.cpp b/Engine/Source/Runtime/Engine/Private/Components/AudioComponent.cpp
--- a/Engine/Source/Runtime/Engine/Private/Components/AudioComponent.cpp
+++ b/Engine/Source/Runtime/Engine/Private/Components/AudioComponent.cpp
@@ -13,6 +13,14 @@
 
 DECLARE_CYCLE_STAT(TEXT("AudioComponent Play"), STAT_AudioComp_Play, STATGROUP_Audio);
 
+static float BakedAnalysisTimeShiftCVar = 0.0f;
+FAutoConsoleVariableRef CVarBackedAnalysisTimeShift(
+	TEXT("au.AnalysisTimeShift"),
+	BakedAnalysisTimeShiftCVar,
+	TEXT("Shifts the timeline for baked analysis playback.\n")
+	TEXT("Value: The time in seconds to shift the timeline."),
+	ECVF_Default);
+
 
 /*-----------------------------------------------------------------------------
 UAudioComponent implementation.
@@ -408,6 +416,17 @@ void UAudioComponent::PlayInternal(const float StartTime, const float FadeInDura
 			NewActiveSound.bUpdatePlayPercentage = OnAudioPlaybackPercentNative.IsBound() || OnAudioPlaybackPercent.IsBound();
 			NewActiveSound.bUpdateSingleEnvelopeValue = OnAudioSingleEnvelopeValue.IsBound() || OnAudioSingleEnvelopeValueNative.IsBound();
 			NewActiveSound.bUpdateMultiEnvelopeValue = OnAudioMultiEnvelopeValue.IsBound() || OnAudioMultiEnvelopeValueNative.IsBound();
+			
+			// Setup audio component cooked analysis data playback data set
+			TArray<USoundWave*> SoundWavesWithCookedData;
+			NewActiveSound.bUpdatePlaybackTime = Sound->GetSoundWavesWithCookedAnalysisData(SoundWavesWithCookedData);
+
+			// Reset the audio component's soundwave playback times
+			SoundWavePlaybackTimes.Reset();
+			for (USoundWave* SoundWave : SoundWavesWithCookedData)
+			{
+				SoundWavePlaybackTimes.Add(SoundWave->GetUniqueID(), FSoundWavePlaybackTimeData(SoundWave));
+			}
 
 			NewActiveSound.MaxDistance = MaxDistance;
 
@@ -1110,4 +1129,167 @@ void UAudioComponent::SetLowPassFilterFrequency(float InLowPassFilterFrequency)
 	}
 }
 
+bool UAudioComponent::HasCookedFFTData() const
+{
+	if (Sound)
+	{
+		USoundWave* SoundWave = Cast<USoundWave>(Sound);
+		if (SoundWave)
+		{
+			return SoundWave->CookedSpectralTimeData.Num() > 0;
+		}
+	}
+	return false;
+}
+
+bool UAudioComponent::HasCookedAmplitudeEnvelopeData() const
+{
+	if (Sound)
+	{
+		USoundWave* SoundWave = Cast<USoundWave>(Sound);
+		if (SoundWave)
+		{
+			return SoundWave->CookedEnvelopeTimeData.Num() > 0;
+		}
+	}
+	return false;
+}
+
+void UAudioComponent::SetPlaybackTimes(const TMap<uint32, float>& InSoundWavePlaybackTimes)
+{
+	// Reset the playback times for everything in case the wave instance stops and is not updated
+	for (auto& Elem : SoundWavePlaybackTimes)
+	{
+		Elem.Value.PlaybackTime = 0.0f;
+	}
+
+	for (auto& Elem : InSoundWavePlaybackTimes)
+	{
+		uint32 ObjectId = Elem.Key;
+		FSoundWavePlaybackTimeData* PlaybackTimeData = SoundWavePlaybackTimes.Find(ObjectId);
+		if (PlaybackTimeData)
+		{
+			PlaybackTimeData->PlaybackTime = FMath::Max(Elem.Value - BakedAnalysisTimeShiftCVar, 0.0f);
+		}
+	}
+}
+
+bool UAudioComponent::GetCookedFFTData(const TArray<float>& FrequenciesToGet, TArray<FSoundWaveSpectralData>& OutSoundWaveSpectralData)
+{
+	bool bHadData = false;
+	if (IsPlaying() && SoundWavePlaybackTimes.Num() > 0 && FrequenciesToGet.Num() > 0)
+	{
+		OutSoundWaveSpectralData.Reset();
+		for (float Frequency : FrequenciesToGet)
+		{
+			FSoundWaveSpectralData NewEntry;
+			NewEntry.FrequencyHz = Frequency;
+			OutSoundWaveSpectralData.Add(NewEntry);
+		}
+
+		// Sort by frequency (lowest frequency first).
+		OutSoundWaveSpectralData.Sort(FCompareSpectralDataByFrequencyHz());
+
+		int32 NumEntriesAdded = 0;
+		for (auto& Entry : SoundWavePlaybackTimes)
+		{
+			if (Entry.Value.PlaybackTime > 0.0f && Entry.Value.SoundWave->CookedSpectralTimeData.Num() > 0)
+			{
+				static TArray<FSoundWaveSpectralData> CookedSpectralData;
+				CookedSpectralData.Reset();
+
+				// Find the point in the spectral data that corresponds to the time
+				Entry.Value.SoundWave->GetInterpolatedCookedFFTDataForTime(Entry.Value.PlaybackTime, Entry.Value.LastFFTCookedIndex, CookedSpectralData, Sound->IsLooping());
+
+				if (CookedSpectralData.Num() > 0)
+				{
+					// Find the interpolated values given the frequencies we want to get
+					for (FSoundWaveSpectralData& OutSpectralData : OutSoundWaveSpectralData)
+					{
+						// Check min edge case: we're requesting cooked FFT data lower than what we have cooked
+						if (OutSpectralData.FrequencyHz < CookedSpectralData[0].FrequencyHz)
+						{
+							// Just mix in the lowest value we have cooked
+							OutSpectralData.Magnitude += CookedSpectralData[0].Magnitude;
+							OutSpectralData.NormalizedMagnitude += CookedSpectralData[0].NormalizedMagnitude;
+						}
+						// Check max edge case: we're requesting cooked FFT data at a higher frequency than what we have cooked
+						else if (OutSpectralData.FrequencyHz >= CookedSpectralData.Last().FrequencyHz)
+						{
+							// Just mix in the highest value we have cooked
+							OutSpectralData.Magnitude += CookedSpectralData.Last().Magnitude;
+							OutSpectralData.NormalizedMagnitude += CookedSpectralData.Last().NormalizedMagnitude;
+						}
+						// We need to find the 2 closest cooked results and interpolate those
+						else
+						{
+							for (int32 SpectralDataIndex = 0; SpectralDataIndex < CookedSpectralData.Num() - 1; ++SpectralDataIndex)
+							{
+								const FSoundWaveSpectralData& CurrentSpectralData = CookedSpectralData[SpectralDataIndex];
+								const FSoundWaveSpectralData& NextSpectralData = CookedSpectralData[SpectralDataIndex + 1];
+								if (OutSpectralData.FrequencyHz >= CurrentSpectralData.FrequencyHz && OutSpectralData.FrequencyHz < NextSpectralData.FrequencyHz)
+								{
+									float Alpha = (OutSpectralData.FrequencyHz - CurrentSpectralData.FrequencyHz) / (NextSpectralData.FrequencyHz - CurrentSpectralData.FrequencyHz);
+									OutSpectralData.Magnitude += FMath::Lerp(CurrentSpectralData.Magnitude, NextSpectralData.Magnitude, Alpha);
+									OutSpectralData.NormalizedMagnitude += FMath::Lerp(CurrentSpectralData.NormalizedMagnitude, NextSpectralData.NormalizedMagnitude, Alpha);
+
+									break;
+								}
+							}
+						}
+					}
+					
+					++NumEntriesAdded;
+					bHadData = true;
+				}
+			}
+		}
+
+		// Divide by the number of entries we added (i.e. we are averaging together multiple cooked FFT data in the case of multiple sound waves playing with cooked data)
+		if (NumEntriesAdded > 1)
+		{
+			for (FSoundWaveSpectralData& OutSpectralData : OutSoundWaveSpectralData)
+			{
+				OutSpectralData.Magnitude /= NumEntriesAdded;
+				OutSpectralData.NormalizedMagnitude /= NumEntriesAdded;
+			}
+		}
+	}
+
+	return bHadData;
+}
+
+bool UAudioComponent::GetCookedEnvelopeData(float& OutEnvelopeData)
+{
+	bool bHadData = false;
+	if (IsPlaying() && SoundWavePlaybackTimes.Num() > 0)
+	{
+		static TArray<FSoundWaveEnvelopeTimeData> CookedEnvelopeData;
+		int32 NumEntriesAdded = 0;
+		OutEnvelopeData = 0.0f;
+		for (auto& Entry : SoundWavePlaybackTimes)
+		{
+			if (Entry.Value.SoundWave->CookedEnvelopeTimeData.Num() > 0 && Entry.Value.PlaybackTime > 0.0f)
+			{
+				CookedEnvelopeData.Reset();
 
+				// Find the point in the spectral data that corresponds to the time
+				float SoundWaveAmplitude = 0.0f;
+				if (Entry.Value.SoundWave->GetInterpolatedCookedEnvelopeDataForTime(Entry.Value.PlaybackTime, Entry.Value.LastEnvelopeCookedIndex, SoundWaveAmplitude, Sound->IsLooping()))
+				{
+					OutEnvelopeData += SoundWaveAmplitude;
+					++NumEntriesAdded;
+					bHadData = true;
+				}
+			}
+		}
+
+		// Divide by number of entries we added... get average amplitude envelope
+		if (bHadData)
+		{
+			OutEnvelopeData /= NumEntriesAdded;
+		}
+	}
+
+	return bHadData;
+}
diff --git a/Engine/Source/Runtime/Engine/Private/DataReplication.cpp b/Engine/Source/Runtime/Engine/Private/DataReplication.cpp
--- a/Engine/Source/Runtime/Engine/Private/DataReplication.cpp
+++ b/Engine/Source/Runtime/Engine/Private/DataReplication.cpp
@@ -14,9 +14,9 @@
 #include "Engine/PackageMapClient.h"
 #include "Net/RepLayout.h"
 #include "Engine/ActorChannel.h"
-#include "Engine/DemoNetDriver.h"
 #include "ProfilingDebugging/CsvProfiler.h"
 #include "Engine/Engine.h"
+#include "Engine/NetConnection.h"
 
 DECLARE_CYCLE_STAT(TEXT("Custom Delta Property Rep Time"), STAT_NetReplicateCustomDeltaPropTime, STATGROUP_Game);
 DECLARE_CYCLE_STAT(TEXT("ReceiveRPC"), STAT_NetReceiveRPC, STATGROUP_Game);
diff --git a/Engine/Source/Runtime/Engine/Private/SoundBase.cpp b/Engine/Source/Runtime/Engine/Private/SoundBase.cpp
--- a/Engine/Source/Runtime/Engine/Private/SoundBase.cpp
+++ b/Engine/Source/Runtime/Engine/Private/SoundBase.cpp
@@ -166,6 +166,11 @@ uint32 USoundBase::GetSoundConcurrencyObjectID() const
 	return 0;
 }
 
+bool USoundBase::GetSoundWavesWithCookedAnalysisData(TArray<USoundWave*>& OutSoundWaves)
+{
+	return false;
+}
+
 void USoundBase::PostLoad()
 {
 	Super::PostLoad();
diff --git a/Engine/Source/Runtime/Engine/Private/SoundCue.cpp b/Engine/Source/Runtime/Engine/Private/SoundCue.cpp
--- a/Engine/Source/Runtime/Engine/Private/SoundCue.cpp
+++ b/Engine/Source/Runtime/Engine/Private/SoundCue.cpp
@@ -487,6 +487,24 @@ float USoundCue::GetSubtitlePriority() const
 	return SubtitlePriority;
 }
 
+bool USoundCue::GetSoundWavesWithCookedAnalysisData(TArray<USoundWave*>& OutSoundWaves)
+{
+	// Check this sound cue's wave players to see if any of their soundwaves have cooked analysis data
+	TArray<USoundNodeWavePlayer*> WavePlayers;
+	RecursiveFindNode<USoundNodeWavePlayer>(FirstNode, WavePlayers);
+
+	bool bHasAnalysisData = false;
+	for (USoundNodeWavePlayer* Player : WavePlayers)
+	{
+		USoundWave* SoundWave = Player->GetSoundWave();
+		if (SoundWave->GetSoundWavesWithCookedAnalysisData(OutSoundWaves))
+		{
+			bHasAnalysisData = true;
+		}
+	}
+	return bHasAnalysisData;
+}
+
 #if WITH_EDITOR
 UEdGraph* USoundCue::GetGraph()
 { 
diff --git a/Engine/Source/Runtime/Engine/Private/SoundWave.cpp b/Engine/Source/Runtime/Engine/Private/SoundWave.cpp
--- a/Engine/Source/Runtime/Engine/Private/SoundWave.cpp
+++ b/Engine/Source/Runtime/Engine/Private/SoundWave.cpp
@@ -20,6 +20,9 @@
 #include "ProfilingDebugging/CookStats.h"
 #include "HAL/LowLevelMemTracker.h"
 #include "AudioCompressionSettingsUtils.h"
+#include "DSP/SpectrumAnalyzer.h"
+#include "DSP/EnvelopeFollower.h"
+#include "DSP/BufferVectorOperations.h"
 
 #if ENABLE_COOK_STATS
 namespace SoundWaveCookStats
@@ -111,6 +114,18 @@ USoundWave::USoundWave(const FObjectInitializer& ObjectInitializer)
 	RawPCMDataSize = 0;
 	SetPrecacheState(ESoundWavePrecacheState::NotStarted);
 
+#if WITH_EDITORONLY_DATA
+	FFTSize = ESoundWaveFFTSize::Medium_512;
+	FrequenciesToAnalyze.Add(100.0f);
+	FrequenciesToAnalyze.Add(500.0f);
+	FrequenciesToAnalyze.Add(1000.0f);
+	FrequenciesToAnalyze.Add(5000.0f);
+	FFTAnalysisFrameSize = 1024;
+	EnvelopeFollowerFrameSize = 1024;
+	EnvelopeFollowerAttackTime = 10;
+	EnvelopeFollowerReleaseTime = 100;
+#endif
+
 #if !WITH_EDITOR
 	bCachedSampleRateFromPlatformSettings = false;
 	bSampleRateManuallyReset = false;
@@ -436,6 +451,42 @@ const FPlatformAudioCookOverrides* USoundWave::GetPlatformCompressionOverridesFo
 	return FPlatformCompressionUtilities::GetCookOverridesForCurrentPlatform();
 }
 
+#if WITH_EDITOR
+bool USoundWave::GetImportedSoundWaveData(TArray<uint8>& OutRawPCMData, uint32& OutSampleRate, uint16& OutNumChannels)
+{
+	// Can only get sound wave data if there is bulk data and if we don't have some weird munging of multi-channel files (e.g. mono stereo only)
+	if (RawData.GetBulkDataSize() > 0)
+	{
+		FWaveModInfo WaveInfo;
+
+		const uint8* RawWaveData = (const uint8*)RawData.LockReadOnly();
+		int32 RawDataSize = RawData.GetBulkDataSize();
+
+		// parse the wave data
+		if (!WaveInfo.ReadWaveHeader(RawWaveData, RawDataSize, 0))
+		{
+			UE_LOG(LogAudio, Warning, TEXT("Only mono or stereo 16 bit waves allowed: %s."), *GetFullName());
+			RawData.Unlock();
+			return false;
+		}
+
+		// Copy the raw PCM data and the header info that was parsed
+		OutRawPCMData.Reset();
+		OutRawPCMData.AddUninitialized(WaveInfo.SampleDataSize);
+		FMemory::Memcpy(OutRawPCMData.GetData(), WaveInfo.SampleDataStart, WaveInfo.SampleDataSize);
+
+		OutSampleRate = *WaveInfo.pSamplesPerSec;
+		OutNumChannels = *WaveInfo.pChannels;
+
+		RawData.Unlock();
+		return true;
+	}
+
+	UE_LOG(LogAudio, Warning, TEXT("Failed to get imported raw data for sound wave '%s'"), *GetFullName());
+	return false;
+}
+#endif
+
 FName USoundWave::GetPlatformSpecificFormat(FName Format, const FPlatformAudioCookOverrides* CompressionOverrides)
 {
 	// Platforms that require compression overrides get concatenated formats.
@@ -732,6 +783,239 @@ float USoundWave::GetSampleRateForTargetPlatform(const ITargetPlatform* TargetPl
 	}
 }
 
+static bool AnyFFTAnalysisPropertiesChanged(const FName& PropertyName)
+{
+	// List of properties which cause re-cooking to get triggered
+	static FName EnableFFTAnalysisFName					= GET_MEMBER_NAME_CHECKED(USoundWave, bEnableBakedFFTAnalysis);
+	static FName FFTSizeFName							= GET_MEMBER_NAME_CHECKED(USoundWave, FFTSize);
+	static FName FFTAnalysisFrameSizeFName				= GET_MEMBER_NAME_CHECKED(USoundWave, FFTAnalysisFrameSize);
+	static FName FrequenciesToAnalyzeFName				= GET_MEMBER_NAME_CHECKED(USoundWave, FrequenciesToAnalyze);
+
+	return	PropertyName == EnableFFTAnalysisFName ||
+			PropertyName == FFTSizeFName ||
+			PropertyName == FFTAnalysisFrameSizeFName ||
+			PropertyName == FrequenciesToAnalyzeFName;
+}
+
+static bool AnyEnvelopeAnalysisPropertiesChanged(const FName& PropertyName)
+{
+	static FName EnableAmplitudeEnvelopeAnalysisFName = GET_MEMBER_NAME_CHECKED(USoundWave, bEnableAmplitudeEnvelopeAnalysis);
+	static FName EnvelopeFollowerFrameSizeFName = GET_MEMBER_NAME_CHECKED(USoundWave, EnvelopeFollowerFrameSize);
+	static FName EnvelopeFollowerAttackTimeFName = GET_MEMBER_NAME_CHECKED(USoundWave, EnvelopeFollowerAttackTime);
+	static FName EnvelopeFollowerReleaseTimeFName = GET_MEMBER_NAME_CHECKED(USoundWave, EnvelopeFollowerReleaseTime);
+
+	return	PropertyName == EnableAmplitudeEnvelopeAnalysisFName ||
+			PropertyName == EnvelopeFollowerFrameSizeFName ||
+			PropertyName == EnvelopeFollowerAttackTimeFName ||
+			PropertyName == EnvelopeFollowerReleaseTimeFName;
+
+}
+
+void USoundWave::BakeFFTAnalysis()
+{
+	// Clear any existing spectral data regardless of if it's enabled. If this was enabled and is now toggled, this will clear previous data.
+	CookedSpectralTimeData.Reset();
+
+	// Perform analysis if enabled on the sound wave
+	if (bEnableBakedFFTAnalysis)
+	{
+		// If there are no frequencies to analyze, we can't do the analysis
+		if (!FrequenciesToAnalyze.Num())
+		{
+			UE_LOG(LogAudio, Warning, TEXT("Soundwave '%s' had baked FFT analysis enabled without specifying any frequencies to analyze."), *GetFullName());
+			return;
+		}
+
+		if (ChannelSizes.Num() > 0)
+		{
+			UE_LOG(LogAudio, Warning, TEXT("Soundwave '%s' has multi-channel audio (channels greater than 2). Baking FFT analysis is not currently supported for this yet."), *GetFullName());
+			return;
+		}
+
+		// Retrieve the raw imported data
+		TArray<uint8> RawImportedWaveData;
+		uint32 RawDataSampleRate = 0;
+		uint16 RawDataNumChannels = 0;
+
+		if (!GetImportedSoundWaveData(RawImportedWaveData, RawDataSampleRate, RawDataNumChannels))
+		{
+			return;
+		}
+
+		if (RawDataSampleRate == 0 || RawDataNumChannels == 0)
+		{
+			UE_LOG(LogAudio, Error, TEXT("Failed to parse the raw imported data for '%s' for baked FFT analysis."), *GetFullName());
+			return;
+		}
+
+		const uint32 NumFrames = (RawImportedWaveData.Num() / sizeof(int16)) / RawDataNumChannels;
+		int16* InputData = (int16*)RawImportedWaveData.GetData();
+
+		Audio::FSpectrumAnalyzerSettings SpectrumAnalyzerSettings;
+		switch (FFTSize)
+		{
+		case ESoundWaveFFTSize::VerySmall_64:
+			SpectrumAnalyzerSettings.FFTSize = Audio::FSpectrumAnalyzerSettings::EFFTSize::Min_64;
+			break;
+
+		case ESoundWaveFFTSize::Small_256:
+			SpectrumAnalyzerSettings.FFTSize = Audio::FSpectrumAnalyzerSettings::EFFTSize::Small_256;
+			break;
+
+		default:
+		case ESoundWaveFFTSize::Medium_512:
+			SpectrumAnalyzerSettings.FFTSize = Audio::FSpectrumAnalyzerSettings::EFFTSize::Medium_512;
+			break;
+
+		case ESoundWaveFFTSize::Large_1024:
+			SpectrumAnalyzerSettings.FFTSize = Audio::FSpectrumAnalyzerSettings::EFFTSize::Large_1024;
+			break;
+
+		case ESoundWaveFFTSize::VeryLarge_2048:
+			SpectrumAnalyzerSettings.FFTSize = Audio::FSpectrumAnalyzerSettings::EFFTSize::VeryLarge_2048;
+			break;
+
+		}
+
+		// Build a new spectrum analyzer
+		Audio::FSpectrumAnalyzer SpectrumAnalyzer(SpectrumAnalyzerSettings, (float)SampleRate);
+
+		// The audio data block to use to submit audio data to the spectrum analyzer
+		Audio::AlignedFloatBuffer AnalysisData;
+		check(FFTAnalysisFrameSize > 256);
+		AnalysisData.Reserve(FFTAnalysisFrameSize);
+
+		float MaximumMagnitude = 0.0f;
+		for (uint32 FrameIndex = 0; FrameIndex < NumFrames; ++FrameIndex)
+		{
+			// Get the averaged sample value of all the channels
+			float SampleValue = 0.0f;
+			for (uint16 ChannelIndex = 0; ChannelIndex < RawDataNumChannels; ++ChannelIndex)
+			{
+				SampleValue += (float)InputData[FrameIndex * RawDataNumChannels] / 32767.0f;
+			}
+			SampleValue /= RawDataNumChannels;
+
+			// Accumate the samples in the scratch buffer
+			AnalysisData.Add(SampleValue);
+
+			// Until we reached the frame size
+			if (AnalysisData.Num() == FFTAnalysisFrameSize)
+			{
+				SpectrumAnalyzer.PushAudio(AnalysisData.GetData(), AnalysisData.Num());
+
+				// Block while the analyzer does the analysis
+				while (SpectrumAnalyzer.PerformAnalysisIfPossible()) {}
+
+				FSoundWaveSpectralTimeData NewData;
+
+				// Don't need to lock here since we're doing this sync, but it's here as that's the expected pattern for the Spectrum analyzer
+				SpectrumAnalyzer.LockOutputBuffer();
+
+				// Get the magntiudes for the specified frequencies
+				for (float Frequency : FrequenciesToAnalyze)
+				{
+					FSoundWaveSpectralDataEntry DataEntry;
+					DataEntry.Magnitude = SpectrumAnalyzer.GetMagnitudeForFrequency(Frequency);
+
+					// Track the max magnitude so we can later set normalized magnitudes
+					if (DataEntry.Magnitude > MaximumMagnitude)
+					{
+						MaximumMagnitude = DataEntry.Magnitude;
+					}
+
+					NewData.Data.Add(DataEntry);
+				}
+
+				SpectrumAnalyzer.UnlockOutputBuffer();
+
+				// The time stamp is derived from the frame index and sample rate
+				NewData.TimeSec = FMath::Max((float)(FrameIndex - FFTAnalysisFrameSize + 1) / RawDataSampleRate, 0.0f);
+
+				CookedSpectralTimeData.Add(NewData);
+
+				AnalysisData.Reset();
+			}
+		}
+
+		// It's possible for the maximum magnitude to be 0.0 if the audio file was silent. 
+		if (MaximumMagnitude > 0.0f)
+		{
+			// Normalize all the magnitude values based on the highest magnitude
+			for (FSoundWaveSpectralTimeData& SpectralTimeData : CookedSpectralTimeData)
+			{
+				for (FSoundWaveSpectralDataEntry& DataEntry : SpectralTimeData.Data)
+				{
+					DataEntry.NormalizedMagnitude = DataEntry.Magnitude / MaximumMagnitude;
+				}
+			}
+		}
+
+	}
+}
+
+void USoundWave::BakeEnvelopeAnalysis()
+{
+	// Clear any existing envelope data regardless of if it's enabled. If this was enabled and is now toggled, this will clear previous data.
+	CookedEnvelopeTimeData.Reset();
+
+	// Perform analysis if enabled on the sound wave
+	if (bEnableAmplitudeEnvelopeAnalysis)
+	{
+		if (ChannelSizes.Num() > 0)
+		{
+			UE_LOG(LogAudio, Warning, TEXT("Soundwave '%s' has multi-channel audio (channels greater than 2). Baking envelope analysis is not currently supported for this yet."), *GetFullName());
+			return;
+		}
+
+		// Retrieve the raw imported data
+		TArray<uint8> RawImportedWaveData;
+		uint32 RawDataSampleRate = 0;
+		uint16 RawDataNumChannels = 0;
+
+		if (!GetImportedSoundWaveData(RawImportedWaveData, RawDataSampleRate, RawDataNumChannels))
+		{
+			return;
+		}
+
+		if (RawDataSampleRate == 0 || RawDataNumChannels == 0)
+		{
+			UE_LOG(LogAudio, Error, TEXT("Failed to parse the raw imported data for '%s' for baked FFT analysis."), *GetFullName());
+			return;
+		}
+
+		const uint32 NumFrames = (RawImportedWaveData.Num() / sizeof(int16)) / RawDataNumChannels;
+		int16* InputData = (int16*)RawImportedWaveData.GetData();
+
+		Audio::FEnvelopeFollower EnvelopeFollower;
+		EnvelopeFollower.Init(RawDataSampleRate, EnvelopeFollowerAttackTime, EnvelopeFollowerReleaseTime);
+
+		for (uint32 FrameIndex = 0; FrameIndex < NumFrames; ++FrameIndex)
+		{
+			// Get the averaged sample value of all the channels
+			float SampleValue = 0.0f;
+			for (uint16 ChannelIndex = 0; ChannelIndex < RawDataNumChannels; ++ChannelIndex)
+			{
+				SampleValue += (float)InputData[FrameIndex * RawDataNumChannels] / 32767.0f;
+			}
+			SampleValue /= RawDataNumChannels;
+
+			float Output = EnvelopeFollower.ProcessAudio(SampleValue);
+
+			// Until we reached the frame size
+			if (FrameIndex % EnvelopeFollowerFrameSize == 0)
+			{
+				FSoundWaveEnvelopeTimeData NewData;
+				NewData.Amplitude = Output;
+				NewData.TimeSec = (float)FrameIndex / RawDataSampleRate;
+				CookedEnvelopeTimeData.Add(NewData);
+			}
+		}
+
+	}
+
+}
+
 void USoundWave::PostEditChangeProperty(FPropertyChangedEvent& PropertyChangedEvent)
 {
 	Super::PostEditChangeProperty(PropertyChangedEvent);
@@ -742,7 +1026,7 @@ void USoundWave::PostEditChangeProperty(FPropertyChangedEvent& PropertyChangedEv
 	// Prevent constant re-compression of SoundWave while properties are being changed interactively
 	if (PropertyChangedEvent.ChangeType != EPropertyChangeType::Interactive)
 	{
-		// Regenerate on save any compressed sound formats
+		// Regenerate on save any compressed sound formats or if analysis needs to be re-done
 		if (UProperty* PropertyThatChanged = PropertyChangedEvent.Property)
 		{
 			const FName& Name = PropertyThatChanged->GetFName();
@@ -753,6 +1037,14 @@ void USoundWave::PostEditChangeProperty(FPropertyChangedEvent& PropertyChangedEv
 				UpdatePlatformData();
 				MarkPackageDirty();
 			}
+			else if (AnyFFTAnalysisPropertiesChanged(Name))
+			{
+				BakeFFTAnalysis();
+			}
+			else if (AnyEnvelopeAnalysisPropertiesChanged(Name))
+			{
+				BakeEnvelopeAnalysis();
+			}
 		}
 	}
 }
@@ -1185,6 +1477,16 @@ bool USoundWave::IsStreaming(const FPlatformAudioCookOverrides* Overrides /* = n
 			&& Duration > Overrides->AutoStreamingThreshold);
 }
 
+bool USoundWave::GetSoundWavesWithCookedAnalysisData(TArray<USoundWave*>& OutSoundWaves)
+{
+	if (CookedSpectralTimeData.Num() > 0 || CookedEnvelopeTimeData.Num() > 0)
+	{
+		OutSoundWaves.Add(this);
+		return true;
+	}
+	return false;
+}
+
 void USoundWave::UpdatePlatformData()
 {
 	if (IsStreaming())
@@ -1281,3 +1583,214 @@ bool USoundWave::GetChunkData(int32 ChunkIndex, uint8** OutChunkData)
 	}
 	return true;
 }
+
+uint32 USoundWave::GetInterpolatedCookedFFTDataForTimeInternal(float InTime, uint32 StartingIndex, TArray<FSoundWaveSpectralData>& OutData, bool bLoop)
+{
+	// Find the two entries on either side of the input time
+	int32 NumDataEntries = CookedSpectralTimeData.Num();
+	for (int32 Index = StartingIndex; Index < NumDataEntries; ++Index)
+	{
+		// Get the current data at this index
+		const FSoundWaveSpectralTimeData& CurrentData = CookedSpectralTimeData[Index];
+
+		// Get the next data, wrap if needed (i.e. if current is last index, we'll lerp to the first index)
+		const FSoundWaveSpectralTimeData& NextData = CookedSpectralTimeData[(Index + 1) % NumDataEntries];
+
+		if (InTime >= CurrentData.TimeSec && InTime < NextData.TimeSec)
+		{
+			// Lerping alpha is fraction from current to next data
+			const float Alpha = (InTime - CurrentData.TimeSec) / (NextData.TimeSec - CurrentData.TimeSec);
+			for (int32 FrequencyIndex = 0; FrequencyIndex < FrequenciesToAnalyze.Num(); ++FrequencyIndex)
+			{
+				FSoundWaveSpectralData InterpolatedData;
+				InterpolatedData.FrequencyHz = FrequenciesToAnalyze[FrequencyIndex];
+				InterpolatedData.Magnitude = FMath::Lerp(CurrentData.Data[FrequencyIndex].Magnitude, NextData.Data[FrequencyIndex].Magnitude, Alpha);
+				InterpolatedData.NormalizedMagnitude = FMath::Lerp(CurrentData.Data[FrequencyIndex].NormalizedMagnitude, NextData.Data[FrequencyIndex].NormalizedMagnitude, Alpha);
+
+				OutData.Add(InterpolatedData);
+			}
+
+			// Sort by frequency (lowest frequency first).
+			OutData.Sort(FCompareSpectralDataByFrequencyHz());
+
+			// We found cooked spectral data which maps to these indices
+			return Index;
+		}
+	}
+
+	return INDEX_NONE;
+}
+
+bool USoundWave::GetInterpolatedCookedFFTDataForTime(float InTime, uint32& InOutLastIndex, TArray<FSoundWaveSpectralData>& OutData, bool bLoop)
+{
+	if (CookedSpectralTimeData.Num() > 0)
+	{
+		// Handle edge cases
+		if (!bLoop)
+		{
+			// Pointer to which data to use
+			FSoundWaveSpectralTimeData* SpectralTimeData = nullptr;
+
+			// We are past the edge
+			if (InTime >= CookedSpectralTimeData.Last().TimeSec)
+			{
+				SpectralTimeData = &CookedSpectralTimeData.Last();
+				InOutLastIndex = CookedPlatformData.Num() - 1;
+			}
+			// We are before the first data point
+			else if (InTime < CookedSpectralTimeData[0].TimeSec)
+			{
+				SpectralTimeData = &CookedSpectralTimeData[0];
+				InOutLastIndex = 0;
+			}
+
+			// If we were either case before we have a non-nullptr here
+			if (SpectralTimeData != nullptr)
+			{
+				// Create an entry for this clamped output
+				for (int32 FrequencyIndex = 0; FrequencyIndex < FrequenciesToAnalyze.Num(); ++FrequencyIndex)
+				{
+					FSoundWaveSpectralData InterpolatedData;
+					InterpolatedData.FrequencyHz = FrequenciesToAnalyze[FrequencyIndex];
+					InterpolatedData.Magnitude = SpectralTimeData->Data[FrequencyIndex].Magnitude;
+					InterpolatedData.NormalizedMagnitude = SpectralTimeData->Data[FrequencyIndex].NormalizedMagnitude;
+
+					OutData.Add(InterpolatedData);
+				}
+
+				return true;
+			}
+		}
+		// We're looping 
+		else 
+		{
+			// Need to check initial wrap-around case (i.e. we're reading earlier than first data point so need to lerp from last data point to first
+			if (InTime >= 0.0f && InTime < CookedSpectralTimeData[0].TimeSec)
+			{
+				const FSoundWaveSpectralTimeData& CurrentData = CookedSpectralTimeData.Last();
+
+				// Get the next data, wrap if needed (i.e. if current is last index, we'll lerp to the first index)
+				const FSoundWaveSpectralTimeData& NextData = CookedSpectralTimeData[0];
+
+				float TimeLeftFromLastDataToEnd = Duration - CurrentData.TimeSec;
+				float Alpha = (TimeLeftFromLastDataToEnd + InTime) / (TimeLeftFromLastDataToEnd + NextData.TimeSec);
+
+				for (int32 FrequencyIndex = 0; FrequencyIndex < FrequenciesToAnalyze.Num(); ++FrequencyIndex)
+				{
+					FSoundWaveSpectralData InterpolatedData;
+					InterpolatedData.FrequencyHz = FrequenciesToAnalyze[FrequencyIndex];
+					InterpolatedData.Magnitude = FMath::Lerp(CurrentData.Data[FrequencyIndex].Magnitude, NextData.Data[FrequencyIndex].Magnitude, Alpha);
+					InterpolatedData.NormalizedMagnitude = FMath::Lerp(CurrentData.Data[FrequencyIndex].NormalizedMagnitude, NextData.Data[FrequencyIndex].NormalizedMagnitude, Alpha);
+
+					OutData.Add(InterpolatedData);
+
+					InOutLastIndex = 0;
+				}
+				return true;
+			}
+			// Or we've been offset a bit in the negative.
+			else if (InTime < 0.0f)
+			{
+				// Wrap the time to the end of the sound wave file 
+				InTime = FMath::Clamp(Duration + InTime, 0.0f, Duration);
+			}
+		}
+
+		uint32 StartingIndex = InOutLastIndex == INDEX_NONE ? 0 : InOutLastIndex;
+
+		InOutLastIndex = GetInterpolatedCookedFFTDataForTimeInternal(InTime, StartingIndex, OutData, bLoop);
+		if (InOutLastIndex == INDEX_NONE && StartingIndex != 0)
+		{
+			InOutLastIndex = GetInterpolatedCookedFFTDataForTimeInternal(InTime, 0, OutData, bLoop);
+		}
+		return InOutLastIndex != INDEX_NONE;
+	}
+
+	return false;
+}
+
+uint32 USoundWave::GetInterpolatedCookedEnvelopeDataForTimeInternal(float InTime, uint32 StartingIndex, float& OutAmplitude, bool bLoop)
+{
+	if (StartingIndex == INDEX_NONE || StartingIndex == CookedEnvelopeTimeData.Num())
+	{
+		StartingIndex = 0;
+	}
+
+	// Find the two entries on either side of the input time
+	int32 NumDataEntries = CookedEnvelopeTimeData.Num();
+	for (int32 Index = StartingIndex; Index < NumDataEntries; ++Index)
+	{
+		const FSoundWaveEnvelopeTimeData& CurrentData = CookedEnvelopeTimeData[Index];
+		const FSoundWaveEnvelopeTimeData& NextData = CookedEnvelopeTimeData[(Index + 1) % NumDataEntries];
+
+		if (InTime >= CurrentData.TimeSec && InTime < NextData.TimeSec)
+		{
+			// Lerping alpha is fraction from current to next data
+			const float Alpha = (InTime - CurrentData.TimeSec) / (NextData.TimeSec - CurrentData.TimeSec);
+			OutAmplitude = FMath::Lerp(CurrentData.Amplitude, NextData.Amplitude, Alpha);
+
+			// We found cooked spectral data which maps to these indices
+			return Index;
+		}
+	}
+
+	// Did not find the data
+	return INDEX_NONE;
+}
+
+bool USoundWave::GetInterpolatedCookedEnvelopeDataForTime(float InTime, uint32& InOutLastIndex, float& OutAmplitude, bool bLoop)
+{
+	InOutLastIndex = INDEX_NONE;
+	if (CookedEnvelopeTimeData.Num() > 0 && InTime >= 0.0f)
+	{
+		// Handle edge cases
+		if (!bLoop)
+		{
+			// We are past the edge
+			if (InTime >= CookedEnvelopeTimeData.Last().TimeSec)
+			{
+				OutAmplitude = CookedEnvelopeTimeData.Last().Amplitude;
+				InOutLastIndex = CookedEnvelopeTimeData.Num() - 1;
+				return true;
+			}
+			// We are before the first data point
+			else if (InTime < CookedEnvelopeTimeData[0].TimeSec)
+			{
+				OutAmplitude = CookedEnvelopeTimeData[0].Amplitude;
+				InOutLastIndex = 0;
+				return true;
+			}
+		}	
+		else
+		{
+			// Need to check initial wrap-around case (i.e. we're reading earlier than first data point so need to lerp from last data point to first
+			if (InTime >= 0.0f && InTime < CookedEnvelopeTimeData[0].TimeSec)
+			{
+				const FSoundWaveEnvelopeTimeData& CurrentData = CookedEnvelopeTimeData.Last();
+				const FSoundWaveEnvelopeTimeData& NextData = CookedEnvelopeTimeData[0];
+
+				float TimeLeftFromLastDataToEnd = Duration - CurrentData.TimeSec;
+				float Alpha = (TimeLeftFromLastDataToEnd + InTime) / (TimeLeftFromLastDataToEnd + NextData.TimeSec);
+
+				OutAmplitude = FMath::Lerp(CurrentData.Amplitude, NextData.Amplitude, Alpha);
+				InOutLastIndex = 0;
+				return true;
+			}
+			// Or we've been offset a bit in the negative.
+			else if (InTime < 0.0f)
+			{
+				// Wrap the time to the end of the sound wave file 
+				InTime = FMath::Clamp(Duration + InTime, 0.0f, Duration);
+			}
+
+			uint32 StartingIndex = InOutLastIndex == INDEX_NONE ? 0 : InOutLastIndex;
+
+			InOutLastIndex = GetInterpolatedCookedEnvelopeDataForTimeInternal(InTime, StartingIndex, OutAmplitude, bLoop);
+			if (InOutLastIndex == INDEX_NONE && StartingIndex != 0)
+			{
+				InOutLastIndex = GetInterpolatedCookedEnvelopeDataForTimeInternal(InTime, 0, OutAmplitude, bLoop);
+			}
+		}
+	}
+	return InOutLastIndex != INDEX_NONE;
+}
diff --git a/Engine/Source/Runtime/Engine/Public/ActiveSound.h b/Engine/Source/Runtime/Engine/Public/ActiveSound.h
--- a/Engine/Source/Runtime/Engine/Public/ActiveSound.h
+++ b/Engine/Source/Runtime/Engine/Public/ActiveSound.h
@@ -366,6 +366,9 @@ public:
 	/** Whether or not this active sound will update the average envelope value of every wave instance that plays a sound source. Based on set delegates on audio component. */
 	uint8 bUpdateMultiEnvelopeValue:1;
 
+	/** Whether or not the active sound should update it's owning audio component's playback time. */
+	uint8 bUpdatePlaybackTime:1;
+
 	/** Whether or not this active sound is playing audio, as in making audible sounds. */
 	uint8 bIsPlayingAudio:1;
 
diff --git a/Engine/Source/Runtime/Engine/Public/Audio.h b/Engine/Source/Runtime/Engine/Public/Audio.h
--- a/Engine/Source/Runtime/Engine/Public/Audio.h
+++ b/Engine/Source/Runtime/Engine/Public/Audio.h
@@ -427,6 +427,9 @@ public:
 	/** The absolute position of the wave instance relative to forward vector of listener. */
 	float AbsoluteAzimuth; 
 
+	/** The playback time of the wave instance. Updated from active sound. */
+	float PlaybackTime;
+
 	/** The reverb send method to use. */
 	EReverbSendMethod ReverbSendMethod;
 
@@ -846,7 +849,6 @@ protected:
 	friend struct FActiveSound;
 };
 
-
 /*-----------------------------------------------------------------------------
 	FWaveModInfo. 
 -----------------------------------------------------------------------------*/
diff --git a/Engine/Source/Runtime/Engine/Public/AudioDevice.h b/Engine/Source/Runtime/Engine/Public/AudioDevice.h
--- a/Engine/Source/Runtime/Engine/Public/AudioDevice.h
+++ b/Engine/Source/Runtime/Engine/Public/AudioDevice.h
@@ -15,6 +15,7 @@
 #include "Sound/SoundSourceBus.h"
 #include "Sound/AudioSettings.h"
 #include "AudioDeviceManager.h"
+#include "DSP/SpectrumAnalyzer.h"
 #include "EngineGlobals.h"
 
 class FAudioEffectsManager;
@@ -1204,6 +1205,26 @@ public:
 		UE_LOG(LogAudio, Error, TEXT("Envelope following submixes only works with the audio mixer. Please run using -audiomixer or set INI file to use submix recording."));
 	}
 
+	virtual void StartSpectrumAnalysis(USoundSubmix* InSubmix, const Audio::FSpectrumAnalyzerSettings& InSettings)
+	{
+		UE_LOG(LogAudio, Error, TEXT("Spectrum analysis of submixes only works with the audio mixer. Please run using -audiomixer or set INI file to use submix recording."));
+	}
+
+	virtual void StopSpectrumAnalysis(USoundSubmix* InSubmix)
+	{
+		UE_LOG(LogAudio, Error, TEXT("Spectrum analysis of submixes only works with the audio mixer. Please run using -audiomixer or set INI file to use submix recording."));
+	}
+
+	virtual void GetMagnitudesForFrequencies(USoundSubmix* InSubmix, const TArray<float>& InFrequencies, TArray<float>& OutMagnitudes)
+	{
+		UE_LOG(LogAudio, Error, TEXT("Spectrum analysis of submixes only works with the audio mixer. Please run using -audiomixer or set INI file to use submix recording."));
+	}
+
+	virtual void GetPhasesForFrequencies(USoundSubmix* InSubmix, const TArray<float>& InFrequencies, TArray<float>& OutPhases)
+	{
+		UE_LOG(LogAudio, Error, TEXT("Spectrum analysis of submixes only works with the audio mixer. Please run using -audiomixer or set INI file to use submix recording."));
+	}
+
 protected:
 	friend class FSoundSource;
 
diff --git a/Engine/Source/Runtime/JsonUtilities/Public/JsonDomBuilder.h b/Engine/Source/Runtime/JsonUtilities/Public/JsonDomBuilder.h
--- a/Engine/Source/Runtime/JsonUtilities/Public/JsonDomBuilder.h
+++ b/Engine/Source/Runtime/JsonUtilities/Public/JsonDomBuilder.h
@@ -4,6 +4,12 @@
 
 #include "Dom/JsonValue.h"
 #include "Dom/JsonObject.h"
+#include "Serialization/JsonSerializer.h"
+
+#include "Templates/IsFloatingPoint.h"
+#include "Templates/IsIntegral.h"
+#include "Templates/EnableIf.h"
+#include "Templates/Invoke.h"
 
 /**
  * Helpers for creating TSharedPtr<FJsonValue> JSON trees
@@ -40,20 +46,39 @@ public:
 			return MakeShared<FJsonValueObject>(Object);
 		}
 
-		FObject& Set(const FString& Key, const FArray& Arr)        { Object->SetField(Key, Arr.AsJsonValue());                      return *this; }
-		FObject& Set(const FString& Key, const FObject& Obj)       { Object->SetField(Key, Obj.AsJsonValue());                      return *this; }
+		template <template <class> class TPrintPolicy = TPrettyJsonPrintPolicy>
+		FString ToString() const
+		{
+			FString Result;
+			auto JsonWriter = TJsonWriterFactory<TCHAR, TPrintPolicy<TCHAR>>::Create(&Result);
+			FJsonSerializer::Serialize(Object, JsonWriter);
+			return Result;
+		}
+
+		FObject& Set(const FString& Key, const FArray& Arr)            { Object->SetField(Key, Arr.AsJsonValue());                      return *this; }
+		FObject& Set(const FString& Key, const FObject& Obj)           { Object->SetField(Key, Obj.AsJsonValue());                      return *this; }
 
-		FObject& Set(const FString& Key, const FString& Str)       { Object->SetField(Key, MakeShared<FJsonValueString>(Str));      return *this; }
+		FObject& Set(const FString& Key, const FString& Str)           { Object->SetField(Key, MakeShared<FJsonValueString>(Str));      return *this; }
 
 		template <class FNumber>
 		typename TEnableIf<TIsIntegral<FNumber>::Value || TIsFloatingPoint<FNumber>::Value, FObject&>::Type
-			Set(const FString& Key, FNumber Number)                { Object->SetField(Key, MakeShared<FJsonValueNumber>(Number));   return *this; }
+			Set(const FString& Key, FNumber Number)                    { Object->SetField(Key, MakeShared<FJsonValueNumber>(Number));   return *this; }
+
+		FObject& Set(const FString& Key, bool Boolean)                 { Object->SetField(Key, MakeShared<FJsonValueBoolean>(Boolean)); return *this; }
+		FObject& Set(const FString& Key, TYPE_OF_NULLPTR)              { Object->SetField(Key, MakeShared<FJsonValueNull>());           return *this; }
 
-		FObject& Set(const FString& Key, bool Boolean)             { Object->SetField(Key, MakeShared<FJsonValueBoolean>(Boolean)); return *this; }
-		FObject& Set(const FString& Key, TYPE_OF_NULLPTR)          { Object->SetField(Key, MakeShared<FJsonValueNull>());           return *this; }
+		FObject& Set(const FString& Key, TSharedPtr<FJsonValue> Value) { Object->SetField(Key, Value);                                  return *this; }
 
-		template <class FValue>
-		FObject& Set(const FString& Key, TSharedPtr<FValue> Value) { Object->SetField(Key, Value);                                  return *this; }
+		void CopyIf(const FJsonObject& Src, TFunctionRef<bool (const FString&, const FJsonValue&)> Pred)
+		{
+			for (const TPair<FString, TSharedPtr<FJsonValue>>& KV: Src.Values)
+			{
+				if (ensure(KV.Value) && Pred(KV.Key, *KV.Value))
+				{
+					Object->SetField(KV.Key, KV.Value);
+				}
+			}
+		}
 
 	private:
 		TSharedRef<FJsonObject> Object;
@@ -67,21 +92,30 @@ public:
 			return MakeShared<FJsonValueArray>(Array);
 		}
 
-		FArray& Add(const FArray& Arr)        { Array.Emplace(Arr.AsJsonValue());                      return *this; }
-		FArray& Add(const FObject& Obj)       { Array.Emplace(Obj.AsJsonValue());                      return *this; }
+		template <template <class> class TPrintPolicy = TPrettyJsonPrintPolicy>
+		FString ToString() const
+		{
+			FString Result;
+			auto JsonWriter = TJsonWriterFactory<TCHAR, TPrintPolicy<TCHAR>>::Create(&Result);
+			FJsonSerializer::Serialize(Array, JsonWriter);
+			return Result;
+		}
+
+		FArray& Add(const FArray& Arr)            { Array.Emplace(Arr.AsJsonValue());                      return *this; }
+		FArray& Add(const FObject& Obj)           { Array.Emplace(Obj.AsJsonValue());                      return *this; }
 
-		FArray& Add(const FString& Str)       { Array.Emplace(MakeShared<FJsonValueString>(Str));      return *this; }
+		FArray& Add(const FString& Str)           { Array.Emplace(MakeShared<FJsonValueString>(Str));      return *this; }
 
 		template <class FNumber>
 		typename TEnableIf<TIsIntegral<FNumber>::Value || TIsFloatingPoint<FNumber>::Value, FArray&>::Type
-			Add(FNumber Number)               { Array.Emplace(MakeShared<FJsonValueNumber>(Number));   return *this; }
+			Add(FNumber Number)                   { Array.Emplace(MakeShared<FJsonValueNumber>(Number));   return *this; }
 
-		FArray& Add(bool Boolean)             { Array.Emplace(MakeShared<FJsonValueBoolean>(Boolean)); return *this; }
-		FArray& Add(TYPE_OF_NULLPTR)          { Array.Emplace(MakeShared<FJsonValueNull>());           return *this; }
+		FArray& Add(bool Boolean)                 { Array.Emplace(MakeShared<FJsonValueBoolean>(Boolean)); return *this; }
+		FArray& Add(TYPE_OF_NULLPTR)              { Array.Emplace(MakeShared<FJsonValueNull>());           return *this; }
 
-		template <class FValue>
-		FArray& Add(TSharedPtr<FValue> Value) { Array.Emplace(Value);                                  return *this; }
+		FArray& Add(TSharedPtr<FJsonValue> Value) { Array.Emplace(Value);                                  return *this; }
 
+		/** Add multiple values */
 		template <class... FValue>
 		typename TEnableIf<(sizeof...(FValue) > 1), FArray&>::Type
 			Add(FValue&&... Value)
@@ -92,6 +126,16 @@ public:
 			return *this;
 		}
 
+		void CopyIf(const TArray<TSharedPtr<FJsonValue>>& Src, TFunctionRef<bool (const FJsonValue&)> Pred)
+		{
+			for (const TSharedPtr<FJsonValue>& Value: Src)
+			{
+				if (ensure(Value) && Pred(*Value))
+				{
+					Array.Emplace(Value);
+				}
+			}
+		}
 	private:
 		TArray<TSharedPtr<FJsonValue>> Array;
 	};
diff --git a/Engine/Source/Runtime/MediaAssets/Private/Assets/MediaSoundComponent.cpp b/Engine/Source/Runtime/MediaAssets/Private/Assets/MediaSoundComponent.cpp
--- a/Engine/Source/Runtime/MediaAssets/Private/Assets/MediaSoundComponent.cpp
+++ b/Engine/Source/Runtime/MediaAssets/Private/Assets/MediaSoundComponent.cpp
@@ -31,6 +31,7 @@ FAutoConsoleVariableRef CVarSyncAudioAfterDropouts(
 
 DECLARE_FLOAT_COUNTER_STAT(TEXT("MediaUtils MediaSoundComponent Sync"), STAT_MediaUtils_MediaSoundComponentSync, STATGROUP_Media);
 DECLARE_FLOAT_COUNTER_STAT(TEXT("MediaUtils MediaSoundComponent SampleTime"), STAT_MediaUtils_MediaSoundComponentSampleTime, STATGROUP_Media);
+DECLARE_DWORD_COUNTER_STAT(TEXT("MediaUtils MediaSoundComponent Queued"), STAT_Media_SoundCompQueued, STATGROUP_Media);
 
 /* Static initialization
  *****************************************************************************/
@@ -52,8 +53,13 @@ UMediaSoundComponent::UMediaSoundComponent(const FObjectInitializer& ObjectIniti
 	, RateAdjustment(1.0f)
 	, Resampler(new FMediaAudioResampler)
 	, FrameSyncOffset(0)
-	, bSyncAudioAfterDropouts(false)
 	, LastPlaySampleTime(FTimespan::MinValue())
+	, EnvelopeFollowerAttackTime(10)
+	, EnvelopeFollowerReleaseTime(100)
+	, CurrentEnvelopeValue(0.0f)
+	, bSyncAudioAfterDropouts(false)
+	, bSpectralAnalysisEnabled(false)
+	, bEnvelopeFollowingEnabled(false)
 {
 	PrimaryComponentTick.bCanEverTick = true;
 	bAutoActivate = true;
@@ -414,36 +420,62 @@ int32 UMediaSoundComponent::OnGenerateAudio(float* OutAudio, int32 NumSamples)
 
 		LastPlaySampleTime = OutTime;
 
-		if (bSpectralAnalysisEnabled)
+
+		if (bSpectralAnalysisEnabled || bEnvelopeFollowingEnabled)
 		{
-			// If we have stereo audio, sum to mono before sending to analyzer
+			float* BufferToUseForAnalysis = nullptr;
+			int32 NumFrames = NumSamples;
+			
 			if (NumChannels == 2)
 			{
-				int32 NumFrames = NumSamples / 2;
+				NumFrames = NumSamples / 2;
 
 				// Use the scratch buffer to sum the audio to mono
 				AudioScratchBuffer.Reset();
 				AudioScratchBuffer.AddUninitialized(NumFrames);
-				float* AudioScratchBufferPtr = AudioScratchBuffer.GetData();
+				BufferToUseForAnalysis = AudioScratchBuffer.GetData();
 				int32 SampleIndex = 0;
 				for (int32 FrameIndex = 0; FrameIndex < NumFrames; ++FrameIndex, SampleIndex += NumChannels)
 				{
-					AudioScratchBufferPtr[FrameIndex] = 0.5f * (OutAudio[SampleIndex] + OutAudio[SampleIndex + 1]);
+					BufferToUseForAnalysis[FrameIndex] = 0.5f * (OutAudio[SampleIndex] + OutAudio[SampleIndex + 1]);
 				}
-
-				SpectrumAnalyzer.PushAudio(AudioScratchBufferPtr, NumFrames);
 			}
 			else
 			{
-				SpectrumAnalyzer.PushAudio(OutAudio, NumSamples);
+				BufferToUseForAnalysis = OutAudio;
+			}
+
+			if (bSpectralAnalysisEnabled)
+			{
+				SpectrumAnalyzer.PushAudio(BufferToUseForAnalysis, NumFrames);
+				SpectrumAnalyzer.PerformAnalysisIfPossible(true, true);
 			}
 
-			// Launch an analysis task with this audio
-			(new FAutoDeleteAsyncTask<FMediaSoundComponentSpectrumAnalysisTask>(&SpectrumAnalyzer, &SpectrumAnalysisCounter))->StartBackgroundTask();
+			{
+				FScopeLock ScopeLock(&EnvelopeFollowerCriticalSection);
+				if (bEnvelopeFollowingEnabled)
+				{
+					if (bEnvelopeFollowerSettingsChanged)
+					{
+						EnvelopeFollower.SetAttackTime((float)EnvelopeFollowerAttackTime);
+						EnvelopeFollower.SetReleaseTime((float)EnvelopeFollowerReleaseTime);
+
+						bEnvelopeFollowerSettingsChanged = false;
+					}
+
+					for (int32 FrameIndex = 0; FrameIndex < NumFrames; ++FrameIndex)
+					{
+						EnvelopeFollower.ProcessAudio(BufferToUseForAnalysis[FrameIndex]);
+					}
+
+					CurrentEnvelopeValue = EnvelopeFollower.GetCurrentValue();
+				}
+			}
 		}
 
 		SET_FLOAT_STAT(STAT_MediaUtils_MediaSoundComponentSync, FMath::Abs((Time - OutTime).GetTotalMilliseconds()));
 		SET_FLOAT_STAT(STAT_MediaUtils_MediaSoundComponentSampleTime, OutTime.GetTotalMilliseconds());
+		SET_DWORD_STAT(STAT_Media_SoundCompQueued, PinnedSampleQueue->Num());
 	}
 	else
 	{
@@ -467,25 +499,25 @@ void UMediaSoundComponent::SetEnableSpectralAnalysis(bool bInSpectralAnalysisEna
 
 void UMediaSoundComponent::SetSpectralAnalysisSettings(TArray<float> InFrequenciesToAnalyze, EMediaSoundComponentFFTSize InFFTSize)
 {
-	Audio::SpectrumAnalyzerSettings::EFFTSize SpectrumAnalyzerSize;
+	Audio::FSpectrumAnalyzerSettings::EFFTSize SpectrumAnalyzerSize;
 
 	switch (InFFTSize)
 	{
 		case EMediaSoundComponentFFTSize::Min_64: 
-			SpectrumAnalyzerSize = Audio::SpectrumAnalyzerSettings::EFFTSize::Min_64; 
+			SpectrumAnalyzerSize = Audio::FSpectrumAnalyzerSettings::EFFTSize::Min_64;
 			break;
 		
 		case EMediaSoundComponentFFTSize::Small_256: 
-			SpectrumAnalyzerSize = Audio::SpectrumAnalyzerSettings::EFFTSize::Small_256; 
+			SpectrumAnalyzerSize = Audio::FSpectrumAnalyzerSettings::EFFTSize::Small_256;
 			break;
 		
 		default:
 		case EMediaSoundComponentFFTSize::Medium_512:
-			SpectrumAnalyzerSize = Audio::SpectrumAnalyzerSettings::EFFTSize::Medium_512; 
+			SpectrumAnalyzerSize = Audio::FSpectrumAnalyzerSettings::EFFTSize::Medium_512;
 			break;
 
 		case EMediaSoundComponentFFTSize::Large_1024: 
-			SpectrumAnalyzerSize = Audio::SpectrumAnalyzerSettings::EFFTSize::Large_1024; 
+			SpectrumAnalyzerSize = Audio::FSpectrumAnalyzerSettings::EFFTSize::Large_1024;
 			break;
 	}
 
@@ -517,6 +549,26 @@ TArray<FMediaSoundComponentSpectralData> UMediaSoundComponent::GetSpectralData()
 	return TArray<FMediaSoundComponentSpectralData>();
 }
 
+void UMediaSoundComponent::SetEnableEnvelopeFollowing(bool bInEnvelopeFollowing)
+{
+	FScopeLock ScopeLock(&EnvelopeFollowerCriticalSection);
+	bEnvelopeFollowingEnabled = bInEnvelopeFollowing;
+	CurrentEnvelopeValue = 0.0f;
+}
+
+void UMediaSoundComponent::SetEnvelopeFollowingsettings(int32 AttackTimeMsec, int32 ReleaseTimeMsec)
+{
+	FScopeLock ScopeLock(&EnvelopeFollowerCriticalSection);
+	EnvelopeFollowerAttackTime = AttackTimeMsec;
+	EnvelopeFollowerReleaseTime = ReleaseTimeMsec;
+	bEnvelopeFollowerSettingsChanged = true;
+}
+
+float UMediaSoundComponent::GetEnvelopeValue() const
+{
+	return CurrentEnvelopeValue;
+}
+
 /* UMediaSoundComponent implementation
  *****************************************************************************/
 
diff --git a/Engine/Source/Runtime/MediaAssets/Public/MediaSoundComponent.h b/Engine/Source/Runtime/MediaAssets/Public/MediaSoundComponent.h
--- a/Engine/Source/Runtime/MediaAssets/Public/MediaSoundComponent.h
+++ b/Engine/Source/Runtime/MediaAssets/Public/MediaSoundComponent.h
@@ -14,6 +14,8 @@
 #include "UObject/ScriptMacros.h"
 #include "DSP/SpectrumAnalyzer.h"
 #include "DSP/BufferVectorOperations.h"
+#include "DSP/EnvelopeFollower.h"
+
 
 #include "MediaSoundComponent.generated.h"
 
@@ -64,34 +66,6 @@ struct FMediaSoundComponentSpectralData
 	float Magnitude;
 };
 
-class FMediaSoundComponentSpectrumAnalysisTask : public FNonAbandonableTask
-{
-	friend class FAutoDeleteAsyncTask<FMediaSoundComponentSpectrumAnalysisTask>;
-
-	FMediaSoundComponentSpectrumAnalysisTask(Audio::FSpectrumAnalyzer* InAnalyzer, FThreadSafeCounter* InTaskCounter)
-		: Analyzer(InAnalyzer)
-		, TaskCounter(InTaskCounter)
-	{
-		TaskCounter->Increment();
-	}
-
-	void DoWork()
-	{
-		while (Analyzer->PerformAnalysisIfPossible());
-
-		TaskCounter->Decrement();
-	}
-
-	FORCEINLINE TStatId GetStatId() const
-	{
-		RETURN_QUICK_DECLARE_CYCLE_STAT(FTimeSynthSpectrumAnalysisTask, STATGROUP_ThreadPoolAsyncTasks);
-	}
-
-	Audio::FSpectrumAnalyzer* Analyzer;
-	FThreadSafeCounter* TaskCounter;
-};
-
-
 /**
  * Implements a sound component for playing a media player's audio output.
  */
@@ -183,6 +157,18 @@ public:
 	UFUNCTION(BlueprintCallable, Category = "TimeSynth")
 	TArray<FMediaSoundComponentSpectralData> GetSpectralData();
 
+	/** Turns on amplitude envelope following the audio in the media sound component. */
+	UFUNCTION(BlueprintCallable, Category = "Media|MediaSoundComponent")
+	void SetEnableEnvelopeFollowing(bool bInEnvelopeFollowing);
+
+	/** Sets the envelope attack and release times (in ms). */
+	UFUNCTION(BlueprintCallable, Category = "Media|MediaSoundComponent")
+	void SetEnvelopeFollowingsettings(int32 AttackTimeMsec, int32 ReleaseTimeMsec);
+
+	/** Retrieves the current amplitude envelope. */
+	UFUNCTION(BlueprintCallable, Category = "TimeSynth")
+	float GetEnvelopeValue() const;
+
 public:
 
 	void UpdatePlayer();
@@ -285,11 +271,6 @@ private:
 	/** Handle SampleQueue running dry. Ensure audio resumes playback at correct position. */
 	int32 FrameSyncOffset;
 
-	/**
-	 * Sync forward after input audio buffer runs dry due to a hitch or decoder not being able to keep up
-	 * Without this audio will resume playing exactly where it last left off (far behind current player time)
-	 */
-	bool bSyncAudioAfterDropouts;
 
 	/* Time of last sample played. */
 	TAtomic<FTimespan> LastPlaySampleTime;
@@ -302,15 +283,32 @@ private:
 
 	/** Spectrum analyzer used for anlayzing audio in media. */
 	Audio::FSpectrumAnalyzer SpectrumAnalyzer;
-	Audio::SpectrumAnalyzerSettings::FSettings SpectrumAnalyzerSettings;
-	FThreadSafeCounter SpectrumAnalysisCounter;
+	Audio::FSpectrumAnalyzerSettings SpectrumAnalyzerSettings;
+
+	Audio::FEnvelopeFollower EnvelopeFollower;
+	int32 EnvelopeFollowerAttackTime;
+	int32 EnvelopeFollowerReleaseTime;
+	float CurrentEnvelopeValue;
+	FCriticalSection EnvelopeFollowerCriticalSection;
 
 	/** Scratch buffer to mix in source audio to from decoder */
 	Audio::AlignedFloatBuffer AudioScratchBuffer;
 
+	/**
+	 * Sync forward after input audio buffer runs dry due to a hitch or decoder not being able to keep up
+	 * Without this audio will resume playing exactly where it last left off (far behind current player time)
+	 */
+	bool bSyncAudioAfterDropouts;
+
 	/** Whether or not spectral analysis is enabled. */
 	bool bSpectralAnalysisEnabled;
 
+	/** Whether or not envelope following is enabled. */
+	bool bEnvelopeFollowingEnabled;
+
+	/** Whether or not envelope follower settings changed. */
+	bool bEnvelopeFollowerSettingsChanged;
+
 private:
 
 	static USoundClass* DefaultMediaSoundClassObject;
diff --git a/Engine/Source/Runtime/MovieScene/Private/Evaluation/MovieSceneEvaluationTemplateInstance.cpp b/Engine/Source/Runtime/MovieScene/Private/Evaluation/MovieSceneEvaluationTemplateInstance.cpp
--- a/Engine/Source/Runtime/MovieScene/Private/Evaluation/MovieSceneEvaluationTemplateInstance.cpp
+++ b/Engine/Source/Runtime/MovieScene/Private/Evaluation/MovieSceneEvaluationTemplateInstance.cpp
@@ -542,7 +542,7 @@ void FMovieSceneRootEvaluationTemplateInstance::CallSetupTearDown(const FMovieSc
 			{
 				Track->OnBeginEvaluation(PersistentDataProxy, Player);
 			}
-			else
+			else if (Track->HasChildTemplate(Key.SectionIndex))
 			{
 				PersistentDataProxy.SetSectionKey(Key);
 				Track->GetChildTemplate(Key.SectionIndex).OnBeginEvaluation(PersistentDataProxy, Player);
diff --git a/Engine/Source/Runtime/Slate/Private/Widgets/Input/SButton.cpp b/Engine/Source/Runtime/Slate/Private/Widgets/Input/SButton.cpp
--- a/Engine/Source/Runtime/Slate/Private/Widgets/Input/SButton.cpp
+++ b/Engine/Source/Runtime/Slate/Private/Widgets/Input/SButton.cpp
@@ -383,6 +383,12 @@ bool SButton::IsInteractable() const
 	return IsEnabled();
 }
 
+bool SButton::ComputeVolatility() const
+{
+	// Note: we need to be careful with button volatility.  The parent SBorder class always has bound delegates to button but the following are the only thing that actually would be bound that would not be caught by an Invalidate call alone
+	return ContentScale.IsBound() || DesiredSizeScale.IsBound() || BorderBackgroundColor.IsBound() || ContentPadding.IsBound() || ForegroundColor.IsBound();
+}
+
 bool SButton::IsPreciseTapOrClick(const FPointerEvent& MouseEvent) const
 {
 	return ( TouchMethod == EButtonTouchMethod::PreciseTap && MouseEvent.IsTouchEvent() ) ||
diff --git a/Engine/Source/Runtime/Slate/Private/Widgets/Input/SMenuAnchor.cpp b/Engine/Source/Runtime/Slate/Private/Widgets/Input/SMenuAnchor.cpp
--- a/Engine/Source/Runtime/Slate/Private/Widgets/Input/SMenuAnchor.cpp
+++ b/Engine/Source/Runtime/Slate/Private/Widgets/Input/SMenuAnchor.cpp
@@ -174,7 +174,7 @@ bool SMenuAnchor::ComputeVolatility() const
 
 void SMenuAnchor::OnArrangeChildren( const FGeometry& AllottedGeometry, FArrangedChildren& ArrangedChildren ) const
 {
-	ArrangeSingleChild( AllottedGeometry, ArrangedChildren, Children[0], FVector2D::UnitVector );
+	ArrangeSingleChild(GSlateFlowDirection, AllottedGeometry, ArrangedChildren, Children[0], FVector2D::UnitVector);
 	const TSharedPtr<SWindow> PresentingWindow = PopupWindowPtr.Pin();
 	if (IsOpenAndReusingWindow() && PresentingWindow.IsValid())
 	{
diff --git a/Engine/Source/Runtime/Slate/Private/Widgets/LayerManager/STooltipPresenter.cpp b/Engine/Source/Runtime/Slate/Private/Widgets/LayerManager/STooltipPresenter.cpp
--- a/Engine/Source/Runtime/Slate/Private/Widgets/LayerManager/STooltipPresenter.cpp
+++ b/Engine/Source/Runtime/Slate/Private/Widgets/LayerManager/STooltipPresenter.cpp
@@ -40,7 +40,14 @@ void STooltipPresenter::OnArrangeChildren(const FGeometry& AllottedGeometry, FAr
 
 FVector2D STooltipPresenter::ComputeDesiredSize( float ) const
 {
-	return ChildSlot.GetWidget()->GetDesiredSize();
+	if (ChildSlot.GetWidget() != SNullWidget::NullWidget)
+	{
+		return ChildSlot.GetWidget()->GetDesiredSize();
+	}
+	else
+	{
+		return FVector2D(0, 0);
+	}
 }
 
 FChildren* STooltipPresenter::GetChildren()
diff --git a/Engine/Source/Runtime/Slate/Private/Widgets/Layout/SBorder.cpp b/Engine/Source/Runtime/Slate/Private/Widgets/Layout/SBorder.cpp
--- a/Engine/Source/Runtime/Slate/Private/Widgets/Layout/SBorder.cpp
+++ b/Engine/Source/Runtime/Slate/Private/Widgets/Layout/SBorder.cpp
@@ -27,6 +27,8 @@ void SBorder::Construct( const SBorder::FArguments& InArgs )
 
 	ShowDisabledEffect = InArgs._ShowEffectWhenDisabled;
 
+	bFlipForRightToLeftFlowDirection = InArgs._FlipForRightToLeftFlowDirection;
+
 	BorderImage = InArgs._BorderImage;
 	BorderBackgroundColor = InArgs._BorderBackgroundColor;
 	ForegroundColor = InArgs._ForegroundColor;
@@ -89,19 +91,39 @@ int32 SBorder::OnPaint( const FPaintArgs& Args, const FGeometry& AllottedGeometr
 		const bool bShowDisabledEffect = ShowDisabledEffect.Get();
 		const ESlateDrawEffect DrawEffects = (bShowDisabledEffect && !bEnabled) ? ESlateDrawEffect::DisabledEffect : ESlateDrawEffect::None;
 
-		FSlateDrawElement::MakeBox(
-			OutDrawElements,
-			LayerId,
-			AllottedGeometry.ToPaintGeometry(),
-			BrushResource,
-			DrawEffects,
-			BrushResource->GetTint(InWidgetStyle) * InWidgetStyle.GetColorAndOpacityTint() * BorderBackgroundColor.Get().GetColor(InWidgetStyle)
-		);
+		if (bFlipForRightToLeftFlowDirection && GSlateFlowDirection == EFlowDirection::RightToLeft)
+		{
+			const FGeometry FlippedGeometry = AllottedGeometry.MakeChild(FSlateRenderTransform(FScale2D(-1, 1)));
+			FSlateDrawElement::MakeBox(
+				OutDrawElements,
+				LayerId,
+				FlippedGeometry.ToPaintGeometry(),
+				BrushResource,
+				DrawEffects,
+				BrushResource->GetTint(InWidgetStyle) * InWidgetStyle.GetColorAndOpacityTint() * BorderBackgroundColor.Get().GetColor(InWidgetStyle)
+			);
+		}
+		else
+		{
+			FSlateDrawElement::MakeBox(
+				OutDrawElements,
+				LayerId,
+				AllottedGeometry.ToPaintGeometry(),
+				BrushResource,
+				DrawEffects,
+				BrushResource->GetTint(InWidgetStyle) * InWidgetStyle.GetColorAndOpacityTint() * BorderBackgroundColor.Get().GetColor(InWidgetStyle)
+			);
+		}
 	}
 
 	return SCompoundWidget::OnPaint(Args, AllottedGeometry, MyCullingRect, OutDrawElements, LayerId, InWidgetStyle, bEnabled );
 }
 
+bool SBorder::ComputeVolatility() const
+{
+	return BorderImage.IsBound() || BorderBackgroundColor.IsBound() || DesiredSizeScale.IsBound() || ShowDisabledEffect.IsBound();
+}
+
 FVector2D SBorder::ComputeDesiredSize(float LayoutScaleMultiplier) const
 {
 	return DesiredSizeScale.Get() * SCompoundWidget::ComputeDesiredSize(LayoutScaleMultiplier);
diff --git a/Engine/Source/Runtime/Slate/Private/Widgets/Layout/SWidgetSwitcher.cpp b/Engine/Source/Runtime/Slate/Private/Widgets/Layout/SWidgetSwitcher.cpp
--- a/Engine/Source/Runtime/Slate/Private/Widgets/Layout/SWidgetSwitcher.cpp
+++ b/Engine/Source/Runtime/Slate/Private/Widgets/Layout/SWidgetSwitcher.cpp
@@ -140,7 +140,7 @@ void SWidgetSwitcher::OnArrangeChildren( const FGeometry& AllottedGeometry, FArr
 		const FSlot* ActiveSlotPtr = GetActiveSlot();	// Returns null if unsafe attribute WidgetIndex is out-of-bounds
 		if (ActiveSlotPtr != nullptr)
 		{
-			ArrangeSingleChild(AllottedGeometry, ArrangedChildren, *ActiveSlotPtr, ContentScale);
+			ArrangeSingleChild(GSlateFlowDirection, AllottedGeometry, ArrangedChildren, *ActiveSlotPtr, ContentScale);
 		}
 	}
 }
diff --git a/Engine/Source/Runtime/Slate/Private/Widgets/Notifications/SProgressBar.cpp b/Engine/Source/Runtime/Slate/Private/Widgets/Notifications/SProgressBar.cpp
--- a/Engine/Source/Runtime/Slate/Private/Widgets/Notifications/SProgressBar.cpp
+++ b/Engine/Source/Runtime/Slate/Private/Widgets/Notifications/SProgressBar.cpp
@@ -174,8 +174,22 @@ int32 SProgressBar::OnPaint( const FPaintArgs& Args, const FGeometry& AllottedGe
 	
 	if( ProgressFraction.IsSet() )
 	{
+		EProgressBarFillType::Type ComputedBarFillType = BarFillType;
+		if (GSlateFlowDirection == EFlowDirection::RightToLeft)
+		{
+			switch (ComputedBarFillType)
+			{
+			case EProgressBarFillType::LeftToRight:
+				ComputedBarFillType = EProgressBarFillType::RightToLeft;
+				break;
+			case EProgressBarFillType::RightToLeft:
+				ComputedBarFillType = EProgressBarFillType::LeftToRight;
+				break;
+			}
+		}
+
 		const float ClampedFraction = FMath::Clamp(ProgressFraction.GetValue(), 0.0f, 1.0f);
-		switch (BarFillType)
+		switch (ComputedBarFillType)
 		{
 			case EProgressBarFillType::RightToLeft:
 			{
diff --git a/Engine/Source/Runtime/Slate/Private/Widgets/Views/SHeaderRow.cpp b/Engine/Source/Runtime/Slate/Private/Widgets/Views/SHeaderRow.cpp
--- a/Engine/Source/Runtime/Slate/Private/Widgets/Views/SHeaderRow.cpp
+++ b/Engine/Source/Runtime/Slate/Private/Widgets/Views/SHeaderRow.cpp
@@ -48,7 +48,7 @@ public:
 	{
 		check(InArgs._Style);
 
-		SWidget::Construct( InArgs._ToolTipText, InArgs._ToolTip, InArgs._Cursor, InArgs._IsEnabled, InArgs._Visibility, InArgs._RenderOpacity, InArgs._RenderTransform, InArgs._RenderTransformPivot, InArgs._Tag, InArgs._ForceVolatile, InArgs._Clipping, InArgs.MetaData );
+		SWidget::Construct( InArgs._ToolTipText, InArgs._ToolTip, InArgs._Cursor, InArgs._IsEnabled, InArgs._Visibility, InArgs._RenderOpacity, InArgs._RenderTransform, InArgs._RenderTransformPivot, InArgs._Tag, InArgs._ForceVolatile, InArgs._Clipping, InArgs._FlowDirectionPreference, InArgs.MetaData );
 
 		Style = InArgs._Style;
 		ColumnId = Column.ColumnId;
@@ -414,7 +414,7 @@ void SHeaderRow::Construct( const FArguments& InArgs )
 {
 	check(InArgs._Style);
 
-	SWidget::Construct( InArgs._ToolTipText, InArgs._ToolTip, InArgs._Cursor, InArgs._IsEnabled, InArgs._Visibility, InArgs._RenderOpacity, InArgs._RenderTransform, InArgs._RenderTransformPivot, InArgs._Tag, InArgs._ForceVolatile, InArgs._Clipping, InArgs.MetaData );
+	SWidget::Construct( InArgs._ToolTipText, InArgs._ToolTip, InArgs._Cursor, InArgs._IsEnabled, InArgs._Visibility, InArgs._RenderOpacity, InArgs._RenderTransform, InArgs._RenderTransformPivot, InArgs._Tag, InArgs._ForceVolatile, InArgs._Clipping, InArgs._FlowDirectionPreference, InArgs.MetaData );
 
 	ScrollBarThickness = FVector2D::ZeroVector;
 	ScrollBarVisibility = EVisibility::Collapsed;
diff --git a/Engine/Source/Runtime/Slate/Public/Widgets/Input/SButton.h b/Engine/Source/Runtime/Slate/Public/Widgets/Input/SButton.h
--- a/Engine/Source/Runtime/Slate/Public/Widgets/Input/SButton.h
+++ b/Engine/Source/Runtime/Slate/Public/Widgets/Input/SButton.h
@@ -177,6 +177,7 @@ public:
 	virtual void OnMouseLeave( const FPointerEvent& MouseEvent ) override;
 	virtual void OnMouseCaptureLost(const FCaptureLostEvent& CaptureLostEvent) override;
 	virtual bool IsInteractable() const override;
+	virtual bool ComputeVolatility() const override;
 	// SWidget
 
 protected:
diff --git a/Engine/Source/Runtime/Slate/Public/Widgets/Layout/SBorder.h b/Engine/Source/Runtime/Slate/Public/Widgets/Layout/SBorder.h
--- a/Engine/Source/Runtime/Slate/Public/Widgets/Layout/SBorder.h
+++ b/Engine/Source/Runtime/Slate/Public/Widgets/Layout/SBorder.h
@@ -40,6 +40,7 @@ public:
 		, _BorderBackgroundColor( FLinearColor::White )
 		, _ForegroundColor( FSlateColor::UseForeground() )
 		, _ShowEffectWhenDisabled( true )
+		, _FlipForRightToLeftFlowDirection(false)
 		{}
 
 		SLATE_DEFAULT_SLOT( FArguments, Content )
@@ -67,6 +68,9 @@ public:
 		SLATE_ATTRIBUTE( FSlateColor, ForegroundColor )
 		/** Whether or not to show the disabled effect when this border is disabled */
 		SLATE_ATTRIBUTE( bool, ShowEffectWhenDisabled )
+
+		/** Flips the background image if the localization's flow direction is RightToLeft */
+		SLATE_ARGUMENT(bool, FlipForRightToLeftFlowDirection)
 	SLATE_END_ARGS()
 
 	/**
@@ -122,6 +126,7 @@ public:
 public:
 	// SWidget interface
 	virtual int32 OnPaint( const FPaintArgs& Args, const FGeometry& AllottedGeometry, const FSlateRect& MyCullingRect, FSlateWindowElementList& OutDrawElements, int32 LayerId, const FWidgetStyle& InWidgetStyle, bool bParentEnabled ) const override;
+	virtual bool ComputeVolatility() const override;
 	// End of SWidget interface
 
 protected:
@@ -135,4 +140,7 @@ protected:
 	TAttribute<FVector2D> DesiredSizeScale;
 	/** Whether or not to show the disabled effect when this border is disabled */
 	TAttribute<bool> ShowDisabledEffect;
+
+	/** Flips the image if the localization's flow direction is RightToLeft */
+	bool bFlipForRightToLeftFlowDirection;
  };
diff --git a/Engine/Source/Runtime/SlateCore/Private/Layout/FlowDirection.cpp b/Engine/Source/Runtime/SlateCore/Private/Layout/FlowDirection.cpp
new file mode 100644
--- /dev/null
+++ b/Engine/Source/Runtime/SlateCore/Private/Layout/FlowDirection.cpp
@@ -0,0 +1,43 @@
+// Copyright 1998-2019 Epic Games, Inc. All Rights Reserved.
+
+#include "Layout/FlowDirection.h"
+#include "Internationalization/Internationalization.h"
+#include "Internationalization/Culture.h"
+#include "HAL/IConsoleManager.h"
+#include "Misc/FrameValue.h"
+
+EFlowDirection GSlateFlowDirection = EFlowDirection::LeftToRight;
+
+static int32 GSlateEnableLayoutLocalization = 1;
+static FAutoConsoleVariableRef CVarSlateEnableLayoutLocalization(
+	TEXT("Slate.EnableLayoutLocalization"),
+	GSlateEnableLayoutLocalization,
+	TEXT("Controls if we enable or disable localized layout, which affects left to right or right to left detection for cultures."),
+	ECVF_Default
+);
+
+int32 GSlateFlowDirectionShouldFollowCultureByDefault = 0;
+static FAutoConsoleVariableRef CVarSlateFlowDirectionShouldFollowCultureByDefault(
+	TEXT("Slate.ShouldFollowCultureByDefault"),
+	GSlateFlowDirectionShouldFollowCultureByDefault,
+	TEXT("Should we initially follow the culture's flow direction at the window level."),
+	ECVF_Default
+);
+
+EFlowDirection FLayoutLocalization::GetLocalizedLayoutDirection()
+{
+	if (!GSlateEnableLayoutLocalization)
+	{
+		return EFlowDirection::LeftToRight;
+	}
+
+	// Only do this once per frame max no matter how many times people call the function.
+	static TFrameValue<EFlowDirection> FrameCoherentDirection;
+	if (!FrameCoherentDirection.IsSet())
+	{
+		//HACK: Normally we'd get this from the culture, but our cultures don't tell us if they're RightToLeft.  Newer versions of ICU tell you, but we need to upgrade.
+		FrameCoherentDirection = FInternationalization::Get().GetCurrentCulture()->GetTwoLetterISOLanguageName().StartsWith(TEXT("ar")) ? EFlowDirection::RightToLeft : EFlowDirection::LeftToRight;
+	}
+
+	return FrameCoherentDirection.GetValue();
+}
diff --git a/Engine/Source/Runtime/SlateCore/Private/Widgets/Images/SImage.cpp b/Engine/Source/Runtime/SlateCore/Private/Widgets/Images/SImage.cpp
--- a/Engine/Source/Runtime/SlateCore/Private/Widgets/Images/SImage.cpp
+++ b/Engine/Source/Runtime/SlateCore/Private/Widgets/Images/SImage.cpp
@@ -7,8 +7,8 @@ void SImage::Construct( const FArguments& InArgs )
 {
 	Image = InArgs._Image;
 	ColorAndOpacity = InArgs._ColorAndOpacity;
-
-	SetOnMouseButtonDown(InArgs._OnMouseButtonDown);
+	bFlipForRightToLeftFlowDirection = InArgs._FlipForRightToLeftFlowDirection;
+	OnMouseButtonDownHandler = InArgs._OnMouseButtonDown;
 }
 
 int32 SImage::OnPaint( const FPaintArgs& Args, const FGeometry& AllottedGeometry, const FSlateRect& MyCullingRect, FSlateWindowElementList& OutDrawElements, int32 LayerId, const FWidgetStyle& InWidgetStyle, bool bParentEnabled ) const
@@ -22,11 +22,34 @@ int32 SImage::OnPaint( const FPaintArgs& Args, const FGeometry& AllottedGeometry
 
 		const FLinearColor FinalColorAndOpacity( InWidgetStyle.GetColorAndOpacityTint() * ColorAndOpacity.Get().GetColor(InWidgetStyle) * ImageBrush->GetTint( InWidgetStyle ) );
 
-		FSlateDrawElement::MakeBox(OutDrawElements, LayerId, AllottedGeometry.ToPaintGeometry(), ImageBrush, DrawEffects, FinalColorAndOpacity );
+		if (bFlipForRightToLeftFlowDirection && GSlateFlowDirection == EFlowDirection::RightToLeft)
+		{
+			const FGeometry FlippedGeometry = AllottedGeometry.MakeChild(FSlateRenderTransform(FScale2D(-1, 1)));
+			FSlateDrawElement::MakeBox(OutDrawElements, LayerId, FlippedGeometry.ToPaintGeometry(), ImageBrush, DrawEffects, FinalColorAndOpacity);
+		}
+		else
+		{
+			FSlateDrawElement::MakeBox(OutDrawElements, LayerId, AllottedGeometry.ToPaintGeometry(), ImageBrush, DrawEffects, FinalColorAndOpacity);
+		}
 	}
+
 	return LayerId;
 }
 
+FReply SImage::OnMouseButtonDown( const FGeometry& MyGeometry, const FPointerEvent& MouseEvent )
+{
+	if ( OnMouseButtonDownHandler.IsBound() )
+	{
+		// If a handler is assigned, call it.
+		return OnMouseButtonDownHandler.Execute(MyGeometry, MouseEvent);
+	}
+	else
+	{
+		// otherwise the event is unhandled.
+		return FReply::Unhandled();
+	}	
+}
+
 FVector2D SImage::ComputeDesiredSize( float ) const
 {
 	const FSlateBrush* ImageBrush = Image.Get();
@@ -63,3 +86,8 @@ void SImage::SetImage(TAttribute<const FSlateBrush*> InImage)
 		Invalidate(EInvalidateWidget::LayoutAndVolatility);
 	}
 }
+
+void SImage::SetOnMouseButtonDown(FPointerEventHandler EventHandler)
+{
+	OnMouseButtonDownHandler = EventHandler;
+}
diff --git a/Engine/Source/Runtime/SlateCore/Private/Widgets/SBoxPanel.cpp b/Engine/Source/Runtime/SlateCore/Private/Widgets/SBoxPanel.cpp
--- a/Engine/Source/Runtime/SlateCore/Private/Widgets/SBoxPanel.cpp
+++ b/Engine/Source/Runtime/SlateCore/Private/Widgets/SBoxPanel.cpp
@@ -1,4 +1,4 @@
-﻿// Copyright 1998-2019 Epic Games, Inc. All Rights Reserved.
+// Copyright 1998-2018 Epic Games, Inc. All Rights Reserved.
 
 #include "Widgets/SBoxPanel.h"
 #include "Layout/LayoutUtils.h"
@@ -34,7 +34,7 @@ void SVerticalBox::Construct( const SVerticalBox::FArguments& InArgs )
 }
 
 template<EOrientation Orientation>
-static void ArrangeChildrenAlong( const TPanelChildren<SBoxPanel::FSlot>& Children, const FGeometry& AllottedGeometry, FArrangedChildren& ArrangedChildren )
+static void ArrangeChildrenAlong(EFlowDirection InLayoutFlow, const TPanelChildren<SBoxPanel::FSlot>& Children, const FGeometry& AllottedGeometry, FArrangedChildren& ArrangedChildren )
 {
 	// Allotted space will be given to fixed-size children first.
 	// Remaining space will be proportionately divided between stretch children (SizeRule_Stretch)
@@ -74,7 +74,7 @@ static void ArrangeChildrenAlong( const TPanelChildren<SBoxPanel::FSlot>& Childr
 
 					// Clamp to the max size if it was specified
 					float MaxSize = CurChild.MaxSize.Get();
-					FixedTotal += MaxSize > 0 ? FMath::Min( MaxSize, ChildSize ) : ChildSize ;
+					FixedTotal += MaxSize > 0 ? FMath::Min( MaxSize, ChildSize ) : ChildSize;
 				}
 			}
 		}
@@ -93,9 +93,9 @@ static void ArrangeChildrenAlong( const TPanelChildren<SBoxPanel::FSlot>& Childr
 
 		// Now that we have the total fixed-space requirement and the total stretch coefficients we can
 		// arrange widgets top-to-bottom or left-to-right (depending on the orientation).
-		for( int32 ChildIndex=0; ChildIndex < Children.Num(); ++ChildIndex )
+		for (TPanelChildrenConstIterator<SBoxPanel::FSlot> It(Children, Orientation, InLayoutFlow); It; ++It)
 		{
-			const SBoxPanel::FSlot& CurChild = Children[ChildIndex];
+			const SBoxPanel::FSlot& CurChild = *It;
 			const EVisibility ChildVisibility = CurChild.GetWidget()->GetVisibility();
 
 			// Figure out the area allocated to the child in the direction of BoxPanel
@@ -130,14 +130,14 @@ static void ArrangeChildrenAlong( const TPanelChildren<SBoxPanel::FSlot>& Childr
 				}
 			}
 
-			const FMargin SlotPadding(CurChild.SlotPadding.Get());
+			const FMargin SlotPadding(LayoutPaddingWithFlow(InLayoutFlow, CurChild.SlotPadding.Get()));
 
 			FVector2D SlotSize = (Orientation == Orient_Vertical)
 				? FVector2D(AllottedGeometry.GetLocalSize().X, ChildSize + SlotPadding.GetTotalSpaceAlong<Orient_Vertical>())
 				: FVector2D(ChildSize + SlotPadding.GetTotalSpaceAlong<Orient_Horizontal>(), AllottedGeometry.GetLocalSize().Y);
 
 			// Figure out the size and local position of the child within the slot			
-			AlignmentArrangeResult XAlignmentResult = AlignChild<Orient_Horizontal>( SlotSize.X, CurChild, SlotPadding);
+			AlignmentArrangeResult XAlignmentResult = AlignChild<Orient_Horizontal>(InLayoutFlow, SlotSize.X, CurChild, SlotPadding);
 			AlignmentArrangeResult YAlignmentResult = AlignChild<Orient_Vertical>(SlotSize.Y, CurChild, SlotPadding);
 
 			const FVector2D LocalPosition = (Orientation == Orient_Vertical)
@@ -176,11 +176,11 @@ void SBoxPanel::OnArrangeChildren( const FGeometry& AllottedGeometry, FArrangedC
 {
 	if ( this->Orientation == Orient_Horizontal )
 	{
-		ArrangeChildrenAlong<Orient_Horizontal>(this->Children, AllottedGeometry, ArrangedChildren );
+		ArrangeChildrenAlong<Orient_Horizontal>(GSlateFlowDirection, this->Children, AllottedGeometry, ArrangedChildren );
 	}
 	else
 	{
-		ArrangeChildrenAlong<Orient_Vertical>(this->Children, AllottedGeometry, ArrangedChildren );
+		ArrangeChildrenAlong<Orient_Vertical>(GSlateFlowDirection, this->Children, AllottedGeometry, ArrangedChildren );
 	}
 }
 
diff --git a/Engine/Source/Runtime/SlateCore/Private/Widgets/SCompoundWidget.cpp b/Engine/Source/Runtime/SlateCore/Private/Widgets/SCompoundWidget.cpp
--- a/Engine/Source/Runtime/SlateCore/Private/Widgets/SCompoundWidget.cpp
+++ b/Engine/Source/Runtime/SlateCore/Private/Widgets/SCompoundWidget.cpp
@@ -58,7 +58,7 @@ FVector2D SCompoundWidget::ComputeDesiredSize( float ) const
 
 void SCompoundWidget::OnArrangeChildren( const FGeometry& AllottedGeometry, FArrangedChildren& ArrangedChildren ) const
 {
-	ArrangeSingleChild( AllottedGeometry, ArrangedChildren, ChildSlot, ContentScale );
+	ArrangeSingleChild(GSlateFlowDirection, AllottedGeometry, ArrangedChildren, ChildSlot, ContentScale);
 }
 
 FSlateColor SCompoundWidget::GetForegroundColor() const
diff --git a/Engine/Source/Runtime/SlateCore/Private/Widgets/SOverlay.cpp b/Engine/Source/Runtime/SlateCore/Private/Widgets/SOverlay.cpp
--- a/Engine/Source/Runtime/SlateCore/Private/Widgets/SOverlay.cpp
+++ b/Engine/Source/Runtime/SlateCore/Private/Widgets/SOverlay.cpp
@@ -30,8 +30,8 @@ void SOverlay::OnArrangeChildren( const FGeometry& AllottedGeometry, FArrangedCh
 		const EVisibility ChildVisibility = CurChild.GetWidget()->GetVisibility();
 		if ( ArrangedChildren.Accepts(ChildVisibility) )
 		{
-			const FMargin SlotPadding(CurChild.SlotPadding.Get());
-			AlignmentArrangeResult XResult = AlignChild<Orient_Horizontal>(AllottedGeometry.GetLocalSize().X, CurChild, SlotPadding);
+			const FMargin SlotPadding(LayoutPaddingWithFlow(GSlateFlowDirection, CurChild.SlotPadding.Get()));
+			AlignmentArrangeResult XResult = AlignChild<Orient_Horizontal>(GSlateFlowDirection, AllottedGeometry.GetLocalSize().X, CurChild, SlotPadding);
 			AlignmentArrangeResult YResult = AlignChild<Orient_Vertical>(AllottedGeometry.GetLocalSize().Y, CurChild, SlotPadding);
 
 			ArrangedChildren.AddWidget( ChildVisibility, AllottedGeometry.MakeChild(
diff --git a/Engine/Source/Runtime/SlateCore/Private/Widgets/SWidget.cpp b/Engine/Source/Runtime/SlateCore/Private/Widgets/SWidget.cpp
--- a/Engine/Source/Runtime/SlateCore/Private/Widgets/SWidget.cpp
+++ b/Engine/Source/Runtime/SlateCore/Private/Widgets/SWidget.cpp
@@ -112,6 +112,7 @@ SWidget::SWidget()
 	, bNeedsDesiredSize(true)
 	, bUpdatingDesiredSize(false)
 	, Clipping(EWidgetClipping::Inherit)
+	, FlowDirectionPreference(EFlowDirectionPreference::Inherit)
 	, UpdateFlags(EWidgetUpdateFlags::NeedsTick)
 	, DesiredSize()
 	, PrepassLayoutScaleMultiplier(1.0f)
@@ -159,6 +160,7 @@ void SWidget::Construct(
 	const FName& InTag,
 	const bool InForceVolatile,
 	const EWidgetClipping InClipping,
+	const EFlowDirectionPreference InFlowPreference,
 	const TArray<TSharedRef<ISlateMetaData>>& InMetaData
 )
 {
@@ -187,12 +189,13 @@ void SWidget::Construct(
 	Tag = InTag;
 	bForceVolatile = InForceVolatile;
 	Clipping = InClipping;
+	FlowDirectionPreference = InFlowPreference;
 	MetaData = InMetaData;
 }
 
-void SWidget::SWidgetConstruct(const TAttribute<FText>& InToolTipText, const TSharedPtr<IToolTip>& InToolTip, const TAttribute< TOptional<EMouseCursor::Type> >& InCursor, const TAttribute<bool>& InEnabledState, const TAttribute<EVisibility>& InVisibility, const float InRenderOpacity, const TAttribute<TOptional<FSlateRenderTransform>>& InTransform, const TAttribute<FVector2D>& InTransformPivot, const FName& InTag, const bool InForceVolatile, const EWidgetClipping InClipping, const TArray<TSharedRef<ISlateMetaData>>& InMetaData)
+void SWidget::SWidgetConstruct(const TAttribute<FText>& InToolTipText, const TSharedPtr<IToolTip>& InToolTip, const TAttribute< TOptional<EMouseCursor::Type> >& InCursor, const TAttribute<bool>& InEnabledState, const TAttribute<EVisibility>& InVisibility, const float InRenderOpacity, const TAttribute<TOptional<FSlateRenderTransform>>& InTransform, const TAttribute<FVector2D>& InTransformPivot, const FName& InTag, const bool InForceVolatile, const EWidgetClipping InClipping, const EFlowDirectionPreference InFlowPreference, const TArray<TSharedRef<ISlateMetaData>>& InMetaData)
 {
-	Construct(InToolTipText, InToolTip, InCursor, InEnabledState, InVisibility, InRenderOpacity, InTransform, InTransformPivot, InTag, InForceVolatile, InClipping, InMetaData);
+	Construct(InToolTipText, InToolTip, InCursor, InEnabledState, InVisibility, InRenderOpacity, InTransform, InTransformPivot, InTag, InForceVolatile, InClipping, InFlowPreference, InMetaData);
 }
 
 FReply SWidget::OnFocusReceived(const FGeometry& MyGeometry, const FFocusEvent& InFocusEvent)
@@ -1029,6 +1032,11 @@ int32 SWidget::Paint(const FPaintArgs& Args, const FGeometry& AllottedGeometry,
 	FSlateDebugging::BeginWidgetPaint.Broadcast(this, UpdatedArgs, AllottedGeometry, CullingBounds, OutDrawElements, LayerId);
 #endif
 
+	// Establish the flow direction if we're changing from inherit.
+	// FOR RB mode, this should first set GSlateFlowDirection to the incoming state that was cached for the widget, then paint
+	// will override it here to reflow is needed.
+	TGuardValue<EFlowDirection> FlowGuard(GSlateFlowDirection, ComputeFlowDirection());
+	
 	// Paint the geometry of this widget.
 	int32 NewLayerID = OnPaint(UpdatedArgs, AllottedGeometry, CullingBounds, OutDrawElements, LayerId, ContentWidgetStyle, bParentEnabled);
 
diff --git a/Engine/Source/Runtime/SlateCore/Private/Widgets/SWindow.cpp b/Engine/Source/Runtime/SlateCore/Private/Widgets/SWindow.cpp
--- a/Engine/Source/Runtime/SlateCore/Private/Widgets/SWindow.cpp
+++ b/Engine/Source/Runtime/SlateCore/Private/Widgets/SWindow.cpp
@@ -1922,12 +1922,20 @@ SWindow::SWindow()
 int32 SWindow::PaintWindow( const FPaintArgs& Args, const FGeometry& AllottedGeometry, const FSlateRect& MyCullingRect, FSlateWindowElementList& OutDrawElements, int32 LayerId, const FWidgetStyle& InWidgetStyle, bool bParentEnabled ) const
 {
 	LayerId = Paint(Args, AllottedGeometry, MyCullingRect, OutDrawElements, LayerId, InWidgetStyle, bParentEnabled);
-	//LayerId = OutDrawElements.PaintDeferred( LayerId );
+
 	return LayerId;
 }
 
 int32 SWindow::OnPaint(const FPaintArgs& Args, const FGeometry& AllottedGeometry, const FSlateRect& MyCullingRect, FSlateWindowElementList& OutDrawElements, int32 LayerId, const FWidgetStyle& InWidgetStyle, bool bParentEnabled) const
 {
+	EFlowDirection NewFlowDirection = GSlateFlowDirection;
+	if (GetFlowDirectionPreference() == EFlowDirectionPreference::Inherit)
+	{
+		NewFlowDirection = GSlateFlowDirectionShouldFollowCultureByDefault ? FLayoutLocalization::GetLocalizedLayoutDirection() : EFlowDirection::LeftToRight;
+	}
+
+	TGuardValue<EFlowDirection> FlowGuard(GSlateFlowDirection, NewFlowDirection);
+
 	OutDrawElements.BeginDeferredGroup();
 	int32 MaxLayer = SCompoundWidget::OnPaint(Args, AllottedGeometry, MyCullingRect, OutDrawElements, LayerId, InWidgetStyle, bParentEnabled);
 	OutDrawElements.EndDeferredGroup();
diff --git a/Engine/Source/Runtime/SlateCore/Public/Layout/Children.h b/Engine/Source/Runtime/SlateCore/Public/Layout/Children.h
--- a/Engine/Source/Runtime/SlateCore/Public/Layout/Children.h
+++ b/Engine/Source/Runtime/SlateCore/Public/Layout/Children.h
@@ -9,6 +9,7 @@
 #include "Widgets/SNullWidget.h"
 #include "SlotBase.h"
 #include "Widgets/SWidget.h"
+#include "FlowDirection.h"
 
 /**
  * FChildren is an interface that must be implemented by all child containers.
@@ -223,6 +224,8 @@ public:
 };
 
 
+
+
 /**
  * A generic FChildren that stores children along with layout-related information.
  * The type containing Widget* and layout info is specified by ChildType.
@@ -350,6 +353,121 @@ private:
 	bool bEmptying;
 };
 
+
+
+template<typename SlotType>
+class TPanelChildrenConstIterator
+{
+public:
+	TPanelChildrenConstIterator(const TPanelChildren<SlotType>& InContainer, EFlowDirection InLayoutFlow)
+		: Container(InContainer)
+		, LayoutFlow(InLayoutFlow)
+	{
+		Reset();
+	}
+
+	TPanelChildrenConstIterator(const TPanelChildren<SlotType>& InContainer, EOrientation InOrientation, EFlowDirection InLayoutFlow)
+		: Container(InContainer)
+		, LayoutFlow(InOrientation == Orient_Vertical ? EFlowDirection::LeftToRight : InLayoutFlow)
+	{
+		Reset();
+	}
+
+	/** Advances iterator to the next element in the container. */
+	TPanelChildrenConstIterator<SlotType>& operator++()
+	{
+		switch (LayoutFlow)
+		{
+		default:
+		case EFlowDirection::LeftToRight:
+			++Index;
+			break;
+		case EFlowDirection::RightToLeft:
+			--Index;
+			break;
+		}
+
+		return *this;
+	}
+
+	/** Moves iterator to the previous element in the container. */
+	TPanelChildrenConstIterator<SlotType>& operator--()
+	{
+		switch (LayoutFlow)
+		{
+		default:
+		case EFlowDirection::LeftToRight:
+			--Index;
+			break;
+		case EFlowDirection::RightToLeft:
+			++Index;
+			break;
+		}
+
+		return *this;
+	}
+
+	const SlotType& operator* () const
+	{
+		return Container[Index];
+	}
+
+	const SlotType* operator->() const
+	{
+		return &Container[Index];
+	}
+
+	/** conversion to "bool" returning true if the iterator has not reached the last element. */
+	FORCEINLINE explicit operator bool() const
+	{
+		return Container.IsValidIndex(Index);
+	}
+
+	/** Returns an index to the current element. */
+	int32 GetIndex() const
+	{
+		return Index;
+	}
+
+	/** Resets the iterator to the first element. */
+	void Reset()
+	{
+		switch (LayoutFlow)
+		{
+		default:
+		case EFlowDirection::LeftToRight:
+			Index = 0;
+			break;
+		case EFlowDirection::RightToLeft:
+			Index = Container.Num() - 1;
+			break;
+		}
+	}
+
+	/** Sets iterator to the last element. */
+	void SetToEnd()
+	{
+		switch (LayoutFlow)
+		{
+		default:
+		case EFlowDirection::LeftToRight:
+			Index = Container.Num() - 1;
+			break;
+		case EFlowDirection::RightToLeft:
+			Index = 0;
+			break;
+		}
+	}
+
+private:
+
+	const TPanelChildren<SlotType>& Container;
+	int32 Index;
+	EFlowDirection LayoutFlow;
+};
+
+
+
 /**
  * Some advanced widgets contain no layout information, and do not require slots.
  * Those widgets may wish to store a specialized type of child widget.
diff --git a/Engine/Source/Runtime/SlateCore/Public/Layout/FlowDirection.h b/Engine/Source/Runtime/SlateCore/Public/Layout/FlowDirection.h
new file mode 100644
--- /dev/null
+++ b/Engine/Source/Runtime/SlateCore/Public/Layout/FlowDirection.h
@@ -0,0 +1,51 @@
+// Copyright 1998-2019 Epic Games, Inc. All Rights Reserved.
+
+#pragma once
+
+#include "CoreMinimal.h"
+#include "UObject/ObjectMacros.h"
+#include "UObject/Class.h"
+#include "FlowDirection.generated.h"
+
+/**
+ * Widgets may need to flow left or right depending upon the current culture/localization that's active.
+ * This enum is used to request a specific layout flow.
+ */
+enum class EFlowDirection : uint8
+{
+	/** Desires content flows using a LTR layout */
+	LeftToRight,
+	/** Desires content flows using a RTL layout */
+	RightToLeft,
+};
+
+
+/**
+ * 
+ */
+UENUM()
+enum class EFlowDirectionPreference : uint8
+{
+	/** */
+	Inherit,
+	/** */
+	Culture,
+	/** */
+	LeftToRight,
+	/**  */
+	RightToLeft
+};
+
+
+extern SLATECORE_API EFlowDirection GSlateFlowDirection;
+extern SLATECORE_API int32 GSlateFlowDirectionShouldFollowCultureByDefault;
+
+/**
+ * Class containing utilities for getting layout localization information.
+ */
+class SLATECORE_API FLayoutLocalization
+{
+public:
+	/** Gets the current expected flow direction based on localization. */
+	static EFlowDirection GetLocalizedLayoutDirection();
+};
diff --git a/Engine/Source/Runtime/SlateCore/Public/Layout/Geometry.h b/Engine/Source/Runtime/SlateCore/Public/Layout/Geometry.h
--- a/Engine/Source/Runtime/SlateCore/Public/Layout/Geometry.h
+++ b/Engine/Source/Runtime/SlateCore/Public/Layout/Geometry.h
@@ -229,6 +229,11 @@ public:
 		return FGeometry(InLocalSize, LayoutTransform, GetAccumulatedLayoutTransform(), GetAccumulatedRenderTransform(), bHasRenderTransform);
 	}
 
+	FORCEINLINE_DEBUGGABLE FGeometry MakeChild(const FSlateRenderTransform& RenderTransform, const FVector2D& RenderTransformPivot = FVector2D(0.5f, 0.5)) const
+	{
+		return FGeometry(GetLocalSize(), FSlateLayoutTransform(), RenderTransform, RenderTransformPivot, GetAccumulatedLayoutTransform(), GetAccumulatedRenderTransform());
+	}
+
 	/**
 	 * Create a child geometry+widget relative to this one using the given LayoutGeometry.
 	 *
diff --git a/Engine/Source/Runtime/SlateCore/Public/Layout/LayoutUtils.h b/Engine/Source/Runtime/SlateCore/Public/Layout/LayoutUtils.h
--- a/Engine/Source/Runtime/SlateCore/Public/Layout/LayoutUtils.h
+++ b/Engine/Source/Runtime/SlateCore/Public/Layout/LayoutUtils.h
@@ -22,21 +22,36 @@ struct AlignmentArrangeResult
 
 namespace ArrangeUtils
 {
-	/** Gets the alignment of an axis-agnostic int32 so that we can do aligment on an axis without caring about its orientation */
+	/** Gets the alignment of an axis-agnostic int32 so that we can do alignment on an axis without caring about its orientation */
 	template<EOrientation Orientation>
 	struct GetChildAlignment
 	{
 		template<typename SlotType>
-		static int32 AsInt( const SlotType& InSlot );
+		static int32 AsInt(EFlowDirection InFlowDirection, const SlotType& InSlot );
 	};
 
 	template<>
 	struct GetChildAlignment<Orient_Horizontal>
 	{
 		template<typename SlotType>
-		static int32 AsInt( const SlotType& InSlot )
+		static int32 AsInt(EFlowDirection InFlowDirection, const SlotType& InSlot )
 		{
-			return static_cast<int32>(InSlot.HAlignment);
+			switch (InFlowDirection)
+			{
+			default:
+			case EFlowDirection::LeftToRight:
+				return static_cast<int32>(InSlot.HAlignment);
+			case EFlowDirection::RightToLeft:
+				switch (InSlot.HAlignment)
+				{
+				case HAlign_Left:
+					return static_cast<int32>(HAlign_Right);
+				case HAlign_Right:
+					return static_cast<int32>(HAlign_Left);
+				default:
+					return static_cast<int32>(InSlot.HAlignment);
+				}
+			}
 		}
 	};
 
@@ -44,15 +59,16 @@ namespace ArrangeUtils
 	struct GetChildAlignment<Orient_Vertical>
 	{
 		template<typename SlotType>
-		static int32 AsInt( const SlotType& InSlot )
+		static int32 AsInt(EFlowDirection InFlowDirection, const SlotType& InSlot )
 		{
+			// InFlowDirection has no effect in vertical orientations.
 			return static_cast<int32>(InSlot.VAlignment);
 		}
 	};
-
 }
 
 
+
 /**
  * Helper method to BoxPanel::ArrangeChildren.
  * 
@@ -65,21 +81,21 @@ namespace ArrangeUtils
  * @return  Offset and Size of widget
  */
 template<EOrientation Orientation, typename SlotType>
-static AlignmentArrangeResult AlignChild(float AllottedSize, float ChildDesiredSize, const SlotType& ChildToArrange, const FMargin& SlotPadding, const float& ContentScale = 1.0f, bool bClampToParent = true)
+static AlignmentArrangeResult AlignChild(EFlowDirection InLayoutFlow, float AllottedSize, float ChildDesiredSize, const SlotType& ChildToArrange, const FMargin& SlotPadding, const float& ContentScale = 1.0f, bool bClampToParent = true)
 {
 	const FMargin& Margin = SlotPadding;
 	const float TotalMargin = Margin.GetTotalSpaceAlong<Orientation>();
 	const float MarginPre = ( Orientation == Orient_Horizontal ) ? Margin.Left : Margin.Top;
 	const float MarginPost = ( Orientation == Orient_Horizontal ) ? Margin.Right : Margin.Bottom;
 
-	const int32 Alignment = ArrangeUtils::GetChildAlignment<Orientation>::AsInt(ChildToArrange);
+	const int32 Alignment = ArrangeUtils::GetChildAlignment<Orientation>::AsInt(InLayoutFlow, ChildToArrange);
 
-	switch ( Alignment )
+	switch (Alignment)
 	{
 	case HAlign_Fill:
-		return AlignmentArrangeResult(MarginPre, ( AllottedSize - TotalMargin ) * ContentScale);
+		return AlignmentArrangeResult(MarginPre, (AllottedSize - TotalMargin) * ContentScale);
 	}
-
+	
 	const float ChildSize = bClampToParent ? FMath::Min(ChildDesiredSize, AllottedSize - TotalMargin) : ChildDesiredSize;
 
 	switch( Alignment )
@@ -97,19 +113,25 @@ static AlignmentArrangeResult AlignChild(float AllottedSize, float ChildDesiredS
 }
 
 template<EOrientation Orientation, typename SlotType>
-static AlignmentArrangeResult AlignChild(float AllottedSize, const SlotType& ChildToArrange, const FMargin& SlotPadding, const float& ContentScale = 1.0f, bool bClampToParent = true)
+static AlignmentArrangeResult AlignChild(float AllottedSize, float ChildDesiredSize, const SlotType& ChildToArrange, const FMargin& SlotPadding, const float& ContentScale = 1.0f, bool bClampToParent = true)
+{
+	return AlignChild<Orientation, SlotType>(EFlowDirection::LeftToRight, AllottedSize, ChildDesiredSize, ChildToArrange, SlotPadding, ContentScale, bClampToParent);
+}
+
+template<EOrientation Orientation, typename SlotType>
+static AlignmentArrangeResult AlignChild(EFlowDirection InLayoutFlow, float AllottedSize, const SlotType& ChildToArrange, const FMargin& SlotPadding, const float& ContentScale = 1.0f, bool bClampToParent = true)
 {
 	const FMargin& Margin = SlotPadding;
 	const float TotalMargin = Margin.GetTotalSpaceAlong<Orientation>();
 	const float MarginPre = ( Orientation == Orient_Horizontal ) ? Margin.Left : Margin.Top;
 	const float MarginPost = ( Orientation == Orient_Horizontal ) ? Margin.Right : Margin.Bottom;
 
-	const int32 Alignment = ArrangeUtils::GetChildAlignment<Orientation>::AsInt(ChildToArrange);
+	const int32 Alignment = ArrangeUtils::GetChildAlignment<Orientation>::AsInt(InLayoutFlow, ChildToArrange);
 
-	switch ( Alignment )
+	switch (Alignment)
 	{
 	case HAlign_Fill:
-		return AlignmentArrangeResult(MarginPre, ( AllottedSize - TotalMargin ) * ContentScale);
+		return AlignmentArrangeResult(MarginPre, (AllottedSize - TotalMargin) * ContentScale);
 	}
 
 	const float ChildDesiredSize = ( Orientation == Orient_Horizontal )
@@ -132,29 +154,52 @@ static AlignmentArrangeResult AlignChild(float AllottedSize, const SlotType& Chi
 	return AlignmentArrangeResult(MarginPre, ( AllottedSize - TotalMargin ) * ContentScale);
 }
 
+template<EOrientation Orientation, typename SlotType>
+static AlignmentArrangeResult AlignChild(float AllottedSize, const SlotType& ChildToArrange, const FMargin& SlotPadding, const float& ContentScale = 1.0f, bool bClampToParent = true)
+{
+	return AlignChild<Orientation, SlotType>(EFlowDirection::LeftToRight, AllottedSize, ChildToArrange, SlotPadding, ContentScale, bClampToParent);
+}
+
 
 /**
  * Arrange a ChildSlot within the AllottedGeometry and populate ArrangedChildren with the arranged result.
  * The code makes certain assumptions about the type of ChildSlot.
  */
 template<typename SlotType>
-void ArrangeSingleChild( const FGeometry& AllottedGeometry, FArrangedChildren& ArrangedChildren, const SlotType& ChildSlot, const TAttribute<FVector2D>& ContentScale )
+static void ArrangeSingleChild(const FGeometry& AllottedGeometry, FArrangedChildren& ArrangedChildren, const SlotType& ChildSlot, const TAttribute<FVector2D>& ContentScale)
+{
+	ArrangeSingleChild<SlotType>(EFlowDirection::LeftToRight, AllottedGeometry, ArrangedChildren, ChildSlot, ContentScale);
+}
+
+template<typename SlotType>
+static void ArrangeSingleChild(EFlowDirection InFlowDirection, const FGeometry& AllottedGeometry, FArrangedChildren& ArrangedChildren, const SlotType& ChildSlot, const TAttribute<FVector2D>& ContentScale)
 {
 	const EVisibility ChildVisibility = ChildSlot.GetWidget()->GetVisibility();
 	if ( ArrangedChildren.Accepts(ChildVisibility) )
 	{
 		const FVector2D ThisContentScale = ContentScale.Get();
-		const FMargin SlotPadding(ChildSlot.SlotPadding.Get());
-		AlignmentArrangeResult XResult = AlignChild<Orient_Horizontal>(AllottedGeometry.GetLocalSize().X, ChildSlot, SlotPadding, ThisContentScale.X);
-		AlignmentArrangeResult YResult = AlignChild<Orient_Vertical>(AllottedGeometry.GetLocalSize().Y, ChildSlot, SlotPadding, ThisContentScale.Y);
+		const FMargin SlotPadding(LayoutPaddingWithFlow(InFlowDirection, ChildSlot.SlotPadding.Get()));
+		const AlignmentArrangeResult XResult = AlignChild<Orient_Horizontal>(InFlowDirection, AllottedGeometry.GetLocalSize().X, ChildSlot, SlotPadding, ThisContentScale.X);
+		const AlignmentArrangeResult YResult = AlignChild<Orient_Vertical>(AllottedGeometry.GetLocalSize().Y, ChildSlot, SlotPadding, ThisContentScale.Y);
 
 		ArrangedChildren.AddWidget( ChildVisibility, AllottedGeometry.MakeChild(
 				ChildSlot.GetWidget(),
-				FVector2D(XResult.Offset,YResult.Offset),
+				FVector2D(XResult.Offset, YResult.Offset),
 				FVector2D(XResult.Size, YResult.Size)
 		) );
 	}
+}
 
+static FMargin LayoutPaddingWithFlow(EFlowDirection InLayoutFlow, const FMargin& InPadding)
+{
+	FMargin ReturnPadding(InPadding);
+	if (InLayoutFlow == EFlowDirection::RightToLeft)
+	{
+		float Temp = ReturnPadding.Left;
+		ReturnPadding.Left = ReturnPadding.Right;
+		ReturnPadding.Right = Temp;
+	}
+	return ReturnPadding;
 }
 
 /**
diff --git a/Engine/Source/Runtime/SlateCore/Public/Widgets/DeclarativeSyntaxSupport.h b/Engine/Source/Runtime/SlateCore/Public/Widgets/DeclarativeSyntaxSupport.h
--- a/Engine/Source/Runtime/SlateCore/Public/Widgets/DeclarativeSyntaxSupport.h
+++ b/Engine/Source/Runtime/SlateCore/Public/Widgets/DeclarativeSyntaxSupport.h
@@ -6,6 +6,7 @@
 #include "Misc/Attribute.h"
 #include "Layout/Visibility.h"
 #include "Layout/Clipping.h"
+#include "Layout/FlowDirection.h"
 #include "Rendering/SlateRenderTransform.h"
 #include "GenericPlatform/ICursor.h"
 #include "Types/ISlateMetaData.h"
@@ -774,6 +775,7 @@ struct TSlateBaseNamedArgs
 	, _RenderTransformPivot( FVector2D::ZeroVector )
 	, _ForceVolatile( false )
 	, _Clipping( EWidgetClipping::Inherit )
+	, _FlowDirectionPreference( EFlowDirectionPreference::Inherit )
 	{
 	}
 
@@ -817,6 +819,7 @@ struct TSlateBaseNamedArgs
 	SLATE_ARGUMENT( FName, Tag )
 	SLATE_ARGUMENT( bool, ForceVolatile )
 	SLATE_ARGUMENT( EWidgetClipping, Clipping )
+	SLATE_ARGUMENT( EFlowDirectionPreference, FlowDirectionPreference)
 
 	TArray<TSharedRef<ISlateMetaData>> MetaData;
 };
@@ -1089,6 +1092,7 @@ struct TDecl
 			InArgs._Tag,
 			InArgs._ForceVolatile,
 			InArgs._Clipping,
+			InArgs._FlowDirectionPreference,
 			InArgs.MetaData );
 
 		_RequiredArgs.CallConstruct(_Widget, InArgs);
diff --git a/Engine/Source/Runtime/SlateCore/Public/Widgets/Images/SImage.h b/Engine/Source/Runtime/SlateCore/Public/Widgets/Images/SImage.h
--- a/Engine/Source/Runtime/SlateCore/Public/Widgets/Images/SImage.h
+++ b/Engine/Source/Runtime/SlateCore/Public/Widgets/Images/SImage.h
@@ -20,19 +20,24 @@ class SLATECORE_API SImage
 	: public SLeafWidget
 {
 public:
-	SLATE_BEGIN_ARGS(SImage)
-		: _Image(FCoreStyle::Get().GetDefaultBrush())
-		, _ColorAndOpacity(FLinearColor::White)
-	{}
+	SLATE_BEGIN_ARGS( SImage )
+		: _Image( FCoreStyle::Get().GetDefaultBrush() )
+		, _ColorAndOpacity( FLinearColor::White )
+		, _FlipForRightToLeftFlowDirection( false )
+		, _OnMouseButtonDown()
+		{}
 
-	/** Image resource */
-	SLATE_ATTRIBUTE(const FSlateBrush*, Image)
+		/** Image resource */
+		SLATE_ATTRIBUTE( const FSlateBrush*, Image )
 
 		/** Color and opacity */
-		SLATE_ATTRIBUTE(FSlateColor, ColorAndOpacity)
+		SLATE_ATTRIBUTE( FSlateColor, ColorAndOpacity )
+
+		/** Flips the image if the localization's flow direction is RightToLeft */
+		SLATE_ARGUMENT( bool, FlipForRightToLeftFlowDirection )
 
 		/** Invoked when the mouse is pressed in the widget. */
-		SLATE_EVENT(FPointerEventHandler, OnMouseButtonDown)
+		SLATE_EVENT( FPointerEventHandler, OnMouseButtonDown )
 
 	SLATE_END_ARGS()
 
@@ -48,23 +53,27 @@ public:
 	 *
 	 * @param	InArgs	The declaration data for this widget
 	 */
-	void Construct(const FArguments& InArgs);
+	void Construct( const FArguments& InArgs );
 
 public:
 
 	/** See the ColorAndOpacity attribute */
-	void SetColorAndOpacity(const TAttribute<FSlateColor>& InColorAndOpacity);
-
+	void SetColorAndOpacity( const TAttribute<FSlateColor>& InColorAndOpacity );
+	
 	/** See the ColorAndOpacity attribute */
-	void SetColorAndOpacity(FLinearColor InColorAndOpacity);
+	void SetColorAndOpacity( FLinearColor InColorAndOpacity );
 
 	/** See the Image attribute */
 	void SetImage(TAttribute<const FSlateBrush*> InImage);
+	
+	/** See OnMouseButtonDown event */
+	void SetOnMouseButtonDown(FPointerEventHandler EventHandler);
 
 public:
 
 	// SWidget overrides
-	virtual int32 OnPaint(const FPaintArgs& Args, const FGeometry& AllottedGeometry, const FSlateRect& MyCullingRect, FSlateWindowElementList& OutDrawElements, int32 LayerId, const FWidgetStyle& InWidgetStyle, bool bParentEnabled) const override;
+	virtual int32 OnPaint( const FPaintArgs& Args, const FGeometry& AllottedGeometry, const FSlateRect& MyCullingRect, FSlateWindowElementList& OutDrawElements, int32 LayerId, const FWidgetStyle& InWidgetStyle, bool bParentEnabled ) const override;
+	virtual FReply OnMouseButtonDown( const FGeometry& MyGeometry, const FPointerEvent& MouseEvent ) override;
 
 protected:
 	// Begin SWidget overrides.
@@ -78,4 +87,10 @@ protected:
 
 	/** Color and opacity scale for this image */
 	TAttribute<FSlateColor> ColorAndOpacity;
+
+	/** Flips the image if the localization's flow direction is RightToLeft */
+	bool bFlipForRightToLeftFlowDirection;
+
+	/** Invoked when the mouse is pressed in the image */
+	FPointerEventHandler OnMouseButtonDownHandler;
 };
diff --git a/Engine/Source/Runtime/SlateCore/Public/Widgets/SWidget.h b/Engine/Source/Runtime/SlateCore/Public/Widgets/SWidget.h
--- a/Engine/Source/Runtime/SlateCore/Public/Widgets/SWidget.h
+++ b/Engine/Source/Runtime/SlateCore/Public/Widgets/SWidget.h
@@ -14,6 +14,7 @@
 #include "Layout/ArrangedWidget.h"
 #include "Layout/LayoutGeometry.h"
 #include "Layout/Margin.h"
+#include "Layout/FlowDirection.h"
 #include "Rendering/SlateLayoutTransform.h"
 #include "Input/CursorReply.h"
 #include "Input/Reply.h"
@@ -282,14 +283,13 @@ public:
 		const FName& InTag,
 		const bool InForceVolatile,
 		const EWidgetClipping InClipping,
+		const EFlowDirectionPreference InFlowPreference,
 		const TArray<TSharedRef<ISlateMetaData>>& InMetaData);
 
 	void SWidgetConstruct(const TAttribute<FText>& InToolTipText,
 		const TSharedPtr<IToolTip>& InToolTip,
 		const TAttribute< TOptional<EMouseCursor::Type> >& InCursor,
 		const TAttribute<bool>& InEnabledState,
-
-
 		const TAttribute<EVisibility>& InVisibility,
 		const float InRenderOpacity,
 		const TAttribute<TOptional<FSlateRenderTransform>>& InTransform,
@@ -297,6 +297,7 @@ public:
 		const FName& InTag,
 		const bool InForceVolatile,
 		const EWidgetClipping InClipping,
+		const EFlowDirectionPreference InFlowPreference,
 		const TArray<TSharedRef<ISlateMetaData>>& InMetaData);
 
 	//
@@ -1068,6 +1069,25 @@ public:
 		return RenderTransform.Get();
 	}
 
+	FORCEINLINE TOptional<FSlateRenderTransform> GetRenderTransformWithRespectToFlowDirection() const
+	{
+		if (LIKELY(GSlateFlowDirection == EFlowDirection::LeftToRight))
+		{
+			return RenderTransform.Get();
+		}
+		else
+		{
+			// If we're going right to left, flip the X translation on render transforms.
+			TOptional<FSlateRenderTransform> Transform = RenderTransform.Get();
+			if (Transform.IsSet())
+			{
+				FVector2D Translation = Transform.GetValue().GetTranslation();
+				Transform.GetValue().SetTranslation(FVector2D(-Translation.X, Translation.Y));
+			}
+			return Transform;
+		}
+	}
+
 	/** @param InTransform the render transform to set for the widget (transforms from widget's local space). TOptional<> to allow code to skip expensive overhead if there is no render transform applied. */
 	FORCEINLINE void SetRenderTransform(TAttribute<TOptional<FSlateRenderTransform>> InTransform)
 	{
@@ -1079,7 +1099,7 @@ public:
 	}
 
 	/** @return the pivot point of the render transform. */
-	FORCEINLINE const FVector2D& GetRenderTransformPivot() const
+	FORCEINLINE FVector2D GetRenderTransformPivot() const
 	{
 		return RenderTransformPivot.Get();
 	}
@@ -1132,6 +1152,22 @@ public:
 		return CullingBoundsExtension;
 	}
 
+	/**
+	 * Sets how content should flow in this panel, based on the current culture.  By default all panels inherit 
+	 * the state of the widget above.  If they set a new flow direction it will be inherited down the tree.
+	 */
+	void SetFlowDirectionPreference(EFlowDirectionPreference InFlowDirectionPreference)
+	{
+		if (FlowDirectionPreference != InFlowDirectionPreference)
+		{
+			FlowDirectionPreference = InFlowDirectionPreference;
+			Invalidate(EInvalidateWidget::Paint);
+		}
+	}
+
+	/** Gets the desired flow direction for the layout. */
+	EFlowDirectionPreference GetFlowDirectionPreference() const { return FlowDirectionPreference; }
+
 	/**
 	 * Set the tool tip that should appear when this widget is hovered.
 	 *
@@ -1479,6 +1515,26 @@ protected:
 	 */
 	EWidgetClipping Clipping;
 
+protected:
+	/** Establishes a new flow direction potentially, if this widget has a particular preference for it and all its children. */
+	EFlowDirection ComputeFlowDirection() const
+	{
+		switch (FlowDirectionPreference)
+		{
+		case EFlowDirectionPreference::Culture:
+			return FLayoutLocalization::GetLocalizedLayoutDirection();
+		case EFlowDirectionPreference::LeftToRight:
+			return EFlowDirection::LeftToRight;
+		case EFlowDirectionPreference::RightToLeft:
+			return EFlowDirection::RightToLeft;
+		}
+
+		return GSlateFlowDirection;
+	}
+
+private:
+	/** Flow direction preference */
+	EFlowDirectionPreference FlowDirectionPreference;
 
 private:
 	EWidgetUpdateFlags UpdateFlags;
@@ -1573,9 +1629,10 @@ FORCEINLINE_DEBUGGABLE FArrangedWidget FGeometry::MakeChild(const TSharedRef<SWi
 {
 	// If there is no render transform set, use the simpler MakeChild call that doesn't bother concatenating the render transforms.
 	// This saves a significant amount of overhead since every widget does this, and most children don't have a render transform.
-	if ( ChildWidget->GetRenderTransform().IsSet() )
+	TOptional<FSlateRenderTransform> RenderTransform = ChildWidget->GetRenderTransformWithRespectToFlowDirection();
+	if (RenderTransform.IsSet() )
 	{
-		return FArrangedWidget(ChildWidget, MakeChild(InLocalSize, LayoutTransform, ChildWidget->GetRenderTransform().GetValue(), ChildWidget->GetRenderTransformPivot()));
+		return FArrangedWidget(ChildWidget, MakeChild(InLocalSize, LayoutTransform, RenderTransform.GetValue(), ChildWidget->GetRenderTransformPivot()));
 	}
 	else
 	{
diff --git a/Engine/Source/Runtime/UMG/Private/Components/Border.cpp b/Engine/Source/Runtime/UMG/Private/Components/Border.cpp
--- a/Engine/Source/Runtime/UMG/Private/Components/Border.cpp
+++ b/Engine/Source/Runtime/UMG/Private/Components/Border.cpp
@@ -43,7 +43,8 @@ void UBorder::ReleaseSlateResources(bool bReleaseChildren)
 
 TSharedRef<SWidget> UBorder::RebuildWidget()
 {
-	MyBorder = SNew(SBorder);
+	MyBorder = SNew(SBorder)
+		.FlipForRightToLeftFlowDirection(bFlipForRightToLeftFlowDirection);
 	
 	if ( GetChildrenCount() > 0 )
 	{
diff --git a/Engine/Source/Runtime/UMG/Private/Components/Image.cpp b/Engine/Source/Runtime/UMG/Private/Components/Image.cpp
--- a/Engine/Source/Runtime/UMG/Private/Components/Image.cpp
+++ b/Engine/Source/Runtime/UMG/Private/Components/Image.cpp
@@ -44,7 +44,9 @@ void UImage::ReleaseSlateResources(bool bReleaseChildren)
 
 TSharedRef<SWidget> UImage::RebuildWidget()
 {
-	MyImage = SNew(SImage);
+	MyImage = SNew(SImage)
+			.FlipForRightToLeftFlowDirection(bFlipForRightToLeftFlowDirection);
+
 	return MyImage.ToSharedRef();
 }
 
diff --git a/Engine/Source/Runtime/UMG/Private/Components/Widget.cpp b/Engine/Source/Runtime/UMG/Private/Components/Widget.cpp
--- a/Engine/Source/Runtime/UMG/Private/Components/Widget.cpp
+++ b/Engine/Source/Runtime/UMG/Private/Components/Widget.cpp
@@ -1121,6 +1121,8 @@ void UWidget::SynchronizeProperties()
 	SafeWidget->SetClipping(Clipping);
 #endif
 
+	SafeWidget->SetFlowDirectionPreference(FlowDirectionPreference);
+
 	SafeWidget->ForceVolatile(bIsVolatile);
 
 	SafeWidget->SetRenderOpacity(RenderOpacity);
diff --git a/Engine/Source/Runtime/UMG/Public/Components/Border.h b/Engine/Source/Runtime/UMG/Public/Components/Border.h
--- a/Engine/Source/Runtime/UMG/Public/Components/Border.h
+++ b/Engine/Source/Runtime/UMG/Public/Components/Border.h
@@ -1,4 +1,4 @@
-﻿// Copyright 1998-2019 Epic Games, Inc. All Rights Reserved.
+// Copyright 1998-2018 Epic Games, Inc. All Rights Reserved.
 
 #pragma once
 
@@ -81,6 +81,10 @@ public:
 	UPROPERTY(EditAnywhere, BlueprintReadOnly, Category = Appearance)
 	FVector2D DesiredSizeScale;
 
+	/** Flips the background image if the localization's flow direction is RightToLeft */
+	UPROPERTY(EditAnywhere, Category = "Localization")
+	bool bFlipForRightToLeftFlowDirection;
+
 public:
 
 	UPROPERTY(EditAnywhere, Category=Events, meta=( IsBindableEvent="True" ))
diff --git a/Engine/Source/Runtime/UMG/Public/Components/Image.h b/Engine/Source/Runtime/UMG/Public/Components/Image.h
--- a/Engine/Source/Runtime/UMG/Public/Components/Image.h
+++ b/Engine/Source/Runtime/UMG/Public/Components/Image.h
@@ -55,6 +55,10 @@ public:
 	UPROPERTY()
 	FGetLinearColor ColorAndOpacityDelegate;
 
+	/** Flips the image if the localization's flow direction is RightToLeft */
+	UPROPERTY(EditAnywhere, Category = "Localization")
+	bool bFlipForRightToLeftFlowDirection;
+
 public:
 
 	UPROPERTY(EditAnywhere, Category=Events, meta=( IsBindableEvent="True" ))
diff --git a/Engine/Source/Runtime/UMG/Public/Components/Widget.h b/Engine/Source/Runtime/UMG/Public/Components/Widget.h
--- a/Engine/Source/Runtime/UMG/Public/Components/Widget.h
+++ b/Engine/Source/Runtime/UMG/Public/Components/Widget.h
@@ -371,6 +371,10 @@ public:
 	UPROPERTY(Instanced, EditAnywhere, BlueprintReadOnly, Category="Navigation")
 	class UWidgetNavigation* Navigation;
 
+	/** Allows you to set a new flow direction */
+	UPROPERTY(EditAnywhere, Category = "Localization")
+	EFlowDirectionPreference FlowDirectionPreference;
+
 #if WITH_EDITORONLY_DATA
 
 	/** Stores a reference to the asset responsible for this widgets construction. */
