Added "NetErrorLogInterval" config to reduce network connection failure log spam

diff --git a/Engine/Config/BaseEngine.ini b/Engine/Config/BaseEngine.ini
--- a/Engine/Config/BaseEngine.ini
+++ b/Engine/Config/BaseEngine.ini
@@ -211,6 +211,7 @@ InitialButtonRepeatDelay=0.2
 ButtonRepeatDelay=0.1
 NetDriverDefinitions=(DefName="GameNetDriver",DriverClassName="/Script/OnlineSubsystemUtils.IpNetDriver",DriverClassNameFallback="/Script/OnlineSubsystemUtils.IpNetDriver")
 +NetDriverDefinitions=(DefName="DemoNetDriver",DriverClassName="/Script/Engine.DemoNetDriver",DriverClassNameFallback="/Script/Engine.DemoNetDriver")
+NetErrorLogInterval=1.0
 SerializationOutOfBoundsErrorMessage=NSLOCTEXT("","SerializationOutOfBoundsErrorMessage","Corrupt data found, please verify your installation.")
 SerializationOutOfBoundsErrorMessageCaption=NSLOCTEXT("","SerializationOutOfBoundsErrorMessageCaption","Serialization Error : Action Needed")
 bSmoothFrameRate=true
diff --git a/Engine/Source/Runtime/Engine/Classes/Engine/Engine.h b/Engine/Source/Runtime/Engine/Classes/Engine/Engine.h
--- a/Engine/Source/Runtime/Engine/Classes/Engine/Engine.h
+++ b/Engine/Source/Runtime/Engine/Classes/Engine/Engine.h
@@ -2749,6 +2749,10 @@ public:
 	UPROPERTY()
 	TArray<FString> RuntimeServerActors;
 
+	/** Amount of time in seconds between network error logging */
+	UPROPERTY(globalconfig)
+	float NetErrorLogInterval;
+
 	/** Spawns all of the registered server actors */
 	virtual void SpawnServerActors(UWorld *World);
 
diff --git a/Engine/Source/Runtime/Engine/Private/GameViewportClient.cpp b/Engine/Source/Runtime/Engine/Private/GameViewportClient.cpp
--- a/Engine/Source/Runtime/Engine/Private/GameViewportClient.cpp
+++ b/Engine/Source/Runtime/Engine/Private/GameViewportClient.cpp
@@ -1920,7 +1920,12 @@ void UGameViewportClient::PeekTravelFailureMessages(UWorld* InWorld, ETravelFail
 
 void UGameViewportClient::PeekNetworkFailureMessages(UWorld *InWorld, UNetDriver *NetDriver, ENetworkFailure::Type FailureType, const FString& ErrorString)
 {
-	UE_LOG(LogNet, Warning, TEXT("Network Failure: %s[%s]: %s"), NetDriver ? *NetDriver->NetDriverName.ToString() : TEXT("NULL"), ENetworkFailure::ToString(FailureType), *ErrorString);
+	static double LastTimePrinted = 0.0f;
+	if (FPlatformTime::Seconds() - LastTimePrinted > GEngine->NetErrorLogInterval)
+	{
+		UE_LOG(LogNet, Warning, TEXT("Network Failure: %s[%s]: %s"), NetDriver ? *NetDriver->NetDriverName.ToString() : TEXT("NULL"), ENetworkFailure::ToString(FailureType), *ErrorString);
+		LastTimePrinted = FPlatformTime::Seconds();
+	}
 }
 
 void UGameViewportClient::SSSwapControllers()
diff --git a/Engine/Source/Runtime/Engine/Private/NetConnection.cpp b/Engine/Source/Runtime/Engine/Private/NetConnection.cpp
--- a/Engine/Source/Runtime/Engine/Private/NetConnection.cpp
+++ b/Engine/Source/Runtime/Engine/Private/NetConnection.cpp
@@ -2533,7 +2533,13 @@ void UNetConnection::Tick()
 			Driver->Time,
 			Timeout,
 			*Describe());
-		UE_LOG(LogNet, Warning, TEXT("%s"), *Error);
+		
+		static double LastTimePrinted = 0.0f;
+		if (FPlatformTime::Seconds() - LastTimePrinted > GEngine->NetErrorLogInterval)
+		{
+			UE_LOG(LogNet, Warning, TEXT("%s"), *Error);
+			LastTimePrinted = FPlatformTime::Seconds();
+		}
 
 		if (!bPendingDestroy)
 		{
diff --git a/Engine/Source/Runtime/Engine/Private/UnrealEngine.cpp b/Engine/Source/Runtime/Engine/Private/UnrealEngine.cpp
--- a/Engine/Source/Runtime/Engine/Private/UnrealEngine.cpp
+++ b/Engine/Source/Runtime/Engine/Private/UnrealEngine.cpp
@@ -11061,7 +11061,12 @@ void UEngine::HandleTravelFailure(UWorld* InWorld, ETravelFailure::Type FailureT
 
 void UEngine::HandleNetworkFailure(UWorld *World, UNetDriver *NetDriver, ENetworkFailure::Type FailureType, const FString& ErrorString)
 {
-	UE_LOG(LogNet, Log, TEXT("NetworkFailure: %s, Error: '%s'"), ENetworkFailure::ToString(FailureType), *ErrorString);
+	static double LastTimePrinted = 0.0f;
+	if (FPlatformTime::Seconds() - LastTimePrinted > NetErrorLogInterval)
+	{
+		UE_LOG(LogNet, Log, TEXT("NetworkFailure: %s, Error: '%s'"), ENetworkFailure::ToString(FailureType), *ErrorString);
+		LastTimePrinted = FPlatformTime::Seconds();
+	}
 
 	if (!NetDriver)
 	{
