Use move semantics to reduce memcpy() overhead when evaluating animation

Poses & curves are now moved instead of copied where possible.
Also removed needless passing of bone space transforms to UAnimInstance::ParallelEvaluateAnimation, which further removes an extra allocation for sub-instances.

diff --git a/Engine/Source/Editor/AnimGraph/Private/AnimPreviewInstance.cpp b/Engine/Source/Editor/AnimGraph/Private/AnimPreviewInstance.cpp
--- a/Engine/Source/Editor/AnimGraph/Private/AnimPreviewInstance.cpp
+++ b/Engine/Source/Editor/AnimGraph/Private/AnimPreviewInstance.cpp
@@ -223,7 +223,7 @@ bool FAnimPreviewInstanceProxy::Evaluate(FPoseContext& Output)
 					ApplyBoneControllers(BoneControllers, ComponentSpacePoseContext);
 
 					// convert back to local @todo check this
-					ComponentSpacePoseContext.Pose.ConvertToLocalPoses(Output.Pose);
+					FCSPose<FCompactPose>::ConvertComponentPosesToLocalPoses(ComponentSpacePoseContext.Pose, Output.Pose);
 
 					if(bSetKey)
 					{
diff --git a/Engine/Source/Runtime/AnimGraphRuntime/Private/AnimNodes/AnimNode_BlendBoneByChannel.cpp b/Engine/Source/Runtime/AnimGraphRuntime/Private/AnimNodes/AnimNode_BlendBoneByChannel.cpp
--- a/Engine/Source/Runtime/AnimGraphRuntime/Private/AnimNodes/AnimNode_BlendBoneByChannel.cpp
+++ b/Engine/Source/Runtime/AnimGraphRuntime/Private/AnimNodes/AnimNode_BlendBoneByChannel.cpp
@@ -163,7 +163,7 @@ void FAnimNode_BlendBoneByChannel::Evaluate_AnyThread(FPoseContext& Output)
 				TargetPoseCmpntSpace.SafeSetCSBoneTransforms(QueuedModifiedBoneTransforms);
 
 				// Turn Component Space poses back into local space.
-				TargetPoseCmpntSpace.ConvertToLocalPoses(Output.Pose);
+				FCSPose<FCompactPose>::ConvertComponentPosesToLocalPoses(MoveTemp(TargetPoseCmpntSpace), Output.Pose);
 			}
 		}
 	}
diff --git a/Engine/Source/Runtime/AnimGraphRuntime/Private/AnimNodes/AnimNode_BlendListBase.cpp b/Engine/Source/Runtime/AnimGraphRuntime/Private/AnimNodes/AnimNode_BlendListBase.cpp
--- a/Engine/Source/Runtime/AnimGraphRuntime/Private/AnimNodes/AnimNode_BlendListBase.cpp
+++ b/Engine/Source/Runtime/AnimGraphRuntime/Private/AnimNodes/AnimNode_BlendListBase.cpp
@@ -225,8 +225,8 @@ void FAnimNode_BlendListBase::Evaluate_AnyThread(FPoseContext& Output)
 			FPoseLink& CurrentPose = BlendPose[PoseIndex];
 			CurrentPose.Evaluate(EvaluateContext);
 
-			FilteredPoses[i].CopyBonesFrom(EvaluateContext.Pose);
-			FilteredCurve[i] = EvaluateContext.Curve;
+			FilteredPoses[i].MoveBonesFrom(EvaluateContext.Pose);
+			FilteredCurve[i].MoveFrom(EvaluateContext.Curve);
 		}
 
 		// Use the calculated blend sample data if we're blending per-bone
diff --git a/Engine/Source/Runtime/AnimGraphRuntime/Private/AnimNodes/AnimNode_LayeredBoneBlend.cpp b/Engine/Source/Runtime/AnimGraphRuntime/Private/AnimNodes/AnimNode_LayeredBoneBlend.cpp
--- a/Engine/Source/Runtime/AnimGraphRuntime/Private/AnimNodes/AnimNode_LayeredBoneBlend.cpp
+++ b/Engine/Source/Runtime/AnimGraphRuntime/Private/AnimNodes/AnimNode_LayeredBoneBlend.cpp
@@ -230,8 +230,8 @@ void FAnimNode_LayeredBoneBlend::Evaluate_AnyThread(FPoseContext& Output)
 				FPoseContext CurrentPoseContext(Output);
 				BlendPoses[ChildIndex].Evaluate(CurrentPoseContext);
 
-				TargetBlendPoses[ChildIndex].CopyBonesFrom(CurrentPoseContext.Pose);
-				TargetBlendCurves[ChildIndex].CopyFrom(CurrentPoseContext.Curve);
+				TargetBlendPoses[ChildIndex].MoveBonesFrom(CurrentPoseContext.Pose);
+				TargetBlendCurves[ChildIndex].MoveFrom(CurrentPoseContext.Curve);
 			}
 			else
 			{
diff --git a/Engine/Source/Runtime/AnimGraphRuntime/Private/AnimNodes/AnimNode_MultiWayBlend.cpp b/Engine/Source/Runtime/AnimGraphRuntime/Private/AnimNodes/AnimNode_MultiWayBlend.cpp
--- a/Engine/Source/Runtime/AnimGraphRuntime/Private/AnimNodes/AnimNode_MultiWayBlend.cpp
+++ b/Engine/Source/Runtime/AnimGraphRuntime/Private/AnimNodes/AnimNode_MultiWayBlend.cpp
@@ -122,8 +122,12 @@ void FAnimNode_MultiWayBlend::Evaluate_AnyThread(FPoseContext& Output)
 				// total alpha shouldn't be zero
 				Poses[PoseIndex].Evaluate(PoseContext);
 
-				SourcePoses.Add(PoseContext.Pose);
-				SourceCurves.Add(PoseContext.Curve);
+				FCompactPose& SourcePose = SourcePoses.AddDefaulted_GetRef();
+				SourcePose.MoveBonesFrom(PoseContext.Pose);
+
+				FBlendedCurve& SourceCurve = SourceCurves.AddDefaulted_GetRef();
+				SourceCurve.MoveFrom(PoseContext.Curve);
+
 				SourceWeights.Add(CurrentAlpha);
 			}
 		}
diff --git a/Engine/Source/Runtime/AnimGraphRuntime/Private/BoneControllers/AnimNode_ApplyLimits.cpp b/Engine/Source/Runtime/AnimGraphRuntime/Private/BoneControllers/AnimNode_ApplyLimits.cpp
--- a/Engine/Source/Runtime/AnimGraphRuntime/Private/BoneControllers/AnimNode_ApplyLimits.cpp
+++ b/Engine/Source/Runtime/AnimGraphRuntime/Private/BoneControllers/AnimNode_ApplyLimits.cpp
@@ -32,8 +32,8 @@ void FAnimNode_ApplyLimits::EvaluateSkeletalControl_AnyThread(FComponentSpacePos
 
 	FPoseContext LocalPose0(Output.AnimInstanceProxy);
 	FPoseContext LocalPose1(Output.AnimInstanceProxy);
-	Output.Pose.ConvertToLocalPoses(LocalPose0.Pose);
-	Output.Pose.ConvertToLocalPoses(LocalPose1.Pose);
+	FCSPose<FCompactPose>::ConvertComponentPosesToLocalPoses(Output.Pose, LocalPose0.Pose);
+	FCSPose<FCompactPose>::ConvertComponentPosesToLocalPoses(Output.Pose, LocalPose1.Pose);
 	LocalPose0.Curve = Output.Curve;
 	LocalPose1.Curve = Output.Curve;
 
diff --git a/Engine/Source/Runtime/Engine/Classes/Animation/AnimInstance.h b/Engine/Source/Runtime/Engine/Classes/Animation/AnimInstance.h
--- a/Engine/Source/Runtime/Engine/Classes/Animation/AnimInstance.h
+++ b/Engine/Source/Runtime/Engine/Classes/Animation/AnimInstance.h
@@ -184,6 +184,17 @@ struct FSlotEvaluationPose
 		, Weight(InWeight)
 	{
 	}
+
+	FSlotEvaluationPose(FSlotEvaluationPose&& InEvaluationPose)
+		: AdditiveType(InEvaluationPose.AdditiveType)
+		, Weight(InEvaluationPose.Weight)
+	{
+		Pose.MoveBonesFrom(InEvaluationPose.Pose);
+		Curve.MoveFrom(InEvaluationPose.Curve);
+	}
+
+	FSlotEvaluationPose(const FSlotEvaluationPose& InEvaluationPose) = default;
+	FSlotEvaluationPose& operator=(const FSlotEvaluationPose& InEvaluationPose) = default;
 };
 
 /** Helper struct to store a Queued Montage BlendingOut event. */
@@ -967,6 +978,9 @@ public:
 	bool ParallelCanEvaluate(const USkeletalMesh* InSkeletalMesh) const;
 
 	/** Perform evaluation. Can be called from worker threads. */
+	void ParallelEvaluateAnimation(bool bForceRefPose, const USkeletalMesh* InSkeletalMesh, FBlendedHeapCurve& OutCurve, FCompactPose& OutPose);
+
+	UE_DEPRECATED(4.32, "Please use ParallelEvaluateAnimation without passing OutBoneSpaceTransforms.")
 	void ParallelEvaluateAnimation(bool bForceRefPose, const USkeletalMesh* InSkeletalMesh, TArray<FTransform>& OutBoneSpaceTransforms, FBlendedHeapCurve& OutCurve, FCompactPose& OutPose);
 
 	void PostEvaluateAnimation();
diff --git a/Engine/Source/Runtime/Engine/Classes/Animation/AnimNode_SaveCachedPose.h b/Engine/Source/Runtime/Engine/Classes/Animation/AnimNode_SaveCachedPose.h
--- a/Engine/Source/Runtime/Engine/Classes/Animation/AnimNode_SaveCachedPose.h
+++ b/Engine/Source/Runtime/Engine/Classes/Animation/AnimNode_SaveCachedPose.h
@@ -25,8 +25,8 @@ struct ENGINE_API FAnimNode_SaveCachedPose : public FAnimNode_Base
 	float GlobalWeight;
 
 protected:
-	FCompactHeapPose CachedPose;
-	FBlendedHeapCurve CachedCurve;
+	FCompactPose CachedPose;
+	FBlendedCurve CachedCurve;
 
 	TArray<FAnimationUpdateContext> CachedUpdateContexts;
 
diff --git a/Engine/Source/Runtime/Engine/Classes/Animation/AnimNode_SubInstance.h b/Engine/Source/Runtime/Engine/Classes/Animation/AnimNode_SubInstance.h
--- a/Engine/Source/Runtime/Engine/Classes/Animation/AnimNode_SubInstance.h
+++ b/Engine/Source/Runtime/Engine/Classes/Animation/AnimNode_SubInstance.h
@@ -57,8 +57,6 @@ public:
 	TArray<UProperty*> SubInstanceProperties;
 
 	// Temporary storage for the output of the subinstance, will be copied into output pose.
-	// Declared at member level to avoid reallocating the objects
-	TArray<FTransform> BoneTransforms;
 	FBlendedHeapCurve BlendedCurve;
 
 	// FAnimNode_Base interface
@@ -67,8 +65,6 @@ public:
 	virtual void Update_AnyThread(const FAnimationUpdateContext& Context) override;
 	virtual void Evaluate_AnyThread(FPoseContext& Output) override;
 	virtual void GatherDebugData(FNodeDebugData& DebugData) override;
-	virtual bool HasPreUpdate() const override;
-	virtual void PreUpdate(const UAnimInstance* InAnimInstance) override;
 
 protected:
 	virtual void OnInitializeAnimInstance(const FAnimInstanceProxy* InProxy, const UAnimInstance* InAnimInstance) override;
@@ -76,7 +72,4 @@ protected:
 
 	// Shutdown the currently running instance
 	void TeardownInstance();
-
-	// Makes sure the bone transforms array can contain the pose information from the provided anim instance
-	void AllocateBoneTransforms(const UAnimInstance* InAnimInstance);
 };
diff --git a/Engine/Source/Runtime/Engine/Classes/Components/SkeletalMeshComponent.h b/Engine/Source/Runtime/Engine/Classes/Components/SkeletalMeshComponent.h
--- a/Engine/Source/Runtime/Engine/Classes/Components/SkeletalMeshComponent.h
+++ b/Engine/Source/Runtime/Engine/Classes/Components/SkeletalMeshComponent.h
@@ -1875,7 +1875,7 @@ private:
 	void EndPhysicsTickComponent(FSkeletalMeshComponentEndPhysicsTickFunction& ThisTickFunction);
 
 	/** Evaluate Anim System **/
-	void EvaluateAnimation(const USkeletalMesh* InSkeletalMesh, UAnimInstance* InAnimInstance, TArray<FTransform>& OutBoneSpaceTransforms, FVector& OutRootBoneTranslation, FBlendedHeapCurve& OutCurve, FCompactPose& OutPose) const;
+	void EvaluateAnimation(const USkeletalMesh* InSkeletalMesh, UAnimInstance* InAnimInstance, FVector& OutRootBoneTranslation, FBlendedHeapCurve& OutCurve, FCompactPose& OutPose) const;
 
 	/** Queues up tasks for parallel update/evaluation, as well as the chained game thread completion task */
 	void DispatchParallelEvaluationTasks(FActorComponentTickFunction* TickFunction);
diff --git a/Engine/Source/Runtime/Engine/Private/Animation/AnimInstance.cpp b/Engine/Source/Runtime/Engine/Private/Animation/AnimInstance.cpp
--- a/Engine/Source/Runtime/Engine/Private/Animation/AnimInstance.cpp
+++ b/Engine/Source/Runtime/Engine/Private/Animation/AnimInstance.cpp
@@ -600,10 +600,14 @@ bool UAnimInstance::ParallelCanEvaluate(const USkeletalMesh* InSkeletalMesh) con
 }
 
 void UAnimInstance::ParallelEvaluateAnimation(bool bForceRefPose, const USkeletalMesh* InSkeletalMesh, TArray<FTransform>& OutBoneSpaceTransforms, FBlendedHeapCurve& OutCurve, FCompactPose& OutPose)
+{
+	ParallelEvaluateAnimation(bForceRefPose, InSkeletalMesh, OutCurve, OutPose);
+}
+
+void UAnimInstance::ParallelEvaluateAnimation(bool bForceRefPose, const USkeletalMesh* InSkeletalMesh, FBlendedHeapCurve& OutCurve, FCompactPose& OutPose)
 {
 	FAnimInstanceProxy& Proxy = GetProxyOnAnyThread<FAnimInstanceProxy>();
 	OutPose.SetBoneContainer(&Proxy.GetRequiredBones());
-	OutPose.ResetToRefPose();
 
 	FMemMark Mark(FMemStack::Get());
 
diff --git a/Engine/Source/Runtime/Engine/Private/Animation/AnimInstanceProxy.cpp b/Engine/Source/Runtime/Engine/Private/Animation/AnimInstanceProxy.cpp
--- a/Engine/Source/Runtime/Engine/Private/Animation/AnimInstanceProxy.cpp
+++ b/Engine/Source/Runtime/Engine/Private/Animation/AnimInstanceProxy.cpp
@@ -1227,11 +1227,11 @@ void FAnimInstanceProxy::SlotEvaluatePose(const FName& SlotNodeName, const FComp
 #endif // DEBUG_MONTAGEINSTANCE_WEIGHT
 			if (AdditiveAnimType == AAT_None)
 			{
-				NonAdditivePoses.Add(NewPose);
+				NonAdditivePoses.Add(FSlotEvaluationPose(MoveTemp(NewPose)));
 			}
 			else
 			{
-				AdditivePoses.Add(NewPose);
+				AdditivePoses.Add(FSlotEvaluationPose(MoveTemp(NewPose)));
 			}
 		}
 	}
@@ -2175,7 +2175,7 @@ void FAnimInstanceProxy::RegisterWatchedPose(const FCSPose<FCompactPose>& Pose,
 			if (PoseWatch.NodeID == LinkID)
 			{
 				FCompactPose TempPose;
-				Pose.ConvertToLocalPoses(TempPose);
+				FCSPose<FCompactPose>::ConvertComponentPosesToLocalPoses(Pose, TempPose);
 				PoseWatch.PoseInfo->CopyBonesFrom(TempPose);
 				break;
 			}
diff --git a/Engine/Source/Runtime/Engine/Private/Animation/AnimNodeSpaceConversions.cpp b/Engine/Source/Runtime/Engine/Private/Animation/AnimNodeSpaceConversions.cpp
--- a/Engine/Source/Runtime/Engine/Private/Animation/AnimNodeSpaceConversions.cpp
+++ b/Engine/Source/Runtime/Engine/Private/Animation/AnimNodeSpaceConversions.cpp
@@ -31,8 +31,8 @@ void FAnimNode_ConvertComponentToLocalSpace::Evaluate_AnyThread(FPoseContext & O
 	ComponentPose.EvaluateComponentSpace(InputCSPose);
 
 	checkSlow( InputCSPose.Pose.GetPose().IsValid() );
-	InputCSPose.Pose.ConvertToLocalPoses(Output.Pose);
-	Output.Curve = InputCSPose.Curve;
+	FCSPose<FCompactPose>::ConvertComponentPosesToLocalPoses(MoveTemp(InputCSPose.Pose), Output.Pose);
+	Output.Curve = MoveTemp(InputCSPose.Curve);
 }
 
 void FAnimNode_ConvertComponentToLocalSpace::GatherDebugData(FNodeDebugData& DebugData)
@@ -77,6 +77,6 @@ void FAnimNode_ConvertLocalToComponentSpace::EvaluateComponentSpace_AnyThread(FC
 	FPoseContext InputPose(OutputCSPose.AnimInstanceProxy);
 	LocalPose.Evaluate(InputPose);
 
-	OutputCSPose.Pose.InitPose(InputPose.Pose);
-	OutputCSPose.Curve = InputPose.Curve;
+	OutputCSPose.Pose.InitPose(MoveTemp(InputPose.Pose));
+	OutputCSPose.Curve = MoveTemp(InputPose.Curve);
 }
diff --git a/Engine/Source/Runtime/Engine/Private/Animation/AnimNode_SaveCachedPose.cpp b/Engine/Source/Runtime/Engine/Private/Animation/AnimNode_SaveCachedPose.cpp
--- a/Engine/Source/Runtime/Engine/Private/Animation/AnimNode_SaveCachedPose.cpp
+++ b/Engine/Source/Runtime/Engine/Private/Animation/AnimNode_SaveCachedPose.cpp
@@ -52,8 +52,8 @@ void FAnimNode_SaveCachedPose::Evaluate_AnyThread(FPoseContext& Output)
 
 		FPoseContext CachingContext(Output);
 		Pose.Evaluate(CachingContext);
-		CachedPose.CopyBonesFrom(CachingContext.Pose);
-		CachedCurve.CopyFrom(CachingContext.Curve);
+		CachedPose.MoveBonesFrom(CachingContext.Pose);
+		CachedCurve.MoveFrom(CachingContext.Curve);
 	}
 
 	// Return the cached result
diff --git a/Engine/Source/Runtime/Engine/Private/Animation/AnimNode_SubInstance.cpp b/Engine/Source/Runtime/Engine/Private/Animation/AnimNode_SubInstance.cpp
--- a/Engine/Source/Runtime/Engine/Private/Animation/AnimNode_SubInstance.cpp
+++ b/Engine/Source/Runtime/Engine/Private/Animation/AnimNode_SubInstance.cpp
@@ -77,7 +77,7 @@ void FAnimNode_SubInstance::Evaluate_AnyThread(FPoseContext& Output)
 			InputNode->InputCurve.CopyFrom(Output.Curve);
 		}
 
-		InstanceToRun->ParallelEvaluateAnimation(false, nullptr, BoneTransforms, BlendedCurve, Output.Pose);
+		InstanceToRun->ParallelEvaluateAnimation(false, nullptr, BlendedCurve, Output.Pose);
 
 		Output.Curve.CopyFrom(BlendedCurve);
 	}
@@ -107,28 +107,6 @@ void FAnimNode_SubInstance::GatherDebugData(FNodeDebugData& DebugData)
 	InPose.GatherDebugData(DebugData.BranchFlow(1.0f));
 }
 
-bool FAnimNode_SubInstance::HasPreUpdate() const
-{
-	return true;
-}
-
-void FAnimNode_SubInstance::PreUpdate(const UAnimInstance* InAnimInstance)
-{
-	AllocateBoneTransforms(InAnimInstance);
-}
-
-void FAnimNode_SubInstance::AllocateBoneTransforms(const UAnimInstance* InAnimInstance)
-{
-	if(USkeletalMeshComponent* SkelComp = InAnimInstance->GetSkelMeshComponent())
-	{
-		USkeletalMesh* SkelMesh = SkelComp->SkeletalMesh;
-
-		const int32 NumTransforms = SkelComp->GetComponentSpaceTransforms().Num();
-		BoneTransforms.Empty(NumTransforms);
-		BoneTransforms.AddZeroed(NumTransforms);
-	}
-}
-
 void FAnimNode_SubInstance::OnInitializeAnimInstance(const FAnimInstanceProxy* InProxy, const UAnimInstance* InAnimInstance)
 {
 	if(*InstanceClass)
@@ -149,9 +127,6 @@ void FAnimNode_SubInstance::OnInitializeAnimInstance(const FAnimInstanceProxy* I
 		// dont need to verify it is unique here.
 		InstanceToRun = NewObject<UAnimInstance>(MeshComp, InstanceClass, Tag);
 
-		// Set up bone transform array
-		AllocateBoneTransforms(InstanceToRun);
-
 		// Initialize the new instance
 		InstanceToRun->InitializeAnimation();
 
diff --git a/Engine/Source/Runtime/Engine/Private/Animation/AnimSequence.cpp b/Engine/Source/Runtime/Engine/Private/Animation/AnimSequence.cpp
--- a/Engine/Source/Runtime/Engine/Private/Animation/AnimSequence.cpp
+++ b/Engine/Source/Runtime/Engine/Private/Animation/AnimSequence.cpp
@@ -1674,7 +1674,7 @@ void UAnimSequence::GetBonePose(FCompactPose& OutPose, FBlendedCurve& OutCurve,
 		{
 			// get the remaining bone atoms
 			FTransformArray LocalBones;
-			OutPose.CopyBonesTo(LocalBones);
+			OutPose.MoveBonesTo(LocalBones);
 
 			AnimationFormat_GetAnimationPose(
 				LocalBones,
@@ -1683,7 +1683,7 @@ void UAnimSequence::GetBonePose(FCompactPose& OutPose, FBlendedCurve& OutCurve,
 				RotationScalePairs,
 				EvalDecompContext);
 
-			OutPose.CopyBonesFrom(LocalBones);
+			OutPose.MoveBonesFrom(MoveTemp(LocalBones));
 		}
 	}
 
diff --git a/Engine/Source/Runtime/Engine/Private/Components/SkeletalMeshComponent.cpp b/Engine/Source/Runtime/Engine/Private/Components/SkeletalMeshComponent.cpp
--- a/Engine/Source/Runtime/Engine/Private/Components/SkeletalMeshComponent.cpp
+++ b/Engine/Source/Runtime/Engine/Private/Components/SkeletalMeshComponent.cpp
@@ -1692,7 +1692,7 @@ bool USkeletalMeshComponent::AreRequiredCurvesUpToDate() const
 	return (!SkeletalMesh || !SkeletalMesh->Skeleton || CachedAnimCurveUidVersion == SkeletalMesh->Skeleton->GetAnimCurveUidVersion());
 }
 
-void USkeletalMeshComponent::EvaluateAnimation(const USkeletalMesh* InSkeletalMesh, UAnimInstance* InAnimInstance, TArray<FTransform>& OutBoneSpaceTransforms, FVector& OutRootBoneTranslation, FBlendedHeapCurve& OutCurve, FCompactPose& OutPose) const
+void USkeletalMeshComponent::EvaluateAnimation(const USkeletalMesh* InSkeletalMesh, UAnimInstance* InAnimInstance, FVector& OutRootBoneTranslation, FBlendedHeapCurve& OutCurve, FCompactPose& OutPose) const
 {
 	ANIM_MT_SCOPE_CYCLE_COUNTER(SkeletalComponentAnimEvaluate, !IsInGameThread());
 
@@ -1706,7 +1706,7 @@ void USkeletalMeshComponent::EvaluateAnimation(const USkeletalMesh* InSkeletalMe
 		InAnimInstance &&
 		InAnimInstance->ParallelCanEvaluate(InSkeletalMesh))
 	{
-		InAnimInstance->ParallelEvaluateAnimation(bForceRefpose, InSkeletalMesh, OutBoneSpaceTransforms, OutCurve, OutPose);
+		InAnimInstance->ParallelEvaluateAnimation(bForceRefpose, InSkeletalMesh, OutCurve, OutPose);
 	}
 	else
 	{
@@ -1793,7 +1793,7 @@ void USkeletalMeshComponent::PerformAnimationProcessing(const USkeletalMesh* InS
 		FCompactPose EvaluatedPose;
 
 		// evaluate pure animations, and fill up BoneSpaceTransforms
-		EvaluateAnimation(InSkeletalMesh, InAnimInstance, OutBoneSpaceTransforms, OutRootBoneTranslation, OutCurve, EvaluatedPose);
+		EvaluateAnimation(InSkeletalMesh, InAnimInstance, OutRootBoneTranslation, OutCurve, EvaluatedPose);
 		EvaluatePostProcessMeshInstance(OutBoneSpaceTransforms, EvaluatedPose, OutCurve, InSkeletalMesh, OutRootBoneTranslation);
 
 		// Finalize the transforms from the evaluation
@@ -1825,7 +1825,7 @@ void USkeletalMeshComponent::EvaluatePostProcessMeshInstance(TArray<FTransform>&
 			}
 		}
 
-		EvaluateAnimation(InSkeletalMesh, PostProcessAnimInstance, OutBoneSpaceTransforms, OutRootBoneTranslation, OutCurve, InOutPose);
+		EvaluateAnimation(InSkeletalMesh, PostProcessAnimInstance, OutRootBoneTranslation, OutCurve, InOutPose);
 	}
 }
 
diff --git a/Engine/Source/Runtime/Engine/Public/Animation/AnimCurveTypes.h b/Engine/Source/Runtime/Engine/Public/Animation/AnimCurveTypes.h
--- a/Engine/Source/Runtime/Engine/Public/Animation/AnimCurveTypes.h
+++ b/Engine/Source/Runtime/Engine/Public/Animation/AnimCurveTypes.h
@@ -519,13 +519,39 @@ struct FBaseBlendedCurve
 	}
 
 	/**
-	 * Override with inupt curve 
+	 * Override with input curve 
 	 */
 	void Override(const FBaseBlendedCurve& CurveToOverrideFrom)
 	{
-		InitFrom(CurveToOverrideFrom);
-		Elements.Reset();
-		Elements.Append(CurveToOverrideFrom.Elements);
+		// make sure this doesn't happen
+		if (ensure(&CurveToOverrideFrom != this))
+		{
+			check(CurveToOverrideFrom.UIDToArrayIndexLUT != nullptr);
+			UIDToArrayIndexLUT = CurveToOverrideFrom.UIDToArrayIndexLUT;
+			NumValidCurveCount = GetValidElementCount(UIDToArrayIndexLUT);
+			Elements.Reset();
+			Elements.Append(CurveToOverrideFrom.Elements);
+			bInitialized = true;
+		}
+	}
+
+	/**
+	 * Override with input curve, leaving input curve invalid
+	 */
+	void OverrideMove(FBaseBlendedCurve& CurveToOverrideFrom)
+	{
+		// make sure this doesn't happen
+		if (ensure(&CurveToOverrideFrom != this))
+		{
+			check(CurveToOverrideFrom.UIDToArrayIndexLUT != nullptr);
+			UIDToArrayIndexLUT = CurveToOverrideFrom.UIDToArrayIndexLUT;
+			CurveToOverrideFrom.UIDToArrayIndexLUT = nullptr;
+			NumValidCurveCount = GetValidElementCount(UIDToArrayIndexLUT);
+			CurveToOverrideFrom.NumValidCurveCount = 0;
+			Elements = MoveTemp(CurveToOverrideFrom.Elements);
+			bInitialized = true;
+			CurveToOverrideFrom.bInitialized = false;
+		}
 	}
 
 	/** Return number of elements */
@@ -555,6 +581,19 @@ struct FBaseBlendedCurve
 			bInitialized = true;
 		}
 	}
+
+	/** Once moved, source is invalid */
+	void MoveFrom(FBaseBlendedCurve<Allocator>& CurveToMoveFrom)
+	{
+		UIDToArrayIndexLUT = CurveToMoveFrom.UIDToArrayIndexLUT;
+		CurveToMoveFrom.UIDToArrayIndexLUT = nullptr;
+		NumValidCurveCount = CurveToMoveFrom.NumValidCurveCount;
+		CurveToMoveFrom.NumValidCurveCount = 0;
+		Elements = MoveTemp(CurveToMoveFrom.Elements);
+		bInitialized = true;
+		CurveToMoveFrom.bInitialized = false;
+	}
+
 	/** Empty */
 	void Empty()
 	{
@@ -576,26 +615,12 @@ struct FBaseBlendedCurve
 	}
 };
 
-struct FBlendedHeapCurve;
-
 struct ENGINE_API FBlendedCurve : public FBaseBlendedCurve<FAnimStackAllocator>
 {
 };
 
 struct ENGINE_API FBlendedHeapCurve : public FBaseBlendedCurve<FDefaultAllocator>
 {
-	/** Once moved, source is invalid */
-	void MoveFrom(FBlendedHeapCurve& CurveToMoveFrom)
-	{
-		UIDToArrayIndexLUT = CurveToMoveFrom.UIDToArrayIndexLUT;
-		CurveToMoveFrom.UIDToArrayIndexLUT = nullptr;
-		NumValidCurveCount = CurveToMoveFrom.NumValidCurveCount;
-		CurveToMoveFrom.NumValidCurveCount = 0;
-		Elements = MoveTemp(CurveToMoveFrom.Elements);
-		bInitialized = true;
-		CurveToMoveFrom.bInitialized = false;
-	}
-
 };
 
 UENUM()
diff --git a/Engine/Source/Runtime/Engine/Public/BonePose.h b/Engine/Source/Runtime/Engine/Public/BonePose.h
--- a/Engine/Source/Runtime/Engine/Public/BonePose.h
+++ b/Engine/Source/Runtime/Engine/Public/BonePose.h
@@ -87,6 +87,9 @@ public:
 	};
 
 	const TArray<FTransform, InAllocator>& GetBones() const { return Bones; }
+
+	TArray<FTransform, InAllocator>&& MoveBones() { return MoveTemp(Bones); }
+
 protected:
 	TArray<FTransform, InAllocator> Bones;
 };
@@ -175,6 +178,13 @@ public:
 		this->Bones = SrcPose.Bones;
 	}
 
+	void InitFrom(FBaseCompactPose&& SrcPose)
+	{
+		SetBoneContainer(SrcPose.BoneContainer);
+		this->Bones = MoveTemp(SrcPose.Bones);
+		SrcPose.BoneContainer = nullptr;
+	}
+
 	// Copy bone transform from SrcPose to this
 	template <typename OtherAllocator>
 	void CopyBonesFrom(const FBaseCompactPose<OtherAllocator>& SrcPose)
@@ -192,6 +202,16 @@ public:
 		}
 	}
 	
+	void MoveBonesFrom(FBaseCompactPose<Allocator>& SrcPose)
+	{
+		if (this != &SrcPose)
+		{
+			this->Bones = SrcPose.MoveBones();
+			BoneContainer = &SrcPose.GetBoneContainer();
+			SrcPose.BoneContainer = nullptr;
+		}
+	}
+
 	template <typename OtherAllocator>
 	void CopyBonesFrom(const TArray<FTransform, OtherAllocator>& SrcPoseBones)
 	{
@@ -213,6 +233,18 @@ public:
 		DestPoseBones = this->Bones;
 	}
 
+	// Moves transform data out of the supplied InTransforms. InTransform will be left empty
+	void MoveBonesFrom(TArray<FTransform, Allocator>&& InTransforms)
+	{
+		this->Bones = MoveTemp(InTransforms);
+	}
+
+	// Moves transform data out of this to the supplied InTransforms. Bones will be left empty
+	void MoveBonesTo(TArray<FTransform, Allocator>& InTransforms)
+	{
+		InTransforms = MoveTemp(this->Bones);
+	}
+
 	void Empty()
 	{
 		BoneContainer = nullptr;
@@ -326,38 +358,12 @@ protected:
 	const FBoneContainer* BoneContainer;
 };
 
-struct FCompactHeapPose;
-
 struct ENGINE_API FCompactPose : public FBaseCompactPose<FAnimStackAllocator>
 {
 };
 
 struct ENGINE_API FCompactHeapPose : public FBaseCompactPose<FDefaultAllocator>
 {
-	// Moves transform data out of the supplied SrcPose (SrcPose will be left empty)
-	void MoveBonesFrom(FCompactHeapPose& SrcPose)
-	{
-		Bones = MoveTemp(SrcPose.Bones);
-		BoneContainer = SrcPose.BoneContainer;
-		SrcPose.BoneContainer = NULL;
-	}
-
-	//Moves transform data to supplied array (pose will be left empty)
-	void MoveBonesTo(TArray<FTransform>& OutTransforms)
-	{
-		OutTransforms = MoveTemp(Bones);
-		BoneContainer = NULL;
-	}
-
-	// Moves transform data out of the supplied InTransforms. InTransform will be left empty
-	void MoveBonesFrom(TArray<FTransform>& InTransforms)
-	{
-		// if number doesn't match it is not compatible to start with
-		if (InTransforms.Num() == Bones.Num())
-		{
-			Bones = MoveTemp(InTransforms);
-		}
-	}
 };
 
 struct FMeshPose : public FBasePose<FMeshPoseBoneIndex, FDefaultAllocator>
@@ -378,14 +384,6 @@ public:
 		InitBones(BoneContainer->GetNumBones());
 	}
 
-	//void InitFrom(FCompactPose& CompactPose);
-
-	void MoveBonesTo(TArray<FTransform>& OutTransforms)
-	{
-		OutTransforms = MoveTemp(Bones);
-		BoneContainer = NULL;
-	}
-
 	// Sets this pose to its ref pose
 	ENGINE_API void ResetToRefPose();
 
@@ -435,6 +433,14 @@ struct FCSPose
 		ComponentSpaceFlags[0] = 1;
 	}
 
+	void InitPose(PoseType&& SrcPose)
+	{
+		Pose.InitFrom(MoveTemp(SrcPose));
+		ComponentSpaceFlags.Empty(Pose.GetNumBones());
+		ComponentSpaceFlags.AddZeroed(Pose.GetNumBones());
+		ComponentSpaceFlags[0] = 1;
+	}
+
 	// Copy Pose
 	template <typename OtherPoseType>
 	void CopyPose(const OtherPoseType& SrcPose)
@@ -494,7 +500,8 @@ struct FCSPose
 	void LocalBlendCSBoneTransforms(const TArray<struct FBoneTransform>& BoneTransforms, float Alpha);
 
 	// Convert any component space transforms back to local space
-	void ConvertToLocalPoses(PoseType& OutPose) const;
+	static void ConvertComponentPosesToLocalPoses(FCSPose<PoseType>&& InPose, PoseType& OutPose);
+	static void ConvertComponentPosesToLocalPoses(const FCSPose<PoseType>& InPose, PoseType& OutPose);
 
 protected:
 	PoseType Pose;
@@ -812,10 +819,10 @@ void FCSPose<PoseType>::LocalBlendCSBoneTransforms(const TArray<struct FBoneTran
 }
 
 template<class PoseType>
-void FCSPose<PoseType>::ConvertToLocalPoses(PoseType& OutPose) const
+void FCSPose<PoseType>::ConvertComponentPosesToLocalPoses(const FCSPose<PoseType>& InPose, PoseType& OutPose)
 {
-	checkSlow(Pose.IsValid());
-	OutPose = Pose;
+	checkSlow(InPose.IsValid());
+	OutPose = InPose.Pose;
 
 	// now we need to convert back to local bases
 	// only convert back that has been converted to mesh base
@@ -825,18 +832,52 @@ void FCSPose<PoseType>::ConvertToLocalPoses(PoseType& OutPose) const
 	// go from child to parent since I need parent inverse to go back to local
 	// root is same, so no need to do Index == 0
 	const BoneIndexType RootBoneIndex(0);
-	if (ComponentSpaceFlags[RootBoneIndex])
+	if (InPose.ComponentSpaceFlags[RootBoneIndex])
 	{
-		OutPose[RootBoneIndex] = Pose[RootBoneIndex];
+		OutPose[RootBoneIndex] = InPose.Pose[RootBoneIndex];
 	}
 
-	const int32 NumBones = Pose.GetNumBones();
+	const int32 NumBones = InPose.Pose.GetNumBones();
 	for (int32 Index = NumBones - 1; Index > 0; Index--)
 	{
 		const BoneIndexType BoneIndex(Index);
-		if (ComponentSpaceFlags[BoneIndex])
+		if (InPose.ComponentSpaceFlags[BoneIndex])
 		{
-			const BoneIndexType ParentIndex = Pose.GetParentBoneIndex(BoneIndex);
+			const BoneIndexType ParentIndex = InPose.Pose.GetParentBoneIndex(BoneIndex);
+			OutPose[BoneIndex].SetToRelativeTransform(OutPose[ParentIndex]);
+			OutPose[BoneIndex].NormalizeRotation();
+		}
+	}
+}
+
+template<class PoseType>
+void FCSPose<PoseType>::ConvertComponentPosesToLocalPoses(FCSPose<PoseType>&& InPose, PoseType& OutPose)
+{
+	checkSlow(InPose.IsValid());
+
+	const int32 NumBones = InPose.Pose.GetNumBones();
+
+	// now we need to convert back to local bases
+	// only convert back that has been converted to mesh base
+	// if it was local base, and if it hasn't been modified
+	// that is still okay even if parent is changed, 
+	// that doesn't mean this local has to change
+	// go from child to parent since I need parent inverse to go back to local
+	// root is same, so no need to do Index == 0
+	const BoneIndexType RootBoneIndex(0);
+	if (InPose.ComponentSpaceFlags[RootBoneIndex])
+	{
+		OutPose[RootBoneIndex] = InPose.Pose[RootBoneIndex];
+	}
+
+	OutPose = MoveTemp(InPose.Pose);
+	
+	for (int32 Index = NumBones - 1; Index > 0; Index--)
+	{
+		const BoneIndexType BoneIndex(Index);
+		if (InPose.ComponentSpaceFlags[BoneIndex])
+		{
+			const BoneIndexType ParentIndex = OutPose.GetParentBoneIndex(BoneIndex);
 			OutPose[BoneIndex].SetToRelativeTransform(OutPose[ParentIndex]);
 			OutPose[BoneIndex].NormalizeRotation();
 		}
