Fix for recording audio buffers to wav file

diff --git a/Engine/Source/Runtime/Engine/Private/SampleBuffer.cpp b/Engine/Source/Runtime/Engine/Private/SampleBuffer.cpp
--- a/Engine/Source/Runtime/Engine/Private/SampleBuffer.cpp
+++ b/Engine/Source/Runtime/Engine/Private/SampleBuffer.cpp
@@ -498,7 +498,7 @@ namespace Audio
 			IFileHandle* FileHandle = PlatformFile.OpenWrite(*AbsoluteFilePath);
 			if (FileHandle)
 			{
-				int32 NumChunks = FMath::CeilToInt(SerializedWavData.Num() / ChunkSize);
+				int32 NumChunks = FMath::CeilToInt((float)SerializedWavData.Num() / ChunkSize);
 				UE_LOG(LogAudio, Display, TEXT("Writing wav file in %d chunks..."), NumChunks);
 				for (int32 ChunkIndex = 0; ChunkIndex < NumChunks; ChunkIndex++)
 				{
