Fix uninitialized FWaveInstance variables

diff --git a/Engine/Source/Runtime/Engine/Private/Audio.cpp b/Engine/Source/Runtime/Engine/Private/Audio.cpp
--- a/Engine/Source/Runtime/Engine/Private/Audio.cpp
+++ b/Engine/Source/Runtime/Engine/Private/Audio.cpp
@@ -761,6 +761,7 @@ FWaveInstance::FWaveInstance( FActiveSound* InActiveSound )
 	: WaveData(nullptr)
 	, SoundClass(nullptr)
 	, SoundSubmix(nullptr)
+	, SourceEffectChain(nullptr)
 	, ActiveSound(InActiveSound)
 	, Volume(0.0f)
 	, DistanceAttenuation(1.0f)
@@ -790,11 +791,14 @@ FWaveInstance::FWaveInstance( FActiveSound* InActiveSound )
 	, bIsMusic(false)
 	, bReverb(true)
 	, bCenterChannelOnly(false)
+	, bIsPaused(false)
 	, bReportedSpatializationWarning(false)
 	, bIsAmbisonics(false)
 	, bIsStopping(false)
 	, SpatializationMethod(ESoundSpatializationAlgorithm::SPATIALIZATION_Default)
+	, SpatializationPluginSettings(nullptr)
 	, OcclusionPluginSettings(nullptr)
+	, ReverbPluginSettings(nullptr)
 	, OutputTarget(EAudioOutputTarget::Speaker)
 	, LowPassFilterFrequency(MAX_FILTER_FREQUENCY)
 	, OcclusionFilterFrequency(MAX_FILTER_FREQUENCY)
@@ -808,8 +812,13 @@ FWaveInstance::FWaveInstance( FActiveSound* InActiveSound )
 	, AttenuationDistance(0.0f)
 	, ListenerToSoundDistance(0.0f)
 	, AbsoluteAzimuth(0.0f)
+	, PlaybackTime(0.0f)
+	, ReverbSendMethod(EReverbSendMethod::Linear)
 	, ReverbSendLevelRange(0.0f, 0.0f)
 	, ReverbSendLevelDistanceRange(0.0f, 0.0f)
+	, ManualReverbSendLevel(0.0f)
+	, TypeHash(0)
+	, WaveInstanceHash(0)
 	, UserIndex(0)
 {
 	TypeHash = ++TypeHashCounter;
