Fix debug build

diff --git a/Engine/Source/Runtime/Engine/Public/BonePose.h b/Engine/Source/Runtime/Engine/Public/BonePose.h
--- a/Engine/Source/Runtime/Engine/Public/BonePose.h
+++ b/Engine/Source/Runtime/Engine/Public/BonePose.h
@@ -821,7 +821,7 @@ void FCSPose<PoseType>::LocalBlendCSBoneTransforms(const TArray<struct FBoneTran
 template<class PoseType>
 void FCSPose<PoseType>::ConvertComponentPosesToLocalPoses(const FCSPose<PoseType>& InPose, PoseType& OutPose)
 {
-	checkSlow(InPose.IsValid());
+	checkSlow(InPose.Pose.IsValid());
 	OutPose = InPose.Pose;
 
 	// now we need to convert back to local bases
@@ -853,7 +853,7 @@ void FCSPose<PoseType>::ConvertComponentPosesToLocalPoses(const FCSPose<PoseType
 template<class PoseType>
 void FCSPose<PoseType>::ConvertComponentPosesToLocalPoses(FCSPose<PoseType>&& InPose, PoseType& OutPose)
 {
-	checkSlow(InPose.IsValid());
+	checkSlow(InPose.Pose.IsValid());
 
 	const int32 NumBones = InPose.Pose.GetNumBones();
 
