Fixed Controller.OnNewPawn not being broadcasted when PlayerController is switching possessed pawns #UE4

diff --git a/Engine/Source/Runtime/Engine/Private/PlayerController.cpp b/Engine/Source/Runtime/Engine/Private/PlayerController.cpp
--- a/Engine/Source/Runtime/Engine/Private/PlayerController.cpp
+++ b/Engine/Source/Runtime/Engine/Private/PlayerController.cpp
@@ -752,7 +752,9 @@ void APlayerController::OnPossess(APawn* PawnToPossess)
 	if ( PawnToPossess != NULL && 
 		(PlayerState == NULL || !PlayerState->bOnlySpectator) )
 	{
-		if (GetPawn() && GetPawn() != PawnToPossess)
+		const bool bNewPawn = (GetPawn() != PawnToPossess);
+
+		if (GetPawn() && bNewPawn)
 		{
 			UnPossess();
 		}
@@ -801,6 +803,11 @@ void APlayerController::OnPossess(APawn* PawnToPossess)
 		// not calling UpdateNavigationComponents() anymore. The
 		// PathFollowingComponent is now observing newly possessed
 		// pawns (via OnNewPawn)
+		// need to broadcast here since we don't call Super::OnPossess
+		if (bNewPawn)
+		{
+			OnNewPawn.Broadcast(GetPawn());
+		}
 	}
 }
 
