Adding IsBotControlled accessor just as other IsXControlled accessors

#nrx

diff --git a/Engine/Source/Runtime/Engine/Classes/GameFramework/Pawn.h b/Engine/Source/Runtime/Engine/Classes/GameFramework/Pawn.h
--- a/Engine/Source/Runtime/Engine/Classes/GameFramework/Pawn.h
+++ b/Engine/Source/Runtime/Engine/Classes/GameFramework/Pawn.h
@@ -315,7 +315,11 @@ public:
 	/** @return true if controlled by a human player (possessed by a PlayerController).	 */
 	UFUNCTION(BlueprintPure, Category=Pawn)
 	virtual bool IsPlayerControlled() const;
-	
+
+	/** @return true if controlled by a bot.	 */
+	UFUNCTION(BlueprintPure, Category = Pawn)
+	virtual bool IsBotControlled() const;
+
 	/**
 	 * Get the view rotation of the Pawn (direction they are looking, normally Controller->ControlRotation).
 	 * @return The view rotation of the Pawn.
diff --git a/Engine/Source/Runtime/Engine/Private/Pawn.cpp b/Engine/Source/Runtime/Engine/Private/Pawn.cpp
--- a/Engine/Source/Runtime/Engine/Private/Pawn.cpp
+++ b/Engine/Source/Runtime/Engine/Private/Pawn.cpp
@@ -216,6 +216,11 @@ bool APawn::IsPlayerControlled() const
 	return PlayerState && !PlayerState->bIsABot;
 }
 
+bool APawn::IsBotControlled() const
+{
+	return PlayerState && PlayerState->bIsABot;
+}
+
 bool APawn::ReachedDesiredRotation()
 {
 	// Only base success on Yaw 
