Exported some virtual functions so cross module extension of the class could work

diff --git a/Engine/Source/Runtime/Engine/Classes/Engine/TextureCube.h b/Engine/Source/Runtime/Engine/Classes/Engine/TextureCube.h
--- a/Engine/Source/Runtime/Engine/Classes/Engine/TextureCube.h
+++ b/Engine/Source/Runtime/Engine/Classes/Engine/TextureCube.h
@@ -21,11 +21,11 @@ public:
 	TMap<FString, FTexturePlatformData*> CookedPlatformData;
 
 	//~ Begin UObject Interface.
-	virtual void Serialize(FArchive& Ar) override;
-	virtual void PostLoad() override;
-	virtual void GetAssetRegistryTags(TArray<FAssetRegistryTag>& OutTags) const override;
-	virtual FString GetDesc() override;
-	virtual void GetResourceSizeEx(FResourceSizeEx& CumulativeResourceSize) override;
+	ENGINE_API virtual void Serialize(FArchive& Ar) override;
+	ENGINE_API virtual void PostLoad() override;
+	ENGINE_API virtual void GetAssetRegistryTags(TArray<FAssetRegistryTag>& OutTags) const override;
+	ENGINE_API virtual FString GetDesc() override;
+	ENGINE_API virtual void GetResourceSizeEx(FResourceSizeEx& CumulativeResourceSize) override;
 	//~ End UObject Interface.
 
 	/** Trivial accessors. */
@@ -65,8 +65,8 @@ public:
 	//~ Begin UTexture Interface
 	virtual float GetSurfaceWidth() const override { return GetSizeX(); }
 	virtual float GetSurfaceHeight() const override { return GetSizeY(); }
-	virtual FTextureResource* CreateResource() override;
-	virtual void UpdateResource() override;
+	ENGINE_API virtual FTextureResource* CreateResource() override;
+	ENGINE_API virtual void UpdateResource() override;
 	virtual EMaterialValueType GetMaterialType() const override { return MCT_TextureCube; }
 	virtual FTexturePlatformData** GetRunningPlatformData() override { return &PlatformData; }
 	virtual TMap<FString, FTexturePlatformData*> *GetCookedPlatformData() override { return &CookedPlatformData; }
@@ -86,13 +86,13 @@ public:
 	 * @param	Enum	Which mips to calculate size for.
 	 * @return	Total size of all specified mips, in bytes
 	 */
-	virtual uint32 CalcTextureMemorySizeEnum( ETextureMipCount Enum ) const override;
+	ENGINE_API virtual uint32 CalcTextureMemorySizeEnum( ETextureMipCount Enum ) const override;
 
 #if WITH_EDITOR
 	/**
 	* Return maximum dimension for this texture type.
 	*/
-	virtual uint32 GetMaximumDimension() const override;
+	ENGINE_API virtual uint32 GetMaximumDimension() const override;
 #endif
 };
 
