#DEV-VR - Fix for crash in DebugCanvas on application termination

diff --git a/Engine/Source/Runtime/Engine/Private/Slate/DebugCanvas.cpp b/Engine/Source/Runtime/Engine/Private/Slate/DebugCanvas.cpp
--- a/Engine/Source/Runtime/Engine/Private/Slate/DebugCanvas.cpp
+++ b/Engine/Source/Runtime/Engine/Private/Slate/DebugCanvas.cpp
@@ -102,7 +102,10 @@ FDebugCanvasDrawer::~FDebugCanvasDrawer()
 
 	if (LayerID != INVALID_LAYER_ID)
 	{
-		GEngine->StereoRenderingDevice->GetStereoLayers()->DestroyLayer(LayerID);
+		if (GEngine->StereoRenderingDevice.IsValid())
+		{
+			GEngine->StereoRenderingDevice->GetStereoLayers()->DestroyLayer(LayerID);
+		}
 		LayerID = INVALID_LAYER_ID;
 	}
 }
