Fix for stereo layers crashing on exit, where the DebugCanvas was destroyed after the layer manager, causing a nullptr dereference

diff --git a/Engine/Source/Runtime/Engine/Private/Slate/DebugCanvas.cpp b/Engine/Source/Runtime/Engine/Private/Slate/DebugCanvas.cpp
--- a/Engine/Source/Runtime/Engine/Private/Slate/DebugCanvas.cpp
+++ b/Engine/Source/Runtime/Engine/Private/Slate/DebugCanvas.cpp
@@ -100,7 +100,7 @@ FDebugCanvasDrawer::~FDebugCanvasDrawer()
 		RenderThreadCanvas = nullptr;
 	}
 
-	if (LayerID != INVALID_LAYER_ID)
+	if ((LayerID != INVALID_LAYER_ID) && GEngine->StereoRenderingDevice.IsValid())
 	{
 		if (GEngine->StereoRenderingDevice.IsValid())
 		{
