#DEV-VR - Fix for dragging a "skylight" from lighting menu directly into sequencer.

diff --git a/Engine/Source/Runtime/Engine/Private/Components/SkyLightComponent.cpp b/Engine/Source/Runtime/Engine/Private/Components/SkyLightComponent.cpp
--- a/Engine/Source/Runtime/Engine/Private/Components/SkyLightComponent.cpp
+++ b/Engine/Source/Runtime/Engine/Private/Components/SkyLightComponent.cpp
@@ -559,7 +559,7 @@ void USkyLightComponent::UpdateSkyCaptureContentsArray(UWorld* WorldToUpdate, TA
 		USkyLightComponent* CaptureComponent = ComponentArray[CaptureIndex];
 		AActor* Owner = CaptureComponent->GetOwner();
 
-		if (((!Owner || Owner->GetLevel()->bIsVisible) && CaptureComponent->GetWorld() == WorldToUpdate)
+		if (((!Owner || (Owner->GetLevel() && Owner->GetLevel()->bIsVisible)) && CaptureComponent->GetWorld() == WorldToUpdate)
 			// Only process sky capture requests once async shader compiling completes, otherwise we will capture the scene with temporary shaders
 			&& (!bIsCompilingShaders || CaptureComponent->SourceType == SLS_SpecifiedCubemap))
 		{
