Split ImportFileAsTexture2D into two functions so you can also import from a buffer

diff --git a/Engine/Source/Runtime/Engine/Classes/Kismet/KismetRenderingLibrary.h b/Engine/Source/Runtime/Engine/Classes/Kismet/KismetRenderingLibrary.h
--- a/Engine/Source/Runtime/Engine/Classes/Kismet/KismetRenderingLibrary.h
+++ b/Engine/Source/Runtime/Engine/Classes/Kismet/KismetRenderingLibrary.h
@@ -124,7 +124,13 @@ class UKismetRenderingLibrary : public UBlueprintFunctionLibrary
 	UFUNCTION(BlueprintCallable, Category = "Rendering", meta = (WorldContext = "WorldContextObject"))
 	static ENGINE_API UTexture2D* ImportFileAsTexture2D(UObject* WorldContextObject, const FString& Filename);
 
-	/** 
+	/**
+	 * Imports a texture from a buffer and creates Texture2D from it.
+	 */
+	UFUNCTION(BlueprintCallable, Category = "Rendering", meta = (WorldContext = "WorldContextObject"))
+	static ENGINE_API UTexture2D* ImportBufferAsTexture2D(UObject* WorldContextObject, const TArray<uint8>& Buffer);
+	
+	/**
 	 * Returns a Canvas object that can be used to draw to the specified render target.
 	 * Canvas has functions like DrawMaterial with size parameters that can be used to draw to a specific area of a render target.
 	 * Be sure to call EndDrawCanvasToRenderTarget to complete the rendering!
diff --git a/Engine/Source/Runtime/Engine/Private/ImageUtils.cpp b/Engine/Source/Runtime/Engine/Private/ImageUtils.cpp
--- a/Engine/Source/Runtime/Engine/Private/ImageUtils.cpp
+++ b/Engine/Source/Runtime/Engine/Private/ImageUtils.cpp
@@ -766,6 +766,8 @@ UTexture2D* FImageUtils::ImportFileAsTexture2D(const FString& Filename)
 						FMemory::Memcpy(MipData, HDRDDSLoadHelper.GetDDSDataPointer(), NewTexture->PlatformData->Mips[0].BulkData.GetBulkDataSize());
 
 						NewTexture->PlatformData->Mips[0].BulkData.Unlock();
+
+						NewTexture->UpdateResource();
 					}
 
 				}
@@ -773,67 +775,85 @@ UTexture2D* FImageUtils::ImportFileAsTexture2D(const FString& Filename)
 		}
 		else
 		{
-			EImageFormat Format = ImageWrapperModule.DetectImageFormat(Buffer.GetData(), Buffer.GetAllocatedSize());
-
-			if (Format != EImageFormat::Invalid)
-			{
-				TSharedPtr<IImageWrapper> ImageWrapper = ImageWrapperModule.CreateImageWrapper(Format);
-
-				if (ImageWrapper->SetCompressed((void*)Buffer.GetData(), Buffer.GetAllocatedSize()))
-				{
-					PixelFormat = PF_Unknown;
-
-					ERGBFormat RGBFormat = ERGBFormat::Invalid;
-		
-					BitDepth = ImageWrapper->GetBitDepth();
+			NewTexture = FImageUtils::ImportBufferAsTexture2D(Buffer);
+		}
 
-					Width = ImageWrapper->GetWidth();
-					Height = ImageWrapper->GetHeight();
+		if(!NewTexture)
+		{
+			UE_LOG(LogImageUtils, Warning, TEXT("Error creating texture. %s is not a supported file format"), *Filename)
+		}	
+	}
+	else
+	{
+		UE_LOG(LogImageUtils, Warning, TEXT("Error creating texture. %s could not be found"), *Filename)
+	}
 
-					if (BitDepth == 16)
-					{
-						PixelFormat = PF_FloatRGBA;
-						RGBFormat = ERGBFormat::BGRA;
-					}
-					else if (BitDepth == 8)
-					{
-						PixelFormat = PF_B8G8R8A8;
-						RGBFormat = ERGBFormat::BGRA;
-					}
-					else
-					{
-						UE_LOG(LogImageUtils, Warning, TEXT("Error creating texture. %s is not a valid bit depth (%d)"), BitDepth);
-					}
+	return NewTexture;
+}
 
-					const TArray<uint8>* UncompressedData = nullptr;
-					ImageWrapper->GetRaw(RGBFormat, BitDepth, UncompressedData);
+UTexture2D* FImageUtils::ImportBufferAsTexture2D(const TArray<uint8>& Buffer)
+{
+	IImageWrapperModule& ImageWrapperModule = FModuleManager::Get().LoadModuleChecked<IImageWrapperModule>(TEXT("ImageWrapper"));
 
-					NewTexture = UTexture2D::CreateTransient(Width, Height, PixelFormat);
-					if (NewTexture)
-					{
-						uint8* MipData = static_cast<uint8*>(NewTexture->PlatformData->Mips[0].BulkData.Lock(LOCK_READ_WRITE));
+	EImageFormat Format = ImageWrapperModule.DetectImageFormat(Buffer.GetData(), Buffer.GetAllocatedSize());
 
-						// Bulk data was already allocated for the correct size when we called CreateTransient above
-						FMemory::Memcpy(MipData, UncompressedData->GetData(), NewTexture->PlatformData->Mips[0].BulkData.GetBulkDataSize());
+	UTexture2D* NewTexture = nullptr;
+	EPixelFormat PixelFormat = PF_Unknown;
 
-						NewTexture->PlatformData->Mips[0].BulkData.Unlock();
-					}
-				}
-			}
-		}
+	if (Format != EImageFormat::Invalid)
+	{
+		TSharedPtr<IImageWrapper> ImageWrapper = ImageWrapperModule.CreateImageWrapper(Format);
+		
+		int32 BitDepth = 0;
+		int32 Width = 0;
+		int32 Height = 0;
 
-		if(!NewTexture)
-		{
-			UE_LOG(LogImageUtils, Warning, TEXT("Error creating texture. %s is not a supported file format"), *Filename)
-		}	
-		else
+		if (ImageWrapper->SetCompressed((void*)Buffer.GetData(), Buffer.GetAllocatedSize()))
 		{
-			NewTexture->UpdateResource();
+			PixelFormat = PF_Unknown;
+			
+			ERGBFormat RGBFormat = ERGBFormat::Invalid;
+			
+			BitDepth = ImageWrapper->GetBitDepth();
+			
+			Width = ImageWrapper->GetWidth();
+			Height = ImageWrapper->GetHeight();
+			
+			if (BitDepth == 16)
+			{
+				PixelFormat = PF_FloatRGBA;
+				RGBFormat = ERGBFormat::BGRA;
+			}
+			else if (BitDepth == 8)
+			{
+				PixelFormat = PF_B8G8R8A8;
+				RGBFormat = ERGBFormat::BGRA;
+			}
+			else
+			{
+				UE_LOG(LogImageUtils, Warning, TEXT("Error creating texture. %s is not a valid bit depth (%d)"), BitDepth);
+			}
+			
+			const TArray<uint8>* UncompressedData = nullptr;
+			ImageWrapper->GetRaw(RGBFormat, BitDepth, UncompressedData);
+			
+			NewTexture = UTexture2D::CreateTransient(Width, Height, PixelFormat);
+			if (NewTexture)
+			{
+				uint8* MipData = static_cast<uint8*>(NewTexture->PlatformData->Mips[0].BulkData.Lock(LOCK_READ_WRITE));
+				
+				// Bulk data was already allocated for the correct size when we called CreateTransient above
+				FMemory::Memcpy(MipData, UncompressedData->GetData(), NewTexture->PlatformData->Mips[0].BulkData.GetBulkDataSize());
+				
+				NewTexture->PlatformData->Mips[0].BulkData.Unlock();
+
+				NewTexture->UpdateResource();
+			}
 		}
 	}
 	else
 	{
-		UE_LOG(LogImageUtils, Warning, TEXT("Error creating texture. %s could not be found"), *Filename)
+		UE_LOG(LogImageUtils, Warning, TEXT("Error creating texture. Couldn't determine the file format"));
 	}
 
 	return NewTexture;
diff --git a/Engine/Source/Runtime/Engine/Private/KismetRenderingLibrary.cpp b/Engine/Source/Runtime/Engine/Private/KismetRenderingLibrary.cpp
--- a/Engine/Source/Runtime/Engine/Private/KismetRenderingLibrary.cpp
+++ b/Engine/Source/Runtime/Engine/Private/KismetRenderingLibrary.cpp
@@ -514,6 +514,11 @@ UTexture2D* UKismetRenderingLibrary::ImportFileAsTexture2D(UObject* WorldContext
 	return FImageUtils::ImportFileAsTexture2D(Filename);
 }
 
+UTexture2D* UKismetRenderingLibrary::ImportBufferAsTexture2D(UObject* WorldContextObject, const TArray<uint8>& Buffer)
+{
+	return FImageUtils::ImportBufferAsTexture2D(Buffer);
+}
+
 void UKismetRenderingLibrary::BeginDrawCanvasToRenderTarget(UObject* WorldContextObject, UTextureRenderTarget2D* TextureRenderTarget, UCanvas*& Canvas, FVector2D& Size, FDrawToRenderTargetContext& Context)
 {
 	UWorld* World = GEngine->GetWorldFromContextObject(WorldContextObject, EGetWorldErrorMode::LogAndReturnNull);
diff --git a/Engine/Source/Runtime/Engine/Public/ImageUtils.h b/Engine/Source/Runtime/Engine/Public/ImageUtils.h
--- a/Engine/Source/Runtime/Engine/Public/ImageUtils.h
+++ b/Engine/Source/Runtime/Engine/Public/ImageUtils.h
@@ -148,6 +148,11 @@ public:
 	 */
 	ENGINE_API static UTexture2D* ImportFileAsTexture2D(const FString& Filename);
 
+	/**
+	 * Imports a texture a buffer and creates Texture2D from it
+	 */
+	ENGINE_API static UTexture2D* ImportBufferAsTexture2D(const TArray<uint8>& Buffer);
+	
 	/**
 	* Exports a UTextureRenderTargetCube as an HDR image on the disk.
 	*
