#DEV-VR - Moving Lumin Audio Platform above Android because Lumin is Android

diff --git a/Engine/Source/Runtime/Engine/Public/AudioPluginUtilities.h b/Engine/Source/Runtime/Engine/Public/AudioPluginUtilities.h
--- a/Engine/Source/Runtime/Engine/Public/AudioPluginUtilities.h
+++ b/Engine/Source/Runtime/Engine/Public/AudioPluginUtilities.h
@@ -16,6 +16,9 @@ struct ENGINE_API AudioPluginUtilities
 	static const EAudioPlatform CurrentPlatform = EAudioPlatform::Linux;
 #elif PLATFORM_IOS
 	static const EAudioPlatform CurrentPlatform = EAudioPlatform::IOS;
+//Lumin needs to be before Android because it derives off of it, therefore defines PLATFORM_ANDROID=1
+#elif PLATFORM_LUMIN
+	static const EAudioPlatform CurrentPlatform = EAudioPlatform::Lumin;
 #elif PLATFORM_ANDROID
 	static const EAudioPlatform CurrentPlatform = EAudioPlatform::Android;
 #elif PLATFORM_XBOXONE
@@ -26,8 +29,6 @@ struct ENGINE_API AudioPluginUtilities
 	static const EAudioPlatform CurrentPlatform = EAudioPlatform::Switch;
 #elif PLATFORM_HTML5
 	static const EAudioPlatform CurrentPlatform = EAudioPlatform::HTML5;
-#elif PLATFORM_LUMIN
-	static const EAudioPlatform CurrentPlatform = EAudioPlatform::Lumin;
 #else
 	static const EAudioPlatform CurrentPlatform = EAudioPlatform::Unknown;
 #endif
