We need to recompute the view frustum and related distances if the projection matrix is changed.

diff --git a/Engine/Source/Runtime/Engine/Private/SceneView.cpp b/Engine/Source/Runtime/Engine/Private/SceneView.cpp
--- a/Engine/Source/Runtime/Engine/Private/SceneView.cpp
+++ b/Engine/Source/Runtime/Engine/Private/SceneView.cpp
@@ -543,6 +543,36 @@ FViewMatrices::FViewMatrices(const FSceneViewInitOptions& InitOptions) : FViewMa
 	);
 };
 
+static void SetupViewFrustum(FSceneView& View)
+{
+	if (View.SceneViewInitOptions.OverrideFarClippingPlaneDistance > 0.0f)
+	{
+		const FPlane FarPlane(View.ViewMatrices.GetViewOrigin() + View.GetViewDirection() * View.SceneViewInitOptions.OverrideFarClippingPlaneDistance, View.GetViewDirection());
+		// Derive the view frustum from the view projection matrix, overriding the far plane
+		GetViewFrustumBounds(View.ViewFrustum, View.ViewMatrices.GetViewProjectionMatrix(), FarPlane, true, false);
+	}
+	else
+	{
+		// Derive the view frustum from the view projection matrix.
+		GetViewFrustumBounds(View.ViewFrustum, View.ViewMatrices.GetViewProjectionMatrix(), false);
+	}
+
+	// Derive the view's near clipping distance and plane.
+	// The GetFrustumFarPlane() is the near plane because of reverse Z projection.
+	static_assert((int32)ERHIZBuffer::IsInverted != 0, "Fix Near Clip distance!");
+	View.bHasNearClippingPlane = View.ViewMatrices.GetViewProjectionMatrix().GetFrustumFarPlane(View.NearClippingPlane);
+	if (View.ViewMatrices.GetProjectionMatrix().M[2][3] > DELTA)
+	{
+		// Infinite projection with reversed Z.
+		View.NearClippingDistance = View.ViewMatrices.GetProjectionMatrix().M[3][2];
+	}
+	else
+	{
+		// Ortho projection with reversed Z.
+		View.NearClippingDistance = (1.0f - View.ViewMatrices.GetProjectionMatrix().M[3][2]) / View.ViewMatrices.GetProjectionMatrix().M[2][2];
+	}
+}
+
 FSceneView::FSceneView(const FSceneViewInitOptions& InitOptions)
 	: Family(InitOptions.ViewFamily)
 	, State(InitOptions.SceneViewStateInterface)
@@ -630,32 +660,7 @@ FSceneView::FSceneView(const FSceneViewInitOptions& InitOptions)
 	}
 #endif // !(UE_BUILD_SHIPPING || UE_BUILD_TEST)
 
-	if (InitOptions.OverrideFarClippingPlaneDistance > 0.0f)
-	{
-		const FPlane FarPlane(ViewMatrices.GetViewOrigin() + GetViewDirection() * InitOptions.OverrideFarClippingPlaneDistance, GetViewDirection());
-		// Derive the view frustum from the view projection matrix, overriding the far plane
-		GetViewFrustumBounds(ViewFrustum, ViewMatrices.GetViewProjectionMatrix(),FarPlane,true,false);
-	}
-	else
-	{
-		// Derive the view frustum from the view projection matrix.
-		GetViewFrustumBounds(ViewFrustum, ViewMatrices.GetViewProjectionMatrix(),false);
-	}
-
-	// Derive the view's near clipping distance and plane.
-	// The GetFrustumFarPlane() is the near plane because of reverse Z projection.
-	static_assert((int32)ERHIZBuffer::IsInverted != 0, "Fix Near Clip distance!");
-	bHasNearClippingPlane = ViewMatrices.GetViewProjectionMatrix().GetFrustumFarPlane(NearClippingPlane);
-	if(ViewMatrices.GetProjectionMatrix().M[2][3] > DELTA)
-	{
-		// Infinite projection with reversed Z.
-		NearClippingDistance = ViewMatrices.GetProjectionMatrix().M[3][2];
-	}
-	else
-	{
-		// Ortho projection with reversed Z.
-		NearClippingDistance = (1.0f - ViewMatrices.GetProjectionMatrix().M[3][2]) / ViewMatrices.GetProjectionMatrix().M[2][2];
-	}
+	SetupViewFrustum(*this);
 
 	// Determine whether the view should reverse the cull mode due to a negative determinant.  Only do this for a valid scene
 	bReverseCulling = (Family && Family->Scene) ? FMath::IsNegativeFloat(ViewMatrices.GetViewMatrix().Determinant()) : false;
@@ -866,6 +871,8 @@ void FSceneView::UpdateProjectionMatrix(const FMatrix& NewProjectionMatrix)
 	// Create new matrices
 	FViewMatrices NewViewMatrices = FViewMatrices(SceneViewInitOptions);
 	ViewMatrices = NewViewMatrices;
+
+	SetupViewFrustum(*this);
 }
 
 void FViewMatrices::UpdateViewMatrix(const FVector& ViewLocation, const FRotator& ViewRotation)
