Debug canvas layer fixes

diff --git a/Engine/Source/Runtime/Engine/Private/Slate/DebugCanvas.cpp b/Engine/Source/Runtime/Engine/Private/Slate/DebugCanvas.cpp
--- a/Engine/Source/Runtime/Engine/Private/Slate/DebugCanvas.cpp
+++ b/Engine/Source/Runtime/Engine/Private/Slate/DebugCanvas.cpp
@@ -120,6 +120,7 @@ void FDebugCanvasDrawer::BeginRenderingCanvas( const FIntRect& CanvasRect )
 {
 	if( CanvasRect.Size().X > 0 && CanvasRect.Size().Y > 0 )
 	{
+		bCanvasRenderedLastFrame = true;
 		ENQUEUE_UNIQUE_RENDER_COMMAND_THREEPARAMETER
 		(
 			BeginRenderingDebugCanvas,
@@ -194,6 +195,8 @@ void FDebugCanvasDrawer::InitDebugCanvas(FViewportClient* ViewportClient, UWorld
 				LayerID = INVALID_LAYER_ID;
 			}
 		}
+
+		bCanvasRenderedLastFrame = false;
 	}
 }
 
@@ -242,7 +245,6 @@ void FDebugCanvasDrawer::DrawRenderThread(FRHICommandListImmediate& RHICmdList,
 			RenderThreadCanvas->SetRenderTargetRect( RenderTarget->GetViewRect() );
 		}
 
-		bCanvasRenderedLastFrame = RenderThreadCanvas->HasBatchesToRender();
 		RenderThreadCanvas->Flush_RenderThread(RHICmdList, true);
 		RenderThreadCanvas->SetAllowSwitchVerticalAxis(bNeedToFlipVertical);
 		RenderTarget->ClearRenderTargetTexture();
diff --git a/Engine/Source/Runtime/Engine/Public/IStereoLayers.h b/Engine/Source/Runtime/Engine/Public/IStereoLayers.h
--- a/Engine/Source/Runtime/Engine/Public/IStereoLayers.h
+++ b/Engine/Source/Runtime/Engine/Public/IStereoLayers.h
@@ -175,7 +175,7 @@ public:
 		// Default debug layer desc
 		IStereoLayers::FLayerDesc StereoLayerDesc;
 		StereoLayerDesc.Transform = FTransform(FVector(100.f, 0, 0));
-		StereoLayerDesc.QuadSize = FVector2D(200.f, 200.f);
+		StereoLayerDesc.QuadSize = FVector2D(120.f, 120.f);
 		StereoLayerDesc.PositionType = IStereoLayers::ELayerType::FaceLocked;
 		StereoLayerDesc.ShapeType = IStereoLayers::ELayerShape::QuadLayer;
 		StereoLayerDesc.Texture = Texture;
