Add physics command for listing all physics meshes, their collision modes, and the number of actors which use them.

diff --git a/Engine/Source/Runtime/Engine/Private/PhysicsEngine/PhysUtils.cpp b/Engine/Source/Runtime/Engine/Private/PhysicsEngine/PhysUtils.cpp
--- a/Engine/Source/Runtime/Engine/Private/PhysicsEngine/PhysUtils.cpp
+++ b/Engine/Source/Runtime/Engine/Private/PhysicsEngine/PhysUtils.cpp
@@ -6,6 +6,7 @@
 #include "EngineDefines.h"
 #include "Engine/EngineTypes.h"
 #include "Engine/World.h"
+#include "EngineUtils.h"
 #include "PhysxUserData.h"
 #include "PhysicsEngine/BodyInstance.h"
 #include "Components/PrimitiveComponent.h"
@@ -15,6 +16,10 @@
 #include "PhysicsEngine/ConvexElem.h"
 #include "PhysicsEngine/BodySetup.h"
 #include "Physics/PhysicsInterfaceCore.h"
+#include "Components/StaticMeshComponent.h"
+#include "Components/SkeletalMeshComponent.h"
+#include "Engine/StaticMesh.h"
+
 
 /** Returns false if ModelToHulls operation should halt because of vertex count overflow. */
 static bool AddConvexPrim(FKAggregateGeom* OutGeom, TArray<FPlane> &Planes, UModel* InModel)
@@ -478,7 +483,7 @@ bool FPhysScene::HandleExecCommands(const TCHAR* Cmd, FOutputDevice* Ar)
 		return ExecPxVis(Cmd, Ar);
 	}
 
-	return false;
+return false;
 }
 
 //// EXEC
@@ -491,15 +496,15 @@ bool FPhysicsInterface::ExecPhysCommands(const TCHAR* Cmd, FOutputDevice* Ar, UW
 	{
 		return true;
 	}
-	else if(!IsRunningCommandlet() && GPhysXSDK && FParse::Command(&Cmd, TEXT("PVD")) )
+	else if (!IsRunningCommandlet() && GPhysXSDK && FParse::Command(&Cmd, TEXT("PVD")))
 	{
 		// check if PvdConnection manager is available on this platform
-		if(GPhysXVisualDebugger != NULL)
+		if (GPhysXVisualDebugger != NULL)
 		{
-			if(FParse::Command(&Cmd, TEXT("CONNECT")))
+			if (FParse::Command(&Cmd, TEXT("CONNECT")))
 			{
-				
-				
+
+
 				const bool bVizualization = !FParse::Command(&Cmd, TEXT("NODEBUG"));
 
 				// setup connection parameters
@@ -509,12 +514,12 @@ bool FPhysicsInterface::ExecPhysCommands(const TCHAR* Cmd, FOutputDevice* Ar, UW
 					Host = Cmd;
 				}
 
-				
+
 
 				PvdConnect(Host, bVizualization);
 
 			}
-			else if(FParse::Command(&Cmd, TEXT("DISCONNECT")))
+			else if (FParse::Command(&Cmd, TEXT("DISCONNECT")))
 			{
 				GPhysXVisualDebugger->disconnect();
 			}
@@ -523,18 +528,18 @@ bool FPhysicsInterface::ExecPhysCommands(const TCHAR* Cmd, FOutputDevice* Ar, UW
 		return 1;
 	}
 #if PHYSX_MEMORY_STATS
-	else if(GPhysXAllocator && FParse::Command(&Cmd, TEXT("PHYSXALLOC")) )
+	else if (GPhysXAllocator && FParse::Command(&Cmd, TEXT("PHYSXALLOC")))
 	{
 		GPhysXAllocator->DumpAllocations(Ar);
 		return 1;
 	}
 #endif
-	else if(FParse::Command(&Cmd, TEXT("PHYSXSHARED")) )
+	else if (FParse::Command(&Cmd, TEXT("PHYSXSHARED")))
 	{
 		FPhysxSharedData::Get().DumpSharedMemoryUsage(Ar);
 		return 1;
 	}
-	else if(FParse::Command(&Cmd, TEXT("PHYSXINFO")))
+	else if (FParse::Command(&Cmd, TEXT("PHYSXINFO")))
 	{
 		Ar->Logf(TEXT("PhysX Info:"));
 		Ar->Logf(TEXT("  Version: %d.%d.%d"), PX_PHYSICS_VERSION_MAJOR, PX_PHYSICS_VERSION_MINOR, PX_PHYSICS_VERSION_BUGFIX);
@@ -545,7 +550,7 @@ bool FPhysicsInterface::ExecPhysCommands(const TCHAR* Cmd, FOutputDevice* Ar, UW
 #else
 		Ar->Logf(TEXT("  Configuration: PROFILE"));
 #endif
-		if(GetPhysXCookingModule())
+		if (GetPhysXCookingModule())
 		{
 			Ar->Logf(TEXT("  Cooking Module: TRUE"));
 		}
@@ -556,6 +561,61 @@ bool FPhysicsInterface::ExecPhysCommands(const TCHAR* Cmd, FOutputDevice* Ar, UW
 
 		return 1;
 	}
+	else if (FParse::Command(&Cmd, TEXT("PHYSCOLLISIONACTORS")))
+	{
+
+		TMap<TPair<FName, FName>, int32> ActorCounts;
+
+		static const FName QueryAndPhysName = TEXT("QueryAndPhysics");
+		static const FName PhysOnlyName = TEXT("PhysOnlyName");
+
+		int32 TotalActors = 0;
+		for (TActorIterator<AActor> Itr(InWorld); Itr; ++Itr)
+		{
+			AActor* Actor = *Itr;
+			if (USceneComponent* SC = Actor->GetRootComponent())
+			{
+				ECollisionEnabled::Type Collision = SC->GetCollisionEnabled();
+				const bool bQueryAndPhys = Collision == ECollisionEnabled::QueryAndPhysics;
+				const bool bPhysOnly = Collision == ECollisionEnabled::PhysicsOnly;
+				if (bQueryAndPhys || bPhysOnly)
+				{
+					FName PhysicsName = SC->GetFName();
+
+					if (UStaticMeshComponent* StaticMeshComp = Cast<UStaticMeshComponent>(SC))
+					{
+						if (UStaticMesh* SM = StaticMeshComp->GetStaticMesh())
+						{
+							PhysicsName = SM->GetFName();
+						}
+					}
+					else if (USkeletalMeshComponent* SkelMeshComp = Cast<USkeletalMeshComponent>(SC))
+					{
+						if (USkeletalMesh* SM = SkelMeshComp->SkeletalMesh)
+						{
+							PhysicsName = SM->GetFName();
+						}
+					}
+
+					const TPair<FName, FName> ActorKey(PhysicsName, bQueryAndPhys ? QueryAndPhysName : PhysOnlyName);
+					int32& ActorCount = ActorCounts.FindOrAdd(ActorKey);
+					++ActorCount;
+					++TotalActors;
+				}
+			}
+		}
+
+		ActorCounts.ValueSort([](int32 A, int32 B) { return B < A; });
+
+		Ar->Logf(TEXT("Actors with Physics Enabled:"));
+		for (auto ActorIt = ActorCounts.CreateIterator(); ActorIt; ++ActorIt)
+		{
+			Ar->Logf(TEXT("x%d: %s (%s)"), ActorIt->Value, *(ActorIt->Key.Key.ToString()), *(ActorIt->Key.Value.ToString()));
+		}
+		Ar->Logf(TEXT("Total Physics Actors: %d\n"), TotalActors);
+
+		return 1;
+	}
 
 #endif // WITH_PHYSX
 
