Optimized processing of acks in NetConnection by tracking written channel ids per packet
thus avoiding to touch all open channels when processing received acks.
Consolidated older similar optimization made for demo recording to use same path.

Local instrumentations shows average time to process acks with ~900 open channelzs dropping from ~.029ms to ~0.05ms for a single connection. Extrapolating this to 100 connections will provide a nice gain for dedicated server when channel counts gets a bit higher.

Ryan.Gerleve, Brian.Bekich, Peter.Engstrom, Jon.Nabozny

diff --git a/Engine/Source/Runtime/Engine/Classes/Engine/NetConnection.h b/Engine/Source/Runtime/Engine/Classes/Engine/NetConnection.h
--- a/Engine/Source/Runtime/Engine/Classes/Engine/NetConnection.h
+++ b/Engine/Source/Runtime/Engine/Classes/Engine/NetConnection.h
@@ -22,6 +22,7 @@
 #include "ReplicationDriver.h"
 #include "Analytics/EngineNetAnalytics.h"
 #include "PacketTraits.h"
+#include "Net/Util/ResizableCircularQueue.h"
 
 #include "NetConnection.generated.h"
 
@@ -205,6 +206,29 @@ public:
 };
 #endif
 
+/** Record of channels with data written into each outgoing packet. */
+struct FWrittenChannelsRecord
+{
+	enum { DefaultInitialSize = 1024 };
+
+	struct FChannelRecordEntry
+	{
+		uint32 Value : 31;
+		uint32 IsSequence : 1;
+	};
+
+	typedef TResizableCircularQueue<FChannelRecordEntry> FChannelRecordEntryQueue;
+
+	FChannelRecordEntryQueue ChannelRecord;
+	int32 LastPacketId;
+
+public:
+	FWrittenChannelsRecord(size_t InitialSize = DefaultInitialSize)
+		: ChannelRecord(InitialSize)
+		, LastPacketId(-1)
+	{
+	}
+};
 
 UCLASS(customConstructor, Abstract, MinimalAPI, transient, config=Engine)
 class UNetConnection : public UPlayer
@@ -1084,8 +1108,8 @@ private:
 	/** This is only used in UChannel::SendBunch. It's a member so that we can preserve the allocation between calls, as an optimization, and in a thread-safe way to be compatible with demo.ClientRecordAsyncEndOfFrame */
 	TArray<FOutBunch*> OutgoingBunches;
 
-	/** Bookkeeping for connections using internal ack to avoid acking all open channels */
-	TArray<int32> ChannelsWaitingForInternalAck;
+	/** Per packet bookkeeping of written channelIds */
+	FWrittenChannelsRecord ChannelRecord;
 
 	/** Sequence data used to implement reliability */
 	FNetPacketNotify PacketNotify;
diff --git a/Engine/Source/Runtime/Engine/Private/NetConnection.cpp b/Engine/Source/Runtime/Engine/Private/NetConnection.cpp
--- a/Engine/Source/Runtime/Engine/Private/NetConnection.cpp
+++ b/Engine/Source/Runtime/Engine/Private/NetConnection.cpp
@@ -58,6 +58,29 @@ DECLARE_CYCLE_STAT(TEXT("NetConnection Tick"), Stat_NetConnectionTick, STATGROUP
 DECLARE_CYCLE_STAT(TEXT("NetConnection ReceivedNak"), Stat_NetConnectionReceivedNak, STATGROUP_Net);
 DECLARE_CYCLE_STAT(TEXT("NetConnection NetConnectionReceivedAcks"), Stat_NetConnectionReceivedAcks, STATGROUP_Net);
 
+// ChannelRecord Implementation
+namespace FChannelRecordImpl
+{
+
+// Push ChannelRecordEntry for packet if PacketID differs from last PacketId
+static void PushPacketId(FWrittenChannelsRecord& WrittenChannelsRecord, int32 PacketId);
+
+// Push written ChannelIndex to WrittenChanneRecord, Push new packetId if PackedIdDiffers from Last pushed entry */
+static void PushChannelRecord(FWrittenChannelsRecord& WrittenChannelsRecord, int32 PacketId, int32 ChannelIndex);
+
+// Consume all FChannelRecordEntries for the given PacketId and execute a function with the signature (void*)(int32 PacketId, uint32 Channelndex) for each entry
+template<class Functor>
+static void ConsumeChannelRecordsForPacket(FWrittenChannelsRecord& WrittenChannelsRecord, int32 PacketId, Functor&& Func);
+
+// Consume all FChannelRecordEntries and execute a function with the signature (void*)(uint32 Channelndex) each entry
+template<class Functor>
+static void ConsumeAllChannelRecords(FWrittenChannelsRecord& WrittenChannelsRecord, Functor&& Func);
+
+// Returns allocated size of the record
+static SIZE_T CountBytes(FWrittenChannelsRecord& WrittenChannelsRecord);
+
+};
+
 const int32 UNetConnection::DEFAULT_MAX_CHANNEL_SIZE = 32767;
 /*-----------------------------------------------------------------------------
 	UNetConnection implementation.
@@ -528,7 +551,12 @@ void UNetConnection::Serialize( FArchive& Ar )
 
 		GRANULAR_NETWORK_MEMORY_TRACKING_TRACK("IgnoringChannels", IgnoringChannels.CountBytes(Ar));
 		GRANULAR_NETWORK_MEMORY_TRACKING_TRACK("OutgoingBunches", OutgoingBunches.CountBytes(Ar));
-		GRANULAR_NETWORK_MEMORY_TRACKING_TRACK("ChannelsWaitingForInternalAck", ChannelsWaitingForInternalAck.CountBytes(Ar));
+		
+		GRANULAR_NETWORK_MEMORY_TRACKING_TRACK("ChannelRecord",
+			const SIZE_T SizeAllocatedByChannelRecord = FChannelRecordImpl::CountBytes(ChannelRecord);
+			Ar.CountBytes(SizeAllocatedByChannelRecord, SizeAllocatedByChannelRecord)
+		);
+
 		GRANULAR_NETWORK_MEMORY_TRACKING_TRACK("LastOut", LastOut.CountMemory(Ar));
 		GRANULAR_NETWORK_MEMORY_TRACKING_TRACK("SendBunchHeader", SendBunchHeader.CountMemory(Ar));
 
@@ -1239,8 +1267,12 @@ void UNetConnection::FlushNet(bool bIgnoreSimulation)
 		// Increase outgoing sequence number
 		if (!InternalAck)
 		{
-			PacketNotify.CommitAndIncrementOutSeq();			
+			PacketNotify.CommitAndIncrementOutSeq();
 		}
+
+		// Make sure that we always push an ChannelRecordEntry for each transmitted packet even if it is empty
+		FChannelRecordImpl::PushPacketId(ChannelRecord, OutPacketId);
+
 		++OutPacketId; 
 
 		++OutPackets;
@@ -1299,54 +1331,50 @@ void UNetConnection::ReceivedAck(int32 AckPacketId)
 {
 	UE_LOG(LogNetTraffic, Verbose, TEXT("   Received ack %i"), AckPacketId);
 
+	SCOPE_CYCLE_COUNTER(Stat_NetConnectionReceivedAcks);
+	
 	// Advance OutAckPacketId
 	OutAckPacketId = AckPacketId;
 
 	// Process the bunch.
 	LastRecvAckTime = Driver->Time;
 
-	if ( PackageMap != NULL )
+	if (PackageMap != NULL)
 	{
 		PackageMap->ReceivedAck( AckPacketId );
 	}
 
-	// Forward the ack to the channel.
+	auto AckChannelFunc = [this](int32 AckPacketId, uint32 ChannelIndex)
 	{
-		SCOPE_CYCLE_COUNTER(Stat_NetConnectionReceivedAcks);
-	
-		for (int32 i = OpenChannels.Num() - 1; i >= 0; i--)
+		UChannel* const Channel = Channels[ChannelIndex];
+
+		if (Channel)
 		{
-			UChannel* const Channel = OpenChannels[i];
+			if (Channel->OpenPacketId.Last == AckPacketId) // Necessary for unreliable "bNetTemporary" channels.
+			{
+				Channel->OpenAcked = 1;
+			}
 				
-			if (Channel)
+			for (FOutBunch* OutBunch = Channel->OutRec; OutBunch; OutBunch = OutBunch->Next)
 			{
-				if (Channel->OpenPacketId.Last == AckPacketId) // Necessary for unreliable "bNetTemporary" channels.
+				if (OutBunch->bOpen)
 				{
-					Channel->OpenAcked = 1;
+					UE_LOG(LogNet, VeryVerbose, TEXT("Channel %i reset Ackd because open is reliable. "), Channel->ChIndex );
+					Channel->OpenAcked  = 0; // We have a reliable open bunch, don't let the above code set the OpenAcked state,
+											// it must be set in UChannel::ReceivedAcks to verify all open bunches were received.
 				}
-				
-				for (FOutBunch* OutBunch = Channel->OutRec; OutBunch; OutBunch = OutBunch->Next)
-				{
-					if (OutBunch->bOpen)
-					{
-						UE_LOG(LogNet, VeryVerbose, TEXT("Channel %i reset Ackd because open is reliable. "), Channel->ChIndex );
-						Channel->OpenAcked  = 0; // We have a reliable open bunch, don't let the above code set the OpenAcked state,
-												// it must be set in UChannel::ReceivedAcks to verify all open bunches were received.
-					}
 
-					if (OutBunch->PacketId == AckPacketId)
-					{
-						OutBunch->ReceivedAck = 1;
-					}
+				if (OutBunch->PacketId == AckPacketId)
+				{
+					OutBunch->ReceivedAck = 1;
 				}
-				Channel->ReceivedAcks(); //warning: May destroy Channel.
-			}
-			else
-			{
-				UE_LOG(LogNet, Warning, TEXT("UNetConnection::ReceivedPacket: null channel in OpenChannels array while processing ack. %s"), *Describe());
 			}
+			Channel->ReceivedAcks(); //warning: May destroy Channel.
 		}
-	}
+	};
+
+	// Invoke AckChannelFunc on all channels written for this PacketId
+	FChannelRecordImpl::ConsumeChannelRecordsForPacket(ChannelRecord, AckPacketId, AckChannelFunc);
 }
 
 void UNetConnection::ReceivedNak( int32 NakPacketId )
@@ -1358,10 +1386,9 @@ void UNetConnection::ReceivedNak( int32 NakPacketId )
 	// Update pending NetGUIDs
 	PackageMap->ReceivedNak(NakPacketId);
 
-	// Tell channels about Nak
-	for (int32 i = OpenChannels.Num() - 1; i >= 0; i--)
+	auto NakChannelFunc = [this](int32 NakPacketId, uint32 ChannelIndex)
 	{
-		UChannel* const Channel = OpenChannels[i];
+		UChannel* const Channel = Channels[ChannelIndex];
 		if (Channel)
 		{
 			Channel->ReceivedNak(NakPacketId);
@@ -1370,11 +1397,10 @@ void UNetConnection::ReceivedNak( int32 NakPacketId )
 				Channel->ReceivedAcks(); //warning: May destroy Channel.
 			}
 		}
-		else
-		{
-			UE_LOG(LogNet, Warning, TEXT("UNetConnection::ReceivedNak: null channel in OpenChannels array. %s"), *Describe());
-		}
-	}
+	};
+
+	// Invoke NakChannelFunc on all channels written for this PacketId
+	FChannelRecordImpl::ConsumeChannelRecordsForPacket(ChannelRecord, NakPacketId, NakChannelFunc);
 
 	// Stats
 	++OutPacketsLost;
@@ -2254,10 +2280,8 @@ int32 UNetConnection::SendRawBunch( FOutBunch& Bunch, bool InAllowMerge )
 	// Write the bits to the buffer and remember the packet id used
 	Bunch.PacketId = WriteBitsToSendBuffer( SendBunchHeader.GetData(), SendBunchHeader.GetNumBits(), Bunch.GetData(), Bunch.GetNumBits(), EWriteBitsDataType::Bunch );
 
-	if (InternalAck)
-	{
-		ChannelsWaitingForInternalAck.Add(Bunch.ChIndex);
-	}
+	// Track channels that wrote data to this packet.
+	FChannelRecordImpl::PushChannelRecord(ChannelRecord, Bunch.PacketId, Bunch.ChIndex);
 
 	UE_LOG(LogNetTraffic, Verbose, TEXT("UNetConnection::SendRawBunch. ChIndex: %d. Bits: %d. PacketId: %d"), Bunch.ChIndex, Bunch.GetNumBits(), Bunch.PacketId );
 
@@ -2465,14 +2489,14 @@ void UNetConnection::Tick()
 		LastReceiveTime = Driver->Time;
 		LastReceiveRealtime = FPlatformTime::Seconds();
 		LastGoodPacketRealtime = FPlatformTime::Seconds();
-		
-		for (int32 DirtyIt = 0, EndDirtyIt = ChannelsWaitingForInternalAck.Num(); DirtyIt < EndDirtyIt; ++DirtyIt)
+
+		// Consume all records
+		auto InternnalAckChannelFunc = [this, bIsServer](uint32 ChannelIndex)
 		{
-			const int32 ChannelId = ChannelsWaitingForInternalAck[DirtyIt];
-			UChannel* Channel = Channels[ChannelId];
+			UChannel* Channel = Channels[ChannelIndex];
 			if (Channel)
 			{
-				for( FOutBunch* OutBunch=Channel->OutRec; OutBunch; OutBunch=OutBunch->Next )
+				for(FOutBunch* OutBunch=Channel->OutRec; OutBunch; OutBunch=OutBunch->Next)
 				{
 					OutBunch->ReceivedAck = 1;
 				}
@@ -2484,8 +2508,9 @@ void UNetConnection::Tick()
 
 				Channel->ReceivedAcks();
 			}
-		}
-		ChannelsWaitingForInternalAck.Reset();
+		};
+
+		FChannelRecordImpl::ConsumeAllChannelRecords(ChannelRecord, InternnalAckChannelFunc);
 	}
 
 	// Update stats.
@@ -3371,3 +3396,92 @@ static void	PrintActorReportFunc(const TArray<FString>& Args, UWorld* InWorld)
 }
 
 FAutoConsoleCommandWithWorldAndArgs PrintActorReportCmd(TEXT("net.ActorReport"), TEXT(""),	FConsoleCommandWithWorldAndArgsDelegate::CreateStatic(PrintActorReportFunc) );
+
+/*-----------------------------------------------------------------------------
+	FChannelRecordImpl
+-----------------------------------------------------------------------------*/
+
+namespace FChannelRecordImpl
+{
+
+void PushPacketId(FWrittenChannelsRecord& WrittenChannelsRecord, int32 PacketId)
+{
+	if (PacketId != WrittenChannelsRecord.LastPacketId)
+	{
+		FWrittenChannelsRecord::FChannelRecordEntry PacketEntry = { uint32(PacketId), 1u };
+		WrittenChannelsRecord.ChannelRecord.Enqueue(PacketEntry);
+		WrittenChannelsRecord.LastPacketId = PacketId;
+	}
+}
+
+void PushChannelRecord(FWrittenChannelsRecord& WrittenChannelsRecord, int32 PacketId, int32 ChannelIndex)
+{
+	PushPacketId(WrittenChannelsRecord, PacketId);
+
+	FWrittenChannelsRecord::FChannelRecordEntry ChannelEntry = { uint32(ChannelIndex), 0u };
+	WrittenChannelsRecord.ChannelRecord.Enqueue(ChannelEntry);
+}
+
+SIZE_T CountBytes(FWrittenChannelsRecord& WrittenChannelsRecord)
+{
+	return WrittenChannelsRecord.ChannelRecord.AllocatedCapacity() * sizeof(FWrittenChannelsRecord::FChannelRecordEntry);
+}
+
+template<class Functor>
+void ConsumeChannelRecordsForPacket(FWrittenChannelsRecord& WrittenChannelsRecord, int32 PacketId, Functor&& Func)
+{
+	const int32 ExpectedSeq = PacketId;
+	uint32 PreviousChannelIndex = uint32(-1);
+
+	FWrittenChannelsRecord::FChannelRecordEntryQueue& Record = WrittenChannelsRecord.ChannelRecord;
+
+	// We should ALWAYS have data when we get here
+	FWrittenChannelsRecord::FChannelRecordEntry PacketEntry = Record.Peek();
+	Record.Pop();
+
+	// Verify that we got the expected packetId
+	check(PacketEntry.IsSequence == 1u && PacketEntry.Value == (uint32)PacketId);
+
+	while (!Record.IsEmpty() && Record.PeekNoCheck().IsSequence == 0u)
+	{
+		const FWrittenChannelsRecord::FChannelRecordEntry Entry = Record.PeekNoCheck();
+		Record.PopNoCheck();
+
+		const uint32 ChannelIndex = Entry.Value;
+
+		// Only process channel once per packet
+		if (ChannelIndex != PreviousChannelIndex)
+		{
+			Func(PacketId, ChannelIndex);
+			PreviousChannelIndex = ChannelIndex;
+		}
+	}
+}
+
+template<class Functor>
+void ConsumeAllChannelRecords(FWrittenChannelsRecord& WrittenChannelsRecord, Functor&& Func)
+{
+	// Consume all records
+	uint32 PreviousChannelIndex = uint32(-1);
+	FWrittenChannelsRecord::FChannelRecordEntryQueue& Record = WrittenChannelsRecord.ChannelRecord;
+
+	while (!Record.IsEmpty())
+	{
+		const FWrittenChannelsRecord::FChannelRecordEntry Entry = Record.PeekNoCheck();
+		Record.PopNoCheck();
+
+		const uint32 ChannelIndex = Entry.Value;
+
+		// if the channel writes data multiple non-consecutive times between ticks, the func will be invoked multiple times which should not be an issue.
+		if (Entry.IsSequence == 0u && ChannelIndex != PreviousChannelIndex)
+		{
+			Func(ChannelIndex);
+			PreviousChannelIndex = ChannelIndex;
+		}
+	}
+}
+
+}
+
+
+  
\ No newline at end of file
