Fixed warning on Android

diff --git a/Engine/Source/Runtime/Engine/Private/NetConnection.cpp b/Engine/Source/Runtime/Engine/Private/NetConnection.cpp
--- a/Engine/Source/Runtime/Engine/Private/NetConnection.cpp
+++ b/Engine/Source/Runtime/Engine/Private/NetConnection.cpp
@@ -1344,13 +1344,13 @@ void UNetConnection::ReceivedAck(int32 AckPacketId)
 		PackageMap->ReceivedAck( AckPacketId );
 	}
 
-	auto AckChannelFunc = [this](int32 AckPacketId, uint32 ChannelIndex)
+	auto AckChannelFunc = [this](int32 AckedPacketId, uint32 ChannelIndex)
 	{
 		UChannel* const Channel = Channels[ChannelIndex];
 
 		if (Channel)
 		{
-			if (Channel->OpenPacketId.Last == AckPacketId) // Necessary for unreliable "bNetTemporary" channels.
+			if (Channel->OpenPacketId.Last == AckedPacketId) // Necessary for unreliable "bNetTemporary" channels.
 			{
 				Channel->OpenAcked = 1;
 			}
@@ -1364,7 +1364,7 @@ void UNetConnection::ReceivedAck(int32 AckPacketId)
 											// it must be set in UChannel::ReceivedAcks to verify all open bunches were received.
 				}
 
-				if (OutBunch->PacketId == AckPacketId)
+				if (OutBunch->PacketId == AckedPacketId)
 				{
 					OutBunch->ReceivedAck = 1;
 				}
@@ -1386,13 +1386,13 @@ void UNetConnection::ReceivedNak( int32 NakPacketId )
 	// Update pending NetGUIDs
 	PackageMap->ReceivedNak(NakPacketId);
 
-	auto NakChannelFunc = [this](int32 NakPacketId, uint32 ChannelIndex)
+	auto NakChannelFunc = [this](int32 NackedPacketId, uint32 ChannelIndex)
 	{
 		UChannel* const Channel = Channels[ChannelIndex];
 		if (Channel)
 		{
-			Channel->ReceivedNak(NakPacketId);
-			if (Channel->OpenPacketId.InRange(NakPacketId))
+			Channel->ReceivedNak(NackedPacketId);
+			if (Channel->OpenPacketId.InRange(NackedPacketId))
 			{
 				Channel->ReceivedAcks(); //warning: May destroy Channel.
 			}
