Add verbose logging to dump NetworkChecksum for debugging purposes

diff --git a/Engine/Source/Runtime/Engine/Private/PackageMapClient.cpp b/Engine/Source/Runtime/Engine/Private/PackageMapClient.cpp
--- a/Engine/Source/Runtime/Engine/Private/PackageMapClient.cpp
+++ b/Engine/Source/Runtime/Engine/Private/PackageMapClient.cpp
@@ -878,6 +878,8 @@ FNetworkGUID UPackageMapClient::InternalLoadObject( FArchive & Ar, UObject *& Ob
 		if ( ExportFlags.bHasNetworkChecksum )
 		{
 			Ar << NetworkChecksum;
+
+			UE_LOG(LogNetPackageMap, Verbose, TEXT("%s has network checksum %u"), *PathName, NetworkChecksum);
 		}
 
 		const bool bIsPackage = NetGUID.IsStatic() && !OuterGUID.IsValid();
