[NO NON-PII COMMIT MESSAGE]

diff --git a/Engine/Source/Runtime/Engine/Private/InstancedStaticMesh.cpp b/Engine/Source/Runtime/Engine/Private/InstancedStaticMesh.cpp
--- a/Engine/Source/Runtime/Engine/Private/InstancedStaticMesh.cpp
+++ b/Engine/Source/Runtime/Engine/Private/InstancedStaticMesh.cpp
@@ -2395,32 +2395,32 @@ void UInstancedStaticMeshComponent::SelectInstance(bool bInSelected, int32 InIns
 			SelectedInstances.Init(false, PerInstanceSMData.Num());
 		}
 
-		check(SelectedInstances.IsValidIndex(InInstanceIndex));
-		check(SelectedInstances.IsValidIndex(InInstanceIndex + (InInstanceCount - 1)));
-
+		check(InInstanceIndex >= 0 && InInstanceCount > 0);
+		check(InInstanceIndex + InInstanceCount - 1 < SelectedInstances.Num());
+		
 		for (int32 InstanceIndex = InInstanceIndex; InstanceIndex < InInstanceIndex + InInstanceCount; InstanceIndex++)
 		{
-			SelectedInstances[InstanceIndex] = bInSelected;
-		}
-
-		if (PerInstanceRenderData.IsValid())
-		{
-			for (int32 InstanceIndex = InInstanceIndex; InstanceIndex < InInstanceIndex + InInstanceCount; InstanceIndex++)
+			if (SelectedInstances.IsValidIndex(InInstanceIndex))
 			{
-				// Record if the instance is selected
-				FColor HitProxyColor(ForceInit);
-				bool bSelected = SelectedInstances[InstanceIndex] != 0;
-				if (PerInstanceRenderData->HitProxies.IsValidIndex(InstanceIndex))
-				{
-					HitProxyColor = PerInstanceRenderData->HitProxies[InstanceIndex]->Id.GetColor();
-				}
-				
-				int32 RenderIndex = InstanceReorderTable.IsValidIndex(InstanceIndex) ? InstanceReorderTable[InstanceIndex] : InstanceIndex;
-				if (RenderIndex != INDEX_NONE)
+				SelectedInstances[InstanceIndex] = bInSelected;
+
+				if (PerInstanceRenderData.IsValid())
 				{
-					InstanceUpdateCmdBuffer.SetEditorData(RenderIndex, HitProxyColor, bSelected);
+					// Record if the instance is selected
+					FColor HitProxyColor(ForceInit);
+					bool bSelected = SelectedInstances[InstanceIndex] != 0;
+					if (PerInstanceRenderData->HitProxies.IsValidIndex(InstanceIndex))
+					{
+						HitProxyColor = PerInstanceRenderData->HitProxies[InstanceIndex]->Id.GetColor();
+					}
+
+					int32 RenderIndex = InstanceReorderTable.IsValidIndex(InstanceIndex) ? InstanceReorderTable[InstanceIndex] : InstanceIndex;
+					if (RenderIndex != INDEX_NONE)
+					{
+						InstanceUpdateCmdBuffer.SetEditorData(RenderIndex, HitProxyColor, bSelected);
+					}
 				}
-			}
+			}			
 		}
 		
 		MarkRenderStateDirty();
