Scalar/Vector parameters and subclasses with different behavior can no longer have name conflicts.
Renaming parameters to an existing name in the same class will also change the value to the existing value.

diff --git a/Engine/Source/Runtime/Engine/Classes/Materials/MaterialExpression.h b/Engine/Source/Runtime/Engine/Classes/Materials/MaterialExpression.h
--- a/Engine/Source/Runtime/Engine/Classes/Materials/MaterialExpression.h
+++ b/Engine/Source/Runtime/Engine/Classes/Materials/MaterialExpression.h
@@ -410,6 +410,8 @@ class ENGINE_API UMaterialExpression : public UObject
 	* but are not UMaterialExpressionParameter due to class hierarchy. */
 	virtual bool HasAParameterName() const { return false; }
 	virtual void ValidateParameterName(const bool bAllowDuplicateName = true);
+	virtual bool HasClassAndNameCollision(UMaterialExpression* OtherExpression) const;
+	virtual void SetValueToMatchingExpression(UMaterialExpression* OtherExpression) {};
 
 	virtual FName GetParameterName() const { return NAME_None; }
 	virtual void SetParameterName(const FName& Name) {}
diff --git a/Engine/Source/Runtime/Engine/Classes/Materials/MaterialExpressionChannelMaskParameter.h b/Engine/Source/Runtime/Engine/Classes/Materials/MaterialExpressionChannelMaskParameter.h
--- a/Engine/Source/Runtime/Engine/Classes/Materials/MaterialExpressionChannelMaskParameter.h
+++ b/Engine/Source/Runtime/Engine/Classes/Materials/MaterialExpressionChannelMaskParameter.h
@@ -41,5 +41,5 @@ class UMaterialExpressionChannelMaskParameter : public UMaterialExpressionVector
 	virtual uint32 GetInputType(int32 InputIndex) override {return MCT_Float4;}
 #endif
 
-	virtual bool IsUsedAsChannelMask() override {return true;}
+	virtual bool IsUsedAsChannelMask() const override {return true;}
 };
diff --git a/Engine/Source/Runtime/Engine/Classes/Materials/MaterialExpressionCurveAtlasRowParameter.h b/Engine/Source/Runtime/Engine/Classes/Materials/MaterialExpressionCurveAtlasRowParameter.h
--- a/Engine/Source/Runtime/Engine/Classes/Materials/MaterialExpressionCurveAtlasRowParameter.h
+++ b/Engine/Source/Runtime/Engine/Classes/Materials/MaterialExpressionCurveAtlasRowParameter.h
@@ -23,7 +23,7 @@ class UMaterialExpressionCurveAtlasRowParameter : public UMaterialExpressionScal
 	UPROPERTY()
 	FExpressionInput InputTime;
 
-	virtual bool IsUsedAsAtlasPosition() override { return true; }
+	virtual bool IsUsedAsAtlasPosition() const override { return true; }
 	virtual UTexture* GetReferencedTexture() override;
 	virtual bool CanReferenceTexture() const { return true; }
 
diff --git a/Engine/Source/Runtime/Engine/Classes/Materials/MaterialExpressionFontSampleParameter.h b/Engine/Source/Runtime/Engine/Classes/Materials/MaterialExpressionFontSampleParameter.h
--- a/Engine/Source/Runtime/Engine/Classes/Materials/MaterialExpressionFontSampleParameter.h
+++ b/Engine/Source/Runtime/Engine/Classes/Materials/MaterialExpressionFontSampleParameter.h
@@ -48,6 +48,7 @@ class UMaterialExpressionFontSampleParameter : public UMaterialExpressionFontSam
 	virtual FName GetParameterName() const override { return ParameterName; }
 	virtual void SetParameterName(const FName& Name) override { ParameterName = Name; }
 	virtual void ValidateParameterName(const bool bAllowDuplicateName) override;
+	virtual void SetValueToMatchingExpression(UMaterialExpression* OtherExpression) override;
 #endif
 	//~ End UMaterialExpression Interface
 	
diff --git a/Engine/Source/Runtime/Engine/Classes/Materials/MaterialExpressionScalarParameter.h b/Engine/Source/Runtime/Engine/Classes/Materials/MaterialExpressionScalarParameter.h
--- a/Engine/Source/Runtime/Engine/Classes/Materials/MaterialExpressionScalarParameter.h
+++ b/Engine/Source/Runtime/Engine/Classes/Materials/MaterialExpressionScalarParameter.h
@@ -47,10 +47,14 @@ class UMaterialExpressionScalarParameter : public UMaterialExpressionParameter
 	bool SetParameterValue(FName InParameterName, float InValue);
 
 	virtual void PostEditChangeProperty(FPropertyChangedEvent& PropertyChangedEvent) override;
+	virtual void ValidateParameterName(const bool bAllowDuplicateName) override;
+	virtual bool HasClassAndNameCollision(UMaterialExpression* OtherExpression) const override;
+	virtual void SetValueToMatchingExpression(UMaterialExpression* OtherExpression) override;
 #endif
 
-	virtual bool IsUsedAsAtlasPosition() { return false; }
+	virtual bool IsUsedAsAtlasPosition() const { return false; }
 };
 
 
 
+
diff --git a/Engine/Source/Runtime/Engine/Classes/Materials/MaterialExpressionStaticBoolParameter.h b/Engine/Source/Runtime/Engine/Classes/Materials/MaterialExpressionStaticBoolParameter.h
--- a/Engine/Source/Runtime/Engine/Classes/Materials/MaterialExpressionStaticBoolParameter.h
+++ b/Engine/Source/Runtime/Engine/Classes/Materials/MaterialExpressionStaticBoolParameter.h
@@ -23,6 +23,7 @@ public:
 	virtual int32 CompilePreview(class FMaterialCompiler* Compiler, int32 OutputIndex) override;
 	virtual void GetCaption(TArray<FString>& OutCaptions) const override;
 	virtual uint32 GetOutputType(int32 OutputIndex) override {return MCT_StaticBool;}
+	virtual void SetValueToMatchingExpression(UMaterialExpression* OtherExpression) override;
 #endif
 	//~ End UMaterialExpression Interface
 
diff --git a/Engine/Source/Runtime/Engine/Classes/Materials/MaterialExpressionStaticComponentMaskParameter.h b/Engine/Source/Runtime/Engine/Classes/Materials/MaterialExpressionStaticComponentMaskParameter.h
--- a/Engine/Source/Runtime/Engine/Classes/Materials/MaterialExpressionStaticComponentMaskParameter.h
+++ b/Engine/Source/Runtime/Engine/Classes/Materials/MaterialExpressionStaticComponentMaskParameter.h
@@ -36,6 +36,7 @@ public:
 #if WITH_EDITOR
 	virtual int32 Compile(class FMaterialCompiler* Compiler, int32 OutputIndex) override;
 	virtual void GetCaption(TArray<FString>& OutCaptions) const override;
+	virtual void SetValueToMatchingExpression(UMaterialExpression* OtherExpression) override;
 #endif
 	//~ End UMaterialExpression Interface
 
diff --git a/Engine/Source/Runtime/Engine/Classes/Materials/MaterialExpressionTextureSampleParameter.h b/Engine/Source/Runtime/Engine/Classes/Materials/MaterialExpressionTextureSampleParameter.h
--- a/Engine/Source/Runtime/Engine/Classes/Materials/MaterialExpressionTextureSampleParameter.h
+++ b/Engine/Source/Runtime/Engine/Classes/Materials/MaterialExpressionTextureSampleParameter.h
@@ -47,6 +47,7 @@ class ENGINE_API UMaterialExpressionTextureSampleParameter : public UMaterialExp
 	virtual FName GetParameterName() const override { return ParameterName; }
 	virtual void SetParameterName(const FName& Name) override { ParameterName = Name; }
 	virtual void ValidateParameterName(const bool bAllowDuplicateName) override;
+	virtual void SetValueToMatchingExpression(UMaterialExpression* OtherExpression) override;
 #endif
 	//~ End UMaterialExpression Interface
 
diff --git a/Engine/Source/Runtime/Engine/Classes/Materials/MaterialExpressionVectorParameter.h b/Engine/Source/Runtime/Engine/Classes/Materials/MaterialExpressionVectorParameter.h
--- a/Engine/Source/Runtime/Engine/Classes/Materials/MaterialExpressionVectorParameter.h
+++ b/Engine/Source/Runtime/Engine/Classes/Materials/MaterialExpressionVectorParameter.h
@@ -31,8 +31,11 @@ class UMaterialExpressionVectorParameter : public UMaterialExpressionParameter
 	virtual bool SetParameterValue(FName InParameterName, FLinearColor InValue);
 
 	virtual void PostEditChangeProperty(FPropertyChangedEvent& PropertyChangedEvent) override;	
+	virtual void ValidateParameterName(const bool bAllowDuplicateName) override;
+	virtual bool HasClassAndNameCollision(UMaterialExpression* OtherExpression) const override;
+	virtual void SetValueToMatchingExpression(UMaterialExpression* OtherExpression) override;
 #endif
 
-	virtual bool IsUsedAsChannelMask() {return false;}
+	virtual bool IsUsedAsChannelMask() const {return false;}
 };
 
diff --git a/Engine/Source/Runtime/Engine/Private/Materials/MaterialExpressions.cpp b/Engine/Source/Runtime/Engine/Private/Materials/MaterialExpressions.cpp
--- a/Engine/Source/Runtime/Engine/Private/Materials/MaterialExpressions.cpp
+++ b/Engine/Source/Runtime/Engine/Private/Materials/MaterialExpressions.cpp
@@ -409,8 +409,8 @@ void ValidateParameterNameInternal(class UMaterialExpression* ExpressionToValida
 				{
 					if (Expression != nullptr && Expression->HasAParameterName())
 					{
-						// Name are unique per class type
-						if (Expression != ExpressionToValidate && Expression->GetClass() == ExpressionToValidate->GetClass() && Expression->GetParameterName() == PotentialName)
+						// Validate that the new name doesn't violate the expression's rules (by default, same name as another of the same class)
+						if (Expression != ExpressionToValidate && Expression->GetParameterName() == PotentialName && Expression->HasClassAndNameCollision(ExpressionToValidate))
 						{
 							bFoundValidName = false;
 							break;
@@ -426,6 +426,23 @@ void ValidateParameterNameInternal(class UMaterialExpression* ExpressionToValida
 			}
 		}
 
+		if (bAllowDuplicateName)
+		{
+			// Check for any matching values
+			for (UMaterialExpression* Expression : OwningMaterial->Expressions)
+			{
+				if (Expression != nullptr && Expression->HasAParameterName())
+				{
+					// Name are unique per class type
+					if (Expression != ExpressionToValidate && Expression->GetParameterName() == PotentialName && Expression->GetClass() == ExpressionToValidate->GetClass())
+					{
+						ExpressionToValidate->SetValueToMatchingExpression(Expression);
+						break;
+					}
+				}
+			}
+		}
+
 		ExpressionToValidate->SetParameterName(PotentialName);
 	}
 }
@@ -1220,6 +1237,11 @@ void UMaterialExpression::ValidateParameterName(const bool bAllowDuplicateName)
 	// Incrementing the name is now handled in UMaterialExpressionParameter::ValidateParameterName
 }
 
+bool UMaterialExpression::HasClassAndNameCollision(UMaterialExpression* OtherExpression) const
+{
+	return GetClass() == OtherExpression->GetClass();
+}
+
 #endif // WITH_EDITOR
 
 
@@ -1898,6 +1920,15 @@ void UMaterialExpressionTextureSampleParameter::ValidateParameterName(const bool
 	ValidateParameterNameInternal(this, Material, bAllowDuplicateName);
 }
 
+void UMaterialExpressionTextureSampleParameter::SetValueToMatchingExpression(UMaterialExpression* OtherExpression)
+{
+	UTexture* ExistingValue;
+	Material->GetTextureParameterValue(OtherExpression->GetParameterName(), ExistingValue);
+	Texture = ExistingValue;
+	UProperty* ParamProperty = FindField<UProperty>(UMaterialExpressionTextureSampleParameter::StaticClass(), GET_MEMBER_NAME_STRING_CHECKED(UMaterialExpressionTextureSampleParameter, Texture));
+	FPropertyChangedEvent PropertyChangedEvent(ParamProperty);
+	PostEditChangeProperty(PropertyChangedEvent);
+}
 #endif // WITH_EDITOR
 
 bool UMaterialExpressionTextureSampleParameter::IsNamedParameter(const FMaterialParameterInfo& ParameterInfo, UTexture*& OutValue) const
@@ -3464,6 +3495,32 @@ void UMaterialExpressionStaticComponentMaskParameter::GetCaption(TArray<FString>
 	OutCaptions.Add(TEXT("Mask Param")); 
 	OutCaptions.Add(FString::Printf(TEXT("'%s'"), *ParameterName.ToString()));
 }
+
+void UMaterialExpressionStaticComponentMaskParameter::SetValueToMatchingExpression(UMaterialExpression* OtherExpression)
+{
+	bool R = false;
+	bool G = false;
+	bool B = false;
+	bool A = false;
+	FGuid Guid;
+	Material->GetStaticComponentMaskParameterValue(OtherExpression->GetParameterName(), R, G, B, A, Guid);
+	DefaultR = R;
+	DefaultG = G;
+	DefaultB = B;
+	DefaultA = A;
+	UProperty* ParamProperty = FindField<UProperty>(UMaterialExpressionStaticComponentMaskParameter::StaticClass(), GET_MEMBER_NAME_STRING_CHECKED(UMaterialExpressionStaticComponentMaskParameter, DefaultR));
+	FPropertyChangedEvent RChangedEvent(ParamProperty);
+	PostEditChangeProperty(RChangedEvent);
+	ParamProperty = FindField<UProperty>(UMaterialExpressionStaticComponentMaskParameter::StaticClass(), GET_MEMBER_NAME_STRING_CHECKED(UMaterialExpressionStaticComponentMaskParameter, DefaultG));
+	FPropertyChangedEvent GChangedEvent(ParamProperty);
+	PostEditChangeProperty(GChangedEvent);
+	ParamProperty = FindField<UProperty>(UMaterialExpressionStaticComponentMaskParameter::StaticClass(), GET_MEMBER_NAME_STRING_CHECKED(UMaterialExpressionStaticComponentMaskParameter, DefaultB));
+	FPropertyChangedEvent BChangedEvent(ParamProperty);
+	PostEditChangeProperty(BChangedEvent);
+	ParamProperty = FindField<UProperty>(UMaterialExpressionStaticComponentMaskParameter::StaticClass(), GET_MEMBER_NAME_STRING_CHECKED(UMaterialExpressionStaticComponentMaskParameter, DefaultA));
+	FPropertyChangedEvent AChangedEvent(ParamProperty);
+	PostEditChangeProperty(AChangedEvent);
+}
 #endif // WITH_EDITOR
 
 bool UMaterialExpressionStaticComponentMaskParameter::IsNamedParameter(const FMaterialParameterInfo& ParameterInfo, bool& OutR, bool& OutG, bool& OutB, bool& OutA, FGuid&OutExpressionGuid) const
@@ -6200,6 +6257,48 @@ void UMaterialExpressionVectorParameter::PostEditChangeProperty(FPropertyChanged
 
 	Super::PostEditChangeProperty(PropertyChangedEvent);
 }
+
+void UMaterialExpressionVectorParameter::ValidateParameterName(const bool bAllowDuplicateName)
+{
+	bool bOverrideDuplicateBehavior = false;
+	for (UMaterialExpression* Expression : Material->Expressions)
+	{
+		if (Expression != nullptr && Expression->HasAParameterName())
+		{
+			UMaterialExpressionVectorParameter* VectorExpression = Cast<UMaterialExpressionVectorParameter>(Expression);
+			if (VectorExpression
+				&& GetParameterName() == VectorExpression->GetParameterName()
+				&& IsUsedAsChannelMask() != VectorExpression->IsUsedAsChannelMask())
+			{
+				bOverrideDuplicateBehavior = true;
+				break;
+			}
+		}
+	}
+	Super::ValidateParameterName(bOverrideDuplicateBehavior ? false : bAllowDuplicateName);
+}
+
+bool UMaterialExpressionVectorParameter::HasClassAndNameCollision(UMaterialExpression* OtherExpression) const
+{
+	UMaterialExpressionVectorParameter* VectorExpression = Cast<UMaterialExpressionVectorParameter>(OtherExpression);
+	if (VectorExpression
+		&& GetParameterName() == VectorExpression->GetParameterName()
+		&& IsUsedAsChannelMask() != VectorExpression->IsUsedAsChannelMask())
+	{
+		return true;
+	}
+	return Super::HasClassAndNameCollision(OtherExpression);
+}
+
+void UMaterialExpressionVectorParameter::SetValueToMatchingExpression(UMaterialExpression* OtherExpression)
+{
+	FLinearColor ExistingValue = FLinearColor::Transparent;
+	Material->GetVectorParameterValue(FMaterialParameterInfo(OtherExpression->GetParameterName()), ExistingValue);
+	DefaultValue = ExistingValue;
+	UProperty* ParamProperty = FindField<UProperty>(UMaterialExpressionVectorParameter::StaticClass(), GET_MEMBER_NAME_STRING_CHECKED(UMaterialExpressionVectorParameter, DefaultValue));
+	FPropertyChangedEvent PropertyChangedEvent(ParamProperty);
+	PostEditChangeProperty(PropertyChangedEvent);
+}
 #endif
 
 //
@@ -6423,6 +6522,48 @@ void UMaterialExpressionScalarParameter::PostEditChangeProperty(FPropertyChanged
 	Super::PostEditChangeProperty(PropertyChangedEvent);
 }
 
+void UMaterialExpressionScalarParameter::ValidateParameterName(const bool bAllowDuplicateName)
+{
+	bool bOverrideDuplicateBehavior = false;
+	for (UMaterialExpression* Expression : Material->Expressions)
+	{
+		if (Expression != nullptr && Expression->HasAParameterName())
+		{
+			UMaterialExpressionScalarParameter* ScalarExpression = Cast<UMaterialExpressionScalarParameter>(Expression);
+			if (ScalarExpression 
+				&& GetParameterName() == ScalarExpression->GetParameterName() 
+				&& IsUsedAsAtlasPosition() != ScalarExpression->IsUsedAsAtlasPosition())
+			{
+				bOverrideDuplicateBehavior = true;
+				break;
+			}
+		}
+	}
+	Super::ValidateParameterName(bOverrideDuplicateBehavior ? false : bAllowDuplicateName);
+}
+
+bool UMaterialExpressionScalarParameter::HasClassAndNameCollision(UMaterialExpression* OtherExpression) const
+{
+	UMaterialExpressionScalarParameter* ScalarExpression = Cast<UMaterialExpressionScalarParameter>(OtherExpression);
+	if (ScalarExpression
+		&& GetParameterName() == ScalarExpression->GetParameterName()
+		&& IsUsedAsAtlasPosition() != ScalarExpression->IsUsedAsAtlasPosition())
+	{
+		return true;
+	}
+	return Super::HasClassAndNameCollision(OtherExpression);
+}
+
+void UMaterialExpressionScalarParameter::SetValueToMatchingExpression(UMaterialExpression* OtherExpression)
+{
+	float ExistingValue = 0.0f;
+	Material->GetScalarParameterDefaultValue(FMaterialParameterInfo(OtherExpression->GetParameterName()), ExistingValue);
+	DefaultValue = ExistingValue;
+	UProperty* ParamProperty = FindField<UProperty>(UMaterialExpressionScalarParameter::StaticClass(), GET_MEMBER_NAME_STRING_CHECKED(UMaterialExpressionScalarParameter, DefaultValue));
+	FPropertyChangedEvent PropertyChangedEvent(ParamProperty);
+	PostEditChangeProperty(PropertyChangedEvent);
+}
+
 #endif
 
 //
@@ -6530,6 +6671,17 @@ void UMaterialExpressionStaticBoolParameter::GetCaption(TArray<FString>& OutCapt
 	OutCaptions.Add(FString::Printf(TEXT("Static Bool Param (%s)"), (DefaultValue ? TEXT("True") : TEXT("False")))); 
 	OutCaptions.Add(FString::Printf(TEXT("'%s'"), *ParameterName.ToString())); 
 }
+
+void UMaterialExpressionStaticBoolParameter::SetValueToMatchingExpression(UMaterialExpression* OtherExpression)
+{
+	bool ExistingValue = false;
+	FGuid Guid;
+	Material->GetStaticSwitchParameterValue(OtherExpression->GetParameterName(), ExistingValue, Guid);
+	DefaultValue = ExistingValue;
+	UProperty* ParamProperty = FindField<UProperty>(UMaterialExpressionStaticBoolParameter::StaticClass(), GET_MEMBER_NAME_STRING_CHECKED(UMaterialExpressionStaticBoolParameter, DefaultValue));
+	FPropertyChangedEvent PropertyChangedEvent(ParamProperty);
+	PostEditChangeProperty(PropertyChangedEvent);
+}
 #endif // WITH_EDITOR
 
 bool UMaterialExpressionStaticBoolParameter::IsNamedParameter(const FMaterialParameterInfo& ParameterInfo, bool& OutValue, FGuid& OutExpressionGuid) const
@@ -8936,6 +9088,21 @@ void UMaterialExpressionFontSampleParameter::ValidateParameterName(const bool bA
 {
 	ValidateParameterNameInternal(this, Material, bAllowDuplicateName);
 }
+
+void UMaterialExpressionFontSampleParameter::SetValueToMatchingExpression(UMaterialExpression* OtherExpression)
+{
+	UFont* FontValue;
+	int32 FontPage;
+	Material->GetFontParameterValue(FMaterialParameterInfo(OtherExpression->GetParameterName()), FontValue, FontPage);
+	Font = FontValue;
+	FontTexturePage = FontPage;
+	UProperty* ParamProperty = FindField<UProperty>(UMaterialExpressionFontSampleParameter::StaticClass(), GET_MEMBER_NAME_STRING_CHECKED(UMaterialExpressionFontSampleParameter, Font));
+	FPropertyChangedEvent PropertyChangedEvent(ParamProperty);
+	PostEditChangeProperty(PropertyChangedEvent);
+	ParamProperty = FindField<UProperty>(UMaterialExpressionFontSampleParameter::StaticClass(), GET_MEMBER_NAME_STRING_CHECKED(UMaterialExpressionFontSampleParameter, FontTexturePage));
+	FPropertyChangedEvent PageChangedEvent(ParamProperty);
+	PostEditChangeProperty(PageChangedEvent);
+}
 #endif // WITH_EDITOR
 
 bool UMaterialExpressionFontSampleParameter::IsNamedParameter(const FMaterialParameterInfo& ParameterInfo, UFont*& OutFontValue, int32& OutFontPage) const
