Invalidate lighting when removing instance in editor as it was missing
Override GetLightRelevance and if we have no instance but we have a foiliage actor in the map do not report a light relevance otherwise it will display the "lighting need to be rebuild" while it's not true

diff --git a/Engine/Source/Runtime/Engine/Private/InstancedStaticMesh.cpp b/Engine/Source/Runtime/Engine/Private/InstancedStaticMesh.cpp
--- a/Engine/Source/Runtime/Engine/Private/InstancedStaticMesh.cpp
+++ b/Engine/Source/Runtime/Engine/Private/InstancedStaticMesh.cpp
@@ -826,6 +826,16 @@ void FInstancedStaticMeshSceneProxy::SetupInstancedMeshBatch(int32 LODIndex, int
 	}
 }
 
+void FInstancedStaticMeshSceneProxy::GetLightRelevance(const FLightSceneProxy* LightSceneProxy, bool& bDynamic, bool& bRelevant, bool& bLightMapped, bool& bShadowMapped) const
+{
+	FStaticMeshSceneProxy::GetLightRelevance(LightSceneProxy, bDynamic, bRelevant, bLightMapped, bShadowMapped);
+
+	if (InstancedRenderData.PerInstanceRenderData->InstanceBuffer.GetNumInstances() == 0)
+	{
+		bRelevant = false;
+	}
+}
+
 bool FInstancedStaticMeshSceneProxy::GetShadowMeshElement(int32 LODIndex, int32 BatchIndex, uint8 InDepthPriorityGroup, FMeshBatch& OutMeshBatch, bool bDitheredLODTransition) const
 {
 	if (LODIndex < InstancedRenderData.VertexFactories.Num() && FStaticMeshSceneProxy::GetShadowMeshElement(LODIndex, BatchIndex, InDepthPriorityGroup, OutMeshBatch, bDitheredLODTransition))
diff --git a/Engine/Source/Runtime/Engine/Private/InstancedStaticMesh.h b/Engine/Source/Runtime/Engine/Private/InstancedStaticMesh.h
--- a/Engine/Source/Runtime/Engine/Private/InstancedStaticMesh.h
+++ b/Engine/Source/Runtime/Engine/Private/InstancedStaticMesh.h
@@ -555,6 +555,7 @@ public:
 		return Result;
 	}
 
+	virtual void GetLightRelevance(const FLightSceneProxy* LightSceneProxy, bool& bDynamic, bool& bRelevant, bool& bLightMapped, bool& bShadowMapped) const override;
 	virtual void GetDynamicMeshElements(const TArray<const FSceneView*>& Views, const FSceneViewFamily& ViewFamily, uint32 VisibilityMap, FMeshElementCollector& Collector) const override;
 
 	virtual int32 GetNumMeshBatches() const override;
diff --git a/Engine/Source/Runtime/Foliage/Private/InstancedFoliage.cpp b/Engine/Source/Runtime/Foliage/Private/InstancedFoliage.cpp
--- a/Engine/Source/Runtime/Foliage/Private/InstancedFoliage.cpp
+++ b/Engine/Source/Runtime/Foliage/Private/InstancedFoliage.cpp
@@ -1196,6 +1196,7 @@ void FFoliageMeshInfo::RemoveInstances(AInstancedFoliageActor* InIFA, const TArr
 	Instances.Shrink();
 		
 	Component->bAutoRebuildTreeOnInstanceChanges = PreviousbAutoRebuildTreeOnInstanceChanges;
+	Component->InvalidateLightingCache();
 
 	if (RebuildFoliageTree)
 	{
