Function inputs and outputs now rename correctly on duplication

diff --git a/Engine/Source/Runtime/Engine/Private/Materials/MaterialExpressions.cpp b/Engine/Source/Runtime/Engine/Private/Materials/MaterialExpressions.cpp
--- a/Engine/Source/Runtime/Engine/Private/Materials/MaterialExpressions.cpp
+++ b/Engine/Source/Runtime/Engine/Private/Materials/MaterialExpressions.cpp
@@ -11928,7 +11928,7 @@ void UMaterialExpressionFunctionInput::ValidateName()
 {
 	if (Material)
 	{
-		int32 InputNameIndex = 0;
+		int32 InputNameIndex = 1;
 		bool bResultNameIndexValid = true;
 		FName PotentialInputName;
 
@@ -11936,9 +11936,9 @@ void UMaterialExpressionFunctionInput::ValidateName()
 		do 
 		{
 			PotentialInputName = InputName;
-			if (InputNameIndex != 0)
+			if (InputNameIndex != 1)
 			{
-				PotentialInputName = *FString::Printf(TEXT("%s%d"), *InputName.ToString(), InputNameIndex);
+				PotentialInputName.SetNumber(InputNameIndex);
 			}
 
 			bResultNameIndexValid = true;
@@ -12133,7 +12133,7 @@ void UMaterialExpressionFunctionOutput::ValidateName()
 {
 	if (Material)
 	{
-		int32 OutputNameIndex = 0;
+		int32 OutputNameIndex = 1;
 		bool bResultNameIndexValid = true;
 		FName PotentialOutputName;
 
@@ -12141,9 +12141,9 @@ void UMaterialExpressionFunctionOutput::ValidateName()
 		do 
 		{
 			PotentialOutputName = OutputName;
-			if (OutputNameIndex != 0)
+			if (OutputNameIndex != 1)
 			{
-				PotentialOutputName = *FString::Printf(TEXT("%s%d"), *OutputName.ToString(), OutputNameIndex);
+				PotentialOutputName.SetNumber(OutputNameIndex);
 			}
 
 			bResultNameIndexValid = true;
