Exposed Disconnect/Reconnect/Travel to GameInstance to be overridden by project-specific implementations
Added proper support for the disconnect/reconnect/travel exec commands in Shootergame

diff --git a/Engine/Source/Runtime/Engine/Classes/Engine/GameInstance.h b/Engine/Source/Runtime/Engine/Classes/Engine/GameInstance.h
--- a/Engine/Source/Runtime/Engine/Classes/Engine/GameInstance.h
+++ b/Engine/Source/Runtime/Engine/Classes/Engine/GameInstance.h
@@ -126,6 +126,9 @@ protected:
 	// @todo jcf list of logged-in players?
 
 	virtual bool HandleOpenCommand(const TCHAR* Cmd, FOutputDevice& Ar, UWorld* InWorld);
+	virtual bool HandleDisconnectCommand(const TCHAR* Cmd, FOutputDevice& Ar, UWorld* InWorld);
+	virtual bool HandleReconnectCommand(const TCHAR* Cmd, FOutputDevice& Ar, UWorld* InWorld);
+	virtual bool HandleTravelCommand(const TCHAR* Cmd, FOutputDevice& Ar, UWorld* InWorld);
 
 	/** Delegate for handling PS4 play together system events */
 	void OnPlayTogetherEventReceived(int32 UserIndex, const TArray<const FUniqueNetId&>& UserList);
diff --git a/Engine/Source/Runtime/Engine/Private/GameInstance.cpp b/Engine/Source/Runtime/Engine/Private/GameInstance.cpp
--- a/Engine/Source/Runtime/Engine/Private/GameInstance.cpp
+++ b/Engine/Source/Runtime/Engine/Private/GameInstance.cpp
@@ -555,6 +555,30 @@ bool UGameInstance::HandleOpenCommand(const TCHAR* Cmd, FOutputDevice& Ar, UWorl
 	return Engine->HandleOpenCommand(Cmd, Ar, InWorld);
 }
 
+bool UGameInstance::HandleDisconnectCommand(const TCHAR* Cmd, FOutputDevice& Ar, UWorld* InWorld)
+{
+	check(WorldContext && WorldContext->World() == InWorld);
+
+	UEngine* const Engine = GetEngine();
+	return Engine->HandleDisconnectCommand(Cmd, Ar, InWorld);
+}
+
+bool UGameInstance::HandleReconnectCommand(const TCHAR* Cmd, FOutputDevice& Ar, UWorld* InWorld)
+{
+	check(WorldContext && WorldContext->World() == InWorld);
+
+	UEngine* const Engine = GetEngine();
+	return Engine->HandleReconnectCommand(Cmd, Ar, InWorld);
+}
+
+bool UGameInstance::HandleTravelCommand(const TCHAR* Cmd, FOutputDevice& Ar, UWorld* InWorld)
+{
+	check(WorldContext && WorldContext->World() == InWorld);
+
+	UEngine* const Engine = GetEngine();
+	return Engine->HandleTravelCommand(Cmd, Ar, InWorld);
+}
+
 bool UGameInstance::Exec(UWorld* InWorld, const TCHAR* Cmd, FOutputDevice& Ar)
 {
 	// @todo a bunch of stuff in UEngine probably belongs here as well
@@ -562,6 +586,18 @@ bool UGameInstance::Exec(UWorld* InWorld, const TCHAR* Cmd, FOutputDevice& Ar)
 	{
 		return HandleOpenCommand(Cmd, Ar, InWorld);
 	}
+	else if (FParse::Command(&Cmd, TEXT("DISCONNECT")))
+	{
+		return HandleDisconnectCommand(Cmd, Ar, InWorld);
+	}
+	else if (FParse::Command(&Cmd, TEXT("RECONNECT")))
+	{
+		return HandleReconnectCommand(Cmd, Ar, InWorld);
+	}
+	else if (FParse::Command(&Cmd, TEXT("TRAVEL")))
+	{
+		return HandleTravelCommand(Cmd, Ar, InWorld);
+	}
 
 	return false;
 }
