Fix for rare crash when the mouse enters or leaves a game viewport while it dosn't have an underlying SceneViewport.

diff --git a/Engine/Source/Runtime/Engine/Private/GameViewportClient.cpp b/Engine/Source/Runtime/Engine/Private/GameViewportClient.cpp
--- a/Engine/Source/Runtime/Engine/Private/GameViewportClient.cpp
+++ b/Engine/Source/Runtime/Engine/Private/GameViewportClient.cpp
@@ -721,7 +721,7 @@ void UGameViewportClient::MouseEnter(FViewport* InViewport, int32 x, int32 y)
 	Super::MouseEnter(InViewport, x, y);
 
 #if PLATFORM_DESKTOP || PLATFORM_HTML5
-	if (InViewport && GetUseMouseForTouch() && !GetGameViewport()->GetPlayInEditorIsSimulate())
+	if (InViewport && GetUseMouseForTouch() && GetGameViewport() && !GetGameViewport()->GetPlayInEditorIsSimulate())
 	{
 		FSlateApplication::Get().SetGameIsFakingTouchEvents(true);
 	}
@@ -753,7 +753,8 @@ void UGameViewportClient::MouseLeave(FViewport* InViewport)
 			InViewport->GetMousePos(LastViewportCursorPos, false);
 
 #if PLATFORM_DESKTOP || PLATFORM_HTML5
-			if (!GetGameViewportWidget()->HasFocusedDescendants())
+			TSharedPtr<class SViewport> ViewportWidget = GetGameViewportWidget();
+			if (ViewportWidget.IsValid() && !ViewportWidget->HasFocusedDescendants())
 			{
 				FVector2D CursorPos(LastViewportCursorPos.X, LastViewportCursorPos.Y);
 				FSlateApplication::Get().SetGameIsFakingTouchEvents(false, &CursorPos);
