Fixed bounds generation issue for skeletal meshes with non-uniform scales. FTransform::Inverse does not always yield correct results with transforms containing non-uniform scales, so convert the transform to a matrix and invert that instead. This is more expensive, but is only performed when caching a new local space transform (when the component transform changes).

diff --git a/Engine/Source/Runtime/Engine/Private/Components/SkeletalMeshComponent.cpp b/Engine/Source/Runtime/Engine/Private/Components/SkeletalMeshComponent.cpp
--- a/Engine/Source/Runtime/Engine/Private/Components/SkeletalMeshComponent.cpp
+++ b/Engine/Source/Runtime/Engine/Private/Components/SkeletalMeshComponent.cpp
@@ -2499,7 +2499,7 @@ FBoxSphereBounds USkeletalMeshComponent::CalcBounds(const FTransform& LocalToWor
 #endif// #if WITH_APEX_CLOTHING
 
 		bCachedLocalBoundsUpToDate = true;
-		CachedLocalBounds = NewBounds.TransformBy(LocalToWorld.Inverse());
+		CachedLocalBounds = NewBounds.TransformBy(LocalToWorld.ToInverseMatrixWithScale());
 
 		return NewBounds;
 	}
