Export some engine functions
PR author MilkyEngineer
PR #5143

diff --git a/Engine/Source/Runtime/Engine/Public/Rendering/SkeletalMeshRenderData.h b/Engine/Source/Runtime/Engine/Public/Rendering/SkeletalMeshRenderData.h
--- a/Engine/Source/Runtime/Engine/Public/Rendering/SkeletalMeshRenderData.h
+++ b/Engine/Source/Runtime/Engine/Public/Rendering/SkeletalMeshRenderData.h
@@ -42,7 +42,7 @@ public:
 	void GetResourceSizeEx(FResourceSizeEx& CumulativeResourceSize);
 
 	/** Returns true if this resource must be skinned on the CPU for the given feature level. */
-	bool RequiresCPUSkinning(ERHIFeatureLevel::Type FeatureLevel) const;
+	ENGINE_API bool RequiresCPUSkinning(ERHIFeatureLevel::Type FeatureLevel) const;
 
 	/** Returns true if there are more than MAX_INFLUENCES_PER_STREAM influences per vertex. */
 	bool HasExtraBoneInfluences() const;
diff --git a/Engine/Source/Runtime/Engine/Public/Rendering/SkinWeightVertexBuffer.h b/Engine/Source/Runtime/Engine/Public/Rendering/SkinWeightVertexBuffer.h
--- a/Engine/Source/Runtime/Engine/Public/Rendering/SkinWeightVertexBuffer.h
+++ b/Engine/Source/Runtime/Engine/Public/Rendering/SkinWeightVertexBuffer.h
@@ -207,7 +207,7 @@ private:
 	uint32 NumVertices;
 
 	/** Allocates the vertex data storage type. */
-	void AllocateData();
+	ENGINE_API void AllocateData();
 
 #if WITH_EDITOR
 	/** Helper for concrete types */
