Fix vertex painting migration issue with old skeletal mesh asset (prior to 4.18)

diff --git a/Engine/Source/Runtime/Engine/Private/SkeletalMeshLODModel.cpp b/Engine/Source/Runtime/Engine/Private/SkeletalMeshLODModel.cpp
--- a/Engine/Source/Runtime/Engine/Private/SkeletalMeshLODModel.cpp
+++ b/Engine/Source/Runtime/Engine/Private/SkeletalMeshLODModel.cpp
@@ -693,6 +693,21 @@ void FSkeletalMeshLODModel::Serialize(FArchive& Ar, UObject* Owner, int32 Idx)
 				{
 					FColorVertexBuffer DummyColorBuffer;
 					DummyColorBuffer.Serialize(Ar, false);
+					//Copy the data to the softVertices
+					int32 VertexColorCount = DummyColorBuffer.GetNumVertices();
+					check(NumVertices == VertexColorCount);
+					TArray<FColor> OutColors;
+					DummyColorBuffer.GetVertexColors(OutColors);
+					int32 DummyVertexColorIndex = 0;
+					for (int32 SectionIndex = 0; SectionIndex < Sections.Num(); ++SectionIndex)
+					{
+						int32 SectionVertexCount = Sections[SectionIndex].GetNumVertices();
+						TArray<FSoftSkinVertex>& SoftVertices = Sections[SectionIndex].SoftVertices;
+						for (int32 SectionVertexIndex = 0; SectionVertexIndex < SectionVertexCount; ++SectionVertexIndex)
+						{
+							SoftVertices[SectionVertexIndex].Color = OutColors[DummyVertexColorIndex++];
+						}
+					}
 				}
 			}
 
diff --git a/Engine/Source/Runtime/Engine/Private/SkeletalMeshLODRenderData.cpp b/Engine/Source/Runtime/Engine/Private/SkeletalMeshLODRenderData.cpp
--- a/Engine/Source/Runtime/Engine/Private/SkeletalMeshLODRenderData.cpp
+++ b/Engine/Source/Runtime/Engine/Private/SkeletalMeshLODRenderData.cpp
@@ -500,7 +500,6 @@ void FSkeletalMeshLODRenderData::BuildFromLODModel(const FSkeletalMeshLODModel*
 	if (bHasVertexColors && Vertices.Num() > 0 && StaticVertexBuffers.ColorVertexBuffer.GetAllocatedSize() == 0)
 	{
 		StaticVertexBuffers.ColorVertexBuffer.InitFromColorArray(&Vertices[0].Color, Vertices.Num(), sizeof(FSoftSkinVertex));
-		StaticVertexBuffers.ColorVertexBuffer.InitFromColorArray(&Vertices[0].Color, Vertices.Num(), sizeof(FSoftSkinVertex));
 	}
 
 	if (ImportedModel->HasClothData())
