Engine and Editor subsystems cleaned up the collection to hide module complexity
Fixed some api that were missed in conversion.

diff --git a/Engine/Plugins/Runtime/GameplayAbilities/Source/GameplayAbilities/Private/AbilitySystemGlobals.cpp b/Engine/Plugins/Runtime/GameplayAbilities/Source/GameplayAbilities/Private/AbilitySystemGlobals.cpp
--- a/Engine/Plugins/Runtime/GameplayAbilities/Source/GameplayAbilities/Private/AbilitySystemGlobals.cpp
+++ b/Engine/Plugins/Runtime/GameplayAbilities/Source/GameplayAbilities/Private/AbilitySystemGlobals.cpp
@@ -340,7 +340,7 @@ void UAbilitySystemGlobals::InitAttributeDefaults()
 #if WITH_EDITOR
 		if (GIsEditor && !RegisteredReimportCallback)
 		{
-			GEditor->OnObjectReimported().AddUObject(this, &UAbilitySystemGlobals::OnTableReimported);
+			GEditor->GetEditorSubsystem<UImportSubsystem>()->OnAssetReimport.AddUObject(this, &UAbilitySystemGlobals::OnTableReimported);
 			RegisteredReimportCallback = true;
 		}
 #endif
diff --git a/Engine/Source/Runtime/Engine/Private/Subsystems/SubsystemCollection.cpp b/Engine/Source/Runtime/Engine/Private/Subsystems/SubsystemCollection.cpp
--- a/Engine/Source/Runtime/Engine/Private/Subsystems/SubsystemCollection.cpp
+++ b/Engine/Source/Runtime/Engine/Private/Subsystems/SubsystemCollection.cpp
@@ -8,8 +8,27 @@
 #include "Modules/ModuleManager.h"
 #include "Misc/PackageName.h"
 
+/** FSubsystemModuleWatcher class to hide the implementation of keeping the DynamicSystemModuleMap up to date*/
+class FSubsystemModuleWatcher
+{
+public:
+	static void OnModulesChanged(FName ModuleThatChanged, EModuleChangeReason ReasonForChange);
+
+	/** Init / Deinit the Module watcher, this tracks module startup and shutdown to ensure only the appropriate dynamic subsystems are instantiated */
+	static void InitializeModuleWatcher();
+	static void DeinitializeModuleWatcher();
+
+private:
+	static void AddClassesForModule(const FName& InModuleName);
+	static void RemoveClassesForModule(const FName& InModuleName);
+
+	static FDelegateHandle ModulesChangedHandle;
+};
+
+FDelegateHandle FSubsystemModuleWatcher::ModulesChangedHandle;
+
+
 
-FDelegateHandle FSubsystemCollectionBase::ModulesChangedHandle;
 TArray<FSubsystemCollectionBase*> FSubsystemCollectionBase::SubsystemCollections;
 TMap<FName, TArray<TSubclassOf<UDynamicSubsystem>>> FSubsystemCollectionBase::DynamicSystemModuleMap;
 
@@ -70,61 +89,6 @@ const TArray<USubsystem*>& FSubsystemCollectionBase::GetSubsystemArrayInternal(T
 	return List;
 }
 
-class FSubsystemModuleWatcher
-{
-public :
-	static void OnModulesChanged(FName ModuleThatChanged, EModuleChangeReason ReasonForChange)
-	{
-
-		switch (ReasonForChange)
-		{
-		case EModuleChangeReason::ModuleLoaded:
-			FSubsystemCollectionBase::AddClassesForModule(ModuleThatChanged);
-			break;
-
-		case EModuleChangeReason::ModuleUnloaded:
-			FSubsystemCollectionBase::RemoveClassesForModule(ModuleThatChanged);
-			break;
-		}
-	}
-};
-
-void FSubsystemCollectionBase::InitializeModuleWatcher()
-{
-	check(!ModulesChangedHandle.IsValid());
-
-	// Add Loaded Modules
-	TArray<UClass*> SubsystemClasses;
-	GetDerivedClasses(UDynamicSubsystem::StaticClass(), SubsystemClasses, true);
-
-	for (UClass* SubsystemClass : SubsystemClasses)
-	{
-		if (!SubsystemClass->HasAllClassFlags(CLASS_Abstract))
-		{
-			UPackage* const ClassPackage = SubsystemClass->GetOuterUPackage();
-			if (ClassPackage)
-			{
-				const FName ModuleName = FPackageName::GetShortFName(ClassPackage->GetFName());
-				if (FModuleManager::Get().IsModuleLoaded(ModuleName))
-				{
-					TArray<TSubclassOf<UDynamicSubsystem>>& ModuleSubsystemClasses = DynamicSystemModuleMap.FindOrAdd(ModuleName);
-					ModuleSubsystemClasses.Add(SubsystemClass);
-				}
-			}
-		}
-	}
-
-	ModulesChangedHandle = FModuleManager::Get().OnModulesChanged().AddStatic(&FSubsystemModuleWatcher::OnModulesChanged);
-}
-
-void FSubsystemCollectionBase::DeinitializeModuleWatcher()
-{
-	if (ModulesChangedHandle.IsValid())
-	{
-		FModuleManager::Get().OnModulesChanged().Remove(ModulesChangedHandle);
-	}
-}
-
 void FSubsystemCollectionBase::Initialize()
 {
 	if (ensure(BaseType) && ensureMsgf(SubsystemMap.Num() == 0, TEXT("Currently don't support repopulation of Subsystem Collections.")))
@@ -134,7 +98,7 @@ void FSubsystemCollectionBase::Initialize()
 		
 		if (SubsystemCollections.Num() == 0)
 		{
-			InitializeModuleWatcher();
+			FSubsystemModuleWatcher::InitializeModuleWatcher();
 		}
 		
 		TGuardValue<bool> PopulatingGuard(bPopulating, true);
@@ -158,7 +122,7 @@ void FSubsystemCollectionBase::Deinitialize()
 	SubsystemCollections.Remove(this);
 	if (SubsystemCollections.Num() == 0)
 	{
-		DeinitializeModuleWatcher();
+		FSubsystemModuleWatcher::DeinitializeModuleWatcher();
 	}
 
 	// Deinit and clean up existing systems
@@ -227,9 +191,89 @@ void FSubsystemCollectionBase::RemoveAndDeinitializeSubsystem(USubsystem* Subsys
 	Subsystem->InternalOwningSubsystem = nullptr;
 }
 
-void FSubsystemCollectionBase::AddClassesForModule(const FName& InModuleName)
+void FSubsystemCollectionBase::AddAllInstances(UClass* SubsystemClass)
+{
+	for (FSubsystemCollectionBase* SubsystemCollection : SubsystemCollections)
+	{
+		if (SubsystemClass->IsChildOf(SubsystemCollection->BaseType))
+		{
+			SubsystemCollection->AddAndInitializeSubsystem(SubsystemClass);
+		}
+	}
+}
+
+void FSubsystemCollectionBase::RemoveAllInstances(UClass* SubsystemClass)
 {
-	check(!DynamicSystemModuleMap.Contains(InModuleName));
+	ForEachObjectOfClass(SubsystemClass, [](UObject* SubsystemObj)
+	{
+		USubsystem* Subsystem = CastChecked<USubsystem>(SubsystemObj);
+
+		if (Subsystem->InternalOwningSubsystem)
+		{
+			Subsystem->InternalOwningSubsystem->RemoveAndDeinitializeSubsystem(Subsystem);
+		}
+	});
+}
+
+
+
+
+/** FSubsystemModuleWatcher Implementations */
+void FSubsystemModuleWatcher::OnModulesChanged(FName ModuleThatChanged, EModuleChangeReason ReasonForChange)
+{
+
+	switch (ReasonForChange)
+	{
+	case EModuleChangeReason::ModuleLoaded:
+		AddClassesForModule(ModuleThatChanged);
+		break;
+
+	case EModuleChangeReason::ModuleUnloaded:
+		RemoveClassesForModule(ModuleThatChanged);
+		break;
+	}
+}
+
+
+void FSubsystemModuleWatcher::InitializeModuleWatcher()
+{
+	check(!ModulesChangedHandle.IsValid());
+
+	// Add Loaded Modules
+	TArray<UClass*> SubsystemClasses;
+	GetDerivedClasses(UDynamicSubsystem::StaticClass(), SubsystemClasses, true);
+
+	for (UClass* SubsystemClass : SubsystemClasses)
+	{
+		if (!SubsystemClass->HasAllClassFlags(CLASS_Abstract))
+		{
+			UPackage* const ClassPackage = SubsystemClass->GetOuterUPackage();
+			if (ClassPackage)
+			{
+				const FName ModuleName = FPackageName::GetShortFName(ClassPackage->GetFName());
+				if (FModuleManager::Get().IsModuleLoaded(ModuleName))
+				{
+					TArray<TSubclassOf<UDynamicSubsystem>>& ModuleSubsystemClasses = FSubsystemCollectionBase::DynamicSystemModuleMap.FindOrAdd(ModuleName);
+					ModuleSubsystemClasses.Add(SubsystemClass);
+				}
+			}
+		}
+	}
+
+	ModulesChangedHandle = FModuleManager::Get().OnModulesChanged().AddStatic(&FSubsystemModuleWatcher::OnModulesChanged);
+}
+
+void FSubsystemModuleWatcher::DeinitializeModuleWatcher()
+{
+	if (ModulesChangedHandle.IsValid())
+	{
+		FModuleManager::Get().OnModulesChanged().Remove(ModulesChangedHandle);
+	}
+}
+
+void FSubsystemModuleWatcher::AddClassesForModule(const FName& InModuleName)
+{
+	check(!FSubsystemCollectionBase::DynamicSystemModuleMap.Contains(InModuleName));
 
 	// Find the class package for this module
 	const UPackage* const ClassPackage = FindPackage(nullptr, *(FString("/Script/") + InModuleName.ToString()));
@@ -252,42 +296,18 @@ void FSubsystemCollectionBase::AddClassesForModule(const FName& InModuleName)
 	}
 	if (SubsystemClasses.Num() > 0)
 	{
-		DynamicSystemModuleMap.Add(InModuleName, MoveTemp(SubsystemClasses));
+		FSubsystemCollectionBase::DynamicSystemModuleMap.Add(InModuleName, MoveTemp(SubsystemClasses));
 	}
 }
-void FSubsystemCollectionBase::RemoveClassesForModule(const FName& InModuleName)
+void FSubsystemModuleWatcher::RemoveClassesForModule(const FName& InModuleName)
 {
-	TArray<TSubclassOf<UDynamicSubsystem>>* SubsystemClasses = DynamicSystemModuleMap.Find(InModuleName);
+	TArray<TSubclassOf<UDynamicSubsystem>>* SubsystemClasses = FSubsystemCollectionBase::DynamicSystemModuleMap.Find(InModuleName);
 	if (SubsystemClasses)
 	{
 		for (TSubclassOf<UDynamicSubsystem>& SubsystemClass : *SubsystemClasses)
 		{
 			FSubsystemCollectionBase::RemoveAllInstances(SubsystemClass);
 		}
-		DynamicSystemModuleMap.Remove(InModuleName);
-	}
-}
-
-void FSubsystemCollectionBase::AddAllInstances(UClass* SubsystemClass)
-{
-	for (FSubsystemCollectionBase* SubsystemCollection : SubsystemCollections)
-	{
-		if (SubsystemClass->IsChildOf(SubsystemCollection->BaseType))
-		{
-			SubsystemCollection->AddAndInitializeSubsystem(SubsystemClass);
-		}
+		FSubsystemCollectionBase::DynamicSystemModuleMap.Remove(InModuleName);
 	}
-}
-
-void FSubsystemCollectionBase::RemoveAllInstances(UClass* SubsystemClass)
-{
-	ForEachObjectOfClass(SubsystemClass, [](UObject* SubsystemObj)
-	{
-		USubsystem* Subsystem = CastChecked<USubsystem>(SubsystemObj);
-
-		if (Subsystem->InternalOwningSubsystem)
-		{
-			Subsystem->InternalOwningSubsystem->RemoveAndDeinitializeSubsystem(Subsystem);
-		}
-	});
-}
+}
\ No newline at end of file
diff --git a/Engine/Source/Runtime/Engine/Public/Subsystems/SubsystemCollection.h b/Engine/Source/Runtime/Engine/Public/Subsystems/SubsystemCollection.h
--- a/Engine/Source/Runtime/Engine/Public/Subsystems/SubsystemCollection.h
+++ b/Engine/Source/Runtime/Engine/Public/Subsystems/SubsystemCollection.h
@@ -57,24 +57,15 @@ private:
 
 	bool bPopulating;
 
-public:
-	/** Init / Deinit the Module watcher, this tracks module startup and shutdown to ensure only the appropriate dynamic subsystems are instantiated */
-	static void InitializeModuleWatcher();
-	static void DeinitializeModuleWatcher();
-
 private:
 	friend class FSubsystemModuleWatcher;
 
-	static void AddClassesForModule(const FName& InModuleName);
-	static void RemoveClassesForModule(const FName& InModuleName);
-
 	/** Add Instances of the specified Subsystem class to all existing SubsystemCollections of the correct type */
 	static void AddAllInstances(UClass* SubsystemClass);
 
 	/** Remove Instances of the specified Subsystem class from all existing SubsystemCollections of the correct type */
 	static void RemoveAllInstances(UClass* SubsystemClass);
 
-	static FDelegateHandle ModulesChangedHandle;
 	static TArray<FSubsystemCollectionBase*> SubsystemCollections;
 	static TMap<FName, TArray<TSubclassOf<UDynamicSubsystem>>> DynamicSystemModuleMap;
 };
