PR #5451: UE-51783: Parameter needs to be passed through to Level (Contributed by [REDACTED_PERSON])

diff --git a/Engine/Source/Runtime/Engine/Private/World.cpp b/Engine/Source/Runtime/Engine/Private/World.cpp
--- a/Engine/Source/Runtime/Engine/Private/World.cpp
+++ b/Engine/Source/Runtime/Engine/Private/World.cpp
@@ -4434,7 +4434,7 @@ AWorldSettings* UWorld::GetWorldSettings( const bool bCheckStreamingPersistent,
 				ULevel* Level = StreamingLevels[0]->GetLoadedLevel();
 				if (Level != nullptr)
 				{
-					WorldSettings = Level->GetWorldSettings();
+					WorldSettings = Level->GetWorldSettings(bChecked);
 				}
 			}
 		}
