Removed some dead code (from UE3)

diff --git a/Engine/Source/Runtime/Engine/Classes/Engine/EngineTypes.h b/Engine/Source/Runtime/Engine/Classes/Engine/EngineTypes.h
--- a/Engine/Source/Runtime/Engine/Classes/Engine/EngineTypes.h
+++ b/Engine/Source/Runtime/Engine/Classes/Engine/EngineTypes.h
@@ -1505,43 +1505,6 @@ struct FSubtitleCue
 };
 
 
-/**	A subtitle localized to a specific language. */
-USTRUCT()
-struct FLocalizedSubtitle
-{
-	GENERATED_USTRUCT_BODY()
-
-	/** The 3-letter language for this subtitle */
-	UPROPERTY()
-	FString LanguageExt;
-
-	/**
-	 * Subtitle cues.  If empty, use SoundNodeWave's SpokenText as the subtitle.  Will often be empty,
-	 * as the contents of the subtitle is commonly identical to what is spoken.
-	 */
-	UPROPERTY()
-	TArray<FSubtitleCue> Subtitles;
-
-	/** true if this sound is considered to contain mature content. */
-	UPROPERTY()
-	uint32 bMature:1;
-
-	/** true if the subtitles have been split manually. */
-	UPROPERTY()
-	uint32 bManualWordWrap:1;
-
-	/** true if the subtitles should be displayed one line at a time. */
-	UPROPERTY()
-	uint32 bSingleLine:1;
-
-	FLocalizedSubtitle()
-		: bMature(false)
-		, bManualWordWrap(false)
-		, bSingleLine(false)
-	{ }
-};
-
-
 /**	Per-light settings for Lightmass */
 USTRUCT()
 struct FLightmassLightSettings
diff --git a/Engine/Source/Runtime/Engine/Classes/Sound/SoundWave.h b/Engine/Source/Runtime/Engine/Classes/Sound/SoundWave.h
--- a/Engine/Source/Runtime/Engine/Classes/Sound/SoundWave.h
+++ b/Engine/Source/Runtime/Engine/Classes/Sound/SoundWave.h
@@ -280,12 +280,6 @@ public:
 
 #endif // WITH_EDITORONLY_DATA
 
-	/**
-	 * The array of the subtitles for each language. Generated at cook time.
-	 */
-	UPROPERTY()
-	TArray<struct FLocalizedSubtitle> LocalizedSubtitles;
-
 #if WITH_EDITORONLY_DATA
 	UPROPERTY()
 	FString SourceFilePath_DEPRECATED;
