Hopefully the Linux Compile fix.

diff --git a/Engine/Source/Runtime/Engine/Public/Subsystems/SubsystemCollection.h b/Engine/Source/Runtime/Engine/Public/Subsystems/SubsystemCollection.h
--- a/Engine/Source/Runtime/Engine/Public/Subsystems/SubsystemCollection.h
+++ b/Engine/Source/Runtime/Engine/Public/Subsystems/SubsystemCollection.h
@@ -2,6 +2,7 @@
 
 #pragma once
 
+#include "Templates/Casts.h"
 #include "Templates/SubclassOf.h"
 #include "UObject/GCObject.h"
 
