Set bHDRSource flag in texture build settings.

diff --git a/Engine/Source/Runtime/Engine/Private/TextureDerivedData.cpp b/Engine/Source/Runtime/Engine/Private/TextureDerivedData.cpp
--- a/Engine/Source/Runtime/Engine/Private/TextureDerivedData.cpp
+++ b/Engine/Source/Runtime/Engine/Private/TextureDerivedData.cpp
@@ -85,6 +85,7 @@ static void SerializeForKey(FArchive& Ar, const FTextureBuildSettings& Settings)
 	TempUint32 = Settings.DiffuseConvolveMipLevel; Ar << TempUint32;
 	TempUint32 = Settings.SharpenMipKernelSize; Ar << TempUint32;
 	// NOTE: TextureFormatName is not stored in the key here.
+	// NOTE: bHDRSource is not stored in the key here.
 	TempByte = Settings.MipGenSettings; Ar << TempByte;
 	TempByte = Settings.bCubemap; Ar << TempByte;
 	TempByte = Settings.bSRGB ? (Settings.bSRGB | ( Settings.bUseLegacyGamma ? 0 : 0x2 )) : 0; Ar << TempByte;
@@ -273,6 +274,7 @@ static void GetTextureBuildSettings(
 	OutBuildSettings.bReplicateRed = false;
 	OutBuildSettings.bVolume = false;
 	OutBuildSettings.bCubemap = false;
+	OutBuildSettings.bHDRSource = Texture.HasHDRSource();
 
 	if (Texture.MaxTextureSize > 0)
 	{
