Prevent log spam from adjacency required by some material feature. (force the option to true in that case)

diff --git a/Engine/Source/Editor/StaticMeshEditor/Private/StaticMeshEditorTools.cpp b/Engine/Source/Editor/StaticMeshEditor/Private/StaticMeshEditorTools.cpp
--- a/Engine/Source/Editor/StaticMeshEditor/Private/StaticMeshEditorTools.cpp
+++ b/Engine/Source/Editor/StaticMeshEditor/Private/StaticMeshEditorTools.cpp
@@ -372,8 +372,9 @@ static void FillEnumOptions(TArray<TSharedPtr<FString> >& OutStrings, UEnum& InE
 	}
 }
 
-FMeshBuildSettingsLayout::FMeshBuildSettingsLayout( TSharedRef<FLevelOfDetailSettingsLayout> InParentLODSettings )
+FMeshBuildSettingsLayout::FMeshBuildSettingsLayout( TSharedRef<FLevelOfDetailSettingsLayout> InParentLODSettings, int32 InLODIndex)
 	: ParentLODSettings( InParentLODSettings )
+	, LODIndex(InLODIndex)
 {
 
 }
@@ -862,6 +863,20 @@ void FMeshBuildSettingsLayout::OnBuildAdjacencyBufferChanged(ECheckBoxState NewS
 			FEngineAnalytics::GetProvider().RecordEvent(TEXT("Editor.Usage.StaticMesh.BuildSettings"), TEXT("bBuildAdjacencyBuffer"), bBuildAdjacencyBuffer ? TEXT("True") : TEXT("False"));
 		}
 		BuildSettings.bBuildAdjacencyBuffer = bBuildAdjacencyBuffer;
+		if (!BuildSettings.bBuildAdjacencyBuffer && ParentLODSettings.IsValid())
+		{
+			if (ParentLODSettings.Pin()->PreviewLODRequiresAdjacencyInformation(LODIndex))
+			{
+				//Prompt the user
+				FText ConfirmRequiredAdjacencyText = LOCTEXT("ConfirmRequiredAdjacencyBufferRemove", "This LOD is using at least one tessellation material that required the adjacency buffer to be computed.\nAre you sure to want to remove the adjacency buffer?");
+				EAppReturnType::Type Result = FMessageDialog::Open(EAppMsgType::YesNo, ConfirmRequiredAdjacencyText);
+				if (Result == EAppReturnType::No)
+				{
+					//Put back the adjacency buffer option to true
+					BuildSettings.bBuildAdjacencyBuffer = true;
+				}
+			}
+		}
 	}
 }
 
@@ -1892,6 +1907,7 @@ void FMeshSectionSettingsLayout::OnSectionChanged(int32 ForLODIndex, int32 Secti
 	FStaticMeshRenderData* RenderData = StaticMesh.RenderData.Get();
 	if (RenderData && RenderData->LODResources.IsValidIndex(LODIndex))
 	{
+		bool bRefreshAll = false;
 		FStaticMeshLODResources& LOD = RenderData->LODResources[LODIndex];
 		if (LOD.Sections.IsValidIndex(SectionIndex))
 		{
@@ -1902,10 +1918,23 @@ void FMeshSectionSettingsLayout::OnSectionChanged(int32 ForLODIndex, int32 Secti
 			FScopedTransaction Transaction(LOCTEXT("StaticMeshOnSectionChangedTransaction", "Staticmesh editor: Section material slot changed"));
 			GetStaticMesh().Modify();
 			FMeshSectionInfo Info = StaticMesh.SectionInfoMap.Get(LODIndex, SectionIndex);
+			int32 CancelOldValue = Info.MaterialIndex;
 			Info.MaterialIndex = NewStaticMaterialIndex;
 			StaticMesh.SectionInfoMap.Set(LODIndex, SectionIndex, Info);
+			bool bUserCancel = false;
+			bRefreshAll = StaticMesh.FixLODRequiresAdjacencyInformation(ForLODIndex, false, true, &bUserCancel);
+			if (bUserCancel)
+			{
+				//Revert the section info map change
+				Info.MaterialIndex = CancelOldValue;
+				StaticMesh.SectionInfoMap.Set(LODIndex, SectionIndex, Info);
+			}
 			CallPostEditChange();
 		}
+		if (bRefreshAll)
+		{
+			StaticMeshEditor.RefreshTool();
+		}
 	}
 }
 
@@ -2444,7 +2473,7 @@ void FMeshMaterialsLayout::OnMaterialChanged(UMaterialInterface* NewMaterial, UM
 {
 	UStaticMesh& StaticMesh = GetStaticMesh();
 	StaticMesh.SetMaterial(MaterialIndex, NewMaterial);
-	StaticMeshEditor.RefreshViewport();
+	StaticMeshEditor.RefreshTool();
 }
 
 TSharedRef<SWidget> FMeshMaterialsLayout::OnGenerateWidgetsForMaterial(UMaterialInterface* Material, int32 SlotIndex)
@@ -3237,7 +3266,7 @@ void FLevelOfDetailSettingsLayout::AddLODLevelCategories( IDetailLayoutBuilder&
 
 				if (StaticMesh->IsMeshDescriptionValid(LODIndex))
 				{
-					BuildSettingsWidgets[LODIndex] = MakeShareable( new FMeshBuildSettingsLayout( AsShared() ) );
+					BuildSettingsWidgets[LODIndex] = MakeShareable( new FMeshBuildSettingsLayout( AsShared(), LODIndex ) );
 					BuildSettingsWidgets[LODIndex]->UpdateSettings(SrcModel.BuildSettings);
 				}
 
@@ -3881,6 +3910,13 @@ void FLevelOfDetailSettingsLayout::ApplyChanges()
 	StaticMeshEditor.RefreshTool();
 }
 
+bool FLevelOfDetailSettingsLayout::PreviewLODRequiresAdjacencyInformation(int32 LODIndex)
+{
+	UStaticMesh* StaticMesh = StaticMeshEditor.GetStaticMesh();
+	check(StaticMesh);
+	return StaticMesh->FixLODRequiresAdjacencyInformation(LODIndex, true, false, nullptr);
+}
+
 FReply FLevelOfDetailSettingsLayout::OnApply()
 {
 	ApplyChanges();
diff --git a/Engine/Source/Editor/StaticMeshEditor/Private/StaticMeshEditorTools.h b/Engine/Source/Editor/StaticMeshEditor/Private/StaticMeshEditorTools.h
--- a/Engine/Source/Editor/StaticMeshEditor/Private/StaticMeshEditorTools.h
+++ b/Engine/Source/Editor/StaticMeshEditor/Private/StaticMeshEditorTools.h
@@ -150,7 +150,7 @@ private:
 class FMeshBuildSettingsLayout : public IDetailCustomNodeBuilder, public TSharedFromThis<FMeshBuildSettingsLayout>
 {
 public:
-	FMeshBuildSettingsLayout( TSharedRef<FLevelOfDetailSettingsLayout> InParentLODSettings );
+	FMeshBuildSettingsLayout( TSharedRef<FLevelOfDetailSettingsLayout> InParentLODSettings, int32 InLODIndex );
 	virtual ~FMeshBuildSettingsLayout();
 
 	const FMeshBuildSettings& GetSettings() const { return BuildSettings; }
@@ -210,6 +210,7 @@ private:
 private:
 	TWeakPtr<FLevelOfDetailSettingsLayout> ParentLODSettings;
 	FMeshBuildSettings BuildSettings;
+	int32 LODIndex;
 };
 
 class FMeshReductionSettingsLayout : public IDetailCustomNodeBuilder, public TSharedFromThis<FMeshReductionSettingsLayout>
@@ -499,6 +500,8 @@ public:
 	/** Apply current LOD settings to the mesh. */
 	void ApplyChanges();
 
+	bool PreviewLODRequiresAdjacencyInformation(int32 LODIndex);
+
 private:
 
 	/** Creates the UI for Current LOD panel */
diff --git a/Engine/Source/Editor/UnrealEd/Private/Fbx/FbxStaticMeshImport.cpp b/Engine/Source/Editor/UnrealEd/Private/Fbx/FbxStaticMeshImport.cpp
--- a/Engine/Source/Editor/UnrealEd/Private/Fbx/FbxStaticMeshImport.cpp
+++ b/Engine/Source/Editor/UnrealEd/Private/Fbx/FbxStaticMeshImport.cpp
@@ -1934,6 +1934,11 @@ void UnFbx::FFbxImporter::PostImportStaticMesh(UStaticMesh* StaticMesh, TArray<F
 	{
 		MeshDescription->TriangulateMesh();
 	}
+	bool bUserCancel = false;
+	if (StaticMesh->FixLODRequiresAdjacencyInformation(LODIndex))
+	{
+		AddTokenizedErrorMessage(FTokenizedMessage::Create(EMessageSeverity::Warning, FText::Format(LOCTEXT("Warning_AdjacencyOptionForced", "Adjacency information not built for static mesh with a material that requires it. Forcing build setting to use adjacency. LOD Index: {0} StaticMesh: {1}"), LODIndex, FText::FromString(StaticMesh->GetPathName()))), FFbxErrors::StaticMesh_AdjacencyOptionForced);
+	}
 
 	//Prebuild the static mesh when we use LodGroup and we want to modify the LodNumber
 	if (!ImportOptions->bImportScene)
diff --git a/Engine/Source/Runtime/Core/Private/Misc/FbxErrors.cpp b/Engine/Source/Runtime/Core/Private/Misc/FbxErrors.cpp
--- a/Engine/Source/Runtime/Core/Private/Misc/FbxErrors.cpp
+++ b/Engine/Source/Runtime/Core/Private/Misc/FbxErrors.cpp
@@ -68,6 +68,8 @@ FName FFbxErrors::StaticMesh_BuildError(TEXT("StaticMesh_BuildError"));
 
 FName FFbxErrors::StaticMesh_AllTrianglesDegenerate(TEXT("StaticMesh_AllTrianglesDegenerate"));
 
+FName FFbxErrors::StaticMesh_AdjacencyOptionForced(TEXT("StaticMesh_AdjacencyOptionForced"));
+
 /** SkeletalMesh **/
 FName FFbxErrors::SkeletalMesh_DifferentRoots(TEXT("SkeletalMesh_DifferentRoot"));
 
diff --git a/Engine/Source/Runtime/Core/Public/Misc/FbxErrors.h b/Engine/Source/Runtime/Core/Public/Misc/FbxErrors.h
--- a/Engine/Source/Runtime/Core/Public/Misc/FbxErrors.h
+++ b/Engine/Source/Runtime/Core/Public/Misc/FbxErrors.h
@@ -80,6 +80,8 @@ struct CORE_API FFbxErrors
 
 	static FName StaticMesh_AllTrianglesDegenerate;
 
+	static FName StaticMesh_AdjacencyOptionForced;
+
 	/** SkeletalMesh **/
 	static FName SkeletalMesh_DifferentRoots;
 
diff --git a/Engine/Source/Runtime/Engine/Classes/Engine/StaticMesh.h b/Engine/Source/Runtime/Engine/Classes/Engine/StaticMesh.h
--- a/Engine/Source/Runtime/Engine/Classes/Engine/StaticMesh.h
+++ b/Engine/Source/Runtime/Engine/Classes/Engine/StaticMesh.h
@@ -885,6 +885,18 @@ public:
 	ENGINE_API void SetNumSourceModels(int32 Num);
 	ENGINE_API void RemoveSourceModel(int32 Index);
 
+	/*
+	 * Verify that a specific LOD using a material needing the adjacency buffer have the build option set to create the adjacency buffer.
+	 *
+	 * LODIndex: The LOD to fix
+	 * bPreviewMode: If true the the function will not fix the build option. It will also change the return behavior, return true if the LOD need adjacency buffer, false otherwise
+	 * bPromptUser: if true a dialog will ask the user if he agree changing the build option to allow adjacency buffer
+	 * OutUserCancel: if the value is not null and the bPromptUser is true, the prompt dialog will have a cancel button and the result will be put in the parameter.
+	 *
+	 * The function will return true if any LOD build settings option is fix to add adjacency option. It will return false if no action was done. In case bPreviewMode is true it return true if the LOD need adjacency buffer, false otherwise.
+	 */
+	ENGINE_API bool FixLODRequiresAdjacencyInformation(const int32 LODIndex, const bool bPreviewMode = false, bool bPromptUser = false, bool* OutUserCancel = nullptr);
+	
 #endif // WITH_EDITOR
 
 	ENGINE_API virtual void Serialize(FArchive& Ar) override;
diff --git a/Engine/Source/Runtime/Engine/Private/StaticMesh.cpp b/Engine/Source/Runtime/Engine/Private/StaticMesh.cpp
--- a/Engine/Source/Runtime/Engine/Private/StaticMesh.cpp
+++ b/Engine/Source/Runtime/Engine/Private/StaticMesh.cpp
@@ -54,6 +54,8 @@
 #include "MeshDescriptionOperations.h"
 #include "IMeshReductionManagerModule.h"
 #include "IMeshReductionInterfaces.h"
+#include "TessellationRendering.h"
+#include "Misc/MessageDialog.h"
 
 #endif // #if WITH_EDITOR
 
@@ -2407,6 +2409,79 @@ void UStaticMesh::RemoveSourceModel(const int32 Index)
 	SourceModels.RemoveAt(Index);
 }
 
+bool UStaticMesh::FixLODRequiresAdjacencyInformation(const int32 LODIndex, const bool bPreviewMode, bool bPromptUser, bool* OutUserCancel)
+{
+	if (OutUserCancel != nullptr)
+	{
+		*OutUserCancel = false;
+	}
+
+	bool bIsUnattended = FApp::IsUnattended() == true || GIsRunningUnattendedScript || GIsAutomationTesting;
+	//Cannot prompt user in unattended mode
+	if (!SourceModels.IsValidIndex(LODIndex) || (bIsUnattended && bPromptUser))
+	{
+		return false;
+	}
+	FStaticMeshSourceModel& SourceModel = SourceModels[LODIndex];
+	FMeshDescription* MeshDescription = GetMeshDescription(LODIndex);
+	//In preview mode we simulate a false BuildAdjacencyBuffer
+	if (MeshDescription && (!(SourceModel.BuildSettings.bBuildAdjacencyBuffer) || bPreviewMode))
+	{
+		TPolygonGroupAttributesRef<FName> PolygonGroupImportedMaterialSlotNames = MeshDescription->PolygonGroupAttributes().GetAttributesRef<FName>(MeshAttribute::PolygonGroup::ImportedMaterialSlotName);
+		int32 SectionIndex = 0;
+		
+		for (const FPolygonGroupID& PolygonGroupID : MeshDescription->PolygonGroups().GetElementIDs())
+		{
+			const FName MaterialImportedName = PolygonGroupImportedMaterialSlotNames[PolygonGroupID];
+			int32 MaterialIndex = 0;
+			for (FStaticMaterial& Material : StaticMaterials)
+			{
+				if (Material.ImportedMaterialSlotName != NAME_None && Material.ImportedMaterialSlotName == MaterialImportedName)
+				{
+					FStaticMaterial *RemapMaterial = &Material;
+					FMeshSectionInfo SectionInfo = SectionInfoMap.Get(LODIndex, SectionIndex);
+					if (StaticMaterials.IsValidIndex(SectionInfo.MaterialIndex))
+					{
+						RemapMaterial = &StaticMaterials[SectionInfo.MaterialIndex];
+					}
+					const bool bRequiresAdjacencyInformation = RequiresAdjacencyInformation(RemapMaterial->MaterialInterface, nullptr, GWorld->FeatureLevel);
+					if (bRequiresAdjacencyInformation)
+					{
+						if (bPromptUser)
+						{
+							FText ConfirmRequiredAdjacencyText = FText::Format(LOCTEXT("ConfirmRequiredAdjacency", "Using a tessellation material required the adjacency buffer to be computed.\nDo you want to set the adjacency options to true?\n\n\tSaticMesh: {0}\n\tLOD Index: {1}\n\tMaterial: {2}"), FText::FromString(GetPathName()), LODIndex, FText::FromString(RemapMaterial->MaterialInterface->GetPathName()));
+							EAppReturnType::Type Result = FMessageDialog::Open((OutUserCancel != nullptr) ? EAppMsgType::YesNoCancel : EAppMsgType::YesNo, ConfirmRequiredAdjacencyText);
+							switch(Result)
+							{
+								//Handle cancel and negative answer
+								case EAppReturnType::Cancel:
+								{
+									check(OutUserCancel != nullptr);
+									*OutUserCancel = true;
+									return false;
+								}
+								case EAppReturnType::No:
+								{
+									return false;
+								}
+							}
+						}
+						if (!bPreviewMode)
+						{
+							UE_LOG(LogStaticMesh, Warning, TEXT("Adjacency information not built for static mesh with a material that requires it. Forcing build setting to use adjacency.\n\tLOD Index: %d\n\tMaterial: %s\n\tStaticMesh: %s"), LODIndex, *RemapMaterial->MaterialInterface->GetPathName(), *GetPathName());
+							SourceModel.BuildSettings.bBuildAdjacencyBuffer = true;
+						}
+						return true;
+					}
+				}
+				MaterialIndex++;
+			}
+			SectionIndex++;
+		}
+	}
+	return false;
+}
+
 #endif // WITH_EDITOR
 
 void UStaticMesh::BeginDestroy()
@@ -4628,7 +4703,7 @@ void UStaticMesh::SetMaterial(int32 MaterialIndex, UMaterialInterface* NewMateri
 		UProperty* ChangedProperty = FindField<UProperty>(UStaticMesh::StaticClass(), NAME_StaticMaterials);
 		check(ChangedProperty);
 		PreEditChange(ChangedProperty);
-
+		UMaterialInterface* CancelOldMaterial = StaticMaterials[MaterialIndex].MaterialInterface;
 		StaticMaterials[MaterialIndex].MaterialInterface = NewMaterial;
 		if (NewMaterial != nullptr)
 		{
@@ -4673,6 +4748,48 @@ void UStaticMesh::SetMaterial(int32 MaterialIndex, UMaterialInterface* NewMateri
 				}
 				StaticMaterials[MaterialIndex].ImportedMaterialSlotName = FName(*MaterialSlotName);
 			}
+
+			//Make sure adjacency information fit new material change
+			TArray<bool> FixLODAdjacencyOption;
+			FixLODAdjacencyOption.AddZeroed(GetNumLODs());
+			bool bPromptUser = false;
+			for (int32 LODIndex = 0; LODIndex < GetNumLODs(); ++LODIndex)
+			{
+				FixLODAdjacencyOption[LODIndex] = FixLODRequiresAdjacencyInformation(LODIndex);
+				bPromptUser |= FixLODAdjacencyOption[LODIndex];
+			}
+
+			//Prompt the user only once
+			if (bPromptUser)
+			{
+				FText ConfirmRequiredAdjacencyText = FText::Format(LOCTEXT("ConfirmRequiredAdjacencyNoLODIndex", "Using a tessellation material required the adjacency buffer to be computed.\nDo you want to set the adjacency options to true?\n\n\tSaticMesh: {0}\n\tMaterial: {1}"), FText::FromString(GetPathName()), FText::FromString(StaticMaterials[MaterialIndex].MaterialInterface->GetPathName()));
+				EAppReturnType::Type Result = FMessageDialog::Open(EAppMsgType::YesNoCancel, ConfirmRequiredAdjacencyText);
+				bool bRevertAdjacency = false;
+				switch(Result)
+				{
+					//Handle cancel and negative answer
+					case EAppReturnType::Cancel:
+					{
+						StaticMaterials[MaterialIndex].MaterialInterface = CancelOldMaterial;
+						bRevertAdjacency = true;
+					}
+					case EAppReturnType::No:
+					{
+						bRevertAdjacency = true;
+					}
+				}
+				if (bRevertAdjacency)
+				{
+					//Revert previous change since the material was reverse
+					for (int32 FixLODIndex = 0; FixLODIndex < FixLODAdjacencyOption.Num(); ++FixLODIndex)
+					{
+						if (FixLODAdjacencyOption[FixLODIndex])
+						{
+							SourceModels[FixLODIndex].BuildSettings.bBuildAdjacencyBuffer = false;
+						}
+					}
+				}
+			}
 		}
 
 		if (ChangedProperty)
