Added a new serializable TStructOnScope which is a Type restricted FStructOnScope
Modified FActorComponentInstanceData to all be USTRUCT
Modified FActorComponentInstanceDataCache to be serializable so it can persisted or sent over the wire

Not all ComponentInstanceData properties has been properly made into UPROPERTY yet however

diff --git a/Engine/Plugins/Experimental/ControlRig/Source/ControlRig/Private/ControlRigComponent.cpp b/Engine/Plugins/Experimental/ControlRig/Source/ControlRig/Private/ControlRigComponent.cpp
--- a/Engine/Plugins/Experimental/ControlRig/Source/ControlRig/Private/ControlRigComponent.cpp
+++ b/Engine/Plugins/Experimental/ControlRig/Source/ControlRig/Private/ControlRigComponent.cpp
@@ -12,52 +12,29 @@
 #endif
 ////////////////////////////////////////////////////////////////////////////////////////
 
-/** Used to store animation ControlRig data during recompile of BP */
-class FControlRigComponentInstanceData : public FActorComponentInstanceData
+void FControlRigComponentInstanceData::ApplyToComponent(UActorComponent* Component, const ECacheApplyPhase CacheApplyPhase) 
 {
-public:
-	FControlRigComponentInstanceData(const UControlRigComponent* SourceComponent)
-		: FActorComponentInstanceData(SourceComponent)
-		, AnimControlRig(SourceComponent->ControlRig)
-	{}
+	FActorComponentInstanceData::ApplyToComponent(Component, CacheApplyPhase);
+	UControlRigComponent* NewComponent = CastChecked<UControlRigComponent>(Component);
 
-	virtual void ApplyToComponent(UActorComponent* Component, const ECacheApplyPhase CacheApplyPhase) override
+	UControlRig* NewControlRig = NewComponent->ControlRig;
+	if (NewControlRig && AnimControlRig)
 	{
-		FActorComponentInstanceData::ApplyToComponent(Component, CacheApplyPhase);
-		UControlRigComponent* NewComponent = CastChecked<UControlRigComponent>(Component);
-
-		UControlRig* NewControlRig = NewComponent->ControlRig;
-		if (NewControlRig && AnimControlRig)
-		{
-			// it will just copy same property if not same class
-			TArray<uint8> SavedPropertyBuffer;
-			FObjectWriter Writer(AnimControlRig, SavedPropertyBuffer);
-			FObjectReader Reader(NewComponent->ControlRig, SavedPropertyBuffer);
-		}
-	}
-
-	virtual void FindAndReplaceInstances(const TMap<UObject*, UObject*>& OldToNewInstanceMap) override
-	{
-		FActorComponentInstanceData::FindAndReplaceInstances(OldToNewInstanceMap);
-	}
-
-	virtual void AddReferencedObjects(FReferenceCollector& Collector) override
-	{
-		FActorComponentInstanceData::AddReferencedObjects(Collector);
-		if (AnimControlRig)
-		{
-			Collector.AddReferencedObject(AnimControlRig);
-		}
+		// it will just copy same property if not same class
+		TArray<uint8> SavedPropertyBuffer;
+		FObjectWriter Writer(AnimControlRig, SavedPropertyBuffer);
+		FObjectReader Reader(NewComponent->ControlRig, SavedPropertyBuffer);
 	}
+}
 
-	bool ContainsData() const
+void FControlRigComponentInstanceData::AddReferencedObjects(FReferenceCollector& Collector)
+{
+	FActorComponentInstanceData::AddReferencedObjects(Collector);
+	if (AnimControlRig)
 	{
-		return (AnimControlRig != nullptr);
+		Collector.AddReferencedObject(AnimControlRig);
 	}
-
-	// stored object
-	UControlRig* AnimControlRig;
-};
+}
 
 ////////////////////////////////////////////////////////////////////////////////////////
 UControlRigComponent::UControlRigComponent(const FObjectInitializer& ObjectInitializer)
@@ -147,15 +124,7 @@ void UControlRigComponent::OnPostEvaluate_Implementation()
 	OnPostEvaluateDelegate.Broadcast(this);
 }
 
-FActorComponentInstanceData* UControlRigComponent::GetComponentInstanceData() const
+TStructOnScope<FActorComponentInstanceData> UControlRigComponent::GetComponentInstanceData() const
 {
-	FControlRigComponentInstanceData* InstanceData = new FControlRigComponentInstanceData(this);
-
-	if (!InstanceData->ContainsData())
-	{
-		delete InstanceData;
-		InstanceData = nullptr;
-	}
-
-	return InstanceData;
+	return MakeStructOnScope<FActorComponentInstanceData, FControlRigComponentInstanceData>(this);
 }
diff --git a/Engine/Plugins/Experimental/ControlRig/Source/ControlRig/Public/ControlRigComponent.h b/Engine/Plugins/Experimental/ControlRig/Source/ControlRig/Public/ControlRigComponent.h
--- a/Engine/Plugins/Experimental/ControlRig/Source/ControlRig/Public/ControlRigComponent.h
+++ b/Engine/Plugins/Experimental/ControlRig/Source/ControlRig/Public/ControlRigComponent.h
@@ -48,7 +48,7 @@ public:
 	virtual void OnRegister() override;
 	virtual void OnUnregister() override;
 	virtual void TickComponent(float DeltaTime, enum ELevelTick TickType, FActorComponentTickFunction *ThisTickFunction) override;
-	virtual class FActorComponentInstanceData* GetComponentInstanceData() const override;
+	virtual TStructOnScope<FActorComponentInstanceData> GetComponentInstanceData() const override;
 
 	/** Get the ControlRig hosted by this component */
 	UFUNCTION(BlueprintPure, Category = "ControlRig", meta = (DisplayName = "Get ControlRig", ScriptName = "GetControlRig"))
@@ -74,5 +74,33 @@ public:
 	}
 
 private:
-	friend class FControlRigComponentInstanceData;
-};
\ No newline at end of file
+	friend struct FControlRigComponentInstanceData;
+};
+
+/** Used to store animation ControlRig data during recompile of BP */
+USTRUCT()
+struct FControlRigComponentInstanceData : public FActorComponentInstanceData
+{
+	GENERATED_BODY()
+public:
+	FControlRigComponentInstanceData()
+		: AnimControlRig(nullptr)
+	{}
+	FControlRigComponentInstanceData(const UControlRigComponent* SourceComponent)
+		: FActorComponentInstanceData(SourceComponent)
+		, AnimControlRig(SourceComponent->ControlRig)
+	{}
+	virtual ~FControlRigComponentInstanceData() = default;
+
+	virtual bool ContainsData() const override
+	{
+		return (AnimControlRig != nullptr) || Super::ContainsData();
+	}
+
+	virtual void ApplyToComponent(UActorComponent* Component, const ECacheApplyPhase CacheApplyPhase) override;
+	virtual void AddReferencedObjects(FReferenceCollector& Collector) override;
+
+	// stored object
+	UPROPERTY()
+	UControlRig* AnimControlRig;
+};
diff --git a/Engine/Source/Runtime/CoreUObject/Public/UObject/StructOnScope.h b/Engine/Source/Runtime/CoreUObject/Public/UObject/StructOnScope.h
--- a/Engine/Source/Runtime/CoreUObject/Public/UObject/StructOnScope.h
+++ b/Engine/Source/Runtime/CoreUObject/Public/UObject/StructOnScope.h
@@ -7,6 +7,7 @@
 #include "UObject/WeakObjectPtr.h"
 #include "UObject/Package.h"
 #include "Templates/Casts.h"
+#include "Templates/RemoveReference.h"
 
 class FStructOnScope
 {
@@ -159,3 +160,163 @@ public:
 		Initialize();
 	}
 };
+
+/**
+ * Typed FStructOnScope that exposes type-safe access to the wrapped struct
+ * @note The second template argument is there to restrict to type that are actually USTRUCT.
+ *		It will be replaced to a static_assert with a better compiler error with a "is ustruct" type trait
+ */
+template<typename T, typename = decltype(TBaseStructure<T>::Get())>
+class TStructOnScope final : public FStructOnScope
+{
+public:
+	TStructOnScope() = default;
+	virtual ~TStructOnScope() = default;
+
+	TStructOnScope(TStructOnScope&& InOther) = default;
+	TStructOnScope& operator=(TStructOnScope&& InOther) = default;
+
+	template<typename U, typename = typename TEnableIf<TIsDerivedFrom<typename TRemoveReference<U>::Type, T>::IsDerived, void>::Type>
+	explicit TStructOnScope(U&& InStruct)
+		: FStructOnScope(TBaseStructure<typename TRemoveReference<U>::Type>::Get())
+	{
+		if (const UScriptStruct* ScriptStructPtr = ::Cast<UScriptStruct>(ScriptStruct.Get()))
+		{
+			ScriptStructPtr->CopyScriptStruct(SampleStructMemory, &InStruct);
+		}
+	}
+
+	template<typename U, typename = typename TEnableIf<TIsDerivedFrom<typename TRemoveReference<U>::Type, T>::IsDerived, void>::Type>
+	TStructOnScope& operator=(U&& InStruct)
+	{
+		Initialize(TBaseStructure<typename TRemoveReference<U>::Type>::Get());
+		if (const UScriptStruct* ScriptStructPtr = ::Cast<UScriptStruct>(ScriptStruct.Get()))
+		{
+			ScriptStructPtr->CopyScriptStruct(SampleStructMemory, &InStruct);
+		}
+		return *this;
+	}
+
+	/**
+	 * Initialize the TStructOnScope as a struct of type U which needs to derive from T
+	 * @params InArgs The arguments to pass to the constructor of type U
+	 */
+	template<typename U
+		, typename = typename TEnableIf<TIsDerivedFrom<typename TRemoveReference<U>::Type, T>::IsDerived, void>::Type
+		, typename... TArgs>
+	void InitializeAs(TArgs&&... InArgs)
+	{
+		Destroy();
+		if (UScriptStruct* ScriptStructPtr = TBaseStructure<U>::Get())
+		{
+			ScriptStruct = ScriptStructPtr;
+			SampleStructMemory = (uint8*)FMemory::Malloc(ScriptStructPtr->GetStructureSize() ? ScriptStructPtr->GetStructureSize() : 1);
+			new (SampleStructMemory) U(Forward<TArgs>(InArgs)...);
+			OwnsMemory = true;
+		}
+	}
+
+	T* Get() const
+	{
+		return reinterpret_cast<T*>(SampleStructMemory);
+	}
+
+	T* operator->() const
+	{
+		return Get();
+	}
+
+	operator bool() const
+	{
+		return IsValid();
+	}
+
+	template<typename U>
+	U* Cast()
+	{
+		if (TBaseStructure<U>::Get()->IsChildOf(ScriptStruct.Get()))
+		{
+			return reinterpret_cast<U*>(SampleStructMemory);
+		}
+		return nullptr;
+	}
+
+	template<typename U>
+	U* CastChecked()
+	{
+		U* Result = nullptr;
+		if (!IsValid())
+		{
+			UE_LOG(LogClass, Fatal, TEXT("Cast of nullptr to %s failed"), *GetTypeName<U>());
+			return Result;
+		}
+
+		Result = Cast<U>();
+		if (!Result)
+		{
+			UE_LOG(LogClass, Fatal, TEXT("Cast of %s to %s failed"), *GetTypeName<T>(), *GetTypeName<U>());
+		}
+		return Result;
+	}
+
+	friend FArchive& operator<<(FArchive& Ar, TStructOnScope& InStruct)
+	{
+		if (Ar.IsLoading())
+		{
+			FString StructPath;
+			Ar << StructPath;
+			if (!StructPath.IsEmpty())
+			{
+				UScriptStruct* ScriptStructPtr = FindObject<UScriptStruct>(nullptr, *StructPath, false);
+				if (ScriptStructPtr == nullptr || !ScriptStructPtr->IsChildOf(TBaseStructure<T>::Get()))
+				{
+					Ar.SetError();
+					return Ar;
+				}
+				InStruct.ScriptStruct = ScriptStructPtr;
+				InStruct.Initialize();
+				ScriptStructPtr->SerializeItem(Ar, InStruct.SampleStructMemory, nullptr);
+			}
+		}
+		// Saving
+		else
+		{
+			FString StructPath;
+			if (UScriptStruct* ScriptStructPtr = const_cast<UScriptStruct*>(::Cast<UScriptStruct>(InStruct.ScriptStruct.Get())))
+			{
+				StructPath = ScriptStructPtr->GetPathName();
+				Ar << StructPath;
+				ScriptStructPtr->SerializeItem(Ar, InStruct.SampleStructMemory, nullptr);
+			}
+			else
+			{
+				Ar << StructPath;
+			}
+		}
+		
+		return Ar;
+	}
+
+private:
+	using FStructOnScope::OwnsStructMemory;
+	using FStructOnScope::GetStructMemory;	
+	using FStructOnScope::GetPackage;
+	using FStructOnScope::SetPackage;
+	using FStructOnScope::Destroy;
+	using FStructOnScope::Initialize;
+};
+
+/**
+ * Allocates a new struct of type U with the given arguments and returns it as a Typed StructOnScope of type T.
+ *
+ * @param Args The arguments to pass to the constructor of U.
+ *
+ * @return A TStructOnScope<T> which holds a struct to a newly-constructed U with the specified Args.
+ */
+template <typename T, typename U = T, typename... TArgs>
+FORCEINLINE TStructOnScope<T> MakeStructOnScope(TArgs&&... Args)
+{
+	TStructOnScope<T> Struct;
+	Struct.template InitializeAs<U>(Forward<TArgs>(Args)...);
+	return Struct;
+}
diff --git a/Engine/Source/Runtime/Engine/Classes/Atmosphere/AtmosphericFogComponent.h b/Engine/Source/Runtime/Engine/Classes/Atmosphere/AtmosphericFogComponent.h
--- a/Engine/Source/Runtime/Engine/Classes/Atmosphere/AtmosphericFogComponent.h
+++ b/Engine/Source/Runtime/Engine/Classes/Atmosphere/AtmosphericFogComponent.h
@@ -269,10 +269,10 @@ public:
 	ENGINE_API void ReleaseResource();
 
 	//~ Begin UActorComponent Interface.
-	virtual FActorComponentInstanceData* GetComponentInstanceData() const override;
+	virtual TStructOnScope<FActorComponentInstanceData> GetComponentInstanceData() const override;
 	//~ End UActorComponent Interface.
 
-	void ApplyComponentInstanceData(class FAtmospherePrecomputeInstanceData* ComponentInstanceData);
+	void ApplyComponentInstanceData(struct FAtmospherePrecomputeInstanceData* ComponentInstanceData);
 	const FAtmospherePrecomputeParameters& GetPrecomputeParameters() const { return PrecomputeParams;  }
 
 private:
@@ -283,3 +283,33 @@ private:
 	friend class FAtmosphericFogSceneInfo;
 };
 
+/** Used to store lightmap data during RerunConstructionScripts */
+USTRUCT()
+struct FAtmospherePrecomputeInstanceData : public FSceneComponentInstanceData
+{
+	GENERATED_BODY()
+public:
+	FAtmospherePrecomputeInstanceData() = default;
+	FAtmospherePrecomputeInstanceData(const UAtmosphericFogComponent* SourceComponent)
+		: FSceneComponentInstanceData(SourceComponent)
+	{}
+
+	virtual ~FAtmospherePrecomputeInstanceData() = default;
+
+	virtual bool ContainsData() const override
+	{
+		return true;
+	}
+
+	virtual void ApplyToComponent(UActorComponent* Component, const ECacheApplyPhase CacheApplyPhase) override
+	{
+		FSceneComponentInstanceData::ApplyToComponent(Component, CacheApplyPhase);
+		CastChecked<UAtmosphericFogComponent>(Component)->ApplyComponentInstanceData(this);
+	}
+
+	struct FAtmospherePrecomputeParameters PrecomputeParameter;
+
+	FByteBulkData TransmittanceData;
+	FByteBulkData IrradianceData;
+	FByteBulkData InscatterData;
+};
diff --git a/Engine/Source/Runtime/Engine/Classes/Components/ActorComponent.h b/Engine/Source/Runtime/Engine/Classes/Components/ActorComponent.h
--- a/Engine/Source/Runtime/Engine/Classes/Components/ActorComponent.h
+++ b/Engine/Source/Runtime/Engine/Classes/Components/ActorComponent.h
@@ -14,6 +14,8 @@
 #include "EdGraph/EdGraphPin.h"
 #include "Interfaces/Interface_AssetUserData.h"
 #include "UObject/UObjectAnnotation.h"
+#include "UObject/StructOnScope.h"
+#include "ComponentInstanceDataCache.h"
 #include "ActorComponent.generated.h"
 
 class AActor;
@@ -29,20 +31,6 @@ struct FMinimalViewInfo;
 extern ENGINE_API FUObjectAnnotationSparseBool GSelectedComponentAnnotation;
 #endif
 
-UENUM()
-enum class EComponentCreationMethod : uint8
-{
-	/** A component that is part of a native class. */
-	Native,
-	/** A component that is created from a template defined in the Components section of the Blueprint. */
-	SimpleConstructionScript,	
-	/**A dynamically created component, either from the UserConstructionScript or from a Add Component node in a Blueprint event graph. */
-	UserConstructionScript,
-	/** A component added to a single Actor instance via the Component section of the Actor's details panel. */
-	Instance,
-};
-
-
 /** Information about how to update transform*/
 enum class EUpdateTransformFlags : int32
 {
@@ -256,7 +244,7 @@ private:
 	uint8 MarkedForEndOfFrameUpdateState:2;
 	friend struct FMarkComponentEndOfFrameUpdateState;
 
-	friend class FActorComponentInstanceData;
+	friend struct FActorComponentInstanceData;
 	friend class FActorComponentDetails;
 
 public:
@@ -791,7 +779,7 @@ public:
 	virtual void PostNetReceive() override { }
 
 	/** Called before we throw away components during RerunConstructionScripts, to cache any data we wish to persist across that operation */
-	virtual class FActorComponentInstanceData* GetComponentInstanceData() const;
+	virtual TStructOnScope<FActorComponentInstanceData> GetComponentInstanceData() const;
 
 	//~ Begin UObject Interface.
 	virtual void BeginDestroy() override;
diff --git a/Engine/Source/Runtime/Engine/Classes/Components/ChildActorComponent.h b/Engine/Source/Runtime/Engine/Classes/Components/ChildActorComponent.h
--- a/Engine/Source/Runtime/Engine/Classes/Components/ChildActorComponent.h
+++ b/Engine/Source/Runtime/Engine/Classes/Components/ChildActorComponent.h
@@ -11,32 +11,44 @@
 #include "ChildActorComponent.generated.h"
 
 class AActor;
-struct FAttachedActorInfo;
 
-class ENGINE_API FChildActorComponentInstanceData : public FSceneComponentInstanceData
+USTRUCT()
+struct FChildActorAttachedActorInfo
 {
+	GENERATED_BODY()
+
+	UPROPERTY()
+	TWeakObjectPtr<AActor> Actor;
+	UPROPERTY()
+	FName SocketName;
+	UPROPERTY()
+	FTransform RelativeTransform;
+};
+
+USTRUCT()
+struct ENGINE_API FChildActorComponentInstanceData : public FSceneComponentInstanceData
+{
+	GENERATED_BODY()
 public:
+	FChildActorComponentInstanceData() = default;
 	FChildActorComponentInstanceData(const class UChildActorComponent* Component);
 
-	virtual ~FChildActorComponentInstanceData();
+	virtual ~FChildActorComponentInstanceData() = default;
+
+	virtual bool ContainsData() const override;
 
 	virtual void ApplyToComponent(UActorComponent* Component, const ECacheApplyPhase CacheApplyPhase) override;
 	virtual void AddReferencedObjects(FReferenceCollector& Collector) override;
 
 	// The name of the spawned child actor so it (attempts to) remain constant across construction script reruns
+	UPROPERTY()
 	FName ChildActorName;
 
-	// The component instance data cache for the ChildActor spawned by this component
-	FComponentInstanceDataCache* ComponentInstanceData;
-
-	struct FAttachedActorInfo
-	{
-		TWeakObjectPtr<AActor> Actor;
-		FName SocketName;
-		FTransform RelativeTransform;
-	};
+	UPROPERTY()
+	TArray<FChildActorAttachedActorInfo> AttachedActors;
 
-	TArray<FAttachedActorInfo> AttachedActors;
+	// The component instance data cache for the ChildActor spawned by this component
+	TSharedPtr<FComponentInstanceDataCache> ComponentInstanceData;
 };
 
 /** A component that spawns an Actor when registered, and destroys it when unregistered.*/
@@ -95,12 +107,12 @@ public:
 	virtual void OnComponentDestroyed(bool bDestroyingHierarchy) override;
 	virtual void OnRegister() override;
 	virtual void OnUnregister() override;
-	virtual FActorComponentInstanceData* GetComponentInstanceData() const override;
+	virtual TStructOnScope<FActorComponentInstanceData> GetComponentInstanceData() const override;
 	virtual void BeginPlay() override;
 	//~ End ActorComponent Interface.
 
 	/** Apply the component instance data to the child actor component */
-	void ApplyComponentInstanceData(class FChildActorComponentInstanceData* ComponentInstanceData, const ECacheApplyPhase CacheApplyPhase);
+	void ApplyComponentInstanceData(struct FChildActorComponentInstanceData* ComponentInstanceData, const ECacheApplyPhase CacheApplyPhase);
 
 	/** Create the child actor */
 	virtual void CreateChildActor();
diff --git a/Engine/Source/Runtime/Engine/Classes/Components/HierarchicalInstancedStaticMeshComponent.h b/Engine/Source/Runtime/Engine/Classes/Components/HierarchicalInstancedStaticMeshComponent.h
--- a/Engine/Source/Runtime/Engine/Classes/Components/HierarchicalInstancedStaticMeshComponent.h
+++ b/Engine/Source/Runtime/Engine/Classes/Components/HierarchicalInstancedStaticMeshComponent.h
@@ -237,7 +237,7 @@ public:
 	/** Heuristic for the number of leaves in the tree **/
 	int32 DesiredInstancesPerLeaf();
 
-	virtual void ApplyComponentInstanceData(class FInstancedStaticMeshComponentInstanceData* InstancedMeshData) override;
+	virtual void ApplyComponentInstanceData(struct FInstancedStaticMeshComponentInstanceData* InstancedMeshData) override;
 	
 	// Number of instances in the render-side instance buffer
 	virtual int32 GetNumRenderInstances() const { return SortedInstances.Num(); }
diff --git a/Engine/Source/Runtime/Engine/Classes/Components/InstancedStaticMeshComponent.h b/Engine/Source/Runtime/Engine/Classes/Components/InstancedStaticMeshComponent.h
--- a/Engine/Source/Runtime/Engine/Classes/Components/InstancedStaticMeshComponent.h
+++ b/Engine/Source/Runtime/Engine/Classes/Components/InstancedStaticMeshComponent.h
@@ -240,7 +240,7 @@ public:
 	TArray<FBodyInstance*> InstanceBodies;
 
 	//~ Begin UActorComponent Interface
-	virtual FActorComponentInstanceData* GetComponentInstanceData() const override;
+	virtual TStructOnScope<FActorComponentInstanceData> GetComponentInstanceData() const override;
 	//~ End UActorComponent Interface
 
 	//~ Begin UPrimitiveComponent Interface
@@ -277,7 +277,7 @@ public:
 	//~ End UObject Interface
 
 	/** Applies the cached component instance data to a newly blueprint constructed component. */
-	virtual void ApplyComponentInstanceData(class FInstancedStaticMeshComponentInstanceData* ComponentInstanceData);
+	virtual void ApplyComponentInstanceData(struct FInstancedStaticMeshComponentInstanceData* ComponentInstanceData);
 
 	/** Check to see if an instance is selected. */
 	bool IsInstanceSelected(int32 InInstanceIndex) const;
@@ -371,3 +371,67 @@ struct HInstancedStaticMeshInstance : public HHitProxy
 		return EMouseCursor::CardinalCross;
 	}
 };
+
+/** Used to store lightmap data during RerunConstructionScripts */
+USTRUCT()
+struct FInstancedStaticMeshLightMapInstanceData
+{
+	GENERATED_BODY()
+
+	/** Transform of component */
+	UPROPERTY()
+	FTransform Transform;
+
+	/** guid from LODData */
+	UPROPERTY()
+	TArray<FGuid> MapBuildDataIds;
+};
+
+/** Helper class used to preserve lighting/selection state across blueprint reinstancing */
+USTRUCT()
+struct FInstancedStaticMeshComponentInstanceData : public FSceneComponentInstanceData
+{
+	GENERATED_BODY()
+public:
+	FInstancedStaticMeshComponentInstanceData() = default;
+	FInstancedStaticMeshComponentInstanceData(const UInstancedStaticMeshComponent* InComponent)
+		: FSceneComponentInstanceData(InComponent)
+		, StaticMesh(InComponent->GetStaticMesh())
+	{}
+	virtual ~FInstancedStaticMeshComponentInstanceData() = default;
+
+	virtual bool ContainsData() const override
+	{
+		return true;
+	}
+
+	virtual void ApplyToComponent(UActorComponent* Component, const ECacheApplyPhase CacheApplyPhase) override
+	{
+		Super::ApplyToComponent(Component, CacheApplyPhase);
+		CastChecked<UInstancedStaticMeshComponent>(Component)->ApplyComponentInstanceData(this);
+	}
+
+	virtual void AddReferencedObjects(FReferenceCollector& Collector) override
+	{
+		Super::AddReferencedObjects(Collector);
+		Collector.AddReferencedObject(StaticMesh);
+	}
+
+public:
+	/** Mesh being used by component */
+	UPROPERTY()
+	UStaticMesh* StaticMesh;
+
+	// Static lighting info
+	UPROPERTY()
+	FInstancedStaticMeshLightMapInstanceData CachedStaticLighting;
+	UPROPERTY()
+	TArray<FInstancedStaticMeshInstanceData> PerInstanceSMData;
+
+	/** The cached selected instances */
+	TBitArray<> SelectedInstances;
+
+	/* The cached random seed */
+	UPROPERTY()
+	int32 InstancingRandomSeed;
+};
diff --git a/Engine/Source/Runtime/Engine/Classes/Components/LightComponent.h b/Engine/Source/Runtime/Engine/Classes/Components/LightComponent.h
--- a/Engine/Source/Runtime/Engine/Classes/Components/LightComponent.h
+++ b/Engine/Source/Runtime/Engine/Classes/Components/LightComponent.h
@@ -364,8 +364,8 @@ public:
 	virtual bool IsReadyForFinishDestroy() override;
 	//~ End UObject Interface.
 
-	virtual FActorComponentInstanceData* GetComponentInstanceData() const override;
-	void ApplyComponentInstanceData(class FPrecomputedLightInstanceData* ComponentInstanceData);
+	virtual TStructOnScope<FActorComponentInstanceData> GetComponentInstanceData() const override;
+	void ApplyComponentInstanceData(struct FPrecomputedLightInstanceData* ComponentInstanceData);
 	virtual void PropagateLightingScenarioChange() override;
 	virtual bool IsPrecomputedLightingValid() const override;
 
@@ -423,4 +423,40 @@ public:
 };
 
 
+/** Used to store lightmap data during RerunConstructionScripts */
+USTRUCT()
+struct FPrecomputedLightInstanceData : public FSceneComponentInstanceData
+{
+	GENERATED_BODY()
+
+	FPrecomputedLightInstanceData() = default;
+	FPrecomputedLightInstanceData(const ULightComponent* SourceComponent)
+		: FSceneComponentInstanceData(SourceComponent)
+		, Transform(SourceComponent->GetComponentTransform())
+		, LightGuid(SourceComponent->LightGuid)
+		, PreviewShadowMapChannel(SourceComponent->PreviewShadowMapChannel)
+	{}
+	virtual ~FPrecomputedLightInstanceData() = default;
+
+	virtual bool ContainsData() const override
+	{
+		return true;
+	}
+
+	virtual void ApplyToComponent(UActorComponent* Component, const ECacheApplyPhase CacheApplyPhase) override
+	{
+		Super::ApplyToComponent(Component, CacheApplyPhase);
+		CastChecked<ULightComponent>(Component)->ApplyComponentInstanceData(this);
+	}
+
+	UPROPERTY()
+	FTransform Transform;
+
+	UPROPERTY()
+	FGuid LightGuid;
+
+	UPROPERTY()
+	int32 PreviewShadowMapChannel;
+};
+
 
diff --git a/Engine/Source/Runtime/Engine/Classes/Components/PrimitiveComponent.h b/Engine/Source/Runtime/Engine/Classes/Components/PrimitiveComponent.h
--- a/Engine/Source/Runtime/Engine/Classes/Components/PrimitiveComponent.h
+++ b/Engine/Source/Runtime/Engine/Classes/Components/PrimitiveComponent.h
@@ -1482,7 +1482,7 @@ public:
 	virtual bool IsEditorOnly() const override;
 	virtual bool ShouldCreatePhysicsState() const override;
 	virtual bool HasValidPhysicsState() const override;
-	virtual class FActorComponentInstanceData* GetComponentInstanceData() const override;
+	virtual TStructOnScope<FActorComponentInstanceData> GetComponentInstanceData() const override;
 	virtual void OnComponentDestroyed(bool bDestroyingHierarchy) override;
 	//~ End UActorComponent Interface
 
@@ -2307,24 +2307,31 @@ public:
  *  Component instance cached data base class for primitive components. 
  *  Stores a list of instance components attached to the 
  */
-class ENGINE_API FPrimitiveComponentInstanceData : public FSceneComponentInstanceData
+USTRUCT()
+struct ENGINE_API FPrimitiveComponentInstanceData : public FSceneComponentInstanceData
 {
+	GENERATED_BODY()
 public:
+	FPrimitiveComponentInstanceData() = default;
 	FPrimitiveComponentInstanceData(const UPrimitiveComponent* SourceComponent);
-	virtual ~FPrimitiveComponentInstanceData()
-	{}
+	virtual ~FPrimitiveComponentInstanceData() = default;
+
+	virtual bool ContainsData() const override;
 
 	virtual void ApplyToComponent(UActorComponent* Component, const ECacheApplyPhase CacheApplyPhase) override;
 	virtual void FindAndReplaceInstances(const TMap<UObject*, UObject*>& OldToNewInstanceMap) override;
 	virtual void AddReferencedObjects(FReferenceCollector& Collector) override;
 
-	bool ContainsData() const;
-
 	const FTransform& GetComponentTransform() const { return ComponentTransform; }
 
 private:
+	UPROPERTY()
 	FTransform ComponentTransform;
+
+	UPROPERTY()
 	int32 VisibilityId;
+
+	UPROPERTY()
 	UPrimitiveComponent* LODParent;
 };
 
diff --git a/Engine/Source/Runtime/Engine/Classes/Components/SceneComponent.h b/Engine/Source/Runtime/Engine/Classes/Components/SceneComponent.h
--- a/Engine/Source/Runtime/Engine/Classes/Components/SceneComponent.h
+++ b/Engine/Source/Runtime/Engine/Classes/Components/SceneComponent.h
@@ -905,7 +905,7 @@ public:
 	virtual void DestroyComponent(bool bPromoteChildren = false) override;
 	virtual void OnComponentDestroyed(bool bDestroyingHierarchy) override;
 	virtual void ApplyWorldOffset(const FVector& InOffset, bool bWorldShift) override;
-	virtual class FActorComponentInstanceData* GetComponentInstanceData() const override;
+	virtual TStructOnScope<FActorComponentInstanceData> GetComponentInstanceData() const override;
 	//~ End ActorComponent Interface
 
 	// Call UpdateComponentToWorld if bComponentToWorldUpdated is false.
@@ -1375,19 +1375,24 @@ FORCEINLINE_DEBUGGABLE void USceneComponent::AddRelativeRotation(FRotator DeltaR
  *  Component instance cached data base class for scene components. 
  *  Stores a list of instance components attached to the 
  */
-class ENGINE_API FSceneComponentInstanceData : public FActorComponentInstanceData
+USTRUCT()
+struct ENGINE_API FSceneComponentInstanceData : public FActorComponentInstanceData
 {
-public:
+	GENERATED_BODY()
+
+	FSceneComponentInstanceData() = default;
 	FSceneComponentInstanceData(const USceneComponent* SourceComponent);
 			
-	virtual ~FSceneComponentInstanceData()
-	{}
+	virtual ~FSceneComponentInstanceData() = default;
+
+	virtual bool ContainsData() const override;
 
 	virtual void ApplyToComponent(UActorComponent* Component, const ECacheApplyPhase CacheApplyPhase) override;
 	virtual void FindAndReplaceInstances(const TMap<UObject*, UObject*>& OldToNewInstanceMap) override;
 	virtual void AddReferencedObjects(FReferenceCollector& Collector) override;
 
-	TArray<TPair<USceneComponent*, FTransform>> AttachedInstanceComponents;
+	UPROPERTY() 
+	TMap<USceneComponent*, FTransform> AttachedInstanceComponents;
 };
 
 
diff --git a/Engine/Source/Runtime/Engine/Classes/Components/SkyLightComponent.h b/Engine/Source/Runtime/Engine/Classes/Components/SkyLightComponent.h
--- a/Engine/Source/Runtime/Engine/Classes/Components/SkyLightComponent.h
+++ b/Engine/Source/Runtime/Engine/Classes/Components/SkyLightComponent.h
@@ -179,8 +179,8 @@ class ENGINE_API USkyLightComponent : public ULightComponentBase
 	virtual bool IsReadyForFinishDestroy() override;
 	//~ End UObject Interface
 
-	virtual FActorComponentInstanceData* GetComponentInstanceData() const override;
-	void ApplyComponentInstanceData(class FPrecomputedSkyLightInstanceData* ComponentInstanceData);
+	virtual TStructOnScope<FActorComponentInstanceData> GetComponentInstanceData() const override;
+	void ApplyComponentInstanceData(struct FPrecomputedSkyLightInstanceData* ComponentInstanceData);
 
 	/** Called each tick to recapture and queued sky captures. */
 	static void UpdateSkyCaptureContents(UWorld* WorldToUpdate);
@@ -307,4 +307,38 @@ protected:
 };
 
 
+/** Used to store lightmap data during RerunConstructionScripts */
+USTRUCT()
+struct FPrecomputedSkyLightInstanceData : public FSceneComponentInstanceData
+{
+	GENERATED_BODY()
+public:
+	FPrecomputedSkyLightInstanceData() = default;
+	FPrecomputedSkyLightInstanceData(const USkyLightComponent* SourceComponent)
+		: FSceneComponentInstanceData(SourceComponent)
+	{}
+	virtual ~FPrecomputedSkyLightInstanceData() = default;
+
+	virtual bool ContainsData() const override
+	{
+		return true;
+	}
+
+	virtual void ApplyToComponent(UActorComponent* Component, const ECacheApplyPhase CacheApplyPhase) override
+	{
+		Super::ApplyToComponent(Component, CacheApplyPhase);
+		CastChecked<USkyLightComponent>(Component)->ApplyComponentInstanceData(this);
+	}
+
+	UPROPERTY()
+	FGuid LightGuid;
+
+	UPROPERTY()
+	float AverageBrightness;
+
+	// This has to be refcounted to keep it alive during the handoff without doing a deep copy
+	TRefCountPtr<FSkyTextureCubeResource> ProcessedSkyTexture;
+
+	FSHVectorRGB3 IrradianceEnvironmentMap;
+};
 
diff --git a/Engine/Source/Runtime/Engine/Classes/Components/SplineComponent.h b/Engine/Source/Runtime/Engine/Classes/Components/SplineComponent.h
--- a/Engine/Source/Runtime/Engine/Classes/Components/SplineComponent.h
+++ b/Engine/Source/Runtime/Engine/Classes/Components/SplineComponent.h
@@ -272,7 +272,7 @@ public:
 	//~ End UObject Interface
 
 	//~ Begin UActorComponent Interface.
-	virtual FActorComponentInstanceData* GetComponentInstanceData() const override;
+	virtual TStructOnScope<FActorComponentInstanceData> GetComponentInstanceData() const override;
 	//~ End UActorComponent Interface.
 
 #if !UE_BUILD_SHIPPING
@@ -295,7 +295,7 @@ public:
 	FInterpCurveVector& GetSplinePointsScale() { return SplineCurves.Scale; }
 	const FInterpCurveVector& GetSplinePointsScale() const { return SplineCurves.Scale; }
 
-	void ApplyComponentInstanceData(class FSplineInstanceData* ComponentInstanceData, const bool bPostUCS);
+	void ApplyComponentInstanceData(struct FSplineInstanceData* ComponentInstanceData, const bool bPostUCS);
 
 	/** Update the spline tangents and SplineReparamTable */
 	UFUNCTION(BlueprintCallable, Category = Spline)
@@ -722,6 +722,44 @@ private:
 	}
 };
 
+/** Used to store spline data during RerunConstructionScripts */
+USTRUCT()
+struct FSplineInstanceData : public FSceneComponentInstanceData
+{
+	GENERATED_BODY()
+public:
+	FSplineInstanceData()
+		: bSplineHasBeenEdited(false)
+	{}
+	explicit FSplineInstanceData(const USplineComponent* SourceComponent)
+		: FSceneComponentInstanceData(SourceComponent)
+		, bSplineHasBeenEdited(false)
+	{}
+	virtual ~FSplineInstanceData() = default;
+
+	virtual bool ContainsData() const override
+	{
+		return Super::ContainsData() || bSplineHasBeenEdited;
+	}
+
+	virtual void ApplyToComponent(UActorComponent* Component, const ECacheApplyPhase CacheApplyPhase) override
+	{
+		Super::ApplyToComponent(Component, CacheApplyPhase);
+		CastChecked<USplineComponent>(Component)->ApplyComponentInstanceData(this, (CacheApplyPhase == ECacheApplyPhase::PostUserConstructionScript));
+	}
+
+	UPROPERTY()
+	bool bSplineHasBeenEdited;
+
+	UPROPERTY()
+	FSplineCurves SplineCurves;
+
+	UPROPERTY()
+	FSplineCurves SplineCurvesPreUCS;
+};
+
+
+
 ENGINE_API EInterpCurveMode ConvertSplinePointTypeToInterpCurveMode(ESplinePointType::Type SplinePointType);
 ENGINE_API ESplinePointType::Type ConvertInterpCurveModeToSplinePointType(EInterpCurveMode InterpCurveMode);
 
diff --git a/Engine/Source/Runtime/Engine/Classes/Components/SplineMeshComponent.h b/Engine/Source/Runtime/Engine/Classes/Components/SplineMeshComponent.h
--- a/Engine/Source/Runtime/Engine/Classes/Components/SplineMeshComponent.h
+++ b/Engine/Source/Runtime/Engine/Classes/Components/SplineMeshComponent.h
@@ -154,10 +154,10 @@ class ENGINE_API USplineMeshComponent : public UStaticMeshComponent, public IInt
 	//End UObject Interface
 
 	//~ Begin UActorComponent Interface.
-	virtual FActorComponentInstanceData* GetComponentInstanceData() const override;
+	virtual TStructOnScope<FActorComponentInstanceData> GetComponentInstanceData() const override;
 	//~ End UActorComponent Interface.
 
-	void ApplyComponentInstanceData(class FSplineMeshInstanceData* ComponentInstanceData);
+	void ApplyComponentInstanceData(struct FSplineMeshInstanceData* ComponentInstanceData);
 
 	//Begin USceneComponent Interface
 	virtual FPrimitiveSceneProxy* CreateSceneProxy() override;
@@ -350,6 +350,42 @@ private:
 	void UpdateRenderStateAndCollision_Internal(bool bConcurrent);
 };
 
+/** Used to store spline mesh data during RerunConstructionScripts */
+USTRUCT()
+struct FSplineMeshInstanceData : public FSceneComponentInstanceData
+{
+	GENERATED_BODY()
+public:
+	FSplineMeshInstanceData() = default;
+	explicit FSplineMeshInstanceData(const USplineMeshComponent* SourceComponent)
+		: FSceneComponentInstanceData(SourceComponent)
+	{}
+	virtual ~FSplineMeshInstanceData() = default;
+
+	virtual bool ContainsData() const override
+	{
+		return true;
+	}
+
+	virtual void ApplyToComponent(UActorComponent* Component, const ECacheApplyPhase CacheApplyPhase) override
+	{
+		Super::ApplyToComponent(Component, CacheApplyPhase);
+		CastChecked<USplineMeshComponent>(Component)->ApplyComponentInstanceData(this);
+	}
+
+	UPROPERTY()
+	FVector StartPos;
+
+	UPROPERTY()
+	FVector EndPos;
+
+	UPROPERTY()
+	FVector StartTangent;
+
+	UPROPERTY()
+	FVector EndTangent;
+};
+
 const float& USplineMeshComponent::GetAxisValue(const FVector& InVector, ESplineMeshAxis::Type InAxis)
 {
 	switch (InAxis)
diff --git a/Engine/Source/Runtime/Engine/Classes/Components/StaticMeshComponent.h b/Engine/Source/Runtime/Engine/Classes/Components/StaticMeshComponent.h
--- a/Engine/Source/Runtime/Engine/Classes/Components/StaticMeshComponent.h
+++ b/Engine/Source/Runtime/Engine/Classes/Components/StaticMeshComponent.h
@@ -447,7 +447,7 @@ public:
 #if WITH_EDITOR
 	virtual void CheckForErrors() override;
 #endif
-	virtual FActorComponentInstanceData* GetComponentInstanceData() const override;
+	virtual TStructOnScope<FActorComponentInstanceData> GetComponentInstanceData() const override;
 	//~ End UActorComponent Interface.
 
 	//~ Begin UPrimitiveComponent Interface.
@@ -665,7 +665,7 @@ public:
 	/** Get this components index in its parents blueprint created components array (used for matching instance data) */
 	int32 GetBlueprintCreatedComponentIndex() const;
 
-	void ApplyComponentInstanceData(class FStaticMeshComponentInstanceData* ComponentInstanceData);
+	void ApplyComponentInstanceData(struct FStaticMeshComponentInstanceData* ComponentInstanceData);
 
 	virtual void PropagateLightingScenarioChange() override;
 
@@ -681,5 +681,73 @@ private:
 #endif
 };
 
+/** Vertex data stored per-LOD */
+USTRUCT()
+struct FStaticMeshVertexColorLODData
+{
+	GENERATED_BODY()
+
+	/** copy of painted vertex data */
+	UPROPERTY()
+	TArray<FPaintedVertex> PaintedVertices;
+
+	/** Copy of vertex buffer colors */
+	UPROPERTY()
+	TArray<FColor> VertexBufferColors;
+
+	/** Index of the LOD that this data came from */
+	UPROPERTY()
+	uint32 LODIndex;
+
+	/** Check whether this contains valid data */
+	bool IsValid() const
+	{
+		return PaintedVertices.Num() > 0 && VertexBufferColors.Num() > 0;
+	}
+};
+
+USTRUCT()
+struct FStaticMeshComponentInstanceData : public FPrimitiveComponentInstanceData
+{
+	GENERATED_BODY()
+public:
+	FStaticMeshComponentInstanceData() = default;
+	FStaticMeshComponentInstanceData(const UStaticMeshComponent* SourceComponent);
+	virtual ~FStaticMeshComponentInstanceData() = default;
+
+	virtual bool ContainsData() const override;
+
+	virtual void ApplyToComponent(UActorComponent* Component, const ECacheApplyPhase CacheApplyPhase) override;
+
+	virtual void AddReferencedObjects(FReferenceCollector& Collector) override;
+
+	/** Add vertex color data for a specified LOD before RerunConstructionScripts is called */
+	void AddVertexColorData(const struct FStaticMeshComponentLODInfo& LODInfo, uint32 LODIndex);
+
+	/** Re-apply vertex color data after RerunConstructionScripts is called */
+	bool ApplyVertexColorData(UStaticMeshComponent* StaticMeshComponent) const;
+
+	/** Mesh being used by component */
+	UPROPERTY()
+	class UStaticMesh* StaticMesh;
+
+	/** Array of cached vertex colors for each LOD */
+	UPROPERTY()
+	TArray<FStaticMeshVertexColorLODData> VertexColorLODs;
+
+	/** Used to store lightmap data during RerunConstructionScripts */
+	UPROPERTY()
+	TArray<FGuid> CachedStaticLighting;
+
+	/** Texture streaming build data */
+	UPROPERTY()
+	TArray<FStreamingTextureBuildInfo> StreamingTextureData;
+
+#if WITH_EDITORONLY_DATA
+	/** Texture streaming editor data (for viewmodes) */
+	UPROPERTY()
+	TArray<uint32> MaterialStreamingRelativeBoxes;
+#endif
+};
 
 
diff --git a/Engine/Source/Runtime/Engine/Private/ActorEditor.cpp b/Engine/Source/Runtime/Engine/Private/ActorEditor.cpp
--- a/Engine/Source/Runtime/Engine/Private/ActorEditor.cpp
+++ b/Engine/Source/Runtime/Engine/Private/ActorEditor.cpp
@@ -426,6 +426,7 @@ void AActor::FActorTransactionAnnotation::Serialize(FArchive& Ar)
 	enum class EVersion : uint8
 	{
 		InitialVersion = 0,
+		WithInstanceCache,
 		// -----<new versions can be added above this line>-------------------------------------------------
 		VersionPlusOne,
 		LatestVersion = VersionPlusOne - 1
@@ -440,17 +441,21 @@ void AActor::FActorTransactionAnnotation::Serialize(FArchive& Ar)
 		return;
 	}
 
+	// InitialVersion
 	Ar << Actor;
+	Ar << bRootComponentDataCached;
+	if (bRootComponentDataCached)
+	{
+		Ar << RootComponentData;
+	}
+	// WithInstanceCache
 	if (Ar.IsLoading())
 	{
 		ComponentInstanceData = FComponentInstanceDataCache(Actor.Get());
 	}
-
-	Ar << bRootComponentDataCached;
-
-	if (bRootComponentDataCached)
+	if (Version >= EVersion::WithInstanceCache)
 	{
-		Ar << RootComponentData;
+		ComponentInstanceData.Serialize(Ar);
 	}
 }
 
diff --git a/Engine/Source/Runtime/Engine/Private/Atmosphere/Atmosphere.cpp b/Engine/Source/Runtime/Engine/Private/Atmosphere/Atmosphere.cpp
--- a/Engine/Source/Runtime/Engine/Private/Atmosphere/Atmosphere.cpp
+++ b/Engine/Source/Runtime/Engine/Private/Atmosphere/Atmosphere.cpp
@@ -731,40 +731,17 @@ void UAtmosphericFogComponent::Serialize(FArchive& Ar)
 	}
 }
 
-/** Used to store lightmap data during RerunConstructionScripts */
-class FAtmospherePrecomputeInstanceData : public FSceneComponentInstanceData
-{
-public:
-	FAtmospherePrecomputeInstanceData(const UAtmosphericFogComponent* SourceComponent)
-		: FSceneComponentInstanceData(SourceComponent)
-	{}
-
-	virtual ~FAtmospherePrecomputeInstanceData()
-	{}
-
-	virtual void ApplyToComponent(UActorComponent* Component, const ECacheApplyPhase CacheApplyPhase) override
-	{
-		FSceneComponentInstanceData::ApplyToComponent(Component, CacheApplyPhase);
-		CastChecked<UAtmosphericFogComponent>(Component)->ApplyComponentInstanceData(this);
-	}
-
-	struct FAtmospherePrecomputeParameters PrecomputeParameter;
-
-	FByteBulkData TransmittanceData;
-	FByteBulkData IrradianceData;
-	FByteBulkData InscatterData;
-};
 
 // Backup the precomputed data before re-running Blueprint construction script
-FActorComponentInstanceData* UAtmosphericFogComponent::GetComponentInstanceData() const
+TStructOnScope<FActorComponentInstanceData> UAtmosphericFogComponent::GetComponentInstanceData() const
 {
-	FActorComponentInstanceData* InstanceData = nullptr;
+	TStructOnScope<FActorComponentInstanceData> InstanceData;
 
 	if (TransmittanceData.GetElementCount() && IrradianceData.GetElementCount() && InscatterData.GetElementCount() && PrecomputeCounter == EValid)
 	{
 		// Allocate new struct for holding light map data
-		FAtmospherePrecomputeInstanceData* PrecomputedData = new FAtmospherePrecomputeInstanceData(this);
-		InstanceData = PrecomputedData;
+		InstanceData.InitializeAs<FAtmospherePrecomputeInstanceData>(this);
+		FAtmospherePrecomputeInstanceData* PrecomputedData = InstanceData.Cast<FAtmospherePrecomputeInstanceData>();
 
 		// Fill in info
 		PrecomputedData->PrecomputeParameter = PrecomputeParams;
@@ -796,7 +773,6 @@ FActorComponentInstanceData* UAtmosphericFogComponent::GetComponentInstanceData(
 	{
 		InstanceData = Super::GetComponentInstanceData();
 	}
-
 	return InstanceData;
 }
 
diff --git a/Engine/Source/Runtime/Engine/Private/ComponentInstanceDataCache.cpp b/Engine/Source/Runtime/Engine/Private/ComponentInstanceDataCache.cpp
--- a/Engine/Source/Runtime/Engine/Private/ComponentInstanceDataCache.cpp
+++ b/Engine/Source/Runtime/Engine/Private/ComponentInstanceDataCache.cpp
@@ -4,6 +4,8 @@
 #include "Serialization/ObjectWriter.h"
 #include "Serialization/ObjectReader.h"
 #include "Serialization/DuplicatedObject.h"
+#include "Serialization/MemoryWriter.h"
+#include "Serialization/MemoryReader.h"
 #include "UObject/Package.h"
 #include "UObject/UObjectAnnotation.h"
 #include "UObject/UObjectGlobals.h"
@@ -16,11 +18,10 @@
 class FComponentPropertyWriter : public FObjectWriter
 {
 public:
-	FComponentPropertyWriter(const UActorComponent* InComponent, TArray<uint8>& InBytes, TArray<UObject*>& InInstancedObjects, TArray<uint32_t>& InObjectReferenceIndicesInByteArray)
-		: FObjectWriter(InBytes)
+	FComponentPropertyWriter(const UActorComponent* InComponent, FActorComponentInstanceData& InActorInstanceData)
+		: FObjectWriter(InActorInstanceData.SavedProperties)
 		, Component(InComponent)
-		, InstancedObjects(InInstancedObjects)
-		, ObjectReferenceIndicesInByteArray(InObjectReferenceIndicesInByteArray)
+		, ActorInstanceData(InActorInstanceData)
 	{
 		// Include properties that would normally skip tagged serialization (e.g. bulk serialization of array properties).
 		ArPortFlags |= PPF_ForceTaggedSerialization;
@@ -80,7 +81,7 @@ public:
 			else if (Object->GetOuter() == Component)
 			{
 				Result = DuplicateObject(Object, GetTransientPackage());
-				InstancedObjects.Add(Result);
+				ActorInstanceData.DuplicatedObjects.Emplace(Result);
 			}
 			else
 			{
@@ -107,6 +108,16 @@ public:
 		return Result;
 	}
 
+	virtual FArchive& operator<<(FName& Name) override
+	{
+		// store the reference to this name in the array instead of the global table, this allow to for persistence
+		int32 ReferenceIndex = ActorInstanceData.ReferencedNames.AddUnique(Name);
+		// save the name as an index in the referenced name array
+		*this << ReferenceIndex;
+
+		return *this;
+	}
+
 	virtual FArchive& operator<<(UObject*& Object) override
 	{
 		UObject* SerializedObject = Object;
@@ -116,41 +127,38 @@ public:
 		}
 
 		// store the pointer to this object
+		int32 ReferenceIndex = INDEX_NONE;
 		if (SerializedObject)
 		{
-			ObjectReferenceIndicesInByteArray.Add(Bytes.Num());
-		}
-		Serialize(&SerializedObject, sizeof(UObject*));
-#if DO_CHECK
-		if (SerializedObject)
-		{
-			uint32_t PtrIdx = Bytes.Num() - sizeof(UObject*);
-			void* NewObjAddress = static_cast<void*>(&Bytes[PtrIdx]);
-			UObject* NewObjValue = *(reinterpret_cast<UObject**>(NewObjAddress));
-			ensure(NewObjValue == SerializedObject);
+			ReferenceIndex = ActorInstanceData.ReferencedObjects.Num();
+			ActorInstanceData.ReferencedObjects.Add(SerializedObject);
 		}
-#endif
+		// save the pointer as an index in the referenced object array
+		*this << ReferenceIndex;
 
 		return *this;
 	}
 
-private:
+	virtual FArchive& operator<<(FLazyObjectPtr& LazyObjectPtr) override
+	{
+		UObject* Obj = LazyObjectPtr.Get();
+		*this << Obj;
+		return *this;
+	}
 
+private:
 	const UActorComponent* Component;
-
+	FActorComponentInstanceData& ActorInstanceData;
 	TSet<const UProperty*> PropertiesToSkip;
-	
-	TArray<UObject*>& InstancedObjects;
-	TArray<uint32_t>& ObjectReferenceIndicesInByteArray;
-
 	FUObjectAnnotationSparse<FDuplicatedObject,false> DuplicatedObjectAnnotation;
 };
 
 class FComponentPropertyReader : public FObjectReader
 {
 public:
-	FComponentPropertyReader(UActorComponent* InComponent, TArray<uint8>& InBytes)
-		: FObjectReader(InBytes)
+	FComponentPropertyReader(UActorComponent* InComponent, FActorComponentInstanceData& InActorInstanceData)
+		: FObjectReader(InActorInstanceData.SavedProperties)
+		, ActorInstanceData(InActorInstanceData)
 	{
 		// Include properties that would normally skip tagged serialization (e.g. bulk serialization of array properties).
 		ArPortFlags |= PPF_ForceTaggedSerialization;
@@ -166,22 +174,129 @@ public:
 		return PropertiesToSkip.Contains(InProperty);
 	}
 
+	virtual FArchive& operator<<(FName& Name) override
+	{
+		// FName are serialized as Index in ActorInstanceData instead of the normal FName table
+		int32 ReferenceIndex = INDEX_NONE;
+		*this << ReferenceIndex;
+		Name = ActorInstanceData.ReferencedNames.IsValidIndex(ReferenceIndex) ? ActorInstanceData.ReferencedNames[ReferenceIndex] : FName();
+		return *this;
+	}
+
+	virtual FArchive& operator<<(UObject*& Object) override
+	{
+		// UObject pointer are serialized as Index in ActorInstanceData
+		int32 ReferenceIndex = INDEX_NONE;
+		*this << ReferenceIndex;
+		Object = ActorInstanceData.ReferencedObjects.IsValidIndex(ReferenceIndex) ? ActorInstanceData.ReferencedObjects[ReferenceIndex] : nullptr;
+		return *this;
+	}
+
+	virtual FArchive& operator<<(FLazyObjectPtr& LazyObjectPtr) override
+	{
+		UObject* Obj = LazyObjectPtr.Get();
+		*this << Obj;
+		return *this;
+	}
+
+	FActorComponentInstanceData& ActorInstanceData;
 	TSet<const UProperty*> PropertiesToSkip;
 };
 
+
+FActorComponentDuplicatedObjectData::FActorComponentDuplicatedObjectData(UObject* InObject)
+	: DuplicatedObject(InObject)
+	, ObjectPathDepth(0)
+{
+	if (DuplicatedObject)
+	{
+		for (UObject* Outer = DuplicatedObject; Outer; Outer = Outer->GetOuter())
+		{
+			++ObjectPathDepth;
+		}
+	}
+}
+
+bool FActorComponentDuplicatedObjectData::Serialize(FArchive& Ar)
+{
+	enum class EVersion : uint8
+	{
+		InitialVersion = 0,
+		// -----<new versions can be added above this line>-------------------------------------------------
+		VersionPlusOne,
+		LatestVersion = VersionPlusOne - 1
+	};
+
+	EVersion Version = EVersion::LatestVersion;
+	Ar << Version;
+
+	if (Version > EVersion::LatestVersion)
+	{
+		Ar.SetError();
+		return true;
+	}
+
+	FString ObjectClassPath;
+	FString ObjectOuterPath;
+	FName ObjectName;
+	uint32 ObjectPersistentFlags;
+	TArray<uint8> ObjectData;
+
+	if (Ar.IsSaving())
+	{
+		UClass* ObjectClass = DuplicatedObject ? DuplicatedObject->GetClass() : nullptr;
+		ObjectClassPath = ObjectClass ? ObjectClass->GetPathName() : FString();
+		ObjectOuterPath = DuplicatedObject && DuplicatedObject->GetOuter() ? DuplicatedObject->GetOuter()->GetPathName() : FString();
+		ObjectName = DuplicatedObject->GetFName();
+		ObjectPersistentFlags = DuplicatedObject ? DuplicatedObject->GetFlags() & RF_Load : 0;
+		if (ObjectClass)
+		{
+			FMemoryWriter Writer(ObjectData);
+			ObjectClass->SerializeTaggedProperties(Writer, (uint8*)DuplicatedObject, ObjectClass, nullptr);
+		}
+	}
+
+	Ar << ObjectClassPath;
+	Ar << ObjectOuterPath;
+	Ar << ObjectName;
+	Ar << ObjectPersistentFlags;
+	Ar << ObjectData;
+
+	// If loading use the de-serialized property to assign DuplicatedObject
+	if (Ar.IsLoading())
+	{
+		DuplicatedObject = nullptr;
+		// Get the object class
+		if (UClass* ObjectClass = LoadObject<UClass>(nullptr, *ObjectClassPath))
+		{
+			// Get the object outer
+			if (UObject* FoundOuter = StaticFindObject(UObject::StaticClass(), nullptr, *ObjectOuterPath))
+			{
+				// Create the duplicated object
+				DuplicatedObject = NewObject<UObject>(FoundOuter, ObjectClass, *ObjectName.ToString(), (EObjectFlags)ObjectPersistentFlags);
+
+				// Deserialize the duplicated object
+				FMemoryReader Reader(ObjectData);
+				ObjectClass->SerializeTaggedProperties(Reader, (uint8*)DuplicatedObject, ObjectClass, nullptr);
+			}
+		}
+	}
+
+	return true;
+}
+
 FActorComponentInstanceData::FActorComponentInstanceData()
 	: SourceComponentTemplate(nullptr)
-	, SourceComponentTypeSerializedIndex(-1)
 	, SourceComponentCreationMethod(EComponentCreationMethod::Native)
-{
-}
+	, SourceComponentTypeSerializedIndex(INDEX_NONE)
+{}
 
 FActorComponentInstanceData::FActorComponentInstanceData(const UActorComponent* SourceComponent)
 {
 	check(SourceComponent);
 	SourceComponentTemplate = SourceComponent->GetArchetype();
 	SourceComponentCreationMethod = SourceComponent->CreationMethod;
-	SourceComponentTypeSerializedIndex = -1;
+	SourceComponentTypeSerializedIndex = INDEX_NONE;
 
 	// UCS components can share the same template (e.g. an AddComponent node inside a loop), so we also cache their serialization index here (relative to the shared template) as a means for identification
 	if(SourceComponentCreationMethod == EComponentCreationMethod::UserConstructionScript)
@@ -209,30 +324,28 @@ FActorComponentInstanceData::FActorComponentInstanceData(const UActorComponent*
 			}
 			if (!bFound)
 			{
-				SourceComponentTypeSerializedIndex = -1;
+				SourceComponentTypeSerializedIndex = INDEX_NONE;
 			}
 		}
 	}
 
 	if (SourceComponent->IsEditableWhenInherited())
 	{
-		FComponentPropertyWriter ComponentPropertyWriter(SourceComponent, SavedProperties, InstancedObjects, ObjectReferenceIndicesInByteArray);
+		FComponentPropertyWriter ComponentPropertyWriter(SourceComponent, *this);
 
 		// Cache off the length of an array that will come from SerializeTaggedProperties that had no properties saved in to it.
 		auto GetSizeOfEmptyArchive = []() -> int32
 		{
 			const UActorComponent* DummyComponent = GetDefault<UActorComponent>();
-			TArray<uint8> NoWrittenPropertyReference;
-			TArray<UObject*> NoInstances;
-			TArray<uint32_t> NoObjectReferences;
-			FComponentPropertyWriter NullWriter(nullptr, NoWrittenPropertyReference, NoInstances, NoObjectReferences);
+			FActorComponentInstanceData DummyInstanceData;
+			FComponentPropertyWriter NullWriter(nullptr, DummyInstanceData);
 			UClass* ComponentClass = DummyComponent->GetClass();
 			
 			// By serializing the component with itself as its defaults we guarantee that no properties will be written out
 			ComponentClass->SerializeTaggedProperties(NullWriter, (uint8*)DummyComponent, ComponentClass, (uint8*)DummyComponent);
 
-			check(NoInstances.Num() == 0 && NoObjectReferences.Num() == 0);
-			return NoWrittenPropertyReference.Num();
+			check(DummyInstanceData.DuplicatedObjects.Num() == 0 && DummyInstanceData.ReferencedObjects.Num() == 0);
+			return DummyInstanceData.SavedProperties.Num();
 		};
 
 		static const int32 SizeOfEmptyArchive = GetSizeOfEmptyArchive();
@@ -244,6 +357,15 @@ FActorComponentInstanceData::FActorComponentInstanceData(const UActorComponent*
 			SavedProperties.Empty();
 		}
 	}
+
+	// Sort duplicated objects, if any, this ensures that lower depth duplicated objects are first in the array, which ensures proper creation order when deserializing
+	if (DuplicatedObjects.Num() > 0)
+	{
+		DuplicatedObjects.StableSort([](const FActorComponentDuplicatedObjectData& One, const FActorComponentDuplicatedObjectData& Two) -> bool
+		{
+			return One.ObjectPathDepth < Two.ObjectPathDepth;
+		});
+	}
 }
 
 bool FActorComponentInstanceData::MatchesComponent(const UActorComponent* Component, const UObject* ComponentTemplate, const TMap<UActorComponent*, const UObject*>& ComponentToArchetypeMap) const
@@ -290,15 +412,15 @@ void FActorComponentInstanceData::ApplyToComponent(UActorComponent* Component, c
 	{
 		Component->DetermineUCSModifiedProperties();
 
-		for (UObject* InstancedObject : InstancedObjects)
+		for (const FActorComponentDuplicatedObjectData& DuplicatedObjectData : DuplicatedObjects)
 		{
-			if (InstancedObject)
+			if (DuplicatedObjectData.DuplicatedObject)
 			{
-				InstancedObject->Rename(nullptr, Component, REN_DontCreateRedirectors | REN_ForceNoResetLoaders);
+				DuplicatedObjectData.DuplicatedObject->Rename(nullptr, Component, REN_DontCreateRedirectors | REN_ForceNoResetLoaders);
 			}
 		}
 
-		FComponentPropertyReader ComponentPropertyReader(Component, SavedProperties);
+		FComponentPropertyReader ComponentPropertyReader(Component, *this);
 
 		if (Component->IsRegistered())
 		{
@@ -310,16 +432,7 @@ void FActorComponentInstanceData::ApplyToComponent(UActorComponent* Component, c
 void FActorComponentInstanceData::AddReferencedObjects(FReferenceCollector& Collector)
 {
 	Collector.AddReferencedObject(SourceComponentTemplate);
-	Collector.AddReferencedObjects(InstancedObjects);
-
-	for (uint32_t Idx : ObjectReferenceIndicesInByteArray)
-	{
-		UObject** Obj = (UObject**)&SavedProperties[Idx];
-		if (*Obj)
-		{
-			Collector.AddReferencedObject(*Obj);
-		}
-	}
+	Collector.AddReferencedObjects(ReferencedObjects);
 }
 
 FComponentInstanceDataCache::FComponentInstanceDataCache(const AActor* Actor)
@@ -337,10 +450,10 @@ FComponentInstanceDataCache::FComponentInstanceDataCache(const AActor* Actor)
 		{
 			if (bIsChildActor || Component->IsCreatedByConstructionScript()) // Only cache data from 'created by construction script' components
 			{
-				FActorComponentInstanceData* ComponentInstanceData = Component->GetComponentInstanceData();
-				if (ComponentInstanceData)
+				TStructOnScope<FActorComponentInstanceData> ComponentInstanceData = Component->GetComponentInstanceData();
+				if (ComponentInstanceData.IsValid())
 				{
-					ComponentsInstanceData.Add(ComponentInstanceData);
+					ComponentsInstanceData.Add(MoveTemp(ComponentInstanceData));
 				}
 			}
 			else if (Component->CreationMethod == EComponentCreationMethod::Instance)
@@ -370,12 +483,40 @@ FComponentInstanceDataCache::FComponentInstanceDataCache(const AActor* Actor)
 	}
 }
 
-FComponentInstanceDataCache::~FComponentInstanceDataCache()
+void FComponentInstanceDataCache::Serialize(FArchive& Ar)
 {
-	for (FActorComponentInstanceData* ComponentData : ComponentsInstanceData)
+	enum class EVersion : uint8
 	{
-		delete ComponentData;
+		InitialVersion = 0,
+		// -----<new versions can be added above this line>-------------------------------------------------
+		VersionPlusOne,
+		LatestVersion = VersionPlusOne - 1
+	};
+
+	EVersion Version = EVersion::LatestVersion;
+	Ar << Version;
+
+	if (Version > EVersion::LatestVersion)
+	{
+		Ar.SetError();
+		return;
+	}
+
+
+	if (Ar.IsLoading())
+	{
+		// Since not all properties are serializable we don't want to deserialize the map directly,
+		// so we deserialize it in a temp array and copy it over
+		TArray<TStructOnScope<FActorComponentInstanceData>> TempInstanceData;
+		Ar << TempInstanceData;
+		CopySerializableProperties(MoveTemp(TempInstanceData));
 	}
+	else
+	{
+		Ar << ComponentsInstanceData;
+	}
+
+	Ar << InstanceComponentTransformToRootMap;
 }
 
 void FComponentInstanceDataCache::ApplyToActor(AActor* Actor, const ECacheApplyPhase CacheApplyPhase) const
@@ -465,9 +606,9 @@ void FComponentInstanceDataCache::ApplyToActor(AActor* Actor, const ECacheApplyP
 				// Cache template here to avoid redundant calls in the loop below
 				const UObject* ComponentTemplate = ComponentToArchetypeMap.FindChecked(ComponentInstance);
 
-				for (FActorComponentInstanceData* ComponentInstanceData : ComponentsInstanceData)
+				for (const TStructOnScope<FActorComponentInstanceData>& ComponentInstanceData : ComponentsInstanceData)
 				{
-					if (	ComponentInstanceData
+					if (	ComponentInstanceData.IsValid()
 						&&	ComponentInstanceData->GetComponentClass() == ComponentTemplate->GetClass() // filter on class early to avoid unnecessary virtual and expensive tests
 						&&	ComponentInstanceData->MatchesComponent(ComponentInstance, ComponentTemplate, ComponentToArchetypeMap))
 					{
@@ -495,9 +636,9 @@ void FComponentInstanceDataCache::ApplyToActor(AActor* Actor, const ECacheApplyP
 
 void FComponentInstanceDataCache::FindAndReplaceInstances(const TMap<UObject*, UObject*>& OldToNewInstanceMap)
 {
-	for (FActorComponentInstanceData* ComponentInstanceData : ComponentsInstanceData)
+	for (TStructOnScope<FActorComponentInstanceData>& ComponentInstanceData : ComponentsInstanceData)
 	{
-		if (ComponentInstanceData)
+		if (ComponentInstanceData.IsValid())
 		{
 			ComponentInstanceData->FindAndReplaceInstances(OldToNewInstanceMap);
 		}
@@ -525,11 +666,42 @@ void FComponentInstanceDataCache::AddReferencedObjects(FReferenceCollector& Coll
 {
 	Collector.AddReferencedObjects(InstanceComponentTransformToRootMap);
 
-	for (FActorComponentInstanceData* ComponentInstanceData : ComponentsInstanceData)
+	for (TStructOnScope<FActorComponentInstanceData>& ComponentInstanceData : ComponentsInstanceData)
 	{
-		if (ComponentInstanceData)
+		if (ComponentInstanceData.IsValid())
 		{
 			ComponentInstanceData->AddReferencedObjects(Collector);
 		}
 	}
 }
+
+void FComponentInstanceDataCache::CopySerializableProperties(TArray<TStructOnScope<FActorComponentInstanceData>> InComponentsInstanceData)
+{
+	auto CopyProperties = [](TStructOnScope<FActorComponentInstanceData>& DestData, const TStructOnScope<FActorComponentInstanceData>& SrcData)
+	{
+		for (TFieldIterator<const UProperty> PropertyIt(SrcData.GetStruct(), EFieldIteratorFlags::IncludeSuper, EFieldIteratorFlags::IncludeDeprecated, EFieldIteratorFlags::ExcludeInterfaces); PropertyIt; ++PropertyIt)
+		{
+			const void* SrcValuePtr = PropertyIt->ContainerPtrToValuePtr<void>(SrcData.Get());
+			void* DestValuePtr = PropertyIt->ContainerPtrToValuePtr<void>(DestData.Get());
+			PropertyIt->CopyCompleteValue(DestValuePtr, SrcValuePtr);
+		}
+	};
+
+	for (TStructOnScope<FActorComponentInstanceData>& InstanceData : InComponentsInstanceData)
+	{
+		TStructOnScope<FActorComponentInstanceData>* DestInstanceData = ComponentsInstanceData.FindByPredicate([&InstanceData](TStructOnScope<FActorComponentInstanceData>& InInstanceData)
+		{
+			return InstanceData->GetComponentTemplate() == InInstanceData->GetComponentTemplate() && InstanceData.GetStruct() == InInstanceData.GetStruct();
+		});
+		// if we find a ComponentInstanceData to apply it too, we copy the properties over
+		if (DestInstanceData)
+		{
+			CopyProperties(*DestInstanceData, InstanceData);
+		}
+		// otherwise we just add our to the list, since no component instance data was created for it
+		else
+		{
+			ComponentsInstanceData.Add(MoveTemp(InstanceData));
+		}
+	}
+}
diff --git a/Engine/Source/Runtime/Engine/Private/Components/ActorComponent.cpp b/Engine/Source/Runtime/Engine/Private/Components/ActorComponent.cpp
--- a/Engine/Source/Runtime/Engine/Private/Components/ActorComponent.cpp
+++ b/Engine/Source/Runtime/Engine/Private/Components/ActorComponent.cpp
@@ -797,17 +797,9 @@ void UActorComponent::EndPlay(const EEndPlayReason::Type EndPlayReason)
 	bHasBegunPlay = false;
 }
 
-FActorComponentInstanceData* UActorComponent::GetComponentInstanceData() const
+TStructOnScope<FActorComponentInstanceData> UActorComponent::GetComponentInstanceData() const
 {
-	FActorComponentInstanceData* InstanceData = new FActorComponentInstanceData(this);
-
-	if (!InstanceData->ContainsSavedProperties())
-	{
-		delete InstanceData;
-		InstanceData = nullptr;
-	}
-
-	return InstanceData;
+	return MakeStructOnScope<FActorComponentInstanceData>(this);
 }
 
 void FActorComponentTickFunction::ExecuteTick(float DeltaTime, enum ELevelTick TickType, ENamedThreads::Type CurrentThread, const FGraphEventRef& MyCompletionGraphEvent)
@@ -1731,7 +1723,7 @@ bool UActorComponent::IsEditableWhenInherited() const
 		}
 		else
 #endif
-			if (CreationMethod == EComponentCreationMethod::UserConstructionScript)
+		if (CreationMethod == EComponentCreationMethod::UserConstructionScript)
 		{
 			bCanEdit = false;
 		}
diff --git a/Engine/Source/Runtime/Engine/Private/Components/ChildActorComponent.cpp b/Engine/Source/Runtime/Engine/Private/Components/ChildActorComponent.cpp
--- a/Engine/Source/Runtime/Engine/Private/Components/ChildActorComponent.cpp
+++ b/Engine/Source/Runtime/Engine/Private/Components/ChildActorComponent.cpp
@@ -293,12 +293,11 @@ FChildActorComponentInstanceData::FChildActorComponentInstanceData(const UChildA
 {
 	if (Component->GetChildActor())
 	{
-		ComponentInstanceData = new FComponentInstanceDataCache(Component->GetChildActor());
+		ComponentInstanceData = MakeShared<FComponentInstanceDataCache>(Component->GetChildActor());
 		// If it is empty dump it
 		if (!ComponentInstanceData->HasInstanceData())
 		{
-			delete ComponentInstanceData;
-			ComponentInstanceData = nullptr;
+			ComponentInstanceData.Reset();
 		}
 
 		USceneComponent* ChildRootComponent = Component->GetChildActor()->GetRootComponent();
@@ -311,7 +310,7 @@ FChildActorComponentInstanceData::FChildActorComponentInstanceData(const UChildA
 					AActor* AttachedActor = AttachedComponent->GetOwner();
 					if (AttachedActor != Component->GetChildActor())
 					{
-						FAttachedActorInfo Info;
+						FChildActorAttachedActorInfo Info;
 						Info.Actor = AttachedActor;
 						Info.SocketName = AttachedComponent->GetAttachSocketName();
 						Info.RelativeTransform = AttachedComponent->GetRelativeTransform();
@@ -323,21 +322,20 @@ FChildActorComponentInstanceData::FChildActorComponentInstanceData(const UChildA
 	}
 }
 
-FChildActorComponentInstanceData::~FChildActorComponentInstanceData()
+bool FChildActorComponentInstanceData::ContainsData() const
 {
-	delete ComponentInstanceData;
+	return AttachedActors.Num() > 0 || (ComponentInstanceData && ComponentInstanceData->HasInstanceData()) || Super::ContainsData();
 }
 
 void FChildActorComponentInstanceData::ApplyToComponent(UActorComponent* Component, const ECacheApplyPhase CacheApplyPhase)
 {
-	FSceneComponentInstanceData::ApplyToComponent(Component, CacheApplyPhase);
+	Super::ApplyToComponent(Component, CacheApplyPhase);
 	CastChecked<UChildActorComponent>(Component)->ApplyComponentInstanceData(this, CacheApplyPhase);
 }
 
 void FChildActorComponentInstanceData::AddReferencedObjects(FReferenceCollector& Collector)
 {
-	FSceneComponentInstanceData::AddReferencedObjects(Collector);
-
+	Super::AddReferencedObjects(Collector);
 	if (ComponentInstanceData)
 	{
 		ComponentInstanceData->AddReferencedObjects(Collector);
@@ -367,20 +365,19 @@ void UChildActorComponent::BeginDestroy()
 	}
 }
 
-FActorComponentInstanceData* UChildActorComponent::GetComponentInstanceData() const
+TStructOnScope<FActorComponentInstanceData> UChildActorComponent::GetComponentInstanceData() const
 {
-	FChildActorComponentInstanceData* InstanceData = CachedInstanceData;
-
+	TStructOnScope<FActorComponentInstanceData> InstanceData;
 	if (CachedInstanceData)
 	{
-		// We've handed over ownership of the pointer to the instance cache, so drop our reference
+		InstanceData.InitializeAs<FChildActorComponentInstanceData>(*CachedInstanceData);
+		delete CachedInstanceData;
 		CachedInstanceData = nullptr;
 	}
 	else
 	{
-		InstanceData = new FChildActorComponentInstanceData(this);
+		InstanceData.InitializeAs<FChildActorComponentInstanceData>(this);
 	}
-	
 	return InstanceData;
 }
 
@@ -409,7 +406,7 @@ void UChildActorComponent::ApplyComponentInstanceData(FChildActorComponentInstan
 		USceneComponent* ChildActorRoot = ChildActor->GetRootComponent();
 		if (ChildActorRoot)
 		{
-			for (const FChildActorComponentInstanceData::FAttachedActorInfo& AttachInfo : ChildActorInstanceData->AttachedActors)
+			for (const FChildActorAttachedActorInfo& AttachInfo : ChildActorInstanceData->AttachedActors)
 			{
 				AActor* AttachedActor = AttachInfo.Actor.Get();
 				if (AttachedActor)
@@ -575,7 +572,7 @@ void UChildActorComponent::CreateChildActor()
 					FActorParentComponentSetter::Set(ChildActor, this);
 
 					// Parts that we deferred from SpawnActor
-					const FComponentInstanceDataCache* ComponentInstanceData = (CachedInstanceData ? CachedInstanceData->ComponentInstanceData : nullptr);
+					const FComponentInstanceDataCache* ComponentInstanceData = (CachedInstanceData ? CachedInstanceData->ComponentInstanceData.Get() : nullptr);
 					ChildActor->FinishSpawning(GetComponentTransform(), false, ComponentInstanceData);
 
 					ChildActor->AttachToComponent(this, FAttachmentTransformRules::SnapToTargetNotIncludingScale);
@@ -584,7 +581,7 @@ void UChildActorComponent::CreateChildActor()
 
 					if (CachedInstanceData)
 					{
-						for (const FChildActorComponentInstanceData::FAttachedActorInfo& AttachedActorInfo : CachedInstanceData->AttachedActors)
+						for (const FChildActorAttachedActorInfo& AttachedActorInfo : CachedInstanceData->AttachedActors)
 						{
 							AActor* AttachedActor = AttachedActorInfo.Actor.Get();
 
diff --git a/Engine/Source/Runtime/Engine/Private/Components/LightComponent.cpp b/Engine/Source/Runtime/Engine/Private/Components/LightComponent.cpp
--- a/Engine/Source/Runtime/Engine/Private/Components/LightComponent.cpp
+++ b/Engine/Source/Runtime/Engine/Private/Components/LightComponent.cpp
@@ -1060,32 +1060,10 @@ void ULightComponent::InvalidateLightingCacheInner(bool bRecreateLightGuids)
 	}
 }
 
-/** Used to store lightmap data during RerunConstructionScripts */
-class FPrecomputedLightInstanceData : public FSceneComponentInstanceData
-{
-public:
-	FPrecomputedLightInstanceData(const ULightComponent* SourceComponent)
-		: FSceneComponentInstanceData(SourceComponent)
-		, Transform(SourceComponent->GetComponentTransform())
-		, LightGuid(SourceComponent->LightGuid)
-		, PreviewShadowMapChannel(SourceComponent->PreviewShadowMapChannel)
-	{}
-
-	virtual void ApplyToComponent(UActorComponent* Component, const ECacheApplyPhase CacheApplyPhase) override
-	{
-		FSceneComponentInstanceData::ApplyToComponent(Component, CacheApplyPhase);
-		CastChecked<ULightComponent>(Component)->ApplyComponentInstanceData(this);
-	}
-
-	FTransform Transform;
-	FGuid LightGuid;
-	int32 PreviewShadowMapChannel;
-};
-
-FActorComponentInstanceData* ULightComponent::GetComponentInstanceData() const
+TStructOnScope<FActorComponentInstanceData> ULightComponent::GetComponentInstanceData() const
 {
 	// Allocate new struct for holding light map data
-	return new FPrecomputedLightInstanceData(this);
+	return MakeStructOnScope<FActorComponentInstanceData, FPrecomputedLightInstanceData>(this);
 }
 
 void ULightComponent::ApplyComponentInstanceData(FPrecomputedLightInstanceData* LightMapData)
diff --git a/Engine/Source/Runtime/Engine/Private/Components/PrimitiveComponent.cpp b/Engine/Source/Runtime/Engine/Private/Components/PrimitiveComponent.cpp
--- a/Engine/Source/Runtime/Engine/Private/Components/PrimitiveComponent.cpp
+++ b/Engine/Source/Runtime/Engine/Private/Components/PrimitiveComponent.cpp
@@ -627,7 +627,7 @@ void FPrimitiveComponentInstanceData::ApplyToComponent(UActorComponent* Componen
 		PrimitiveComponent->VisibilityId = VisibilityId;
 	}
 
-	if (Component->IsRegistered() && ((VisibilityId != INDEX_NONE) || ContainsSavedProperties()))
+	if (Component->IsRegistered() && ((VisibilityId != INDEX_NONE) || SavedProperties.Num() > 0))
 	{
 		Component->MarkRenderStateDirty();
 	}
@@ -635,12 +635,12 @@ void FPrimitiveComponentInstanceData::ApplyToComponent(UActorComponent* Componen
 
 bool FPrimitiveComponentInstanceData::ContainsData() const
 {
-	return (ContainsSavedProperties() || AttachedInstanceComponents.Num() > 0 || LODParent || (VisibilityId != INDEX_NONE));
+	return (Super::ContainsData() || LODParent || (VisibilityId != INDEX_NONE));
 }
 
 void FPrimitiveComponentInstanceData::AddReferencedObjects(FReferenceCollector& Collector)
 {
-	FSceneComponentInstanceData::AddReferencedObjects(Collector);
+	Super::AddReferencedObjects(Collector);
 
 	// if LOD Parent
 	if (LODParent)
@@ -663,17 +663,9 @@ void FPrimitiveComponentInstanceData::FindAndReplaceInstances(const TMap<UObject
 	}
 }
 
-FActorComponentInstanceData* UPrimitiveComponent::GetComponentInstanceData() const
+TStructOnScope<FActorComponentInstanceData> UPrimitiveComponent::GetComponentInstanceData() const
 {
-	FPrimitiveComponentInstanceData* InstanceData = new FPrimitiveComponentInstanceData(this);
-
-	if (!InstanceData->ContainsData())
-	{
-		delete InstanceData;
-		InstanceData = nullptr;
-	}
-
-	return InstanceData;
+	return MakeStructOnScope<FActorComponentInstanceData, FPrimitiveComponentInstanceData>(this);
 }
 
 void UPrimitiveComponent::OnAttachmentChanged()
diff --git a/Engine/Source/Runtime/Engine/Private/Components/SceneComponent.cpp b/Engine/Source/Runtime/Engine/Private/Components/SceneComponent.cpp
--- a/Engine/Source/Runtime/Engine/Private/Components/SceneComponent.cpp
+++ b/Engine/Source/Runtime/Engine/Private/Components/SceneComponent.cpp
@@ -2214,13 +2214,18 @@ FSceneComponentInstanceData::FSceneComponentInstanceData(const USceneComponent*
 	}
 }
 
+bool FSceneComponentInstanceData::ContainsData() const
+{
+	return AttachedInstanceComponents.Num() > 0 || Super::ContainsData();
+}
+
 void FSceneComponentInstanceData::ApplyToComponent(UActorComponent* Component, const ECacheApplyPhase CacheApplyPhase)
 {
-	FActorComponentInstanceData::ApplyToComponent(Component, CacheApplyPhase);
+	Super::ApplyToComponent(Component, CacheApplyPhase);
 
 	USceneComponent* SceneComponent = CastChecked<USceneComponent>(Component);
 
-	if (ContainsSavedProperties())
+	if (SavedProperties.Num() > 0)
 	{
 		SceneComponent->UpdateComponentToWorld();
 	}
@@ -2251,34 +2256,28 @@ void FSceneComponentInstanceData::AddReferencedObjects(FReferenceCollector& Coll
 
 void FSceneComponentInstanceData::FindAndReplaceInstances(const TMap<UObject*, UObject*>& OldToNewInstanceMap)
 {
-	for (TPair<USceneComponent*, FTransform>& ChildComponentPair : AttachedInstanceComponents)
+	TArray<USceneComponent*> SceneComponents;
+	AttachedInstanceComponents.GenerateKeyArray(SceneComponents);
+
+	for (USceneComponent* SceneComponent : SceneComponents)
 	{
-		if (UObject* const* NewChildComponent = OldToNewInstanceMap.Find(ChildComponentPair.Key))
+		if (UObject* const* NewChildComponent = OldToNewInstanceMap.Find(SceneComponent))
 		{
-			ChildComponentPair.Key = CastChecked<USceneComponent>(*NewChildComponent, ECastCheckedType::NullAllowed);
+			if (*NewChildComponent)
+			{
+				AttachedInstanceComponents.Add(CastChecked<USceneComponent>(*NewChildComponent), AttachedInstanceComponents.FindAndRemoveChecked(SceneComponent));
+			}
+			else
+			{
+				AttachedInstanceComponents.Remove(SceneComponent);
+			}
 		}
 	}
 }
 
-FActorComponentInstanceData* USceneComponent::GetComponentInstanceData() const
+TStructOnScope<FActorComponentInstanceData> USceneComponent::GetComponentInstanceData() const
 {
-	FActorComponentInstanceData* InstanceData = nullptr;
-
-	for (USceneComponent* Child : GetAttachChildren())
-	{
-		if (Child && !Child->IsCreatedByConstructionScript() && !Child->HasAnyFlags(RF_DefaultSubObject))
-		{
-			InstanceData = new FSceneComponentInstanceData(this);
-			break;
-		}
-	}
-
-	if (InstanceData == nullptr)
-	{
-		InstanceData = Super::GetComponentInstanceData();
-	}
-
-	return InstanceData;
+	return MakeStructOnScope<FActorComponentInstanceData, FSceneComponentInstanceData>(this);;
 }
 
 void USceneComponent::UpdateChildTransforms(EUpdateTransformFlags UpdateTransformFlags, ETeleportType Teleport)
diff --git a/Engine/Source/Runtime/Engine/Private/Components/SkyLightComponent.cpp b/Engine/Source/Runtime/Engine/Private/Components/SkyLightComponent.cpp
--- a/Engine/Source/Runtime/Engine/Private/Components/SkyLightComponent.cpp
+++ b/Engine/Source/Runtime/Engine/Private/Components/SkyLightComponent.cpp
@@ -495,37 +495,17 @@ bool USkyLightComponent::IsReadyForFinishDestroy()
 	return Super::IsReadyForFinishDestroy() && ReleaseResourcesFence.IsFenceComplete();
 }
 
-/** Used to store lightmap data during RerunConstructionScripts */
-class FPrecomputedSkyLightInstanceData : public FSceneComponentInstanceData
+TStructOnScope<FActorComponentInstanceData> USkyLightComponent::GetComponentInstanceData() const
 {
-public:
-	FPrecomputedSkyLightInstanceData(const USkyLightComponent* SourceComponent)
-		: FSceneComponentInstanceData(SourceComponent)
-	{}
-
-	virtual void ApplyToComponent(UActorComponent* Component, const ECacheApplyPhase CacheApplyPhase) override
-	{
-		FSceneComponentInstanceData::ApplyToComponent(Component, CacheApplyPhase);
-		CastChecked<USkyLightComponent>(Component)->ApplyComponentInstanceData(this);
-	}
-
-	FGuid LightGuid;
-	// This has to be refcounted to keep it alive during the handoff without doing a deep copy
-	TRefCountPtr<FSkyTextureCubeResource> ProcessedSkyTexture;
-	FSHVectorRGB3 IrradianceEnvironmentMap;
-	float AverageBrightness;
-};
-
-FActorComponentInstanceData* USkyLightComponent::GetComponentInstanceData() const
-{
-	FPrecomputedSkyLightInstanceData* InstanceData = new FPrecomputedSkyLightInstanceData(this);
-	InstanceData->LightGuid = LightGuid;
-	InstanceData->ProcessedSkyTexture = ProcessedSkyTexture;
+	TStructOnScope<FActorComponentInstanceData> InstanceData = MakeStructOnScope<FActorComponentInstanceData, FPrecomputedSkyLightInstanceData>(this);
+	FPrecomputedSkyLightInstanceData* SkyLightInstanceData = InstanceData.Cast<FPrecomputedSkyLightInstanceData>();
+	SkyLightInstanceData->LightGuid = LightGuid;
+	SkyLightInstanceData->ProcessedSkyTexture = ProcessedSkyTexture;
 
 	// Block until the rendering thread has completed its writes from a previous capture
 	IrradianceMapFence.Wait();
-	InstanceData->IrradianceEnvironmentMap = IrradianceEnvironmentMap;
-	InstanceData->AverageBrightness = AverageBrightness;
+	SkyLightInstanceData->IrradianceEnvironmentMap = IrradianceEnvironmentMap;
+	SkyLightInstanceData->AverageBrightness = AverageBrightness;
 
 	return InstanceData;
 }
diff --git a/Engine/Source/Runtime/Engine/Private/Components/SplineComponent.cpp b/Engine/Source/Runtime/Engine/Private/Components/SplineComponent.cpp
--- a/Engine/Source/Runtime/Engine/Private/Components/SplineComponent.cpp
+++ b/Engine/Source/Runtime/Engine/Private/Components/SplineComponent.cpp
@@ -1534,38 +1534,18 @@ FBoxSphereBounds USplineComponent::CalcBounds(const FTransform& LocalToWorld) co
 
 #endif
 
-
-/** Used to store spline data during RerunConstructionScripts */
-class FSplineInstanceData : public FSceneComponentInstanceData
-{
-public:
-	explicit FSplineInstanceData(const USplineComponent* SourceComponent)
-		: FSceneComponentInstanceData(SourceComponent)
-	{
-	}
-
-	virtual void ApplyToComponent(UActorComponent* Component, const ECacheApplyPhase CacheApplyPhase) override
-	{
-		FSceneComponentInstanceData::ApplyToComponent(Component, CacheApplyPhase);
-		CastChecked<USplineComponent>(Component)->ApplyComponentInstanceData(this, (CacheApplyPhase == ECacheApplyPhase::PostUserConstructionScript));
-	}
-
-	FSplineCurves SplineCurves;
-	FSplineCurves SplineCurvesPreUCS;
-	bool bSplineHasBeenEdited;
-};
-
-
-FActorComponentInstanceData* USplineComponent::GetComponentInstanceData() const
+TStructOnScope<FActorComponentInstanceData> USplineComponent::GetComponentInstanceData() const
 {
-	FSplineInstanceData* SplineInstanceData = new FSplineInstanceData(this);
+	TStructOnScope<FActorComponentInstanceData> InstanceData = MakeStructOnScope<FActorComponentInstanceData, FSplineInstanceData>(this);
+	FSplineInstanceData* SplineInstanceData = InstanceData.Cast<FSplineInstanceData>();
+	
 	if (bSplineHasBeenEdited)
 	{
 		SplineInstanceData->SplineCurves = SplineCurves;
 	}
 	SplineInstanceData->bSplineHasBeenEdited = bSplineHasBeenEdited;
 
-	return SplineInstanceData;
+	return InstanceData;
 }
 
 
diff --git a/Engine/Source/Runtime/Engine/Private/Components/SplineMeshComponent.cpp b/Engine/Source/Runtime/Engine/Private/Components/SplineMeshComponent.cpp
--- a/Engine/Source/Runtime/Engine/Private/Components/SplineMeshComponent.cpp
+++ b/Engine/Source/Runtime/Engine/Private/Components/SplineMeshComponent.cpp
@@ -1101,38 +1101,17 @@ void USplineMeshComponent::RecreateCollision()
 	}
 }
 
-/** Used to store spline mesh data during RerunConstructionScripts */
-class FSplineMeshInstanceData : public FSceneComponentInstanceData
+TStructOnScope<FActorComponentInstanceData> USplineMeshComponent::GetComponentInstanceData() const
 {
-public:
-	explicit FSplineMeshInstanceData(const USplineMeshComponent* SourceComponent)
-		: FSceneComponentInstanceData(SourceComponent)
-	{
-	}
-
-	virtual void ApplyToComponent(UActorComponent* Component, const ECacheApplyPhase CacheApplyPhase) override
-	{
-		FSceneComponentInstanceData::ApplyToComponent(Component, CacheApplyPhase);
-		CastChecked<USplineMeshComponent>(Component)->ApplyComponentInstanceData(this);
-	}
-
-	FVector StartPos;
-	FVector EndPos;
-	FVector StartTangent;
-	FVector EndTangent;
-};
-
-FActorComponentInstanceData* USplineMeshComponent::GetComponentInstanceData() const
-{
-	FActorComponentInstanceData* InstanceData = nullptr;
+	TStructOnScope<FActorComponentInstanceData> InstanceData;
 	if (bAllowSplineEditingPerInstance)
 	{
-		FSplineMeshInstanceData *SplineMeshInstanceData = new FSplineMeshInstanceData(this);
+		InstanceData.InitializeAs<FSplineMeshInstanceData>(this);
+		FSplineMeshInstanceData *SplineMeshInstanceData = InstanceData.Cast<FSplineMeshInstanceData>();
 		SplineMeshInstanceData->StartPos = SplineParams.StartPos;
 		SplineMeshInstanceData->EndPos = SplineParams.EndPos;
 		SplineMeshInstanceData->StartTangent = SplineParams.StartTangent;
 		SplineMeshInstanceData->EndTangent = SplineParams.EndTangent;
-		InstanceData = SplineMeshInstanceData;
 	}
 	else
 	{
diff --git a/Engine/Source/Runtime/Engine/Private/Components/StaticMeshComponent.cpp b/Engine/Source/Runtime/Engine/Private/Components/StaticMeshComponent.cpp
--- a/Engine/Source/Runtime/Engine/Private/Components/StaticMeshComponent.cpp
+++ b/Engine/Source/Runtime/Engine/Private/Components/StaticMeshComponent.cpp
@@ -45,145 +45,109 @@
 DECLARE_MEMORY_STAT( TEXT( "StaticMesh VxColor Inst Mem" ), STAT_InstVertexColorMemory, STATGROUP_MemoryStaticMesh );
 DECLARE_MEMORY_STAT( TEXT( "StaticMesh PreCulled Index Memory" ), STAT_StaticMeshPreCulledIndexMemory, STATGROUP_MemoryStaticMesh );
 
-class FStaticMeshComponentInstanceData : public FPrimitiveComponentInstanceData
+
+FStaticMeshComponentInstanceData::FStaticMeshComponentInstanceData(const UStaticMeshComponent* SourceComponent)
+	: FPrimitiveComponentInstanceData(SourceComponent)
+	, StaticMesh(SourceComponent->GetStaticMesh())
+{}
+
+bool FStaticMeshComponentInstanceData::ContainsData() const 
 {
-public:
-	FStaticMeshComponentInstanceData(const UStaticMeshComponent* SourceComponent)
-		: FPrimitiveComponentInstanceData(SourceComponent)
-		, StaticMesh(SourceComponent->GetStaticMesh())
-	{
-	}
+	return Super::ContainsData() 
+		|| StreamingTextureData.Num() > 0 
+		|| CachedStaticLighting.Num() > 0 
+#if WITH_EDITORONLY_DATA
+		|| MaterialStreamingRelativeBoxes.Num() > 0 
+#endif
+		|| StaticMesh != nullptr;
+}
 
-	virtual void ApplyToComponent(UActorComponent* Component, const ECacheApplyPhase CacheApplyPhase) override
+void FStaticMeshComponentInstanceData::ApplyToComponent(UActorComponent* Component, const ECacheApplyPhase CacheApplyPhase) 
+{
+	Super::ApplyToComponent(Component, CacheApplyPhase);
+	if (CacheApplyPhase == ECacheApplyPhase::PostUserConstructionScript)
 	{
-		FPrimitiveComponentInstanceData::ApplyToComponent(Component, CacheApplyPhase);
-		if (CacheApplyPhase == ECacheApplyPhase::PostUserConstructionScript)
-		{
-			CastChecked<UStaticMeshComponent>(Component)->ApplyComponentInstanceData(this);
-		}
+		CastChecked<UStaticMeshComponent>(Component)->ApplyComponentInstanceData(this);
 	}
+}
 
-	virtual void AddReferencedObjects(FReferenceCollector& Collector) override
-	{
-		FPrimitiveComponentInstanceData::AddReferencedObjects(Collector);
-
-		Collector.AddReferencedObject(StaticMesh);
-	}
+void FStaticMeshComponentInstanceData::AddReferencedObjects(FReferenceCollector& Collector)
+{
+	Super::AddReferencedObjects(Collector);
+	Collector.AddReferencedObject(StaticMesh);
+}
 
-	/** Add vertex color data for a specified LOD before RerunConstructionScripts is called */
-	void AddVertexColorData(const struct FStaticMeshComponentLODInfo& LODInfo, uint32 LODIndex)
+/** Add vertex color data for a specified LOD before RerunConstructionScripts is called */
+void FStaticMeshComponentInstanceData::AddVertexColorData(const struct FStaticMeshComponentLODInfo& LODInfo, uint32 LODIndex)
+{
+	if (VertexColorLODs.Num() <= (int32)LODIndex)
 	{
-		if (VertexColorLODs.Num() <= (int32)LODIndex)
-		{
-			VertexColorLODs.SetNum(LODIndex + 1);
-		}
-		FVertexColorLODData& VertexColorData = VertexColorLODs[LODIndex];
-		VertexColorData.LODIndex = LODIndex;
-		VertexColorData.PaintedVertices = LODInfo.PaintedVertices;
-		LODInfo.OverrideVertexColors->GetVertexColors(VertexColorData.VertexBufferColors);
+		VertexColorLODs.SetNum(LODIndex + 1);
 	}
+	FStaticMeshVertexColorLODData& VertexColorData = VertexColorLODs[LODIndex];
+	VertexColorData.LODIndex = LODIndex;
+	VertexColorData.PaintedVertices = LODInfo.PaintedVertices;
+	LODInfo.OverrideVertexColors->GetVertexColors(VertexColorData.VertexBufferColors);
+}
+
+/** Re-apply vertex color data after RerunConstructionScripts is called */
+bool FStaticMeshComponentInstanceData::ApplyVertexColorData(UStaticMeshComponent* StaticMeshComponent) const
+{
+	bool bAppliedAnyData = false;
 
-	/** Re-apply vertex color data after RerunConstructionScripts is called */
-	bool ApplyVertexColorData(UStaticMeshComponent* StaticMeshComponent) const
+	if (StaticMeshComponent != NULL)
 	{
-		bool bAppliedAnyData = false;
+		StaticMeshComponent->SetLODDataCount(VertexColorLODs.Num(), StaticMeshComponent->LODData.Num());
 
-		if (StaticMeshComponent != NULL)
+		// Its possible that we have recreated LODs in SetLODDataCount that existed prior
+		// to reconstruction, but not *rebuilt* them because static lighting usage was clobbered
+		// by the construction script. In this case we should recover the GUIDs we had before
+		// so we dont end up creating new (non-deterministic) data
+		for(int32 LODIndex = 0; LODIndex < StaticMeshComponent->LODData.Num(); ++LODIndex)
 		{
-			StaticMeshComponent->SetLODDataCount(VertexColorLODs.Num(), StaticMeshComponent->LODData.Num());
-
-			// Its possible that we have recreated LODs in SetLODDataCount that existed prior
-			// to reconstruction, but not *rebuilt* them because static lighting usage was clobbered
-			// by the construction script. In this case we should recover the GUIDs we had before
-			// so we dont end up creating new (non-deterministic) data
-			for(int32 LODIndex = 0; LODIndex < StaticMeshComponent->LODData.Num(); ++LODIndex)
+			FStaticMeshComponentLODInfo& LODInfo = StaticMeshComponent->LODData[LODIndex];
+			if(CachedStaticLighting.IsValidIndex(LODIndex))
 			{
-				FStaticMeshComponentLODInfo& LODInfo = StaticMeshComponent->LODData[LODIndex];
-				if(CachedStaticLighting.MapBuildDataIds.IsValidIndex(LODIndex))
-				{
-					LODInfo.MapBuildDataId = CachedStaticLighting.MapBuildDataIds[LODIndex];
+				LODInfo.MapBuildDataId = CachedStaticLighting[LODIndex];
 #if WITH_EDITOR
-					LODInfo.bMapBuildDataIdLoaded = true;
+				LODInfo.bMapBuildDataIdLoaded = true;
 #endif
-				}
 			}
+		}
 
-			for (int32 LODDataIndex = 0; LODDataIndex < VertexColorLODs.Num(); ++LODDataIndex)
-			{
-				const FVertexColorLODData& VertexColorLODData = VertexColorLODs[LODDataIndex];
-				uint32 LODIndex = VertexColorLODData.LODIndex;
+		for (int32 LODDataIndex = 0; LODDataIndex < VertexColorLODs.Num(); ++LODDataIndex)
+		{
+			const FStaticMeshVertexColorLODData& VertexColorLODData = VertexColorLODs[LODDataIndex];
+			uint32 LODIndex = VertexColorLODData.LODIndex;
 
-				if (StaticMeshComponent->LODData.IsValidIndex(LODIndex))
+			if (StaticMeshComponent->LODData.IsValidIndex(LODIndex))
+			{
+				FStaticMeshComponentLODInfo& LODInfo = StaticMeshComponent->LODData[LODIndex];
+				// this component could have been constructed from a template
+				// that had its own vert color overrides; so before we apply
+				// the instance's color data, we need to clear the old
+				// vert colors (so we can properly call InitFromColorArray())
+				StaticMeshComponent->RemoveInstanceVertexColorsFromLOD(LODIndex);
+				// may not be null at the start (could have been initialized 
+				// from a  component template with vert coloring), but should
+				// be null at this point, after RemoveInstanceVertexColorsFromLOD()
+				if (LODInfo.OverrideVertexColors == NULL)
 				{
-					FStaticMeshComponentLODInfo& LODInfo = StaticMeshComponent->LODData[LODIndex];
-					// this component could have been constructed from a template
-					// that had its own vert color overrides; so before we apply
-					// the instance's color data, we need to clear the old
-					// vert colors (so we can properly call InitFromColorArray())
-					StaticMeshComponent->RemoveInstanceVertexColorsFromLOD(LODIndex);
-					// may not be null at the start (could have been initialized 
-					// from a  component template with vert coloring), but should
-					// be null at this point, after RemoveInstanceVertexColorsFromLOD()
-					if (LODInfo.OverrideVertexColors == NULL)
-					{
-						LODInfo.PaintedVertices = VertexColorLODData.PaintedVertices;
+					LODInfo.PaintedVertices = VertexColorLODData.PaintedVertices;
 
-						LODInfo.OverrideVertexColors = new FColorVertexBuffer;
-						LODInfo.OverrideVertexColors->InitFromColorArray(VertexColorLODData.VertexBufferColors);
+					LODInfo.OverrideVertexColors = new FColorVertexBuffer;
+					LODInfo.OverrideVertexColors->InitFromColorArray(VertexColorLODData.VertexBufferColors);
 						
-						check(LODInfo.OverrideVertexColors->GetStride() > 0);
-						BeginInitResource(LODInfo.OverrideVertexColors);
-						bAppliedAnyData = true;
-					}
+					check(LODInfo.OverrideVertexColors->GetStride() > 0);
+					BeginInitResource(LODInfo.OverrideVertexColors);
+					bAppliedAnyData = true;
 				}
 			}
 		}
-
-		return bAppliedAnyData;
 	}
 
-	/** Used to store lightmap data during RerunConstructionScripts */
-	struct FLightMapInstanceData
-	{
-		/** MapBuildDataId from LODData */
-		TArray<FGuid>	MapBuildDataIds;
-	};
-
-	/** Vertex data stored per-LOD */
-	struct FVertexColorLODData
-	{
-		/** copy of painted vertex data */
-		TArray<FPaintedVertex> PaintedVertices;
-
-		/** Copy of vertex buffer colors */
-		TArray<FColor> VertexBufferColors;
-
-		/** Index of the LOD that this data came from */
-		uint32 LODIndex;
-
-		/** Check whether this contains valid data */
-		bool IsValid() const
-		{
-			return PaintedVertices.Num() > 0 && VertexBufferColors.Num() > 0;
-		}
-	};
-
-	/** Mesh being used by component */
-	class UStaticMesh* StaticMesh;
-
-	/** Array of cached vertex colors for each LOD */
-	TArray<FVertexColorLODData> VertexColorLODs;
-
-	FLightMapInstanceData CachedStaticLighting;
-
-	/** Texture streaming build data */
-	TArray<FStreamingTextureBuildInfo> StreamingTextureData;
-
-#if WITH_EDITORONLY_DATA
-	/** Texture streaming editor data (for viewmodes) */
-	TArray<uint32> MaterialStreamingRelativeBoxes;
-#endif
-};
+	return bAppliedAnyData;
+}
 
 UStaticMeshComponent::UStaticMeshComponent(const FObjectInitializer& ObjectInitializer)
 	: Super(ObjectInitializer)
@@ -2256,14 +2220,15 @@ int32 UStaticMeshComponent::GetBlueprintCreatedComponentIndex() const
 	return INDEX_NONE;
 }
 
-FActorComponentInstanceData* UStaticMeshComponent::GetComponentInstanceData() const
+TStructOnScope<FActorComponentInstanceData> UStaticMeshComponent::GetComponentInstanceData() const
 {
-	FStaticMeshComponentInstanceData* StaticMeshInstanceData = new FStaticMeshComponentInstanceData(this);
-
+	TStructOnScope<FActorComponentInstanceData> InstanceData = MakeStructOnScope<FActorComponentInstanceData, FStaticMeshComponentInstanceData>(this);
+	FStaticMeshComponentInstanceData* StaticMeshInstanceData = InstanceData.Cast<FStaticMeshComponentInstanceData>();
+	
 	// Fill in info
 	for (const FStaticMeshComponentLODInfo& LODDataEntry : LODData)
 	{
-		StaticMeshInstanceData->CachedStaticLighting.MapBuildDataIds.Add(LODDataEntry.MapBuildDataId);
+		StaticMeshInstanceData->CachedStaticLighting.Add(LODDataEntry.MapBuildDataId);
 	}
 
 	// Backup the texture streaming data.
@@ -2279,16 +2244,11 @@ FActorComponentInstanceData* UStaticMeshComponent::GetComponentInstanceData() co
 
 		if ( LODInfo.OverrideVertexColors && LODInfo.OverrideVertexColors->GetNumVertices() > 0 && LODInfo.PaintedVertices.Num() > 0 )
 		{
-			if (!StaticMeshInstanceData)
-			{
-				StaticMeshInstanceData = new FStaticMeshComponentInstanceData(this);
-			}
-
 			StaticMeshInstanceData->AddVertexColorData(LODInfo, LODIndex);
 		}
 	}
 
-	return StaticMeshInstanceData;
+	return InstanceData;
 }
 
 void UStaticMeshComponent::ApplyComponentInstanceData(FStaticMeshComponentInstanceData* StaticMeshInstanceData)
@@ -2303,7 +2263,7 @@ void UStaticMeshComponent::ApplyComponentInstanceData(FStaticMeshComponentInstan
 		return;
 	}
 
-	const int32 NumLODLightMaps = StaticMeshInstanceData->CachedStaticLighting.MapBuildDataIds.Num();
+	const int32 NumLODLightMaps = StaticMeshInstanceData->CachedStaticLighting.Num();
 
 	if (HasStaticLighting() && NumLODLightMaps > 0)
 	{
@@ -2314,7 +2274,7 @@ void UStaticMeshComponent::ApplyComponentInstanceData(FStaticMeshComponentInstan
 
 			for (int32 i = 0; i < NumLODLightMaps; ++i)
 			{
-				LODData[i].MapBuildDataId = StaticMeshInstanceData->CachedStaticLighting.MapBuildDataIds[i];
+				LODData[i].MapBuildDataId = StaticMeshInstanceData->CachedStaticLighting[i];
 			#if WITH_EDITOR
 				LODData[i].bMapBuildDataIdLoaded = true;
 			#endif
diff --git a/Engine/Source/Runtime/Engine/Private/InstancedStaticMesh.cpp b/Engine/Source/Runtime/Engine/Private/InstancedStaticMesh.cpp
--- a/Engine/Source/Runtime/Engine/Private/InstancedStaticMesh.cpp
+++ b/Engine/Source/Runtime/Engine/Private/InstancedStaticMesh.cpp
@@ -943,63 +943,12 @@ UInstancedStaticMeshComponent::~UInstancedStaticMeshComponent()
 	ReleasePerInstanceRenderData();
 }
 
-#if WITH_EDITOR
-/** Helper class used to preserve lighting/selection state across blueprint reinstancing */
-class FInstancedStaticMeshComponentInstanceData : public FSceneComponentInstanceData
-{
-public:
-	
-	FInstancedStaticMeshComponentInstanceData(const UInstancedStaticMeshComponent& InComponent)
-		: FSceneComponentInstanceData(&InComponent)
-		, StaticMesh(InComponent.GetStaticMesh())
-	{
-	}
-
-	virtual void ApplyToComponent(UActorComponent* Component, const ECacheApplyPhase CacheApplyPhase) override
-	{
-		FSceneComponentInstanceData::ApplyToComponent(Component, CacheApplyPhase);
-		CastChecked<UInstancedStaticMeshComponent>(Component)->ApplyComponentInstanceData(this);
-	}
-
-	virtual void AddReferencedObjects(FReferenceCollector& Collector) override
-	{
-		FSceneComponentInstanceData::AddReferencedObjects(Collector);
-
-		Collector.AddReferencedObject(StaticMesh);
-	}
-
-	/** Used to store lightmap data during RerunConstructionScripts */
-	struct FLightMapInstanceData
-	{
-		/** Transform of component */
-		FTransform Transform;
-
-		/** guid from LODData */
-		TArray<FGuid> MapBuildDataIds;
-	};
-
-public:
-	/** Mesh being used by component */
-	UStaticMesh* StaticMesh;
-
-	// Static lighting info
-	FLightMapInstanceData CachedStaticLighting;
-	TArray<FInstancedStaticMeshInstanceData> PerInstanceSMData;
-
-	/** The cached selected instances */
-	TBitArray<> SelectedInstances;
-
-	/* The cached random seed */
-	int32 InstancingRandomSeed;
-};
-#endif
-
-FActorComponentInstanceData* UInstancedStaticMeshComponent::GetComponentInstanceData() const
+TStructOnScope<FActorComponentInstanceData> UInstancedStaticMeshComponent::GetComponentInstanceData() const
 {
+	TStructOnScope<FActorComponentInstanceData> InstanceData;
 #if WITH_EDITOR
-	FActorComponentInstanceData* InstanceData = nullptr;
-	FInstancedStaticMeshComponentInstanceData* StaticMeshInstanceData = nullptr;
-	InstanceData = StaticMeshInstanceData = new FInstancedStaticMeshComponentInstanceData(*this);	
+	InstanceData.InitializeAs<FInstancedStaticMeshComponentInstanceData>(this);
+	FInstancedStaticMeshComponentInstanceData* StaticMeshInstanceData = InstanceData.Cast<FInstancedStaticMeshComponentInstanceData>();
 
 	// Fill in info (copied from UStaticMeshComponent::GetComponentInstanceData)
 	StaticMeshInstanceData->CachedStaticLighting.Transform = GetComponentTransform();
@@ -1017,11 +966,8 @@ FActorComponentInstanceData* UInstancedStaticMeshComponent::GetComponentInstance
 
 	// Back up random seed
 	StaticMeshInstanceData->InstancingRandomSeed = InstancingRandomSeed;
-
-	return InstanceData;
-#else
-	return nullptr;
 #endif
+	return InstanceData;
 }
 
 void UInstancedStaticMeshComponent::ApplyComponentInstanceData(FInstancedStaticMeshComponentInstanceData* InstancedMeshData)
diff --git a/Engine/Source/Runtime/Engine/Public/ComponentInstanceDataCache.h b/Engine/Source/Runtime/Engine/Public/ComponentInstanceDataCache.h
--- a/Engine/Source/Runtime/Engine/Public/ComponentInstanceDataCache.h
+++ b/Engine/Source/Runtime/Engine/Public/ComponentInstanceDataCache.h
@@ -4,6 +4,8 @@
 
 #include "CoreMinimal.h"
 #include "UObject/Object.h"
+#include "UObject/StructOnScope.h"
+#include "ComponentInstanceDataCache.generated.h"
 
 class AActor;
 class UActorComponent;
@@ -17,19 +19,62 @@ enum class ECacheApplyPhase
 	PostUserConstructionScript,		// After the user construction script has been run
 };
 
+UENUM()
+enum class EComponentCreationMethod : uint8
+{
+	/** A component that is part of a native class. */
+	Native,
+	/** A component that is created from a template defined in the Components section of the Blueprint. */
+	SimpleConstructionScript,
+	/**A dynamically created component, either from the UserConstructionScript or from a Add Component node in a Blueprint event graph. */
+	UserConstructionScript,
+	/** A component added to a single Actor instance via the Component section of the Actor's details panel. */
+	Instance,
+};
+
+USTRUCT()
+struct FActorComponentDuplicatedObjectData
+{
+	GENERATED_BODY()
+
+	FActorComponentDuplicatedObjectData(UObject* InObject = nullptr);
+
+	bool Serialize(FArchive& Ar);
+
+	// The duplicated object
+	UObject* DuplicatedObject;
+
+	// Object Outer Depth so we can sort creation order
+	int32 ObjectPathDepth;
+};
+
+// Trait to signal ActorCompomentInstanceData duplicated objects uses a serialize function
+template<>
+struct TStructOpsTypeTraits<FActorComponentDuplicatedObjectData> : public TStructOpsTypeTraitsBase2<FActorComponentDuplicatedObjectData>
+{
+	enum
+	{
+		WithSerializer = true,
+	};
+};
+
 /** Base class for component instance cached data of a particular type. */
-class ENGINE_API FActorComponentInstanceData
+USTRUCT()
+struct ENGINE_API FActorComponentInstanceData
 {
+	GENERATED_BODY()
 public:
 	FActorComponentInstanceData();
 	FActorComponentInstanceData(const UActorComponent* SourceComponent);
 
-	virtual ~FActorComponentInstanceData()
-	{}
+	virtual ~FActorComponentInstanceData() = default;
 
 	/** Determines whether this component instance data matches the component */
 	bool MatchesComponent(const UActorComponent* Component, const UObject* ComponentTemplate, const TMap<UActorComponent*, const UObject*>& ComponentToArchetypeMap) const;
 
+	/** Determines if any instance data was actually saved. */
+	virtual bool ContainsData() const { return SavedProperties.Num() > 0; }
+
 	/** Applies this component instance data to the supplied component */
 	virtual void ApplyToComponent(UActorComponent* Component, const ECacheApplyPhase CacheApplyPhase);
 
@@ -38,24 +83,41 @@ public:
 
 	virtual void AddReferencedObjects(FReferenceCollector& Collector);
 
-	bool ContainsSavedProperties() const { return SavedProperties.Num() > 0; }
-
 	const UClass* GetComponentClass() const { return SourceComponentTemplate ? SourceComponentTemplate->GetClass() : nullptr; }
 
+	const UObject* GetComponentTemplate() const { return SourceComponentTemplate; }
+
 protected:
+	friend class FComponentPropertyWriter;
+	friend class FComponentPropertyReader;
+
 	/** The template used to create the source component */
+	UPROPERTY()
 	const UObject* SourceComponentTemplate;
 
-	/** The index of the source component in its owner's serialized array 
-		when filtered to just that component type */
-	int32 SourceComponentTypeSerializedIndex;
-
 	/** The method that was used to create the source component */
+	UPROPERTY() 
 	EComponentCreationMethod SourceComponentCreationMethod;
 
+	/** The index of the source component in its owner's serialized array
+	when filtered to just that component type */
+	UPROPERTY()
+	int32 SourceComponentTypeSerializedIndex;
+
+	UPROPERTY()
 	TArray<uint8> SavedProperties;
-	TArray<UObject*> InstancedObjects;
-	TArray<uint32_t> ObjectReferenceIndicesInByteArray;
+
+	// Duplicated objects created when saving component instance properties
+	UPROPERTY()
+	TArray<FActorComponentDuplicatedObjectData> DuplicatedObjects;
+
+	// Referenced objects in component instance saved properties
+	UPROPERTY()
+	TArray<UObject*> ReferencedObjects;
+
+	// Referenced names in component instance saved properties
+	UPROPERTY()
+	TArray<FName> ReferencedNames;
 };
 
 /** 
@@ -66,12 +128,12 @@ class ENGINE_API FComponentInstanceDataCache
 {
 public:
 
-	FComponentInstanceDataCache() {}
+	FComponentInstanceDataCache() = default;
 
 	/** Constructor that also populates cache from Actor */
 	FComponentInstanceDataCache(const AActor* InActor);
 
-	~FComponentInstanceDataCache();
+	~FComponentInstanceDataCache() = default;
 
 	/** Non-copyable */
 	FComponentInstanceDataCache(const FComponentInstanceDataCache&) = delete;
@@ -81,6 +143,9 @@ public:
 	FComponentInstanceDataCache(FComponentInstanceDataCache&&) = default;
 	FComponentInstanceDataCache& operator=(FComponentInstanceDataCache&&) = default;
 
+	/** Serialize Instance data for persistence or transmission. */
+	void Serialize(FArchive& Ar);
+
 	/** Iterates over an Actor's components and applies the stored component instance data to each */
 	void ApplyToActor(AActor* Actor, const ECacheApplyPhase CacheApplyPhase) const;
 
@@ -92,8 +157,12 @@ public:
 	void AddReferencedObjects(FReferenceCollector& Collector);
 
 private:
-	/** Map of data type name to data of that type */
-	TArray< FActorComponentInstanceData* > ComponentsInstanceData;
+	// called during de-serialization to copy serialized properties over existing component instance data and keep non UPROPERTY data intact
+	void CopySerializableProperties(TArray<TStructOnScope<FActorComponentInstanceData>> InComponentsInstanceData);
+
+	/** Map of component instance data struct (template -> instance data) */
+	TArray<TStructOnScope<FActorComponentInstanceData>> ComponentsInstanceData;
 
+	/** Map of the actor instanced scene component to their transform relative to the root. */
 	TMap< USceneComponent*, FTransform > InstanceComponentTransformToRootMap;
 };
diff --git a/Engine/Source/Runtime/UMG/Private/Components/WidgetComponent.cpp b/Engine/Source/Runtime/UMG/Private/Components/WidgetComponent.cpp
--- a/Engine/Source/Runtime/UMG/Private/Components/WidgetComponent.cpp
+++ b/Engine/Source/Runtime/UMG/Private/Components/WidgetComponent.cpp
@@ -1176,38 +1176,9 @@ void UWidgetComponent::RemoveWidgetFromScreen()
 #endif // !UE_SERVER
 }
 
-class FWidgetComponentInstanceData : public FSceneComponentInstanceData
+TStructOnScope<FActorComponentInstanceData> UWidgetComponent::GetComponentInstanceData() const
 {
-public:
-	FWidgetComponentInstanceData( const UWidgetComponent* SourceComponent )
-		: FSceneComponentInstanceData(SourceComponent)
-		, WidgetClass ( SourceComponent->GetWidgetClass() )
-		, RenderTarget( SourceComponent->GetRenderTarget() )
-	{}
-
-	virtual void ApplyToComponent(UActorComponent* Component, const ECacheApplyPhase CacheApplyPhase) override
-	{
-		FSceneComponentInstanceData::ApplyToComponent(Component, CacheApplyPhase);
-		CastChecked<UWidgetComponent>(Component)->ApplyComponentInstanceData(this);
-	}
-
-	virtual void AddReferencedObjects(FReferenceCollector& Collector) override
-	{
-		FSceneComponentInstanceData::AddReferencedObjects(Collector);
-
-		UClass* WidgetUClass = *WidgetClass;
-		Collector.AddReferencedObject(WidgetUClass);
-		Collector.AddReferencedObject(RenderTarget);
-	}
-
-public:
-	TSubclassOf<UUserWidget> WidgetClass;
-	UTextureRenderTarget2D* RenderTarget;
-};
-
-FActorComponentInstanceData* UWidgetComponent::GetComponentInstanceData() const
-{
-	return new FWidgetComponentInstanceData( this );
+	return MakeStructOnScope<FActorComponentInstanceData, FWidgetComponentInstanceData>(this);
 }
 
 void UWidgetComponent::ApplyComponentInstanceData(FWidgetComponentInstanceData* WidgetInstanceData)
diff --git a/Engine/Source/Runtime/UMG/Public/Components/WidgetComponent.h b/Engine/Source/Runtime/UMG/Public/Components/WidgetComponent.h
--- a/Engine/Source/Runtime/UMG/Public/Components/WidgetComponent.h
+++ b/Engine/Source/Runtime/UMG/Public/Components/WidgetComponent.h
@@ -90,8 +90,8 @@ public:
 	int32 GetNumMaterials() const override;
 	virtual void TickComponent(float DeltaTime, enum ELevelTick TickType, FActorComponentTickFunction *ThisTickFunction) override;
 
-	virtual FActorComponentInstanceData* GetComponentInstanceData() const override;
-	void ApplyComponentInstanceData(class FWidgetComponentInstanceData* ComponentInstanceData);
+	virtual TStructOnScope<FActorComponentInstanceData> GetComponentInstanceData() const override;
+	void ApplyComponentInstanceData(struct FWidgetComponentInstanceData* ComponentInstanceData);
 	virtual void GetUsedMaterials(TArray<UMaterialInterface*>& OutMaterials, bool bGetDebugMaterials = false) const override;
 
 #if WITH_EDITOR
@@ -571,3 +571,44 @@ protected:
 	/** Helper class for drawing widgets to a render target. */
 	class FWidgetRenderer* WidgetRenderer;
 };
+
+USTRUCT()
+struct FWidgetComponentInstanceData : public FSceneComponentInstanceData
+{
+	GENERATED_BODY()
+public:
+	FWidgetComponentInstanceData()
+		: RenderTarget(nullptr)
+	{}
+
+	FWidgetComponentInstanceData(const UWidgetComponent* SourceComponent)
+		: FSceneComponentInstanceData(SourceComponent)
+		, WidgetClass(SourceComponent->GetWidgetClass())
+		, RenderTarget(SourceComponent->GetRenderTarget())
+	{}
+	virtual ~FWidgetComponentInstanceData() = default;
+
+	virtual bool ContainsData() const override
+	{
+		return true;
+	}
+
+	virtual void ApplyToComponent(UActorComponent* Component, const ECacheApplyPhase CacheApplyPhase) override
+	{
+		Super::ApplyToComponent(Component, CacheApplyPhase);
+		CastChecked<UWidgetComponent>(Component)->ApplyComponentInstanceData(this);
+	}
+
+	virtual void AddReferencedObjects(FReferenceCollector& Collector) override
+	{
+		Super::AddReferencedObjects(Collector);
+
+		UClass* WidgetUClass = *WidgetClass;
+		Collector.AddReferencedObject(WidgetUClass);
+		Collector.AddReferencedObject(RenderTarget);
+	}
+
+public:
+	TSubclassOf<UUserWidget> WidgetClass;
+	UTextureRenderTarget2D* RenderTarget;
+};
