Fix cook build

diff --git a/Engine/Source/Runtime/Engine/Private/SkeletalMeshLODModel.cpp b/Engine/Source/Runtime/Engine/Private/SkeletalMeshLODModel.cpp
--- a/Engine/Source/Runtime/Engine/Private/SkeletalMeshLODModel.cpp
+++ b/Engine/Source/Runtime/Engine/Private/SkeletalMeshLODModel.cpp
@@ -695,17 +695,19 @@ void FSkeletalMeshLODModel::Serialize(FArchive& Ar, UObject* Owner, int32 Idx)
 					DummyColorBuffer.Serialize(Ar, false);
 					//Copy the data to the softVertices
 					int32 VertexColorCount = DummyColorBuffer.GetNumVertices();
-					check(NumVertices == VertexColorCount);
-					TArray<FColor> OutColors;
-					DummyColorBuffer.GetVertexColors(OutColors);
-					int32 DummyVertexColorIndex = 0;
-					for (int32 SectionIndex = 0; SectionIndex < Sections.Num(); ++SectionIndex)
+					if (NumVertices == VertexColorCount)
 					{
-						int32 SectionVertexCount = Sections[SectionIndex].GetNumVertices();
-						TArray<FSoftSkinVertex>& SoftVertices = Sections[SectionIndex].SoftVertices;
-						for (int32 SectionVertexIndex = 0; SectionVertexIndex < SectionVertexCount; ++SectionVertexIndex)
+						TArray<FColor> OutColors;
+						DummyColorBuffer.GetVertexColors(OutColors);
+						int32 DummyVertexColorIndex = 0;
+						for (int32 SectionIndex = 0; SectionIndex < Sections.Num(); ++SectionIndex)
 						{
-							SoftVertices[SectionVertexIndex].Color = OutColors[DummyVertexColorIndex++];
+							int32 SectionVertexCount = Sections[SectionIndex].GetNumVertices();
+							TArray<FSoftSkinVertex>& SoftVertices = Sections[SectionIndex].SoftVertices;
+							for (int32 SectionVertexIndex = 0; SectionVertexIndex < SectionVertexCount; ++SectionVertexIndex)
+							{
+								SoftVertices[SectionVertexIndex].Color = OutColors[DummyVertexColorIndex++];
+							}
 						}
 					}
 				}
