Fix for CIS error and static analysis warning

diff --git a/Engine/Source/Runtime/Engine/Private/ComponentInstanceDataCache.cpp b/Engine/Source/Runtime/Engine/Private/ComponentInstanceDataCache.cpp
--- a/Engine/Source/Runtime/Engine/Private/ComponentInstanceDataCache.cpp
+++ b/Engine/Source/Runtime/Engine/Private/ComponentInstanceDataCache.cpp
@@ -239,16 +239,16 @@ bool FActorComponentDuplicatedObjectData::Serialize(FArchive& Ar)
 	FString ObjectClassPath;
 	FString ObjectOuterPath;
 	FName ObjectName;
-	uint32 ObjectPersistentFlags;
+	uint32 ObjectPersistentFlags = 0;
 	TArray<uint8> ObjectData;
 
-	if (Ar.IsSaving())
+	if (Ar.IsSaving() && DuplicatedObject)
 	{
-		UClass* ObjectClass = DuplicatedObject ? DuplicatedObject->GetClass() : nullptr;
+		UClass* ObjectClass = DuplicatedObject->GetClass();
 		ObjectClassPath = ObjectClass ? ObjectClass->GetPathName() : FString();
-		ObjectOuterPath = DuplicatedObject && DuplicatedObject->GetOuter() ? DuplicatedObject->GetOuter()->GetPathName() : FString();
+		ObjectOuterPath = DuplicatedObject->GetOuter() ? DuplicatedObject->GetOuter()->GetPathName() : FString();
 		ObjectName = DuplicatedObject->GetFName();
-		ObjectPersistentFlags = DuplicatedObject ? DuplicatedObject->GetFlags() & RF_Load : 0;
+		ObjectPersistentFlags = DuplicatedObject->GetFlags() & RF_Load;
 		if (ObjectClass)
 		{
 			FMemoryWriter Writer(ObjectData);
