Fix accidental removal of the non dynamic subsystem code path.

diff --git a/Engine/Source/Runtime/Engine/Private/Subsystems/SubsystemCollection.cpp b/Engine/Source/Runtime/Engine/Private/Subsystems/SubsystemCollection.cpp
--- a/Engine/Source/Runtime/Engine/Private/Subsystems/SubsystemCollection.cpp
+++ b/Engine/Source/Runtime/Engine/Private/Subsystems/SubsystemCollection.cpp
@@ -103,9 +103,22 @@ void FSubsystemCollectionBase::Initialize()
 		
 		TGuardValue<bool> PopulatingGuard(bPopulating, true);
 
-		for (const TPair<FName, TArray<TSubclassOf<UDynamicSubsystem>>>& SubsystemClasses : DynamicSystemModuleMap)
+		if (BaseType->IsChildOf(UDynamicSubsystem::StaticClass()))
 		{
-			for (const TSubclassOf<UDynamicSubsystem>& SubsystemClass : SubsystemClasses.Value)
+			for (const TPair<FName, TArray<TSubclassOf<UDynamicSubsystem>>>& SubsystemClasses : DynamicSystemModuleMap)
+			{
+				for (const TSubclassOf<UDynamicSubsystem>& SubsystemClass : SubsystemClasses.Value)
+				{
+					AddAndInitializeSubsystem(SubsystemClass);
+				}
+			}
+		}
+		else
+		{
+			TArray<UClass*> SubsystemClasses;
+			GetDerivedClasses(BaseType, SubsystemClasses, true);
+
+			for (UClass* SubsystemClass : SubsystemClasses)
 			{
 				AddAndInitializeSubsystem(SubsystemClass);
 			}
