Properly filter dynamic systems when adding to collections also added a Check in the internal add code to catch any future mistakes in the higher level code.

diff --git a/Engine/Source/Runtime/Engine/Private/Subsystems/SubsystemCollection.cpp b/Engine/Source/Runtime/Engine/Private/Subsystems/SubsystemCollection.cpp
--- a/Engine/Source/Runtime/Engine/Private/Subsystems/SubsystemCollection.cpp
+++ b/Engine/Source/Runtime/Engine/Private/Subsystems/SubsystemCollection.cpp
@@ -109,7 +109,10 @@ void FSubsystemCollectionBase::Initialize()
 			{
 				for (const TSubclassOf<UDynamicSubsystem>& SubsystemClass : SubsystemClasses.Value)
 				{
-					AddAndInitializeSubsystem(SubsystemClass);
+					if (SubsystemClass->IsChildOf(BaseType))
+					{
+						AddAndInitializeSubsystem(SubsystemClass);
+					}
 				}
 			}
 		}
@@ -153,7 +156,7 @@ void FSubsystemCollectionBase::Deinitialize()
 	SubsystemMap.Empty();
 }
 
-bool FSubsystemCollectionBase::InitializeDependancy(TSubclassOf<USubsystem> SubsystemClass)
+bool FSubsystemCollectionBase::InitializeDependency(TSubclassOf<USubsystem> SubsystemClass)
 {
 	if (ensureMsgf(SubsystemClass, TEXT("Attempting to add invalid subsystem as dependancy."))
 		&& ensureMsgf(bPopulating, TEXT("InitializeDependancy() should only be called from System USubsystem::Initialization() implementations."))
@@ -174,6 +177,9 @@ bool FSubsystemCollectionBase::AddAndInitializeSubsystem(UClass* SubsystemClass)
 {
 	if (!SubsystemMap.Contains(SubsystemClass))
 	{
+		// Catch any attempt to add a subsystem of the wrong type
+		checkf(SubsystemClass->IsChildOf(BaseType), TEXT("ClassType (%s) must be a subclass of BaseType(%s)."), *SubsystemClass->GetName(), *BaseType->GetName());
+
 		// Only add instances for non abstract Subsystems
 		if (SubsystemClass && !SubsystemClass->HasAllClassFlags(CLASS_Abstract))
 		{
diff --git a/Engine/Source/Runtime/Engine/Public/Subsystems/SubsystemCollection.h b/Engine/Source/Runtime/Engine/Public/Subsystems/SubsystemCollection.h
--- a/Engine/Source/Runtime/Engine/Public/Subsystems/SubsystemCollection.h
+++ b/Engine/Source/Runtime/Engine/Public/Subsystems/SubsystemCollection.h
@@ -22,7 +22,7 @@ public:
 	 * Only call from Initialize() of Systems to ensure initialization order
 	 * Note: Dependencies only work within a collection
 	 */
-	bool InitializeDependancy(TSubclassOf<USubsystem> SubsystemClass);
+	bool InitializeDependency(TSubclassOf<USubsystem> SubsystemClass);
 
 	/* FGCObject Interface */
 	virtual void AddReferencedObjects(FReferenceCollector& Collector) override;
