Exposed GetLoadedLevel to scripting

diff --git a/Engine/Source/Runtime/Engine/Classes/Engine/LevelStreaming.h b/Engine/Source/Runtime/Engine/Classes/Engine/LevelStreaming.h
--- a/Engine/Source/Runtime/Engine/Classes/Engine/LevelStreaming.h
+++ b/Engine/Source/Runtime/Engine/Classes/Engine/LevelStreaming.h
@@ -359,6 +359,7 @@ public:
 	FBox GetStreamingVolumeBounds();
 
 	/** Gets a pointer to the LoadedLevel value */
+	UFUNCTION(BlueprintCallable, Category="Game")
 	ULevel* GetLoadedLevel() const { return LoadedLevel; }
 	
 	/** Sets the LoadedLevel value to NULL */
