Add object scope on the initializing animation to capture initialization time per object.

diff --git a/Engine/Source/Runtime/Engine/Private/Animation/AnimInstance.cpp b/Engine/Source/Runtime/Engine/Private/Animation/AnimInstance.cpp
--- a/Engine/Source/Runtime/Engine/Private/Animation/AnimInstance.cpp
+++ b/Engine/Source/Runtime/Engine/Private/Animation/AnimInstance.cpp
@@ -166,6 +166,7 @@ UWorld* UAnimInstance::GetWorld() const
 
 void UAnimInstance::InitializeAnimation()
 {
+	FScopeCycleCounterUObject ContextScope(this);
 	SCOPE_CYCLE_COUNTER(STAT_AnimInitTime);
 
 	UninitializeAnimation();
