Add info the assertion

diff --git a/Engine/Source/Runtime/Engine/Private/EdGraph/EdGraphPin.cpp b/Engine/Source/Runtime/Engine/Private/EdGraph/EdGraphPin.cpp
--- a/Engine/Source/Runtime/Engine/Private/EdGraph/EdGraphPin.cpp
+++ b/Engine/Source/Runtime/Engine/Private/EdGraph/EdGraphPin.cpp
@@ -1842,7 +1842,26 @@ bool UEdGraphPin::SerializePin(FArchive& Ar, UEdGraphPin*& PinRef, int32 ArrayId
 		FGuid PinGuid;
 		if (!Ar.IsLoading())
 		{
-			check(!PinRef->bWasTrashed);
+#if DO_CHECK
+			if(PinRef->bWasTrashed)
+			{
+				if(RequestingPin)
+				{
+					if(!RequestingPin->bWasTrashed)
+					{
+						checkf(false, TEXT("Pin named %s was serialized while trashed - requesting pin named %s owned by node %s"), *(PinRef->PinName.ToString()), *RequestingPin->PinName.ToString(), *RequestingPin->GetOwningNode()->GetName());
+					}
+					else
+					{
+						checkf(false, TEXT("Pin named %s was serialized while trashed - requesting pin named %s"), *(PinRef->PinName.ToString()), *RequestingPin->PinName.ToString());
+					}
+				}
+				else
+				{
+					checkf(false, TEXT("Pin named %s was serialized while trashed with no requesting pin"), *(PinRef->PinName.ToString()));
+				}
+			}
+#endif //DO_CHECK
 			LocalOwningNode = PinRef->OwningNode;
 			PinGuid = PinRef->PinId;
 		}
