Add checks to prevent vector nan ensure in ConvertOverlappedShapeToImpactHit and ConvertQueryImpactHit.

diff --git a/Engine/Source/Runtime/Engine/Private/Collision/CollisionConversions.cpp b/Engine/Source/Runtime/Engine/Private/Collision/CollisionConversions.cpp
--- a/Engine/Source/Runtime/Engine/Private/Collision/CollisionConversions.cpp
+++ b/Engine/Source/Runtime/Engine/Private/Collision/CollisionConversions.cpp
@@ -237,23 +237,26 @@ EConvertQueryResult ConvertQueryImpactHit(const UWorld* World, const FHitLocatio
 	OutResult.Location = SafeLocationToFitShape;
 
 	const bool bUseReturnedPoint = ((Flags & EHitFlags::Position) && !bInitialOverlap);
-	const FVector Position = GetPosition(Hit);
-	if (bUseReturnedPoint && Position.ContainsNaN())
+	FVector Position = StartLoc;
+	if (bUseReturnedPoint)
 	{
+		Position = GetPosition(Hit);
+		if (Position.ContainsNaN())
+		{
 #if ENABLE_NAN_DIAGNOSTIC
-		SetHitResultFromShapeAndFaceIndex(HitShape, HitActor, InternalFaceIndex, OutResult, bReturnPhysMat);
-		UE_LOG(LogCore, Error, TEXT("ConvertQueryImpactHit() NaN details:\n>> Actor:%s (%s)\n>> Component:%s\n>> Item:%d\n>> BoneName:%s\n>> Time:%f\n>> Distance:%f\n>> Location:%s\n>> bIsBlocking:%d\n>> bStartPenetrating:%d"),
-			*GetNameSafe(OutResult.GetActor()), OutResult.Actor.IsValid() ? *OutResult.GetActor()->GetPathName() : TEXT("no path"),
-			*GetNameSafe(OutResult.GetComponent()), OutResult.Item, *OutResult.BoneName.ToString(),
-			OutResult.Time, OutResult.Distance, *OutResult.Location.ToString(), OutResult.bBlockingHit ? 1 : 0, OutResult.bStartPenetrating ? 1 : 0);
+			SetHitResultFromShapeAndFaceIndex(HitShape, HitActor, InternalFaceIndex, OutResult, bReturnPhysMat);
+			UE_LOG(LogCore, Error, TEXT("ConvertQueryImpactHit() NaN details:\n>> Actor:%s (%s)\n>> Component:%s\n>> Item:%d\n>> BoneName:%s\n>> Time:%f\n>> Distance:%f\n>> Location:%s\n>> bIsBlocking:%d\n>> bStartPenetrating:%d"),
+				*GetNameSafe(OutResult.GetActor()), OutResult.Actor.IsValid() ? *OutResult.GetActor()->GetPathName() : TEXT("no path"),
+				*GetNameSafe(OutResult.GetComponent()), OutResult.Item, *OutResult.BoneName.ToString(),
+				OutResult.Time, OutResult.Distance, *OutResult.Location.ToString(), OutResult.bBlockingHit ? 1 : 0, OutResult.bStartPenetrating ? 1 : 0);
 #endif // ENABLE_NAN_DIAGNOSTIC
 
-		OutResult.Reset();
-		logOrEnsureNanError(TEXT("ConvertQueryImpactHit() received NaN/Inf for position: %s"), *Position.ToString());
-		return EConvertQueryResult::Invalid;
+			OutResult.Reset();
+			logOrEnsureNanError(TEXT("ConvertQueryImpactHit() received NaN/Inf for position: %s"), *Position.ToString());
+			return EConvertQueryResult::Invalid;
+		}
 	}
-
-	OutResult.ImpactPoint = bUseReturnedPoint ? Position : StartLoc;
+	OutResult.ImpactPoint = Position;
 	
 	// Caution: we may still have an initial overlap, but with null Geom. This is the case for RayCast results.
 	const bool bUseReturnedNormal = ((Flags & EHitFlags::Normal) && !bInitialOverlap);
@@ -400,8 +403,26 @@ static bool ConvertOverlappedShapeToImpactHit(const UWorld* World, const FHitLoc
 
 	// Return start location as 'safe location'
 	OutResult.Location = QueryTM.GetLocation();
-	OutResult.ImpactPoint = GetPosition(Hit);
 
+	const bool bValidPosition = (GetFlags(Hit) & EHitFlags::Position);
+	if (bValidPosition)
+	{
+		const FVector HitPosition = GetPosition(Hit);
+		const bool bFinitePosition = !HitPosition.ContainsNaN();
+		if (bFinitePosition)
+		{
+			OutResult.ImpactPoint = HitPosition;
+		}
+		else
+		{
+			OutResult.ImpactPoint = StartLoc;
+			UE_LOG(LogPhysics, Verbose, TEXT("Warning: ConvertOverlappedShapeToImpactHit: MTD returned NaN :( position: %s"), *HitPosition.ToString());
+		}
+	}
+	else
+	{
+		OutResult.ImpactPoint = StartLoc;
+	}
 	OutResult.TraceStart = StartLoc;
 	OutResult.TraceEnd = EndLoc;
 
