Renderpass updates.
Canvas and Unreal updates.

diff --git a/Engine/Source/Runtime/Engine/Private/CanvasRenderTarget2D.cpp b/Engine/Source/Runtime/Engine/Private/CanvasRenderTarget2D.cpp
--- a/Engine/Source/Runtime/Engine/Private/CanvasRenderTarget2D.cpp
+++ b/Engine/Source/Runtime/Engine/Private/CanvasRenderTarget2D.cpp
@@ -66,32 +66,26 @@ void UCanvasRenderTarget2D::RepaintCanvas()
 	// Create the FCanvas which does the actual rendering.
 	const UWorld* WorldPtr = World.Get();
 	const ERHIFeatureLevel::Type FeatureLevel = WorldPtr != nullptr ? World->FeatureLevel : GMaxRHIFeatureLevel;
-	FCanvas RenderCanvas(GameThread_GetRenderTargetResource(), nullptr, FApp::GetCurrentTime() - GStartTime, FApp::GetDeltaTime(), FApp::GetCurrentTime() - GStartTime, FeatureLevel);
 
+	FTextureRenderTarget2DResource* TextureRenderTarget = (FTextureRenderTarget2DResource*) GameThread_GetRenderTargetResource();
+
+	FCanvas RenderCanvas(TextureRenderTarget, nullptr, FApp::GetCurrentTime() - GStartTime, FApp::GetDeltaTime(), FApp::GetCurrentTime() - GStartTime, FeatureLevel);
 	Canvas->Init(GetSurfaceWidth(), GetSurfaceHeight(), nullptr, &RenderCanvas);
 
 	// Enqueue the rendering command to set up the rendering canvas.
-	ENQUEUE_UNIQUE_RENDER_COMMAND_TWOPARAMETER
-	(
-		CanvasRenderTargetMakeCurrentCommand,
-		FTextureRenderTarget2DResource*,
-		TextureRenderTarget,
-		(FTextureRenderTarget2DResource*)GameThread_GetRenderTargetResource(),
-		bool,
-		bClearRenderTarget,
-		bShouldClearRenderTargetOnReceiveUpdate,
+	bool bClearRenderTarget = bShouldClearRenderTargetOnReceiveUpdate;
+	ENQUEUE_RENDER_COMMAND(CanvasRenderTargetMakeCurrentCommand)(
+		[&TextureRenderTarget, bClearRenderTarget](FRHICommandListImmediate& RHICmdList)
+	{
+		RHICmdList.TransitionResource(EResourceTransitionAccess::EWritable, TextureRenderTarget->GetRenderTargetTexture());
+
+		if (bClearRenderTarget)
 		{
-			RHICmdList.SetViewport(0, 0, 0.0f, TextureRenderTarget->GetSizeXY().X, TextureRenderTarget->GetSizeXY().Y, 1.0f);
-			if (bClearRenderTarget)
-			{
-				SetRenderTarget(RHICmdList, TextureRenderTarget->GetRenderTargetTexture(), FTextureRHIRef(), ESimpleRenderTargetMode::EClearColorExistingDepth, FExclusiveDepthStencil::DepthWrite_StencilWrite, true);
-			}
-			else
-			{
-				SetRenderTarget(RHICmdList, TextureRenderTarget->GetRenderTargetTexture(), FTexture2DRHIRef(), true);
-			}
+			FRHIRenderPassInfo RPInfo(TextureRenderTarget->GetRenderTargetTexture(), ERenderTargetActions::Clear_Store);
+			RHICmdList.BeginRenderPass(RPInfo, TEXT("ClearUCanvas"));
+			RHICmdList.EndRenderPass();
 		}
-	);
+	});
 
 	if (!IsPendingKill() && OnCanvasRenderTargetUpdate.IsBound())
 	{
diff --git a/Engine/Source/Runtime/Engine/Private/TextureRenderTarget2D.cpp b/Engine/Source/Runtime/Engine/Private/TextureRenderTarget2D.cpp
--- a/Engine/Source/Runtime/Engine/Private/TextureRenderTarget2D.cpp
+++ b/Engine/Source/Runtime/Engine/Private/TextureRenderTarget2D.cpp
@@ -511,9 +511,13 @@ void FTextureRenderTarget2DResource::UpdateDeferredResource( FRHICommandListImme
 	{
 		RHICmdList.SetViewport(0, 0, 0.0f, TargetSizeX, TargetSizeY, 1.0f);
 		ensure(RenderTargetTextureRHI.IsValid() && (RenderTargetTextureRHI->GetClearColor() == ClearColor));
-		SetRenderTarget(RHICmdList, RenderTargetTextureRHI, FTextureRHIRef(), ESimpleRenderTargetMode::EClearColorExistingDepth, FExclusiveDepthStencil::DepthWrite_StencilWrite, true);
+
+		FRHIRenderPassInfo RPInfo(RenderTargetTextureRHI, ERenderTargetActions::Clear_Store);
+		RHICmdList.BeginRenderPass(RPInfo, TEXT("ClearRenderTarget"));
+		RHICmdList.EndRenderPass();
 	}
  
+	// #todo-renderpasses must generate mips outside of a renderpass?
 	if (Owner->bAutoGenerateMips)
 	{
 		RHICmdList.GenerateMips(RenderTargetTextureRHI);
diff --git a/Engine/Source/Runtime/Engine/Private/TextureRenderTargetCube.cpp b/Engine/Source/Runtime/Engine/Private/TextureRenderTargetCube.cpp
--- a/Engine/Source/Runtime/Engine/Private/TextureRenderTargetCube.cpp
+++ b/Engine/Source/Runtime/Engine/Private/TextureRenderTargetCube.cpp
@@ -298,22 +298,18 @@ void FTextureRenderTargetCubeResource::UpdateDeferredResource(FRHICommandListImm
 	const FIntPoint Dims = GetSizeXY();
 	for(int32 FaceIdx = CubeFace_PosX; FaceIdx < CubeFace_MAX; FaceIdx++)
 	{
+		ERenderTargetLoadAction LoadAction = ERenderTargetLoadAction::ELoad;
 		// clear each face of the cube target texture to ClearColor
 		if (bClearRenderTarget)
 		{
-			TransitionSetRenderTargetsHelper(RHICmdList, RenderTargetTextureRHI, FTextureRHIParamRef(), FExclusiveDepthStencil::DepthWrite_StencilWrite);
-
-			FRHIRenderTargetView RtView = FRHIRenderTargetView(RenderTargetTextureRHI, ERenderTargetLoadAction::EClear);
-			FRHISetRenderTargetsInfo Info(1, &RtView, FRHIDepthRenderTargetView());
-			RHICmdList.SetRenderTargetsAndClear(Info);
-			RHICmdList.SetViewport(0, 0, 0.0f, Dims.X, Dims.Y, 1.0f);
-		}
-		else
-		{
-			SetRenderTarget(RHICmdList, RenderTargetTextureRHI, FTextureRHIParamRef(), true);
-			RHICmdList.SetViewport(0, 0, 0.0f, Dims.X, Dims.Y, 1.0f);
+			LoadAction = ERenderTargetLoadAction::EClear;
 		}
-
+		
+		FRHIRenderPassInfo RPInfo(RenderTargetTextureRHI, MakeRenderTargetActions(LoadAction, ERenderTargetStoreAction::EStore));
+		TransitionRenderPassTargets(RHICmdList, RPInfo);
+		RHICmdList.BeginRenderPass(RPInfo, TEXT("UpdateTargetCube"));
+		RHICmdList.SetViewport(0, 0, 0.0f, Dims.X, Dims.Y, 1.0f);
+		RHICmdList.EndRenderPass();
 		// copy surface to the texture for use
 		FResolveParams ResolveParams;
 		ResolveParams.CubeFace = (ECubeFace)FaceIdx;
diff --git a/Engine/Source/Runtime/Engine/Private/UnrealClient.cpp b/Engine/Source/Runtime/Engine/Private/UnrealClient.cpp
--- a/Engine/Source/Runtime/Engine/Private/UnrealClient.cpp
+++ b/Engine/Source/Runtime/Engine/Private/UnrealClient.cpp
@@ -1610,7 +1610,9 @@ const TArray<FColor>& FViewport::GetRawHitProxyData(FIntRect InRect)
 			{
 			// Set the hit proxy map's render target.
 			// Clear the hit proxy map to white, which is overloaded to mean no hit proxy.
-			SetRenderTarget(RHICmdList, Viewport->HitProxyMap.GetRenderTargetTexture(), FTextureRHIRef(), ESimpleRenderTargetMode::EClearColorExistingDepth, FExclusiveDepthStencil::DepthWrite_StencilWrite, true);
+			FRHIRenderPassInfo RPInfo(Viewport->HitProxyMap.GetRenderTargetTexture(), ERenderTargetActions::Clear_Store);
+			RHICmdList.BeginRenderPass(RPInfo, TEXT("ClearHitProxyMap"));
+			RHICmdList.EndRenderPass();
 		});
 
 		// Let the viewport client draw its hit proxies.
@@ -1625,7 +1627,7 @@ const TArray<FColor>& FViewport::GetRawHitProxyData(FIntRect InRect)
 		ENQUEUE_UNIQUE_RENDER_COMMAND_ONEPARAMETER(
 			UpdateHitProxyRTCommand,
 			FHitProxyMap*, HitProxyMap, &HitProxyMap,
-			{
+		{
 			// Copy (resolve) the rendered thumbnail from the render target to its texture
 			RHICmdList.CopyToResolveTarget(HitProxyMap->GetRenderTargetTexture(), HitProxyMap->GetHitProxyTexture(), FResolveParams());
 			RHICmdList.CopyToResolveTarget(HitProxyMap->GetRenderTargetTexture(), HitProxyMap->GetHitProxyCPUTexture(), FResolveParams());
diff --git a/Engine/Source/Runtime/Engine/Private/UnrealEngine.cpp b/Engine/Source/Runtime/Engine/Private/UnrealEngine.cpp
--- a/Engine/Source/Runtime/Engine/Private/UnrealEngine.cpp
+++ b/Engine/Source/Runtime/Engine/Private/UnrealEngine.cpp
@@ -2834,9 +2834,10 @@ public:
 	{
 		check(IsInRenderingThread());
 
-		FRHIRenderTargetView BackBufferView = FRHIRenderTargetView(BackBuffer, ERenderTargetLoadAction::EClear);
-		FRHISetRenderTargetsInfo Info(1, &BackBufferView, FRHIDepthRenderTargetView());
-		RHICmdList.SetRenderTargetsAndClear(Info);
+		FRHIRenderPassInfo RPInfo(BackBuffer, ERenderTargetActions::Clear_Store);
+		RHICmdList.BeginRenderPass(RPInfo, TEXT("RenderTexture_RenderThread"));
+		RHICmdList.EndRenderPass();
+
 		const uint32 ViewportWidth = BackBuffer->GetSizeX();
 		const uint32 ViewportHeight = BackBuffer->GetSizeY();
 		RHICmdList.SetViewport( 0,0,0,ViewportWidth, ViewportHeight, 1.0f );
diff --git a/Engine/Source/Runtime/Engine/Private/UserInterface/Canvas.cpp b/Engine/Source/Runtime/Engine/Private/UserInterface/Canvas.cpp
--- a/Engine/Source/Runtime/Engine/Private/UserInterface/Canvas.cpp
+++ b/Engine/Source/Runtime/Engine/Private/UserInterface/Canvas.cpp
@@ -645,6 +645,9 @@ FCanvas::~FCanvas()
 
 void FCanvas::Flush_RenderThread(FRHICommandListImmediate& RHICmdList, bool bForce)
 {
+	// This renderpass is self contained.
+	check(RHICmdList.IsOutsideRenderPass());
+
 	SCOPE_CYCLE_COUNTER(STAT_Canvas_FlushTime);
 
 	if (!(AllowedModes&Allow_Flush) && !bForce)
@@ -681,64 +684,66 @@ void FCanvas::Flush_RenderThread(FRHICommandListImmediate& RHICmdList, bool bFor
 	const FTexture2DRHIRef& RenderTargetTexture = RenderTarget->GetRenderTargetTexture();
 
 	check(IsValidRef(RenderTargetTexture));
+
+	FRHIRenderPassInfo RPInfo(RenderTargetTexture, ERenderTargetActions::Load_Store);
 	
 	// Set the RHI render target.
 	if (IsUsingInternalTexture())
 	{
-		::SetRenderTarget(RHICmdList, RenderTargetTexture, FTexture2DRHIRef(), ESimpleRenderTargetMode::EClearColorAndDepth);
+		RPInfo.ColorRenderTargets[0].Action = ERenderTargetActions::Clear_Store;
 	}
-	else
+	
+	RHICmdList.BeginRenderPass(RPInfo, TEXT("CanvasRenderThread"));
 	{
-		::SetRenderTarget(RHICmdList, RenderTargetTexture, FTexture2DRHIRef());
-	}
+		FDrawingPolicyRenderState DrawRenderState;
+		// disable depth test & writes
+		DrawRenderState.SetDepthStencilState(TStaticDepthStencilState<false, CF_Always>::GetRHI());
 
-	FDrawingPolicyRenderState DrawRenderState;
-	// disable depth test & writes
-	DrawRenderState.SetDepthStencilState(TStaticDepthStencilState<false, CF_Always>::GetRHI());
-
-	if (ViewRect.Area() <= 0)
-	{
-		ViewRect = FIntRect(FIntPoint::ZeroValue, RenderTarget->GetSizeXY());
-	}
+		if (ViewRect.Area() <= 0)
+		{
+			ViewRect = FIntRect(FIntPoint::ZeroValue, RenderTarget->GetSizeXY());
+		}
 
-	// set viewport to RT size
-	RHICmdList.SetViewport(ViewRect.Min.X, ViewRect.Min.Y, 0.0f, ViewRect.Max.X, ViewRect.Max.Y, 1.0f);
+		// set viewport to RT size
+		RHICmdList.SetViewport(ViewRect.Min.X, ViewRect.Min.Y, 0.0f, ViewRect.Max.X, ViewRect.Max.Y, 1.0f);
 
-	// Set scissor rect if valid applied
-	if (ScissorRect.Area() > 0)
-	{
-		RHICmdList.SetScissorRect(true, ScissorRect.Min.X, ScissorRect.Min.Y, ScissorRect.Max.X, ScissorRect.Max.Y);
-	}
+		// Set scissor rect if valid applied
+		if (ScissorRect.Area() > 0)
+		{
+			RHICmdList.SetScissorRect(true, ScissorRect.Min.X, ScissorRect.Min.Y, ScissorRect.Max.X, ScissorRect.Max.Y);
+		}
 
-	// iterate over the FCanvasSortElements in sorted order and render all the batched items for each entry
-	for (int32 Idx = 0; Idx < SortedElements.Num(); Idx++)
-	{
-		FCanvasSortElement& SortElement = SortedElements[Idx];
-		for (int32 BatchIdx = 0; BatchIdx < SortElement.RenderBatchArray.Num(); BatchIdx++)
+		// iterate over the FCanvasSortElements in sorted order and render all the batched items for each entry
+		for (int32 Idx = 0; Idx < SortedElements.Num(); Idx++)
 		{
-			FCanvasBaseRenderItem* RenderItem = SortElement.RenderBatchArray[BatchIdx];
-			if (RenderItem)
+			FCanvasSortElement& SortElement = SortedElements[Idx];
+			for (int32 BatchIdx = 0; BatchIdx < SortElement.RenderBatchArray.Num(); BatchIdx++)
 			{
-				// mark current render target as dirty since we are drawing to it
-				bRenderTargetDirty |= RenderItem->Render_RenderThread(RHICmdList, DrawRenderState, this);
-				if (AllowedModes & Allow_DeleteOnRender)
+				FCanvasBaseRenderItem* RenderItem = SortElement.RenderBatchArray[BatchIdx];
+				if (RenderItem)
 				{
-					delete RenderItem;
+					// mark current render target as dirty since we are drawing to it
+					bRenderTargetDirty |= RenderItem->Render_RenderThread(RHICmdList, DrawRenderState, this);
+					if (AllowedModes & Allow_DeleteOnRender)
+					{
+						delete RenderItem;
+					}
 				}
 			}
+			if (AllowedModes & Allow_DeleteOnRender)
+			{
+				SortElement.RenderBatchArray.Empty();
+			}
 		}
 		if (AllowedModes & Allow_DeleteOnRender)
 		{
-			SortElement.RenderBatchArray.Empty();
+			// empty the array of FCanvasSortElement entries after finished with rendering	
+			SortedElements.Empty();
+			SortedElementLookupMap.Empty();
+			LastElementIndex = INDEX_NONE;
 		}
 	}
-	if (AllowedModes & Allow_DeleteOnRender)
-	{
-		// empty the array of FCanvasSortElement entries after finished with rendering	
-		SortedElements.Empty();
-		SortedElementLookupMap.Empty();
-		LastElementIndex = INDEX_NONE;
-	}
+	RHICmdList.EndRenderPass();	
 }
 
 void FCanvas::Flush_GameThread(bool bForce)
@@ -801,8 +806,12 @@ void FCanvas::Flush_GameThread(bool bForce)
 	ENQUEUE_RENDER_COMMAND(CanvasFlushSetupCommand)(
 		[FlushParameters](FRHICommandList& RHICmdList)
 		{
+			check(RHICmdList.IsOutsideRenderPass());
+
 			// Set the RHI render target.
-			::SetRenderTarget(RHICmdList, FlushParameters.CanvasRenderTarget->GetRenderTargetTexture(), FTextureRHIRef(), true);
+			FRHIRenderPassInfo RPInfo(FlushParameters.CanvasRenderTarget->GetRenderTargetTexture(), ERenderTargetActions::Load_Store);
+			TransitionRenderPassTargets(RHICmdList, RPInfo);
+			RHICmdList.BeginRenderPass(RPInfo, TEXT("Canvas_GameThreadFlush"));
 
 			FIntRect ViewportRect = FlushParameters.ViewRect;
 			FIntRect ScissorRectParam = FlushParameters.ScissorRect;
@@ -849,6 +858,12 @@ void FCanvas::Flush_GameThread(bool bForce)
 		SortedElementLookupMap.Empty();
 		LastElementIndex = INDEX_NONE;
 	}
+
+	ENQUEUE_RENDER_COMMAND(CanvasFlushEndCommand)(
+		[](FRHICommandList& RHICmdList)
+	{
+		RHICmdList.EndRenderPass();
+	});
 }
 
 void FCanvas::PushRelativeTransform(const FMatrix& Transform)
@@ -957,13 +972,18 @@ void FCanvas::Clear(const FLinearColor& ClearColor)
 				if (CanvasRenderTarget->GetRenderTargetTexture() && CanvasRenderTarget->GetRenderTargetTexture()->GetClearBinding() == FClearValueBinding(ClearColor))
 				{
 					// do fast clear
-					SetRenderTarget(RHICmdList, CanvasRenderTarget->GetRenderTargetTexture(), FTextureRHIRef(), ESimpleRenderTargetMode::EClearColorAndDepth);
+					FRHIRenderPassInfo RPInfo(CanvasRenderTarget->GetRenderTargetTexture(), ERenderTargetActions::Clear_Store);
+					RHICmdList.BeginRenderPass(RPInfo, TEXT("ClearCanvas"));
+					RHICmdList.EndRenderPass();					
 				}
 				else
 				{
-					::SetRenderTarget(RHICmdList, CanvasRenderTarget->GetRenderTargetTexture(), FTextureRHIRef(), true);
+					FRHIRenderPassInfo RPInfo(CanvasRenderTarget->GetRenderTargetTexture(), ERenderTargetActions::Load_Store);
+					TransitionRenderPassTargets(RHICmdList, RPInfo);
+					RHICmdList.BeginRenderPass(RPInfo, TEXT("ClearCanvas"));
 					RHICmdList.SetViewport(0, 0, 0.0f, CanvasRenderTarget->GetSizeXY().X, CanvasRenderTarget->GetSizeXY().Y, 1.0f);
 					DrawClearQuad(RHICmdList, ClearColor);
+					RHICmdList.EndRenderPass();
 				}
 			}
 			else
diff --git a/Engine/Source/Runtime/Engine/Public/UnrealClient.h b/Engine/Source/Runtime/Engine/Public/UnrealClient.h
--- a/Engine/Source/Runtime/Engine/Public/UnrealClient.h
+++ b/Engine/Source/Runtime/Engine/Public/UnrealClient.h
@@ -1087,7 +1087,7 @@ public:
 	virtual void BeginRenderFrame(FRHICommandListImmediate& RHICmdList) override
 	{
 		check( IsInRenderingThread() );
-		SetRenderTarget(RHICmdList,  RenderTargetTextureRHI,  FTexture2DRHIRef() );
+		//SetRenderTarget(RHICmdList,  RenderTargetTextureRHI,  FTexture2DRHIRef() );
 	};
 
 	virtual void EndRenderFrame(FRHICommandListImmediate& RHICmdList, bool bPresent, bool bLockToVsync) override
