Make sure all per-instance state is properly created for (H)ISM on newly created copies

diff --git a/Engine/Source/Runtime/Engine/Classes/Components/HierarchicalInstancedStaticMeshComponent.h b/Engine/Source/Runtime/Engine/Classes/Components/HierarchicalInstancedStaticMeshComponent.h
--- a/Engine/Source/Runtime/Engine/Classes/Components/HierarchicalInstancedStaticMeshComponent.h
+++ b/Engine/Source/Runtime/Engine/Classes/Components/HierarchicalInstancedStaticMeshComponent.h
@@ -191,6 +191,8 @@ class ENGINE_API UHierarchicalInstancedStaticMeshComponent : public UInstancedSt
 	// Apply the results of the async build
 	void ApplyBuildTreeAsync(ENamedThreads::Type CurrentThread, const FGraphEventRef& MyCompletionGraphEvent, TSharedRef<FClusterBuilder, ESPMode::ThreadSafe> Builder, double StartTime);
 
+	virtual void OnComponentCreated() override;
+
 public:
 
 	//Begin UObject Interface
diff --git a/Engine/Source/Runtime/Engine/Private/HierarchicalInstancedStaticMesh.cpp b/Engine/Source/Runtime/Engine/Private/HierarchicalInstancedStaticMesh.cpp
--- a/Engine/Source/Runtime/Engine/Private/HierarchicalInstancedStaticMesh.cpp
+++ b/Engine/Source/Runtime/Engine/Private/HierarchicalInstancedStaticMesh.cpp
@@ -2595,6 +2595,18 @@ void UHierarchicalInstancedStaticMeshComponent::ApplyBuildTreeAsync(ENamedThread
 	PostBuildStats();
 }
 
+void UHierarchicalInstancedStaticMeshComponent::OnComponentCreated()
+{
+	Super::OnComponentCreated();
+
+	if (FApp::CanEverRender() && !HasAnyFlags(RF_ClassDefaultObject | RF_ArchetypeObject))
+	{
+		// if we are pasting/duplicating this component, it may be created with some instances already in place
+		// in this case, need to ensure that the tree is properly created
+		BuildTreeIfOutdated(false, false);
+	}
+}
+
 bool UHierarchicalInstancedStaticMeshComponent::BuildTreeIfOutdated(bool Async, bool ForceUpdate)
 {
 	if (ForceUpdate 
diff --git a/Engine/Source/Runtime/Engine/Private/InstancedStaticMesh.cpp b/Engine/Source/Runtime/Engine/Private/InstancedStaticMesh.cpp
--- a/Engine/Source/Runtime/Engine/Private/InstancedStaticMesh.cpp
+++ b/Engine/Source/Runtime/Engine/Private/InstancedStaticMesh.cpp
@@ -2186,6 +2186,13 @@ void UInstancedStaticMeshComponent::OnComponentCreated()
 
 	if (FApp::CanEverRender() && !HasAnyFlags(RF_ClassDefaultObject | RF_ArchetypeObject))
 	{
+		// if we are pasting/duplicating this component, it may be created with some instances already in place
+		// in this case, need to ensure that the instance render data is properly created
+		if (PerInstanceSMData.Num() > 0)
+		{
+			InstanceUpdateCmdBuffer.Edit();
+		}
+
 		InitPerInstanceRenderData(false);
 	}
 }
