Minor revision to CL 4548117
- cleaner method of updating render thread state in ISM
- guard against rebuilding tree when running cooked data in HISM

diff --git a/Engine/Source/Runtime/Engine/Private/HierarchicalInstancedStaticMesh.cpp b/Engine/Source/Runtime/Engine/Private/HierarchicalInstancedStaticMesh.cpp
--- a/Engine/Source/Runtime/Engine/Private/HierarchicalInstancedStaticMesh.cpp
+++ b/Engine/Source/Runtime/Engine/Private/HierarchicalInstancedStaticMesh.cpp
@@ -2603,7 +2603,11 @@ void UHierarchicalInstancedStaticMeshComponent::OnComponentCreated()
 	{
 		// if we are pasting/duplicating this component, it may be created with some instances already in place
 		// in this case, need to ensure that the tree is properly created
-		BuildTreeIfOutdated(false, false);
+		if (PerInstanceSMData.Num() > 0 && ClusterTreePtr->Num() == 0)
+		{
+			const bool bForceUpdate = true;
+			BuildTreeIfOutdated(false, bForceUpdate);
+		}
 	}
 }
 
diff --git a/Engine/Source/Runtime/Engine/Private/InstancedStaticMesh.cpp b/Engine/Source/Runtime/Engine/Private/InstancedStaticMesh.cpp
--- a/Engine/Source/Runtime/Engine/Private/InstancedStaticMesh.cpp
+++ b/Engine/Source/Runtime/Engine/Private/InstancedStaticMesh.cpp
@@ -2188,12 +2188,8 @@ void UInstancedStaticMeshComponent::OnComponentCreated()
 	{
 		// if we are pasting/duplicating this component, it may be created with some instances already in place
 		// in this case, need to ensure that the instance render data is properly created
-		if (PerInstanceSMData.Num() > 0)
-		{
-			InstanceUpdateCmdBuffer.Edit();
-		}
-
-		InitPerInstanceRenderData(false);
+		const bool InitializeFromCurrentData = PerInstanceSMData.Num() > 0;
+		InitPerInstanceRenderData(InitializeFromCurrentData);
 	}
 }
 
