The CanvasRenderTarget may be null if running under the NullRHI.

diff --git a/Engine/Source/Runtime/Engine/Private/UserInterface/Canvas.cpp b/Engine/Source/Runtime/Engine/Private/UserInterface/Canvas.cpp
--- a/Engine/Source/Runtime/Engine/Private/UserInterface/Canvas.cpp
+++ b/Engine/Source/Runtime/Engine/Private/UserInterface/Canvas.cpp
@@ -968,7 +968,7 @@ void FCanvas::Clear(const FLinearColor& ClearColor)
 			SCOPED_DRAW_EVENT(RHICmdList, CanvasClear);
 			if (CanvasRenderTarget)
 			{
-				// possibility for the RTT to be null for nullrhi
+				// possibility for the RTT to be null for nullrhi. In that case we do nothing.
 				if (CanvasRenderTarget->GetRenderTargetTexture() && CanvasRenderTarget->GetRenderTargetTexture()->GetClearBinding() == FClearValueBinding(ClearColor))
 				{
 					// do fast clear
@@ -976,9 +976,9 @@ void FCanvas::Clear(const FLinearColor& ClearColor)
 					RHICmdList.BeginRenderPass(RPInfo, TEXT("ClearCanvas"));
 					RHICmdList.EndRenderPass();					
 				}
-				else
+				else if(CanvasRenderTarget->GetRenderTargetTexture())
 				{
-					FRHIRenderPassInfo RPInfo(CanvasRenderTarget->GetRenderTargetTexture(), ERenderTargetActions::Load_Store);
+					FRHIRenderPassInfo RPInfo(CanvasRenderTarget->GetRenderTargetTexture(), ERenderTargetActions::DontLoad_Store);
 					TransitionRenderPassTargets(RHICmdList, RPInfo);
 					RHICmdList.BeginRenderPass(RPInfo, TEXT("ClearCanvas"));
 					RHICmdList.SetViewport(0, 0, 0.0f, CanvasRenderTarget->GetSizeXY().X, CanvasRenderTarget->GetSizeXY().Y, 1.0f);
