Fix logic for batched sprites.  Previous code would submit extra draw call even none of the sprites in the batch matched the current filter.  This would cause masked sprites to sometimes redraw during the translucent pass.  In many cases, the redrawn sprites would match the existing sprites, so there would be overdraw, but no visible problem.  When calling SetSprite on billboard every frame however, billboard's proxy will also regenerate every frame, which causes memory to move around such that the invalid translucent sprites don't line up with the masked sprites.  This is what caused the flickering mentioned in Jira.

diff --git a/Engine/Source/Runtime/Engine/Private/BatchedElements.cpp b/Engine/Source/Runtime/Engine/Private/BatchedElements.cpp
--- a/Engine/Source/Runtime/Engine/Private/BatchedElements.cpp
+++ b/Engine/Source/Runtime/Engine/Private/BatchedElements.cpp
@@ -1096,98 +1096,100 @@ bool FBatchedElements::Draw(FRHICommandList& RHICmdList, const FDrawingPolicyRen
 			// Sprites don't support batched element parameters (yet!)
 			FBatchedElementParameters* BatchedElementParameters = NULL;
 
-			int32 SpriteCount = Sprites.Num();
-			//Sort sprites by texture
-			/** Compare two texture pointers, return false if equal */
-			struct FCompareTexture
+			//Sort sprites by blend mode and texture
+			struct FCompareSprite
 			{
+				explicit FCompareSprite(EBlendModeFilter::Type Filter) : ValidBlendFilter(Filter) {}
+
 				FORCEINLINE bool operator()( const FBatchedElements::FBatchedSprite& SpriteA, const FBatchedElements::FBatchedSprite& SpriteB ) const
 				{
-					return (SpriteA.Texture == SpriteB.Texture && SpriteA.BlendMode == SpriteB.BlendMode) ? false : true; 
+					if (SpriteA.BlendMode != SpriteB.BlendMode)
+					{
+						const bool bBlendAValid = ValidBlendFilter & SpriteA.BlendMode;
+						const bool bBlendBValid = ValidBlendFilter & SpriteB.BlendMode;
+						if (bBlendAValid != bBlendBValid) return bBlendAValid; // we want valid blend modes to sort in front of invalid blend modes
+						return SpriteA.BlendMode < SpriteB.BlendMode;
+					}
+					return SpriteA.Texture < SpriteB.Texture;
 				}
+
+				EBlendModeFilter::Type ValidBlendFilter;
 			};
 
-			Sprites.Sort(FCompareTexture());
-			
-			//First time init
-			const FTexture* CurrentTexture = Sprites[0].Texture;
-			ESimpleElementBlendMode CurrentBlendMode = (ESimpleElementBlendMode)Sprites[0].BlendMode;
+			Sprites.Sort(FCompareSprite(Filter));
 
-			FRHIResourceCreateInfo CreateInfo;
-			FVertexBufferRHIRef VertexBufferRHI = RHICreateVertexBuffer(sizeof(FSimpleElementVertex) * SpriteCount * 6, BUF_Volatile, CreateInfo);
-			void* VoidPtr = RHILockVertexBuffer(VertexBufferRHI, 0, sizeof(FSimpleElementVertex) * SpriteCount * 6, RLM_WriteOnly);
-			FSimpleElementVertex* SpriteList = reinterpret_cast<FSimpleElementVertex*>(VoidPtr);
+			// count the number of sprites that have valid blend modes
+			// (they have been sorted to the front of the list)
+			int32 ValidSpriteCount = 0;
+			while (ValidSpriteCount < Sprites.Num() &&
+				(Filter & GetBlendModeFilter((ESimpleElementBlendMode)Sprites[ValidSpriteCount].BlendMode)))
+			{
+				++ValidSpriteCount;
+			}
 
-			for (int32 SpriteIndex = 0;SpriteIndex < SpriteCount;SpriteIndex++)
+			if (ValidSpriteCount > 0)
 			{
-				const FBatchedSprite& Sprite = Sprites[SpriteIndex];
-				FSimpleElementVertex* Vertex = &SpriteList[SpriteIndex * 6];
+				FRHIResourceCreateInfo CreateInfo;
+				FVertexBufferRHIRef VertexBufferRHI = RHICreateVertexBuffer(sizeof(FSimpleElementVertex) * ValidSpriteCount * 6, BUF_Volatile, CreateInfo);
+				void* VoidPtr = RHILockVertexBuffer(VertexBufferRHI, 0, sizeof(FSimpleElementVertex) * ValidSpriteCount * 6, RLM_WriteOnly);
+				FSimpleElementVertex* SpriteList = reinterpret_cast<FSimpleElementVertex*>(VoidPtr);
 
-				// Compute the sprite vertices.
-				const FVector WorldSpriteX = CameraX * Sprite.SizeX;
-				const FVector WorldSpriteY = CameraY * -Sprite.SizeY * GProjectionSignY;
+				for (int32 SpriteIndex = 0; SpriteIndex < ValidSpriteCount; SpriteIndex++)
+				{
+					const FBatchedSprite& Sprite = Sprites[SpriteIndex];
+					FSimpleElementVertex* Vertex = &SpriteList[SpriteIndex * 6];
 
-				const float UStart = Sprite.U / Sprite.Texture->GetSizeX();
-				const float UEnd = (Sprite.U + Sprite.UL) / Sprite.Texture->GetSizeX();
-				const float VStart = Sprite.V / Sprite.Texture->GetSizeY();
-				const float VEnd = (Sprite.V + Sprite.VL) / Sprite.Texture->GetSizeY();
+					// Compute the sprite vertices.
+					const FVector WorldSpriteX = CameraX * Sprite.SizeX;
+					const FVector WorldSpriteY = CameraY * -Sprite.SizeY * GProjectionSignY;
 
-				Vertex[0] = FSimpleElementVertex(FVector4(Sprite.Position + WorldSpriteX - WorldSpriteY, 1), FVector2D(UEnd, VStart), Sprite.Color, Sprite.HitProxyId);
-				Vertex[1] = FSimpleElementVertex(FVector4(Sprite.Position + WorldSpriteX + WorldSpriteY, 1), FVector2D(UEnd, VEnd), Sprite.Color, Sprite.HitProxyId);
-				Vertex[2] = FSimpleElementVertex(FVector4(Sprite.Position - WorldSpriteX - WorldSpriteY, 1), FVector2D(UStart, VStart), Sprite.Color, Sprite.HitProxyId);
+					const float UStart = Sprite.U / Sprite.Texture->GetSizeX();
+					const float UEnd = (Sprite.U + Sprite.UL) / Sprite.Texture->GetSizeX();
+					const float VStart = Sprite.V / Sprite.Texture->GetSizeY();
+					const float VEnd = (Sprite.V + Sprite.VL) / Sprite.Texture->GetSizeY();
 
-				Vertex[3] = FSimpleElementVertex(FVector4(Sprite.Position + WorldSpriteX + WorldSpriteY, 1), FVector2D(UEnd, VEnd), Sprite.Color, Sprite.HitProxyId);
-				Vertex[4] = FSimpleElementVertex(FVector4(Sprite.Position - WorldSpriteX - WorldSpriteY, 1), FVector2D(UStart, VStart), Sprite.Color, Sprite.HitProxyId);
-				Vertex[5] = FSimpleElementVertex(FVector4(Sprite.Position - WorldSpriteX + WorldSpriteY, 1), FVector2D(UStart, VEnd), Sprite.Color, Sprite.HitProxyId);
-			}
-			RHIUnlockVertexBuffer(VertexBufferRHI);
+					Vertex[0] = FSimpleElementVertex(FVector4(Sprite.Position + WorldSpriteX - WorldSpriteY, 1), FVector2D(UEnd, VStart), Sprite.Color, Sprite.HitProxyId);
+					Vertex[1] = FSimpleElementVertex(FVector4(Sprite.Position + WorldSpriteX + WorldSpriteY, 1), FVector2D(UEnd, VEnd), Sprite.Color, Sprite.HitProxyId);
+					Vertex[2] = FSimpleElementVertex(FVector4(Sprite.Position - WorldSpriteX - WorldSpriteY, 1), FVector2D(UStart, VStart), Sprite.Color, Sprite.HitProxyId);
 
-			int32 SpriteStartIndex = 0;
-			for(int32 SpriteIndex = 0;SpriteIndex < (SpriteCount + 1);SpriteIndex++)
-			{
-				int32 EndIndex = SpriteIndex;
-				bool isVeryLastDummyIndex = (SpriteIndex == SpriteCount);
-				if (!isVeryLastDummyIndex)
+					Vertex[3] = FSimpleElementVertex(FVector4(Sprite.Position + WorldSpriteX + WorldSpriteY, 1), FVector2D(UEnd, VEnd), Sprite.Color, Sprite.HitProxyId);
+					Vertex[4] = FSimpleElementVertex(FVector4(Sprite.Position - WorldSpriteX - WorldSpriteY, 1), FVector2D(UStart, VStart), Sprite.Color, Sprite.HitProxyId);
+					Vertex[5] = FSimpleElementVertex(FVector4(Sprite.Position - WorldSpriteX + WorldSpriteY, 1), FVector2D(UStart, VEnd), Sprite.Color, Sprite.HitProxyId);
+				}
+				RHIUnlockVertexBuffer(VertexBufferRHI);
+
+				//First time init
+				const FTexture* CurrentTexture = Sprites[0].Texture;
+				ESimpleElementBlendMode CurrentBlendMode = (ESimpleElementBlendMode)Sprites[0].BlendMode;
+				int32 BatchStartIndex = 0;
+
+				// Start loop at 1, since we've already started the first batch with the first sprite in the list
+				for (int32 SpriteIndex = 1; SpriteIndex < ValidSpriteCount + 1; SpriteIndex++)
 				{
-					if (!(Filter & GetBlendModeFilter((ESimpleElementBlendMode)Sprites[SpriteIndex].BlendMode)))
+					// Need to flush the current batch once we hit the end of the list, or if state of this sprite doesn't match current batch
+					if (SpriteIndex == ValidSpriteCount ||
+						CurrentTexture != Sprites[SpriteIndex].Texture ||
+						CurrentBlendMode != Sprites[SpriteIndex].BlendMode)
 					{
-						if ((SpriteStartIndex + 1) == SpriteIndex)
+						const int32 SpriteNum = SpriteIndex - BatchStartIndex;
+						const int32 BaseVertex = BatchStartIndex * 6;
+						const int32 PrimCount = SpriteNum * 2;
+						PrepareShaders(RHICmdList, GraphicsPSOInit, FeatureLevel, CurrentBlendMode, Transform, bNeedToSwitchVerticalAxis, BatchedElementParameters, CurrentTexture, bHitTesting, Gamma, NULL, &View);
+						RHICmdList.SetStencilRef(StencilRef);
+						RHICmdList.SetStreamSource(0, VertexBufferRHI, 0);
+						RHICmdList.DrawPrimitive(BaseVertex, PrimCount, 1);
+
+						// begin the next batch
+						if (SpriteIndex < ValidSpriteCount)
 						{
-							// just skip as there is nothing to draw
-							SpriteStartIndex = SpriteIndex;
-							continue;
+							BatchStartIndex = SpriteIndex;
+							CurrentTexture = Sprites[SpriteIndex].Texture;
+							CurrentBlendMode = (ESimpleElementBlendMode)Sprites[SpriteIndex].BlendMode;
 						}
-						//draw but without this filtered element
-						EndIndex = SpriteIndex - 1;
-					}
-					else if (CurrentTexture == Sprites[SpriteIndex].Texture && CurrentBlendMode == Sprites[SpriteIndex].BlendMode)
-					{
-						//batch the draw
-						continue;
 					}
 				}
-
-				const int32 SpriteNum = EndIndex - SpriteStartIndex;
-				if (SpriteNum > 0)
-				{
-					//New batch, draw previous and clear
-					const int32 BaseVertex = SpriteStartIndex * 6;
-					const int32 PrimCount = SpriteNum * 2;
-					PrepareShaders(RHICmdList, GraphicsPSOInit, FeatureLevel, CurrentBlendMode, Transform, bNeedToSwitchVerticalAxis, BatchedElementParameters, CurrentTexture, bHitTesting, Gamma, NULL, &View);
-					RHICmdList.SetStencilRef(StencilRef);
-					RHICmdList.SetStreamSource(0, VertexBufferRHI, 0);
-
-					RHICmdList.DrawPrimitive(BaseVertex, PrimCount, 1);
-				}
-
-				if (!isVeryLastDummyIndex)
-				{
-					SpriteStartIndex = SpriteIndex;
-					CurrentTexture = Sprites[SpriteIndex].Texture;
-					CurrentBlendMode = (ESimpleElementBlendMode)Sprites[SpriteIndex].BlendMode;
-				}
+				VertexBufferRHI.SafeRelease();
 			}
-			VertexBufferRHI.SafeRelease();
 		}
 
 		if( MeshElements.Num() > 0)
