Fix int->bool error with certain compilers

diff --git a/Engine/Source/Runtime/Engine/Private/BatchedElements.cpp b/Engine/Source/Runtime/Engine/Private/BatchedElements.cpp
--- a/Engine/Source/Runtime/Engine/Private/BatchedElements.cpp
+++ b/Engine/Source/Runtime/Engine/Private/BatchedElements.cpp
@@ -1105,8 +1105,8 @@ bool FBatchedElements::Draw(FRHICommandList& RHICmdList, const FDrawingPolicyRen
 				{
 					if (SpriteA.BlendMode != SpriteB.BlendMode)
 					{
-						const bool bBlendAValid = ValidBlendFilter & SpriteA.BlendMode;
-						const bool bBlendBValid = ValidBlendFilter & SpriteB.BlendMode;
+						const bool bBlendAValid = (ValidBlendFilter & SpriteA.BlendMode) != 0;
+						const bool bBlendBValid = (ValidBlendFilter & SpriteB.BlendMode) != 0;
 						if (bBlendAValid != bBlendBValid) return bBlendAValid; // we want valid blend modes to sort in front of invalid blend modes
 						return SpriteA.BlendMode < SpriteB.BlendMode;
 					}
