Update various AllocCommand to safer version

diff --git a/Engine/Plugins/Experimental/PixelStreaming/Source/PixelStreaming/Private/NvVideoEncoder.cpp b/Engine/Plugins/Experimental/PixelStreaming/Source/PixelStreaming/Private/NvVideoEncoder.cpp
--- a/Engine/Plugins/Experimental/PixelStreaming/Source/PixelStreaming/Private/NvVideoEncoder.cpp
+++ b/Engine/Plugins/Experimental/PixelStreaming/Source/PixelStreaming/Private/NvVideoEncoder.cpp
@@ -556,7 +556,7 @@ void FNvVideoEncoder::FNvVideoEncoderImpl::EncodeFrame(const FVideoEncoderSettin
 		}
 		else
 		{
-			new (RHICmdList.AllocCommand<FRHITransferRenderTargetToNvEnc>()) FRHITransferRenderTargetToNvEnc(this, &Frame);
+			ALLOC_COMMAND_CL(RHICmdList, FRHITransferRenderTargetToNvEnc)(this, &Frame);
 		}
 	}
 
diff --git a/Engine/Plugins/Runtime/Steam/SteamVR/Source/SteamVR/Private/SteamVRRender.cpp b/Engine/Plugins/Runtime/Steam/SteamVR/Source/SteamVR/Private/SteamVRRender.cpp
--- a/Engine/Plugins/Runtime/Steam/SteamVR/Source/SteamVR/Private/SteamVRRender.cpp
+++ b/Engine/Plugins/Runtime/Steam/SteamVR/Source/SteamVR/Private/SteamVRRender.cpp
@@ -102,7 +102,7 @@ void FSteamVRHMD::BridgeBaseImpl::BeginRendering_RenderThread(FRHICommandListImm
 {
 	if (IsUsingExplicitTimingMode())
 	{
-		new (RHICmdList.AllocCommand<FRHICommandExecute_BeginRendering>()) FRHICommandExecute_BeginRendering(this);
+		ALLOC_COMMAND_CL(RHICmdList, FRHICommandExecute_BeginRendering)(this);
 	}
 }
 
diff --git a/Engine/Source/Runtime/Engine/Private/GPUSkinVertexFactory.cpp b/Engine/Source/Runtime/Engine/Private/GPUSkinVertexFactory.cpp
--- a/Engine/Source/Runtime/Engine/Private/GPUSkinVertexFactory.cpp
+++ b/Engine/Source/Runtime/Engine/Private/GPUSkinVertexFactory.cpp
@@ -273,7 +273,7 @@ bool FGPUBaseSkinVertexFactory::FShaderDataType::UpdateBoneData(FRHICommandListI
 		{
 			if (!bUseSkinCache && DeferSkeletalLockAndFillToRHIThread())
 			{
-				new (RHICmdList.AllocCommand<FRHICommandUpdateBoneBuffer>()) FRHICommandUpdateBoneBuffer(CurrentBoneBuffer->VertexBufferRHI, VectorArraySize, ReferenceToLocalMatrices, BoneMap);
+				ALLOC_COMMAND_CL(RHICmdList, FRHICommandUpdateBoneBuffer)(CurrentBoneBuffer->VertexBufferRHI, VectorArraySize, ReferenceToLocalMatrices, BoneMap);
 				return true;
 			}
 			ChunkMatrices = (FSkinMatrix3x4*)RHILockVertexBuffer(CurrentBoneBuffer->VertexBufferRHI, 0, VectorArraySize, RLM_WriteOnly);
@@ -966,7 +966,7 @@ bool FGPUBaseSkinAPEXClothVertexFactory::ClothShaderType::UpdateClothSimulData(F
 		{
 			if (DeferSkeletalLockAndFillToRHIThread())
 			{
-				new (RHICmdList.AllocCommand<FRHICommandUpdateClothBuffer>()) FRHICommandUpdateClothBuffer(CurrentClothBuffer->VertexBufferRHI, VectorArraySize, InSimulPositions, InSimulNormals);
+				ALLOC_COMMAND_CL(RHICmdList, FRHICommandUpdateClothBuffer)(CurrentClothBuffer->VertexBufferRHI, VectorArraySize, InSimulPositions, InSimulNormals);
 				return true;
 			}
 			float* RESTRICT Data = (float* RESTRICT)RHILockVertexBuffer(CurrentClothBuffer->VertexBufferRHI, 0, VectorArraySize, RLM_WriteOnly);
diff --git a/Engine/Source/Runtime/HeadMountedDisplay/Private/DefaultSpectatorScreenController.cpp b/Engine/Source/Runtime/HeadMountedDisplay/Private/DefaultSpectatorScreenController.cpp
--- a/Engine/Source/Runtime/HeadMountedDisplay/Private/DefaultSpectatorScreenController.cpp
+++ b/Engine/Source/Runtime/HeadMountedDisplay/Private/DefaultSpectatorScreenController.cpp
@@ -121,7 +121,7 @@ void FDefaultSpectatorScreenController::SetSpectatorScreenTextureRenderCommand(U
 				Command.Execute(RHICmdList);
 				return;
 			}
-			new (RHICmdList.AllocCommand<FRHISetSpectatorScreenTexture>()) FRHISetSpectatorScreenTexture(SpectatorScreenController, Texture);
+			ALLOC_COMMAND_CL(RHICmdList, FRHISetSpectatorScreenTexture)(SpectatorScreenController, Texture);
 		}
 	);
 }
@@ -165,7 +165,7 @@ void FDefaultSpectatorScreenController::SetSpectatorScreenModeTexturePlusEyeLayo
 				Command.Execute(RHICmdList);
 				return;
 			}
-			new (RHICmdList.AllocCommand<FRHISetSpectatorScreenModeTexturePlusEyeLayout>()) FRHISetSpectatorScreenModeTexturePlusEyeLayout(SpectatorScreenController, Layout);
+			ALLOC_COMMAND_CL(RHICmdList, FRHISetSpectatorScreenModeTexturePlusEyeLayout)(SpectatorScreenController, Layout);
 		}
 	);
 }
diff --git a/Engine/Source/Runtime/HeadMountedDisplay/Private/XRThreadUtils.cpp b/Engine/Source/Runtime/HeadMountedDisplay/Private/XRThreadUtils.cpp
--- a/Engine/Source/Runtime/HeadMountedDisplay/Private/XRThreadUtils.cpp
+++ b/Engine/Source/Runtime/HeadMountedDisplay/Private/XRThreadUtils.cpp
@@ -92,7 +92,7 @@ namespace
 
 		if (GRHIThreadId && !IsInRHIThread() && !RHICmdList.Bypass())
 		{
-			new (RHICmdList.AllocCommand<FXRFunctionWrapperRHICommand<T>>()) FXRFunctionWrapperRHICommand<T>(Function);
+			ALLOC_COMMAND_CL(RHICmdList, FXRFunctionWrapperRHICommand<T>)(Function);
 			if (bFlush)
 			{
 				RHICmdList.ImmediateFlush(EImmediateFlushType::FlushRHIThread);
diff --git a/Engine/Source/Runtime/RenderCore/Private/RenderingThread.cpp b/Engine/Source/Runtime/RenderCore/Private/RenderingThread.cpp
--- a/Engine/Source/Runtime/RenderCore/Private/RenderingThread.cpp
+++ b/Engine/Source/Runtime/RenderCore/Private/RenderingThread.cpp
@@ -966,7 +966,7 @@ void FRenderCommandFence::BeginFence(bool bSyncToRHIAndGPU)
 			{
 				if (GRHIThread_InternalUseOnly)
 				{
-					new (RHICmdList.AllocCommand<FRHISyncFrameCommand>()) FRHISyncFrameCommand(CompletionEvent, GTSyncType);
+					ALLOC_COMMAND_CL(RHICmdList, FRHISyncFrameCommand)(CompletionEvent, GTSyncType);
 					RHICmdList.ImmediateFlush(EImmediateFlushType::DispatchToRHIThread);
 				}
 				else
diff --git a/Engine/Source/Runtime/SlateRHIRenderer/Private/Slate3DRenderer.cpp b/Engine/Source/Runtime/SlateRHIRenderer/Private/Slate3DRenderer.cpp
--- a/Engine/Source/Runtime/SlateRHIRenderer/Private/Slate3DRenderer.cpp
+++ b/Engine/Source/Runtime/SlateRHIRenderer/Private/Slate3DRenderer.cpp
@@ -200,5 +200,5 @@ void FSlate3DRenderer::DrawWindowToTarget_RenderThread( FRHICommandListImmediate
 
 	ISlate3DRendererPtr Self = SharedThis(this);
 
-	new ( InRHICmdList.AllocCommand<TKeepAliveCommand<ISlate3DRendererPtr>>() ) TKeepAliveCommand<ISlate3DRendererPtr>(Self);
+	ALLOC_COMMAND_CL(InRHICmdList, TKeepAliveCommand<ISlate3DRendererPtr>)(Self);
 }
diff --git a/Engine/Source/Runtime/SlateRHIRenderer/Private/SlateRHIRenderer.cpp b/Engine/Source/Runtime/SlateRHIRenderer/Private/SlateRHIRenderer.cpp
--- a/Engine/Source/Runtime/SlateRHIRenderer/Private/SlateRHIRenderer.cpp
+++ b/Engine/Source/Runtime/SlateRHIRenderer/Private/SlateRHIRenderer.cpp
@@ -1601,7 +1601,7 @@ void FSlateEndDrawingWindowsCommand::EndDrawingWindows(FRHICommandListImmediate&
 {
 	if (!RHICmdList.Bypass())
 	{
-		new (RHICmdList.AllocCommand<FSlateEndDrawingWindowsCommand>()) FSlateEndDrawingWindowsCommand(Policy, DrawBuffer);
+		ALLOC_COMMAND_CL(RHICmdList, FSlateEndDrawingWindowsCommand)(Policy, DrawBuffer);
 	}
 	else
 	{
diff --git a/Engine/Source/Runtime/SlateRHIRenderer/Private/SlateRHIRenderingPolicy.cpp b/Engine/Source/Runtime/SlateRHIRenderer/Private/SlateRHIRenderingPolicy.cpp
--- a/Engine/Source/Runtime/SlateRHIRenderer/Private/SlateRHIRenderingPolicy.cpp
+++ b/Engine/Source/Runtime/SlateRHIRenderer/Private/SlateRHIRenderingPolicy.cpp
@@ -232,7 +232,7 @@ void FSlateRHIRenderingPolicy::UpdateVertexAndIndexBuffers(FRHICommandListImmedi
 		}
 		else
 		{
-			new (RHICmdList.AllocCommand<FSlateUpdateVertexAndIndexBuffers>()) FSlateUpdateVertexAndIndexBuffers(VertexBuffer, IndexBuffer, InBatchData, bAbsoluteIndices);
+			ALLOC_COMMAND_CL(RHICmdList, FSlateUpdateVertexAndIndexBuffers)(VertexBuffer, IndexBuffer, InBatchData, bAbsoluteIndices);
 		}
 	}
 
diff --git a/Engine/Source/Runtime/SlateRHIRenderer/Private/SlateUpdatableBuffer.cpp b/Engine/Source/Runtime/SlateRHIRenderer/Private/SlateUpdatableBuffer.cpp
--- a/Engine/Source/Runtime/SlateRHIRenderer/Private/SlateUpdatableBuffer.cpp
+++ b/Engine/Source/Runtime/SlateRHIRenderer/Private/SlateUpdatableBuffer.cpp
@@ -100,7 +100,7 @@ void FSlateUpdatableInstanceBuffer::UpdateRenderingData_RenderThread(FRHICommand
 	}
 	else
 	{
-		new (RHICmdList.AllocCommand<FSlateUpdateInstanceBufferCommand>()) FSlateUpdateInstanceBufferCommand(InstanceBufferResource, RenderThreadBufferData);
+		ALLOC_COMMAND_CL(RHICmdList, FSlateUpdateInstanceBufferCommand)(InstanceBufferResource, RenderThreadBufferData);
 	}
 }
 
