Potential fix for https://udn.unrealengine.com/questions/467418/view.html

diff --git a/Engine/Source/Runtime/Engine/Classes/Materials/Material.h b/Engine/Source/Runtime/Engine/Classes/Materials/Material.h
--- a/Engine/Source/Runtime/Engine/Classes/Materials/Material.h
+++ b/Engine/Source/Runtime/Engine/Classes/Materials/Material.h
@@ -982,6 +982,10 @@ public:
 
 	/** Checks to see if an input property should be active, based on the state of the material */
 	ENGINE_API virtual bool IsPropertyActive(EMaterialProperty InProperty) const override;
+
+	/** Liked IsPropertyActive(), but considers any state overriden by DerivedMaterial */
+	ENGINE_API bool IsPropertyActiveInDerived(EMaterialProperty InProperty, const UMaterialInterface* DerivedMaterial) const;
+
 #if WITH_EDITOR
 	/** Allows material properties to be compiled with the option of being overridden by the material attributes input. */
 	ENGINE_API virtual int32 CompilePropertyEx( class FMaterialCompiler* Compiler, const FGuid& AttributeID ) override;
diff --git a/Engine/Source/Runtime/Engine/Private/Materials/Material.cpp b/Engine/Source/Runtime/Engine/Private/Materials/Material.cpp
--- a/Engine/Source/Runtime/Engine/Private/Materials/Material.cpp
+++ b/Engine/Source/Runtime/Engine/Private/Materials/Material.cpp
@@ -5800,8 +5800,16 @@ USubsurfaceProfile* UMaterial::GetSubsurfaceProfile_Internal() const
 	return SubsurfaceProfile; 
 }
 
-bool UMaterial::IsPropertyActive(EMaterialProperty InProperty) const 
+bool UMaterial::IsPropertyActive(EMaterialProperty InProperty) const
 {
+	return IsPropertyActiveInDerived(InProperty, this);
+}
+
+bool UMaterial::IsPropertyActiveInDerived(EMaterialProperty InProperty, const UMaterialInterface* DerivedMaterial) const
+{
+	const EBlendMode DerivedBlendMode = DerivedMaterial->GetBlendMode();
+	const EMaterialShadingModel DerivedShadingModel = DerivedMaterial->GetShadingModel();
+
 	if(MaterialDomain == MD_PostProcess)
 	{
 		return InProperty == MP_EmissiveColor || ( BlendableOutputAlpha && InProperty == MP_Opacity );
@@ -5953,15 +5961,15 @@ bool UMaterial::IsPropertyActive(EMaterialProperty InProperty) const
 	{
 		return InProperty == MP_EmissiveColor
 			|| ( InProperty == MP_WorldPositionOffset )
-			|| ( InProperty == MP_OpacityMask && BlendMode == BLEND_Masked ) 
-			|| ( InProperty == MP_Opacity && IsTranslucentBlendMode((EBlendMode)BlendMode) && BlendMode != BLEND_Modulate )
+			|| ( InProperty == MP_OpacityMask && DerivedBlendMode == BLEND_Masked )
+			|| ( InProperty == MP_Opacity && IsTranslucentBlendMode((EBlendMode)DerivedBlendMode) && DerivedBlendMode != BLEND_Modulate )
 			|| ( InProperty >= MP_CustomizedUVs0 && InProperty <= MP_CustomizedUVs7);
 		{
 			return true;
 		}
 	}
 
-	const bool bIsTranslucentBlendMode = IsTranslucentBlendMode((EBlendMode)BlendMode);
+	const bool bIsTranslucentBlendMode = IsTranslucentBlendMode((EBlendMode)DerivedBlendMode);
 	const bool bIsNonDirectionalTranslucencyLightingMode = TranslucencyLightingMode == TLM_VolumetricNonDirectional || TranslucencyLightingMode == TLM_VolumetricPerVertexNonDirectional;
 	const bool bIsVolumetricTranslucencyLightingMode = TranslucencyLightingMode == TLM_VolumetricNonDirectional 
 		|| TranslucencyLightingMode == TLM_VolumetricDirectional 
@@ -5977,41 +5985,41 @@ bool UMaterial::IsPropertyActive(EMaterialProperty InProperty) const
 		Active = false;
 		break;
 	case MP_Refraction:
-		Active =bIsTranslucentBlendMode && BlendMode != BLEND_Modulate;
+		Active =bIsTranslucentBlendMode && DerivedBlendMode != BLEND_Modulate;
 		break;
 	case MP_Opacity:
-		Active = bIsTranslucentBlendMode && BlendMode != BLEND_Modulate;
-		if (IsSubsurfaceShadingModel(ShadingModel))
+		Active = bIsTranslucentBlendMode && DerivedBlendMode != BLEND_Modulate;
+		if (IsSubsurfaceShadingModel(DerivedShadingModel))
 		{
 			Active = true;
 		}
 		break;
 	case MP_OpacityMask:
-		Active = BlendMode == BLEND_Masked;
+		Active = DerivedBlendMode == BLEND_Masked;
 		break;
 	case MP_BaseColor:
 	case MP_AmbientOcclusion:
-		Active = ShadingModel != MSM_Unlit;
+		Active = DerivedShadingModel != MSM_Unlit;
 		break;
 	case MP_Specular:
 	case MP_Roughness:
-		Active = ShadingModel != MSM_Unlit && (!bIsTranslucentBlendMode || !bIsVolumetricTranslucencyLightingMode);
+		Active = DerivedShadingModel != MSM_Unlit && (!bIsTranslucentBlendMode || !bIsVolumetricTranslucencyLightingMode);
 		break;
 	case MP_Metallic:
 		// Subsurface models store opacity in place of Metallic in the GBuffer
-		Active = ShadingModel != MSM_Unlit && (!bIsTranslucentBlendMode || !bIsVolumetricTranslucencyLightingMode);
+		Active = DerivedShadingModel != MSM_Unlit && (!bIsTranslucentBlendMode || !bIsVolumetricTranslucencyLightingMode);
 		break;
 	case MP_Normal:
-		Active = (ShadingModel != MSM_Unlit && (!bIsTranslucentBlendMode || !bIsNonDirectionalTranslucencyLightingMode)) || Refraction.IsConnected();
+		Active = (DerivedShadingModel != MSM_Unlit && (!bIsTranslucentBlendMode || !bIsNonDirectionalTranslucencyLightingMode)) || Refraction.IsConnected();
 		break;
 	case MP_SubsurfaceColor:
-		Active = ShadingModel == MSM_Subsurface || ShadingModel == MSM_PreintegratedSkin || ShadingModel == MSM_TwoSidedFoliage || ShadingModel == MSM_Cloth;
+		Active = DerivedShadingModel == MSM_Subsurface || DerivedShadingModel == MSM_PreintegratedSkin || DerivedShadingModel == MSM_TwoSidedFoliage || DerivedShadingModel == MSM_Cloth;
 		break;
 	case MP_CustomData0:
-		Active = ShadingModel == MSM_ClearCoat || ShadingModel == MSM_Hair || ShadingModel == MSM_Cloth || ShadingModel == MSM_Eye;
+		Active = DerivedShadingModel == MSM_ClearCoat || DerivedShadingModel == MSM_Hair || DerivedShadingModel == MSM_Cloth || DerivedShadingModel == MSM_Eye;
 		break;
 	case MP_CustomData1:
-		Active = ShadingModel == MSM_ClearCoat || ShadingModel == MSM_Eye;
+		Active = DerivedShadingModel == MSM_ClearCoat || DerivedShadingModel == MSM_Eye;
 		break;
 	case MP_TessellationMultiplier:
 	case MP_WorldDisplacement:
diff --git a/Engine/Source/Runtime/Engine/Private/Materials/MaterialInstance.cpp b/Engine/Source/Runtime/Engine/Private/Materials/MaterialInstance.cpp
--- a/Engine/Source/Runtime/Engine/Private/Materials/MaterialInstance.cpp
+++ b/Engine/Source/Runtime/Engine/Private/Materials/MaterialInstance.cpp
@@ -3994,13 +3994,8 @@ USubsurfaceProfile* UMaterialInstance::GetSubsurfaceProfile_Internal() const
 /** Checks to see if an input property should be active, based on the state of the material */
 bool UMaterialInstance::IsPropertyActive(EMaterialProperty InProperty) const
 {
-	if(InProperty == MP_DiffuseColor || InProperty == MP_SpecularColor)
-	{
-		// to suppress some CompilePropertyEx calls
-		return false;
-	}
-
-	return true;
+	const UMaterial* Material = GetMaterial();
+	return Material ? Material->IsPropertyActiveInDerived(InProperty, this) : false;
 }
 
 #if WITH_EDITOR
