Duplicate 4565983
Fix some rare material crashes.

diff --git a/Engine/Source/Editor/DetailCustomizations/Private/SlateBrushCustomization.cpp b/Engine/Source/Editor/DetailCustomizations/Private/SlateBrushCustomization.cpp
--- a/Engine/Source/Editor/DetailCustomizations/Private/SlateBrushCustomization.cpp
+++ b/Engine/Source/Editor/DetailCustomizations/Private/SlateBrushCustomization.cpp
@@ -1226,7 +1226,8 @@ private:
 				UProperty* MaterialDomainProp = FindField<UProperty>(UMaterial::StaticClass(), GET_MEMBER_NAME_CHECKED(UMaterial,MaterialDomain) );
 
 				FScopedTransaction Transaction( FText::Format( NSLOCTEXT("FSlateBrushStructCustomization", "ChangeMaterialDomainTransaction", "Changed {0} to use the UI material domain"), FText::FromString( BaseMaterial->GetName() ) ) );
-
+				FMaterialUpdateContext MaterialUpdateContext;
+				MaterialUpdateContext.AddMaterial(BaseMaterial);
 				BaseMaterial->PreEditChange( MaterialDomainProp );
 
 				BaseMaterial->MaterialDomain = MD_UI;
diff --git a/Engine/Source/Runtime/Engine/Private/Materials/Material.cpp b/Engine/Source/Runtime/Engine/Private/Materials/Material.cpp
--- a/Engine/Source/Runtime/Engine/Private/Materials/Material.cpp
+++ b/Engine/Source/Runtime/Engine/Private/Materials/Material.cpp
@@ -4282,13 +4282,15 @@ bool UMaterial::CanEditChange(const UProperty* InProperty) const
 void UMaterial::PreEditChange(UProperty* PropertyThatChanged)
 {
 	Super::PreEditChange(PropertyThatChanged);
-
-	// Flush all pending rendering commands.
-	FlushRenderingCommands();
 }
 
 void UMaterial::PostEditChangeProperty(FPropertyChangedEvent& PropertyChangedEvent)
 {
+	// PreEditChange is not enforced to be called before PostEditChange.
+	// CacheResourceShadersForRendering if called will cause a rendering thread race condition with a debug mechanism (bDeletedThroughDeferredCleanup) if there is no flush or
+	// FMaterialUpdateContext present.
+	FlushRenderingCommands();
+
 	// If the material changes, then the debug view material must reset to prevent parameters mismatch
 	void ClearAllDebugViewMaterials();
 	ClearAllDebugViewMaterials();
