Fix bogus compile warning

diff --git a/Engine/Source/Runtime/Engine/Private/BatchedElements.cpp b/Engine/Source/Runtime/Engine/Private/BatchedElements.cpp
--- a/Engine/Source/Runtime/Engine/Private/BatchedElements.cpp
+++ b/Engine/Source/Runtime/Engine/Private/BatchedElements.cpp
@@ -1099,7 +1099,7 @@ bool FBatchedElements::Draw(FRHICommandList& RHICmdList, const FDrawingPolicyRen
 			//Sort sprites by blend mode and texture
 			struct FCompareSprite
 			{
-				explicit FCompareSprite(EBlendModeFilter::Type Filter) : ValidBlendFilter(Filter) {}
+				explicit FCompareSprite(EBlendModeFilter::Type InFilter) : ValidBlendFilter(InFilter) {}
 
 				FORCEINLINE bool operator()( const FBatchedElements::FBatchedSprite& SpriteA, const FBatchedElements::FBatchedSprite& SpriteB ) const
 				{
