Add new IsPropertyActiveInEditor(), fixes problem with IsPropertyActive() introduced in CL 4564983
When making decisions about properties in editor, we don't want to apply any optimizations that may be
introduced by UMaterialInterface::Get* methods

diff --git a/Engine/Source/Editor/UnrealEd/Private/MaterialGraph.cpp b/Engine/Source/Editor/UnrealEd/Private/MaterialGraph.cpp
--- a/Engine/Source/Editor/UnrealEd/Private/MaterialGraph.cpp
+++ b/Engine/Source/Editor/UnrealEd/Private/MaterialGraph.cpp
@@ -333,7 +333,7 @@ bool UMaterialGraph::IsInputActive(UEdGraphPin* GraphPin) const
 		{
 			if (RootNode->Pins[Index] == GraphPin)
 			{
-				return Material->IsPropertyActive(MaterialInputs[Index].GetProperty());
+				return Material->IsPropertyActiveInEditor(MaterialInputs[Index].GetProperty());
 			}
 		}
 	}
diff --git a/Engine/Source/Runtime/Engine/Classes/Materials/Material.h b/Engine/Source/Runtime/Engine/Classes/Materials/Material.h
--- a/Engine/Source/Runtime/Engine/Classes/Materials/Material.h
+++ b/Engine/Source/Runtime/Engine/Classes/Materials/Material.h
@@ -983,6 +983,16 @@ public:
 	/** Checks to see if an input property should be active, based on the state of the material */
 	ENGINE_API virtual bool IsPropertyActive(EMaterialProperty InProperty) const override;
 
+#if WITH_EDITOR
+	/** 
+	* Like IsPropertyActive(), but should be used in context of editor.
+	* For example, there is an optimization that transforms masked materials into opaque materials in certain situations.  If this optimization is active,
+	* the opacity mask input will no longer be active normally (since blend mode will be reported as opaque),
+	* but we still want to be able to connect this input from within the material editor.
+	*/
+	ENGINE_API bool IsPropertyActiveInEditor(EMaterialProperty InProperty) const;
+#endif
+
 	/** Like IsPropertyActive(), but considers any state overriden by DerivedMaterial */
 	ENGINE_API bool IsPropertyActiveInDerived(EMaterialProperty InProperty, const UMaterialInterface* DerivedMaterial) const;
 
diff --git a/Engine/Source/Runtime/Engine/Private/Materials/Material.cpp b/Engine/Source/Runtime/Engine/Private/Materials/Material.cpp
--- a/Engine/Source/Runtime/Engine/Private/Materials/Material.cpp
+++ b/Engine/Source/Runtime/Engine/Private/Materials/Material.cpp
@@ -5816,26 +5816,26 @@ USubsurfaceProfile* UMaterial::GetSubsurfaceProfile_Internal() const
 	return SubsurfaceProfile; 
 }
 
-bool UMaterial::IsPropertyActive(EMaterialProperty InProperty) const
-{
-	return IsPropertyActiveInDerived(InProperty, this);
-}
-
-bool UMaterial::IsPropertyActiveInDerived(EMaterialProperty InProperty, const UMaterialInterface* DerivedMaterial) const
+static bool IsPropertyActive_Internal(EMaterialProperty InProperty,
+	EMaterialDomain Domain,
+	EBlendMode BlendMode,
+	EMaterialShadingModel ShadingModel,
+	ETranslucencyLightingMode TranslucencyLightingMode,
+	EDecalBlendMode DecalBlendMode,
+	bool bBlendableOutputAlpha,
+	bool bHasTessellation,
+	bool bHasRefraction )
 {
-	const EBlendMode DerivedBlendMode = DerivedMaterial->GetBlendMode();
-	const EMaterialShadingModel DerivedShadingModel = DerivedMaterial->GetShadingModel();
-
-	if(MaterialDomain == MD_PostProcess)
+	if (Domain == MD_PostProcess)
 	{
-		return InProperty == MP_EmissiveColor || ( BlendableOutputAlpha && InProperty == MP_Opacity );
+		return InProperty == MP_EmissiveColor || (bBlendableOutputAlpha && InProperty == MP_Opacity);
 	}
-	else if(MaterialDomain == MD_LightFunction)
+	else if (Domain == MD_LightFunction)
 	{
 		// light functions should already use MSM_Unlit but we also we don't want WorldPosOffset
 		return InProperty == MP_EmissiveColor;
 	}
-	else if(MaterialDomain == MD_DeferredDecal)
+	else if (Domain == MD_DeferredDecal)
 	{
 		if (InProperty >= MP_CustomizedUVs0 && InProperty <= MP_CustomizedUVs7)
 		{
@@ -5846,150 +5846,150 @@ bool UMaterial::IsPropertyActiveInDerived(EMaterialProperty InProperty, const UM
 			// todo: MaterialAttruibutes would not return true, should it? Why we don't check for the checkbox in the material
 			return true;
 		}
-		else if( InProperty == MP_WorldPositionOffset )
+		else if (InProperty == MP_WorldPositionOffset)
 		{
 			// Note: DeferredDecals don't support this but MeshDecals do
 			return true;
 		}
 
-		switch(DecalBlendMode)
-		{
-			case DBM_Translucent:
-				return InProperty == MP_EmissiveColor
-					|| InProperty == MP_Normal
-					|| InProperty == MP_Metallic
-					|| InProperty == MP_Specular
-					|| InProperty == MP_BaseColor
-					|| InProperty == MP_Roughness
-					|| InProperty == MP_Opacity;
-
-			case DBM_Stain:
-				return InProperty == MP_EmissiveColor
-					|| InProperty == MP_Normal
-					|| InProperty == MP_Metallic
-					|| InProperty == MP_Specular
-					|| InProperty == MP_BaseColor
-					|| InProperty == MP_Roughness
-					|| InProperty == MP_Opacity;
-
-			case DBM_Normal:
-				return InProperty == MP_Normal
-					|| InProperty == MP_Opacity;
-
-			case DBM_Emissive:
-				// even emissive supports opacity
-				return InProperty == MP_EmissiveColor
-					|| InProperty == MP_Opacity;
-
-			case DBM_AlphaComposite:
-				return InProperty == MP_EmissiveColor
-					|| InProperty == MP_Metallic
-					|| InProperty == MP_Specular
-					|| InProperty == MP_BaseColor
-					|| InProperty == MP_Roughness
-					|| InProperty == MP_Opacity;
-
-			case DBM_DBuffer_AlphaComposite:
-				return InProperty == MP_EmissiveColor
-					|| InProperty == MP_Metallic
-					|| InProperty == MP_Specular
-					|| InProperty == MP_BaseColor
-					|| InProperty == MP_Roughness
-					|| InProperty == MP_Opacity;
-
-			case DBM_DBuffer_ColorNormalRoughness:
-				return InProperty == MP_Normal
-					|| InProperty == MP_BaseColor
-					|| InProperty == MP_Roughness
-					|| InProperty == MP_Metallic
-					|| InProperty == MP_Specular
-					|| InProperty == MP_Opacity
-					|| InProperty == MP_EmissiveColor;
-
-			case DBM_DBuffer_Emissive:
-				return InProperty == MP_EmissiveColor
-					|| InProperty == MP_Opacity;
-
-			case DBM_DBuffer_EmissiveAlphaComposite:
-				return InProperty == MP_EmissiveColor
-					|| InProperty == MP_Opacity;
-
-			case DBM_DBuffer_Color:
-				return InProperty == MP_BaseColor
-					|| InProperty == MP_Opacity
-					|| InProperty == MP_EmissiveColor;
-
-			case DBM_DBuffer_ColorNormal:
-				return InProperty == MP_BaseColor
-					|| InProperty == MP_Normal
-					|| InProperty == MP_Opacity
-					|| InProperty == MP_EmissiveColor;
-
-			case DBM_DBuffer_ColorRoughness:
-				return InProperty == MP_BaseColor
-					|| InProperty == MP_Roughness
-					|| InProperty == MP_Metallic
-					|| InProperty == MP_Specular
-					|| InProperty == MP_Opacity
-					|| InProperty == MP_EmissiveColor;
-
-			case DBM_DBuffer_NormalRoughness:
-				return InProperty == MP_Normal
-					|| InProperty == MP_Roughness
-					|| InProperty == MP_Metallic
-					|| InProperty == MP_Specular
-					|| InProperty == MP_Opacity;
-
-			case DBM_DBuffer_Normal:
-				return InProperty == MP_Normal
-					|| InProperty == MP_Opacity;
-
-			case DBM_DBuffer_Roughness:
-				return InProperty == MP_Roughness
-					|| InProperty == MP_Metallic
-					|| InProperty == MP_Specular
-					|| InProperty == MP_Opacity;
-
-			case DBM_Volumetric_DistanceFunction:
-				return InProperty == MP_EmissiveColor
-					|| InProperty == MP_Normal
-					|| InProperty == MP_Metallic
-					|| InProperty == MP_Specular
-					|| InProperty == MP_BaseColor
-					|| InProperty == MP_Roughness
-					|| InProperty == MP_OpacityMask;
-
-			case DBM_AmbientOcclusion:
-				return InProperty == MP_AmbientOcclusion;
-
-			default:
-				// if you create a new mode it needs to expose the right pins
-				return false;
+		switch (DecalBlendMode)
+		{
+		case DBM_Translucent:
+			return InProperty == MP_EmissiveColor
+				|| InProperty == MP_Normal
+				|| InProperty == MP_Metallic
+				|| InProperty == MP_Specular
+				|| InProperty == MP_BaseColor
+				|| InProperty == MP_Roughness
+				|| InProperty == MP_Opacity;
+
+		case DBM_Stain:
+			return InProperty == MP_EmissiveColor
+				|| InProperty == MP_Normal
+				|| InProperty == MP_Metallic
+				|| InProperty == MP_Specular
+				|| InProperty == MP_BaseColor
+				|| InProperty == MP_Roughness
+				|| InProperty == MP_Opacity;
+
+		case DBM_Normal:
+			return InProperty == MP_Normal
+				|| InProperty == MP_Opacity;
+
+		case DBM_Emissive:
+			// even emissive supports opacity
+			return InProperty == MP_EmissiveColor
+				|| InProperty == MP_Opacity;
+
+		case DBM_AlphaComposite:
+			return InProperty == MP_EmissiveColor
+				|| InProperty == MP_Metallic
+				|| InProperty == MP_Specular
+				|| InProperty == MP_BaseColor
+				|| InProperty == MP_Roughness
+				|| InProperty == MP_Opacity;
+
+		case DBM_DBuffer_AlphaComposite:
+			return InProperty == MP_EmissiveColor
+				|| InProperty == MP_Metallic
+				|| InProperty == MP_Specular
+				|| InProperty == MP_BaseColor
+				|| InProperty == MP_Roughness
+				|| InProperty == MP_Opacity;
+
+		case DBM_DBuffer_ColorNormalRoughness:
+			return InProperty == MP_Normal
+				|| InProperty == MP_BaseColor
+				|| InProperty == MP_Roughness
+				|| InProperty == MP_Metallic
+				|| InProperty == MP_Specular
+				|| InProperty == MP_Opacity
+				|| InProperty == MP_EmissiveColor;
+
+		case DBM_DBuffer_Emissive:
+			return InProperty == MP_EmissiveColor
+				|| InProperty == MP_Opacity;
+
+		case DBM_DBuffer_EmissiveAlphaComposite:
+			return InProperty == MP_EmissiveColor
+				|| InProperty == MP_Opacity;
+
+		case DBM_DBuffer_Color:
+			return InProperty == MP_BaseColor
+				|| InProperty == MP_Opacity
+				|| InProperty == MP_EmissiveColor;
+
+		case DBM_DBuffer_ColorNormal:
+			return InProperty == MP_BaseColor
+				|| InProperty == MP_Normal
+				|| InProperty == MP_Opacity
+				|| InProperty == MP_EmissiveColor;
+
+		case DBM_DBuffer_ColorRoughness:
+			return InProperty == MP_BaseColor
+				|| InProperty == MP_Roughness
+				|| InProperty == MP_Metallic
+				|| InProperty == MP_Specular
+				|| InProperty == MP_Opacity
+				|| InProperty == MP_EmissiveColor;
+
+		case DBM_DBuffer_NormalRoughness:
+			return InProperty == MP_Normal
+				|| InProperty == MP_Roughness
+				|| InProperty == MP_Metallic
+				|| InProperty == MP_Specular
+				|| InProperty == MP_Opacity;
+
+		case DBM_DBuffer_Normal:
+			return InProperty == MP_Normal
+				|| InProperty == MP_Opacity;
+
+		case DBM_DBuffer_Roughness:
+			return InProperty == MP_Roughness
+				|| InProperty == MP_Metallic
+				|| InProperty == MP_Specular
+				|| InProperty == MP_Opacity;
+
+		case DBM_Volumetric_DistanceFunction:
+			return InProperty == MP_EmissiveColor
+				|| InProperty == MP_Normal
+				|| InProperty == MP_Metallic
+				|| InProperty == MP_Specular
+				|| InProperty == MP_BaseColor
+				|| InProperty == MP_Roughness
+				|| InProperty == MP_OpacityMask;
+
+		case DBM_AmbientOcclusion:
+			return InProperty == MP_AmbientOcclusion;
+
+		default:
+			// if you create a new mode it needs to expose the right pins
+			return false;
 		}
 	}
-	else if (MaterialDomain == MD_Volume)
+	else if (Domain == MD_Volume)
 	{
 		return InProperty == MP_EmissiveColor
 			|| InProperty == MP_Opacity
 			|| InProperty == MP_BaseColor;
 	}
-	else if ( MaterialDomain == MD_UI )
+	else if (Domain == MD_UI)
 	{
 		return InProperty == MP_EmissiveColor
-			|| ( InProperty == MP_WorldPositionOffset )
-			|| ( InProperty == MP_OpacityMask && DerivedBlendMode == BLEND_Masked )
-			|| ( InProperty == MP_Opacity && IsTranslucentBlendMode((EBlendMode)DerivedBlendMode) && DerivedBlendMode != BLEND_Modulate )
-			|| ( InProperty >= MP_CustomizedUVs0 && InProperty <= MP_CustomizedUVs7);
+			|| (InProperty == MP_WorldPositionOffset)
+			|| (InProperty == MP_OpacityMask && BlendMode == BLEND_Masked)
+			|| (InProperty == MP_Opacity && IsTranslucentBlendMode(BlendMode) && BlendMode != BLEND_Modulate)
+			|| (InProperty >= MP_CustomizedUVs0 && InProperty <= MP_CustomizedUVs7);
 		{
 			return true;
 		}
 	}
 
-	const bool bIsTranslucentBlendMode = IsTranslucentBlendMode((EBlendMode)DerivedBlendMode);
+	const bool bIsTranslucentBlendMode = IsTranslucentBlendMode(BlendMode);
 	const bool bIsNonDirectionalTranslucencyLightingMode = TranslucencyLightingMode == TLM_VolumetricNonDirectional || TranslucencyLightingMode == TLM_VolumetricPerVertexNonDirectional;
-	const bool bIsVolumetricTranslucencyLightingMode = TranslucencyLightingMode == TLM_VolumetricNonDirectional 
-		|| TranslucencyLightingMode == TLM_VolumetricDirectional 
-		|| TranslucencyLightingMode == TLM_VolumetricPerVertexNonDirectional 
+	const bool bIsVolumetricTranslucencyLightingMode = TranslucencyLightingMode == TLM_VolumetricNonDirectional
+		|| TranslucencyLightingMode == TLM_VolumetricDirectional
+		|| TranslucencyLightingMode == TLM_VolumetricPerVertexNonDirectional
 		|| TranslucencyLightingMode == TLM_VolumetricPerVertexDirectional;
 
 	bool Active = true;
@@ -6001,45 +6001,45 @@ bool UMaterial::IsPropertyActiveInDerived(EMaterialProperty InProperty, const UM
 		Active = false;
 		break;
 	case MP_Refraction:
-		Active =bIsTranslucentBlendMode && DerivedBlendMode != BLEND_Modulate;
+		Active = bIsTranslucentBlendMode && BlendMode != BLEND_Modulate;
 		break;
 	case MP_Opacity:
-		Active = bIsTranslucentBlendMode && DerivedBlendMode != BLEND_Modulate;
-		if (IsSubsurfaceShadingModel(DerivedShadingModel))
+		Active = bIsTranslucentBlendMode && BlendMode != BLEND_Modulate;
+		if (IsSubsurfaceShadingModel(ShadingModel))
 		{
 			Active = true;
 		}
 		break;
 	case MP_OpacityMask:
-		Active = DerivedBlendMode == BLEND_Masked;
+		Active = BlendMode == BLEND_Masked;
 		break;
 	case MP_BaseColor:
 	case MP_AmbientOcclusion:
-		Active = DerivedShadingModel != MSM_Unlit;
+		Active = ShadingModel != MSM_Unlit;
 		break;
 	case MP_Specular:
 	case MP_Roughness:
-		Active = DerivedShadingModel != MSM_Unlit && (!bIsTranslucentBlendMode || !bIsVolumetricTranslucencyLightingMode);
+		Active = ShadingModel != MSM_Unlit && (!bIsTranslucentBlendMode || !bIsVolumetricTranslucencyLightingMode);
 		break;
 	case MP_Metallic:
 		// Subsurface models store opacity in place of Metallic in the GBuffer
-		Active = DerivedShadingModel != MSM_Unlit && (!bIsTranslucentBlendMode || !bIsVolumetricTranslucencyLightingMode);
+		Active = ShadingModel != MSM_Unlit && (!bIsTranslucentBlendMode || !bIsVolumetricTranslucencyLightingMode);
 		break;
 	case MP_Normal:
-		Active = (DerivedShadingModel != MSM_Unlit && (!bIsTranslucentBlendMode || !bIsNonDirectionalTranslucencyLightingMode)) || Refraction.IsConnected();
+		Active = (ShadingModel != MSM_Unlit && (!bIsTranslucentBlendMode || !bIsNonDirectionalTranslucencyLightingMode)) || bHasRefraction;
 		break;
 	case MP_SubsurfaceColor:
-		Active = DerivedShadingModel == MSM_Subsurface || DerivedShadingModel == MSM_PreintegratedSkin || DerivedShadingModel == MSM_TwoSidedFoliage || DerivedShadingModel == MSM_Cloth;
+		Active = ShadingModel == MSM_Subsurface || ShadingModel == MSM_PreintegratedSkin || ShadingModel == MSM_TwoSidedFoliage || ShadingModel == MSM_Cloth;
 		break;
 	case MP_CustomData0:
-		Active = DerivedShadingModel == MSM_ClearCoat || DerivedShadingModel == MSM_Hair || DerivedShadingModel == MSM_Cloth || DerivedShadingModel == MSM_Eye;
+		Active = ShadingModel == MSM_ClearCoat || ShadingModel == MSM_Hair || ShadingModel == MSM_Cloth || ShadingModel == MSM_Eye;
 		break;
 	case MP_CustomData1:
-		Active = DerivedShadingModel == MSM_ClearCoat || DerivedShadingModel == MSM_Eye;
+		Active = ShadingModel == MSM_ClearCoat || ShadingModel == MSM_Eye;
 		break;
 	case MP_TessellationMultiplier:
 	case MP_WorldDisplacement:
-		Active = D3D11TessellationMode != MTM_NoTessellation;
+		Active = bHasTessellation;
 		break;
 	case MP_EmissiveColor:
 		// Emissive is always active, even for light functions and post process materials
@@ -6059,6 +6059,41 @@ bool UMaterial::IsPropertyActiveInDerived(EMaterialProperty InProperty, const UM
 	return Active;
 }
 
+bool UMaterial::IsPropertyActive(EMaterialProperty InProperty) const
+{
+	return IsPropertyActiveInDerived(InProperty, this);
+}
+
+#if WITH_EDITOR
+bool UMaterial::IsPropertyActiveInEditor(EMaterialProperty InProperty) const
+{
+	// explicitly DON'T use getters for BlendMode/ShadingModel...these getters may return an optimized value
+	// we want the actual value that's been set by the user in the material editor
+	return IsPropertyActive_Internal(InProperty,
+		MaterialDomain,
+		BlendMode,
+		ShadingModel,
+		TranslucencyLightingMode,
+		DecalBlendMode,
+		BlendableOutputAlpha,
+		D3D11TessellationMode != MTM_NoTessellation,
+		Refraction.IsConnected());
+}
+#endif // WITH_EDITOR
+
+bool UMaterial::IsPropertyActiveInDerived(EMaterialProperty InProperty, const UMaterialInterface* DerivedMaterial) const
+{
+	return IsPropertyActive_Internal(InProperty,
+		MaterialDomain,
+		DerivedMaterial->GetBlendMode(),
+		DerivedMaterial->GetShadingModel(),
+		TranslucencyLightingMode,
+		DecalBlendMode,
+		BlendableOutputAlpha,
+		D3D11TessellationMode != MTM_NoTessellation,
+		Refraction.IsConnected());
+}
+
 #if WITH_EDITORONLY_DATA
 void UMaterial::FlipExpressionPositions(const TArray<UMaterialExpression*>& Expressions, const TArray<UMaterialExpressionComment*>& Comments, bool bScaleCoords, UMaterial* InMaterial)
 {
