DR - Remove unecessary branch in MorphTarget range caluclation that was also reusing a wrong Cvar also causing this UDN issue: https://udn.unrealengine.com/questions/465992/morph-targets-stutter-with-rmorphtargetweightthres.html

diff --git a/Engine/Source/Runtime/Engine/Private/SkeletalRenderGPUSkin.cpp b/Engine/Source/Runtime/Engine/Private/SkeletalRenderGPUSkin.cpp
--- a/Engine/Source/Runtime/Engine/Private/SkeletalRenderGPUSkin.cpp
+++ b/Engine/Source/Runtime/Engine/Private/SkeletalRenderGPUSkin.cpp
@@ -624,22 +624,18 @@ static void CalculateMorphDeltaBounds(const TArray<float>& MorphTargetWeights, c
 	double MaxScale[4] = { 0, 0, 0, 0 };
 	for (uint32 i = 0; i < MorphTargetVertexInfoBuffers.GetNumMorphs(); i++)
 	{
-		float AbsoluteMorphTargetWeight = FMath::Abs(MorphTargetWeights[i]);
-		if (AbsoluteMorphTargetWeight > GMorphTargetWeightThreshold)
-		{
-			FVector4 MinMorphScale = MorphTargetVertexInfoBuffers.GetMinimumMorphScale(i);
-			FVector4 MaxMorphScale = MorphTargetVertexInfoBuffers.GetMaximumMorphScale(i);
+		FVector4 MinMorphScale = MorphTargetVertexInfoBuffers.GetMinimumMorphScale(i);
+		FVector4 MaxMorphScale = MorphTargetVertexInfoBuffers.GetMaximumMorphScale(i);
 
-			for (uint32 j = 0; j < 4; j++)
-			{
-				MinAccumScale[j] += MorphTargetWeights[i] * MinMorphScale[j];
-				MaxAccumScale[j] += MorphTargetWeights[i] * MaxMorphScale[j];
+		for (uint32 j = 0; j < 4; j++)
+		{
+			MinAccumScale[j] += MorphTargetWeights[i] * MinMorphScale[j];
+			MaxAccumScale[j] += MorphTargetWeights[i] * MaxMorphScale[j];
 
-				double AbsMorphScale = FMath::Max(FMath::Abs(MinMorphScale[j]), FMath::Abs(MaxMorphScale[j]));
-				double AbsAccumScale = FMath::Max(FMath::Abs(MinAccumScale[j]), FMath::Abs(MaxAccumScale[j]));
-				//the maximum accumulated and the maximum local value have to fit into out int24
-				MaxScale[j] = FMath::Max(MaxScale[j], FMath::Max(AbsMorphScale, AbsAccumScale));
-			}
+			double AbsMorphScale = FMath::Max(FMath::Abs(MinMorphScale[j]), FMath::Abs(MaxMorphScale[j]));
+			double AbsAccumScale = FMath::Max(FMath::Abs(MinAccumScale[j]), FMath::Abs(MaxAccumScale[j]));
+			//the maximum accumulated and the maximum local value have to fit into out int24
+			MaxScale[j] = FMath::Max(MaxScale[j], FMath::Max(AbsMorphScale, AbsAccumScale));
 		}
 	}
 
