Unify with 4601059:

Fix asserts during populating DDC (UE-65065).

- Reported on UDN as well, with this fix.
- Indendependtly fixed in CL 4378681 (except it is using a binary &).

diff --git a/Engine/Source/Runtime/Engine/Private/ShaderCompiler/ShaderCompiler.cpp b/Engine/Source/Runtime/Engine/Private/ShaderCompiler/ShaderCompiler.cpp
--- a/Engine/Source/Runtime/Engine/Private/ShaderCompiler/ShaderCompiler.cpp
+++ b/Engine/Source/Runtime/Engine/Private/ShaderCompiler/ShaderCompiler.cpp
@@ -1945,7 +1945,7 @@ void FShaderCompilingManager::ProcessCompiledShaderMaps(
 					for (int32 Index = 0; Index < PipelineJob->StageJobs.Num(); ++Index)
 					{
 						bool bCheckSucceeded = CheckSingleJob(PipelineJob->StageJobs[Index]->GetSingleShaderJob(), MaterialsArray, Errors);
-						bSuccess &= PipelineJob->StageJobs[Index]->bSucceeded && bCheckSucceeded;
+						bSuccess = PipelineJob->StageJobs[Index]->bSucceeded && bCheckSucceeded && bSuccess;
 					}
 				}
 			}
