Fix skylight hiding/showing

diff --git a/Engine/Source/Runtime/Engine/Private/Components/SkyLightComponent.cpp b/Engine/Source/Runtime/Engine/Private/Components/SkyLightComponent.cpp
--- a/Engine/Source/Runtime/Engine/Private/Components/SkyLightComponent.cpp
+++ b/Engine/Source/Runtime/Engine/Private/Components/SkyLightComponent.cpp
@@ -249,17 +249,10 @@ void USkyLightComponent::CreateRenderState_Concurrent()
 {
 	Super::CreateRenderState_Concurrent();
 
-	bool bHidden = false;
-#if WITH_EDITORONLY_DATA
-	bHidden = GetOwner() ? GetOwner()->bHiddenEdLevel : false;
-#endif // WITH_EDITORONLY_DATA
+	bool bHidden = !ShouldComponentAddToScene() || !ShouldRender();
+	const bool bIsValid = SourceType != SLS_SpecifiedCubemap || Cubemap != NULL;
 
-	if(!ShouldComponentAddToScene())
-	{
-		bHidden = true;
-	}
 
-	const bool bIsValid = SourceType != SLS_SpecifiedCubemap || Cubemap != NULL;
 
 	if (bAffectsWorld && bVisible && !bHidden && bIsValid)
 	{
