Back out changelist 4604523
Looks like it broke some automation

diff --git a/Engine/Source/Runtime/Engine/Private/Components/SkyLightComponent.cpp b/Engine/Source/Runtime/Engine/Private/Components/SkyLightComponent.cpp
--- a/Engine/Source/Runtime/Engine/Private/Components/SkyLightComponent.cpp
+++ b/Engine/Source/Runtime/Engine/Private/Components/SkyLightComponent.cpp
@@ -249,10 +249,17 @@ void USkyLightComponent::CreateRenderState_Concurrent()
 {
 	Super::CreateRenderState_Concurrent();
 
-	bool bHidden = !ShouldComponentAddToScene() || !ShouldRender();
-	const bool bIsValid = SourceType != SLS_SpecifiedCubemap || Cubemap != NULL;
+	bool bHidden = false;
+#if WITH_EDITORONLY_DATA
+	bHidden = GetOwner() ? GetOwner()->bHiddenEdLevel : false;
+#endif // WITH_EDITORONLY_DATA
 
+	if(!ShouldComponentAddToScene())
+	{
+		bHidden = true;
+	}
 
+	const bool bIsValid = SourceType != SLS_SpecifiedCubemap || Cubemap != NULL;
 
 	if (bAffectsWorld && bVisible && !bHidden && bIsValid)
 	{
