Adds Render Command scopes. Buffers up lambdas and executes on the render thread.
Switches FCanvas to using this so it no longer races with other threads enqueuing rendering commands.

diff --git a/Engine/Source/Developer/MaterialBaking/Private/MaterialRenderItem.cpp b/Engine/Source/Developer/MaterialBaking/Private/MaterialRenderItem.cpp
--- a/Engine/Source/Developer/MaterialBaking/Private/MaterialRenderItem.cpp
+++ b/Engine/Source/Developer/MaterialBaking/Private/MaterialRenderItem.cpp
@@ -22,7 +22,7 @@ bool FMeshMaterialRenderItem::Render_RenderThread(FRHICommandListImmediate& RHIC
 	return false;
 }
 
-bool FMeshMaterialRenderItem::Render_GameThread(const FCanvas* Canvas)
+bool FMeshMaterialRenderItem::Render_GameThread(const FCanvas* Canvas, FRenderThreadScope& RenderScope)
 {
 	checkSlow(ViewFamily && MaterialSettings && MeshSettings && MaterialRenderProxy);
 	// current render target set for the canvas
@@ -63,7 +63,7 @@ bool FMeshMaterialRenderItem::Render_GameThread(const FCanvas* Canvas)
 
 	if (Vertices.Num() && Indices.Num())
 	{
-		ENQUEUE_RENDER_COMMAND(DrawMaterialCommand)(
+		RenderScope.EnqueueRenderCommand(
 			[Parameters](FRHICommandListImmediate& RHICmdList)
 		{
 			FDrawingPolicyRenderState DrawRenderState(*Parameters.View);
diff --git a/Engine/Source/Developer/MaterialBaking/Private/MaterialRenderItem.h b/Engine/Source/Developer/MaterialBaking/Private/MaterialRenderItem.h
--- a/Engine/Source/Developer/MaterialBaking/Private/MaterialRenderItem.h
+++ b/Engine/Source/Developer/MaterialBaking/Private/MaterialRenderItem.h
@@ -21,7 +21,7 @@ public:
 
 	/** Begin FCanvasBaseRenderItem overrides */
 	virtual bool Render_RenderThread(FRHICommandListImmediate& RHICmdList, FDrawingPolicyRenderState& DrawRenderState, const FCanvas* Canvas) final;
-	virtual bool Render_GameThread(const FCanvas* Canvas) final;
+	virtual bool Render_GameThread(const FCanvas* Canvas, FRenderThreadScope& RenderScope) final;
 	/** End FCanvasBaseRenderItem overrides */
 
 	/** Populate vertices and indices according to available mesh data and otherwise uses simple quad */
diff --git a/Engine/Source/Developer/MaterialUtilities/Private/MeshRendering.cpp b/Engine/Source/Developer/MaterialUtilities/Private/MeshRendering.cpp
--- a/Engine/Source/Developer/MaterialUtilities/Private/MeshRendering.cpp
+++ b/Engine/Source/Developer/MaterialUtilities/Private/MeshRendering.cpp
@@ -546,7 +546,7 @@ public:
 		return true;
 	}
 
-	virtual bool Render_GameThread(const FCanvas* Canvas)
+	virtual bool Render_GameThread(const FCanvas* Canvas, FRenderThreadScope& RenderScope)
 	{
 		checkSlow(Data);
 		// current render target set for the canvas
@@ -581,22 +581,22 @@ public:
 		};
 
 		FDrawMaterialParameters Parameters = DrawMaterialParameters;
-		ENQUEUE_RENDER_COMMAND(DrawMaterialCommand)(
+		RenderScope.EnqueueRenderCommand(
 			[Parameters](FRHICommandListImmediate& RHICmdList)
-			{
-				FDrawingPolicyRenderState DrawRenderState(*Parameters.View);
+		{
+			FDrawingPolicyRenderState DrawRenderState(*Parameters.View);
 
-				// disable depth test & writes
-				DrawRenderState.SetDepthStencilState(TStaticDepthStencilState<false, CF_Always>::GetRHI());
+			// disable depth test & writes
+			DrawRenderState.SetDepthStencilState(TStaticDepthStencilState<false, CF_Always>::GetRHI());
 
-				RenderMaterial(RHICmdList, DrawRenderState, *Parameters.View, *Parameters.RenderData);
+			RenderMaterial(RHICmdList, DrawRenderState, *Parameters.View, *Parameters.RenderData);
 
-				delete Parameters.View;
-				if (Parameters.AllowedCanvasModes & FCanvas::Allow_DeleteOnRender)
-				{
-					delete Parameters.RenderData;
-				}
-			});
+			delete Parameters.View;
+			if (Parameters.AllowedCanvasModes & FCanvas::Allow_DeleteOnRender)
+			{
+				delete Parameters.RenderData;
+			}
+		});
 		if (Canvas->GetAllowedModes() & FCanvas::Allow_DeleteOnRender)
 		{
 			Data = NULL;
diff --git a/Engine/Source/Runtime/Engine/Private/TileRendering.cpp b/Engine/Source/Runtime/Engine/Private/TileRendering.cpp
--- a/Engine/Source/Runtime/Engine/Private/TileRendering.cpp
+++ b/Engine/Source/Runtime/Engine/Private/TileRendering.cpp
@@ -419,7 +419,7 @@ bool FCanvasTileRendererItem::Render_RenderThread(FRHICommandListImmediate& RHIC
 	return true;
 }
 
-bool FCanvasTileRendererItem::Render_GameThread(const FCanvas* Canvas)
+bool FCanvasTileRendererItem::Render_GameThread(const FCanvas* Canvas, FRenderThreadScope& RenderScope)
 {
 	float CurrentRealTime = 0.f;
 	float CurrentWorldTime = 0.f;
@@ -475,41 +475,41 @@ bool FCanvasTileRendererItem::Render_GameThread(const FCanvas* Canvas)
 
 	InitTileBuffers(&Data->VertexFactory, Data->Tiles, *View, bNeedsToSwitchVerticalAxis);
 
-	ENQUEUE_RENDER_COMMAND(DrawTileCommand)(
+	RenderScope.EnqueueRenderCommand(
 		[DrawTileParameters](FRHICommandListImmediate& RHICmdList)
-		{
-			SCOPED_DRAW_EVENTF(RHICmdList, CanvasDrawTile, *DrawTileParameters.RenderData->MaterialRenderProxy->GetMaterial(GMaxRHIFeatureLevel)->GetFriendlyName());
+	{
+		SCOPED_DRAW_EVENTF(RHICmdList, CanvasDrawTile, *DrawTileParameters.RenderData->MaterialRenderProxy->GetMaterial(GMaxRHIFeatureLevel)->GetFriendlyName());
 
-			FDrawingPolicyRenderState DrawRenderState(*DrawTileParameters.View);
+		FDrawingPolicyRenderState DrawRenderState(*DrawTileParameters.View);
 
-			// disable depth test & writes
-			DrawRenderState.SetDepthStencilState(TStaticDepthStencilState<false, CF_Always>::GetRHI());
+		// disable depth test & writes
+		DrawRenderState.SetDepthStencilState(TStaticDepthStencilState<false, CF_Always>::GetRHI());
 
-			for (int32 TileIdx = 0; TileIdx < DrawTileParameters.RenderData->Tiles.Num(); TileIdx++)
-			{
-				FRenderData::FTileInst& Tile = DrawTileParameters.RenderData->Tiles[TileIdx];
+		for (int32 TileIdx = 0; TileIdx < DrawTileParameters.RenderData->Tiles.Num(); TileIdx++)
+		{
+			FRenderData::FTileInst& Tile = DrawTileParameters.RenderData->Tiles[TileIdx];
 
-				// update the FMeshBatch
-				FMeshBatch& Mesh = DrawTileParameters.RenderData->TileMesh.MeshElement;
-				Mesh.MaterialRenderProxy = DrawTileParameters.RenderData->MaterialRenderProxy;
-				Mesh.Elements[0].BaseVertexIndex = NUM_MATERIAL_TILE_VERTS * TileIdx;
+			// update the FMeshBatch
+			FMeshBatch& Mesh = DrawTileParameters.RenderData->TileMesh.MeshElement;
+			Mesh.MaterialRenderProxy = DrawTileParameters.RenderData->MaterialRenderProxy;
+			Mesh.Elements[0].BaseVertexIndex = NUM_MATERIAL_TILE_VERTS * TileIdx;
 
-				GetRendererModule().DrawTileMesh(RHICmdList, DrawRenderState, *DrawTileParameters.View, Mesh, DrawTileParameters.bIsHitTesting, Tile.HitProxyId);
-			}
+			GetRendererModule().DrawTileMesh(RHICmdList, DrawRenderState, *DrawTileParameters.View, Mesh, DrawTileParameters.bIsHitTesting, Tile.HitProxyId);
+		}
 
-			DrawTileParameters.RenderData->StaticMeshVertexBuffers.PositionVertexBuffer.ReleaseResource();
-			DrawTileParameters.RenderData->StaticMeshVertexBuffers.StaticMeshVertexBuffer.ReleaseResource();
-			DrawTileParameters.RenderData->StaticMeshVertexBuffers.ColorVertexBuffer.ReleaseResource();
-			DrawTileParameters.RenderData->TileMesh.ReleaseResource();
-			DrawTileParameters.RenderData->VertexFactory.ReleaseResource();
+		DrawTileParameters.RenderData->StaticMeshVertexBuffers.PositionVertexBuffer.ReleaseResource();
+		DrawTileParameters.RenderData->StaticMeshVertexBuffers.StaticMeshVertexBuffer.ReleaseResource();
+		DrawTileParameters.RenderData->StaticMeshVertexBuffers.ColorVertexBuffer.ReleaseResource();
+		DrawTileParameters.RenderData->TileMesh.ReleaseResource();
+		DrawTileParameters.RenderData->VertexFactory.ReleaseResource();
 
-			delete DrawTileParameters.View->Family;
-			delete DrawTileParameters.View;
-			if (DrawTileParameters.AllowedCanvasModes & FCanvas::Allow_DeleteOnRender)
-			{
-				delete DrawTileParameters.RenderData;
-			}
-		});
+		delete DrawTileParameters.View->Family;
+		delete DrawTileParameters.View;
+		if (DrawTileParameters.AllowedCanvasModes & FCanvas::Allow_DeleteOnRender)
+		{
+			delete DrawTileParameters.RenderData;
+		}
+	});
 
 	if (Canvas->GetAllowedModes() & FCanvas::Allow_DeleteOnRender)
 	{
diff --git a/Engine/Source/Runtime/Engine/Private/TriangleRendering.cpp b/Engine/Source/Runtime/Engine/Private/TriangleRendering.cpp
--- a/Engine/Source/Runtime/Engine/Private/TriangleRendering.cpp
+++ b/Engine/Source/Runtime/Engine/Private/TriangleRendering.cpp
@@ -348,7 +348,7 @@ bool FCanvasTriangleRendererItem::Render_RenderThread(FRHICommandListImmediate&
 	return true;
 }
 
-bool FCanvasTriangleRendererItem::Render_GameThread(const FCanvas* Canvas)
+bool FCanvasTriangleRendererItem::Render_GameThread(const FCanvas* Canvas, FRenderThreadScope& RenderScope)
 {
 	float CurrentRealTime = 0.f;
 	float CurrentWorldTime = 0.f;
@@ -405,40 +405,40 @@ bool FCanvasTriangleRendererItem::Render_GameThread(const FCanvas* Canvas)
 	InitTriangleBuffers(&Data->VertexFactory, Data->Triangles, *View, bNeedsToSwitchVerticalAxis);
 
 	FDrawTriangleParameters Parameters = DrawTriangleParameters;
-	ENQUEUE_RENDER_COMMAND(DrawTriangleCommand)(
+	RenderScope.EnqueueRenderCommand(
 		[Parameters](FRHICommandListImmediate& RHICmdList)	
-		{
-			FDrawingPolicyRenderState DrawRenderState(*Parameters.View);
+	{
+		FDrawingPolicyRenderState DrawRenderState(*Parameters.View);
 
-			// disable depth test & writes
-			DrawRenderState.SetDepthStencilState(TStaticDepthStencilState<false, CF_Always>::GetRHI());
+		// disable depth test & writes
+		DrawRenderState.SetDepthStencilState(TStaticDepthStencilState<false, CF_Always>::GetRHI());
 
-			SCOPED_DRAW_EVENT(RHICmdList, CanvasDrawTriangle);
-			for (int32 TriIdx = 0; TriIdx < Parameters.RenderData->Triangles.Num(); TriIdx++)
-			{
-				const FRenderData::FTriangleInst& Tri = Parameters.RenderData->Triangles[TriIdx];
-				// update the FMeshBatch
-				FMeshBatch& TriMesh = Parameters.RenderData->TriMesh.TriMeshElement;
-				TriMesh.VertexFactory = &Parameters.RenderData->VertexFactory;
-				TriMesh.MaterialRenderProxy = Parameters.RenderData->MaterialRenderProxy;
-				TriMesh.Elements[0].BaseVertexIndex = 3 * TriIdx;
-
-				GetRendererModule().DrawTileMesh(RHICmdList, DrawRenderState, *Parameters.View, TriMesh, Parameters.bIsHitTesting, Tri.HitProxyId);
-			}
+		SCOPED_DRAW_EVENT(RHICmdList, CanvasDrawTriangle);
+		for (int32 TriIdx = 0; TriIdx < Parameters.RenderData->Triangles.Num(); TriIdx++)
+		{
+			const FRenderData::FTriangleInst& Tri = Parameters.RenderData->Triangles[TriIdx];
+			// update the FMeshBatch
+			FMeshBatch& TriMesh = Parameters.RenderData->TriMesh.TriMeshElement;
+			TriMesh.VertexFactory = &Parameters.RenderData->VertexFactory;
+			TriMesh.MaterialRenderProxy = Parameters.RenderData->MaterialRenderProxy;
+			TriMesh.Elements[0].BaseVertexIndex = 3 * TriIdx;
+
+			GetRendererModule().DrawTileMesh(RHICmdList, DrawRenderState, *Parameters.View, TriMesh, Parameters.bIsHitTesting, Tri.HitProxyId);
+		}
 
-			Parameters.RenderData->StaticMeshVertexBuffers.PositionVertexBuffer.ReleaseResource();
-			Parameters.RenderData->StaticMeshVertexBuffers.StaticMeshVertexBuffer.ReleaseResource();
-			Parameters.RenderData->StaticMeshVertexBuffers.ColorVertexBuffer.ReleaseResource();
-			Parameters.RenderData->TriMesh.ReleaseResource();
-			Parameters.RenderData->VertexFactory.ReleaseResource();
+		Parameters.RenderData->StaticMeshVertexBuffers.PositionVertexBuffer.ReleaseResource();
+		Parameters.RenderData->StaticMeshVertexBuffers.StaticMeshVertexBuffer.ReleaseResource();
+		Parameters.RenderData->StaticMeshVertexBuffers.ColorVertexBuffer.ReleaseResource();
+		Parameters.RenderData->TriMesh.ReleaseResource();
+		Parameters.RenderData->VertexFactory.ReleaseResource();
 
-			delete Parameters.View->Family;
-			delete Parameters.View;
-			if (Parameters.AllowedCanvasModes & FCanvas::Allow_DeleteOnRender)
-			{
-				delete Parameters.RenderData;
-			}
-		});
+		delete Parameters.View->Family;
+		delete Parameters.View;
+		if (Parameters.AllowedCanvasModes & FCanvas::Allow_DeleteOnRender)
+		{
+			delete Parameters.RenderData;
+		}
+	});
 
 	if (Canvas->GetAllowedModes() & FCanvas::Allow_DeleteOnRender)
 	{
diff --git a/Engine/Source/Runtime/Engine/Private/UserInterface/Canvas.cpp b/Engine/Source/Runtime/Engine/Private/UserInterface/Canvas.cpp
--- a/Engine/Source/Runtime/Engine/Private/UserInterface/Canvas.cpp
+++ b/Engine/Source/Runtime/Engine/Private/UserInterface/Canvas.cpp
@@ -418,7 +418,7 @@ bool FCanvasBatchedElementRenderItem::Render_RenderThread(FRHICommandListImmedia
 	return bDirty;
 }
 
-bool FCanvasBatchedElementRenderItem::Render_GameThread(const FCanvas* Canvas)
+bool FCanvasBatchedElementRenderItem::Render_GameThread(const FCanvas* Canvas, FRenderThreadScope& RenderScope)
 {	
 	checkSlow(Data);
 	bool bDirty=false;
@@ -462,32 +462,32 @@ bool FCanvasBatchedElementRenderItem::Render_GameThread(const FCanvas* Canvas)
 			Canvas->GetFeatureLevel(),
 			Canvas->GetShaderPlatform()
 		};
-		ENQUEUE_RENDER_COMMAND(BatchedDrawCommand)(
+		RenderScope.EnqueueRenderCommand(
 			[DrawParameters](FRHICommandList& RHICmdList)
-			{
-				FSceneView SceneView = FBatchedElements::CreateProxySceneView(DrawParameters.RenderData->Transform.GetMatrix(), FIntRect(0, 0, DrawParameters.ViewportSizeX, DrawParameters.ViewportSizeY));
+		{
+			FSceneView SceneView = FBatchedElements::CreateProxySceneView(DrawParameters.RenderData->Transform.GetMatrix(), FIntRect(0, 0, DrawParameters.ViewportSizeX, DrawParameters.ViewportSizeY));
 
-				FDrawingPolicyRenderState DrawRenderState(SceneView);
+			FDrawingPolicyRenderState DrawRenderState(SceneView);
 
-				// disable depth test & writes
-				DrawRenderState.SetDepthStencilState(TStaticDepthStencilState<false, CF_Always>::GetRHI());
-				DrawRenderState.SetBlendState(TStaticBlendState<>::GetRHI());
+			// disable depth test & writes
+			DrawRenderState.SetDepthStencilState(TStaticDepthStencilState<false, CF_Always>::GetRHI());
+			DrawRenderState.SetBlendState(TStaticBlendState<>::GetRHI());
 
-				// draw batched items
-				DrawParameters.RenderData->BatchedElements.Draw(
-					RHICmdList,
-					DrawRenderState,
-					DrawParameters.FeatureLevel,
-					DrawParameters.bNeedsToSwitchVerticalAxis,
-					SceneView,
-					DrawParameters.bHitTesting,
-					DrawParameters.DisplayGamma);
+			// draw batched items
+			DrawParameters.RenderData->BatchedElements.Draw(
+				RHICmdList,
+				DrawRenderState,
+				DrawParameters.FeatureLevel,
+				DrawParameters.bNeedsToSwitchVerticalAxis,
+				SceneView,
+				DrawParameters.bHitTesting,
+				DrawParameters.DisplayGamma);
 
-				if(DrawParameters.AllowedCanvasModes & FCanvas::Allow_DeleteOnRender )
-				{
-					delete DrawParameters.RenderData;
-				}
-			});
+			if(DrawParameters.AllowedCanvasModes & FCanvas::Allow_DeleteOnRender )
+			{
+				delete DrawParameters.RenderData;
+			}
+		});
 	}
 	if( Canvas->GetAllowedModes() & FCanvas::Allow_DeleteOnRender )
 	{
@@ -803,8 +803,12 @@ void FCanvas::Flush_GameThread(bool bForce)
 	};
 	bool bEmitCanvasDrawEvents = GetEmitDrawEvents();
 
-	ENQUEUE_RENDER_COMMAND(CanvasFlushSetupCommand)(
-		[FlushParameters](FRHICommandList& RHICmdList)
+	// Only create these render commands if we actually have something to draw.
+	if(SortedElements.Num() > 0)
+	{
+		FRenderThreadScope RenderThreadScope;
+		RenderThreadScope.EnqueueRenderCommand(
+			[FlushParameters](FRHICommandList& RHICmdList)
 		{
 			check(RHICmdList.IsOutsideRenderPass());
 
@@ -829,41 +833,43 @@ void FCanvas::Flush_GameThread(bool bForce)
 			RHICmdList.SetViewport(ViewportRect.Min.X, ViewportRect.Min.Y, 0.0f, ViewportRect.Max.X, ViewportRect.Max.Y, 1.0f);
 		});
 
-	// iterate over the FCanvasSortElements in sorted order and render all the batched items for each entry
-	for( int32 Idx=0; Idx < SortedElements.Num(); Idx++ )
-	{
-		FCanvasSortElement& SortElement = SortedElements[Idx];
-		for( int32 BatchIdx=0; BatchIdx < SortElement.RenderBatchArray.Num(); BatchIdx++ )
+		// iterate over the FCanvasSortElements in sorted order and render all the batched items for each entry
+		for( int32 Idx=0; Idx < SortedElements.Num(); Idx++ )
 		{
-			FCanvasBaseRenderItem* RenderItem = SortElement.RenderBatchArray[BatchIdx];
-			if( RenderItem )
+			FCanvasSortElement& SortElement = SortedElements[Idx];
+			for( int32 BatchIdx=0; BatchIdx < SortElement.RenderBatchArray.Num(); BatchIdx++ )
 			{
-				// mark current render target as dirty since we are drawing to it
-				bRenderTargetDirty |= RenderItem->Render_GameThread(this);
-				if( AllowedModes & Allow_DeleteOnRender )
+				FCanvasBaseRenderItem* RenderItem = SortElement.RenderBatchArray[BatchIdx];
+				if( RenderItem )
 				{
-					delete RenderItem;
+					// mark current render target as dirty since we are drawing to it
+					bRenderTargetDirty |= RenderItem->Render_GameThread(this, RenderThreadScope);
+					if( AllowedModes & Allow_DeleteOnRender )
+					{
+						delete RenderItem;
+					}
 				}
-			}			
+			}
+			if( AllowedModes & Allow_DeleteOnRender )
+			{
+				SortElement.RenderBatchArray.Empty();
+			}
 		}
-		if( AllowedModes & Allow_DeleteOnRender )
+
+		RenderThreadScope.EnqueueRenderCommand(
+			[](FRHICommandList& RHICmdList)
 		{
-			SortElement.RenderBatchArray.Empty();
-		}
+			RHICmdList.EndRenderPass();
+		});
 	}
+	
 	if( AllowedModes & Allow_DeleteOnRender )
 	{
-		// empty the array of FCanvasSortElement entries after finished with rendering	
+		// empty the array of FCanvasSortElement entries after finished with rendering
 		SortedElements.Empty();
 		SortedElementLookupMap.Empty();
 		LastElementIndex = INDEX_NONE;
 	}
-
-	ENQUEUE_RENDER_COMMAND(CanvasFlushEndCommand)(
-		[](FRHICommandList& RHICmdList)
-	{
-		RHICmdList.EndRenderPass();
-	});
 }
 
 void FCanvas::PushRelativeTransform(const FMatrix& Transform)
diff --git a/Engine/Source/Runtime/Engine/Public/CanvasTypes.h b/Engine/Source/Runtime/Engine/Public/CanvasTypes.h
--- a/Engine/Source/Runtime/Engine/Public/CanvasTypes.h
+++ b/Engine/Source/Runtime/Engine/Public/CanvasTypes.h
@@ -795,7 +795,7 @@ public:
 	* @param Canvas - canvas currently being rendered
 	* @return true if anything rendered
 	*/
-	virtual bool Render_GameThread(const FCanvas* Canvas) = 0;
+	virtual bool Render_GameThread(const FCanvas* Canvas, FRenderThreadScope& RenderScope) = 0;
 	
 	/**
 	* FCanvasBatchedElementRenderItem instance accessor
@@ -875,7 +875,7 @@ public:
 	* @param Canvas - canvas currently being rendered
 	* @return true if anything rendered
 	*/
-	virtual bool Render_GameThread(const FCanvas* Canvas) override;
+	virtual bool Render_GameThread(const FCanvas* Canvas, FRenderThreadScope& RenderScope) override;
 
 	/**
 	* Determine if this is a matching set by comparing texture,blendmode,elementype,transform. All must match
@@ -1003,7 +1003,7 @@ public:
 	* @param Canvas - canvas currently being rendered
 	* @return true if anything rendered
 	*/
-	virtual bool Render_GameThread(const FCanvas* Canvas) override;
+	virtual bool Render_GameThread(const FCanvas* Canvas, FRenderThreadScope& RenderScope) override;
 
 	/**
 	* Determine if this is a matching set by comparing material,transform. All must match
@@ -1159,7 +1159,7 @@ public:
 	* @param Canvas - canvas currently being rendered
 	* @return true if anything rendered
 	*/
-	virtual bool Render_GameThread(const FCanvas* Canvas) override;
+	virtual bool Render_GameThread(const FCanvas* Canvas, FRenderThreadScope& RenderScope) override;
 
 	/**
 	* Determine if this is a matching set by comparing material,transform. All must match
diff --git a/Engine/Source/Runtime/RenderCore/Public/RenderingThread.h b/Engine/Source/Runtime/RenderCore/Public/RenderingThread.h
--- a/Engine/Source/Runtime/RenderCore/Public/RenderingThread.h
+++ b/Engine/Source/Runtime/RenderCore/Public/RenderingThread.h
@@ -647,3 +647,42 @@ extern RENDERCORE_API void BeginCleanup(FDeferredCleanupInterface* CleanupObject
  * @return A pointer to the set of pending cleanup objects.  The called is responsible for deletion.
  */
 extern RENDERCORE_API FPendingCleanupObjects* GetPendingCleanupObjects();
+
+////////////////////////////////////
+// RenderThread scoped work
+////////////////////////////////////
+
+class RENDERCORE_API FRenderThreadScope
+{
+	typedef TFunction<void(FRHICommandListImmediate&)> RenderCommandFunction;
+	typedef TArray<RenderCommandFunction> RenderCommandFunctionArray;
+public:
+	FRenderThreadScope()
+	{
+		RenderCommands = new RenderCommandFunctionArray;
+	}
+
+	~FRenderThreadScope()
+	{
+		RenderCommandFunctionArray* RenderCommandArray = RenderCommands;
+
+		ENQUEUE_RENDER_COMMAND(DispatchScopeCommands)(
+			[RenderCommandArray](FRHICommandListImmediate& RHICmdList)
+		{
+			for(uint32 Index = 0; Index < RenderCommandArray->Num(); Index++)
+			{
+				(*RenderCommandArray)[Index](RHICmdList);
+			}
+
+			delete RenderCommandArray;
+		});
+	}
+
+	void EnqueueRenderCommand(RenderCommandFunction&& Lambda)
+	{
+		RenderCommands->Add(MoveTemp(Lambda));
+	}
+
+private:
+	RenderCommandFunctionArray* RenderCommands;
+};
