Add DeltaTime, ShadowReplace material expression

diff --git a/Engine/Shaders/Private/Common.ush b/Engine/Shaders/Private/Common.ush
--- a/Engine/Shaders/Private/Common.ush
+++ b/Engine/Shaders/Private/Common.ush
@@ -1686,14 +1686,20 @@ SafeSaturate_Def(float4)
 // Can be accessed with a Custom material node. If this turns out to be very useful we can expose as MaterialFunction
 // and potentially expose other queries as well (e.g. SkeletalMesh, HitProxy, ).
 // @return 0:no, 1:yes
-float IsShadowDepthShader()
+bool GetShadowReplaceState()
 {
 #ifdef SHADOW_DEPTH_SHADER
-	return 1;
+	return true;
 #else
-	return 0;
+	return false;
 #endif
 }
+
+float IsShadowDepthShader()
+{
+	return GetShadowReplaceState() ? 1.0f : 0.0f;
+}
+
 #define TERRAIN_ZSCALE (1.0f/128.0f)
 
 // Decodes a value which was packed into two 8 bit channels
diff --git a/Engine/Source/Runtime/Engine/Classes/Materials/MaterialExpressionDeltaTime.h b/Engine/Source/Runtime/Engine/Classes/Materials/MaterialExpressionDeltaTime.h
new file mode 100644
--- /dev/null
+++ b/Engine/Source/Runtime/Engine/Classes/Materials/MaterialExpressionDeltaTime.h
@@ -0,0 +1,22 @@
+// Copyright 1998-2018 Epic Games, Inc. All Rights Reserved.
+
+
+#pragma once
+
+#include "CoreMinimal.h"
+#include "UObject/ObjectMacros.h"
+#include "Materials/MaterialExpression.h"
+#include "MaterialExpressionDeltaTime.generated.h"
+
+UCLASS(collapsecategories, hidecategories=Object)
+class UMaterialExpressionDeltaTime : public UMaterialExpression
+{
+	GENERATED_UCLASS_BODY()
+
+	//~ Begin UMaterialExpression Interface
+#if WITH_EDITOR
+	virtual int32 Compile(class FMaterialCompiler* Compiler, int32 OutputIndex) override;
+	virtual void GetCaption(TArray<FString>& OutCaptions) const override;
+#endif
+	//~ End UMaterialExpression Interface
+};
diff --git a/Engine/Source/Runtime/Engine/Classes/Materials/MaterialExpressionShadowReplace.h b/Engine/Source/Runtime/Engine/Classes/Materials/MaterialExpressionShadowReplace.h
new file mode 100644
--- /dev/null
+++ b/Engine/Source/Runtime/Engine/Classes/Materials/MaterialExpressionShadowReplace.h
@@ -0,0 +1,31 @@
+// Copyright 1998-2018 Epic Games, Inc. All Rights Reserved.
+
+
+#pragma once
+
+#include "CoreMinimal.h"
+#include "UObject/ObjectMacros.h"
+#include "MaterialExpressionIO.h"
+#include "Materials/MaterialExpression.h"
+#include "MaterialExpressionShadowReplace.generated.h"
+
+UCLASS()
+class UMaterialExpressionShadowReplace : public UMaterialExpression
+{
+	GENERATED_UCLASS_BODY()
+
+	UPROPERTY()
+	FExpressionInput Default;
+
+	UPROPERTY()
+	FExpressionInput Shadow;
+
+	//~ Begin UMaterialExpression Interface
+#if WITH_EDITOR
+	virtual int32 Compile(class FMaterialCompiler* Compiler, int32 OutputIndex) override;
+	virtual void GetCaption(TArray<FString>& OutCaptions) const override;
+#endif
+	//~ End UMaterialExpression Interface
+};
+
+
diff --git a/Engine/Source/Runtime/Engine/Private/Materials/HLSLMaterialTranslator.h b/Engine/Source/Runtime/Engine/Private/Materials/HLSLMaterialTranslator.h
--- a/Engine/Source/Runtime/Engine/Private/Materials/HLSLMaterialTranslator.h
+++ b/Engine/Source/Runtime/Engine/Private/Materials/HLSLMaterialTranslator.h
@@ -2483,6 +2483,14 @@ protected:
 		return AddInlinedCodeChunk(MCT_Float, TEXT("fmod(View.RealTime,%s)"), *GetParameterCode(PeriodChunk));
 	}
 
+	virtual int32 DeltaTime() override
+	{
+		// explicitly avoid trying to return previous frame's delta time for bCompilingPreviousFrame here
+		// DeltaTime expression is designed to be used when generating custom motion vectors, by using world position offset along with previous frame switch
+		// in this context, we will technically be evaluating the previous frame, but we want to use the current frame's delta tick in order to offset the vector used to create previous position
+		return AddInlinedCodeChunk(MCT_Float, TEXT("View.DeltaTime"));
+	}
+
 	virtual int32 PeriodicHint(int32 PeriodicCode) override
 	{
 		if(PeriodicCode == INDEX_NONE)
@@ -5043,6 +5051,17 @@ protected:
 		return AddCodeChunk(ResultType,TEXT("(GetGIReplaceState() ? (%s) : (%s))"), *GetParameterCode(DynamicIndirect), *GetParameterCode(Direct));
 	}
 
+	virtual int32 ShadowReplace(int32 Default, int32 Shadow) override
+	{
+		if (Default == INDEX_NONE || Shadow == INDEX_NONE)
+		{
+			return INDEX_NONE;
+		}
+
+		EMaterialValueType ResultType = GetArithmeticResultType(Default, Shadow);
+		return AddCodeChunk(ResultType, TEXT("(GetShadowReplaceState() ? (%s) : (%s))"), *GetParameterCode(Shadow), *GetParameterCode(Default));
+	}
+
 	virtual int32 MaterialProxyReplace(int32 Realtime, int32 MaterialProxy) override { return Realtime; }
 
 	virtual int32 ObjectOrientation() override
diff --git a/Engine/Source/Runtime/Engine/Private/Materials/MaterialExpressions.cpp b/Engine/Source/Runtime/Engine/Private/Materials/MaterialExpressions.cpp
--- a/Engine/Source/Runtime/Engine/Private/Materials/MaterialExpressions.cpp
+++ b/Engine/Source/Runtime/Engine/Private/Materials/MaterialExpressions.cpp
@@ -126,6 +126,7 @@
 #include "Materials/MaterialExpressionReroute.h"
 #include "Materials/MaterialExpressionScalarParameter.h"
 #include "Materials/MaterialExpressionSetMaterialAttributes.h"
+#include "Materials/MaterialExpressionShadowReplace.h"
 #include "Materials/MaterialExpressionSign.h"
 #include "Materials/MaterialExpressionStaticBoolParameter.h"
 #include "Materials/MaterialExpressionStaticSwitchParameter.h"
@@ -186,6 +187,7 @@
 #include "Materials/MaterialExpressionTextureSampleParameterVolume.h"
 #include "Materials/MaterialExpressionTextureCoordinate.h"
 #include "Materials/MaterialExpressionTime.h"
+#include "Materials/MaterialExpressionDeltaTime.h"
 #include "Materials/MaterialExpressionTransform.h"
 #include "Materials/MaterialExpressionTransformPosition.h"
 #include "Materials/MaterialExpressionTruncate.h"
@@ -7693,6 +7695,39 @@ void UMaterialExpressionViewSize::GetCaption(TArray<FString>& OutCaptions) const
 }
 #endif // WITH_EDITOR
 
+UMaterialExpressionDeltaTime::UMaterialExpressionDeltaTime(const FObjectInitializer& ObjectInitializer)
+	: Super(ObjectInitializer)
+{
+#if WITH_EDITORONLY_DATA
+	// Structure to hold one-time initialization
+	struct FConstructorStatics
+	{
+		FText NAME_Constants;
+		FConstructorStatics()
+			: NAME_Constants(LOCTEXT("Constants", "Constants"))
+		{
+		}
+	};
+	static FConstructorStatics ConstructorStatics;
+
+	MenuCategories.Add(ConstructorStatics.NAME_Constants);
+
+	bShaderInputData = true;
+#endif
+}
+
+#if WITH_EDITOR
+int32 UMaterialExpressionDeltaTime::Compile(class FMaterialCompiler* Compiler, int32 OutputIndex)
+{
+	return Compiler->DeltaTime();
+}
+
+void UMaterialExpressionDeltaTime::GetCaption(TArray<FString>& OutCaptions) const
+{
+	OutCaptions.Add(TEXT("DeltaTime"));
+}
+#endif // WITH_EDITOR
+
 UMaterialExpressionSceneTexelSize::UMaterialExpressionSceneTexelSize(const FObjectInitializer& ObjectInitializer)
 	: Super(ObjectInitializer)
 {
@@ -10751,7 +10786,7 @@ bool UMaterialFunctionInstance::HasFlippedCoordinates() const
 }
 #endif
 
-bool UMaterialFunctionInstance::OverrideNamedScalarParameter(const FMaterialParameterInfo& ParameterInfo, float& OutValue)
+bool UMaterialFunctionInstance::OverrideNamedScalarParameter(const FMaterialParameterInfo& ParameterInfo, float& OutValue) const
 {
 	for (const FScalarParameterValue& ScalarParameter : ScalarParameterValues)
 	{
@@ -10765,7 +10800,7 @@ bool UMaterialFunctionInstance::OverrideNamedScalarParameter(const FMaterialPara
 	return false;
 }
 
-bool UMaterialFunctionInstance::OverrideNamedVectorParameter(const FMaterialParameterInfo& ParameterInfo, FLinearColor& OutValue)
+bool UMaterialFunctionInstance::OverrideNamedVectorParameter(const FMaterialParameterInfo& ParameterInfo, FLinearColor& OutValue) const
 {
 	for (const FVectorParameterValue& VectorParameter : VectorParameterValues)
 	{
@@ -10779,7 +10814,7 @@ bool UMaterialFunctionInstance::OverrideNamedVectorParameter(const FMaterialPara
 	return false;
 }
 
-bool UMaterialFunctionInstance::OverrideNamedTextureParameter(const FMaterialParameterInfo& ParameterInfo, UTexture*& OutValue)
+bool UMaterialFunctionInstance::OverrideNamedTextureParameter(const FMaterialParameterInfo& ParameterInfo, UTexture*& OutValue) const
 {
 	for (const FTextureParameterValue& TextureParameter : TextureParameterValues)
 	{
@@ -10793,7 +10828,7 @@ bool UMaterialFunctionInstance::OverrideNamedTextureParameter(const FMaterialPar
 	return false;
 }
 
-bool UMaterialFunctionInstance::OverrideNamedFontParameter(const FMaterialParameterInfo& ParameterInfo, UFont*& OutFontValue, int32& OutFontPage)
+bool UMaterialFunctionInstance::OverrideNamedFontParameter(const FMaterialParameterInfo& ParameterInfo, UFont*& OutFontValue, int32& OutFontPage) const
 {
 	for (const FFontParameterValue& FontParameter : FontParameterValues)
 	{
@@ -10808,7 +10843,7 @@ bool UMaterialFunctionInstance::OverrideNamedFontParameter(const FMaterialParame
 	return false;
 }
 
-bool UMaterialFunctionInstance::OverrideNamedStaticSwitchParameter(const FMaterialParameterInfo& ParameterInfo, bool& OutValue, FGuid& OutExpressionGuid)
+bool UMaterialFunctionInstance::OverrideNamedStaticSwitchParameter(const FMaterialParameterInfo& ParameterInfo, bool& OutValue, FGuid& OutExpressionGuid) const
 {
 	for (const FStaticSwitchParameter& StaticSwitchParameter : StaticSwitchParameterValues)
 	{
@@ -10823,7 +10858,7 @@ bool UMaterialFunctionInstance::OverrideNamedStaticSwitchParameter(const FMateri
 	return false;
 }
 
-bool UMaterialFunctionInstance::OverrideNamedStaticComponentMaskParameter(const FMaterialParameterInfo& ParameterInfo, bool& OutR, bool& OutG, bool& OutB, bool& OutA, FGuid& OutExpressionGuid)
+bool UMaterialFunctionInstance::OverrideNamedStaticComponentMaskParameter(const FMaterialParameterInfo& ParameterInfo, bool& OutR, bool& OutG, bool& OutB, bool& OutA, FGuid& OutExpressionGuid) const
 {
 	for (const FStaticComponentMaskParameter& StaticComponentMaskParameter : StaticComponentMaskParameterValues)
 	{
@@ -12251,6 +12286,53 @@ void UMaterialExpressionLightmassReplace::GetCaption(TArray<FString>& OutCaption
 }
 #endif // WITH_EDITOR
 
+//
+//	UMaterialExpressionShadowReplace
+//
+UMaterialExpressionShadowReplace::UMaterialExpressionShadowReplace(const FObjectInitializer& ObjectInitializer)
+	: Super(ObjectInitializer)
+{
+	// Structure to hold one-time initialization
+	struct FConstructorStatics
+	{
+		FText NAME_Utility;
+		FConstructorStatics()
+			: NAME_Utility(LOCTEXT("Utility", "Utility"))
+		{
+		}
+	};
+	static FConstructorStatics ConstructorStatics;
+
+#if WITH_EDITORONLY_DATA
+	MenuCategories.Add(ConstructorStatics.NAME_Utility);
+#endif
+}
+
+#if WITH_EDITOR
+int32 UMaterialExpressionShadowReplace::Compile(class FMaterialCompiler* Compiler, int32 OutputIndex)
+{
+	if (!Default.GetTracedInput().Expression)
+	{
+		return Compiler->Errorf(TEXT("Missing ShadowReplace input Default"));
+	}
+	else if (!Shadow.GetTracedInput().Expression)
+	{
+		return Compiler->Errorf(TEXT("Missing ShadowReplace input Shadow"));
+	}
+	else
+	{
+		const int32 Arg1 = Default.Compile(Compiler);
+		const int32 Arg2 = Shadow.Compile(Compiler);
+		return Compiler->ShadowReplace(Arg1, Arg2);
+	}
+}
+
+void UMaterialExpressionShadowReplace::GetCaption(TArray<FString>& OutCaptions) const
+{
+	OutCaptions.Add(TEXT("ShadowReplace"));
+}
+#endif // WITH_EDITOR
+
 
 //
 //	UMaterialExpressionMaterialProxy
diff --git a/Engine/Source/Runtime/Engine/Private/SceneView.cpp b/Engine/Source/Runtime/Engine/Private/SceneView.cpp
--- a/Engine/Source/Runtime/Engine/Private/SceneView.cpp
+++ b/Engine/Source/Runtime/Engine/Private/SceneView.cpp
@@ -2231,6 +2231,7 @@ void FSceneView::SetupCommonViewUniformBufferParameters(
 
 	ViewUniformShaderParameters.GameTime = Family->CurrentWorldTime;
 	ViewUniformShaderParameters.RealTime = Family->CurrentRealTime;
+	ViewUniformShaderParameters.DeltaTime = Family->DeltaWorldTime;
 	ViewUniformShaderParameters.Random = FMath::Rand();
 	ViewUniformShaderParameters.FrameNumber = Family->FrameNumber;
 
diff --git a/Engine/Source/Runtime/Engine/Public/MaterialCompiler.h b/Engine/Source/Runtime/Engine/Public/MaterialCompiler.h
--- a/Engine/Source/Runtime/Engine/Public/MaterialCompiler.h
+++ b/Engine/Source/Runtime/Engine/Public/MaterialCompiler.h
@@ -105,6 +105,7 @@ public:
 
 	virtual int32 GameTime(bool bPeriodic, float Period) = 0;
 	virtual int32 RealTime(bool bPeriodic, float Period) = 0;
+	virtual int32 DeltaTime() = 0;
 	virtual int32 PeriodicHint(int32 PeriodicCode) { return PeriodicCode; }
 
 	virtual int32 Sine(int32 X) = 0;
@@ -256,6 +257,7 @@ public:
 
 	virtual int32 LightmassReplace(int32 Realtime, int32 Lightmass) = 0;
 	virtual int32 GIReplace(int32 Direct, int32 StaticIndirect, int32 DynamicIndirect) = 0;
+	virtual int32 ShadowReplace(int32 Default, int32 Shadow) = 0;
 	virtual int32 MaterialProxyReplace(int32 Realtime, int32 MaterialProxy) = 0;
 
 	virtual int32 ObjectOrientation() = 0;
@@ -352,6 +354,7 @@ public:
 
 	virtual int32 GameTime(bool bPeriodic, float Period) override { return Compiler->GameTime(bPeriodic, Period); }
 	virtual int32 RealTime(bool bPeriodic, float Period) override { return Compiler->RealTime(bPeriodic, Period); }
+	virtual int32 DeltaTime() override { return Compiler->DeltaTime(); }
 
 	virtual int32 PeriodicHint(int32 PeriodicCode) override { return Compiler->PeriodicHint(PeriodicCode); }
 
@@ -473,6 +476,7 @@ public:
 
 	virtual int32 LightmassReplace(int32 Realtime, int32 Lightmass) override { return Realtime; }
 	virtual int32 GIReplace(int32 Direct, int32 StaticIndirect, int32 DynamicIndirect) override { return Compiler->GIReplace(Direct, StaticIndirect, DynamicIndirect); }
+	virtual int32 ShadowReplace(int32 Default, int32 Shadow) override { return Compiler->ShadowReplace(Default, Shadow); }
 	virtual int32 MaterialProxyReplace(int32 Realtime, int32 MaterialProxy) override { return Realtime; }
 	virtual int32 ObjectOrientation() override { return Compiler->ObjectOrientation(); }
 	virtual int32 RotateAboutAxis(int32 NormalizedRotationAxisAndAngleIndex, int32 PositionOnAxisIndex, int32 PositionIndex) override
diff --git a/Engine/Source/Runtime/Engine/Public/SceneView.h b/Engine/Source/Runtime/Engine/Public/SceneView.h
--- a/Engine/Source/Runtime/Engine/Public/SceneView.h
+++ b/Engine/Source/Runtime/Engine/Public/SceneView.h
@@ -673,6 +673,7 @@ enum ETranslucencyVolumeCascade
 	VIEW_UNIFORM_BUFFER_MEMBER(float, AdaptiveTessellationFactor) \
 	VIEW_UNIFORM_BUFFER_MEMBER(float, GameTime) \
 	VIEW_UNIFORM_BUFFER_MEMBER(float, RealTime) \
+	VIEW_UNIFORM_BUFFER_MEMBER(float, DeltaTime) \
 	VIEW_UNIFORM_BUFFER_MEMBER(float, MaterialTextureMipBias) \
 	VIEW_UNIFORM_BUFFER_MEMBER(float, MaterialTextureDerivativeMultiply) \
 	VIEW_UNIFORM_BUFFER_MEMBER(uint32, Random) \
