Make GetXXXParameterDefaultValue methods on materi instance thread safe from any thread, not just game/render thread.
This is needed as material layers may call these from async job threads during rendering

diff --git a/Engine/Source/Runtime/Engine/Classes/Materials/Material.h b/Engine/Source/Runtime/Engine/Classes/Materials/Material.h
--- a/Engine/Source/Runtime/Engine/Classes/Materials/Material.h
+++ b/Engine/Source/Runtime/Engine/Classes/Materials/Material.h
@@ -1260,12 +1260,12 @@ public:
 
 	ENGINE_API virtual void GetDependentFunctions(TArray<class UMaterialFunctionInterface*>& DependentFunctions) const override;
 
-	ENGINE_API virtual bool GetScalarParameterDefaultValue(const FMaterialParameterInfo& ParameterInfo, float& OutValue, bool bOveriddenOnly = false, bool bCheckOwnedGlobalOverrides = false) const override;
-	ENGINE_API virtual bool GetVectorParameterDefaultValue(const FMaterialParameterInfo& ParameterInfo, FLinearColor& OutValue, bool bOveriddenOnly = false, bool bCheckOwnedGlobalOverrides = false) const override;
-	ENGINE_API virtual bool GetTextureParameterDefaultValue(const FMaterialParameterInfo& ParameterInfo, class UTexture*& OutValue, bool bCheckOwnedGlobalOverrides = false) const override;
-	ENGINE_API virtual bool GetFontParameterDefaultValue(const FMaterialParameterInfo& ParameterInfo, class UFont*& OutFontValue, int32& OutFontPage, bool bCheckOwnedGlobalOverrides = false) const override;
-	ENGINE_API virtual bool GetStaticComponentMaskParameterDefaultValue(const FMaterialParameterInfo& ParameterInfo, bool& OutR, bool& OutG, bool& OutB, bool& OutA, FGuid& OutExpressionGuid, bool bCheckOwnedGlobalOverrides = false) const override;
-	ENGINE_API virtual bool GetStaticSwitchParameterDefaultValue(const FMaterialParameterInfo& ParameterInfo, bool& OutValue, FGuid& OutExpressionGuid, bool bCheckOwnedGlobalOverrides = false) const override;
+	ENGINE_API virtual bool GetScalarParameterDefaultValue_Concurrent(TMicRecursionGuard& RecursionGuard, const FMaterialParameterInfo& ParameterInfo, float& OutValue, bool bOveriddenOnly = false, bool bCheckOwnedGlobalOverrides = false) const override;
+	ENGINE_API virtual bool GetVectorParameterDefaultValue_Concurrent(TMicRecursionGuard& RecursionGuard, const FMaterialParameterInfo& ParameterInfo, FLinearColor& OutValue, bool bOveriddenOnly = false, bool bCheckOwnedGlobalOverrides = false) const override;
+	ENGINE_API virtual bool GetTextureParameterDefaultValue_Concurrent(TMicRecursionGuard& RecursionGuard, const FMaterialParameterInfo& ParameterInfo, class UTexture*& OutValue, bool bCheckOwnedGlobalOverrides = false) const override;
+	ENGINE_API virtual bool GetFontParameterDefaultValue_Concurrent(TMicRecursionGuard& RecursionGuard, const FMaterialParameterInfo& ParameterInfo, class UFont*& OutFontValue, int32& OutFontPage, bool bCheckOwnedGlobalOverrides = false) const override;
+	ENGINE_API virtual bool GetStaticComponentMaskParameterDefaultValue_Concurrent(TMicRecursionGuard& RecursionGuard, const FMaterialParameterInfo& ParameterInfo, bool& OutR, bool& OutG, bool& OutB, bool& OutA, FGuid& OutExpressionGuid, bool bCheckOwnedGlobalOverrides = false) const override;
+	ENGINE_API virtual bool GetStaticSwitchParameterDefaultValue_Concurrent(TMicRecursionGuard& RecursionGuard, const FMaterialParameterInfo& ParameterInfo, bool& OutValue, FGuid& OutExpressionGuid, bool bCheckOwnedGlobalOverrides = false) const override;
 
 	/** Returns the material's decal blend mode, calculated from the DecalBlendMode property and what inputs are connected. */
 	uint32 GetDecalBlendMode() const { return DecalBlendMode; }
diff --git a/Engine/Source/Runtime/Engine/Classes/Materials/MaterialFunctionInstance.h b/Engine/Source/Runtime/Engine/Classes/Materials/MaterialFunctionInstance.h
--- a/Engine/Source/Runtime/Engine/Classes/Materials/MaterialFunctionInstance.h
+++ b/Engine/Source/Runtime/Engine/Classes/Materials/MaterialFunctionInstance.h
@@ -157,12 +157,12 @@ class UMaterialFunctionInstance : public UMaterialFunctionInterface
 	}
 
 public:
-	virtual bool OverrideNamedScalarParameter(const FMaterialParameterInfo& ParameterInfo, float& OutValue) override;
-	virtual bool OverrideNamedVectorParameter(const FMaterialParameterInfo& ParameterInfo, FLinearColor& OutValue) override;
-	virtual bool OverrideNamedTextureParameter(const FMaterialParameterInfo& ParameterInfo, class UTexture*& OutValue) override;
-	virtual bool OverrideNamedFontParameter(const FMaterialParameterInfo& ParameterInfo, class UFont*& OutFontValue, int32& OutFontPage) override;
-	virtual bool OverrideNamedStaticSwitchParameter(const FMaterialParameterInfo& ParameterInfo, bool& OutValue, FGuid& OutExpressionGuid) override;
-	virtual bool OverrideNamedStaticComponentMaskParameter(const FMaterialParameterInfo& ParameterInfo, bool& OutR, bool& OutG, bool& OutB, bool& OutA, FGuid& OutExpressionGuid) override;
+	virtual bool OverrideNamedScalarParameter(const FMaterialParameterInfo& ParameterInfo, float& OutValue) const override;
+	virtual bool OverrideNamedVectorParameter(const FMaterialParameterInfo& ParameterInfo, FLinearColor& OutValue) const override;
+	virtual bool OverrideNamedTextureParameter(const FMaterialParameterInfo& ParameterInfo, class UTexture*& OutValue) const override;
+	virtual bool OverrideNamedFontParameter(const FMaterialParameterInfo& ParameterInfo, class UFont*& OutFontValue, int32& OutFontPage) const override;
+	virtual bool OverrideNamedStaticSwitchParameter(const FMaterialParameterInfo& ParameterInfo, bool& OutValue, FGuid& OutExpressionGuid) const override;
+	virtual bool OverrideNamedStaticComponentMaskParameter(const FMaterialParameterInfo& ParameterInfo, bool& OutR, bool& OutG, bool& OutB, bool& OutA, FGuid& OutExpressionGuid) const override;
 	//~ End UMaterialFunctionInterface interface
 
 protected:
diff --git a/Engine/Source/Runtime/Engine/Classes/Materials/MaterialFunctionInterface.h b/Engine/Source/Runtime/Engine/Classes/Materials/MaterialFunctionInterface.h
--- a/Engine/Source/Runtime/Engine/Classes/Materials/MaterialFunctionInterface.h
+++ b/Engine/Source/Runtime/Engine/Classes/Materials/MaterialFunctionInterface.h
@@ -406,32 +406,32 @@ public:
 		}
 	}
 
-	virtual bool OverrideNamedScalarParameter(const FMaterialParameterInfo& ParameterInfo, float& OutValue)
+	virtual bool OverrideNamedScalarParameter(const FMaterialParameterInfo& ParameterInfo, float& OutValue) const
 	{
 		return false;
 	}
 
-	virtual bool OverrideNamedVectorParameter(const FMaterialParameterInfo& ParameterInfo, FLinearColor& OutValue)
+	virtual bool OverrideNamedVectorParameter(const FMaterialParameterInfo& ParameterInfo, FLinearColor& OutValue) const
 	{
 		return false;
 	}
 
-	virtual bool OverrideNamedTextureParameter(const FMaterialParameterInfo& ParameterInfo, class UTexture*& OutValue)
+	virtual bool OverrideNamedTextureParameter(const FMaterialParameterInfo& ParameterInfo, class UTexture*& OutValue) const
 	{
 		return false;
 	}
 	
-	virtual bool OverrideNamedFontParameter(const FMaterialParameterInfo& ParameterInfo, class UFont*& OutFontValue, int32& OutFontPage)
+	virtual bool OverrideNamedFontParameter(const FMaterialParameterInfo& ParameterInfo, class UFont*& OutFontValue, int32& OutFontPage) const
 	{
 		return false;
 	}
 
-	virtual bool OverrideNamedStaticSwitchParameter(const FMaterialParameterInfo& ParameterInfo, bool& OutValue, FGuid& OutExpressionGuid)
+	virtual bool OverrideNamedStaticSwitchParameter(const FMaterialParameterInfo& ParameterInfo, bool& OutValue, FGuid& OutExpressionGuid) const
 	{
 		return false;
 	}
 
-	virtual bool OverrideNamedStaticComponentMaskParameter(const FMaterialParameterInfo& ParameterInfo, bool& OutR, bool& OutG, bool& OutB, bool& OutA, FGuid& OutExpressionGuid)
+	virtual bool OverrideNamedStaticComponentMaskParameter(const FMaterialParameterInfo& ParameterInfo, bool& OutR, bool& OutG, bool& OutB, bool& OutA, FGuid& OutExpressionGuid) const
 	{
 		return false;
 	}
diff --git a/Engine/Source/Runtime/Engine/Classes/Materials/MaterialInstance.h b/Engine/Source/Runtime/Engine/Classes/Materials/MaterialInstance.h
--- a/Engine/Source/Runtime/Engine/Classes/Materials/MaterialInstance.h
+++ b/Engine/Source/Runtime/Engine/Classes/Materials/MaterialInstance.h
@@ -561,12 +561,12 @@ public:
 
 	ENGINE_API virtual void GetDependentFunctions(TArray<class UMaterialFunctionInterface*>& DependentFunctions) const override;
 	
-	ENGINE_API virtual bool GetScalarParameterDefaultValue(const FMaterialParameterInfo& ParameterInfo, float& OutValue, bool bOveriddenOnly = false, bool bCheckOwnedGlobalOverrides = false) const override;
-	ENGINE_API virtual bool GetVectorParameterDefaultValue(const FMaterialParameterInfo& ParameterInfo, FLinearColor& OutValue, bool bOveriddenOnly = false, bool bCheckOwnedGlobalOverrides = false) const override;
-	ENGINE_API virtual bool GetTextureParameterDefaultValue(const FMaterialParameterInfo& ParameterInfo, UTexture*& OutValue, bool bCheckOwnedGlobalOverrides = false) const override;
-	ENGINE_API virtual bool GetFontParameterDefaultValue(const FMaterialParameterInfo& ParameterInfo, UFont*& OutFontValue, int32& OutFontPage, bool bCheckOwnedGlobalOverrides = false) const override;
-	ENGINE_API virtual bool GetStaticSwitchParameterDefaultValue(const FMaterialParameterInfo& ParameterInfo, bool& OutValue, FGuid& OutExpressionGuid, bool bCheckOwnedGlobalOverrides = false) const override;
-	ENGINE_API virtual bool GetStaticComponentMaskParameterDefaultValue(const FMaterialParameterInfo& ParameterInfo, bool& OutR, bool& OutG, bool& OutB, bool& OutA, FGuid& OutExpressionGuid, bool bCheckOwnedGlobalOverrides = false) const override;
+	ENGINE_API virtual bool GetScalarParameterDefaultValue_Concurrent(TMicRecursionGuard& RecursionGuard, const FMaterialParameterInfo& ParameterInfo, float& OutValue, bool bOveriddenOnly = false, bool bCheckOwnedGlobalOverrides = false) const override;
+	ENGINE_API virtual bool GetVectorParameterDefaultValue_Concurrent(TMicRecursionGuard& RecursionGuard, const FMaterialParameterInfo& ParameterInfo, FLinearColor& OutValue, bool bOveriddenOnly = false, bool bCheckOwnedGlobalOverrides = false) const override;
+	ENGINE_API virtual bool GetTextureParameterDefaultValue_Concurrent(TMicRecursionGuard& RecursionGuard, const FMaterialParameterInfo& ParameterInfo, UTexture*& OutValue, bool bCheckOwnedGlobalOverrides = false) const override;
+	ENGINE_API virtual bool GetFontParameterDefaultValue_Concurrent(TMicRecursionGuard& RecursionGuard, const FMaterialParameterInfo& ParameterInfo, UFont*& OutFontValue, int32& OutFontPage, bool bCheckOwnedGlobalOverrides = false) const override;
+	ENGINE_API virtual bool GetStaticSwitchParameterDefaultValue_Concurrent(TMicRecursionGuard& RecursionGuard, const FMaterialParameterInfo& ParameterInfo, bool& OutValue, FGuid& OutExpressionGuid, bool bCheckOwnedGlobalOverrides = false) const override;
+	ENGINE_API virtual bool GetStaticComponentMaskParameterDefaultValue_Concurrent(TMicRecursionGuard& RecursionGuard, const FMaterialParameterInfo& ParameterInfo, bool& OutR, bool& OutG, bool& OutB, bool& OutA, FGuid& OutExpressionGuid, bool bCheckOwnedGlobalOverrides = false) const override;
 
 #if WITH_EDITOR
 	ENGINE_API virtual bool GetGroupName(const FMaterialParameterInfo& ParameterInfo, FName& OutGroup) const override;
diff --git a/Engine/Source/Runtime/Engine/Classes/Materials/MaterialInterface.h b/Engine/Source/Runtime/Engine/Classes/Materials/MaterialInterface.h
--- a/Engine/Source/Runtime/Engine/Classes/Materials/MaterialInterface.h
+++ b/Engine/Source/Runtime/Engine/Classes/Materials/MaterialInterface.h
@@ -488,18 +488,25 @@ public:
 	virtual void GetDependentFunctions(TArray<class UMaterialFunctionInterface*>& DependentFunctions) const
 		PURE_VIRTUAL(UMaterialInterface::GetDependentFunctions,return;);
 
-	virtual bool GetScalarParameterDefaultValue(const FMaterialParameterInfo& ParameterInfo, float& OutValue, bool bOveriddenOnly = false, bool bCheckOwnedGlobalOverrides = false) const
-		PURE_VIRTUAL(UMaterialInterface::GetScalarParameterDefaultValue,return false;);
-	virtual bool GetVectorParameterDefaultValue(const FMaterialParameterInfo& ParameterInfo, FLinearColor& OutValue, bool bOveriddenOnly = false, bool bCheckOwnedGlobalOverrides = false) const
-		PURE_VIRTUAL(UMaterialInterface::GetVectorParameterDefaultValue,return false;);
-	virtual bool GetTextureParameterDefaultValue(const FMaterialParameterInfo& ParameterInfo, class UTexture*& OutValue, bool bCheckOwnedGlobalOverrides = false) const
-		PURE_VIRTUAL(UMaterialInterface::GetTextureParameterDefaultValue,return false;);
-	virtual bool GetFontParameterDefaultValue(const FMaterialParameterInfo& ParameterInfo, class UFont*& OutFontValue, int32& OutFontPage, bool bCheckOwnedGlobalOverrides = false) const
-		PURE_VIRTUAL(UMaterialInterface::GetFontParameterDefaultValue,return false;);
-	virtual bool GetStaticSwitchParameterDefaultValue(const FMaterialParameterInfo& ParameterInfo, bool& OutValue, FGuid& OutExpressionGuid, bool bCheckOwnedGlobalOverrides = false) const
-		PURE_VIRTUAL(UMaterialInterface::GetStaticSwitchParameterDefaultValue,return false;);
-	virtual bool GetStaticComponentMaskParameterDefaultValue(const FMaterialParameterInfo& ParameterInfo, bool& OutR, bool& OutG, bool& OutB, bool& OutA, FGuid& OutExpressionGuid, bool bCheckOwnedGlobalOverrides = false) const
-		PURE_VIRTUAL(UMaterialInterface::GetStaticComponentMaskParameterDefaultValue,return false;);
+	virtual bool GetScalarParameterDefaultValue_Concurrent(TMicRecursionGuard& RecursionGuard, const FMaterialParameterInfo& ParameterInfo, float& OutValue, bool bOveriddenOnly = false, bool bCheckOwnedGlobalOverrides = false) const
+		PURE_VIRTUAL(UMaterialInterface::GetScalarParameterDefaultValue_Concurrent,return false;);
+	virtual bool GetVectorParameterDefaultValue_Concurrent(TMicRecursionGuard& RecursionGuard, const FMaterialParameterInfo& ParameterInfo, FLinearColor& OutValue, bool bOveriddenOnly = false, bool bCheckOwnedGlobalOverrides = false) const
+		PURE_VIRTUAL(UMaterialInterface::GetVectorParameterDefaultValue_Concurrent,return false;);
+	virtual bool GetTextureParameterDefaultValue_Concurrent(TMicRecursionGuard& RecursionGuard, const FMaterialParameterInfo& ParameterInfo, class UTexture*& OutValue, bool bCheckOwnedGlobalOverrides = false) const
+		PURE_VIRTUAL(UMaterialInterface::GetTextureParameterDefaultValue_Concurrent,return false;);
+	virtual bool GetFontParameterDefaultValue_Concurrent(TMicRecursionGuard& RecursionGuard, const FMaterialParameterInfo& ParameterInfo, class UFont*& OutFontValue, int32& OutFontPage, bool bCheckOwnedGlobalOverrides = false) const
+		PURE_VIRTUAL(UMaterialInterface::GetFontParameterDefaultValue_Concurrent,return false;);
+	virtual bool GetStaticSwitchParameterDefaultValue_Concurrent(TMicRecursionGuard& RecursionGuard, const FMaterialParameterInfo& ParameterInfo, bool& OutValue, FGuid& OutExpressionGuid, bool bCheckOwnedGlobalOverrides = false) const
+		PURE_VIRTUAL(UMaterialInterface::GetStaticSwitchParameterDefaultValue_Concurrent,return false;);
+	virtual bool GetStaticComponentMaskParameterDefaultValue_Concurrent(TMicRecursionGuard& RecursionGuard, const FMaterialParameterInfo& ParameterInfo, bool& OutR, bool& OutG, bool& OutB, bool& OutA, FGuid& OutExpressionGuid, bool bCheckOwnedGlobalOverrides = false) const
+		PURE_VIRTUAL(UMaterialInterface::GetStaticComponentMaskParameterDefaultValue_Concurrent,return false;);
+
+	ENGINE_API bool GetScalarParameterDefaultValue(const FMaterialParameterInfo& ParameterInfo, float& OutValue, bool bOveriddenOnly = false, bool bCheckOwnedGlobalOverrides = false) const;
+	ENGINE_API bool GetVectorParameterDefaultValue(const FMaterialParameterInfo& ParameterInfo, FLinearColor& OutValue, bool bOveriddenOnly = false, bool bCheckOwnedGlobalOverrides = false) const;
+	ENGINE_API bool GetTextureParameterDefaultValue(const FMaterialParameterInfo& ParameterInfo, class UTexture*& OutValue, bool bCheckOwnedGlobalOverrides = false) const;
+	ENGINE_API bool GetFontParameterDefaultValue(const FMaterialParameterInfo& ParameterInfo, class UFont*& OutFontValue, int32& OutFontPage, bool bCheckOwnedGlobalOverrides = false) const;
+	ENGINE_API bool GetStaticSwitchParameterDefaultValue(const FMaterialParameterInfo& ParameterInfo, bool& OutValue, FGuid& OutExpressionGuid, bool bCheckOwnedGlobalOverrides = false) const;
+	ENGINE_API bool GetStaticComponentMaskParameterDefaultValue(const FMaterialParameterInfo& ParameterInfo, bool& OutR, bool& OutG, bool& OutB, bool& OutA, FGuid& OutExpressionGuid, bool bCheckOwnedGlobalOverrides = false) const;
 		
 	/** Appends textures referenced by expressions, including nested functions. */
 	virtual void AppendReferencedTextures(TArray<UTexture*>& InOutTextures) const
diff --git a/Engine/Source/Runtime/Engine/Private/Materials/Material.cpp b/Engine/Source/Runtime/Engine/Private/Materials/Material.cpp
--- a/Engine/Source/Runtime/Engine/Private/Materials/Material.cpp
+++ b/Engine/Source/Runtime/Engine/Private/Materials/Material.cpp
@@ -1987,32 +1987,32 @@ void UMaterial::GetDependentFunctions(TArray<UMaterialFunctionInterface*>& Depen
 	}
 }
 
-bool UMaterial::GetScalarParameterDefaultValue(const FMaterialParameterInfo& ParameterInfo, float& OutValue, bool bOveriddenOnly, bool bCheckOwnedGlobalOverrides) const
+bool UMaterial::GetScalarParameterDefaultValue_Concurrent(TMicRecursionGuard& RecursionGuard, const FMaterialParameterInfo& ParameterInfo, float& OutValue, bool bOveriddenOnly, bool bCheckOwnedGlobalOverrides) const
 {
 	return GetScalarParameterValue(ParameterInfo, OutValue, bOveriddenOnly);
 }
 
-bool UMaterial::GetVectorParameterDefaultValue(const FMaterialParameterInfo& ParameterInfo, FLinearColor& OutValue, bool bOveriddenOnly, bool bCheckOwnedGlobalOverrides) const
+bool UMaterial::GetVectorParameterDefaultValue_Concurrent(TMicRecursionGuard& RecursionGuard, const FMaterialParameterInfo& ParameterInfo, FLinearColor& OutValue, bool bOveriddenOnly, bool bCheckOwnedGlobalOverrides) const
 {
 	return GetVectorParameterValue(ParameterInfo, OutValue, bOveriddenOnly);
 }
 
-bool UMaterial::GetTextureParameterDefaultValue(const FMaterialParameterInfo& ParameterInfo, class UTexture*& OutValue, bool bCheckOwnedGlobalOverrides) const
+bool UMaterial::GetTextureParameterDefaultValue_Concurrent(TMicRecursionGuard& RecursionGuard, const FMaterialParameterInfo& ParameterInfo, class UTexture*& OutValue, bool bCheckOwnedGlobalOverrides) const
 {
 	return GetTextureParameterValue(ParameterInfo, OutValue);
 }
 
-bool UMaterial::GetFontParameterDefaultValue(const FMaterialParameterInfo& ParameterInfo, class UFont*& OutFontValue, int32& OutFontPage, bool bCheckOwnedGlobalOverrides) const
+bool UMaterial::GetFontParameterDefaultValue_Concurrent(TMicRecursionGuard& RecursionGuard, const FMaterialParameterInfo& ParameterInfo, class UFont*& OutFontValue, int32& OutFontPage, bool bCheckOwnedGlobalOverrides) const
 {
 	return GetFontParameterValue(ParameterInfo, OutFontValue, OutFontPage);
 }
 
-bool UMaterial::GetStaticComponentMaskParameterDefaultValue(const FMaterialParameterInfo& ParameterInfo, bool& OutR, bool& OutG, bool& OutB, bool& OutA, FGuid& OutExpressionGuid, bool bCheckOwnedGlobalOverrides) const
+bool UMaterial::GetStaticComponentMaskParameterDefaultValue_Concurrent(TMicRecursionGuard& RecursionGuard, const FMaterialParameterInfo& ParameterInfo, bool& OutR, bool& OutG, bool& OutB, bool& OutA, FGuid& OutExpressionGuid, bool bCheckOwnedGlobalOverrides) const
 {
 	return GetStaticComponentMaskParameterValue(ParameterInfo, OutR, OutG, OutB, OutA, OutExpressionGuid);
 }
 
-bool UMaterial::GetStaticSwitchParameterDefaultValue(const FMaterialParameterInfo& ParameterInfo, bool& OutValue, FGuid& OutExpressionGuid, bool bCheckOwnedGlobalOverrides) const
+bool UMaterial::GetStaticSwitchParameterDefaultValue_Concurrent(TMicRecursionGuard& RecursionGuard, const FMaterialParameterInfo& ParameterInfo, bool& OutValue, FGuid& OutExpressionGuid, bool bCheckOwnedGlobalOverrides) const
 {
 	return GetStaticSwitchParameterValue(ParameterInfo, OutValue, OutExpressionGuid);
 }
diff --git a/Engine/Source/Runtime/Engine/Private/Materials/MaterialInstance.cpp b/Engine/Source/Runtime/Engine/Private/Materials/MaterialInstance.cpp
--- a/Engine/Source/Runtime/Engine/Private/Materials/MaterialInstance.cpp
+++ b/Engine/Source/Runtime/Engine/Private/Materials/MaterialInstance.cpp
@@ -1869,10 +1869,10 @@ void UMaterialInstance::GetDependentFunctions(TArray<UMaterialFunctionInterface*
 	}
 }
 
-bool UMaterialInstance::GetScalarParameterDefaultValue(const FMaterialParameterInfo& ParameterInfo, float& OutValue, bool bOveriddenOnly, bool bCheckOwnedGlobalOverrides) const
+bool UMaterialInstance::GetScalarParameterDefaultValue_Concurrent(TMicRecursionGuard& RecursionGuard, const FMaterialParameterInfo& ParameterInfo, float& OutValue, bool bOveriddenOnly, bool bCheckOwnedGlobalOverrides) const
 {
 #if WITH_EDITOR
-	if (GetReentrantFlag())
+	if(RecursionGuard.Contains(this))
 	{
 		return false;
 	}
@@ -1924,18 +1924,18 @@ bool UMaterialInstance::GetScalarParameterDefaultValue(const FMaterialParameterI
 	if (Parent)
 	{
 #if WITH_EDITOR
-		FMICReentranceGuard	Guard(this);
+		RecursionGuard.Add(this);
 #endif
-		return Parent->GetScalarParameterDefaultValue(ParameterInfo, OutValue, bOveriddenOnly, true);
+		return Parent->GetScalarParameterDefaultValue_Concurrent(RecursionGuard, ParameterInfo, OutValue, bOveriddenOnly, true);
 	}
 
 	return false;
 }
 
-bool UMaterialInstance::GetVectorParameterDefaultValue(const FMaterialParameterInfo& ParameterInfo, FLinearColor& OutValue, bool bOveriddenOnly, bool bCheckOwnedGlobalOverrides) const
+bool UMaterialInstance::GetVectorParameterDefaultValue_Concurrent(TMicRecursionGuard& RecursionGuard, const FMaterialParameterInfo& ParameterInfo, FLinearColor& OutValue, bool bOveriddenOnly, bool bCheckOwnedGlobalOverrides) const
 {
 #if WITH_EDITOR
-	if (GetReentrantFlag())
+	if(RecursionGuard.Contains(this))
 	{
 		return false;
 	}
@@ -1987,18 +1987,18 @@ bool UMaterialInstance::GetVectorParameterDefaultValue(const FMaterialParameterI
 	if (Parent)
 	{
 #if WITH_EDITOR
-		FMICReentranceGuard	Guard(this);
+		RecursionGuard.Add(this);
 #endif
-		return Parent->GetVectorParameterDefaultValue(ParameterInfo, OutValue, bOveriddenOnly, true);
+		return Parent->GetVectorParameterDefaultValue_Concurrent(RecursionGuard, ParameterInfo, OutValue, bOveriddenOnly, true);
 	}
 
 	return false;
 }
 
-bool UMaterialInstance::GetTextureParameterDefaultValue(const FMaterialParameterInfo& ParameterInfo, UTexture*& OutValue, bool bCheckOwnedGlobalOverrides) const
+bool UMaterialInstance::GetTextureParameterDefaultValue_Concurrent(TMicRecursionGuard& RecursionGuard, const FMaterialParameterInfo& ParameterInfo, UTexture*& OutValue, bool bCheckOwnedGlobalOverrides) const
 {
 #if WITH_EDITOR
-	if (GetReentrantFlag())
+	if(RecursionGuard.Contains(this))
 	{
 		return false;
 	}
@@ -2050,18 +2050,18 @@ bool UMaterialInstance::GetTextureParameterDefaultValue(const FMaterialParameter
 	if (Parent)
 	{
 #if WITH_EDITOR
-		FMICReentranceGuard	Guard(this);
+		RecursionGuard.Add(this);
 #endif
-		return Parent->GetTextureParameterDefaultValue(ParameterInfo, OutValue, true);
+		return Parent->GetTextureParameterDefaultValue_Concurrent(RecursionGuard, ParameterInfo, OutValue, true);
 	}
 
 	return false;
 }
 
-bool UMaterialInstance::GetFontParameterDefaultValue(const FMaterialParameterInfo& ParameterInfo, UFont*& OutFontValue, int32& OutFontPage, bool bCheckOwnedGlobalOverrides) const
+bool UMaterialInstance::GetFontParameterDefaultValue_Concurrent(TMicRecursionGuard& RecursionGuard, const FMaterialParameterInfo& ParameterInfo, UFont*& OutFontValue, int32& OutFontPage, bool bCheckOwnedGlobalOverrides) const
 {
 #if WITH_EDITOR
-	if (GetReentrantFlag())
+	if(RecursionGuard.Contains(this))
 	{
 		return false;
 	}
@@ -2114,17 +2114,17 @@ bool UMaterialInstance::GetFontParameterDefaultValue(const FMaterialParameterInf
 	if (Parent)
 	{
 #if WITH_EDITOR
-		FMICReentranceGuard	Guard(this);
+		RecursionGuard.Add(this);
 #endif
-		return Parent->GetFontParameterDefaultValue(ParameterInfo, OutFontValue, OutFontPage, true);
+		return Parent->GetFontParameterDefaultValue_Concurrent(RecursionGuard, ParameterInfo, OutFontValue, OutFontPage, true);
 	}
 
 	return false;
 }
 
-bool UMaterialInstance::GetStaticSwitchParameterDefaultValue(const FMaterialParameterInfo& ParameterInfo, bool& OutValue, FGuid& OutExpressionGuid, bool bCheckOwnedGlobalOverrides) const
+bool UMaterialInstance::GetStaticSwitchParameterDefaultValue_Concurrent(TMicRecursionGuard& RecursionGuard, const FMaterialParameterInfo& ParameterInfo, bool& OutValue, FGuid& OutExpressionGuid, bool bCheckOwnedGlobalOverrides) const
 {
-	if (GetReentrantFlag())
+	if(RecursionGuard.Contains(this))
 	{
 		return false;
 	}
@@ -2175,16 +2175,16 @@ bool UMaterialInstance::GetStaticSwitchParameterDefaultValue(const FMaterialPara
 	
 	if (Parent)
 	{
-		FMICReentranceGuard	Guard(this);
-		return Parent->GetStaticSwitchParameterDefaultValue(ParameterInfo, OutValue, OutExpressionGuid, true);
+		RecursionGuard.Add(this);
+		return Parent->GetStaticSwitchParameterDefaultValue_Concurrent(RecursionGuard, ParameterInfo, OutValue, OutExpressionGuid, true);
 	}
 
 	return false;
 }
 
-bool UMaterialInstance:: GetStaticComponentMaskParameterDefaultValue(const FMaterialParameterInfo& ParameterInfo, bool& OutR, bool& OutG, bool& OutB, bool& OutA, FGuid& OutExpressionGuid, bool bCheckOwnedGlobalOverrides) const
+bool UMaterialInstance:: GetStaticComponentMaskParameterDefaultValue_Concurrent(TMicRecursionGuard& RecursionGuard, const FMaterialParameterInfo& ParameterInfo, bool& OutR, bool& OutG, bool& OutB, bool& OutA, FGuid& OutExpressionGuid, bool bCheckOwnedGlobalOverrides) const
 {
-	if (GetReentrantFlag())
+	if(RecursionGuard.Contains(this))
 	{
 		return false;
 	}
@@ -2238,8 +2238,8 @@ bool UMaterialInstance:: GetStaticComponentMaskParameterDefaultValue(const FMate
 	
 	if (Parent)
 	{
-		FMICReentranceGuard	Guard(this);
-		return Parent->GetStaticComponentMaskParameterDefaultValue(ParameterInfo, OutR, OutG, OutB, OutA, OutExpressionGuid, true);
+		RecursionGuard.Add(this);
+		return Parent->GetStaticComponentMaskParameterDefaultValue_Concurrent(RecursionGuard, ParameterInfo, OutR, OutG, OutB, OutA, OutExpressionGuid, true);
 	}
 
 	return false;
diff --git a/Engine/Source/Runtime/Engine/Private/Materials/MaterialInterface.cpp b/Engine/Source/Runtime/Engine/Private/Materials/MaterialInterface.cpp
--- a/Engine/Source/Runtime/Engine/Private/Materials/MaterialInterface.cpp
+++ b/Engine/Source/Runtime/Engine/Private/Materials/MaterialInterface.cpp
@@ -96,6 +96,42 @@ void UMaterialInterface::PostLoad()
 #endif
 }
 
+bool UMaterialInterface::GetScalarParameterDefaultValue(const FMaterialParameterInfo& ParameterInfo, float& OutValue, bool bOveriddenOnly, bool bCheckOwnedGlobalOverrides) const
+{
+	TMicRecursionGuard RecursionGuard;
+	return GetScalarParameterDefaultValue_Concurrent(RecursionGuard, ParameterInfo, OutValue, bOveriddenOnly, bCheckOwnedGlobalOverrides);
+}
+
+bool UMaterialInterface::GetVectorParameterDefaultValue(const FMaterialParameterInfo& ParameterInfo, FLinearColor& OutValue, bool bOveriddenOnly, bool bCheckOwnedGlobalOverrides) const
+{
+	TMicRecursionGuard RecursionGuard;
+	return GetVectorParameterDefaultValue_Concurrent(RecursionGuard, ParameterInfo, OutValue, bOveriddenOnly, bCheckOwnedGlobalOverrides);
+}
+
+bool UMaterialInterface::GetTextureParameterDefaultValue(const FMaterialParameterInfo& ParameterInfo, class UTexture*& OutValue, bool bCheckOwnedGlobalOverrides) const
+{
+	TMicRecursionGuard RecursionGuard;
+	return GetTextureParameterDefaultValue_Concurrent(RecursionGuard, ParameterInfo, OutValue, bCheckOwnedGlobalOverrides);
+}
+
+bool UMaterialInterface::GetFontParameterDefaultValue(const FMaterialParameterInfo& ParameterInfo, class UFont*& OutFontValue, int32& OutFontPage, bool bCheckOwnedGlobalOverrides) const
+{
+	TMicRecursionGuard RecursionGuard;
+	return GetFontParameterDefaultValue_Concurrent(RecursionGuard, ParameterInfo, OutFontValue, OutFontPage, bCheckOwnedGlobalOverrides);
+}
+
+bool UMaterialInterface::GetStaticSwitchParameterDefaultValue(const FMaterialParameterInfo& ParameterInfo, bool& OutValue, FGuid& OutExpressionGuid, bool bCheckOwnedGlobalOverrides) const
+{
+	TMicRecursionGuard RecursionGuard;
+	return GetStaticSwitchParameterDefaultValue_Concurrent(RecursionGuard, ParameterInfo, OutValue, OutExpressionGuid, bCheckOwnedGlobalOverrides);
+}
+
+bool UMaterialInterface::GetStaticComponentMaskParameterDefaultValue(const FMaterialParameterInfo& ParameterInfo, bool& OutR, bool& OutG, bool& OutB, bool& OutA, FGuid& OutExpressionGuid, bool bCheckOwnedGlobalOverrides) const
+{
+	TMicRecursionGuard RecursionGuard;
+	return GetStaticComponentMaskParameterDefaultValue_Concurrent(RecursionGuard, ParameterInfo, OutR, OutG, OutB, OutA, OutExpressionGuid, bCheckOwnedGlobalOverrides);
+}
+
 void UMaterialInterface::GetUsedTexturesAndIndices(TArray<UTexture*>& OutTextures, TArray< TArray<int32> >& OutIndices, EMaterialQualityLevel::Type QualityLevel, ERHIFeatureLevel::Type FeatureLevel) const
 {
 	GetUsedTextures(OutTextures, QualityLevel, false, FeatureLevel, false);
