Fix compile error from backout

diff --git a/Engine/Source/Runtime/Engine/Private/Materials/MaterialExpressions.cpp b/Engine/Source/Runtime/Engine/Private/Materials/MaterialExpressions.cpp
--- a/Engine/Source/Runtime/Engine/Private/Materials/MaterialExpressions.cpp
+++ b/Engine/Source/Runtime/Engine/Private/Materials/MaterialExpressions.cpp
@@ -10786,7 +10786,7 @@ bool UMaterialFunctionInstance::HasFlippedCoordinates() const
 }
 #endif
 
-bool UMaterialFunctionInstance::OverrideNamedScalarParameter(const FMaterialParameterInfo& ParameterInfo, float& OutValue) const
+bool UMaterialFunctionInstance::OverrideNamedScalarParameter(const FMaterialParameterInfo& ParameterInfo, float& OutValue)
 {
 	for (const FScalarParameterValue& ScalarParameter : ScalarParameterValues)
 	{
@@ -10800,7 +10800,7 @@ bool UMaterialFunctionInstance::OverrideNamedScalarParameter(const FMaterialPara
 	return false;
 }
 
-bool UMaterialFunctionInstance::OverrideNamedVectorParameter(const FMaterialParameterInfo& ParameterInfo, FLinearColor& OutValue) const
+bool UMaterialFunctionInstance::OverrideNamedVectorParameter(const FMaterialParameterInfo& ParameterInfo, FLinearColor& OutValue)
 {
 	for (const FVectorParameterValue& VectorParameter : VectorParameterValues)
 	{
@@ -10814,7 +10814,7 @@ bool UMaterialFunctionInstance::OverrideNamedVectorParameter(const FMaterialPara
 	return false;
 }
 
-bool UMaterialFunctionInstance::OverrideNamedTextureParameter(const FMaterialParameterInfo& ParameterInfo, UTexture*& OutValue) const
+bool UMaterialFunctionInstance::OverrideNamedTextureParameter(const FMaterialParameterInfo& ParameterInfo, UTexture*& OutValue)
 {
 	for (const FTextureParameterValue& TextureParameter : TextureParameterValues)
 	{
@@ -10828,7 +10828,7 @@ bool UMaterialFunctionInstance::OverrideNamedTextureParameter(const FMaterialPar
 	return false;
 }
 
-bool UMaterialFunctionInstance::OverrideNamedFontParameter(const FMaterialParameterInfo& ParameterInfo, UFont*& OutFontValue, int32& OutFontPage) const
+bool UMaterialFunctionInstance::OverrideNamedFontParameter(const FMaterialParameterInfo& ParameterInfo, UFont*& OutFontValue, int32& OutFontPage)
 {
 	for (const FFontParameterValue& FontParameter : FontParameterValues)
 	{
@@ -10843,7 +10843,7 @@ bool UMaterialFunctionInstance::OverrideNamedFontParameter(const FMaterialParame
 	return false;
 }
 
-bool UMaterialFunctionInstance::OverrideNamedStaticSwitchParameter(const FMaterialParameterInfo& ParameterInfo, bool& OutValue, FGuid& OutExpressionGuid) const
+bool UMaterialFunctionInstance::OverrideNamedStaticSwitchParameter(const FMaterialParameterInfo& ParameterInfo, bool& OutValue, FGuid& OutExpressionGuid)
 {
 	for (const FStaticSwitchParameter& StaticSwitchParameter : StaticSwitchParameterValues)
 	{
@@ -10858,7 +10858,7 @@ bool UMaterialFunctionInstance::OverrideNamedStaticSwitchParameter(const FMateri
 	return false;
 }
 
-bool UMaterialFunctionInstance::OverrideNamedStaticComponentMaskParameter(const FMaterialParameterInfo& ParameterInfo, bool& OutR, bool& OutG, bool& OutB, bool& OutA, FGuid& OutExpressionGuid) const
+bool UMaterialFunctionInstance::OverrideNamedStaticComponentMaskParameter(const FMaterialParameterInfo& ParameterInfo, bool& OutR, bool& OutG, bool& OutB, bool& OutA, FGuid& OutExpressionGuid)
 {
 	for (const FStaticComponentMaskParameter& StaticComponentMaskParameter : StaticComponentMaskParameterValues)
 	{
