Fixed error when encoding shadowmaps with allocations that are not multiple of 4 .

diff --git a/Engine/Source/Runtime/Engine/Private/ShadowMap.cpp b/Engine/Source/Runtime/Engine/Private/ShadowMap.cpp
--- a/Engine/Source/Runtime/Engine/Private/ShadowMap.cpp
+++ b/Engine/Source/Runtime/Engine/Private/ShadowMap.cpp
@@ -751,6 +751,10 @@ void FShadowMap2D::EncodeTextures(UWorld* InWorld, ULevel* LightingScenario, boo
 				MaxHeight = FMath::Max(MaxHeight, Allocation->MappedRect.Height());
 			}
 
+			// Assume bAlignByFour (see FTextureLayout::AddElement)
+			MaxWidth = (MaxWidth + 3) & ~3;
+			MaxHeight = (MaxHeight + 3) & ~3;
+
 			FShadowMapPendingTexture* Texture = nullptr;
 
 			// Find an existing texture which the shadow-map can be stored in.
