DR - Warning fix
#rnx

diff --git a/Engine/Source/Runtime/Engine/Private/Particles/ParticleGpuSimulation.cpp b/Engine/Source/Runtime/Engine/Private/Particles/ParticleGpuSimulation.cpp
--- a/Engine/Source/Runtime/Engine/Private/Particles/ParticleGpuSimulation.cpp
+++ b/Engine/Source/Runtime/Engine/Private/Particles/ParticleGpuSimulation.cpp
@@ -2970,16 +2970,16 @@ public:
 	/** Constructor. */
 	explicit FGPUSpriteDynamicEmitterData( const UParticleModuleRequired* InRequiredModule )
 		: FDynamicEmitterDataBase( InRequiredModule )
+		, MaterialProxy(nullptr)
 		, bIsMaterialTranslucent(true)
-		, FXSystem(NULL)
-		, Resources(NULL)
-		, Simulation(NULL)
+		, FXSystem(nullptr)
+		, Resources(nullptr)
+		, Simulation(nullptr)
 		, SortMode(PSORTMODE_None)
 		, bLocalVectorFieldTileX(false)
 		, bLocalVectorFieldTileY(false)
 		, bLocalVectorFieldTileZ(false)
 		, bLocalVectorFieldUseFixDT(false)
-		, MaterialProxy(nullptr)
 	{
 		GetNewParticleArray(NewParticles);
 	}
