Don't modify FMaterialUniformExpressionComponentSwizzle members to replicate scalar when constructed, instead perform the same logic while expression is being evaluated.
Otherwise may end up with multiple FMaterialUniformExpressionComponentSwizzle used in different contexts, with different result types that compare equal to each other,
which breaks assumptions made by other material compiler code

diff --git a/Engine/Source/Runtime/Engine/Private/Materials/MaterialUniformExpressions.h b/Engine/Source/Runtime/Engine/Private/Materials/MaterialUniformExpressions.h
--- a/Engine/Source/Runtime/Engine/Private/Materials/MaterialUniformExpressions.h
+++ b/Engine/Source/Runtime/Engine/Private/Materials/MaterialUniformExpressions.h
@@ -1206,15 +1206,6 @@ public:
 		// At least one proper index
 		check(InR >= 0 && InR <= 3);
 		++NumElements;
-
-		if (NumElements == 1)
-		{
-			// Replicate scalar
-			IndexG = IndexR;
-			IndexB = IndexR;
-			IndexA = IndexR;
-			NumElements = 4;
-		}
 	}
 
 	// FMaterialUniformExpression interface.
@@ -1235,6 +1226,10 @@ public:
 		OutValue *= 0;
 		switch (NumElements)
 		{
+		case 1:
+			// Replicate scalar
+			OutValue.R = OutValue.G = OutValue.B = OutValue.A = Temp.Component(IndexR);
+			break;
 		case 4:
 			OutValue.A = Temp.Component(IndexA);
 			// Fallthrough...
@@ -1243,8 +1238,6 @@ public:
 			// Fallthrough...
 		case 2:
 			OutValue.G = Temp.Component(IndexG);
-			// Fallthrough...
-		case 1:
 			OutValue.R = Temp.Component(IndexR);
 			break;
 		default: UE_LOG(LogMaterial, Fatal, TEXT("Invalid number of swizzle elements: %d"), NumElements);
