Replaced FScene::SetShaderMapsOnMaterialResources which updated uniform expression caches for materials with current scene's feature level, with a function SetShaderMapsOnMaterialResources which updates those caches for all feature level materials. This fixes crashes when accessing uniform expression cache after certain combinations of shader compilation and feature level switching. One example is:
1. Switch preview level from SM5 to SM4.
2. SM4 materials will be compiled.
3. After compilation FScene::SetShaderMapsOnMaterialResources will be called.
4. FScene has still SM5 feature level, so it won't update uniform expression cache for newly compiled materials.
5. Finally scene is recreated with SM4 feature level and engine crashes on null uniform expression access during caching of mesh draw commands.

diff --git a/Engine/Source/Runtime/Engine/Private/Materials/MaterialShared.cpp b/Engine/Source/Runtime/Engine/Private/Materials/MaterialShared.cpp
--- a/Engine/Source/Runtime/Engine/Private/Materials/MaterialShared.cpp
+++ b/Engine/Source/Runtime/Engine/Private/Materials/MaterialShared.cpp
@@ -3368,3 +3368,73 @@ void FMaterialResourceProxyReader::Initialize(
 		OffsetToEnd = OffsetToFirstResource + NumBytes;
 	}
 }
+
+typedef TMap<FMaterial*, FMaterialShaderMap*> FMaterialsToUpdateMap;
+
+void SetShaderMapsOnMaterialResources_RenderThread(FRHICommandListImmediate& RHICmdList, const FMaterialsToUpdateMap& MaterialsToUpdate)
+{
+	SCOPE_CYCLE_COUNTER(STAT_Scene_SetShaderMapsOnMaterialResources_RT);
+
+	TArray<const FMaterial*> MaterialArray;
+	bool bUpdateFeatureLevel[ERHIFeatureLevel::Num] = { false };
+
+	for (FMaterialsToUpdateMap::TConstIterator It(MaterialsToUpdate); It; ++It)
+	{
+		FMaterial* Material = It.Key();
+		FMaterialShaderMap* ShaderMap = It.Value();
+		Material->SetRenderingThreadShaderMap(ShaderMap);
+		check(!ShaderMap || ShaderMap->IsValidForRendering());
+		MaterialArray.Add(Material);
+		bUpdateFeatureLevel[Material->GetFeatureLevel()] = true;
+	}
+
+	bool bFoundAnyInitializedMaterials = false;
+
+	// Iterate through all loaded material render proxies and recache their uniform expressions if needed
+	// This search does not scale well, but is only used when uploading async shader compile results
+	for (int32 FeatureLevelIndex = 0; FeatureLevelIndex < ARRAY_COUNT(bUpdateFeatureLevel); ++FeatureLevelIndex)
+	{
+		if (bUpdateFeatureLevel[FeatureLevelIndex])
+		{
+			const ERHIFeatureLevel::Type MaterialFeatureLevel = (ERHIFeatureLevel::Type) FeatureLevelIndex;
+
+			for (TSet<FMaterialRenderProxy*>::TConstIterator It(FMaterialRenderProxy::GetMaterialRenderProxyMap()); It; ++It)
+			{
+				FMaterialRenderProxy* MaterialProxy = *It;
+				FMaterial* Material = MaterialProxy->GetMaterialNoFallback(MaterialFeatureLevel);
+
+				if (Material && MaterialsToUpdate.Contains(Material))
+				{
+					// Materials used as async fallbacks can't be updated through this mechanism and should have been updated synchronously earlier
+					check(!Material->RequiresSynchronousCompilation());
+					MaterialProxy->CacheUniformExpressions(true);
+					bFoundAnyInitializedMaterials = true;
+
+					const FMaterial& MaterialForRendering = *MaterialProxy->GetMaterial(MaterialFeatureLevel);
+					check(MaterialForRendering.GetRenderingThreadShaderMap());
+
+					check(!MaterialProxy->UniformExpressionCache[MaterialFeatureLevel].bUpToDate
+						|| MaterialProxy->UniformExpressionCache[MaterialFeatureLevel].CachedUniformExpressionShaderMap == MaterialForRendering.GetRenderingThreadShaderMap());
+
+					check(MaterialForRendering.GetRenderingThreadShaderMap()->IsValidForRendering());
+				}
+			}
+		}
+	}
+}
+
+void SetShaderMapsOnMaterialResources(const TMap<FMaterial*, FMaterialShaderMap*>& MaterialsToUpdate)
+{
+	for (TMap<FMaterial*, FMaterialShaderMap*>::TConstIterator It(MaterialsToUpdate); It; ++It)
+	{
+		FMaterial* Material = It.Key();
+		check(!Material->RequiresSynchronousCompilation());
+	}
+
+	ENQUEUE_UNIQUE_RENDER_COMMAND_ONEPARAMETER(
+		FSetShaderMapOnMaterialResources,
+		FMaterialsToUpdateMap, MaterialsToUpdate, MaterialsToUpdate,
+		{
+			SetShaderMapsOnMaterialResources_RenderThread(RHICmdList, MaterialsToUpdate);
+		});
+}
\ No newline at end of file
diff --git a/Engine/Source/Runtime/Engine/Private/ShaderCompiler/ShaderCompiler.cpp b/Engine/Source/Runtime/Engine/Private/ShaderCompiler/ShaderCompiler.cpp
--- a/Engine/Source/Runtime/Engine/Private/ShaderCompiler/ShaderCompiler.cpp
+++ b/Engine/Source/Runtime/Engine/Private/ShaderCompiler/ShaderCompiler.cpp
@@ -2105,10 +2105,7 @@ void FShaderCompilingManager::ProcessCompiledShaderMaps(
 
 		const TSet<FSceneInterface*>& AllocatedScenes = GetRendererModule().GetAllocatedScenes();
 
-		for (TSet<FSceneInterface*>::TConstIterator SceneIt(AllocatedScenes); SceneIt; ++SceneIt)
-		{
-			(*SceneIt)->SetShaderMapsOnMaterialResources(MaterialsToApplyToScene);
-		}
+		SetShaderMapsOnMaterialResources(MaterialsToApplyToScene);
 
 #if WITH_EDITOR
 		for (TMap<FMaterial*, FMaterialShaderMap*>::TIterator It(MaterialsToUpdate); It; ++It)
diff --git a/Engine/Source/Runtime/Engine/Public/MaterialShared.h b/Engine/Source/Runtime/Engine/Public/MaterialShared.h
--- a/Engine/Source/Runtime/Engine/Public/MaterialShared.h
+++ b/Engine/Source/Runtime/Engine/Public/MaterialShared.h
@@ -2378,6 +2378,9 @@ private:
 		bool bSeekToEnd = false);
 };
 
+/** Sets shader maps on the specified materials without blocking. */
+extern ENGINE_API void SetShaderMapsOnMaterialResources(const TMap<FMaterial*, FMaterialShaderMap*>& MaterialsToUpdate);
+
 #if STORE_ONLY_ACTIVE_SHADERMAPS
 bool HasMaterialResource(
 	UMaterial* Material,
diff --git a/Engine/Source/Runtime/Engine/Public/SceneInterface.h b/Engine/Source/Runtime/Engine/Public/SceneInterface.h
--- a/Engine/Source/Runtime/Engine/Public/SceneInterface.h
+++ b/Engine/Source/Runtime/Engine/Public/SceneInterface.h
@@ -186,9 +186,6 @@ public:
 	/** Sets the precomputed volume distance field for the scene, or NULL to clear the current one. */
 	virtual void SetPrecomputedVolumeDistanceField(const class FPrecomputedVolumeDistanceField* PrecomputedVolumeDistanceField) {}
 
-	/** Sets shader maps on the specified materials without blocking. */
-	virtual void SetShaderMapsOnMaterialResources(const TMap<FMaterial*, FMaterialShaderMap*>& MaterialsToUpdate) {}
-
 	/** Updates static draw lists for the given set of materials. */
 	virtual void UpdateStaticDrawListsForMaterials(const TArray<const FMaterial*>& Materials) {}
 
diff --git a/Engine/Source/Runtime/Renderer/Private/RendererScene.cpp b/Engine/Source/Runtime/Renderer/Private/RendererScene.cpp
--- a/Engine/Source/Runtime/Renderer/Private/RendererScene.cpp
+++ b/Engine/Source/Runtime/Renderer/Private/RendererScene.cpp
@@ -2738,66 +2738,6 @@ void FScene::SetPrecomputedVisibility(const FPrecomputedVisibilityHandler* NewPr
 	});
 }
 
-void FScene::SetShaderMapsOnMaterialResources_RenderThread(FRHICommandListImmediate& RHICmdList, const FMaterialsToUpdateMap& MaterialsToUpdate)
-{
-	SCOPE_CYCLE_COUNTER(STAT_Scene_SetShaderMapsOnMaterialResources_RT);
-
-	TArray<const FMaterial*> MaterialArray;
-
-	for (FMaterialsToUpdateMap::TConstIterator It(MaterialsToUpdate); It; ++It)
-	{
-		FMaterial* Material = It.Key();
-		FMaterialShaderMap* ShaderMap = It.Value();
-		Material->SetRenderingThreadShaderMap(ShaderMap);
-		check(!ShaderMap || ShaderMap->IsValidForRendering());
-		MaterialArray.Add(Material);
-	}
-
-	const auto SceneFeatureLevel = GetFeatureLevel();
-	bool bFoundAnyInitializedMaterials = false;
-
-	// Iterate through all loaded material render proxies and recache their uniform expressions if needed
-	// This search does not scale well, but is only used when uploading async shader compile results
-	for (TSet<FMaterialRenderProxy*>::TConstIterator It(FMaterialRenderProxy::GetMaterialRenderProxyMap()); It; ++It)
-	{
-		FMaterialRenderProxy* MaterialProxy = *It;
-		FMaterial* Material = MaterialProxy->GetMaterialNoFallback(SceneFeatureLevel);
-
-		if (Material && MaterialsToUpdate.Contains(Material))
-		{
-			// Materials used as async fallbacks can't be updated through this mechanism and should have been updated synchronously earlier
-			check(!Material->RequiresSynchronousCompilation());
-			MaterialProxy->CacheUniformExpressions(true);
-			bFoundAnyInitializedMaterials = true;
-
-			const FMaterial& MaterialForRendering = *MaterialProxy->GetMaterial(SceneFeatureLevel);
-			check(MaterialForRendering.GetRenderingThreadShaderMap());
-
-			check(!MaterialProxy->UniformExpressionCache[SceneFeatureLevel].bUpToDate
-				|| MaterialProxy->UniformExpressionCache[SceneFeatureLevel].CachedUniformExpressionShaderMap == MaterialForRendering.GetRenderingThreadShaderMap());
-
-			check(MaterialForRendering.GetRenderingThreadShaderMap()->IsValidForRendering());
-		}
-	}
-}
-
-void FScene::SetShaderMapsOnMaterialResources(const TMap<FMaterial*, class FMaterialShaderMap*>& MaterialsToUpdate)
-{
-	for (TMap<FMaterial*, FMaterialShaderMap*>::TConstIterator It(MaterialsToUpdate); It; ++It)
-	{
-		FMaterial* Material = It.Key();
-		check(!Material->RequiresSynchronousCompilation());
-	}
-
-	ENQUEUE_UNIQUE_RENDER_COMMAND_TWOPARAMETER(
-		FSetShaderMapOnMaterialResources,
-		FScene*,Scene,this,
-		FMaterialsToUpdateMap,MaterialsToUpdate,MaterialsToUpdate,
-	{
-		Scene->SetShaderMapsOnMaterialResources_RenderThread(RHICmdList, MaterialsToUpdate);
-	});
-}
-
 void FScene::UpdateStaticDrawListsForMaterials_RenderThread(FRHICommandListImmediate& RHICmdList, const TArray<const FMaterial*>& Materials)
 {
 	SCOPE_CYCLE_COUNTER(STAT_Scene_UpdateStaticDrawListsForMaterials_RT);
diff --git a/Engine/Source/Runtime/Renderer/Private/ScenePrivate.h b/Engine/Source/Runtime/Renderer/Private/ScenePrivate.h
--- a/Engine/Source/Runtime/Renderer/Private/ScenePrivate.h
+++ b/Engine/Source/Runtime/Renderer/Private/ScenePrivate.h
@@ -2251,8 +2251,6 @@ private:
 	void HideNodeChildren(FSceneViewState* ViewState, FLODSceneNode& Node);
 };
 
-typedef TMap<FMaterial*, FMaterialShaderMap*> FMaterialsToUpdateMap;
-
 class FCachedShadowMapData
 {
 public:
@@ -2691,9 +2689,6 @@ public:
 	/** Sets the precomputed visibility handler for the scene, or NULL to clear the current one. */
 	virtual void SetPrecomputedVisibility(const FPrecomputedVisibilityHandler* InPrecomputedVisibilityHandler) override;
 
-	/** Sets shader maps on the specified materials without blocking. */
-	virtual void SetShaderMapsOnMaterialResources(const TMap<FMaterial*, class FMaterialShaderMap*>& MaterialsToUpdate) override;
-
 	/** Updates static draw lists for the given set of materials. */
 	virtual void UpdateStaticDrawListsForMaterials(const TArray<const FMaterial*>& Materials) override;
 
@@ -2897,9 +2892,6 @@ private:
 	/** Updates the contents of all reflection captures in the scene.  Must be called from the game thread. */
 	void UpdateAllReflectionCaptures(const TCHAR* CaptureReason, bool bVerifyOnlyCapturing);
 
-	/** Sets shader maps on the specified materials without blocking. */
-	void SetShaderMapsOnMaterialResources_RenderThread(FRHICommandListImmediate& RHICmdList, const FMaterialsToUpdateMap& MaterialsToUpdate);
-
 	/** Updates static draw lists for the given materials. */
 	void UpdateStaticDrawListsForMaterials_RenderThread(FRHICommandListImmediate& RHICmdList, const TArray<const FMaterial*>& Materials);
 
