Fixed ForniteClient compilation warning.

diff --git a/Engine/Source/Runtime/Engine/Private/Materials/MaterialInstance.cpp b/Engine/Source/Runtime/Engine/Private/Materials/MaterialInstance.cpp
--- a/Engine/Source/Runtime/Engine/Private/Materials/MaterialInstance.cpp
+++ b/Engine/Source/Runtime/Engine/Private/Materials/MaterialInstance.cpp
@@ -2459,19 +2459,19 @@ void UMaterialInstance::UpdatePermutationAllocations()
 		{
 			for (int32 Quality = 0; Quality < EMaterialQualityLevel::Num; ++Quality)
 			{
-				FMaterialResource*& Resource = StaticPermutationMaterialResources[Quality][Feature];
+				FMaterialResource*& StaticPermResource = StaticPermutationMaterialResources[Quality][Feature];
 				if (Feature != ActiveFeatureLevel || Quality != ActiveQualityLevel)
 				{
-					delete Resource;
-					Resource = nullptr;
+					delete StaticPermResource;
+					StaticPermResource = nullptr;
 				}
 				else
 				{
-					if (!Resource)
+					if (!StaticPermResource)
 					{
-						Resource = AllocatePermutationResource();
+						StaticPermResource = AllocatePermutationResource();
 					}
-					Resource->SetMaterial(BaseMaterial, ActiveQualityLevel, true, ActiveFeatureLevel, this);
+					StaticPermResource->SetMaterial(BaseMaterial, ActiveQualityLevel, true, ActiveFeatureLevel, this);
 				}
 			}
 		}
