Sorting translucent mesh draw commands also by a mesh id in primitive. This allows for stable sorting of translucent mesh draw commands which belong to the same primitive.

diff --git a/Engine/Source/Runtime/Engine/Private/SceneManagement.cpp b/Engine/Source/Runtime/Engine/Private/SceneManagement.cpp
--- a/Engine/Source/Runtime/Engine/Private/SceneManagement.cpp
+++ b/Engine/Source/Runtime/Engine/Private/SceneManagement.cpp
@@ -290,6 +290,9 @@ void FMeshElementCollector::AddMesh(int32 ViewIndex, FMeshBatch& MeshBatch)
 
 	MeshBatch.MaterialRenderProxy->UpdateUniformExpressionCacheIfNeeded(Views[ViewIndex]->GetFeatureLevel());
 
+	MeshBatch.MeshIdInPrimitive = MeshIdInPrimitivePerView[ViewIndex];
+	++MeshIdInPrimitivePerView[ViewIndex];
+
 	TArray<FMeshBatchAndRelevance,SceneRenderingAllocator>& ViewMeshBatches = *MeshBatches[ViewIndex];
 	new (ViewMeshBatches) FMeshBatchAndRelevance(MeshBatch, PrimitiveSceneProxy, FeatureLevel);	
 }
diff --git a/Engine/Source/Runtime/Engine/Public/MeshBatch.h b/Engine/Source/Runtime/Engine/Public/MeshBatch.h
--- a/Engine/Source/Runtime/Engine/Public/MeshBatch.h
+++ b/Engine/Source/Runtime/Engine/Public/MeshBatch.h
@@ -131,6 +131,9 @@ struct FMeshBatch
 {
 	TArray<FMeshBatchElement,TInlineAllocator<1> > Elements;
 
+	/* Mesh Id in a primitive. Used for stable sorting of draws belonging to the same primitive. **/
+	uint16 MeshIdInPrimitive;
+
 	/** LOD index of the mesh, used for fading LOD transitions. */
 	int8 LODIndex;
 
@@ -256,7 +259,8 @@ struct FMeshBatch
 
 	/** Default constructor. */
 	FMeshBatch()
-	:	LODIndex(INDEX_NONE)
+	:	MeshIdInPrimitive(0)
+	,	LODIndex(INDEX_NONE)
 #if !(UE_BUILD_SHIPPING || UE_BUILD_TEST)
 	,	VisualizeLODIndex(INDEX_NONE)
 #endif
diff --git a/Engine/Source/Runtime/Engine/Public/SceneManagement.h b/Engine/Source/Runtime/Engine/Public/SceneManagement.h
--- a/Engine/Source/Runtime/Engine/Public/SceneManagement.h
+++ b/Engine/Source/Runtime/Engine/Public/SceneManagement.h
@@ -1809,6 +1809,7 @@ public:
 private:
 
 	FHitProxyId HitProxyId;
+	uint16 PrimitiveMeshId;
 
 	bool bIsMobileHDR;
 
@@ -1955,6 +1956,12 @@ private:
 		for (int32 ViewIndex = 0; ViewIndex < SimpleElementCollectors.Num(); ViewIndex++)
 		{
 			SimpleElementCollectors[ViewIndex]->HitProxyId = DefaultHitProxyId;
+			SimpleElementCollectors[ViewIndex]->PrimitiveMeshId = 0;
+		}
+
+		for (int32 ViewIndex = 0; ViewIndex < MeshIdInPrimitivePerView.Num(); ++ViewIndex)
+		{
+			MeshIdInPrimitivePerView[ViewIndex] = 0;
 		}
 	}
 
@@ -1963,6 +1970,7 @@ private:
 		Views.Empty();
 		MeshBatches.Empty();
 		SimpleElementCollectors.Empty();
+		MeshIdInPrimitivePerView.Empty();
 		DynamicPrimitiveShaderData = nullptr;
 	}
 
@@ -1974,6 +1982,7 @@ private:
 		ERHIFeatureLevel::Type InFeatureLevel)
 	{
 		Views.Add(InView);
+		MeshIdInPrimitivePerView.AddZeroed(Views.Num());
 		MeshBatches.Add(ViewMeshes);
 		SimpleElementCollectors.Add(ViewSimpleElementCollector);
 		DynamicPrimitiveShaderData = InDynamicPrimitiveShaderData;
@@ -1994,6 +2003,9 @@ private:
 	/** Views being collected for */
 	TArray<FSceneView*, TInlineAllocator<2> > Views;
 
+	/** Current Mesh Id In Primitive per view */
+	TArray<uint16, TInlineAllocator<2> > MeshIdInPrimitivePerView;
+
 	/** Material proxies that will be deleted at the end of the frame. */
 	TArray<FMaterialRenderProxy*, SceneRenderingAllocator> TemporaryProxies;
 
diff --git a/Engine/Source/Runtime/Renderer/Private/BasePassRendering.cpp b/Engine/Source/Runtime/Renderer/Private/BasePassRendering.cpp
--- a/Engine/Source/Runtime/Renderer/Private/BasePassRendering.cpp
+++ b/Engine/Source/Runtime/Renderer/Private/BasePassRendering.cpp
@@ -990,7 +990,7 @@ void FBasePassMeshProcessor::Process(
 
 	if (bTranslucentBasePass)
 	{
-		SortKey = CalculateStaticTranslucentMeshSortKey(PrimitiveSceneProxy);
+		SortKey = CalculateStaticTranslucentMeshSortKey(PrimitiveSceneProxy, MeshBatch.MeshIdInPrimitive);
 	}
 	else
 	{
diff --git a/Engine/Source/Runtime/Renderer/Private/MobileBasePass.cpp b/Engine/Source/Runtime/Renderer/Private/MobileBasePass.cpp
--- a/Engine/Source/Runtime/Renderer/Private/MobileBasePass.cpp
+++ b/Engine/Source/Runtime/Renderer/Private/MobileBasePass.cpp
@@ -732,7 +732,7 @@ void FMobileBasePassMeshProcessor::Process(
 	FMeshDrawCommandSortKey SortKey; 
 	if (bTranslucentBasePass)
 	{
-		SortKey = CalculateStaticTranslucentMeshSortKey(PrimitiveSceneProxy);
+		SortKey = CalculateStaticTranslucentMeshSortKey(PrimitiveSceneProxy, MeshBatch.MeshIdInPrimitive);
 	}
 	else
 	{
diff --git a/Engine/Source/Runtime/Renderer/Private/MobileTranslucentRendering.cpp b/Engine/Source/Runtime/Renderer/Private/MobileTranslucentRendering.cpp
--- a/Engine/Source/Runtime/Renderer/Private/MobileTranslucentRendering.cpp
+++ b/Engine/Source/Runtime/Renderer/Private/MobileTranslucentRendering.cpp
@@ -865,7 +865,7 @@ private:
 		FMeshMaterialShaderElementData ShaderElementData;
 		ShaderElementData.InitializeMeshMaterialData(ViewIfDynamicMeshCommand, PrimitiveSceneProxy, MeshBatch, StaticMeshId, false);
 
-		FMeshDrawCommandSortKey SortKey = CalculateStaticTranslucentMeshSortKey(PrimitiveSceneProxy);
+		FMeshDrawCommandSortKey SortKey = CalculateStaticTranslucentMeshSortKey(PrimitiveSceneProxy, MeshBatch.MeshIdInPrimitive);
 		
 		BuildMeshDrawCommands(
 			MeshBatch,
diff --git a/Engine/Source/Runtime/Renderer/Private/TranslucentRendering.cpp b/Engine/Source/Runtime/Renderer/Private/TranslucentRendering.cpp
--- a/Engine/Source/Runtime/Renderer/Private/TranslucentRendering.cpp
+++ b/Engine/Source/Runtime/Renderer/Private/TranslucentRendering.cpp
@@ -289,18 +289,19 @@ float CalculateTranslucentSortKey(const FPrimitiveSceneInfo* RESTRICT PrimitiveS
 	return SortKey;
 }
 
-FMeshDrawCommandSortKey CalculateStaticTranslucentMeshSortKey(const FPrimitiveSceneProxy* RESTRICT PrimitiveSceneProxy)
+FMeshDrawCommandSortKey CalculateStaticTranslucentMeshSortKey(const FPrimitiveSceneProxy* RESTRICT PrimitiveSceneProxy, uint16 MeshIdInPrimitive)
 {
-	uint32 UnsignedSortPriority = 0;
+	uint16 SortKeyPriority = 0;
 
 	if (PrimitiveSceneProxy)
 	{
 		const FPrimitiveSceneInfo* PrimitiveSceneInfo = PrimitiveSceneProxy->GetPrimitiveSceneInfo();
-		UnsignedSortPriority = (int32)PrimitiveSceneInfo->Proxy->GetTranslucencySortPriority() - (int32)SHRT_MIN;
+		SortKeyPriority = (uint16)((int32)PrimitiveSceneInfo->Proxy->GetTranslucencySortPriority() - (int32)SHRT_MIN);
 	}
 
 	UTranslucentMeshSortKey TranslucentMeshSortKey;
-	TranslucentMeshSortKey.Fields.SortPriority = UnsignedSortPriority;
+	TranslucentMeshSortKey.Fields.MeshIdInPrimitive = MeshIdInPrimitive;
+	TranslucentMeshSortKey.Fields.Priority = SortKeyPriority;
 	TranslucentMeshSortKey.Fields.Distance = 0; // View specific, so will be filled later inside VisibleMeshCommands.
 
 	FMeshDrawCommandSortKey SortKey;
diff --git a/Engine/Source/Runtime/Renderer/Private/TranslucentRendering.h b/Engine/Source/Runtime/Renderer/Private/TranslucentRendering.h
--- a/Engine/Source/Runtime/Renderer/Private/TranslucentRendering.h
+++ b/Engine/Source/Runtime/Renderer/Private/TranslucentRendering.h
@@ -17,7 +17,7 @@
 bool UseNearestDepthNeighborUpsampleForSeparateTranslucency(const FSceneRenderTargets& SceneContext);
 
 extern float CalculateTranslucentSortKey(const FPrimitiveSceneInfo* RESTRICT PrimitiveSceneInfo, const FSceneView& View);
-extern FMeshDrawCommandSortKey CalculateStaticTranslucentMeshSortKey(const FPrimitiveSceneProxy* RESTRICT PrimitiveSceneProxy);
+extern FMeshDrawCommandSortKey CalculateStaticTranslucentMeshSortKey(const FPrimitiveSceneProxy* RESTRICT PrimitiveSceneProxy, uint16 MeshIdInPrimitive);
 EMeshPass::Type TranslucencyPassToMeshPass(ETranslucencyPass::Type TranslucencyPass);
 
 /**
@@ -29,8 +29,9 @@ union UTranslucentMeshSortKey
 
 	struct
 	{
-		uint64 Distance		: 32; // Order by distance.
-		uint64 SortPriority : 32; // First order by sort priority.
+		uint64 MeshIdInPrimitive	: 16; // Order meshes belonging to the same primitive by a stable id.
+		uint64 Distance				: 32; // Order by distance.
+		uint64 Priority				: 16; // First order by priority.
 	} Fields;
 };
 
