Fixed lightmaps on BSP brushes. Every mesh in proxy now has now an unique vertex factory uniform buffer, which contains an unique LODLightmapDataIndex, so it can fetch data from a proper LCI.

diff --git a/Engine/Source/Runtime/Engine/Private/ModelRender.cpp b/Engine/Source/Runtime/Engine/Private/ModelRender.cpp
--- a/Engine/Source/Runtime/Engine/Private/ModelRender.cpp
+++ b/Engine/Source/Runtime/Engine/Private/ModelRender.cpp
@@ -234,14 +234,23 @@ public:
 		InComponent->GetModel()->VertexBuffer.Buffers.InitModelVF(&VertexFactory);
 
 		OverrideOwnerName(NAME_BSP);
-		const TIndirectArray<FModelElement>& SourceElements = Component->GetElements();
+		TIndirectArray<FModelElement>& SourceElements = Component->GetElements();
 
 		Elements.Empty(SourceElements.Num());
 		for(int32 ElementIndex = 0;ElementIndex < SourceElements.Num();ElementIndex++)
 		{
-			const FModelElement& SourceElement = SourceElements[ElementIndex];
+			FModelElement& SourceElement = SourceElements[ElementIndex];
 			FElementInfo* Element = new(Elements) FElementInfo(SourceElement, VertexFactory.GetType());
 			MaterialRelevance |= Element->GetMaterial()->GetRelevance(GetScene().GetFeatureLevel());
+
+			if (RHISupportsManualVertexFetch(GMaxRHIShaderPlatform))
+			{
+				ENQUEUE_RENDER_COMMAND(FModelSceneProxySetData)(
+					[UniformBufferPtr = &SourceElement.VertexFactoryUniformBuffer, VertexFactory = &VertexFactory, ElementIndex](FRHICommandListImmediate& RHICmdList)
+				{
+					*UniformBufferPtr = CreateLocalVFUniformBuffer(VertexFactory, ElementIndex, nullptr, 0);
+				});
+			}
 		}
 
 		bGoodCandidateForCachedShadowmap = CacheShadowDepthsFromPrimitivesUsingWPO() || !MaterialRelevance.bUsesWorldPositionOffset;
@@ -442,6 +451,7 @@ public:
 												BatchElement.NumPrimitives = NumIndices / 3;
 												BatchElement.MinVertexIndex = MinVertexIndex;
 												BatchElement.MaxVertexIndex = MaxVertexIndex;
+												BatchElement.VertexFactoryUserData = ModelElement.VertexFactoryUniformBuffer;
 												MeshElement.Type = PT_TriangleList;
 												MeshElement.DepthPriorityGroup = DepthPriorityGroup;
 												MeshElement.bCanApplyViewModeOverrides = true;
@@ -477,6 +487,7 @@ public:
 								BatchElement.NumPrimitives = ModelElement.NumTriangles;
 								BatchElement.MinVertexIndex = ModelElement.MinVertexIndex;
 								BatchElement.MaxVertexIndex = ModelElement.MaxVertexIndex;
+								BatchElement.VertexFactoryUserData = ModelElement.VertexFactoryUniformBuffer;
 								MeshElement.Type = PT_TriangleList;
 								MeshElement.DepthPriorityGroup = DepthPriorityGroup;
 								MeshElement.bCanApplyViewModeOverrides = true;
@@ -517,6 +528,7 @@ public:
 					BatchElement.NumPrimitives = ModelElement.NumTriangles;
 					BatchElement.MinVertexIndex = ModelElement.MinVertexIndex;
 					BatchElement.MaxVertexIndex = ModelElement.MaxVertexIndex;
+					BatchElement.VertexFactoryUserData = ModelElement.VertexFactoryUniformBuffer;
 					MeshElement.Type = PT_TriangleList;
 					MeshElement.DepthPriorityGroup = PrimitiveDPG;
 					MeshElement.LODIndex = 0;
diff --git a/Engine/Source/Runtime/Engine/Public/Model.h b/Engine/Source/Runtime/Engine/Public/Model.h
--- a/Engine/Source/Runtime/Engine/Public/Model.h
+++ b/Engine/Source/Runtime/Engine/Public/Model.h
@@ -623,6 +623,9 @@ public:
 
 	FMeshMapBuildData* LegacyMapBuildData;
 
+	/** Vertex factory uniform buffer. Needs to be per model element, so every element can have it's unique LODLightmapDataIndex. */
+	FUniformBufferRHIRef VertexFactoryUniformBuffer;
+
 	/** Uniquely identifies this component's built map data. */
 	FGuid MapBuildDataId;
 
