BSP brush undo/redo crash fix. Undo/redo is serializing entire UModelComponent including IndexBuffer* which after redo will point to a deleted index buffer. A bit later PostLoad will recreate those, but this is to late with new mesh draw command caching. Afraid to touch serialization, so I just update index buffer pointers on UModelComponent::CreateRenderState_Concurrent.

diff --git a/Engine/Source/Runtime/Engine/Classes/Components/ModelComponent.h b/Engine/Source/Runtime/Engine/Classes/Components/ModelComponent.h
--- a/Engine/Source/Runtime/Engine/Classes/Components/ModelComponent.h
+++ b/Engine/Source/Runtime/Engine/Classes/Components/ModelComponent.h
@@ -79,12 +79,6 @@ public:
 	 */
 	void ShrinkElements();
 
-	/**
-	* Invalidate resources depending on the Model.
-	*/
-	void OnModelResourcesReleased();
-
-
 	/** Calculate the lightmap resolution to be used by the given surface. */
 	ENGINE_API void GetSurfaceLightMapResolution( int32 SurfaceIndex, int32 QualityScale, int32& Width, int32& Height, FMatrix& WorldToMap, TArray<int32>* GatheredNodes=NULL ) const;
 
diff --git a/Engine/Source/Runtime/Engine/Private/Components/ModelComponent.cpp b/Engine/Source/Runtime/Engine/Private/Components/ModelComponent.cpp
--- a/Engine/Source/Runtime/Engine/Private/Components/ModelComponent.cpp
+++ b/Engine/Source/Runtime/Engine/Private/Components/ModelComponent.cpp
@@ -234,17 +234,6 @@ void UModelComponent::ShrinkElements()
 	}
 }
 
-void UModelComponent::OnModelResourcesReleased()
-{
-	// Need to invalidate IndexBuffer*, as it's a pointer into Model resources which were just released.
-
-	for (int32 ElementIndex = 0; ElementIndex < Elements.Num(); ElementIndex++)
-	{
-		FModelElement& Element = Elements[ElementIndex];
-		Element.IndexBuffer = nullptr;
-	}
-}
-
 void UModelComponent::Serialize(FArchive& Ar)
 {
 	Super::Serialize(Ar);
diff --git a/Engine/Source/Runtime/Engine/Private/Level.cpp b/Engine/Source/Runtime/Engine/Private/Level.cpp
--- a/Engine/Source/Runtime/Engine/Private/Level.cpp
+++ b/Engine/Source/Runtime/Engine/Private/Level.cpp
@@ -1391,7 +1391,6 @@ void ULevel::PreEditUndo()
 	{
 		if(ModelComponents[ComponentIndex])
 		{
-			ModelComponents[ComponentIndex]->OnModelResourcesReleased();
 			ModelComponents[ComponentIndex]->UnregisterComponent();
 		}
 	}
@@ -1502,7 +1501,6 @@ void ULevel::InvalidateModelGeometry()
 	{
 		if(ModelComponents[ComponentIndex])
 		{
-			ModelComponents[ComponentIndex]->OnModelResourcesReleased();
 			ModelComponents[ComponentIndex]->Modify();
 			ModelComponents[ComponentIndex]->UnregisterComponent();
 		}
@@ -1549,7 +1547,6 @@ void ULevel::CommitModelSurfaces()
 		{
 			if(ModelComponents[ComponentIndex])
 			{
-				ModelComponents[ComponentIndex]->OnModelResourcesReleased();
 				ModelComponents[ComponentIndex]->CommitSurfaces();
 			}
 		}
diff --git a/Engine/Source/Runtime/Engine/Private/ModelRender.cpp b/Engine/Source/Runtime/Engine/Private/ModelRender.cpp
--- a/Engine/Source/Runtime/Engine/Private/ModelRender.cpp
+++ b/Engine/Source/Runtime/Engine/Private/ModelRender.cpp
@@ -795,6 +795,13 @@ public:
 
 void UModelComponent::CreateRenderState_Concurrent()
 {
+	for (int32 ElementIndex = 0; ElementIndex < Elements.Num(); ElementIndex++)
+	{
+		FModelElement& Element = Elements[ElementIndex];
+		TUniquePtr<FRawIndexBuffer16or32>* IndexBufferRef = Model->MaterialIndexBuffers.Find(Element.Material);
+		Element.IndexBuffer = IndexBufferRef ? IndexBufferRef->Get() : nullptr;
+	}
+
 	if (GetModel())
 	{
 		++GetModel()->VertexBuffer.RefCount;
