Fixed dynamic primitive data upload for multiple views. All uploads were added to the last view, thus breaking rendering of all views up to the last one.

diff --git a/Engine/Source/Runtime/Engine/Private/SceneManagement.cpp b/Engine/Source/Runtime/Engine/Private/SceneManagement.cpp
--- a/Engine/Source/Runtime/Engine/Private/SceneManagement.cpp
+++ b/Engine/Source/Runtime/Engine/Private/SceneManagement.cpp
@@ -278,6 +278,8 @@ void FMeshElementCollector::AddMesh(int32 ViewIndex, FMeshBatch& MeshBatch)
 
 			if (PrimitiveUniformBufferResource)
 			{
+				TArray<FPrimitiveUniformShaderParameters>* DynamicPrimitiveShaderData = DynamicPrimitiveShaderDataPerView[ViewIndex];
+
 				const int32 DataIndex = DynamicPrimitiveShaderData->AddUninitialized(1);
 				MeshBatch.Elements[Index].PrimitiveIdMode = PrimID_DynamicPrimitiveShaderData;
 				MeshBatch.Elements[Index].DynamicPrimitiveShaderDataIndex = DataIndex;
diff --git a/Engine/Source/Runtime/Engine/Public/SceneManagement.h b/Engine/Source/Runtime/Engine/Public/SceneManagement.h
--- a/Engine/Source/Runtime/Engine/Public/SceneManagement.h
+++ b/Engine/Source/Runtime/Engine/Public/SceneManagement.h
@@ -1978,7 +1978,7 @@ private:
 		MeshBatches.Empty();
 		SimpleElementCollectors.Empty();
 		MeshIdInPrimitivePerView.Empty();
-		DynamicPrimitiveShaderData = nullptr;
+		DynamicPrimitiveShaderDataPerView.Empty();
 	}
 
 	void AddViewMeshArrays(
@@ -1989,10 +1989,10 @@ private:
 		ERHIFeatureLevel::Type InFeatureLevel)
 	{
 		Views.Add(InView);
-		MeshIdInPrimitivePerView.AddZeroed(Views.Num());
+		MeshIdInPrimitivePerView.Add(0);
 		MeshBatches.Add(ViewMeshes);
 		SimpleElementCollectors.Add(ViewSimpleElementCollector);
-		DynamicPrimitiveShaderData = InDynamicPrimitiveShaderData;
+		DynamicPrimitiveShaderDataPerView.Add(InDynamicPrimitiveShaderData);
 	}
 
 	/** 
@@ -2030,8 +2030,8 @@ private:
 	/** Tasks to wait for at the end of gathering dynamic mesh elements. */
 	TArray<TFunction<void()>*, SceneRenderingAllocator> ParallelTasks;
 
-	/** Tracks dynamic primitive data for upload to GPU Scene, when enabled. */
-	TArray<FPrimitiveUniformShaderParameters>* DynamicPrimitiveShaderData;
+	/** Tracks dynamic primitive data for upload to GPU Scene for every view, when enabled. */
+	TArray<TArray<FPrimitiveUniformShaderParameters>*, TInlineAllocator<2> > DynamicPrimitiveShaderDataPerView;
 
 	friend class FSceneRenderer;
 	friend class FDeferredShadingSceneRenderer;
diff --git a/Engine/Source/Runtime/Renderer/Private/SceneRendering.cpp b/Engine/Source/Runtime/Renderer/Private/SceneRendering.cpp
--- a/Engine/Source/Runtime/Renderer/Private/SceneRendering.cpp
+++ b/Engine/Source/Runtime/Renderer/Private/SceneRendering.cpp
@@ -3012,6 +3012,8 @@ void FSceneRenderer::RenderCustomDepthPass(FRHICommandListImmediate& RHICmdList)
 
 			if (View.ShouldRenderView())
 			{
+				Scene->UniformBuffers.UpdateViewUniformBuffer(View);
+
 				FUniformBufferRHIParamRef PassUniformBuffer = nullptr;
 
 				if (FSceneInterface::GetShadingPath(View.FeatureLevel) == EShadingPath::Mobile)
@@ -3028,7 +3030,6 @@ void FSceneRenderer::RenderCustomDepthPass(FRHICommandListImmediate& RHICmdList)
 					Scene->UniformBuffers.CustomDepthPassUniformBuffer.UpdateUniformBufferImmediate(SceneTextureParameters);
 					PassUniformBuffer = Scene->UniformBuffers.CustomDepthPassUniformBuffer;
 				}
-
 				
 				if (!View.IsInstancedStereoPass())
 				{
