Fixed InitViews after the merge. Removed ScreenSizes array from Proxy, as it's now stored inside lightweight StaticMeshRelevances. Also reverted most changes to ComputeLODForMeshes, as it was crashing by passing mesh index as a LOD level.

diff --git a/Engine/Source/Runtime/Engine/Private/SceneManagement.cpp b/Engine/Source/Runtime/Engine/Private/SceneManagement.cpp
--- a/Engine/Source/Runtime/Engine/Private/SceneManagement.cpp
+++ b/Engine/Source/Runtime/Engine/Private/SceneManagement.cpp
@@ -538,12 +538,12 @@ int8 ComputeStaticMeshLOD( const FStaticMeshRenderData* RenderData, const FVecto
 	return MinLOD;
 }
 
-FLODMask ComputeLODForMeshes(const TArray<float>& ScreenSizes, bool bDitheredLODTransition, const FSceneView& View, const FVector4& Origin, float SphereRadius, int32 ForcedLODLevel, float& OutScreenRadiusSquared, float ScreenSizeScale)
+FLODMask ComputeLODForMeshes(const TArray<class FStaticMeshRelevance>& StaticMeshRelevances, const FSceneView& View, const FVector4& Origin, float SphereRadius, int32 ForcedLODLevel, float& OutScreenRadiusSquared, float ScreenSizeScale, bool bDitheredLODTransition)
 {
 	FLODMask LODToRender;
 	const FSceneView& LODView = GetLODView(View);
 
-	const int32 NumMeshes = ScreenSizes.Num();
+	const int32 NumMeshes = StaticMeshRelevances.Num();
 
 	// Handle forced LOD level first
 	if (ForcedLODLevel >= 0)
@@ -551,21 +551,20 @@ FLODMask ComputeLODForMeshes(const TArray<float>& ScreenSizes, bool bDitheredLOD
 		OutScreenRadiusSquared = 0.0f;
 
 		int32 MinLOD = 127, MaxLOD = 0;
-		for (int32 MeshIndex = 0; MeshIndex < ScreenSizes.Num(); ++MeshIndex)
+		for (int32 MeshIndex = 0; MeshIndex < StaticMeshRelevances.Num(); ++MeshIndex)
 		{
-			if (ScreenSizes[MeshIndex] == 0.0f)
+			const FStaticMeshRelevance& Mesh = StaticMeshRelevances[MeshIndex];
+			if (Mesh.ScreenSize > 0.0f)
 			{
-				continue;
+				MinLOD = FMath::Min(MinLOD, (int32)Mesh.LODIndex);
+				MaxLOD = FMath::Max(MaxLOD, (int32)Mesh.LODIndex);
 			}
-
-			MinLOD = FMath::Min(MinLOD, MeshIndex);
-			MaxLOD = FMath::Max(MaxLOD, MeshIndex);
 		}
 		LODToRender.SetLOD(FMath::Clamp(ForcedLODLevel, MinLOD, MaxLOD));
 	}
 	else if (LODView.Family->EngineShowFlags.LOD && NumMeshes)
 	{
-		if (bDitheredLODTransition)
+		if (bDitheredLODTransition && StaticMeshRelevances[0].bDitheredLODTransition)
 		{
 			for (int32 SampleIndex = 0; SampleIndex < 2; SampleIndex++)
 			{
@@ -573,23 +572,22 @@ FLODMask ComputeLODForMeshes(const TArray<float>& ScreenSizes, bool bDitheredLOD
 				bool bFoundLOD = false;
 				OutScreenRadiusSquared = ComputeTemporalLODBoundsScreenRadiusSquared(Origin, SphereRadius, LODView, SampleIndex);
 
-				for(int32 MeshIndex = NumMeshes-1 ; MeshIndex >= 0 ; --MeshIndex)
+				for (int32 MeshIndex = NumMeshes - 1; MeshIndex >= 0; --MeshIndex)
 				{
-					if (ScreenSizes[MeshIndex] == 0.0f)
+					const FStaticMeshRelevance& Mesh = StaticMeshRelevances[MeshIndex];
+					if (Mesh.ScreenSize > 0.0f)
 					{
-						continue;
-					}
+						float MeshScreenSize = Mesh.ScreenSize * ScreenSizeScale;
 
-					float MeshScreenSize = ScreenSizes[MeshIndex] * ScreenSizeScale;
+						if (FMath::Square(MeshScreenSize * 0.5f) >= OutScreenRadiusSquared)
+						{
+							LODToRender.SetLODSample(Mesh.LODIndex, SampleIndex);
+							bFoundLOD = true;
+							break;
+						}
 
-					if(FMath::Square(MeshScreenSize * 0.5f) >= OutScreenRadiusSquared)
-					{
-						LODToRender.SetLODSample(MeshIndex, SampleIndex);
-						bFoundLOD = true;
-						break;
+						MinLODFound = FMath::Min<int32>(MinLODFound, Mesh.LODIndex);
 					}
-
-					MinLODFound = FMath::Min<int32>(MinLODFound, MeshIndex);
 				}
 				// If no LOD was found matching the screen size, use the lowest in the array instead of LOD 0, to handle non-zero MinLOD
 				if (!bFoundLOD)
@@ -604,23 +602,20 @@ FLODMask ComputeLODForMeshes(const TArray<float>& ScreenSizes, bool bDitheredLOD
 			bool bFoundLOD = false;
 			OutScreenRadiusSquared = ComputeBoundsScreenRadiusSquared(Origin, SphereRadius, LODView);
 
-			for(int32 MeshIndex = NumMeshes-1 ; MeshIndex >= 0 ; --MeshIndex)
+			for (int32 MeshIndex = NumMeshes - 1; MeshIndex >= 0; --MeshIndex)
 			{
-				if (ScreenSizes[MeshIndex] == 0.0f)
-				{
-					continue;
-				}
+				const FStaticMeshRelevance& Mesh = StaticMeshRelevances[MeshIndex];
 
-				float MeshScreenSize = ScreenSizes[MeshIndex] * ScreenSizeScale;
+				float MeshScreenSize = Mesh.ScreenSize * ScreenSizeScale;
 
-				if(FMath::Square(MeshScreenSize * 0.5f) >= OutScreenRadiusSquared)
+				if (FMath::Square(MeshScreenSize * 0.5f) >= OutScreenRadiusSquared)
 				{
-					LODToRender.SetLOD(MeshIndex);
+					LODToRender.SetLOD(Mesh.LODIndex);
 					bFoundLOD = true;
 					break;
 				}
 
-				MinLODFound = FMath::Min<int32>(MinLODFound, MeshIndex);
+				MinLODFound = FMath::Min<int32>(MinLODFound, Mesh.LODIndex);
 			}
 			// If no LOD was found matching the screen size, use the lowest in the array instead of LOD 0, to handle non-zero MinLOD
 			if (!bFoundLOD)
diff --git a/Engine/Source/Runtime/Engine/Public/PrimitiveSceneProxy.h b/Engine/Source/Runtime/Engine/Public/PrimitiveSceneProxy.h
--- a/Engine/Source/Runtime/Engine/Public/PrimitiveSceneProxy.h
+++ b/Engine/Source/Runtime/Engine/Public/PrimitiveSceneProxy.h
@@ -1024,7 +1024,6 @@ private:
 #endif
 
 public:
-	TArray<float> ScreenSizes;
 #if RHI_RAYTRACING
 	struct FStaticMeshOrCommandIndex
 	{
diff --git a/Engine/Source/Runtime/Engine/Public/SceneManagement.h b/Engine/Source/Runtime/Engine/Public/SceneManagement.h
--- a/Engine/Source/Runtime/Engine/Public/SceneManagement.h
+++ b/Engine/Source/Runtime/Engine/Public/SceneManagement.h
@@ -2653,7 +2653,7 @@ struct FLODMask
 		return DitheredLODIndices[0] != DitheredLODIndices[1];
 	}
 };
-FLODMask ENGINE_API ComputeLODForMeshes(const TArray<float>& ScreenSizes, bool bDitheredLODTransition, const FSceneView& View, const FVector4& Origin, float SphereRadius, int32 ForcedLODLevel, float& OutScreenRadiusSquared, float ScreenSizeScale = 1.0f);
+FLODMask ENGINE_API ComputeLODForMeshes(const TArray<class FStaticMeshRelevance>& StaticMeshRelevances, const FSceneView& View, const FVector4& Origin, float SphereRadius, int32 ForcedLODLevel, float& OutScreenRadiusSquared, float ScreenSizeScale = 1.0f, bool bDitheredLODTransition = true);
 
 class FSharedSamplerState : public FRenderResource
 {
diff --git a/Engine/Source/Runtime/Renderer/Private/DeferredShadingRenderer.cpp b/Engine/Source/Runtime/Renderer/Private/DeferredShadingRenderer.cpp
--- a/Engine/Source/Runtime/Renderer/Private/DeferredShadingRenderer.cpp
+++ b/Engine/Source/Runtime/Renderer/Private/DeferredShadingRenderer.cpp
@@ -667,7 +667,7 @@ bool FDeferredShadingSceneRenderer::GatherRayTracingWorldInstances(FRHICommandLi
 						}
 						else
 						{
-							LODToRender = ComputeLODForMeshes(SceneProxy->ScreenSizes, false, View, Bounds.BoxSphereBounds.Origin, Bounds.BoxSphereBounds.SphereRadius, ForcedLODLevel, MeshScreenSizeSquared, LODScale);
+							LODToRender = ComputeLODForMeshes(SceneProxy->GetPrimitiveSceneInfo()->StaticMeshRelevances, View, Bounds.BoxSphereBounds.Origin, Bounds.BoxSphereBounds.SphereRadius, ForcedLODLevel, MeshScreenSizeSquared, LODScale, false);
 						}
 
 						FRayTracingGeometryRHIRef RayTracingGeometryInstance = SceneProxy->GetRayTracingGeometryInstance(LODToRender.GetRayTracedLOD());
diff --git a/Engine/Source/Runtime/Renderer/Private/PrimitiveSceneInfo.cpp b/Engine/Source/Runtime/Renderer/Private/PrimitiveSceneInfo.cpp
--- a/Engine/Source/Runtime/Renderer/Private/PrimitiveSceneInfo.cpp
+++ b/Engine/Source/Runtime/Renderer/Private/PrimitiveSceneInfo.cpp
@@ -200,8 +200,6 @@ void FPrimitiveSceneInfo::AddStaticMeshes(FRHICommandListImmediate& RHICmdList,
 
 	if (StaticMeshes.Num() > 0)
 	{
-		Proxy->ScreenSizes.AddDefaulted(MaxLOD + 1);
-
 #if RHI_RAYTRACING
 		Proxy->RayTracingLodIndexToMeshDrawCommandIndicies.AddDefaulted(MaxLOD + 1);
 #endif
@@ -210,15 +208,6 @@ void FPrimitiveSceneInfo::AddStaticMeshes(FRHICommandListImmediate& RHICmdList,
 		{
 			FStaticMeshRelevance& MeshRelevance = StaticMeshRelevances[MeshIndex];
 			FStaticMesh& Mesh = StaticMeshes[MeshIndex];
-			if (Proxy->ScreenSizes[Mesh.LODIndex] != 0.0f)
-			{
-				check(Proxy->ScreenSizes[Mesh.LODIndex] == MeshRelevance.ScreenSize);
-			}
-			else
-			{
-				check(MeshRelevance.ScreenSize != 0.0f);
-				Proxy->ScreenSizes[Mesh.LODIndex] = MeshRelevance.ScreenSize;
-			}
 
 #if RHI_RAYTRACING
 			if (Mesh.Elements.Num() > 0)
diff --git a/Engine/Source/Runtime/Renderer/Private/SceneVisibility.cpp b/Engine/Source/Runtime/Renderer/Private/SceneVisibility.cpp
--- a/Engine/Source/Runtime/Renderer/Private/SceneVisibility.cpp
+++ b/Engine/Source/Runtime/Renderer/Private/SceneVisibility.cpp
@@ -2019,7 +2019,7 @@ struct FRelevancePacket
 			}
 			else
 			{
-				LODToRender = ComputeLODForMeshes(PrimitiveSceneInfo->Proxy->ScreenSizes, PrimitiveSceneInfo->StaticMeshes[0].bDitheredLODTransition, View, Bounds.BoxSphereBounds.Origin, Bounds.BoxSphereBounds.SphereRadius, ViewData.ForcedLODLevel, MeshScreenSizeSquared, ViewData.LODScale);
+				LODToRender = ComputeLODForMeshes(PrimitiveSceneInfo->StaticMeshRelevances, View, Bounds.BoxSphereBounds.Origin, Bounds.BoxSphereBounds.SphereRadius, ViewData.ForcedLODLevel, MeshScreenSizeSquared, ViewData.LODScale);
 			}
 
 			PrimitivesLODMask.AddPrim(FRelevancePacket::FPrimitiveLODMask(PrimitiveIndex, LODToRender));
diff --git a/Engine/Source/Runtime/Renderer/Private/ShadowSetup.cpp b/Engine/Source/Runtime/Renderer/Private/ShadowSetup.cpp
--- a/Engine/Source/Runtime/Renderer/Private/ShadowSetup.cpp
+++ b/Engine/Source/Runtime/Renderer/Private/ShadowSetup.cpp
@@ -861,7 +861,7 @@ bool FProjectedShadowInfo::ShouldDrawStaticMeshes(FViewInfo& InCurrentView, bool
 			else
 			{
 				const FBoxSphereBounds& Bounds = InPrimitiveSceneInfo->Proxy->GetBounds();
-				ViewLODToRender = ComputeLODForMeshes(InPrimitiveSceneInfo->Proxy->ScreenSizes, InPrimitiveSceneInfo->StaticMeshes[0].bDitheredLODTransition, InCurrentView, Bounds.Origin, Bounds.SphereRadius, ForcedLOD, MeshScreenSizeSquared, InCurrentView.LODDistanceFactor);
+				ViewLODToRender = ComputeLODForMeshes(InPrimitiveSceneInfo->StaticMeshRelevances, InCurrentView, Bounds.Origin, Bounds.SphereRadius, ForcedLOD, MeshScreenSizeSquared, InCurrentView.LODDistanceFactor);
 			}	
 
 			InCurrentView.PrimitivesLODMask[PrimitiveId] = ViewLODToRender;
