Repair GatherRayTracingWorldInstance() optimizations

diff --git a/Engine/Source/Runtime/Engine/Private/SceneManagement.cpp b/Engine/Source/Runtime/Engine/Private/SceneManagement.cpp
--- a/Engine/Source/Runtime/Engine/Private/SceneManagement.cpp
+++ b/Engine/Source/Runtime/Engine/Private/SceneManagement.cpp
@@ -627,6 +627,98 @@ FLODMask ComputeLODForMeshes(const TArray<class FStaticMeshRelevance>& StaticMes
 	return LODToRender;
 }
 
+FLODMask ComputeFastLODForMeshes(const TArray<float>& ScreenSizes, const FSceneView& View, const FVector4& Origin, float SphereRadius, int32 ForcedLODLevel, float& OutScreenRadiusSquared, float ScreenSizeScale, bool bDitheredLODTransition)
+{
+	FLODMask LODToRender;
+	const FSceneView& LODView = GetLODView(View);
+
+	const int32 NumMeshes = ScreenSizes.Num();
+
+	// Handle forced LOD level first
+	if (ForcedLODLevel >= 0)
+	{
+		OutScreenRadiusSquared = 0.0f;
+
+		int32 MinLOD = 127, MaxLOD = 0;
+		for (int32 MeshIndex = 0; MeshIndex < ScreenSizes.Num(); ++MeshIndex)
+		{
+			if (ScreenSizes[MeshIndex] == 0.0f)
+			{
+				continue;
+			}
+
+			MinLOD = FMath::Min(MinLOD, MeshIndex);
+			MaxLOD = FMath::Max(MaxLOD, MeshIndex);
+		}
+		LODToRender.SetLOD(FMath::Clamp(ForcedLODLevel, MinLOD, MaxLOD));
+	}
+	else if (LODView.Family->EngineShowFlags.LOD && NumMeshes)
+	{
+		if (bDitheredLODTransition)
+		{
+			for (int32 SampleIndex = 0; SampleIndex < 2; SampleIndex++)
+			{
+				int32 MinLODFound = INT_MAX;
+				bool bFoundLOD = false;
+				OutScreenRadiusSquared = ComputeTemporalLODBoundsScreenRadiusSquared(Origin, SphereRadius, LODView, SampleIndex);
+
+				for (int32 MeshIndex = NumMeshes - 1; MeshIndex >= 0; --MeshIndex)
+				{
+					if (ScreenSizes[MeshIndex] > 0.0f)
+					{
+						float MeshScreenSize = ScreenSizes[MeshIndex] * ScreenSizeScale;
+
+						if (FMath::Square(MeshScreenSize * 0.5f) >= OutScreenRadiusSquared)
+						{
+							LODToRender.SetLODSample(MeshIndex, SampleIndex);
+							bFoundLOD = true;
+							break;
+						}
+
+						MinLODFound = FMath::Min<int32>(MinLODFound, MeshIndex);
+					}
+				}
+				// If no LOD was found matching the screen size, use the lowest in the array instead of LOD 0, to handle non-zero MinLOD
+				if (!bFoundLOD)
+				{
+					LODToRender.SetLODSample(MinLODFound, SampleIndex);
+				}
+			}
+		}
+		else
+		{
+			int32 MinLODFound = INT_MAX;
+			bool bFoundLOD = false;
+			OutScreenRadiusSquared = ComputeBoundsScreenRadiusSquared(Origin, SphereRadius, LODView);
+
+			for (int32 MeshIndex = NumMeshes - 1; MeshIndex >= 0; --MeshIndex)
+			{
+				if (ScreenSizes[MeshIndex] == 0.0f)
+				{
+					continue;
+				}
+
+				float MeshScreenSize = ScreenSizes[MeshIndex] * ScreenSizeScale;
+
+				if (FMath::Square(MeshScreenSize * 0.5f) >= OutScreenRadiusSquared)
+				{
+					LODToRender.SetLOD(MeshIndex);
+					bFoundLOD = true;
+					break;
+				}
+
+				MinLODFound = FMath::Min<int32>(MinLODFound, MeshIndex);
+			}
+			// If no LOD was found matching the screen size, use the lowest in the array instead of LOD 0, to handle non-zero MinLOD
+			if (!bFoundLOD)
+			{
+				LODToRender.SetLOD(MinLODFound);
+			}
+		}
+	}
+	return LODToRender;
+}
+
 FMobileDirectionalLightShaderParameters::FMobileDirectionalLightShaderParameters()
 {
 	FMemory::Memzero(*this);
diff --git a/Engine/Source/Runtime/Engine/Public/PrimitiveSceneProxy.h b/Engine/Source/Runtime/Engine/Public/PrimitiveSceneProxy.h
--- a/Engine/Source/Runtime/Engine/Public/PrimitiveSceneProxy.h
+++ b/Engine/Source/Runtime/Engine/Public/PrimitiveSceneProxy.h
@@ -1025,6 +1025,7 @@ private:
 
 public:
 #if RHI_RAYTRACING
+	TArray<float> ScreenSizes;
 	struct FStaticMeshOrCommandIndex
 	{
 		int32 StaticMeshIndex;
diff --git a/Engine/Source/Runtime/Engine/Public/SceneManagement.h b/Engine/Source/Runtime/Engine/Public/SceneManagement.h
--- a/Engine/Source/Runtime/Engine/Public/SceneManagement.h
+++ b/Engine/Source/Runtime/Engine/Public/SceneManagement.h
@@ -2654,6 +2654,7 @@ struct FLODMask
 	}
 };
 FLODMask ENGINE_API ComputeLODForMeshes(const TArray<class FStaticMeshRelevance>& StaticMeshRelevances, const FSceneView& View, const FVector4& Origin, float SphereRadius, int32 ForcedLODLevel, float& OutScreenRadiusSquared, float ScreenSizeScale = 1.0f, bool bDitheredLODTransition = true);
+FLODMask ENGINE_API ComputeFastLODForMeshes(const TArray<float>& ScreenSizes, const FSceneView& View, const FVector4& Origin, float SphereRadius, int32 ForcedLODLevel, float& OutScreenRadiusSquared, float ScreenSizeScale = 1.0f, bool bDitheredLODTransition = true);
 
 class FSharedSamplerState : public FRenderResource
 {
diff --git a/Engine/Source/Runtime/Renderer/Private/DeferredShadingRenderer.cpp b/Engine/Source/Runtime/Renderer/Private/DeferredShadingRenderer.cpp
--- a/Engine/Source/Runtime/Renderer/Private/DeferredShadingRenderer.cpp
+++ b/Engine/Source/Runtime/Renderer/Private/DeferredShadingRenderer.cpp
@@ -667,7 +667,7 @@ bool FDeferredShadingSceneRenderer::GatherRayTracingWorldInstances(FRHICommandLi
 						}
 						else
 						{
-							LODToRender = ComputeLODForMeshes(SceneProxy->GetPrimitiveSceneInfo()->StaticMeshRelevances, View, Bounds.BoxSphereBounds.Origin, Bounds.BoxSphereBounds.SphereRadius, ForcedLODLevel, MeshScreenSizeSquared, LODScale, false);
+							LODToRender = ComputeFastLODForMeshes(SceneProxy->ScreenSizes, View, Bounds.BoxSphereBounds.Origin, Bounds.BoxSphereBounds.SphereRadius, ForcedLODLevel, MeshScreenSizeSquared, LODScale, false);
 						}
 
 						FRayTracingGeometryRHIRef RayTracingGeometryInstance = SceneProxy->GetRayTracingGeometryInstance(LODToRender.GetRayTracedLOD());
diff --git a/Engine/Source/Runtime/Renderer/Private/PrimitiveSceneInfo.cpp b/Engine/Source/Runtime/Renderer/Private/PrimitiveSceneInfo.cpp
--- a/Engine/Source/Runtime/Renderer/Private/PrimitiveSceneInfo.cpp
+++ b/Engine/Source/Runtime/Renderer/Private/PrimitiveSceneInfo.cpp
@@ -169,7 +169,6 @@ void FPrimitiveSceneInfo::AddStaticMeshes(FRHICommandListImmediate& RHICmdList,
 
 	check(StaticMeshRelevances.Num() == StaticMeshes.Num());
 
-	int MaxLOD = -1;
 	for(int32 MeshIndex = 0;MeshIndex < StaticMeshes.Num();MeshIndex++)
 	{
 		FStaticMeshRelevance& MeshRelevance = StaticMeshRelevances[MeshIndex];
@@ -194,35 +193,41 @@ void FPrimitiveSceneInfo::AddStaticMeshes(FRHICommandListImmediate& RHICmdList,
 			// Add the static mesh to the appropriate draw lists.
 			Mesh.AddToDrawLists(RHICmdList, Scene);
 		}
+	}
+
+#if RHI_RAYTRACING
+	int MaxLOD = -1;
 
+	for (int32 MeshIndex = 0; MeshIndex < StaticMeshes.Num(); MeshIndex++)
+	{
+		FStaticMesh& Mesh = StaticMeshes[MeshIndex];
 		MaxLOD = MaxLOD < Mesh.LODIndex ? Mesh.LODIndex : MaxLOD;
 	}
 
 	if (StaticMeshes.Num() > 0)
 	{
-#if RHI_RAYTRACING
+		Proxy->ScreenSizes.AddDefaulted(MaxLOD + 1);
 		Proxy->RayTracingLodIndexToMeshDrawCommandIndicies.AddDefaulted(MaxLOD + 1);
-#endif
 
 		for (int32 MeshIndex = 0; MeshIndex < StaticMeshes.Num(); MeshIndex++)
 		{
 			FStaticMeshRelevance& MeshRelevance = StaticMeshRelevances[MeshIndex];
 			FStaticMesh& Mesh = StaticMeshes[MeshIndex];
 
-#if RHI_RAYTRACING
+			if (Proxy->ScreenSizes[Mesh.LODIndex] != 0.0f)
+			{
+				check(Proxy->ScreenSizes[Mesh.LODIndex] == MeshRelevance.ScreenSize);
+			}
+			else
+			{
+				check(MeshRelevance.ScreenSize != 0.0f);
+				Proxy->ScreenSizes[Mesh.LODIndex] = MeshRelevance.ScreenSize;
+			}
+
 			if (Mesh.Elements.Num() > 0)
 			{
-				//dxr_todo: review LOD indices and CommandIndex
-				const auto& StaticMeshMdcIndices = Proxy->RayTracingLodIndexToMeshDrawCommandIndicies[MaxLOD];
-				int32 CommandIndex = -1;
-				
-				//if (StaticMeshMdcIndices.Num() > 0)
-				//{
-				//	CommandIndex = StaticMeshMdcIndices[MeshIndex].CommandIndex;
-				//}
-				const bool bSupportsCachingMeshDrawCommands = SupportsCachingMeshDrawCommands(Mesh.VertexFactory, Proxy);
-
-				if (bSupportsCachingMeshDrawCommands && CommandIndex != -1)
+				const int32 CommandIndex = Mesh.CachedMeshDrawCommands[EMeshPass::RayTracing].CommandIndex;
+				if (SupportsCachingMeshDrawCommands(Mesh.VertexFactory, Proxy) && CommandIndex != -1)
 				{
 					Proxy->RayTracingLodIndexToMeshDrawCommandIndicies[Mesh.LODIndex].Add({ MeshIndex, CommandIndex });
 				}
@@ -231,8 +236,8 @@ void FPrimitiveSceneInfo::AddStaticMeshes(FRHICommandListImmediate& RHICmdList,
 					Proxy->RayTracingLodIndexToMeshDrawCommandIndicies[Mesh.LODIndex].Add({ MeshIndex, -1 });
 				}
 			}
-#endif
 		}
+#endif
 	}
 }
 
