Make the raytracing requirement for skincache explicit.  Also makes skincache and raytracing off by default again.  Projects wanting to enable raytracing will need to set both project settings.

diff --git a/Engine/Source/Runtime/Engine/Classes/Engine/RendererSettings.h b/Engine/Source/Runtime/Engine/Classes/Engine/RendererSettings.h
--- a/Engine/Source/Runtime/Engine/Classes/Engine/RendererSettings.h
+++ b/Engine/Source/Runtime/Engine/Classes/Engine/RendererSettings.h
@@ -572,7 +572,7 @@ class ENGINE_API URendererSettings : public UDeveloperSettings
 	*/
 	UPROPERTY(config, EditAnywhere, Category = RayTracing, meta = (
 		ConsoleVariable = "r.RayTracing", DisplayName = "Ray Tracing",
-		ToolTip = "Enable Ray Tracing capabilities.",
+		ToolTip = "Enable Ray Tracing capabilities.  Requires 'Support Compute Skincache' before project is allowed to set this.",
 		ConfigRestartRequired = true))
 		uint32 bEnableRayTracing : 1;
 
diff --git a/Engine/Source/Runtime/Engine/Private/GPUSkinCache.cpp b/Engine/Source/Runtime/Engine/Private/GPUSkinCache.cpp
--- a/Engine/Source/Runtime/Engine/Private/GPUSkinCache.cpp
+++ b/Engine/Source/Runtime/Engine/Private/GPUSkinCache.cpp
@@ -26,11 +26,8 @@ DEFINE_STAT(STAT_GPUSkinCache_NumSetVertexStreams);
 DEFINE_STAT(STAT_GPUSkinCache_NumPreGDME);
 DEFINE_LOG_CATEGORY_STATIC(LogSkinCache, Log, All);
 
-#if RHI_RAYTRACING
-static int32 GEnableGPUSkinCacheShaders = 1;
-#else
 static int32 GEnableGPUSkinCacheShaders = 0;
-#endif
+
 static FAutoConsoleVariableRef CVarEnableGPUSkinCacheShaders(
 	TEXT("r.SkinCache.CompileShaders"),
 	GEnableGPUSkinCacheShaders,
diff --git a/Engine/Source/Runtime/Engine/Private/RendererSettings.cpp b/Engine/Source/Runtime/Engine/Private/RendererSettings.cpp
--- a/Engine/Source/Runtime/Engine/Private/RendererSettings.cpp
+++ b/Engine/Source/Runtime/Engine/Private/RendererSettings.cpp
@@ -64,11 +64,7 @@ URendererSettings::URendererSettings(const FObjectInitializer& ObjectInitializer
 	bSupportMaterialLayers = false;
 	GPUSimulationTextureSizeX = 1024;
 	GPUSimulationTextureSizeY = 1024;
-#if RHI_RAYTRACING
-	bEnableRayTracing = 1;
-#else
 	bEnableRayTracing = 0;
-#endif
 }
 
 void URendererSettings::PostInitProperties()
@@ -83,6 +79,10 @@ void URendererSettings::PostInitProperties()
 		ImportConsoleVariableValues();
 	}
 #endif // #if WITH_EDITOR
+
+#if RHI_RAYTRACING
+	ensureMsgf(!bEnableRayTracing || (bEnableRayTracing && bSupportSkinCacheShaders), TEXT("Raytracing enabled, but skin cache is not.  Disable raytracing, or enable r.SkinCache.CompileShaders"));
+#endif
 }
 
 #if WITH_EDITOR
@@ -121,7 +121,8 @@ bool URendererSettings::CanEditChange(const UProperty* InProperty) const
 	if ((InProperty->GetFName() == GET_MEMBER_NAME_CHECKED(URendererSettings, bEnableRayTracing)))
 	{
 #if RHI_RAYTRACING
-		return ParentVal;
+		//only allow enabling raytracing if skin cache is on, but allow disable of raytracing no matter what.
+		return ParentVal && (bSupportSkinCacheShaders || bEnableRayTracing);
 #else // RHI_RAYTRACING
 		return false;
 #endif // RHI_RAYTRACING
diff --git a/Engine/Source/Runtime/Renderer/Private/DeferredShadingRenderer.cpp b/Engine/Source/Runtime/Renderer/Private/DeferredShadingRenderer.cpp
--- a/Engine/Source/Runtime/Renderer/Private/DeferredShadingRenderer.cpp
+++ b/Engine/Source/Runtime/Renderer/Private/DeferredShadingRenderer.cpp
@@ -152,14 +152,10 @@ static FAutoConsoleVariableRef CVarEnableRayTracingMaterials(
 	ECVF_RenderThreadSafe
 );
 
-static int32 GRayTracing = 1;
-
-#else // RHI_RAYTRACING
+#endif // RHI_RAYTRACING
 
 static int32 GRayTracing = 0;
 
-#endif // RHI_RAYTRACING
-
 static TAutoConsoleVariable<int32> CVarRayTracing(
 	TEXT("r.RayTracing"),
 	GRayTracing,
