Fix non-unity compile. It's not necessary to restrict project settings adjustment on other platforms.

diff --git a/Engine/Source/Runtime/Engine/Private/RendererSettings.cpp b/Engine/Source/Runtime/Engine/Private/RendererSettings.cpp
--- a/Engine/Source/Runtime/Engine/Private/RendererSettings.cpp
+++ b/Engine/Source/Runtime/Engine/Private/RendererSettings.cpp
@@ -80,9 +80,7 @@ void URendererSettings::PostInitProperties()
 	}
 #endif // #if WITH_EDITOR
 
-#if RHI_RAYTRACING
 	ensureMsgf(!bEnableRayTracing || (bEnableRayTracing && bSupportSkinCacheShaders), TEXT("Raytracing enabled, but skin cache is not.  Disable raytracing, or enable r.SkinCache.CompileShaders"));
-#endif
 }
 
 #if WITH_EDITOR
@@ -120,12 +118,8 @@ bool URendererSettings::CanEditChange(const UProperty* InProperty) const
 
 	if ((InProperty->GetFName() == GET_MEMBER_NAME_CHECKED(URendererSettings, bEnableRayTracing)))
 	{
-#if RHI_RAYTRACING
 		//only allow enabling raytracing if skin cache is on, but allow disable of raytracing no matter what.
 		return ParentVal && (bSupportSkinCacheShaders || bEnableRayTracing);
-#else // RHI_RAYTRACING
-		return false;
-#endif // RHI_RAYTRACING
 	}
 
 	return ParentVal;
