DR - Minor rename

diff --git a/Engine/Source/Runtime/Engine/Private/PositionVertexBuffer.cpp b/Engine/Source/Runtime/Engine/Private/PositionVertexBuffer.cpp
--- a/Engine/Source/Runtime/Engine/Private/PositionVertexBuffer.cpp
+++ b/Engine/Source/Runtime/Engine/Private/PositionVertexBuffer.cpp
@@ -47,13 +47,13 @@ void FPositionVertexBuffer::CleanUp()
 	}
 }
 
-void FPositionVertexBuffer::Init(uint32 InNumVertices, bool bNeedsCPUAccess)
+void FPositionVertexBuffer::Init(uint32 InNumVertices, bool bInNeedsCPUAccess)
 {
 	NumVertices = InNumVertices;
-	NeedsCPUAccess = bNeedsCPUAccess;
+	bNeedsCPUAccess = bInNeedsCPUAccess;
 
 	// Allocate the vertex data storage type.
-	AllocateData(bNeedsCPUAccess);
+	AllocateData(bInNeedsCPUAccess);
 
 	// Allocate the vertex data buffer.
 	VertexData->ResizeBuffer(NumVertices);
@@ -64,9 +64,9 @@ void FPositionVertexBuffer::Init(uint32 InNumVertices, bool bNeedsCPUAccess)
 * Initializes the buffer with the given vertices, used to convert legacy layouts.
 * @param InVertices - The vertices to initialize the buffer with.
 */
-void FPositionVertexBuffer::Init(const TArray<FStaticMeshBuildVertex>& InVertices, bool bNeedsCPUAccess)
+void FPositionVertexBuffer::Init(const TArray<FStaticMeshBuildVertex>& InVertices, bool bInNeedsCPUAccess)
 {
-	Init(InVertices.Num(), bNeedsCPUAccess);
+	Init(InVertices.Num(), bInNeedsCPUAccess);
 
 	// Copy the vertices into the buffer.
 	for(int32 VertexIndex = 0;VertexIndex < InVertices.Num();VertexIndex++)
@@ -81,12 +81,12 @@ void FPositionVertexBuffer::Init(const TArray<FStaticMeshBuildVertex>& InVertice
 * Initializes this vertex buffer with the contents of the given vertex buffer.
 * @param InVertexBuffer - The vertex buffer to initialize from.
 */
-void FPositionVertexBuffer::Init(const FPositionVertexBuffer& InVertexBuffer, bool bNeedsCPUAccess)
+void FPositionVertexBuffer::Init(const FPositionVertexBuffer& InVertexBuffer, bool bInNeedsCPUAccess)
 {
-	NeedsCPUAccess = bNeedsCPUAccess;
+	bNeedsCPUAccess = bInNeedsCPUAccess;
 	if ( InVertexBuffer.GetNumVertices() )
 	{
-		Init(InVertexBuffer.GetNumVertices(), bNeedsCPUAccess);
+		Init(InVertexBuffer.GetNumVertices(), bInNeedsCPUAccess);
 
 		check( Stride == InVertexBuffer.GetStride() );
 
@@ -95,13 +95,13 @@ void FPositionVertexBuffer::Init(const FPositionVertexBuffer& InVertexBuffer, bo
 	}
 }
 
-void FPositionVertexBuffer::Init(const TArray<FVector>& InPositions, bool bNeedsCPUAccess)
+void FPositionVertexBuffer::Init(const TArray<FVector>& InPositions, bool bInNeedsCPUAccess)
 {
 	NumVertices = InPositions.Num();
-	NeedsCPUAccess = bNeedsCPUAccess;
+	bNeedsCPUAccess = bInNeedsCPUAccess;
 	if ( NumVertices )
 	{
-		AllocateData(bNeedsCPUAccess);
+		AllocateData(bInNeedsCPUAccess);
 		check( Stride == InPositions.GetTypeSize() );
 		VertexData->ResizeBuffer(NumVertices);
 		Data = VertexData->GetDataPointer();
@@ -114,7 +114,7 @@ void FPositionVertexBuffer::AppendVertices( const FStaticMeshBuildVertex* Vertic
 	if (VertexData == nullptr && NumVerticesToAppend > 0)
 	{
 		// Allocate the vertex data storage type if the buffer was never allocated before
-		AllocateData(NeedsCPUAccess);
+		AllocateData(bNeedsCPUAccess);
 	}
 
 	if( NumVerticesToAppend > 0 )
@@ -144,19 +144,19 @@ void FPositionVertexBuffer::AppendVertices( const FStaticMeshBuildVertex* Vertic
 /**
  * Serializer
  *
- * @param	Ar				Archive to serialize with
- * @param	bNeedsCPUAccess	Whether the elements need to be accessed by the CPU
+ * @param	Ar					Archive to serialize with
+ * @param	bInNeedsCPUAccess	Whether the elements need to be accessed by the CPU
  */
-void FPositionVertexBuffer::Serialize( FArchive& Ar, bool bNeedsCPUAccess )
+void FPositionVertexBuffer::Serialize( FArchive& Ar, bool bInNeedsCPUAccess )
 {
-	NeedsCPUAccess = bNeedsCPUAccess;
+	bNeedsCPUAccess = bInNeedsCPUAccess;
 
 	Ar << Stride << NumVertices;
 
 	if(Ar.IsLoading())
 	{
 		// Allocate the vertex data storage type.
-		AllocateData( bNeedsCPUAccess );
+		AllocateData( bInNeedsCPUAccess );
 	}
 
 	if(VertexData != NULL)
@@ -203,12 +203,12 @@ void FPositionVertexBuffer::ReleaseRHI()
 	FVertexBuffer::ReleaseRHI();
 }
 
-void FPositionVertexBuffer::AllocateData( bool bNeedsCPUAccess /*= true*/ )
+void FPositionVertexBuffer::AllocateData( bool bInNeedsCPUAccess /*= true*/ )
 {
 	// Clear any old VertexData before allocating.
 	CleanUp();
 
-	VertexData = new FPositionVertexData(bNeedsCPUAccess);
+	VertexData = new FPositionVertexData(bInNeedsCPUAccess);
 	// Calculate the vertex stride.
 	Stride = VertexData->GetStride();
 }
diff --git a/Engine/Source/Runtime/Engine/Public/Rendering/PositionVertexBuffer.h b/Engine/Source/Runtime/Engine/Public/Rendering/PositionVertexBuffer.h
--- a/Engine/Source/Runtime/Engine/Public/Rendering/PositionVertexBuffer.h
+++ b/Engine/Source/Runtime/Engine/Public/Rendering/PositionVertexBuffer.h
@@ -32,21 +32,21 @@ public:
 	/** Delete existing resources */
 	ENGINE_API void CleanUp();
 
-	void ENGINE_API Init(uint32 NumVertices, bool bNeedsCPUAccess = true);
+	void ENGINE_API Init(uint32 NumVertices, bool bInNeedsCPUAccess = true);
 
 	/**
 	* Initializes the buffer with the given vertices, used to convert legacy layouts.
 	* @param InVertices - The vertices to initialize the buffer with.
 	*/
-	ENGINE_API void Init(const TArray<FStaticMeshBuildVertex>& InVertices, bool bNeedsCPUAccess = true);
+	ENGINE_API void Init(const TArray<FStaticMeshBuildVertex>& InVertices, bool bInNeedsCPUAccess = true);
 
 	/**
 	* Initializes this vertex buffer with the contents of the given vertex buffer.
 	* @param InVertexBuffer - The vertex buffer to initialize from.
 	*/
-	void Init(const FPositionVertexBuffer& InVertexBuffer, bool bNeedsCPUAccess = true);
+	void Init(const FPositionVertexBuffer& InVertexBuffer, bool bInNeedsCPUAccess = true);
 
-	ENGINE_API void Init(const TArray<FVector>& InPositions, bool bNeedsCPUAccess = true);
+	ENGINE_API void Init(const TArray<FVector>& InPositions, bool bInNeedsCPUAccess = true);
 
 	/**
 	 * Appends the specified vertices to the end of the buffer
@@ -59,10 +59,10 @@ public:
 	/**
 	* Serializer
 	*
-	* @param	Ar				Archive to serialize with
-	* @param	bNeedsCPUAccess	Whether the elements need to be accessed by the CPU
+	* @param	Ar					Archive to serialize with
+	* @param	bInNeedsCPUAccess	Whether the elements need to be accessed by the CPU
 	*/
-	void Serialize(FArchive& Ar, bool bNeedsCPUAccess);
+	void Serialize(FArchive& Ar, bool bInNeedsCPUAccess);
 
 	/**
 	* Specialized assignment operator, only used when importing LOD's.
@@ -118,8 +118,8 @@ private:
 	/** The cached number of vertices. */
 	uint32 NumVertices;
 
-	bool NeedsCPUAccess = true;
+	bool bNeedsCPUAccess = true;
 
 	/** Allocates the vertex data storage type. */
-	void AllocateData(bool bNeedsCPUAccess = true);
+	void AllocateData(bool bInNeedsCPUAccess = true);
 };
