Make it a little harder to get raytracing into a bad state

diff --git a/Engine/Source/Runtime/Engine/Classes/Engine/RendererSettings.h b/Engine/Source/Runtime/Engine/Classes/Engine/RendererSettings.h
--- a/Engine/Source/Runtime/Engine/Classes/Engine/RendererSettings.h
+++ b/Engine/Source/Runtime/Engine/Classes/Engine/RendererSettings.h
@@ -613,6 +613,7 @@ class ENGINE_API URendererSettings : public UDeveloperSettings
 	*/
 	UPROPERTY(config, EditAnywhere, Category = Optimizations, meta = (
 		ConsoleVariable = "r.SkinCache.CompileShaders", DisplayName = "Support Compute Skincache",
+		ToolTip = "Cannot be disabled while Ray Tracing is enabled as it is then required.",
 		ConfigRestartRequired = true))
 		uint32 bSupportSkinCacheShaders : 1;
 
diff --git a/Engine/Source/Runtime/Engine/Private/RendererSettings.cpp b/Engine/Source/Runtime/Engine/Private/RendererSettings.cpp
--- a/Engine/Source/Runtime/Engine/Private/RendererSettings.cpp
+++ b/Engine/Source/Runtime/Engine/Private/RendererSettings.cpp
@@ -122,6 +122,12 @@ bool URendererSettings::CanEditChange(const UProperty* InProperty) const
 		return ParentVal && (bSupportSkinCacheShaders || bEnableRayTracing);
 	}
 
+	if ((InProperty->GetFName() == GET_MEMBER_NAME_CHECKED(URendererSettings, bSupportSkinCacheShaders)))
+	{
+		//only allow DISABLE of skincache shaders if raytracing is also disabled as skincache is a dependency of raytracing.
+		return ParentVal && (!bSupportSkinCacheShaders || !bEnableRayTracing);
+	}
+
 	return ParentVal;
 }
 #endif // #if WITH_EDITOR
