PR #5359: Fix crash in UCanvasRenderTarget2D::RepaintCanvas (Contributed by [REDACTED_PERSON])

diff --git a/Engine/Source/Runtime/Engine/Private/CanvasRenderTarget2D.cpp b/Engine/Source/Runtime/Engine/Private/CanvasRenderTarget2D.cpp
--- a/Engine/Source/Runtime/Engine/Private/CanvasRenderTarget2D.cpp
+++ b/Engine/Source/Runtime/Engine/Private/CanvasRenderTarget2D.cpp
@@ -75,7 +75,7 @@ void UCanvasRenderTarget2D::RepaintCanvas()
 	// Enqueue the rendering command to set up the rendering canvas.
 	bool bClearRenderTarget = bShouldClearRenderTargetOnReceiveUpdate;
 	ENQUEUE_RENDER_COMMAND(CanvasRenderTargetMakeCurrentCommand)(
-		[&TextureRenderTarget, bClearRenderTarget](FRHICommandListImmediate& RHICmdList)
+		[TextureRenderTarget, bClearRenderTarget](FRHICommandListImmediate& RHICmdList)
 	{
 		RHICmdList.TransitionResource(EResourceTransitionAccess::EWritable, TextureRenderTarget->GetRenderTargetTexture());
 
