Attempt to fix static analysis

diff --git a/Engine/Source/Runtime/Engine/Private/DebugViewModeMaterialProxy.cpp b/Engine/Source/Runtime/Engine/Private/DebugViewModeMaterialProxy.cpp
--- a/Engine/Source/Runtime/Engine/Private/DebugViewModeMaterialProxy.cpp
+++ b/Engine/Source/Runtime/Engine/Private/DebugViewModeMaterialProxy.cpp
@@ -115,35 +115,39 @@ void FDebugViewModeMaterialProxy::ValidateAllShaders(TSet<UMaterialInterface*>&
 	{
 		const UMaterialInterface* OriginalMaterialInterface = It.Key().MaterialInterface;
 		FDebugViewModeMaterialProxy* DebugMaterial = It.Value();
-		const FMaterial* OriginalMaterial = OriginalMaterialInterface->GetMaterialResource(DebugMaterial->FMaterial::GetFeatureLevel());
 
-		if (OriginalMaterial && DebugMaterial && OriginalMaterial->GetGameThreadShaderMap() && DebugMaterial->GetGameThreadShaderMap())
+		if (DebugMaterial != nullptr)
 		{
-			const FUniformExpressionSet& DebugViewUniformExpressionSet = DebugMaterial->GetGameThreadShaderMap()->GetUniformExpressionSet();
-			const FUniformExpressionSet& OrignialUniformExpressionSet = OriginalMaterial->GetGameThreadShaderMap()->GetUniformExpressionSet();
+			const FMaterial* OriginalMaterial = OriginalMaterialInterface->GetMaterialResource(DebugMaterial->FMaterial::GetFeatureLevel());
 
-			if (!(DebugViewUniformExpressionSet == OrignialUniformExpressionSet))
+			if (OriginalMaterial != nullptr && OriginalMaterial->GetGameThreadShaderMap() && DebugMaterial->GetGameThreadShaderMap())
 			{
-				// This will happen when the debug shader compiled misses logic. Usually caused by custom features in the original shader compilation not implemented in FDebugViewModeMaterialProxy.
-				UE_LOG(TextureStreamingBuild, Verbose, TEXT("Uniform expression set mismatch for %s, skipping shader"), *DebugMaterial->GetMaterialInterface()->GetName());
-
-				// Here we can't destroy the invalid material because it would trigger ClearAllShaders.
-				DebugMaterial->MarkAsInvalid();
-				Materials.Remove(const_cast<UMaterialInterface*>(DebugMaterial->GetMaterialInterface()));
+				const FUniformExpressionSet& DebugViewUniformExpressionSet = DebugMaterial->GetGameThreadShaderMap()->GetUniformExpressionSet();
+				const FUniformExpressionSet& OrignialUniformExpressionSet = OriginalMaterial->GetGameThreadShaderMap()->GetUniformExpressionSet();
+
+				if (!(DebugViewUniformExpressionSet == OrignialUniformExpressionSet))
+				{
+					// This will happen when the debug shader compiled misses logic. Usually caused by custom features in the original shader compilation not implemented in FDebugViewModeMaterialProxy.
+					UE_LOG(TextureStreamingBuild, Verbose, TEXT("Uniform expression set mismatch for %s, skipping shader"), *DebugMaterial->GetMaterialInterface()->GetName());
+
+					// Here we can't destroy the invalid material because it would trigger ClearAllShaders.
+					DebugMaterial->MarkAsInvalid();
+					Materials.Remove(const_cast<UMaterialInterface*>(DebugMaterial->GetMaterialInterface()));
+				}
 			}
-		}
-		else if (DebugMaterial)
-		{
-			// When using synchronous compilation, it is normal for the original material to not be ready yet.
-			// In this case, we can't validate that the shader will be 100% compatible for overrides, meaning it is risky to use for viewmodes.
-			// This implies that viewmode can't use synchronous compilation.
-			if (!DebugMaterial->GetGameThreadShaderMap() || !DebugMaterial->RequiresSynchronousCompilation())
+			else
 			{
-				UE_LOG(TextureStreamingBuild, Verbose, TEXT("Can't get valid shadermap for %s, skipping shader"), *DebugMaterial->GetMaterialInterface()->GetName());
-
-				// Here we can't destroy the invalid material because it would trigger ClearAllShaders.
-				DebugMaterial->MarkAsInvalid();
-				Materials.Remove(const_cast<UMaterialInterface*>(DebugMaterial->GetMaterialInterface()));
+				// When using synchronous compilation, it is normal for the original material to not be ready yet.
+				// In this case, we can't validate that the shader will be 100% compatible for overrides, meaning it is risky to use for viewmodes.
+				// This implies that viewmode can't use synchronous compilation.
+				if (!DebugMaterial->GetGameThreadShaderMap() || !DebugMaterial->RequiresSynchronousCompilation())
+				{
+					UE_LOG(TextureStreamingBuild, Verbose, TEXT("Can't get valid shadermap for %s, skipping shader"), *DebugMaterial->GetMaterialInterface()->GetName());
+
+					// Here we can't destroy the invalid material because it would trigger ClearAllShaders.
+					DebugMaterial->MarkAsInvalid();
+					Materials.Remove(const_cast<UMaterialInterface*>(DebugMaterial->GetMaterialInterface()));
+				}
 			}
 		}
 	}
diff --git a/Engine/Source/Runtime/Engine/Private/SkeletalRenderGPUSkin.cpp b/Engine/Source/Runtime/Engine/Private/SkeletalRenderGPUSkin.cpp
--- a/Engine/Source/Runtime/Engine/Private/SkeletalRenderGPUSkin.cpp
+++ b/Engine/Source/Runtime/Engine/Private/SkeletalRenderGPUSkin.cpp
@@ -292,7 +292,7 @@ static TAutoConsoleVariable<int32> CVarDeferSkeletalDynamicDataUpdateUntilGDME(
 void FSkeletalMeshObjectGPUSkin::UpdateDynamicData_RenderThread(FGPUSkinCache* GPUSkinCache, FRHICommandListImmediate& RHICmdList, FDynamicSkelMeshObjectDataGPUSkin* InDynamicData, FSceneInterface* Scene, uint32 FrameNumberToPrepare, uint32 RevisionNumber)
 {
 	SCOPE_CYCLE_COUNTER(STAT_GPUSkinUpdateRTTime);
-	check(InDynamicData);
+	check(InDynamicData != nullptr);
 	bool bMorphNeedsUpdate=false;
 	// figure out if the morphing vertex buffer needs to be updated. compare old vs new active morphs
 	bMorphNeedsUpdate = 
