Fix UModelComponent crash. Moved VertexFactoryUniformBuffer from FModelElement to FElementInfo, as FModelElement is destroy by the gameplay thread.

diff --git a/Engine/Source/Runtime/Engine/Private/ModelRender.cpp b/Engine/Source/Runtime/Engine/Private/ModelRender.cpp
--- a/Engine/Source/Runtime/Engine/Private/ModelRender.cpp
+++ b/Engine/Source/Runtime/Engine/Private/ModelRender.cpp
@@ -234,23 +234,14 @@ public:
 		InComponent->GetModel()->VertexBuffer.Buffers.InitModelVF(&VertexFactory);
 
 		OverrideOwnerName(NAME_BSP);
-		TIndirectArray<FModelElement>& SourceElements = Component->GetElements();
+		const TIndirectArray<FModelElement>& SourceElements = Component->GetElements();
 
 		Elements.Empty(SourceElements.Num());
 		for(int32 ElementIndex = 0;ElementIndex < SourceElements.Num();ElementIndex++)
 		{
-			FModelElement& SourceElement = SourceElements[ElementIndex];
-			FElementInfo* Element = new(Elements) FElementInfo(SourceElement, VertexFactory.GetType());
+			const FModelElement& SourceElement = SourceElements[ElementIndex];
+			FElementInfo* Element = new(Elements) FElementInfo(SourceElement, &VertexFactory, ElementIndex);
 			MaterialRelevance |= Element->GetMaterial()->GetRelevance(GetScene().GetFeatureLevel());
-
-			if (RHISupportsManualVertexFetch(GMaxRHIShaderPlatform))
-			{
-				ENQUEUE_RENDER_COMMAND(FModelSceneProxySetData)(
-					[UniformBufferPtr = &SourceElement.VertexFactoryUniformBuffer, VertexFactory = &VertexFactory, ElementIndex](FRHICommandListImmediate& RHICmdList)
-				{
-					*UniformBufferPtr = CreateLocalVFUniformBuffer(VertexFactory, ElementIndex, nullptr, 0);
-				});
-			}
 		}
 
 		bGoodCandidateForCachedShadowmap = CacheShadowDepthsFromPrimitivesUsingWPO() || !MaterialRelevance.bUsesWorldPositionOffset;
@@ -451,7 +442,7 @@ public:
 												BatchElement.NumPrimitives = NumIndices / 3;
 												BatchElement.MinVertexIndex = MinVertexIndex;
 												BatchElement.MaxVertexIndex = MaxVertexIndex;
-												BatchElement.VertexFactoryUserData = ModelElement.VertexFactoryUniformBuffer;
+												BatchElement.VertexFactoryUserData = ProxyElementInfo.GetVertexFactoryUniformBuffer();
 												MeshElement.Type = PT_TriangleList;
 												MeshElement.DepthPriorityGroup = DepthPriorityGroup;
 												MeshElement.bCanApplyViewModeOverrides = true;
@@ -487,7 +478,7 @@ public:
 								BatchElement.NumPrimitives = ModelElement.NumTriangles;
 								BatchElement.MinVertexIndex = ModelElement.MinVertexIndex;
 								BatchElement.MaxVertexIndex = ModelElement.MaxVertexIndex;
-								BatchElement.VertexFactoryUserData = ModelElement.VertexFactoryUniformBuffer;
+								BatchElement.VertexFactoryUserData = Elements[ElementIndex].GetVertexFactoryUniformBuffer();
 								MeshElement.Type = PT_TriangleList;
 								MeshElement.DepthPriorityGroup = DepthPriorityGroup;
 								MeshElement.bCanApplyViewModeOverrides = true;
@@ -530,7 +521,7 @@ public:
 					BatchElement.NumPrimitives = ModelElement.NumTriangles;
 					BatchElement.MinVertexIndex = ModelElement.MinVertexIndex;
 					BatchElement.MaxVertexIndex = ModelElement.MaxVertexIndex;
-					BatchElement.VertexFactoryUserData = ModelElement.VertexFactoryUniformBuffer;
+					BatchElement.VertexFactoryUserData = Elements[ElementIndex].GetVertexFactoryUniformBuffer();
 					MeshElement.Type = PT_TriangleList;
 					MeshElement.DepthPriorityGroup = PrimitiveDPG;
 					MeshElement.LODIndex = 0;
@@ -665,7 +656,7 @@ private:
 	public:
 
 		/** Initialization constructor. */
-		FElementInfo(const FModelElement& InModelElement, const FVertexFactoryType* VertexFactoryType)
+		FElementInfo(const FModelElement& InModelElement, const FLocalVertexFactory* InVertexFactory, int32 InElementIndex)
 			: FLightCacheInterface()
 			, Bounds(InModelElement.BoundingBox)
 		{
@@ -684,7 +675,7 @@ private:
 			// Determine the material applied to the model element.
 			Material = InModelElement.Material;
 
-			if (RequiresAdjacencyInformation(Material, VertexFactoryType, InModelElement.Component->GetScene()->GetFeatureLevel()))
+			if (RequiresAdjacencyInformation(Material, InVertexFactory->GetType(), InModelElement.Component->GetScene()->GetFeatureLevel()))
 			{
 				UE_LOG(LogModelComponent, Warning, TEXT("Material %s requires adjacency information because of Crack Free Displacement or PN Triangle Tesselation, which is not supported with model components. Falling back to DefaultMaterial."), *Material->GetName());
 				Material = nullptr;
@@ -695,6 +686,17 @@ private:
 			{
 				Material = UMaterial::GetDefaultMaterial(MD_Surface);
 			}
+
+			if (RHISupportsManualVertexFetch(GMaxRHIShaderPlatform))
+			{
+				TUniformBufferRef<FLocalVertexFactoryUniformShaderParameters>* UniformBufferPtr = &VertexFactoryUniformBuffer;
+
+				ENQUEUE_RENDER_COMMAND(FElementInfoCreateLocalVFUniformBuffer)(
+					[UniformBufferPtr, VertexFactory = InVertexFactory, ElementIndex = InElementIndex](FRHICommandListImmediate& RHICmdList)
+				{
+					*UniformBufferPtr = CreateLocalVFUniformBuffer(VertexFactory, ElementIndex, nullptr, 0);
+				});
+			}
 		}
 		
 		// FLightCacheInterface.
@@ -721,11 +723,16 @@ private:
 		// Accessors.
 		UMaterialInterface* GetMaterial() const { return Material; }
 
+		FUniformBufferRHIParamRef GetVertexFactoryUniformBuffer() const { return VertexFactoryUniformBuffer; }
+
 	private:
 
 		/** The element's material. */
 		UMaterialInterface* Material;
 
+		/** Vertex factory uniform buffer. Needs to be per model element, so every element can have it's unique LODLightmapDataIndex. */
+		TUniformBufferRef<FLocalVertexFactoryUniformShaderParameters> VertexFactoryUniformBuffer;
+
 		/** The statically irrelevant lights for this element. */
 		TArray<FGuid> IrrelevantLights;
 
diff --git a/Engine/Source/Runtime/Engine/Public/Model.h b/Engine/Source/Runtime/Engine/Public/Model.h
--- a/Engine/Source/Runtime/Engine/Public/Model.h
+++ b/Engine/Source/Runtime/Engine/Public/Model.h
@@ -623,9 +623,6 @@ public:
 
 	FMeshMapBuildData* LegacyMapBuildData;
 
-	/** Vertex factory uniform buffer. Needs to be per model element, so every element can have it's unique LODLightmapDataIndex. */
-	FUniformBufferRHIRef VertexFactoryUniformBuffer;
-
 	/** Uniquely identifies this component's built map data. */
 	FGuid MapBuildDataId;
 
