Change several shader compilation warnings to logs - they should not be warnings

diff --git a/Engine/Source/Runtime/Engine/Private/Materials/MaterialShader.cpp b/Engine/Source/Runtime/Engine/Private/Materials/MaterialShader.cpp
--- a/Engine/Source/Runtime/Engine/Private/Materials/MaterialShader.cpp
+++ b/Engine/Source/Runtime/Engine/Private/Materials/MaterialShader.cpp
@@ -1573,7 +1573,7 @@ void FMaterialShaderMap::Compile(
 				}
 			}
 
-			UE_LOG(LogShaders, Warning, TEXT("		%u Shaders among %u VertexFactories"), NumShaders, NumVertexFactories);
+			UE_LOG(LogShaders, Log, TEXT("		%u Shaders among %u VertexFactories"), NumShaders, NumVertexFactories);
 
 			// Register this shader map in the global map with the material's ID.
 			Register(InPlatform);
diff --git a/Engine/Source/Runtime/Engine/Private/ShaderCompiler/ShaderCompiler.cpp b/Engine/Source/Runtime/Engine/Private/ShaderCompiler/ShaderCompiler.cpp
--- a/Engine/Source/Runtime/Engine/Private/ShaderCompiler/ShaderCompiler.cpp
+++ b/Engine/Source/Runtime/Engine/Private/ShaderCompiler/ShaderCompiler.cpp
@@ -3763,7 +3763,7 @@ void VerifyGlobalShaders(EShaderPlatform Platform, bool bLoadedFromCacheFile)
 	const bool bEmptyMap = GlobalShaderMap->IsEmpty();
 	if (bEmptyMap)
 	{
-		UE_LOG(LogShaders, Warning, TEXT("	Empty global shader map, recompiling all global shaders"));
+		UE_LOG(LogShaders, Log, TEXT("	Empty global shader map, recompiling all global shaders"));
 	}
 
 	bool bErrorOnMissing = bLoadedFromCacheFile;
@@ -3815,7 +3815,7 @@ void VerifyGlobalShaders(EShaderPlatform Platform, bool bLoadedFromCacheFile)
 
 		if (!bEmptyMap && PermutationCountToCompile > 0)
 		{
-			UE_LOG(LogShaders, Warning, TEXT("	%s (%i out of %i)"),
+			UE_LOG(LogShaders, Log, TEXT("	%s (%i out of %i)"),
 				GlobalShaderType->GetName(), PermutationCountToCompile, GlobalShaderType->GetPermutationCount());
 		}
 	}
@@ -3852,7 +3852,7 @@ void VerifyGlobalShaders(EShaderPlatform Platform, bool bLoadedFromCacheFile)
 
 					if (!bEmptyMap)
 					{
-						UE_LOG(LogShaders, Warning, TEXT("	%s"), Pipeline->GetName());
+						UE_LOG(LogShaders, Log, TEXT("	%s"), Pipeline->GetName());
 					}
 
 					if (Pipeline->ShouldOptimizeUnusedOutputs())
