Removes workaround on the mobile permutation shader that had too many permutation.

diff --git a/Engine/Source/Runtime/Engine/Private/ShaderCompiler/ShaderCompiler.cpp b/Engine/Source/Runtime/Engine/Private/ShaderCompiler/ShaderCompiler.cpp
--- a/Engine/Source/Runtime/Engine/Private/ShaderCompiler/ShaderCompiler.cpp
+++ b/Engine/Source/Runtime/Engine/Private/ShaderCompiler/ShaderCompiler.cpp
@@ -3808,8 +3808,7 @@ void VerifyGlobalShaders(EShaderPlatform Platform, bool bLoadedFromCacheFile)
 		}
 
 		ensureMsgf(
-			PermutationCountToCompile < 200 ||
-			FCString::Strcmp(GlobalShaderType->GetName(), TEXT("FPostProcessTonemapPS_ES2")) == 0, // TODO: UE-58014
+			PermutationCountToCompile < 200,
 			TEXT("Global shader %s has %i permutation: probably more that it needs."),
 			GlobalShaderType->GetName(), PermutationCountToCompile);
 
