Fix primitiveuniformbuffer lifetime issues

diff --git a/Engine/Source/Runtime/Engine/Public/MeshBatch.h b/Engine/Source/Runtime/Engine/Public/MeshBatch.h
--- a/Engine/Source/Runtime/Engine/Public/MeshBatch.h
+++ b/Engine/Source/Runtime/Engine/Public/MeshBatch.h
@@ -44,7 +44,7 @@ struct FMeshBatchElement
 	 * Primitive uniform buffer RHI
 	 * Must be null for vertex factories that manually fetch primitive data from scene data, in which case FPrimitiveSceneProxy::UniformBuffer will be used.
 	 */
-	FUniformBufferRHIParamRef PrimitiveUniformBuffer;
+	FUniformBufferRHIRef PrimitiveUniformBuffer;
 
 	/** 
 	 * Primitive uniform buffer to use for rendering, used when PrimitiveUniformBuffer is null. 
