Cache texture-importance sampling with RectLightSceneProxy.

diff --git a/Engine/Shaders/Public/ShaderVersion.ush b/Engine/Shaders/Public/ShaderVersion.ush
--- a/Engine/Shaders/Public/ShaderVersion.ush
+++ b/Engine/Shaders/Public/ShaderVersion.ush
@@ -3,4 +3,4 @@
 // in Platform.ush (which should be included in any shader) it allows to invalidate the shader DDC.
 // 
 // If you are merging streams and there is a conflict with this GUID you should make a new GUID rather than taking one or the other.
-// GUID = 487F1C22-EBE7-4CEA-A565-E19FE49FE8EF
+// GUID = 2ED93466-374C-4558-B573-BC368C3718D7
diff --git a/Engine/Source/Runtime/Engine/Private/Components/RectLightComponent.cpp b/Engine/Source/Runtime/Engine/Private/Components/RectLightComponent.cpp
--- a/Engine/Source/Runtime/Engine/Private/Components/RectLightComponent.cpp
+++ b/Engine/Source/Runtime/Engine/Private/Components/RectLightComponent.cpp
@@ -10,81 +10,13 @@
 #include "Engine/Texture2D.h"
 #include "SceneManagement.h"
 #include "PointLightSceneProxy.h"
+#include "RectLightSceneProxy.h"
 
-extern int32 GAllowPointLightCubemapShadows;
-
-class FRectLightSceneProxy : public FLocalLightSceneProxy
-{
-public:
-	float		SourceWidth;
-	float		SourceHeight;
-	UTexture*	SourceTexture;
-
-	FRectLightSceneProxy(const URectLightComponent* Component)
-	:	FLocalLightSceneProxy(Component)
-	,	SourceWidth(Component->SourceWidth)
-	,	SourceHeight(Component->SourceHeight)
-	,	SourceTexture(Component->SourceTexture)
-	{}
-
-	virtual bool IsRectLight() const override
-	{
-		return true;
-	}
-
-	virtual bool HasSourceTexture() const override
-	{
-		return SourceTexture != nullptr;
-	}
-
-	/** Accesses parameters needed for rendering the light. */
-	virtual void GetLightShaderParameters(FLightShaderParameters& LightParameters) const override
-	{
-		FLinearColor LightColor = GetColor();
-		LightColor /= 0.5f * SourceWidth * SourceHeight;
-		
-		LightParameters.Position = GetOrigin();
-		LightParameters.InvRadius = InvRadius;
-		LightParameters.Color = FVector(LightColor.R, LightColor.G, LightColor.B);
-		LightParameters.FalloffExponent = 0.0f;
-
-		LightParameters.Direction = -GetDirection();
-		LightParameters.Tangent = FVector(WorldToLight.M[0][2], WorldToLight.M[1][2], WorldToLight.M[2][2]);
-		LightParameters.SpotAngles = FVector2D( -2.0f, 1.0f );
-		LightParameters.SpecularScale = SpecularScale;
-		LightParameters.SourceRadius = SourceWidth * 0.5f;
-		LightParameters.SoftSourceRadius = 0.0f;
-		LightParameters.SourceLength = SourceHeight * 0.5f;
-		LightParameters.SourceTexture = SourceTexture ? SourceTexture->Resource->TextureRHI : GWhiteTexture->TextureRHI;
-	}
-
-	/**
-	 * Sets up a projected shadow initializer for shadows from the entire scene.
-	 * @return True if the whole-scene projected shadow should be used.
-	 */
-	virtual bool GetWholeSceneProjectedShadowInitializer(const FSceneViewFamily& ViewFamily, TArray<FWholeSceneProjectedShadowInitializer, TInlineAllocator<6> >& OutInitializers) const
-	{
-		if (ViewFamily.GetFeatureLevel() >= ERHIFeatureLevel::SM4
-			&& GAllowPointLightCubemapShadows != 0)
-		{
-			FWholeSceneProjectedShadowInitializer& OutInitializer = *new(OutInitializers) FWholeSceneProjectedShadowInitializer;
-			OutInitializer.PreShadowTranslation = -GetLightToWorld().GetOrigin();
-			OutInitializer.WorldToLight = GetWorldToLight().RemoveTranslation();
-			OutInitializer.Scales = FVector(1, 1, 1);
-			OutInitializer.FaceDirection = FVector(0,0,1);
-			OutInitializer.SubjectBounds = FBoxSphereBounds(FVector(0, 0, 0),FVector(Radius,Radius,Radius),Radius);
-			OutInitializer.WAxis = FVector4(0,0,1,0);
-			OutInitializer.MinLightW = 0.1f;
-			OutInitializer.MaxDistanceToCastInLightW = Radius;
-			OutInitializer.bOnePassPointLightShadow = true;
-			OutInitializer.bRayTracedDistanceField = UseRayTracedDistanceFieldShadows() && DoesPlatformSupportDistanceFieldShadowing(ViewFamily.GetShaderPlatform());
-			return true;
-		}
-		
-		return false;
-	}
-};
+#include "RHIUtilities.h"
+#include "GlobalShader.h"
+#include "ShaderParameterUtils.h"
 
+extern int32 GAllowPointLightCubemapShadows;
 
 URectLightComponent::URectLightComponent(const FObjectInitializer& ObjectInitializer)
 	: Super(ObjectInitializer)
@@ -206,4 +138,205 @@ void URectLightComponent::PostEditChangeProperty(FPropertyChangedEvent& Property
 
 	Super::PostEditChangeProperty(PropertyChangedEvent);
 }
-#endif // WITH_EDITOR
\ No newline at end of file
+#endif // WITH_EDITOR
+
+FRectLightSceneProxy::FRectLightSceneProxy(const URectLightComponent* Component)
+	: FLocalLightSceneProxy(Component)
+	, SourceWidth(Component->SourceWidth)
+	, SourceHeight(Component->SourceHeight)
+	, SourceTexture(Component->SourceTexture)
+{
+#if RHI_RAYTRACING
+	ENQUEUE_RENDER_COMMAND(BuildRectLightMipTree)(
+		[this](FRHICommandListImmediate& RHICmdList)
+	{
+		BuildRectLightMipTree(RHICmdList);
+	});
+#endif
+}
+
+FRectLightSceneProxy::~FRectLightSceneProxy() {}
+
+bool FRectLightSceneProxy::IsRectLight() const
+{
+	return true;
+}
+
+bool FRectLightSceneProxy::HasSourceTexture() const
+{
+	return SourceTexture != nullptr;
+}
+
+/** Accesses parameters needed for rendering the light. */
+void FRectLightSceneProxy::GetLightShaderParameters(FLightShaderParameters& LightParameters) const
+{
+	FLinearColor LightColor = GetColor();
+	LightColor /= 0.5f * SourceWidth * SourceHeight;
+
+	LightParameters.Position = GetOrigin();
+	LightParameters.InvRadius = InvRadius;
+	LightParameters.Color = FVector(LightColor.R, LightColor.G, LightColor.B);
+	LightParameters.FalloffExponent = 0.0f;
+
+	LightParameters.Direction = -GetDirection();
+	LightParameters.Tangent = FVector(WorldToLight.M[0][2], WorldToLight.M[1][2], WorldToLight.M[2][2]);
+	LightParameters.SpotAngles = FVector2D(-2.0f, 1.0f);
+	LightParameters.SpecularScale = SpecularScale;
+	LightParameters.SourceRadius = SourceWidth * 0.5f;
+	LightParameters.SoftSourceRadius = 0.0f;
+	LightParameters.SourceLength = SourceHeight * 0.5f;
+	LightParameters.SourceTexture = SourceTexture ? SourceTexture->Resource->TextureRHI : GWhiteTexture->TextureRHI;
+}
+
+/**
+* Sets up a projected shadow initializer for shadows from the entire scene.
+* @return True if the whole-scene projected shadow should be used.
+*/
+bool FRectLightSceneProxy::GetWholeSceneProjectedShadowInitializer(const FSceneViewFamily& ViewFamily, TArray<FWholeSceneProjectedShadowInitializer, TInlineAllocator<6> >& OutInitializers) const
+{
+	if (ViewFamily.GetFeatureLevel() >= ERHIFeatureLevel::SM4
+		&& GAllowPointLightCubemapShadows != 0)
+	{
+		FWholeSceneProjectedShadowInitializer& OutInitializer = *new(OutInitializers) FWholeSceneProjectedShadowInitializer;
+		OutInitializer.PreShadowTranslation = -GetLightToWorld().GetOrigin();
+		OutInitializer.WorldToLight = GetWorldToLight().RemoveTranslation();
+		OutInitializer.Scales = FVector(1, 1, 1);
+		OutInitializer.FaceDirection = FVector(0, 0, 1);
+		OutInitializer.SubjectBounds = FBoxSphereBounds(FVector(0, 0, 0), FVector(Radius, Radius, Radius), Radius);
+		OutInitializer.WAxis = FVector4(0, 0, 1, 0);
+		OutInitializer.MinLightW = 0.1f;
+		OutInitializer.MaxDistanceToCastInLightW = Radius;
+		OutInitializer.bOnePassPointLightShadow = true;
+		OutInitializer.bRayTracedDistanceField = UseRayTracedDistanceFieldShadows() && DoesPlatformSupportDistanceFieldShadowing(ViewFamily.GetShaderPlatform());
+		return true;
+	}
+
+	return false;
+}
+
+#if RHI_RAYTRACING
+
+class FBuildRectLightMipTreeCS : public FGlobalShader
+{
+	DECLARE_SHADER_TYPE(FBuildRectLightMipTreeCS, Global)
+
+public:
+	static bool ShouldCompilePermutation(const FGlobalShaderPermutationParameters& Parameters)
+	{
+		return RHISupportsComputeShaders(Parameters.Platform);
+	}
+
+	static uint32 GetGroupSize()
+	{
+		return 16;
+	}
+
+	static void ModifyCompilationEnvironment(const FGlobalShaderPermutationParameters& Parameters, FShaderCompilerEnvironment& OutEnvironment)
+	{
+		FGlobalShader::ModifyCompilationEnvironment(Parameters, OutEnvironment);
+		OutEnvironment.SetDefine(TEXT("THREADGROUP_SIZE"), GetGroupSize());
+	}
+
+	FBuildRectLightMipTreeCS(const ShaderMetaType::CompiledShaderInitializerType& Initializer)
+		: FGlobalShader(Initializer)
+	{
+		TextureParameter.Bind(Initializer.ParameterMap, TEXT("RectLightTexture"));
+		TextureSamplerParameter.Bind(Initializer.ParameterMap, TEXT("TextureSampler"));
+		DimensionsParameter.Bind(Initializer.ParameterMap, TEXT("Dimensions"));
+		MipLevelParameter.Bind(Initializer.ParameterMap, TEXT("MipLevel"));
+		MipTreeParameter.Bind(Initializer.ParameterMap, TEXT("MipTree"));
+	}
+
+	FBuildRectLightMipTreeCS() {}
+
+	void SetParameters(
+		FRHICommandList& RHICmdList,
+		FTextureRHIRef Texture,
+		const FIntVector& Dimensions,
+		uint32 MipLevel,
+		FRWBuffer& MipTree
+	)
+	{
+		FComputeShaderRHIParamRef ShaderRHI = GetComputeShader();
+
+		SetShaderValue(RHICmdList, ShaderRHI, DimensionsParameter, Dimensions);
+		SetShaderValue(RHICmdList, ShaderRHI, MipLevelParameter, MipLevel);
+		SetTextureParameter(RHICmdList, ShaderRHI, TextureParameter, TextureSamplerParameter, TStaticSamplerState<SF_Bilinear>::GetRHI(), Texture);
+
+		check(MipTreeParameter.IsBound());
+		MipTreeParameter.SetBuffer(RHICmdList, ShaderRHI, MipTree);
+	}
+
+	void UnsetParameters(
+		FRHICommandList& RHICmdList,
+		EResourceTransitionAccess TransitionAccess,
+		EResourceTransitionPipeline TransitionPipeline,
+		FRWBuffer& MipTree,
+		FComputeFenceRHIParamRef Fence)
+	{
+		FComputeShaderRHIParamRef ShaderRHI = GetComputeShader();
+
+		MipTreeParameter.UnsetUAV(RHICmdList, ShaderRHI);
+		RHICmdList.TransitionResource(TransitionAccess, TransitionPipeline, MipTree.UAV, Fence);
+	}
+
+	virtual bool Serialize(FArchive& Ar)
+	{
+		bool bShaderHasOutdatedParameters = FGlobalShader::Serialize(Ar);
+		Ar << TextureParameter;
+		Ar << TextureSamplerParameter;
+		Ar << DimensionsParameter;
+		Ar << MipLevelParameter;
+		Ar << MipTreeParameter;
+		return bShaderHasOutdatedParameters;
+	}
+
+private:
+	FShaderResourceParameter TextureParameter;
+	FShaderResourceParameter TextureSamplerParameter;
+
+	FShaderParameter DimensionsParameter;
+	FShaderParameter MipLevelParameter;
+	FRWShaderParameter MipTreeParameter;
+};
+
+IMPLEMENT_SHADER_TYPE(, FBuildRectLightMipTreeCS, TEXT("/Engine/Private/Raytracing/BuildMipTreeCS.usf"), TEXT("BuildRectLightMipTreeCS"), SF_Compute)
+
+void FRectLightSceneProxy::BuildRectLightMipTree(FRHICommandListImmediate& RHICmdList)
+{
+	check(IsInRenderingThread());
+
+	const auto ShaderMap = GetGlobalShaderMap(ERHIFeatureLevel::SM5);
+	TShaderMapRef<FBuildRectLightMipTreeCS> BuildRectLightMipTreeComputeShader(ShaderMap);
+	RHICmdList.SetComputeShader(BuildRectLightMipTreeComputeShader->GetComputeShader());
+
+	// Allocate MIP tree
+	FLightShaderParameters LightParameters;
+	GetLightShaderParameters(LightParameters);
+	FIntVector TextureSize = LightParameters.SourceTexture->GetSizeXYZ();
+	uint32 MipLevelCount = FMath::Min(FMath::CeilLogTwo(TextureSize.X), FMath::CeilLogTwo(TextureSize.Y));
+	RectLightMipTreeDimensions = FIntVector(1 << MipLevelCount, 1 << MipLevelCount, 1);
+	uint32 NumElements = RectLightMipTreeDimensions.X * RectLightMipTreeDimensions.Y;
+	for (uint32 MipLevel = 1; MipLevel <= MipLevelCount; ++MipLevel)
+	{
+		uint32 NumElementsInLevel = (RectLightMipTreeDimensions.X >> MipLevel) * (RectLightMipTreeDimensions.Y >> MipLevel);
+		NumElements += NumElementsInLevel;
+	}
+
+	RectLightMipTree.Initialize(sizeof(float), NumElements, PF_R32_FLOAT, BUF_UnorderedAccess | BUF_ShaderResource);
+
+	// Execute hierarchical build
+	for (uint32 MipLevel = 0; MipLevel <= MipLevelCount; ++MipLevel)
+	{
+		FComputeFenceRHIRef MipLevelFence = RHICmdList.CreateComputeFence(TEXT("RectLightMipTree Build"));
+		BuildRectLightMipTreeComputeShader->SetParameters(RHICmdList, LightParameters.SourceTexture, RectLightMipTreeDimensions, MipLevel, RectLightMipTree);
+		FIntVector MipLevelDimensions = FIntVector(RectLightMipTreeDimensions.X >> MipLevel, RectLightMipTreeDimensions.Y >> MipLevel, 1);
+		FIntVector NumGroups = FIntVector::DivideAndRoundUp(MipLevelDimensions, FBuildRectLightMipTreeCS::GetGroupSize());
+		DispatchComputeShader(RHICmdList, *BuildRectLightMipTreeComputeShader, NumGroups.X, NumGroups.Y, 1);
+		BuildRectLightMipTreeComputeShader->UnsetParameters(RHICmdList, EResourceTransitionAccess::ERWBarrier, EResourceTransitionPipeline::EComputeToCompute, RectLightMipTree, MipLevelFence);
+	}
+	FComputeFenceRHIRef TransitionFence = RHICmdList.CreateComputeFence(TEXT("RectLightMipTree Transition"));
+	BuildRectLightMipTreeComputeShader->UnsetParameters(RHICmdList, EResourceTransitionAccess::ERWBarrier, EResourceTransitionPipeline::EComputeToCompute, RectLightMipTree, TransitionFence);
+}
+
+#endif // RHI_RAYTRACING
\ No newline at end of file
diff --git a/Engine/Source/Runtime/Engine/Private/LocalLightSceneProxy.cpp b/Engine/Source/Runtime/Engine/Private/LocalLightSceneProxy.cpp
new file mode 100644
--- /dev/null
+++ b/Engine/Source/Runtime/Engine/Private/LocalLightSceneProxy.cpp
@@ -0,0 +1,139 @@
+// Copyright 1998-2019 Epic Games, Inc. All Rights Reserved.
+
+/*=============================================================================
+	LocalLightSceneProxy.cpp: Local light scene info implementation.
+=============================================================================*/
+
+#include "LocalLightSceneProxy.h"
+#include "Components/LocalLightComponent.h"
+
+/** Initialization constructor. */
+FLocalLightSceneProxy::FLocalLightSceneProxy(const ULocalLightComponent* Component)
+	: FLightSceneProxy(Component)
+	, MaxDrawDistance(Component->MaxDrawDistance)
+	, FadeRange(Component->MaxDistanceFadeRange)
+{
+	UpdateRadius(Component->AttenuationRadius);
+}
+
+/**
+* Called on the light scene info after it has been passed to the rendering thread to update the rendering thread's cached info when
+* the light's radius changes.
+*/
+void UpdateRadius_GameThread(float Radius);
+
+// FLightSceneInfo interface.
+float FLocalLightSceneProxy::GetMaxDrawDistance() const
+{
+	return MaxDrawDistance;
+}
+
+float FLocalLightSceneProxy::GetFadeRange() const
+{
+	return FadeRange;
+}
+
+/** @return radius of the light or 0 if no radius */
+float FLocalLightSceneProxy::GetRadius() const
+{
+	return Radius;
+}
+
+bool FLocalLightSceneProxy::AffectsBounds(const FBoxSphereBounds& Bounds) const
+{
+	if ((Bounds.Origin - GetLightToWorld().GetOrigin()).SizeSquared() > FMath::Square(Radius + Bounds.SphereRadius))
+	{
+		return false;
+	}
+
+	if (!FLightSceneProxy::AffectsBounds(Bounds))
+	{
+		return false;
+	}
+
+	return true;
+}
+
+bool FLocalLightSceneProxy::GetScissorRect(FIntRect& ScissorRect, const FSceneView& View, const FIntRect& ViewRect) const
+{
+	ScissorRect = ViewRect;
+	return FMath::ComputeProjectedSphereScissorRect(ScissorRect, GetLightToWorld().GetOrigin(), Radius, View.ViewMatrices.GetViewOrigin(), View.ViewMatrices.GetViewMatrix(), View.ViewMatrices.GetProjectionMatrix()) == 1;
+}
+
+void FLocalLightSceneProxy::SetScissorRect(FRHICommandList& RHICmdList, const FSceneView& View, const FIntRect& ViewRect) const
+{
+	FIntRect ScissorRect;
+
+	if (GetScissorRect(ScissorRect, View, ViewRect))
+	{
+		RHICmdList.SetScissorRect(true, ScissorRect.Min.X, ScissorRect.Min.Y, ScissorRect.Max.X, ScissorRect.Max.Y);
+	}
+	else
+	{
+		RHICmdList.SetScissorRect(false, 0, 0, 0, 0);
+	}
+}
+
+FSphere FLocalLightSceneProxy::GetBoundingSphere() const
+{
+	return FSphere(GetPosition(), GetRadius());
+}
+
+float FLocalLightSceneProxy::GetEffectiveScreenRadius(const FViewMatrices& ShadowViewMatrices) const
+{
+	// Use the distance from the view origin to the light to approximate perspective projection
+	// We do not use projected screen position since it causes problems when the light is behind the camera
+
+	const float LightDistance = (GetOrigin() - ShadowViewMatrices.GetViewOrigin()).Size();
+
+	return ShadowViewMatrices.GetScreenScale() * GetRadius() / FMath::Max(LightDistance, 1.0f);
+}
+
+FVector FLocalLightSceneProxy::GetPerObjectProjectedShadowProjectionPoint(const FBoxSphereBounds& SubjectBounds) const
+{
+	return GetOrigin();
+}
+
+bool FLocalLightSceneProxy::GetPerObjectProjectedShadowInitializer(const FBoxSphereBounds& SubjectBounds, class FPerObjectProjectedShadowInitializer& OutInitializer) const
+{
+	// Use a perspective projection looking at the primitive from the light position.
+	FVector LightPosition = GetPerObjectProjectedShadowProjectionPoint(SubjectBounds);
+	FVector LightVector = SubjectBounds.Origin - LightPosition;
+	float LightDistance = LightVector.Size();
+	float SilhouetteRadius = 1.0f;
+	const float SubjectRadius = SubjectBounds.BoxExtent.Size();
+	const float ShadowRadiusMultiplier = 1.1f;
+
+	if (LightDistance > SubjectRadius)
+	{
+		SilhouetteRadius = FMath::Min(SubjectRadius * FMath::InvSqrt((LightDistance - SubjectRadius) * (LightDistance + SubjectRadius)), 1.0f);
+	}
+
+	if (LightDistance <= SubjectRadius * ShadowRadiusMultiplier)
+	{
+		// Make the primitive fit in a single < 90 degree FOV projection.
+		LightVector = SubjectRadius * LightVector.GetSafeNormal() * ShadowRadiusMultiplier;
+		LightPosition = (SubjectBounds.Origin - LightVector);
+		LightDistance = SubjectRadius * ShadowRadiusMultiplier;
+		SilhouetteRadius = 1.0f;
+	}
+
+	OutInitializer.PreShadowTranslation = -LightPosition;
+	OutInitializer.WorldToLight = FInverseRotationMatrix((LightVector / LightDistance).Rotation());
+	OutInitializer.Scales = FVector(1.0f, 1.0f / SilhouetteRadius, 1.0f / SilhouetteRadius);
+	OutInitializer.FaceDirection = FVector(1, 0, 0);
+	OutInitializer.SubjectBounds = FBoxSphereBounds(SubjectBounds.Origin - LightPosition, SubjectBounds.BoxExtent, SubjectBounds.SphereRadius);
+	OutInitializer.WAxis = FVector4(0, 0, 1, 0);
+	OutInitializer.MinLightW = 0.1f;
+	OutInitializer.MaxDistanceToCastInLightW = Radius;
+	return true;
+}
+
+/** Updates the light scene info's radius from the component. */
+void FLocalLightSceneProxy::UpdateRadius(float ComponentRadius)
+{
+	Radius = ComponentRadius;
+
+	// Min to avoid div by 0 (NaN in InvRadius)
+	InvRadius = 1.0f / FMath::Max(0.00001f, ComponentRadius);
+}
\ No newline at end of file
diff --git a/Engine/Source/Runtime/Engine/Private/PointLightSceneProxy.h b/Engine/Source/Runtime/Engine/Private/PointLightSceneProxy.h
--- a/Engine/Source/Runtime/Engine/Private/PointLightSceneProxy.h
+++ b/Engine/Source/Runtime/Engine/Private/PointLightSceneProxy.h
@@ -8,157 +8,9 @@
 
 #include "CoreMinimal.h"
 #include "Components/PointLightComponent.h"
+#include "LocalLightSceneProxy.h"
 #include "SceneManagement.h"
 
-/** The parts of the point light scene info that aren't dependent on the light policy type. */
-class FLocalLightSceneProxy : public FLightSceneProxy
-{
-public:
-
-	/** The light radius. */
-	float Radius;
-
-	/** One over the light's radius. */
-	float InvRadius;
-
-	/** Initialization constructor. */
-	FLocalLightSceneProxy(const ULocalLightComponent* Component)
-	:	FLightSceneProxy(Component)
-	,	MaxDrawDistance(Component->MaxDrawDistance)
-	,	FadeRange(Component->MaxDistanceFadeRange)
-	{
-		UpdateRadius(Component->AttenuationRadius);
-	}
-
-	/**
-	* Called on the light scene info after it has been passed to the rendering thread to update the rendering thread's cached info when
-	* the light's radius changes.
-	*/
-	void UpdateRadius_GameThread(float Radius);
-
-	// FLightSceneInfo interface.
-	virtual float GetMaxDrawDistance() const final override 
-	{ 
-		return MaxDrawDistance; 
-	}
-
-	virtual float GetFadeRange() const final override 
-	{ 
-		return FadeRange; 
-	}
-
-	/** @return radius of the light or 0 if no radius */
-	virtual float GetRadius() const override
-	{ 
-		return Radius; 
-	}
-
-	virtual bool AffectsBounds(const FBoxSphereBounds& Bounds) const override
-	{
-		if((Bounds.Origin - GetLightToWorld().GetOrigin()).SizeSquared() > FMath::Square(Radius + Bounds.SphereRadius))
-		{
-			return false;
-		}
-
-		if(!FLightSceneProxy::AffectsBounds(Bounds))
-		{
-			return false;
-		}
-
-		return true;
-	}
-
-	virtual bool GetScissorRect(FIntRect& ScissorRect, const FSceneView& View, const FIntRect& ViewRect) const override
-	{
-		ScissorRect = ViewRect;
-		return FMath::ComputeProjectedSphereScissorRect(ScissorRect, GetLightToWorld().GetOrigin(), Radius, View.ViewMatrices.GetViewOrigin(), View.ViewMatrices.GetViewMatrix(), View.ViewMatrices.GetProjectionMatrix()) == 1;
-	}
-
-	virtual void SetScissorRect(FRHICommandList& RHICmdList, const FSceneView& View, const FIntRect& ViewRect) const override
-	{
-		FIntRect ScissorRect;
-
-		if (GetScissorRect(ScissorRect, View, ViewRect))
-		{
-			RHICmdList.SetScissorRect(true, ScissorRect.Min.X, ScissorRect.Min.Y, ScissorRect.Max.X, ScissorRect.Max.Y);
-		}
-		else
-		{
-			RHICmdList.SetScissorRect(false, 0, 0, 0, 0);
-		}
-	}
-
-	virtual FSphere GetBoundingSphere() const
-	{
-		return FSphere(GetPosition(), GetRadius());
-	}
-
-	virtual float GetEffectiveScreenRadius(const FViewMatrices& ShadowViewMatrices) const override
-	{
-		// Use the distance from the view origin to the light to approximate perspective projection
-		// We do not use projected screen position since it causes problems when the light is behind the camera
-
-		const float LightDistance = (GetOrigin() - ShadowViewMatrices.GetViewOrigin()).Size();
-
-		return ShadowViewMatrices.GetScreenScale() * GetRadius() / FMath::Max(LightDistance, 1.0f);
-	}
-
-	virtual FVector GetPerObjectProjectedShadowProjectionPoint(const FBoxSphereBounds& SubjectBounds) const
-	{
-		return GetOrigin();
-	}
-
-	virtual bool GetPerObjectProjectedShadowInitializer(const FBoxSphereBounds& SubjectBounds,class FPerObjectProjectedShadowInitializer& OutInitializer) const override
-	{
-		// Use a perspective projection looking at the primitive from the light position.
-		FVector LightPosition = GetPerObjectProjectedShadowProjectionPoint(SubjectBounds);
-		FVector LightVector = SubjectBounds.Origin - LightPosition;
-		float LightDistance = LightVector.Size();
-		float SilhouetteRadius = 1.0f;
-		const float SubjectRadius = SubjectBounds.BoxExtent.Size();
-		const float ShadowRadiusMultiplier = 1.1f;
-
-		if (LightDistance > SubjectRadius)
-		{
-			SilhouetteRadius = FMath::Min(SubjectRadius * FMath::InvSqrt((LightDistance - SubjectRadius) * (LightDistance + SubjectRadius)), 1.0f);
-		}
-
-		if (LightDistance <= SubjectRadius * ShadowRadiusMultiplier)
-		{
-			// Make the primitive fit in a single < 90 degree FOV projection.
-			LightVector = SubjectRadius * LightVector.GetSafeNormal() * ShadowRadiusMultiplier;
-			LightPosition = (SubjectBounds.Origin - LightVector );
-			LightDistance = SubjectRadius * ShadowRadiusMultiplier;
-			SilhouetteRadius = 1.0f;
-		}
-
-		OutInitializer.PreShadowTranslation = -LightPosition;
-		OutInitializer.WorldToLight = FInverseRotationMatrix((LightVector / LightDistance).Rotation());
-		OutInitializer.Scales = FVector(1.0f,1.0f / SilhouetteRadius,1.0f / SilhouetteRadius);
-		OutInitializer.FaceDirection = FVector(1,0,0);
-		OutInitializer.SubjectBounds = FBoxSphereBounds(SubjectBounds.Origin - LightPosition,SubjectBounds.BoxExtent,SubjectBounds.SphereRadius);
-		OutInitializer.WAxis = FVector4(0,0,1,0);
-		OutInitializer.MinLightW = 0.1f;
-		OutInitializer.MaxDistanceToCastInLightW = Radius;
-		return true;
-	}
-
-protected:
-
-	/** Updates the light scene info's radius from the component. */
-	void UpdateRadius(float ComponentRadius)
-	{
-		Radius = ComponentRadius;
-
-		// Min to avoid div by 0 (NaN in InvRadius)
-		InvRadius = 1.0f / FMath::Max(0.00001f, ComponentRadius);
-	}
-
-	float MaxDrawDistance;
-	float FadeRange;
-};
-
-
 class FPointLightSceneProxy : public FLocalLightSceneProxy
 {
 public:
diff --git a/Engine/Source/Runtime/Engine/Public/LocalLightSceneProxy.h b/Engine/Source/Runtime/Engine/Public/LocalLightSceneProxy.h
new file mode 100644
--- /dev/null
+++ b/Engine/Source/Runtime/Engine/Public/LocalLightSceneProxy.h
@@ -0,0 +1,63 @@
+// Copyright 1998-2019 Epic Games, Inc. All Rights Reserved.
+
+/*=============================================================================
+	LocalLightSceneProxy.h: Local light scene info definition.
+=============================================================================*/
+
+#pragma once
+
+#include "CoreMinimal.h"
+#include "SceneManagement.h"
+
+class ULocalLightComponent;
+
+/** The parts of the point light scene info that aren't dependent on the light policy type. */
+class FLocalLightSceneProxy : public FLightSceneProxy
+{
+public:
+
+	/** The light radius. */
+	float Radius;
+
+	/** One over the light's radius. */
+	float InvRadius;
+
+	/** Initialization constructor. */
+	FLocalLightSceneProxy(const ULocalLightComponent* Component);
+
+	/**
+	* Called on the light scene info after it has been passed to the rendering thread to update the rendering thread's cached info when
+	* the light's radius changes.
+	*/
+	void UpdateRadius_GameThread(float Radius);
+
+	// FLightSceneInfo interface.
+	virtual float GetMaxDrawDistance() const final override;
+
+	virtual float GetFadeRange() const final override;
+
+	/** @return radius of the light or 0 if no radius */
+	virtual float GetRadius() const override;
+
+	virtual bool AffectsBounds(const FBoxSphereBounds& Bounds) const override;
+
+	virtual bool GetScissorRect(FIntRect& ScissorRect, const FSceneView& View, const FIntRect& ViewRect) const override;
+
+	virtual void SetScissorRect(FRHICommandList& RHICmdList, const FSceneView& View, const FIntRect& ViewRect) const override;
+
+	virtual FSphere GetBoundingSphere() const;
+
+	virtual float GetEffectiveScreenRadius(const FViewMatrices& ShadowViewMatrices) const override;
+
+	virtual FVector GetPerObjectProjectedShadowProjectionPoint(const FBoxSphereBounds& SubjectBounds) const;
+
+	virtual bool GetPerObjectProjectedShadowInitializer(const FBoxSphereBounds& SubjectBounds, class FPerObjectProjectedShadowInitializer& OutInitializer) const override;
+
+protected:
+
+	/** Updates the light scene info's radius from the component. */
+	void UpdateRadius(float ComponentRadius);
+
+	float MaxDrawDistance;
+	float FadeRange;
+};
diff --git a/Engine/Source/Runtime/Engine/Public/RectLightSceneProxy.h b/Engine/Source/Runtime/Engine/Public/RectLightSceneProxy.h
new file mode 100644
--- /dev/null
+++ b/Engine/Source/Runtime/Engine/Public/RectLightSceneProxy.h
@@ -0,0 +1,44 @@
+// Copyright 1998-2019 Epic Games, Inc. All Rights Reserved.
+
+/*=============================================================================
+	RectLightSceneProxy.h: FRectLightSceneProxy definition.
+=============================================================================*/
+
+#pragma once
+
+#include "CoreMinimal.h"
+#include "Components/RectLightComponent.h"
+#include "LocalLightSceneProxy.h"
+#include "SceneManagement.h"
+
+class FRectLightSceneProxy : public FLocalLightSceneProxy
+{
+public:
+	float		SourceWidth;
+	float		SourceHeight;
+	UTexture*	SourceTexture;
+#if RHI_RAYTRACING
+	FRWBuffer   RectLightMipTree;
+	FIntVector  RectLightMipTreeDimensions;
+#endif
+
+	FRectLightSceneProxy(const URectLightComponent* Component);
+	virtual ~FRectLightSceneProxy();
+
+	virtual bool IsRectLight() const override;
+
+	virtual bool HasSourceTexture() const override;
+
+	/** Accesses parameters needed for rendering the light. */
+	virtual void GetLightShaderParameters(FLightShaderParameters& LightParameters) const override;
+
+	/**
+	* Sets up a projected shadow initializer for shadows from the entire scene.
+	* @return True if the whole-scene projected shadow should be used.
+	*/
+	virtual bool GetWholeSceneProjectedShadowInitializer(const FSceneViewFamily& ViewFamily, TArray<FWholeSceneProjectedShadowInitializer, TInlineAllocator<6> >& OutInitializers) const;
+
+#if RHI_RAYTRACING
+	void BuildRectLightMipTree(FRHICommandListImmediate& RHICmdList);
+#endif
+};
\ No newline at end of file
diff --git a/Engine/Source/Runtime/Renderer/Private/RayTracing/RayTracingShadowsForRectLight.cpp b/Engine/Source/Runtime/Renderer/Private/RayTracing/RayTracingShadowsForRectLight.cpp
--- a/Engine/Source/Runtime/Renderer/Private/RayTracing/RayTracingShadowsForRectLight.cpp
+++ b/Engine/Source/Runtime/Renderer/Private/RayTracing/RayTracingShadowsForRectLight.cpp
@@ -14,6 +14,7 @@
 
 #include "PostProcess/PostProcessing.h"
 #include "PostProcess/SceneFilterRendering.h"
+#include "RectLightSceneProxy.h"
 
 static int32 GRayTracingRectLight = 1;
 static FAutoConsoleVariableRef CVarRayTracingRectLight(
@@ -391,6 +392,7 @@ void FDeferredShadingSceneRenderer::VisualizeRectLightMipTree(
 
 	// PSO definition
 	FGraphicsPipelineStateInitializer GraphicsPSOInit;
+	SceneContext.BeginRenderingSceneColor(RHICmdList, ESimpleRenderTargetMode::EExistingColorAndDepth, FExclusiveDepthStencil::DepthRead_StencilWrite, true);
 	RHICmdList.ApplyCachedRenderTargets(GraphicsPSOInit);
 	GraphicsPSOInit.BlendState = TStaticBlendState<CW_RGB, BO_Add, BF_One, BF_One>::GetRHI();
 	GraphicsPSOInit.RasterizerState = TStaticRasterizerState<FM_Solid, CM_None>::GetRHI();
@@ -497,21 +499,11 @@ void FDeferredShadingSceneRenderer::RenderRayTracingRectLightInternal(
 )
 {
 	check(RectLightSceneInfo.Proxy);
+	check(RectLightSceneInfo.Proxy->IsRectLight());
+	FRectLightSceneProxy* RectLightSceneProxy = (FRectLightSceneProxy*) RectLightSceneInfo.Proxy;
 
-	// #dxr_todo: Cache MIP tree build with RectLight
-	FRWBuffer RectLightMipTree;
-	FIntVector RectLightMipTreeDimensions;
- 	BuildRectLightMipTree(RHICmdList, RectLightSceneInfo, RectLightMipTree, RectLightMipTreeDimensions);
-	// Debug: Visualization
-#if 0
-	for (int32 ViewIndex = 0; ViewIndex < Views.Num(); ++ViewIndex)
-	{
-		FViewInfo& View = Views[ViewIndex];
-		VisualizeRectLightMipTree(RHICmdList, View, RectLightMipTree, RectLightMipTreeDimensions);
-	}
-#endif
 	FLightShaderParameters LightShaderParameters;
-	RectLightSceneInfo.Proxy->GetLightShaderParameters(LightShaderParameters);
+	RectLightSceneProxy->GetLightShaderParameters(LightShaderParameters);
 
 	FRectLightData RectLightData;
 	RectLightData.SamplesPerPixel = GRayTracingRectLightSamplesPerPixel;
@@ -524,17 +516,17 @@ void FDeferredShadingSceneRenderer::RenderRayTracingRectLightInternal(
 	RectLightData.Color = LightShaderParameters.Color / 2.0;
 
 	// #dxr_todo: Ray traced textured area lights are 1.5X brighter than those in lit mode.
-	if (RectLightSceneInfo.Proxy->HasSourceTexture())
+	if (RectLightSceneProxy->HasSourceTexture())
 	{
 		RectLightData.Color *= 2.0 / 3.0;
 	}
 
 	RectLightData.Width = 2.0f * LightShaderParameters.SourceRadius;
 	RectLightData.Height = 2.0f * LightShaderParameters.SourceLength;
-	RectLightData.MipTreeDimensions = RectLightMipTreeDimensions;
 	RectLightData.Texture = LightShaderParameters.SourceTexture;
 	RectLightData.TextureSampler = RHICreateSamplerState(FSamplerStateInitializerRHI(SF_Bilinear, AM_Border, AM_Border, AM_Border));
-	RectLightData.MipTree = RectLightMipTree.SRV;
+	RectLightData.MipTree = RectLightSceneProxy->RectLightMipTree.SRV;
+	RectLightData.MipTreeDimensions = RectLightSceneProxy->RectLightMipTreeDimensions;
 	FUniformBufferRHIRef RectLightUniformBuffer = RHICreateUniformBuffer(&RectLightData, FRectLightData::StaticStructMetadata.GetLayout(), EUniformBufferUsage::UniformBuffer_SingleDraw);
 
 	for (int32 ViewIndex = 0; ViewIndex < Views.Num(); ++ViewIndex)
