Fixed FMaterialInstanceResource::GetMaterialNoFallback calling into UMaterialInstance::GetMaterialResource which accesses the game thread parent directly (race condition)
Reimplement STORE_ONLY_ACTIVE_SHADERMAPS in FMaterialInstanceResource::GetMaterialWithFallback which was dropped in HLR merge

diff --git a/Engine/Source/Runtime/Engine/Private/Materials/MaterialInstance.cpp b/Engine/Source/Runtime/Engine/Private/Materials/MaterialInstance.cpp
--- a/Engine/Source/Runtime/Engine/Private/Materials/MaterialInstance.cpp
+++ b/Engine/Source/Runtime/Engine/Private/Materials/MaterialInstance.cpp
@@ -109,7 +109,16 @@ const FMaterial& FMaterialInstanceResource::GetMaterialWithFallback(ERHIFeatureL
 		if (Owner->bHasStaticPermutationResource)
 		{
 			EMaterialQualityLevel::Type ActiveQualityLevel = GetCachedScalabilityCVars().MaterialQualityLevel;
-			FMaterialResource* StaticPermutationResource = Owner->StaticPermutationMaterialResources[ActiveQualityLevel][InFeatureLevel];
+
+			FMaterialResource* StaticPermutationResource;
+
+#if STORE_ONLY_ACTIVE_SHADERMAPS
+			StaticPermutationResource = Owner->StaticPermutationMaterialResources[ActiveQualityLevel][InFeatureLevel] ?
+				Owner->StaticPermutationMaterialResources[ActiveQualityLevel][InFeatureLevel] :
+				Owner->StaticPermutationMaterialResources[EMaterialQualityLevel::High][InFeatureLevel];
+#else
+			StaticPermutationResource = Owner->StaticPermutationMaterialResources[ActiveQualityLevel][InFeatureLevel];
+#endif
 
 			if (StaticPermutationResource)
 			{
@@ -147,7 +156,34 @@ const FMaterial& FMaterialInstanceResource::GetMaterialWithFallback(ERHIFeatureL
 FMaterial* FMaterialInstanceResource::GetMaterialNoFallback(ERHIFeatureLevel::Type InFeatureLevel) const
 {
 	checkSlow(IsInParallelRenderingThread());
-	return Owner->GetMaterialResource(InFeatureLevel);
+
+	if (Parent)
+	{
+		if (Owner->bHasStaticPermutationResource)
+		{
+			EMaterialQualityLevel::Type ActiveQualityLevel = GetCachedScalabilityCVars().MaterialQualityLevel;
+			FMaterialResource* StaticPermutationResource;
+
+#if STORE_ONLY_ACTIVE_SHADERMAPS
+			StaticPermutationResource = Owner->StaticPermutationMaterialResources[ActiveQualityLevel][InFeatureLevel] ?
+				Owner->StaticPermutationMaterialResources[ActiveQualityLevel][InFeatureLevel] :
+				Owner->StaticPermutationMaterialResources[EMaterialQualityLevel::High][InFeatureLevel];
+#else
+			StaticPermutationResource = Owner->StaticPermutationMaterialResources[ActiveQualityLevel][InFeatureLevel];
+#endif
+			return StaticPermutationResource;
+		}
+		else
+		{
+			FMaterialRenderProxy* ParentProxy = Parent->GetRenderProxy();
+
+			if (ParentProxy)
+			{
+				return ParentProxy->GetMaterialNoFallback(InFeatureLevel);
+			}
+		}
+	}
+	return NULL;
 }
 
 UMaterialInterface* FMaterialInstanceResource::GetMaterialInterface() const
