Fixed mismatch between when uniform expressions are cached and when mesh passes fall back to default material, which was causing uncached uniform expressions to get used for rendering in the case of a Material Instance whose parent had been deleted.

diff --git a/Engine/Source/Runtime/Engine/Private/Materials/MaterialShared.cpp b/Engine/Source/Runtime/Engine/Private/Materials/MaterialShared.cpp
--- a/Engine/Source/Runtime/Engine/Private/Materials/MaterialShared.cpp
+++ b/Engine/Source/Runtime/Engine/Private/Materials/MaterialShared.cpp
@@ -2165,22 +2165,15 @@ void FMaterialRenderProxy::UpdateDeferredCachedUniformExpressions()
 
 		UMaterialInterface::IterateOverActiveFeatureLevels([&](ERHIFeatureLevel::Type InFeatureLevel)
 		{
-			const FMaterial* MaterialNoFallback = MaterialProxy->GetMaterialNoFallback(InFeatureLevel);
+			const FMaterialRenderProxy* FallbackMaterialProxy = nullptr;
+			const FMaterial& Material = MaterialProxy->GetMaterialWithFallback(InFeatureLevel, FallbackMaterialProxy);
 
-			if (MaterialNoFallback && MaterialNoFallback->GetRenderingThreadShaderMap())
+			// Don't bother caching if we'll be falling back to a different FMaterialRenderProxy for rendering anyway
+			if (!FallbackMaterialProxy)
 			{
-				const FMaterial* Material = MaterialProxy->GetMaterial(InFeatureLevel);
-
-				// Do not cache uniform expressions for fallback materials. This step could
-				// be skipped where we don't allow for asynchronous shader compiling.
-				bool bIsFallbackMaterial = (Material != MaterialNoFallback);
-
-				if (!bIsFallbackMaterial)
-				{
-					FMaterialRenderContext MaterialRenderContext(MaterialProxy, *Material, nullptr);
-					MaterialRenderContext.bShowSelection = GIsEditor;
-					MaterialProxy->EvaluateUniformExpressions(MaterialProxy->UniformExpressionCache[(int32)InFeatureLevel], MaterialRenderContext);
-				}
+				FMaterialRenderContext MaterialRenderContext(MaterialProxy, Material, nullptr);
+				MaterialRenderContext.bShowSelection = GIsEditor;
+				MaterialProxy->EvaluateUniformExpressions(MaterialProxy->UniformExpressionCache[(int32)InFeatureLevel], MaterialRenderContext);
 			}
 		});
 	}
