Local VF uniform buffer is created if UseGPUScene is enabled, even if the platform does not support RHISupportsManualVertexFetch

diff --git a/Engine/Source/Runtime/Engine/Private/LocalVertexFactory.cpp b/Engine/Source/Runtime/Engine/Private/LocalVertexFactory.cpp
--- a/Engine/Source/Runtime/Engine/Private/LocalVertexFactory.cpp
+++ b/Engine/Source/Runtime/Engine/Private/LocalVertexFactory.cpp
@@ -53,20 +53,28 @@ TUniformBufferRef<FLocalVertexFactoryUniformShaderParameters> CreateLocalVFUnifo
 	FLocalVertexFactoryUniformShaderParameters UniformParameters;
 
 	UniformParameters.LODLightmapDataIndex = LODLightmapDataIndex;
-
-	UniformParameters.VertexFetch_PackedTangentsBuffer = LocalVertexFactory->GetTangentsSRV();
-	UniformParameters.VertexFetch_TexCoordBuffer = LocalVertexFactory->GetTextureCoordinatesSRV();
-
 	int32 ColorIndexMask = 0;
-	if (OverrideColorVertexBuffer)
+
+	if (RHISupportsManualVertexFetch(GMaxRHIShaderPlatform))
 	{
-		UniformParameters.VertexFetch_ColorComponentsBuffer = OverrideColorVertexBuffer->GetColorComponentsSRV();
-		ColorIndexMask = OverrideColorVertexBuffer->GetNumVertices() > 1 ? ~0 : 0;
+		UniformParameters.VertexFetch_PackedTangentsBuffer = LocalVertexFactory->GetTangentsSRV();
+		UniformParameters.VertexFetch_TexCoordBuffer = LocalVertexFactory->GetTextureCoordinatesSRV();
+
+		if (OverrideColorVertexBuffer)
+		{
+			UniformParameters.VertexFetch_ColorComponentsBuffer = OverrideColorVertexBuffer->GetColorComponentsSRV();
+			ColorIndexMask = OverrideColorVertexBuffer->GetNumVertices() > 1 ? ~0 : 0;
+		}
+		else
+		{
+			UniformParameters.VertexFetch_ColorComponentsBuffer = LocalVertexFactory->GetColorComponentsSRV();
+			ColorIndexMask = (int32)LocalVertexFactory->GetColorIndexMask();
+		}
 	}
 	else
 	{
-		UniformParameters.VertexFetch_ColorComponentsBuffer = LocalVertexFactory->GetColorComponentsSRV();
-		ColorIndexMask = (int32)LocalVertexFactory->GetColorIndexMask();
+		UniformParameters.VertexFetch_PackedTangentsBuffer = GNullColorVertexBuffer.VertexBufferSRV;
+		UniformParameters.VertexFetch_TexCoordBuffer = GNullColorVertexBuffer.VertexBufferSRV;
 	}
 
 	if (!UniformParameters.VertexFetch_ColorComponentsBuffer)
@@ -149,7 +157,7 @@ void FLocalVertexFactoryShaderParametersBase::GetElementShaderBindings(
 {
 	const auto* LocalVertexFactory = static_cast<const FLocalVertexFactory*>(VertexFactory);
 	
-	if (LocalVertexFactory->SupportsManualVertexFetch(FeatureLevel))
+	if (LocalVertexFactory->SupportsManualVertexFetch(FeatureLevel) || UseGPUScene(GMaxRHIShaderPlatform, FeatureLevel))
 	{
 		FUniformBufferRHIParamRef VertexFactoryUniformBuffer = static_cast<FUniformBufferRHIParamRef>(BatchElement.VertexFactoryUserData);
 
@@ -371,7 +379,7 @@ void FLocalVertexFactory::InitRHI()
 	InitDeclaration(Elements);
 	check(IsValidRef(GetDeclaration()));
 
-	if (RHISupportsManualVertexFetch(GMaxRHIShaderPlatform))
+	if (RHISupportsManualVertexFetch(GMaxRHIShaderPlatform) || bCanUseGPUScene)
 	{
 		UniformBuffer = CreateLocalVFUniformBuffer(this, Data.LODLightmapDataIndex, nullptr, DefaultBaseVertexIndex);
 	}
