Licensee fix: Moved FGlobalComponentRecreateRenderStateContext out of the per-level loop during lighting builds and reflection capture builds.  Reduces the number of render state recreates during a build for each component from 2 * NumSublevels to 2.

diff --git a/Engine/Source/Editor/UnrealEd/Private/EditorEngine.cpp b/Engine/Source/Editor/UnrealEd/Private/EditorEngine.cpp
--- a/Engine/Source/Editor/UnrealEd/Private/EditorEngine.cpp
+++ b/Engine/Source/Editor/UnrealEd/Private/EditorEngine.cpp
@@ -3950,19 +3950,23 @@ void UEditorEngine::BuildReflectionCaptures(UWorld* World)
 		}
 	}
 
-	for (ULevel* Level : World->GetLevels())
 	{
-		if (Level->bIsVisible)
+		FGlobalComponentRecreateRenderStateContext Context;
+
+		for (ULevel* Level : World->GetLevels())
 		{
-			if (Level->MapBuildData)
+			if (Level->bIsVisible)
 			{
-				// Remove all existing reflection capture data from visible levels before the build
-				Level->MapBuildData->InvalidateReflectionCaptures(Level->bIsLightingScenario ? &ResourcesToKeep : nullptr);
-			}
+				if (Level->MapBuildData)
+				{
+					// Remove all existing reflection capture data from visible levels before the build
+					Level->MapBuildData->InvalidateReflectionCaptures(Level->bIsLightingScenario ? &ResourcesToKeep : nullptr);
+				}
 
-			if (Level->bIsLightingScenario)
-			{
-				LightingScenarios.Add(Level);
+				if (Level->bIsLightingScenario)
+				{
+					LightingScenarios.Add(Level);
+				}
 			}
 		}
 	}
diff --git a/Engine/Source/Editor/UnrealEd/Private/StaticLightingSystem/StaticLightingSystem.cpp b/Engine/Source/Editor/UnrealEd/Private/StaticLightingSystem/StaticLightingSystem.cpp
--- a/Engine/Source/Editor/UnrealEd/Private/StaticLightingSystem/StaticLightingSystem.cpp
+++ b/Engine/Source/Editor/UnrealEd/Private/StaticLightingSystem/StaticLightingSystem.cpp
@@ -760,6 +760,7 @@ bool FStaticLightingSystem::BeginLightmassProcess()
 void FStaticLightingSystem::InvalidateStaticLighting()
 {
 	FLightmassStatistics::FScopedGather InvalidationScopeStat(LightmassStatistics.InvalidationTime);
+	FGlobalComponentRecreateRenderStateContext Context;
 
 	for( int32 LevelIndex=0; LevelIndex<World->GetNumLevels(); LevelIndex++ )
 	{
@@ -781,7 +782,7 @@ void FStaticLightingSystem::InvalidateStaticLighting()
 
 				if (Level->MapBuildData)
 				{
-					Level->MapBuildData->InvalidateStaticLighting(World, &BuildDataResourcesToKeep);
+					Level->MapBuildData->InvalidateStaticLighting(World, false, &BuildDataResourcesToKeep);
 				}
 			}
 			if (Level == World->PersistentLevel)
diff --git a/Engine/Source/Runtime/Engine/Classes/Engine/MapBuildDataRegistry.h b/Engine/Source/Runtime/Engine/Classes/Engine/MapBuildDataRegistry.h
--- a/Engine/Source/Runtime/Engine/Classes/Engine/MapBuildDataRegistry.h
+++ b/Engine/Source/Runtime/Engine/Classes/Engine/MapBuildDataRegistry.h
@@ -333,7 +333,7 @@ public:
 	ENGINE_API const FReflectionCaptureMapBuildData* GetReflectionCaptureBuildData(FGuid CaptureId) const;
 	ENGINE_API FReflectionCaptureMapBuildData* GetReflectionCaptureBuildData(FGuid CaptureId);
 
-	ENGINE_API void InvalidateStaticLighting(UWorld* World, const TSet<FGuid>* ResourcesToKeep = nullptr);
+	ENGINE_API void InvalidateStaticLighting(UWorld* World, bool bRecreateRenderState = true, const TSet<FGuid>* ResourcesToKeep = nullptr);
 	ENGINE_API void InvalidateReflectionCaptures(const TSet<FGuid>* ResourcesToKeep = nullptr);
 
 	ENGINE_API bool IsLegacyBuildData() const;
diff --git a/Engine/Source/Runtime/Engine/Private/MapBuildData.cpp b/Engine/Source/Runtime/Engine/Private/MapBuildData.cpp
--- a/Engine/Source/Runtime/Engine/Private/MapBuildData.cpp
+++ b/Engine/Source/Runtime/Engine/Private/MapBuildData.cpp
@@ -610,12 +610,18 @@ FReflectionCaptureMapBuildData* UMapBuildDataRegistry::GetReflectionCaptureBuild
 	return ReflectionCaptureBuildData.Find(CaptureId);
 }
 
-void UMapBuildDataRegistry::InvalidateStaticLighting(UWorld* World, const TSet<FGuid>* ResourcesToKeep)
+void UMapBuildDataRegistry::InvalidateStaticLighting(UWorld* World, bool bRecreateRenderState, const TSet<FGuid>* ResourcesToKeep)
 {
-	if (MeshBuildData.Num() > 0 || LightBuildData.Num() > 0)
+	TUniquePtr<FGlobalComponentRecreateRenderStateContext> RecreateContext;
+
+	if (bRecreateRenderState)
 	{
-		FGlobalComponentRecreateRenderStateContext Context;
+		// Warning: if skipping this, caller is responsible for unregistering any components potentially referencing this UMapBuildDataRegistry before we change its contents!
+		RecreateContext = TUniquePtr<FGlobalComponentRecreateRenderStateContext>(new FGlobalComponentRecreateRenderStateContext);
+	}
 
+	if (MeshBuildData.Num() > 0 || LightBuildData.Num() > 0)
+	{
 		if (!ResourcesToKeep || !ResourcesToKeep->Num())
 		{
 			MeshBuildData.Empty();
@@ -671,7 +677,7 @@ void UMapBuildDataRegistry::InvalidateReflectionCaptures(const TSet<FGuid>* Reso
 {
 	if (ReflectionCaptureBuildData.Num() > 0)
 	{
-		FGlobalComponentRecreateRenderStateContext Context;
+		// Warning: caller is responsible for unregistering any components potentially referencing this UMapBuildDataRegistry before we change its contents!
 
 		TMap<FGuid, FReflectionCaptureMapBuildData> PrevReflectionCapturedData;
 		FMemory::Memswap(&ReflectionCaptureBuildData , &PrevReflectionCapturedData, sizeof(ReflectionCaptureBuildData));
