Fixing android build

diff --git a/Engine/Source/Runtime/Engine/Private/Level.cpp b/Engine/Source/Runtime/Engine/Private/Level.cpp
--- a/Engine/Source/Runtime/Engine/Private/Level.cpp
+++ b/Engine/Source/Runtime/Engine/Private/Level.cpp
@@ -1863,7 +1863,7 @@ UMapBuildDataRegistry* ULevel::GetOrCreateMapBuildData()
 		if (MapBuildData)
 		{
 			// Release rendering data depending on MapBuildData, before we destroy MapBuildData
-			MapBuildData->InvalidateStaticLighting(GetWorld(), nullptr);
+			MapBuildData->InvalidateStaticLighting(GetWorld(), true, nullptr);
 
 			// Allow the legacy registry to be GC'ed
 			MapBuildData->ClearFlags(RF_Standalone);
