DR - Fix missing shader
#rnx

diff --git a/Engine/Source/Runtime/Engine/Private/Components/RectLightComponent.cpp b/Engine/Source/Runtime/Engine/Private/Components/RectLightComponent.cpp
--- a/Engine/Source/Runtime/Engine/Private/Components/RectLightComponent.cpp
+++ b/Engine/Source/Runtime/Engine/Private/Components/RectLightComponent.cpp
@@ -223,7 +223,7 @@ class FBuildRectLightMipTreeCS : public FGlobalShader
 public:
 	static bool ShouldCompilePermutation(const FGlobalShaderPermutationParameters& Parameters)
 	{
-		return ShouldCompileRayTracingShadersForProject(Parameters.Platform);
+		return RHISupportsRayTracingShaders(Parameters.Platform);
 	}
 
 	static uint32 GetGroupSize()
@@ -339,4 +339,4 @@ void FRectLightSceneProxy::BuildRectLightMipTree(FRHICommandListImmediate& RHICm
 	BuildRectLightMipTreeComputeShader->UnsetParameters(RHICmdList, EResourceTransitionAccess::ERWBarrier, EResourceTransitionPipeline::EComputeToCompute, RectLightMipTree, TransitionFence);
 }
 
-#endif // RHI_RAYTRACING
\ No newline at end of file
+#endif // RHI_RAYTRACING
