Always render primitive bounds into the world DPG.  Previously they sometimes rendered to foreground DPG, but foreground DPG didn't actually correctly render to the foreground.
Now, foreground DPG is working as intended, but we prefer primitive bounds to depth sort correctly with the rest of the world

diff --git a/Engine/Source/Runtime/Engine/Private/PrimitiveSceneProxy.cpp b/Engine/Source/Runtime/Engine/Private/PrimitiveSceneProxy.cpp
--- a/Engine/Source/Runtime/Engine/Private/PrimitiveSceneProxy.cpp
+++ b/Engine/Source/Runtime/Engine/Private/PrimitiveSceneProxy.cpp
@@ -716,11 +716,11 @@ void FPrimitiveSceneProxy::RenderBounds(
 	const FBoxSphereBounds& InBounds, 
 	bool bRenderInEditor) const
 {
-	const ESceneDepthPriorityGroup DrawBoundsDPG = EngineShowFlags.Game ? SDPG_World : SDPG_Foreground;
 	if (EngineShowFlags.Bounds && (EngineShowFlags.Game || bRenderInEditor))
 	{
 		// Draw the static mesh's bounding box and sphere.
-		DrawWireBox(PDI,InBounds.GetBox(), FColor(72,72,255),DrawBoundsDPG);
+		const ESceneDepthPriorityGroup DrawBoundsDPG = SDPG_World;
+		DrawWireBox(PDI,InBounds.GetBox(), FColor(72,72,255), DrawBoundsDPG);
 		DrawCircle(PDI, InBounds.Origin, FVector(1, 0, 0), FVector(0, 1, 0), FColor::Yellow, InBounds.SphereRadius, 32, DrawBoundsDPG);
 		DrawCircle(PDI, InBounds.Origin, FVector(1, 0, 0), FVector(0, 0, 1), FColor::Yellow, InBounds.SphereRadius, 32, DrawBoundsDPG);
 		DrawCircle(PDI, InBounds.Origin, FVector(0, 1, 0), FVector(0, 0, 1), FColor::Yellow, InBounds.SphereRadius, 32, DrawBoundsDPG);
