Swap the fix for rectlights to avoid the mip build and dont' compile the shader in

diff --git a/Engine/Source/Runtime/Engine/Private/Components/RectLightComponent.cpp b/Engine/Source/Runtime/Engine/Private/Components/RectLightComponent.cpp
--- a/Engine/Source/Runtime/Engine/Private/Components/RectLightComponent.cpp
+++ b/Engine/Source/Runtime/Engine/Private/Components/RectLightComponent.cpp
@@ -147,11 +147,14 @@ FRectLightSceneProxy::FRectLightSceneProxy(const URectLightComponent* Component)
 	, SourceTexture(Component->SourceTexture)
 {
 #if RHI_RAYTRACING
-	ENQUEUE_RENDER_COMMAND(BuildRectLightMipTree)(
-		[this](FRHICommandListImmediate& RHICmdList)
+	if (IsRayTracingEnabled())
 	{
-		BuildRectLightMipTree(RHICmdList);
-	});
+		ENQUEUE_RENDER_COMMAND(BuildRectLightMipTree)(
+			[this](FRHICommandListImmediate& RHICmdList)
+		{
+			BuildRectLightMipTree(RHICmdList);
+		});
+	}
 #endif
 }
 
@@ -223,7 +226,7 @@ class FBuildRectLightMipTreeCS : public FGlobalShader
 public:
 	static bool ShouldCompilePermutation(const FGlobalShaderPermutationParameters& Parameters)
 	{
-		return RHISupportsRayTracingShaders(Parameters.Platform);
+		return ShouldCompileRayTracingShadersForProject(Parameters.Platform);
 	}
 
 	static uint32 GetGroupSize()
