Fixed crash at startup in packaged builds. Ray tracing geometry primitive count is computed from static mesh sections.

IndexBuffer.GetNumIndices() was previously used for this, however it is not available in cooked builds as it refers to CPU data.

diff --git a/Engine/Source/Runtime/Engine/Private/StaticMesh.cpp b/Engine/Source/Runtime/Engine/Private/StaticMesh.cpp
--- a/Engine/Source/Runtime/Engine/Private/StaticMesh.cpp
+++ b/Engine/Source/Runtime/Engine/Private/StaticMesh.cpp
@@ -697,7 +697,7 @@ void FStaticMeshLODResources::InitResources(UStaticMesh* Parent)
 				Initializer.BaseVertexIndex = 0;
 				Initializer.VertexBufferStride = VertexBuffers.PositionVertexBuffer.GetStride();
 				Initializer.VertexBufferByteOffset = 0;
-				Initializer.TotalPrimitiveCount = IndexBuffer.GetNumIndices() / 3;
+				Initializer.TotalPrimitiveCount = 0; // This is calculated below based on static mesh section data
 				Initializer.VertexBufferElementType = VET_Float3;
 				Initializer.PrimitiveType = PT_TriangleList;
 				Initializer.bFastBuild = false;
@@ -710,6 +710,7 @@ void FStaticMeshLODResources::InitResources(UStaticMesh* Parent)
 					Segment.FirstPrimitive = Section.FirstIndex / 3;
 					Segment.NumPrimitives = Section.NumTriangles;
 					GeometrySections.Add(Segment);
+					Initializer.TotalPrimitiveCount += Section.NumTriangles;
 				}
 				Initializer.Segments = GeometrySections;
 				
