Clamp LightMapCoordinateIndex to be below 4, since Lightmass only supports 4 UV sets.

diff --git a/Engine/Source/Runtime/Engine/Classes/Engine/StaticMesh.h b/Engine/Source/Runtime/Engine/Classes/Engine/StaticMesh.h
--- a/Engine/Source/Runtime/Engine/Classes/Engine/StaticMesh.h
+++ b/Engine/Source/Runtime/Engine/Classes/Engine/StaticMesh.h
@@ -642,7 +642,7 @@ class UStaticMesh : public UObject, public IInterface_CollisionDataProvider, pub
 	int32 LightMapResolution;
 
 	/** The light map coordinate index */
-	UPROPERTY(EditAnywhere, AdvancedDisplay, Category=StaticMesh, meta=(ToolTip="The light map coordinate index"))
+	UPROPERTY(EditAnywhere, AdvancedDisplay, Category=StaticMesh, meta=(ToolTip="The light map coordinate index", UIMin = "0", UIMax = "3"))
 	int32 LightMapCoordinateIndex;
 
 	/** Useful for reducing self shadowing from distance field methods when using world position offset to animate the mesh's vertices. */
diff --git a/Engine/Source/Runtime/Engine/Private/StaticMesh.cpp b/Engine/Source/Runtime/Engine/Private/StaticMesh.cpp
--- a/Engine/Source/Runtime/Engine/Private/StaticMesh.cpp
+++ b/Engine/Source/Runtime/Engine/Private/StaticMesh.cpp
@@ -4343,7 +4343,8 @@ void UStaticMesh::EnforceLightmapRestrictions()
 	// Legacy content may contain a lightmap resolution of 0, which was valid when vertex lightmaps were supported, but not anymore with only texture lightmaps
 	LightMapResolution = FMath::Max(LightMapResolution, 4);
 
-	int32 NumUVs = 16;
+	// Lightmass only supports 4 UVs
+	int32 NumUVs = 4;
 
 	if (RenderData)
 	{
