Fixed FArrowSceneProxy velocity relevance.

diff --git a/Engine/Source/Runtime/Engine/Private/Components/ArrowComponent.cpp b/Engine/Source/Runtime/Engine/Private/Components/ArrowComponent.cpp
--- a/Engine/Source/Runtime/Engine/Private/Components/ArrowComponent.cpp
+++ b/Engine/Source/Runtime/Engine/Private/Components/ArrowComponent.cpp
@@ -180,6 +180,7 @@ public:
 #endif
 		Result.bShadowRelevance = IsShadowCast(View);
 		Result.bEditorPrimitiveRelevance = UseEditorCompositing(View);
+		Result.bVelocityRelevance = IsMovable() && Result.bOpaqueRelevance && Result.bRenderInMainPass;
 		return Result;
 	}
 
