Fix shadowing warnings on mac/linux

diff --git a/Engine/Source/Runtime/Engine/Private/Materials/MaterialInstance.cpp b/Engine/Source/Runtime/Engine/Private/Materials/MaterialInstance.cpp
--- a/Engine/Source/Runtime/Engine/Private/Materials/MaterialInstance.cpp
+++ b/Engine/Source/Runtime/Engine/Private/Materials/MaterialInstance.cpp
@@ -4340,21 +4340,21 @@ void UMaterialInstance::CopyMaterialUniformParametersInternal(UMaterialInterface
 				checkSlow(VectorParameterValues.Num() == 0);
 				checkSlow(TextureParameterValues.Num() == 0);
 
-				const FMaterialResource* Resource = nullptr;
+				const FMaterialResource* MaterialResource = nullptr;
 				if (UWorld* World = AsMaterial->GetWorld())
 				{
-					Resource = AsMaterial->GetMaterialResource(World->FeatureLevel);
+					MaterialResource = AsMaterial->GetMaterialResource(World->FeatureLevel);
 				}
 
-				if (!Resource)
+				if (!MaterialResource)
 				{
-					Resource = AsMaterial->GetMaterialResource(GMaxRHIFeatureLevel);
+					MaterialResource = AsMaterial->GetMaterialResource(GMaxRHIFeatureLevel);
 				}
 
-				if (Resource)
+				if (MaterialResource)
 				{
 					// Scalars
-					const TArray<TRefCountPtr<FMaterialUniformExpression>>& ScalarExpressions = Resource->GetUniformScalarParameterExpressions();
+					const TArray<TRefCountPtr<FMaterialUniformExpression>>& ScalarExpressions = MaterialResource->GetUniformScalarParameterExpressions();
 					for (FMaterialUniformExpression* ScalarExpression : ScalarExpressions)
 					{
 						if (ScalarExpression->GetType() == &FMaterialUniformExpressionScalarParameter::StaticType)
@@ -4368,7 +4368,7 @@ void UMaterialInstance::CopyMaterialUniformParametersInternal(UMaterialInterface
 					}
 
 					// Vectors
-					const TArray<TRefCountPtr<FMaterialUniformExpression>>& VectorExpressions = Resource->GetUniformVectorParameterExpressions();
+					const TArray<TRefCountPtr<FMaterialUniformExpression>>& VectorExpressions = MaterialResource->GetUniformVectorParameterExpressions();
 					for (FMaterialUniformExpression* VectorExpression : VectorExpressions)
 					{
 						if (VectorExpression->GetType() == &FMaterialUniformExpressionVectorParameter::StaticType)
@@ -4384,8 +4384,8 @@ void UMaterialInstance::CopyMaterialUniformParametersInternal(UMaterialInterface
 					// Textures
 					const TArray<TRefCountPtr<FMaterialUniformExpressionTexture>>* TextureExpressions[2] =
 					{
-						&Resource->GetUniform2DTextureExpressions(),
-						&Resource->GetUniformCubeTextureExpressions()
+						&MaterialResource->GetUniform2DTextureExpressions(),
+						&MaterialResource->GetUniformCubeTextureExpressions()
 					};
 
 					for (int32 TypeIndex = 0; TypeIndex < ARRAY_COUNT(TextureExpressions); TypeIndex++)
@@ -4398,7 +4398,7 @@ void UMaterialInstance::CopyMaterialUniformParametersInternal(UMaterialInterface
 
 								FTextureParameterValue* ParameterValue = new(TextureParameterValues) FTextureParameterValue;
 								ParameterValue->ParameterInfo.Name = TextureParameter->GetParameterName();
-								TextureParameter->GetGameThreadTextureValue(AsMaterial, *Resource, ParameterValue->ParameterValue, false);
+								TextureParameter->GetGameThreadTextureValue(AsMaterial, *MaterialResource, ParameterValue->ParameterValue, false);
 							}
 						}
 					}
