Removed legacy Drawing Policies and Static Mesh Draw Lists.  These are now replaced by FMeshDrawCommand / FMeshPassProcessor everywhere in the renderer.
Implemented FMeshDrawCommand support for IndirectArgsBuffer, which is used by Niagara.

diff --git a/Engine/Plugins/FX/Niagara/Source/NiagaraVertexFactories/Private/NiagaraMeshVertexFactory.cpp b/Engine/Plugins/FX/Niagara/Source/NiagaraVertexFactories/Private/NiagaraMeshVertexFactory.cpp
--- a/Engine/Plugins/FX/Niagara/Source/NiagaraVertexFactories/Private/NiagaraMeshVertexFactory.cpp
+++ b/Engine/Plugins/FX/Niagara/Source/NiagaraVertexFactories/Private/NiagaraMeshVertexFactory.cpp
@@ -50,27 +50,6 @@ public:
 		Ar << SortedIndicesOffset;
 	}
 
-	virtual void SetMesh(FRHICommandList& RHICmdList, FShader* Shader, const FVertexFactory* VertexFactory, const FSceneView& View, const FMeshBatchElement& BatchElement, uint32 DataFlags) const override
-	{
-		const bool bInstanced = GRHISupportsInstancing;
-		FNiagaraMeshVertexFactory* NiagaraMeshVF = (FNiagaraMeshVertexFactory*)VertexFactory;
-		FVertexShaderRHIParamRef VertexShaderRHI = Shader->GetVertexShader();
-		SetUniformBufferParameter(RHICmdList, VertexShaderRHI, Shader->GetUniformBufferParameter<FNiagaraMeshUniformParameters>(), NiagaraMeshVF->GetUniformBuffer());
-
-		//SetSRVParameter(RHICmdList, VertexShaderRHI, PrevTransformBuffer, NiagaraMeshVF->GetPreviousTransformBufferSRV());
-
-		SetShaderValue(RHICmdList, VertexShaderRHI, MeshFacingMode, NiagaraMeshVF->GetMeshFacingMode());
-
-		SetSRVParameter(RHICmdList, VertexShaderRHI, NiagaraParticleDataFloat, NiagaraMeshVF->GetParticleDataFloatSRV());
-		SetShaderValue(RHICmdList, VertexShaderRHI, FloatDataOffset, NiagaraMeshVF->GetFloatDataOffset());
-		SetShaderValue(RHICmdList, VertexShaderRHI, FloatDataStride, NiagaraMeshVF->GetFloatDataStride());
-
-		//SetSRVParameter(RHICmdList, VertexShaderRHI, NiagaraParticleDataInt, NiagaraMeshVF->GetIntDataSRV());
-
-		SetSRVParameter(RHICmdList, VertexShaderRHI, SortedIndices, NiagaraMeshVF->GetSortedIndicesSRV() ? NiagaraMeshVF->GetSortedIndicesSRV() : GFNiagaraNullSortedIndicesVertexBuffer.VertexBufferSRV);
-		SetShaderValue(RHICmdList, VertexShaderRHI, SortedIndicesOffset, NiagaraMeshVF->GetSortedIndicesOffset());
-	}
-
 	virtual void GetElementShaderBindings(
 		const FSceneInterface* Scene,
 		const FSceneView* View,
diff --git a/Engine/Plugins/FX/Niagara/Source/NiagaraVertexFactories/Private/NiagaraRibbonVertexFactory.cpp b/Engine/Plugins/FX/Niagara/Source/NiagaraVertexFactories/Private/NiagaraRibbonVertexFactory.cpp
--- a/Engine/Plugins/FX/Niagara/Source/NiagaraVertexFactories/Private/NiagaraRibbonVertexFactory.cpp
+++ b/Engine/Plugins/FX/Niagara/Source/NiagaraVertexFactories/Private/NiagaraRibbonVertexFactory.cpp
@@ -59,22 +59,6 @@ public:
 		Ar << PackedPerRibbonDataByIndex;
 	}
 
-	virtual void SetMesh(FRHICommandList& RHICmdList, FShader* Shader, const FVertexFactory* VertexFactory, const FSceneView& View, const FMeshBatchElement& BatchElement, uint32 DataFlags) const override
-	{
-		FNiagaraRibbonVertexFactory* RibbonVF = (FNiagaraRibbonVertexFactory*)VertexFactory;
-		FVertexShaderRHIParamRef VertexShaderRHI = Shader->GetVertexShader();
-		SetUniformBufferParameter(RHICmdList, Shader->GetVertexShader(), Shader->GetUniformBufferParameter<FNiagaraRibbonUniformParameters>(), RibbonVF->GetRibbonUniformBuffer());
-
-		SetSRVParameter(RHICmdList, VertexShaderRHI, NiagaraParticleDataFloat, RibbonVF->GetParticleDataFloatSRV());
-		SetShaderValue(RHICmdList, VertexShaderRHI, FloatDataOffset, RibbonVF->GetFloatDataOffset());
-		SetShaderValue(RHICmdList, VertexShaderRHI, FloatDataStride, RibbonVF->GetFloatDataStride());
-
-		SetSRVParameter(RHICmdList, VertexShaderRHI, SortedIndices, RibbonVF->GetSortedIndicesSRV());
-		SetSRVParameter(RHICmdList, VertexShaderRHI, SegmentDistances, RibbonVF->GetSegmentDistancesSRV());
-		SetSRVParameter(RHICmdList, VertexShaderRHI, PackedPerRibbonDataByIndex, RibbonVF->GetPackedPerRibbonDataByIndexSRV());
-		SetShaderValue(RHICmdList, VertexShaderRHI, SortedIndicesOffset, RibbonVF->GetSortedIndicesOffset());
-	}
-
 	virtual void GetElementShaderBindings(
 		const FSceneInterface* Scene,
 		const FSceneView* View,
@@ -125,12 +109,6 @@ public:
 	{
 	}
 
-	virtual void SetMesh(FRHICommandList& RHICmdList, FShader* Shader, const FVertexFactory* VertexFactory, const FSceneView& View, const FMeshBatchElement& BatchElement, uint32 DataFlags) const override
-	{
-		FNiagaraRibbonVertexFactory* RibbonVF = (FNiagaraRibbonVertexFactory*)VertexFactory;
-		SetUniformBufferParameter(RHICmdList, Shader->GetPixelShader(), Shader->GetUniformBufferParameter<FNiagaraRibbonUniformParameters>(), RibbonVF->GetRibbonUniformBuffer());
-	}
-
 	virtual void GetElementShaderBindings(
 		const FSceneInterface* Scene,
 		const FSceneView* View,
diff --git a/Engine/Plugins/FX/Niagara/Source/NiagaraVertexFactories/Private/NiagaraSpriteVertexFactory.cpp b/Engine/Plugins/FX/Niagara/Source/NiagaraVertexFactories/Private/NiagaraSpriteVertexFactory.cpp
--- a/Engine/Plugins/FX/Niagara/Source/NiagaraVertexFactories/Private/NiagaraSpriteVertexFactory.cpp
+++ b/Engine/Plugins/FX/Niagara/Source/NiagaraVertexFactories/Private/NiagaraSpriteVertexFactory.cpp
@@ -75,27 +75,6 @@ public:
 		Ar << SortedIndicesOffset;
 	}
 
-	virtual void SetMesh(FRHICommandList& RHICmdList, FShader* Shader,const FVertexFactory* VertexFactory,const FSceneView& View,const FMeshBatchElement& BatchElement,uint32 DataFlags) const override
-	{
-		FNiagaraSpriteVertexFactory* SpriteVF = (FNiagaraSpriteVertexFactory*)VertexFactory;
-		FVertexShaderRHIParamRef VertexShaderRHI = Shader->GetVertexShader();
-		SetUniformBufferParameter(RHICmdList, VertexShaderRHI, Shader->GetUniformBufferParameter<FNiagaraSpriteUniformParameters>(), SpriteVF->GetSpriteUniformBuffer() );
-		
-		SetShaderValue(RHICmdList, VertexShaderRHI, NumCutoutVerticesPerFrame, SpriteVF->GetNumCutoutVerticesPerFrame());
-		FShaderResourceViewRHIParamRef NullSRV = GFNiagaraNullSubUVCutoutVertexBuffer.VertexBufferSRV;
-		SetSRVParameter(RHICmdList, VertexShaderRHI, CutoutGeometry, SpriteVF->GetCutoutGeometrySRV() ? SpriteVF->GetCutoutGeometrySRV() : NullSRV);
-
-		SetShaderValue(RHICmdList, VertexShaderRHI, ParticleAlignmentMode, SpriteVF->GetAlignmentMode());
-		SetShaderValue(RHICmdList, VertexShaderRHI, ParticleFacingMode, SpriteVF->GetFacingMode());
-		
-		SetSRVParameter(RHICmdList, VertexShaderRHI, NiagaraParticleDataFloat, SpriteVF->GetParticleDataFloatSRV());
-		SetShaderValue(RHICmdList, VertexShaderRHI, FloatDataOffset, SpriteVF->GetFloatDataOffset());
-		SetShaderValue(RHICmdList, VertexShaderRHI, FloatDataStride, SpriteVF->GetFloatDataStride());
-
-		SetSRVParameter(RHICmdList, VertexShaderRHI, SortedIndices, SpriteVF->GetSortedIndicesSRV() ? SpriteVF->GetSortedIndicesSRV() : GFNiagaraNullSortedIndicesVertexBuffer.VertexBufferSRV);
-		SetShaderValue(RHICmdList, VertexShaderRHI, SortedIndicesOffset, SpriteVF->GetSortedIndicesOffset());
-	}
-
 	virtual void GetElementShaderBindings(
 		const FSceneInterface* Scene,
 		const FSceneView* View,
@@ -150,13 +129,6 @@ class FNiagaraSpriteVertexFactoryShaderParametersPS : public FNiagaraSpriteVerte
 {
 public:
 
-	virtual void SetMesh(FRHICommandList& RHICmdList, FShader* Shader,const FVertexFactory* VertexFactory,const FSceneView& View,const FMeshBatchElement& BatchElement,uint32 DataFlags) const override
-	{
-		FNiagaraSpriteVertexFactory* SpriteVF = (FNiagaraSpriteVertexFactory*)VertexFactory;
-		FPixelShaderRHIParamRef PixelShaderRHI = Shader->GetPixelShader();
-		SetUniformBufferParameter(RHICmdList, PixelShaderRHI, Shader->GetUniformBufferParameter<FNiagaraSpriteUniformParameters>(), SpriteVF->GetSpriteUniformBuffer() );
-	}
-
 	virtual void GetElementShaderBindings(
 		const FSceneInterface* Scene,
 		const FSceneView* View,
diff --git a/Engine/Shaders/Public/ShaderVersion.ush b/Engine/Shaders/Public/ShaderVersion.ush
--- a/Engine/Shaders/Public/ShaderVersion.ush
+++ b/Engine/Shaders/Public/ShaderVersion.ush
@@ -3,4 +3,4 @@
 // in Platform.ush (which should be included in any shader) it allows to invalidate the shader DDC.
 // 
 // If you are merging streams and there is a conflict with this GUID you should make a new GUID rather than taking one or the other.
-// GUID = 2ED93466-374C-4558-B573-BC368C3718D7
+// GUID = 2ED93466-374C-4558-B573-BC368C3718D8
diff --git a/Engine/Source/Developer/MaterialBaking/Private/MaterialRenderItem.cpp b/Engine/Source/Developer/MaterialBaking/Private/MaterialRenderItem.cpp
--- a/Engine/Source/Developer/MaterialBaking/Private/MaterialRenderItem.cpp
+++ b/Engine/Source/Developer/MaterialBaking/Private/MaterialRenderItem.cpp
@@ -7,8 +7,8 @@
 #include "MeshDescription.h"
 #include "MeshAttributes.h"
 #include "MeshAttributeArray.h"
-#include "DrawingPolicy.h"
 #include "DynamicMeshBuilder.h"
+#include "MeshPassProcessor.h"
 
 #define SHOW_WIREFRAME_MESH 0
 
diff --git a/Engine/Source/Developer/MaterialUtilities/Private/MeshRendering.cpp b/Engine/Source/Developer/MaterialUtilities/Private/MeshRendering.cpp
--- a/Engine/Source/Developer/MaterialUtilities/Private/MeshRendering.cpp
+++ b/Engine/Source/Developer/MaterialUtilities/Private/MeshRendering.cpp
@@ -26,9 +26,8 @@
 #include "CanvasItem.h"
 #include "CanvasTypes.h"
 #include "LocalVertexFactory.h"
-#include "DrawingPolicy.h"
 #include "Rendering/SkeletalMeshLODRenderData.h"
-
+#include "MeshPassProcessor.h"
 #include "RendererInterface.h"
 #include "EngineModule.h"
 #include "LightMapHelpers.h"
diff --git a/Engine/Source/Editor/UnrealEd/Private/Lightmass/LightmassLandscapeRender.cpp b/Engine/Source/Editor/UnrealEd/Private/Lightmass/LightmassLandscapeRender.cpp
--- a/Engine/Source/Editor/UnrealEd/Private/Lightmass/LightmassLandscapeRender.cpp
+++ b/Engine/Source/Editor/UnrealEd/Private/Lightmass/LightmassLandscapeRender.cpp
@@ -15,7 +15,6 @@
 #include "LocalVertexFactory.h"
 #include "CanvasTypes.h"
 #include "MeshBatch.h"
-#include "DrawingPolicy.h"
 
 #include "LandscapeProxy.h"
 #include "LandscapeInfo.h"
@@ -24,6 +23,7 @@
 #include "EngineModule.h"
 #include "LandscapeEdit.h"
 #include "DynamicMeshBuilder.h"
+#include "MeshPassProcessor.h"
 
 void RenderLandscapeMaterialForLightmass(const FLandscapeStaticLightingMesh* LandscapeMesh, FMaterialRenderProxy* MaterialProxy, const FRenderTarget* RenderTarget)
 {
diff --git a/Engine/Source/Runtime/Engine/Private/Components/SplineMeshComponent.cpp b/Engine/Source/Runtime/Engine/Private/Components/SplineMeshComponent.cpp
--- a/Engine/Source/Runtime/Engine/Private/Components/SplineMeshComponent.cpp
+++ b/Engine/Source/Runtime/Engine/Private/Components/SplineMeshComponent.cpp
@@ -57,53 +57,6 @@ void FSplineMeshVertexFactoryShaderParameters::Bind(const FShaderParameterMap& P
 	SplineMeshYParam.Bind(ParameterMap, TEXT("SplineMeshY"), SPF_Mandatory);
 }
 
-void FSplineMeshVertexFactoryShaderParameters::SetMesh(FRHICommandList& RHICmdList, FShader* Shader, const FVertexFactory* VertexFactory, const FSceneView& View, const FMeshBatchElement& BatchElement, uint32 DataFlags) const
-{
-	if (BatchElement.bUserDataIsColorVertexBuffer)
-	{
-		FColorVertexBuffer* OverrideColorVertexBuffer = (FColorVertexBuffer*)BatchElement.UserData;
-		check(OverrideColorVertexBuffer);
-		static_cast<const FLocalVertexFactory*>(VertexFactory)->SetColorOverrideStream(RHICmdList, OverrideColorVertexBuffer);
-	}
-
-	FVertexShaderRHIRef VertexShader = Shader->GetVertexShader();
-
-	if (VertexShader)
-	{
-		checkSlow(BatchElement.bIsSplineProxy);
-		FSplineMeshSceneProxy* SplineProxy = BatchElement.SplineMeshSceneProxy;
-		FSplineMeshParams& SplineParams = SplineProxy->SplineParams;
-
-		SetShaderValue(RHICmdList, VertexShader, SplineStartPosParam, SplineParams.StartPos);
-		SetShaderValue(RHICmdList, VertexShader, SplineStartTangentParam, SplineParams.StartTangent);
-		SetShaderValue(RHICmdList, VertexShader, SplineStartRollParam, SplineParams.StartRoll);
-		SetShaderValue(RHICmdList, VertexShader, SplineStartScaleParam, SplineParams.StartScale);
-		SetShaderValue(RHICmdList, VertexShader, SplineStartOffsetParam, SplineParams.StartOffset);
-
-		SetShaderValue(RHICmdList, VertexShader, SplineEndPosParam, SplineParams.EndPos);
-		SetShaderValue(RHICmdList, VertexShader, SplineEndTangentParam, SplineParams.EndTangent);
-		SetShaderValue(RHICmdList, VertexShader, SplineEndRollParam, SplineParams.EndRoll);
-		SetShaderValue(RHICmdList, VertexShader, SplineEndScaleParam, SplineParams.EndScale);
-		SetShaderValue(RHICmdList, VertexShader, SplineEndOffsetParam, SplineParams.EndOffset);
-
-		SetShaderValue(RHICmdList, VertexShader, SplineUpDirParam, SplineProxy->SplineUpDir);
-		SetShaderValue(RHICmdList, VertexShader, SmoothInterpRollScaleParam, SplineProxy->bSmoothInterpRollScale);
-
-		SetShaderValue(RHICmdList, VertexShader, SplineMeshMinZParam, SplineProxy->SplineMeshMinZ);
-		SetShaderValue(RHICmdList, VertexShader, SplineMeshScaleZParam, SplineProxy->SplineMeshScaleZ);
-
-		FVector DirMask(0, 0, 0);
-		DirMask[SplineProxy->ForwardAxis] = 1;
-		SetShaderValue(RHICmdList, VertexShader, SplineMeshDirParam, DirMask);
-		DirMask = FVector::ZeroVector;
-		DirMask[(SplineProxy->ForwardAxis + 1) % 3] = 1;
-		SetShaderValue(RHICmdList, VertexShader, SplineMeshXParam, DirMask);
-		DirMask = FVector::ZeroVector;
-		DirMask[(SplineProxy->ForwardAxis + 2) % 3] = 1;
-		SetShaderValue(RHICmdList, VertexShader, SplineMeshYParam, DirMask);
-	}
-}
-
 void FSplineMeshVertexFactoryShaderParameters::GetElementShaderBindings(
 	const class FSceneInterface* Scene,
 	const FSceneView* View,
diff --git a/Engine/Source/Runtime/Engine/Private/GPUSkinCache.cpp b/Engine/Source/Runtime/Engine/Private/GPUSkinCache.cpp
--- a/Engine/Source/Runtime/Engine/Private/GPUSkinCache.cpp
+++ b/Engine/Source/Runtime/Engine/Private/GPUSkinCache.cpp
@@ -1010,9 +1010,7 @@ void FGPUSkinCache::ProcessEntry(FRHICommandListImmediate& RHICmdList, FGPUBaseS
 	INC_DWORD_STAT(STAT_GPUSkinCache_NumSectionsProcessed);
 
 	const uint32 NumVertices = BatchElement.GetNumVertices();
-	uint32 StreamStrides[MaxVertexElementCount];
 	//#todo-gpuskin Check that stream 0 is the position stream
-	const uint32 StreamStrideCount = VertexFactory->GetStreamStrides(StreamStrides);
 	const uint32 InputStreamStart = BatchElement.BaseVertexIndex;
 
 	FSkeletalMeshRenderData& SkelMeshRenderData = Skin->GetSkeletalMeshRenderData();
diff --git a/Engine/Source/Runtime/Engine/Private/GPUSkinVertexFactory.cpp b/Engine/Source/Runtime/Engine/Private/GPUSkinVertexFactory.cpp
--- a/Engine/Source/Runtime/Engine/Private/GPUSkinVertexFactory.cpp
+++ b/Engine/Source/Runtime/Engine/Private/GPUSkinVertexFactory.cpp
@@ -525,49 +525,6 @@ public:
 		Ar << PreviousBoneMatrices;
 	}
 
-	/**
-	* Set any shader data specific to this vertex factory
-	*/
-	virtual void SetMesh(FRHICommandList& RHICmdList, FShader* Shader, const FVertexFactory* VertexFactory, const FSceneView& View, const FMeshBatchElement& BatchElement, uint32 DataFlags) const override
-	{
-		FRHIVertexShader* ShaderRHI = Shader->GetVertexShader();
-
-		if(ShaderRHI)
-		{
-			const FGPUBaseSkinVertexFactory::FShaderDataType& ShaderData = ((const FGPUBaseSkinVertexFactory*)VertexFactory)->GetShaderData();
-	
-			const auto FeatureLevel = View.GetFeatureLevel();
-
-			bool bLocalPerBoneMotionBlur = false;
-
-			if (SupportsBonesBufferSRV(FeatureLevel))
-			{
-				if(BoneMatrices.IsBound())
-				{
-					FShaderResourceViewRHIParamRef CurrentData = ShaderData.GetBoneBufferForReading(false).VertexBufferSRV;
-
-					RHICmdList.SetShaderResourceViewParameter(ShaderRHI, BoneMatrices.GetBaseIndex(), CurrentData);
-				}
-				if(PreviousBoneMatrices.IsBound())
-				{
-					// todo: Maybe a check for PreviousData!=CurrentData would save some performance (when objects don't have velocty yet) but removing the bool also might save performance
-					bLocalPerBoneMotionBlur = true;
-
-					FShaderResourceViewRHIParamRef PreviousData = ShaderData.GetBoneBufferForReading(true).VertexBufferSRV;
-
-					RHICmdList.SetShaderResourceViewParameter(ShaderRHI, PreviousBoneMatrices.GetBaseIndex(), PreviousData);
-				}
-			}
-			else
-			{
-				SetUniformBufferParameter(RHICmdList, ShaderRHI, Shader->GetUniformBufferParameter<FBoneMatricesUniformShaderParameters>(), ShaderData.GetUniformBuffer());
-			}
-
-
-			SetShaderValue(RHICmdList, ShaderRHI, PerBoneMotionBlur, (uint32)(bLocalPerBoneMotionBlur ? 1 : 0));
-		}
-	}
-
 	virtual void GetElementShaderBindings(
 		const FSceneInterface* Scene,
 		const FSceneView* View,
@@ -651,18 +608,7 @@ public:
 		FLocalVertexFactoryShaderParametersBase::Serialize(Ar);
 		Ar << GPUSkinCachePreviousPositionBuffer;
 	}
-	/**
-	* Set any shader data specific to this vertex factory
-	*/
-	virtual void SetMesh(FRHICommandList& RHICmdList, FShader* Shader,const FVertexFactory* VertexFactory,const FSceneView& View,const FMeshBatchElement& BatchElement,uint32 DataFlags) const override
-	{
-		FLocalVertexFactoryShaderParametersBase::SetMesh(RHICmdList, Shader, VertexFactory, View, BatchElement, DataFlags);
-		check(VertexFactory->GetType() == &FGPUSkinPassthroughVertexFactory::StaticType);
-		FGPUSkinBatchElementUserData* BatchUserData = (FGPUSkinBatchElementUserData*)BatchElement.VertexFactoryUserData;
-		check(BatchUserData);
-		FGPUSkinCache::SetVertexStreams(BatchUserData->Entry, BatchUserData->Section, RHICmdList, Shader, (FGPUSkinPassthroughVertexFactory*)VertexFactory, BatchElement.MinVertexIndex, GPUSkinCachePreviousPositionBuffer);
-	}
-
+	
 	virtual void GetElementShaderBindings(
 		const FSceneInterface* Scene,
 		const FSceneView* View,
@@ -891,75 +837,6 @@ public:
 		Ar << GPUSkinApexClothParameter;
 		Ar << GPUSkinApexClothStartIndexOffsetParameter;
 	}
-
-	virtual void SetMesh(FRHICommandList& RHICmdList, FShader* Shader, const FVertexFactory* VertexFactory, const FSceneView& View, const FMeshBatchElement& BatchElement, uint32 DataFlags) const override
-	{
-		FRHIVertexShader* VertexShader = Shader->GetVertexShader();
-		if (VertexShader)
-		{
-			// Call regular GPU skinning shader parameters
-			FGPUSkinVertexFactoryShaderParameters::SetMesh(RHICmdList, Shader, VertexFactory, View, BatchElement, DataFlags);
-			const auto* GPUSkinVertexFactory = (const FGPUBaseSkinVertexFactory*)VertexFactory;
-			// A little hacky; problem is we can't upcast from FGPUBaseSkinVertexFactory to FGPUBaseSkinAPEXClothVertexFactory as they are unrelated; a nice solution would be
-			// to use virtual inheritance, but that requires RTTI and complicates things further...
-			const FGPUBaseSkinAPEXClothVertexFactory::ClothShaderType& ClothShaderData = GPUSkinVertexFactory->UsesExtraBoneInfluences()
-				? ((const TGPUSkinAPEXClothVertexFactory<true>*)GPUSkinVertexFactory)->GetClothShaderData()
-				: ((const TGPUSkinAPEXClothVertexFactory<false>*)GPUSkinVertexFactory)->GetClothShaderData();
-
-			SetUniformBufferParameter(RHICmdList, VertexShader, Shader->GetUniformBufferParameter<FAPEXClothUniformShaderParameters>(),ClothShaderData.GetClothUniformBuffer());
-
-			uint32 FrameNumber = View.Family->FrameNumber;
-
-			// we tell the shader where to pickup the data
-			if(ClothSimulVertsPositionsNormalsParameter.IsBound())
-			{
-				RHICmdList.SetShaderResourceViewParameter(VertexShader, ClothSimulVertsPositionsNormalsParameter.GetBaseIndex(),
-														  ClothShaderData.GetClothBufferForReading(false, FrameNumber).VertexBufferSRV);
-			}
-			if(PreviousClothSimulVertsPositionsNormalsParameter.IsBound())
-			{
-				RHICmdList.SetShaderResourceViewParameter(VertexShader, PreviousClothSimulVertsPositionsNormalsParameter.GetBaseIndex(),
-														  ClothShaderData.GetClothBufferForReading(true, FrameNumber).VertexBufferSRV);
-			}
-			
-			SetShaderValue(
-				RHICmdList,
-				VertexShader,
-				ClothLocalToWorldParameter,
-				ClothShaderData.GetClothLocalToWorldForReading(false, FrameNumber)
-				);
-
-			SetShaderValue(
-				RHICmdList,
-				VertexShader,
-				PreviousClothLocalToWorldParameter,
-				ClothShaderData.GetClothLocalToWorldForReading(true, FrameNumber)
-				);
-
-			SetShaderValue(
-				RHICmdList,
-				VertexShader,
-				ClothBlendWeightParameter,
-				ClothShaderData.ClothBlendWeight
-				);
-
-			if (GPUSkinApexClothParameter.IsBound())
-			{
-				RHICmdList.SetShaderResourceViewParameter(
-					VertexShader,
-					GPUSkinApexClothParameter.GetBaseIndex(),
-					GPUSkinVertexFactory->UsesExtraBoneInfluences()
-					? ((const TGPUSkinAPEXClothVertexFactory<true>*)GPUSkinVertexFactory)->GetClothBuffer()
-					: ((const TGPUSkinAPEXClothVertexFactory<false>*)GPUSkinVertexFactory)->GetClothBuffer());
-				int32 ClothIndexOffset =
-					GPUSkinVertexFactory->UsesExtraBoneInfluences()
-					? ((const TGPUSkinAPEXClothVertexFactory<true>*)GPUSkinVertexFactory)->GetClothIndexOffset(BatchElement.MinVertexIndex)
-					: ((const TGPUSkinAPEXClothVertexFactory<false>*)GPUSkinVertexFactory)->GetClothIndexOffset(BatchElement.MinVertexIndex);
-				FIntVector4 GPUSkinApexClothStartIndexOffset(BatchElement.MinVertexIndex, ClothIndexOffset, 0, 0);
-				SetShaderValue(RHICmdList, VertexShader, GPUSkinApexClothStartIndexOffsetParameter, GPUSkinApexClothStartIndexOffset);
-			}
-		}
-	}
 	
 	virtual void GetElementShaderBindings(
 		const FSceneInterface* Scene,
diff --git a/Engine/Source/Runtime/Engine/Private/GeometryCacheVertexFactory.cpp b/Engine/Source/Runtime/Engine/Private/GeometryCacheVertexFactory.cpp
--- a/Engine/Source/Runtime/Engine/Private/GeometryCacheVertexFactory.cpp
+++ b/Engine/Source/Runtime/Engine/Private/GeometryCacheVertexFactory.cpp
@@ -46,38 +46,6 @@ public:
 		Ar << MotionBlurPositionScale;
 	}
 
-	/**
-	* Set any shader data specific to this vertex factory
-	*/
-	virtual void SetMesh(FRHICommandList& RHICmdList, FShader* Shader, const FVertexFactory* GenericVertexFactory, const FSceneView& View, const FMeshBatchElement& BatchElement, uint32 DataFlags) const override
-	{
-		// Ensure the vertex factory matches this parameter object and cast relevant objects
-		check(GenericVertexFactory->GetType() == &FGeometryCacheVertexVertexFactory::StaticType);
-		const FGeometryCacheVertexVertexFactory* GCVertexFactory = static_cast<const FGeometryCacheVertexVertexFactory*>(GenericVertexFactory);
-
-
-		FVertexShaderRHIParamRef VS = Shader->GetVertexShader();
-
-		FGeometryCacheVertexFactoryUserData* BatchData = (FGeometryCacheVertexFactoryUserData*)BatchElement.VertexFactoryUserData;
-
-		// Check the passed in vertex buffers make sense
-		checkf(BatchData->PositionBuffer->IsInitialized(), TEXT("Batch position Vertex buffer was not initialized! Name %s"), *BatchData->PositionBuffer->GetFriendlyName());
-		checkf(BatchData->MotionBlurDataBuffer->IsInitialized(), TEXT("Batch motion blur data buffer was not initialized! Name %s"), *BatchData->MotionBlurDataBuffer->GetFriendlyName());
-
-
-		RHICmdList.SetStreamSource(GCVertexFactory->PositionStreamIndex, BatchData->PositionBuffer->VertexBufferRHI, 0);
-		RHICmdList.SetStreamSource(GCVertexFactory->MotionBlurDataStreamIndex, BatchData->MotionBlurDataBuffer->VertexBufferRHI, 0);
-
-		if (VS)
-		{
-			SetShaderValue(RHICmdList, VS, MeshOrigin, BatchData->MeshOrigin);
-			SetShaderValue(RHICmdList, VS, MeshExtension, BatchData->MeshExtension);
-			SetShaderValue(RHICmdList, VS, MotionBlurDataOrigin, BatchData->MotionBlurDataOrigin);
-			SetShaderValue(RHICmdList, VS, MotionBlurDataExtension, BatchData->MotionBlurDataExtension);
-			SetShaderValue(RHICmdList, VS, MotionBlurPositionScale, BatchData->MotionBlurPositionScale);
-		}
-	}
-
 	virtual void GetElementShaderBindings(
 		const class FSceneInterface* Scene,
 		const FSceneView* View,
diff --git a/Engine/Source/Runtime/Engine/Private/InstancedStaticMesh.cpp b/Engine/Source/Runtime/Engine/Private/InstancedStaticMesh.cpp
--- a/Engine/Source/Runtime/Engine/Private/InstancedStaticMesh.cpp
+++ b/Engine/Source/Runtime/Engine/Private/InstancedStaticMesh.cpp
@@ -2467,193 +2467,6 @@ static TAutoConsoleVariable<int32> CVarCullAllInVertexShader(
 	0,
 	TEXT("Debugging, if this is greater than 0, cull all instances in the vertex shader."));
 
-void FInstancedStaticMeshVertexFactoryShaderParameters::SetMesh( FRHICommandList& RHICmdList, FShader* VertexShader,const class FVertexFactory* VertexFactory,const class FSceneView& View,const struct FMeshBatchElement& BatchElement,uint32 DataFlags ) const 
-{
-	FLocalVertexFactoryShaderParametersBase::SetMesh(RHICmdList, VertexShader, VertexFactory, View, BatchElement, DataFlags);
-
-	FRHIVertexShader* VS = VertexShader->GetVertexShader();
-
-	const int32 InstanceOffsetValue = BatchElement.UserIndex;
-	const auto* InstancedVertexFactory = static_cast<const FInstancedStaticMeshVertexFactory*>(VertexFactory);
-	if (InstancedVertexFactory->SupportsManualVertexFetch(View.GetFeatureLevel()))
-	{
-		SetSRVParameter(RHICmdList, VS, VertexFetch_InstanceOriginBufferParameter, InstancedVertexFactory->GetInstanceOriginSRV());
-		SetSRVParameter(RHICmdList, VS, VertexFetch_InstanceTransformBufferParameter, InstancedVertexFactory->GetInstanceTransformSRV());
-		SetSRVParameter(RHICmdList, VS, VertexFetch_InstanceLightmapBufferParameter, InstancedVertexFactory->GetInstanceLightmapSRV());
-
-		const bool bZeroInstanceOffset = RHISupportsAbsoluteVertexID(GMaxRHIShaderPlatform);
-		SetShaderValue(RHICmdList, VS, InstanceOffset, bZeroInstanceOffset ? 0 : InstanceOffsetValue);
-	}
-
-	if (GRHISupportsInstancing && InstanceOffsetValue > 0)
-	{
-		VertexFactory->OffsetInstanceStreams(RHICmdList, InstanceOffsetValue);
-	}
-
-	const FInstancingUserData* InstancingUserData = (const FInstancingUserData*)BatchElement.UserData;
-	if( InstancingWorldViewOriginOneParameter.IsBound() )
-	{
-		FVector4 InstancingViewZCompareZero(MIN_flt, MIN_flt, MAX_flt, 1.0f);
-		FVector4 InstancingViewZCompareOne(MIN_flt, MIN_flt, MAX_flt, 0.0f);
-		FVector4 InstancingViewZConstant(ForceInit);
-		FVector4 InstancingWorldViewOriginZero(ForceInit);
-		FVector4 InstancingWorldViewOriginOne(ForceInit);
-		InstancingWorldViewOriginOne.W = 1.0f;
-		if (InstancingUserData && BatchElement.InstancedLODRange)
-		{
-			int32 FirstLOD = InstancingUserData->MinLOD;
-
-			int32 DebugMin = FMath::Min(CVarMinLOD.GetValueOnRenderThread(), InstancingUserData->MeshRenderData->LODResources.Num() - 1);
-			if (DebugMin >= 0)
-			{
-				FirstLOD = FMath::Max(FirstLOD, DebugMin);
-			}
-
-
-			float SphereRadius = InstancingUserData->MeshRenderData->Bounds.SphereRadius;
-			float MinSize = View.ViewMatrices.IsPerspectiveProjection() ? CVarFoliageMinimumScreenSize.GetValueOnRenderThread() : 0.0f;
-			float LODScale = CVarFoliageLODDistanceScale.GetValueOnRenderThread();
-			float LODRandom = CVarRandomLODRange.GetValueOnRenderThread();
-			float MaxDrawDistanceScale = GetCachedScalabilityCVars().ViewDistanceScale;
-
-			if (BatchElement.InstancedLODIndex)
-			{
-				InstancingViewZConstant.X = -1.0f;
-			}
-			else
-			{
-				// this is the first LOD, so we don't have a fade-in region
-				InstancingViewZConstant.X = 0.0f;
-			}
-			InstancingViewZConstant.Y = 0.0f;
-			InstancingViewZConstant.Z = 1.0f;
-
-			// now we subtract off the lower segments, since they will be incorporated 
-			InstancingViewZConstant.Y -= InstancingViewZConstant.X;
-			InstancingViewZConstant.Z -= InstancingViewZConstant.X + InstancingViewZConstant.Y;
-			//not using W InstancingViewZConstant.W -= InstancingViewZConstant.X + InstancingViewZConstant.Y + InstancingViewZConstant.Z;
-
-			for (int32 SampleIndex = 0; SampleIndex < 2; SampleIndex++)
-			{
-				FVector4& InstancingViewZCompare(SampleIndex ? InstancingViewZCompareOne : InstancingViewZCompareZero);
-
-				float FinalCull = MAX_flt;
-				if (MinSize > 0.0)
-				{
-					FinalCull = ComputeBoundsDrawDistance(MinSize, SphereRadius, View.ViewMatrices.GetProjectionMatrix()) * LODScale;
-				}
-				if (InstancingUserData->EndCullDistance > 0.0f)
-				{
-					FinalCull = FMath::Min(FinalCull, InstancingUserData->EndCullDistance * MaxDrawDistanceScale);
-				}
-				FinalCull *= MaxDrawDistanceScale;
-
-				InstancingViewZCompare.Z = FinalCull;
-				if (BatchElement.InstancedLODIndex < InstancingUserData->MeshRenderData->LODResources.Num() - 1)
-				{
-					float NextCut = ComputeBoundsDrawDistance(InstancingUserData->MeshRenderData->ScreenSize[BatchElement.InstancedLODIndex + 1].GetValueForFeatureLevel(View.GetFeatureLevel()), SphereRadius, View.ViewMatrices.GetProjectionMatrix()) * LODScale;
-					InstancingViewZCompare.Z = FMath::Min(NextCut, FinalCull);
-				}
-
-				InstancingViewZCompare.X = MIN_flt;
-				if (BatchElement.InstancedLODIndex > FirstLOD)
-				{
-					float CurCut = ComputeBoundsDrawDistance(InstancingUserData->MeshRenderData->ScreenSize[BatchElement.InstancedLODIndex].GetValueForFeatureLevel(View.GetFeatureLevel()), SphereRadius, View.ViewMatrices.GetProjectionMatrix()) * LODScale;
-					if (CurCut < FinalCull)
-					{
-						InstancingViewZCompare.Y = CurCut;
-					}
-					else
-					{
-						// this LOD is completely removed by one of the other two factors
-						InstancingViewZCompare.Y = MIN_flt;
-						InstancingViewZCompare.Z = MIN_flt;
-					}
-				}
-				else
-				{
-					// this is the first LOD, so we don't have a fade-in region
-					InstancingViewZCompare.Y = MIN_flt;
-				}
-			}
-
-
-			InstancingWorldViewOriginZero = View.GetTemporalLODOrigin(0);
-			InstancingWorldViewOriginOne = View.GetTemporalLODOrigin(1);
-
-			float Alpha = View.GetTemporalLODTransition();
-			InstancingWorldViewOriginZero.W = 1.0f - Alpha;
-			InstancingWorldViewOriginOne.W = Alpha;
-
-			InstancingViewZCompareZero.W = LODRandom;
-		}
-		SetShaderValue(RHICmdList, VS, InstancingViewZCompareZeroParameter, InstancingViewZCompareZero );
-		SetShaderValue(RHICmdList, VS, InstancingViewZCompareOneParameter, InstancingViewZCompareOne );
-		SetShaderValue(RHICmdList, VS, InstancingViewZConstantParameter, InstancingViewZConstant );
-		SetShaderValue(RHICmdList, VS, InstancingWorldViewOriginZeroParameter, InstancingWorldViewOriginZero );
-		SetShaderValue(RHICmdList, VS, InstancingWorldViewOriginOneParameter, InstancingWorldViewOriginOne );
-	}
-	if( InstancingFadeOutParamsParameter.IsBound() )
-	{
-		FVector4 InstancingFadeOutParams(MAX_flt,0.f,1.f,1.f);
-		if (InstancingUserData)
-		{
-			const float MaxDrawDistanceScale = GetCachedScalabilityCVars().ViewDistanceScale;
-			const float StartDistance = InstancingUserData->StartCullDistance * MaxDrawDistanceScale;
-			const float EndDistance = InstancingUserData->EndCullDistance * MaxDrawDistanceScale;
-
-			InstancingFadeOutParams.X = StartDistance;
-			if( EndDistance > 0 )
-			{
-				if( EndDistance > StartDistance )
-				{
-					InstancingFadeOutParams.Y = 1.f / (float)(EndDistance - StartDistance);
-				}
-				else
-				{
-					InstancingFadeOutParams.Y = 1.f;
-				}
-			}
-			else
-			{
-				InstancingFadeOutParams.Y = 0.f;
-			}
-			if (CVarCullAllInVertexShader.GetValueOnRenderThread() > 0)
-			{
-				InstancingFadeOutParams.Z = 0.0f;
-				InstancingFadeOutParams.W = 0.0f;
-			}
-			else
-			{
-			InstancingFadeOutParams.Z = InstancingUserData->bRenderSelected ? 1.f : 0.f;
-			InstancingFadeOutParams.W = InstancingUserData->bRenderUnselected ? 1.f : 0.f;
-		}
-		}
-		SetShaderValue(RHICmdList, VS, InstancingFadeOutParamsParameter, InstancingFadeOutParams );
-	}
-
-	auto ShaderPlatform = GetFeatureLevelShaderPlatform(View.GetFeatureLevel());
-	const bool bInstanced = GRHISupportsInstancing;
-	if (!bInstanced)
-	{
-		if (CPUInstanceOrigin.IsBound())
-		{
-			const float ShortScale = 1.0f / 32767.0f;
-			auto* InstancingData = (const FInstancingUserData*)BatchElement.UserData;
-			check(InstancingData);
-
-			FVector4 InstanceTransform[3];
-			FVector4 InstanceLightmapAndShadowMapUVBias;
-			FVector4 InstanceOrigin;
-			InstancingData->RenderData->PerInstanceRenderData->InstanceBuffer.GetInstanceShaderValues(BatchElement.UserIndex, InstanceTransform, InstanceLightmapAndShadowMapUVBias, InstanceOrigin);
-
-			SetShaderValue(RHICmdList, VS, CPUInstanceOrigin, InstanceOrigin);
-			SetShaderValueArray(RHICmdList, VS, CPUInstanceTransform, InstanceTransform, 3);
-			SetShaderValue(RHICmdList, VS, CPUInstanceLightmapAndShadowMapBias, InstanceLightmapAndShadowMapUVBias); 
-		}
-	}
-}
-
 void FInstancedStaticMeshVertexFactoryShaderParameters::GetElementShaderBindings(
 	const class FSceneInterface* Scene,
 	const FSceneView* View,
diff --git a/Engine/Source/Runtime/Engine/Private/InstancedStaticMesh.h b/Engine/Source/Runtime/Engine/Private/InstancedStaticMesh.h
--- a/Engine/Source/Runtime/Engine/Private/InstancedStaticMesh.h
+++ b/Engine/Source/Runtime/Engine/Private/InstancedStaticMesh.h
@@ -329,8 +329,6 @@ class FInstancedStaticMeshVertexFactoryShaderParameters : public FLocalVertexFac
 		InstanceOffset.Bind(ParameterMap, TEXT("InstanceOffset"));
 	}
 
-	virtual void SetMesh(FRHICommandList& RHICmdList, FShader* VertexShader,const class FVertexFactory* VertexFactory,const class FSceneView& View,const struct FMeshBatchElement& BatchElement,uint32 DataFlags) const override;
-
 	virtual void GetElementShaderBindings(
 		const class FSceneInterface* Scene,
 		const FSceneView* View,
diff --git a/Engine/Source/Runtime/Engine/Private/LocalVertexFactory.cpp b/Engine/Source/Runtime/Engine/Private/LocalVertexFactory.cpp
--- a/Engine/Source/Runtime/Engine/Private/LocalVertexFactory.cpp
+++ b/Engine/Source/Runtime/Engine/Private/LocalVertexFactory.cpp
@@ -92,57 +92,6 @@ TUniformBufferRef<FLocalVertexFactoryUniformShaderParameters> CreateLocalVFUnifo
 
 const int32 DefaultBaseVertexIndex = 0;
 
-void FLocalVertexFactoryShaderParametersBase::SetMesh(FRHICommandList& RHICmdList, FShader* Shader, const FVertexFactory* VertexFactory, const FSceneView& View, const FMeshBatchElement& BatchElement, uint32 DataFlags) const
-{
-	const auto* LocalVertexFactory = static_cast<const FLocalVertexFactory*>(VertexFactory);
-	
-	FVertexShaderRHIParamRef VS = Shader->GetVertexShader();
-	if (LocalVertexFactory->SupportsManualVertexFetch(View.GetFeatureLevel()))
-	{
-		FUniformBufferRHIParamRef VertexFactoryUniformBuffer = static_cast<FUniformBufferRHIParamRef>(BatchElement.VertexFactoryUserData);
-
-		if (!VertexFactoryUniformBuffer)
-		{
-			// No batch element override
-			VertexFactoryUniformBuffer = LocalVertexFactory->GetUniformBuffer();
-			// Vertex Offset used by manual vertex fetch must match what we are going to pass to the draw call.  LocalVF default uniform buffer uses VertexOffset of 0.
-			check(BatchElement.BaseVertexIndex == DefaultBaseVertexIndex);
-		}
-
-		SetUniformBufferParameter(RHICmdList, VS, Shader->GetUniformBufferParameter<FLocalVertexFactoryUniformShaderParameters>(), VertexFactoryUniformBuffer);
-	}
-
-	if (BatchElement.bUserDataIsColorVertexBuffer)
-	{
-		FColorVertexBuffer* OverrideColorVertexBuffer = (FColorVertexBuffer*)BatchElement.UserData;
-		check(OverrideColorVertexBuffer);
-
-		if (!LocalVertexFactory->SupportsManualVertexFetch(View.GetFeatureLevel()))
-		{
-			LocalVertexFactory->SetColorOverrideStream(RHICmdList, OverrideColorVertexBuffer);
-		}	
-	}
-
-	if (bAnySpeedTreeParamIsBound && View.Family != NULL && View.Family->Scene != NULL)
-	{
-		QUICK_SCOPE_CYCLE_COUNTER(STAT_FLocalVertexFactoryShaderParameters_SetMesh_SpeedTree);
-		FUniformBufferRHIParamRef SpeedTreeUniformBuffer = View.Family->Scene->GetSpeedTreeUniformBuffer(VertexFactory);
-		if (SpeedTreeUniformBuffer == NULL)
-		{
-			SpeedTreeUniformBuffer = GSpeedTreeWindNullUniformBuffer.GetUniformBufferRHI();
-		}
-		check(SpeedTreeUniformBuffer != NULL);
-
-		SetUniformBufferParameter(RHICmdList, VS, Shader->GetUniformBufferParameter<FSpeedTreeUniformParameters>(), SpeedTreeUniformBuffer);
-
-		if (LODParameter.IsBound())
-		{
-			FVector LODData(BatchElement.MinScreenSize, BatchElement.MaxScreenSize, BatchElement.MaxScreenSize - BatchElement.MinScreenSize);
-			SetShaderValue(RHICmdList, VS, LODParameter, LODData);
-		}
-	}
-}
-
 void FLocalVertexFactoryShaderParametersBase::GetElementShaderBindings(
 	const FSceneInterface* Scene,
 	const FSceneView* View,
diff --git a/Engine/Source/Runtime/Engine/Private/Materials/MaterialShared.cpp b/Engine/Source/Runtime/Engine/Private/Materials/MaterialShared.cpp
--- a/Engine/Source/Runtime/Engine/Private/Materials/MaterialShared.cpp
+++ b/Engine/Source/Runtime/Engine/Private/Materials/MaterialShared.cpp
@@ -2110,16 +2110,12 @@ FMaterialRenderProxy::FMaterialRenderProxy()
 	: SubsurfaceProfileRT(0)
 #if !(UE_BUILD_SHIPPING || UE_BUILD_TEST)
 	, DeletedFlag(0)
-	, bIsStaticDrawListReferenced(0)
 #endif
 {
 }
 
 FMaterialRenderProxy::~FMaterialRenderProxy()
 {
-	// Removed for now to work around UE-31636. Re-enable when the underlying bug is fixed!
-	//check(!IsReferencedInDrawList());
-
 	if(IsInitialized())
 	{
 		check(IsInRenderingThread());
diff --git a/Engine/Source/Runtime/Engine/Private/MeshParticleVertexFactory.cpp b/Engine/Source/Runtime/Engine/Private/MeshParticleVertexFactory.cpp
--- a/Engine/Source/Runtime/Engine/Private/MeshParticleVertexFactory.cpp
+++ b/Engine/Source/Runtime/Engine/Private/MeshParticleVertexFactory.cpp
@@ -47,48 +47,6 @@ public:
 		Ar << PrevTransformBuffer;
 	}
 
-	virtual void SetMesh(FRHICommandList& RHICmdList, FShader* Shader,const FVertexFactory* VertexFactory,const FSceneView& View,const FMeshBatchElement& BatchElement,uint32 DataFlags) const override
-	{
-		const bool bInstanced = GRHISupportsInstancing;
-		FMeshParticleVertexFactory* MeshParticleVF = (FMeshParticleVertexFactory*)VertexFactory;
-		FVertexShaderRHIParamRef VertexShaderRHI = Shader->GetVertexShader();
-		SetUniformBufferParameter(RHICmdList, VertexShaderRHI, Shader->GetUniformBufferParameter<FMeshParticleUniformParameters>(), MeshParticleVF->GetUniformBuffer() );
-
-		if (!bInstanced)
-		{
-			const FMeshParticleVertexFactory::FBatchParametersCPU* BatchParameters = (const FMeshParticleVertexFactory::FBatchParametersCPU*)BatchElement.UserData;
-			const FMeshParticleInstanceVertex* Vertex = BatchParameters->InstanceBuffer + BatchElement.UserIndex;
-			const FMeshParticleInstanceVertexDynamicParameter* DynamicVertex = BatchParameters->DynamicParameterBuffer + BatchElement.UserIndex;
-			const FMeshParticleInstanceVertexPrevTransform* PrevTransformVertex = BatchParameters->PrevTransformBuffer + BatchElement.UserIndex;
-
-			SetShaderValue(RHICmdList, VertexShaderRHI, Transform1, Vertex->Transform[0]);
-			SetShaderValue(RHICmdList, VertexShaderRHI, Transform2, Vertex->Transform[1]);
-			SetShaderValue(RHICmdList, VertexShaderRHI, Transform3, Vertex->Transform[2]);
-			SetShaderValue(RHICmdList, VertexShaderRHI, SubUVParams, FVector4((float)Vertex->SubUVParams[0], (float)Vertex->SubUVParams[1], (float)Vertex->SubUVParams[2], (float)Vertex->SubUVParams[3]));
-			SetShaderValue(RHICmdList, VertexShaderRHI, SubUVLerp, Vertex->SubUVLerp);
-			SetShaderValue(RHICmdList, VertexShaderRHI, ParticleDirection, Vertex->Velocity);
-			SetShaderValue(RHICmdList, VertexShaderRHI, RelativeTime, Vertex->RelativeTime);
-
-			if (BatchParameters->DynamicParameterBuffer)
-			{
-				SetShaderValue(RHICmdList, VertexShaderRHI, DynamicParameter, FVector4(DynamicVertex->DynamicValue[0], DynamicVertex->DynamicValue[1], DynamicVertex->DynamicValue[2], DynamicVertex->DynamicValue[3]));
-			}
-
-			if (BatchParameters->PrevTransformBuffer && View.FeatureLevel >= ERHIFeatureLevel::SM4)
-			{
-				SetShaderValue(RHICmdList, VertexShaderRHI, PrevTransform0, PrevTransformVertex->PrevTransform0);
-				SetShaderValue(RHICmdList, VertexShaderRHI, PrevTransform1, PrevTransformVertex->PrevTransform1);
-				SetShaderValue(RHICmdList, VertexShaderRHI, PrevTransform2, PrevTransformVertex->PrevTransform2);
-			}
-
-			SetShaderValue(RHICmdList, VertexShaderRHI, ParticleColor, FVector4(Vertex->Color.Component(0), Vertex->Color.Component(1), Vertex->Color.Component(2), Vertex->Color.Component(3)));
-		}
-		else if (View.FeatureLevel >= ERHIFeatureLevel::SM4)
-		{
-			SetSRVParameter(RHICmdList, VertexShaderRHI, PrevTransformBuffer, MeshParticleVF->GetPreviousTransformBufferSRV());
-		}
-	}
-
 	virtual void GetElementShaderBindings(
 		const FSceneInterface* Scene,
 		const FSceneView* View,
diff --git a/Engine/Source/Runtime/Engine/Private/Particles/ParticleBeamTrailVertexFactory.cpp b/Engine/Source/Runtime/Engine/Private/Particles/ParticleBeamTrailVertexFactory.cpp
--- a/Engine/Source/Runtime/Engine/Private/Particles/ParticleBeamTrailVertexFactory.cpp
+++ b/Engine/Source/Runtime/Engine/Private/Particles/ParticleBeamTrailVertexFactory.cpp
@@ -25,12 +25,6 @@ public:
 	{
 	}
 
-	virtual void SetMesh(FRHICommandList& RHICmdList, FShader* Shader,const FVertexFactory* VertexFactory,const FSceneView& View,const FMeshBatchElement& BatchElement,uint32 DataFlags) const override
-	{
-		FParticleBeamTrailVertexFactory* BeamTrailVF = (FParticleBeamTrailVertexFactory*)VertexFactory;
-		SetUniformBufferParameter(RHICmdList, Shader->GetVertexShader(), Shader->GetUniformBufferParameter<FParticleBeamTrailUniformParameters>(), BeamTrailVF->GetBeamTrailUniformBuffer() );
-	}
-
 	virtual void GetElementShaderBindings(
 		const FSceneInterface* Scene,
 		const FSceneView* View,
diff --git a/Engine/Source/Runtime/Engine/Private/Particles/ParticleGpuSimulation.cpp b/Engine/Source/Runtime/Engine/Private/Particles/ParticleGpuSimulation.cpp
--- a/Engine/Source/Runtime/Engine/Private/Particles/ParticleGpuSimulation.cpp
+++ b/Engine/Source/Runtime/Engine/Private/Particles/ParticleGpuSimulation.cpp
@@ -597,8 +597,6 @@ public:
 		Ar << CurveTextureSampler;
 	}
 
-	virtual void SetMesh(FRHICommandList& RHICmdList, FShader* Shader,const FVertexFactory* VertexFactory,const FSceneView& View,const FMeshBatchElement& BatchElement,uint32 DataFlags) const override;
-
 	virtual void GetElementShaderBindings(
 		const FSceneInterface* Scene,
 		const FSceneView* View,
@@ -642,8 +640,6 @@ public:
 
 	virtual void Serialize(FArchive& Ar) override {}
 
-	virtual void SetMesh(FRHICommandList& RHICmdList, FShader* Shader,const FVertexFactory* VertexFactory,const FSceneView& View,const FMeshBatchElement& BatchElement,uint32 DataFlags) const override;
-	
 	virtual void GetElementShaderBindings(
 		const FSceneInterface* Scene,
 		const FSceneView* View,
@@ -806,28 +802,6 @@ public:
 	}
 };
 
-/**
- * Set vertex factory shader parameters.
- */
-void FGPUSpriteVertexFactoryShaderParametersVS::SetMesh(FRHICommandList& RHICmdList, FShader* Shader, const FVertexFactory* VertexFactory, const FSceneView& View, const FMeshBatchElement& BatchElement, uint32 DataFlags) const
-{
-	FGPUSpriteVertexFactory* GPUVF = (FGPUSpriteVertexFactory*)VertexFactory;
-	FVertexShaderRHIParamRef VertexShader = Shader->GetVertexShader();
-	FSamplerStateRHIParamRef SamplerStatePoint = TStaticSamplerState<SF_Point>::GetRHI();
-	FSamplerStateRHIParamRef SamplerStateLinear = TStaticSamplerState<SF_Bilinear>::GetRHI();
-	SetUniformBufferParameter(RHICmdList, VertexShader, Shader->GetUniformBufferParameter<FGPUSpriteEmitterUniformParameters>(), GPUVF->EmitterUniformBuffer );
-	SetUniformBufferParameter(RHICmdList, VertexShader, Shader->GetUniformBufferParameter<FGPUSpriteEmitterDynamicUniformParameters>(), GPUVF->EmitterDynamicUniformBuffer );
-	if (ParticleIndices.IsBound())
-	{
-		RHICmdList.SetShaderResourceViewParameter(VertexShader, ParticleIndices.GetBaseIndex(), GPUVF->ParticleIndicesBuffer->VertexBufferSRV);
-	}
-	SetShaderValue(RHICmdList, VertexShader, ParticleIndicesOffset, GPUVF->ParticleIndicesOffset);
-	SetTextureParameter(RHICmdList, VertexShader, PositionTexture, PositionTextureSampler, SamplerStatePoint, GPUVF->PositionTextureRHI );
-	SetTextureParameter(RHICmdList, VertexShader, VelocityTexture, VelocityTextureSampler, SamplerStatePoint, GPUVF->VelocityTextureRHI );
-	SetTextureParameter(RHICmdList, VertexShader, AttributesTexture, AttributesTextureSampler, SamplerStatePoint, GPUVF->AttributesTextureRHI );
-	SetTextureParameter(RHICmdList, VertexShader, CurveTexture, CurveTextureSampler, SamplerStateLinear, GParticleCurveTexture.GetCurveTexture() );
-}
-
 void FGPUSpriteVertexFactoryShaderParametersVS::GetElementShaderBindings(
 	const FSceneInterface* Scene,
 	const FSceneView* View,
@@ -854,13 +828,6 @@ void FGPUSpriteVertexFactoryShaderParametersVS::GetElementShaderBindings(
 	ShaderBindings.AddTexture(CurveTexture, CurveTextureSampler, SamplerStateLinear, GParticleCurveTexture.GetCurveTexture());
 }
 
-void FGPUSpriteVertexFactoryShaderParametersPS::SetMesh(FRHICommandList& RHICmdList, FShader* Shader, const FVertexFactory* VertexFactory, const FSceneView& View, const FMeshBatchElement& BatchElement, uint32 DataFlags) const
-{
-	FGPUSpriteVertexFactory* GPUVF = (FGPUSpriteVertexFactory*)VertexFactory;
-	FPixelShaderRHIParamRef PixelShader = Shader->GetPixelShader();
-	SetUniformBufferParameter(RHICmdList, PixelShader, Shader->GetUniformBufferParameter<FGPUSpriteEmitterDynamicUniformParameters>(), GPUVF->EmitterDynamicUniformBuffer);
-}
-
 void FGPUSpriteVertexFactoryShaderParametersPS::GetElementShaderBindings(
 	const FSceneInterface* Scene,
 	const FSceneView* View,
diff --git a/Engine/Source/Runtime/Engine/Private/Particles/ParticleVertexFactory.cpp b/Engine/Source/Runtime/Engine/Private/Particles/ParticleVertexFactory.cpp
--- a/Engine/Source/Runtime/Engine/Private/Particles/ParticleVertexFactory.cpp
+++ b/Engine/Source/Runtime/Engine/Private/Particles/ParticleVertexFactory.cpp
@@ -77,17 +77,6 @@ public:
 		Ar << CutoutGeometry;
 	}
 
-	virtual void SetMesh(FRHICommandList& RHICmdList, FShader* Shader,const FVertexFactory* VertexFactory,const FSceneView& View,const FMeshBatchElement& BatchElement,uint32 DataFlags) const override
-	{
-		FParticleSpriteVertexFactory* SpriteVF = (FParticleSpriteVertexFactory*)VertexFactory;
-		FVertexShaderRHIParamRef VertexShaderRHI = Shader->GetVertexShader();
-		SetUniformBufferParameter(RHICmdList, VertexShaderRHI, Shader->GetUniformBufferParameter<FParticleSpriteUniformParameters>(), SpriteVF->GetSpriteUniformBuffer() );
-		
-		SetShaderValue(RHICmdList, VertexShaderRHI, NumCutoutVerticesPerFrame, SpriteVF->GetNumCutoutVerticesPerFrame());
-		FShaderResourceViewRHIParamRef NullSRV = GFNullSubUVCutoutVertexBuffer.VertexBufferSRV;
-		SetSRVParameter(RHICmdList, VertexShaderRHI, CutoutGeometry, SpriteVF->GetCutoutGeometrySRV() ? SpriteVF->GetCutoutGeometrySRV() : NullSRV);
-	}
-
 	virtual void GetElementShaderBindings(
 		const FSceneInterface* Scene,
 		const FSceneView* View,
@@ -116,13 +105,6 @@ class FParticleSpriteVertexFactoryShaderParametersPS : public FParticleSpriteVer
 {
 public:
 
-	virtual void SetMesh(FRHICommandList& RHICmdList, FShader* Shader,const FVertexFactory* VertexFactory,const FSceneView& View,const FMeshBatchElement& BatchElement,uint32 DataFlags) const override
-	{
-		FParticleSpriteVertexFactory* SpriteVF = (FParticleSpriteVertexFactory*)VertexFactory;
-		FPixelShaderRHIParamRef PixelShaderRHI = Shader->GetPixelShader();
-		SetUniformBufferParameter(RHICmdList, PixelShaderRHI, Shader->GetUniformBufferParameter<FParticleSpriteUniformParameters>(), SpriteVF->GetSpriteUniformBuffer() );
-	}
-
 	virtual void GetElementShaderBindings(
 		const FSceneInterface* Scene,
 		const FSceneView* View,
diff --git a/Engine/Source/Runtime/Engine/Private/SceneManagement.cpp b/Engine/Source/Runtime/Engine/Private/SceneManagement.cpp
--- a/Engine/Source/Runtime/Engine/Private/SceneManagement.cpp
+++ b/Engine/Source/Runtime/Engine/Private/SceneManagement.cpp
@@ -16,8 +16,6 @@
 #include "Engine/ShadowMapTexture2D.h"
 #include "UnrealEngine.h"
 
-ENGINE_API bool GDrawListsLocked = false;
-
 static TAutoConsoleVariable<float> CVarLODTemporalLag(
 	TEXT("lod.TemporalLag"),
 	0.5f,
diff --git a/Engine/Source/Runtime/Engine/Private/SplineMeshSceneProxy.h b/Engine/Source/Runtime/Engine/Private/SplineMeshSceneProxy.h
--- a/Engine/Source/Runtime/Engine/Private/SplineMeshSceneProxy.h
+++ b/Engine/Source/Runtime/Engine/Private/SplineMeshSceneProxy.h
@@ -71,8 +71,6 @@ class FSplineMeshVertexFactoryShaderParameters : public FVertexFactoryShaderPara
 {
 	void Bind(const FShaderParameterMap& ParameterMap) override;
 
-	void SetMesh(FRHICommandList& RHICmdList, FShader* Shader, const FVertexFactory* VertexFactory, const FSceneView& View, const FMeshBatchElement& BatchElement, uint32 DataFlags) const override;
-
 	virtual void GetElementShaderBindings(
 		const class FSceneInterface* Scene,
 		const FSceneView* View,
diff --git a/Engine/Source/Runtime/Engine/Private/StaticMeshRender.cpp b/Engine/Source/Runtime/Engine/Private/StaticMeshRender.cpp
--- a/Engine/Source/Runtime/Engine/Private/StaticMeshRender.cpp
+++ b/Engine/Source/Runtime/Engine/Private/StaticMeshRender.cpp
@@ -872,8 +872,7 @@ void FStaticMeshSceneProxy::DrawStaticElements(FStaticPrimitiveDrawInterface* PD
 					{
 						bUseUnifiedMeshForShadow = bAllSectionsCastShadow;
 
-						// Depth pass is only used for deferred renderer. The other conditions are meant to match the logic in FStaticMesh::AddToDrawLists.
-						// Could not link to "GEarlyZPassMovable" so moveable are ignored.
+						// Depth pass is only used for deferred renderer. The other conditions are meant to match the logic in FDepthPassMeshProcessor::AddMeshBatch.
 						bUseUnifiedMeshForDepth = ShouldUseAsOccluder() && GetScene().GetShadingPath() == EShadingPath::Deferred && !IsMovable();
 
 						if (bUseUnifiedMeshForShadow || bUseUnifiedMeshForDepth)
@@ -1148,14 +1147,7 @@ void FStaticMeshSceneProxy::GetDynamicMeshElements(const TArray<const FSceneView
 	#endif
 										if (MeshElement.bDitheredLODTransition && LODMask.IsDithered())
 										{
-											if (LODIndex == LODMask.DitheredLODIndices[0])
-											{
-												MeshElement.DitheredLODTransitionAlpha = View->GetTemporalLODTransition();
-											}
-											else
-											{
-												MeshElement.DitheredLODTransitionAlpha = View->GetTemporalLODTransition() - 1.0f;
-											}
+
 										}
 										else
 										{
diff --git a/Engine/Source/Runtime/Engine/Private/UnrealEngine.cpp b/Engine/Source/Runtime/Engine/Private/UnrealEngine.cpp
--- a/Engine/Source/Runtime/Engine/Private/UnrealEngine.cpp
+++ b/Engine/Source/Runtime/Engine/Private/UnrealEngine.cpp
@@ -281,7 +281,7 @@ ENGINE_API UEngine*	GEngine = NULL;
 */
 ENGINE_API bool GShowDebugSelectedLightmap = false;
 
-int32 GShowMaterialDrawEventTypes = 0;
+int32 GShowMaterialDrawEvents = 0;
 
 #if WANTS_DRAW_MESH_EVENTS
 /**
@@ -289,29 +289,8 @@ int32 GShowMaterialDrawEventTypes = 0;
 */
 static FAutoConsoleVariableRef CVARShowMaterialDrawEvents(
 	TEXT("r.ShowMaterialDrawEvents"),
-	GShowMaterialDrawEventTypes,
-	TEXT("Uses a flags array to enable a draw event around specific material draw types if supported by the platform.\n")
-	TEXT("Set to -1 to enable everything. \n")
-	TEXT("Otherwise sum up these flags:   \n")
-	TEXT("None						0	  \n")
-	TEXT("CompositionLighting		1	  \n")
-	TEXT("BasePass					2	  \n")
-	TEXT("DepthPositionOnly			4	  \n")
-	TEXT("Depth						8	  \n")
-	TEXT("DistortionDynamic			16	  \n")
-	TEXT("DistortionStatic			32	  \n")
-	TEXT("MobileBasePass			64	  \n")
-	TEXT("MobileTranslucent			128	  \n")
-	TEXT("MobileTranslucentOpacity	256	  \n")
-	TEXT("ShadowDepth				512	  \n")
-	TEXT("ShadowDepthRsm			1024  \n")
-	TEXT("ShadowDepthStatic			2048  \n")
-	TEXT("StaticDraw				4096  \n")
-	TEXT("StaticDrawStereo			8192  \n")
-	TEXT("TranslucentLighting		16384 \n")
-	TEXT("Translucent				32768 \n")
-	TEXT("Velocity					65536 \n")
-	TEXT("FogVoxelization			131072\n"),
+	GShowMaterialDrawEvents,
+	TEXT("Whether to emit a draw event around every mesh draw call with information about the assets used.  Introduces severe CPU and GPU overhead when enabled, but useful for debugging."),
 	ECVF_Default
 );
 #endif
diff --git a/Engine/Source/Runtime/Engine/Private/UserInterface/Canvas.cpp b/Engine/Source/Runtime/Engine/Private/UserInterface/Canvas.cpp
--- a/Engine/Source/Runtime/Engine/Private/UserInterface/Canvas.cpp
+++ b/Engine/Source/Runtime/Engine/Private/UserInterface/Canvas.cpp
@@ -19,11 +19,10 @@
 #include "EngineFontServices.h"
 #include "Internationalization/BreakIterator.h"
 #include "Misc/CoreDelegates.h"
-#include "DrawingPolicy.h"
 #include "OneColorShader.h"
 #include "PipelineStateCache.h"
 #include "ClearQuad.h"
-
+#include "MeshPassProcessor.h"
 #include "StereoRendering.h"
 #include "Debug/ReporterGraph.h"
 #include "Fonts/FontMeasure.h"
diff --git a/Engine/Source/Runtime/Engine/Private/VectorFieldVisualization.cpp b/Engine/Source/Runtime/Engine/Private/VectorFieldVisualization.cpp
--- a/Engine/Source/Runtime/Engine/Private/VectorFieldVisualization.cpp
+++ b/Engine/Source/Runtime/Engine/Private/VectorFieldVisualization.cpp
@@ -39,8 +39,6 @@ public:
 		Ar << VectorFieldTextureSampler;
 	}
 
-	virtual void SetMesh(FRHICommandList& RHICmdList, FShader* Shader,const FVertexFactory* VertexFactory,const FSceneView& View,const FMeshBatchElement& BatchElement,uint32 DataFlags) const override;
-
 	virtual void GetElementShaderBindings(
 		const FSceneInterface* Scene,
 		const FSceneView* View,
@@ -173,18 +171,6 @@ void FVectorFieldVisualizationVertexFactory::SetParameters(
 	VectorFieldTextureRHI = InVectorFieldTextureRHI;
 }
 
-/**
- * Set vertex factory shader parameters.
- */
-void FVectorFieldVisualizationVertexFactoryShaderParameters::SetMesh(FRHICommandList& RHICmdList, FShader* Shader,const FVertexFactory* InVertexFactory,const FSceneView& View,const FMeshBatchElement& BatchElement,uint32 DataFlags) const
-{
-	FVectorFieldVisualizationVertexFactory* VertexFactory = (FVectorFieldVisualizationVertexFactory*)InVertexFactory;
-	FVertexShaderRHIParamRef VertexShaderRHI = Shader->GetVertexShader();
-	FSamplerStateRHIParamRef SamplerStatePoint = TStaticSamplerState<SF_Point>::GetRHI();
-	SetUniformBufferParameter(RHICmdList, VertexShaderRHI, Shader->GetUniformBufferParameter<FVectorFieldVisualizationParameters>(), VertexFactory->UniformBuffer);
-	SetTextureParameter(RHICmdList, VertexShaderRHI, VectorFieldTexture, VectorFieldTextureSampler, SamplerStatePoint, VertexFactory->VectorFieldTextureRHI);
-}
-
 void FVectorFieldVisualizationVertexFactoryShaderParameters::GetElementShaderBindings(
 	const FSceneInterface* Scene,
 	const FSceneView* View,
diff --git a/Engine/Source/Runtime/Engine/Public/LocalVertexFactory.h b/Engine/Source/Runtime/Engine/Public/LocalVertexFactory.h
--- a/Engine/Source/Runtime/Engine/Public/LocalVertexFactory.h
+++ b/Engine/Source/Runtime/Engine/Public/LocalVertexFactory.h
@@ -166,7 +166,6 @@ class FLocalVertexFactoryShaderParametersBase : public FVertexFactoryShaderParam
 public:
 	virtual void Bind(const FShaderParameterMap& ParameterMap) override;
 	virtual void Serialize(FArchive& Ar) override;
-	virtual void SetMesh(FRHICommandList& RHICmdList, FShader* Shader, const FVertexFactory* VertexFactory, const FSceneView& View, const FMeshBatchElement& BatchElement, uint32 DataFlags) const override;
 	virtual void GetElementShaderBindings(
 		const class FSceneInterface* Scene,
 		const FSceneView* View,
@@ -187,7 +186,7 @@ public:
 	// SpeedTree LOD parameter
 	FShaderParameter LODParameter;
 
-	// True if LODParameter is bound, which puts us on the slow path in SetMesh
+	// True if LODParameter is bound, which puts us on the slow path in GetElementShaderBindings
 	bool bAnySpeedTreeParamIsBound;
 };
 
diff --git a/Engine/Source/Runtime/Engine/Public/MaterialShared.h b/Engine/Source/Runtime/Engine/Public/MaterialShared.h
--- a/Engine/Source/Runtime/Engine/Public/MaterialShared.h
+++ b/Engine/Source/Runtime/Engine/Public/MaterialShared.h
@@ -1731,29 +1731,6 @@ public:
 	void SetSubsurfaceProfileRT(const USubsurfaceProfile* Ptr) { SubsurfaceProfileRT = Ptr; }
 	const USubsurfaceProfile* GetSubsurfaceProfileRT() const { return SubsurfaceProfileRT; }
 
-	void SetReferencedInDrawList() const
-	{
-#if !(UE_BUILD_SHIPPING || UE_BUILD_TEST)
-		bIsStaticDrawListReferenced = 1;
-#endif
-	}
-
-	void SetUnreferencedInDrawList() const
-	{
-#if !(UE_BUILD_SHIPPING || UE_BUILD_TEST)
-		bIsStaticDrawListReferenced = 0;
-#endif
-	}
-
-	bool IsReferencedInDrawList() const
-	{
-#if !(UE_BUILD_SHIPPING || UE_BUILD_TEST)
-		return bIsStaticDrawListReferenced != 0;
-#else
-		return false;
-#endif
-	}
-
 	ENGINE_API static void UpdateDeferredCachedUniformExpressions();
 
 	static inline bool HasDeferredUniformExpressionCacheRequests() 
@@ -1769,7 +1746,6 @@ private:
 	/** For tracking down a bug accessing a deleted proxy. */
 #if !(UE_BUILD_SHIPPING || UE_BUILD_TEST)
 	mutable int32 DeletedFlag : 1;
-	mutable int32 bIsStaticDrawListReferenced : 1;
 #endif
 
 	/** 
diff --git a/Engine/Source/Runtime/Engine/Public/MeshBatch.h b/Engine/Source/Runtime/Engine/Public/MeshBatch.h
--- a/Engine/Source/Runtime/Engine/Public/MeshBatch.h
+++ b/Engine/Source/Runtime/Engine/Public/MeshBatch.h
@@ -188,10 +188,6 @@ struct FMeshBatch
 	/** Whether the mesh batch should apply dithered LOD. */
 	uint32 bDitheredLODTransition : 1;
 
-	//@todo MeshCommandPipeline remove it together with mesh draw policies.
-	/** If bDitheredLODTransition and this is a dynamic mesh element, then this is the alpha for dither fade (static draw lists need to derive this later as it is changes every frame) */
-	float DitheredLODTransitionAlpha;
-
 	// can be NULL
 	const FLightCacheInterface* LCI;
 
@@ -281,7 +277,6 @@ struct FMeshBatch
 	,	bSelectable(true)
 	,	bRequiresPerElementVisibility(false)
 	,	bDitheredLODTransition(false)
-	,   DitheredLODTransitionAlpha(0.0f)
 	,	LCI(NULL)
 	,	VertexFactory(NULL)
 	,	MaterialRenderProxy(NULL)
diff --git a/Engine/Source/Runtime/Engine/Public/SceneInterface.h b/Engine/Source/Runtime/Engine/Public/SceneInterface.h
--- a/Engine/Source/Runtime/Engine/Public/SceneInterface.h
+++ b/Engine/Source/Runtime/Engine/Public/SceneInterface.h
@@ -189,9 +189,6 @@ public:
 	/** Updates all static draw lists. */
 	virtual void UpdateStaticDrawLists() {}
 
-	/** Updates static draw lists for the given set of materials. */
-	virtual void UpdateStaticDrawListsForMaterials(const TArray<const FMaterial*>& Materials) {}
-
 	/** 
 	 * Adds a new exponential height fog component to the scene
 	 * 
@@ -331,11 +328,6 @@ public:
 
 	virtual void DumpUnbuiltLightInteractions( FOutputDevice& Ar ) const { }
 
-	/**
-	 * Dumps static mesh draw list stats to the log.
-	 */
-	virtual void DumpStaticMeshDrawListStats() const {}
-
 	/** Updates the scene's list of parameter collection id's and their uniform buffers. */
 	virtual void UpdateParameterCollections(const TArray<class FMaterialParameterCollectionInstanceResource*>& InParameterCollections) {}
 
diff --git a/Engine/Source/Runtime/Engine/Public/UnrealEngine.h b/Engine/Source/Runtime/Engine/Public/UnrealEngine.h
--- a/Engine/Source/Runtime/Engine/Public/UnrealEngine.h
+++ b/Engine/Source/Runtime/Engine/Public/UnrealEngine.h
@@ -16,7 +16,7 @@ class FViewportClient;
 class UFont;
 class ULocalPlayer;
 
-extern ENGINE_API int32 GShowMaterialDrawEventTypes;
+extern ENGINE_API int32 GShowMaterialDrawEvents;
 
 ENGINE_API DECLARE_LOG_CATEGORY_EXTERN(LogEngine, Log, All);
 //
diff --git a/Engine/Source/Runtime/Landscape/Private/LandscapeGrass.cpp b/Engine/Source/Runtime/Landscape/Private/LandscapeGrass.cpp
--- a/Engine/Source/Runtime/Landscape/Private/LandscapeGrass.cpp
+++ b/Engine/Source/Runtime/Landscape/Private/LandscapeGrass.cpp
@@ -32,7 +32,6 @@
 #include "LandscapeVersion.h"
 #include "MaterialShaderType.h"
 #include "MeshMaterialShaderType.h"
-#include "DrawingPolicy.h"
 #include "MeshMaterialShader.h"
 #include "Materials/Material.h"
 #include "LandscapeGrassType.h"
@@ -242,18 +241,6 @@ public:
 		return ShouldCacheLandscapeGrassShaders(Platform, Material, VertexFactoryType);
 	}
 
-	void SetParameters(FRHICommandList& RHICmdList, const FMaterialRenderProxy* MaterialRenderProxy, const FMaterial& MaterialResource, const FSceneView& View, const FDrawingPolicyRenderState& DrawRenderState, const FVector2D& RenderOffset)
-	{
-		//TUniformBufferRef<FSceneTexturesUniformParameters> SceneTexturePassUniformBuffer = CreateSceneTextureUniformBufferSingleDraw(RHICmdList, ESceneTextureSetupMode::None, View.FeatureLevel);
-		FMeshMaterialShader::SetParameters(RHICmdList, GetVertexShader(), MaterialRenderProxy, MaterialResource, View, DrawRenderState.GetViewUniformBuffer(), nullptr);
-		SetShaderValue(RHICmdList, GetVertexShader(), RenderOffsetParameter, RenderOffset);
-	}
-
-	void SetMesh(FRHICommandList& RHICmdList, const FVertexFactory* VertexFactory, const FSceneView& View, const FPrimitiveSceneProxy* Proxy, const FMeshBatchElement& BatchElement, const FDrawingPolicyRenderState& DrawRenderState)
-	{
-		FMeshMaterialShader::SetMesh(RHICmdList, GetVertexShader(), VertexFactory, View, Proxy, BatchElement, DrawRenderState);
-	}
-
 	void GetShaderBindings(
 		const FScene* Scene,
 		ERHIFeatureLevel::Type FeatureLevel,
@@ -301,21 +288,6 @@ public:
 	FLandscapeGrassWeightPS()
 	{}
 
-	void SetParameters(FRHICommandList& RHICmdList, const FMaterialRenderProxy* MaterialRenderProxy, const FMaterial& MaterialResource, const FSceneView* View, const FDrawingPolicyRenderState& DrawRenderState, int32 OutputPass)
-	{
-		//TUniformBufferRef<FSceneTexturesUniformParameters> SceneTexturePassUniformBuffer = CreateSceneTextureUniformBufferSingleDraw(RHICmdList, ESceneTextureSetupMode::None, View->FeatureLevel);
-		FMeshMaterialShader::SetParameters(RHICmdList, GetPixelShader(), MaterialRenderProxy, MaterialResource, *View, DrawRenderState.GetViewUniformBuffer(), nullptr);
-		if (OutputPassParameter.IsBound())
-		{
-			SetShaderValue(RHICmdList, GetPixelShader(), OutputPassParameter, OutputPass);
-		}
-	}
-
-	void SetMesh(FRHICommandList& RHICmdList, const FVertexFactory* VertexFactory, const FSceneView& View, const FPrimitiveSceneProxy* Proxy, const FMeshBatchElement& BatchElement, const FDrawingPolicyRenderState& DrawRenderState)
-	{
-		FMeshMaterialShader::SetMesh(RHICmdList, GetPixelShader(), VertexFactory, View, Proxy, BatchElement, DrawRenderState);
-	}
-
 	void GetShaderBindings(
 		const FScene* Scene,
 		ERHIFeatureLevel::Type FeatureLevel,
diff --git a/Engine/Source/Runtime/Landscape/Private/LandscapeRender.cpp b/Engine/Source/Runtime/Landscape/Private/LandscapeRender.cpp
--- a/Engine/Source/Runtime/Landscape/Private/LandscapeRender.cpp
+++ b/Engine/Source/Runtime/Landscape/Private/LandscapeRender.cpp
@@ -3415,142 +3415,6 @@ public:
 		Ar << XYOffsetTextureParameterSampler;
 	}
 
-	/**
-	* Set any shader data specific to this vertex factory
-	*/
-	virtual void SetMesh(FRHICommandList& RHICmdList, FShader* VertexShader, const FVertexFactory* VertexFactory, const FSceneView& InView, const FMeshBatchElement& BatchElement, uint32 DataFlags) const override
-	{
-		SCOPE_CYCLE_COUNTER(STAT_LandscapeVFDrawTimeVS);
-
-		const FLandscapeBatchElementParams* BatchElementParams = (const FLandscapeBatchElementParams*)BatchElement.UserData;
-		check(BatchElementParams);
-
-		const FLandscapeComponentSceneProxy* SceneProxy = BatchElementParams->SceneProxy;
-		FVertexShaderRHIParamRef VertexShaderParamRef = VertexShader->GetVertexShader();
-
-		SetUniformBufferParameter(RHICmdList, VertexShaderParamRef, VertexShader->GetUniformBufferParameter<FLandscapeUniformShaderParameters>(), *BatchElementParams->LandscapeUniformShaderParametersResource);
-
-		if (HeightmapTextureParameter.IsBound())
-		{
-			SetTextureParameter(RHICmdList, VertexShaderParamRef, HeightmapTextureParameter, HeightmapTextureParameterSampler, TStaticSamplerState<SF_Point>::GetRHI(), SceneProxy->HeightmapTexture->Resource->TextureRHI);
-		}
-
-		if (LodValuesParameter.IsBound())
-		{
-			SetShaderValue(RHICmdList, VertexShaderParamRef, LodValuesParameter, SceneProxy->GetShaderLODValues(BatchElementParams->CurrentLOD));
-		}
-
-		const FSceneView& View = GetLODView(InView);
-		int32 SubSectionIndex = BatchElementParams->SubX + BatchElementParams->SubY * SceneProxy->NumSubsections;
-
-		// If we have no custom data for this primitive we will compute of the fly the proper values, this will happen if the shader is not used for normal landscape rendering(i.e grass rendering)
-		FLandscapeComponentSceneProxy::FViewCustomDataLOD* LODData = (FLandscapeComponentSceneProxy::FViewCustomDataLOD*)View.GetCustomData(SceneProxy->GetPrimitiveSceneInfo()->GetIndex());
-
-		if (LODData != nullptr)
-		{
-			if (LodTessellationParameter.IsBound())
-			{
-				SetShaderValue(RHICmdList, VertexShaderParamRef, LodTessellationParameter, LODData->LodTessellationParams);
-			}
-
-			if (LodBiasParameter.IsBound())
-			{
-				check(LODData->LodBias == SceneProxy->GetShaderLODBias());
-				SetShaderValue(RHICmdList, VertexShaderParamRef, LodBiasParameter, LODData->LodBias);
-			}
-
-			if (SectionLodsParameter.IsBound())
-			{
-				if (LODData->UseCombinedMeshBatch)
-				{
-					SetShaderValue(RHICmdList, VertexShaderParamRef, SectionLodsParameter, LODData->ShaderCurrentLOD);
-				}
-				else // in non combined, only the one representing us as we'll be called 4 times (once per sub section)
-				{
-					check(SubSectionIndex >= 0);
-					FVector4 ShaderCurrentLOD(ForceInitToZero);
-					ShaderCurrentLOD.Component(SubSectionIndex) = LODData->ShaderCurrentLOD.Component(SubSectionIndex);
-
-					SetShaderValue(RHICmdList, VertexShaderParamRef, SectionLodsParameter, ShaderCurrentLOD);
-				}
-			}
-
-			if (NeighborSectionLodParameter.IsBound())
-			{
-				FVector4 ShaderCurrentNeighborLOD[FLandscapeComponentSceneProxy::NEIGHBOR_COUNT] = { FVector4(ForceInitToZero), FVector4(ForceInitToZero), FVector4(ForceInitToZero), FVector4(ForceInitToZero) };
-
-				if (LODData->UseCombinedMeshBatch)
-				{
-					int32 SubSectionCount = SceneProxy->NumSubsections == 1 ? 1 : FLandscapeComponentSceneProxy::MAX_SUBSECTION_COUNT;
-
-					for (int32 NeighborSubSectionIndex = 0; NeighborSubSectionIndex < SubSectionCount; ++NeighborSubSectionIndex)
-					{
-						ShaderCurrentNeighborLOD[NeighborSubSectionIndex] = LODData->SubSections[NeighborSubSectionIndex].ShaderCurrentNeighborLOD;
-						check(ShaderCurrentNeighborLOD[NeighborSubSectionIndex].X != -1.0f); // they should all match so only check the 1st one for simplicity
-					}
-
-					SetShaderValue(RHICmdList, VertexShaderParamRef, NeighborSectionLodParameter, ShaderCurrentNeighborLOD);
-				}
-				else // in non combined, only the one representing us as we'll be called 4 times (once per sub section)
-				{
-					check(SubSectionIndex >= 0);
-					ShaderCurrentNeighborLOD[SubSectionIndex] = LODData->SubSections[SubSectionIndex].ShaderCurrentNeighborLOD;
-					check(ShaderCurrentNeighborLOD[SubSectionIndex].X != -1.0f); // they should all match so only check the 1st one for simplicity
-
-					SetShaderValue(RHICmdList, VertexShaderParamRef, NeighborSectionLodParameter, ShaderCurrentNeighborLOD);
-				}
-			}
-		}
-		else
-		{
-			float ComponentScreenSize = SceneProxy->GetComponentScreenSize(&View, SceneProxy->LandscapeComponent->Bounds.Origin, SceneProxy->ComponentMaxExtend, SceneProxy->LandscapeComponent->Bounds.SphereRadius);
-
-			FLandscapeComponentSceneProxy::FViewCustomDataLOD CurrentLODData;
-			SceneProxy->CalculateBatchElementLOD(InView, ComponentScreenSize, View.LODDistanceFactor, CurrentLODData, true);
-			check(CurrentLODData.UseCombinedMeshBatch);
-
-			if (LodBiasParameter.IsBound())
-			{
-				SetShaderValue(RHICmdList, VertexShaderParamRef, LodBiasParameter, SceneProxy->GetShaderLODBias());
-			}
-
-			if (LodTessellationParameter.IsBound())
-			{
-				FVector4 LodTessellationParams(ForceInitToZero);
-				SceneProxy->ComputeTessellationFalloffShaderValues(CurrentLODData, View.ViewMatrices.GetProjectionMatrix(), LodTessellationParams.X, LodTessellationParams.Y);
-
-				SetShaderValue(RHICmdList, VertexShaderParamRef, LodTessellationParameter, LodTessellationParams);
-			}
-
-			if (SectionLodsParameter.IsBound())
-			{
-				SetShaderValue(RHICmdList, VertexShaderParamRef, SectionLodsParameter, CurrentLODData.ShaderCurrentLOD);
-			}
-
-			if (NeighborSectionLodParameter.IsBound())
-			{
-				FVector4 CurrentNeighborLOD[4] = { FVector4(ForceInitToZero), FVector4(ForceInitToZero), FVector4(ForceInitToZero), FVector4(ForceInitToZero) };
-
-				for (int32 SubY = 0; SubY < SceneProxy->NumSubsections; SubY++)
-				{
-					for (int32 SubX = 0; SubX < SceneProxy->NumSubsections; SubX++)
-					{
-						int32 NeighborSubSectionIndex = SubX + SubY * SceneProxy->NumSubsections;
-						SceneProxy->GetShaderCurrentNeighborLOD(View, CurrentLODData.SubSections[NeighborSubSectionIndex].fBatchElementCurrentLOD, SceneProxy->NumSubsections > 1 ? SubX : INDEX_NONE, SceneProxy->NumSubsections > 1 ? SubY : INDEX_NONE, NeighborSubSectionIndex, CurrentNeighborLOD[NeighborSubSectionIndex]);
-						check(CurrentNeighborLOD[NeighborSubSectionIndex].X != -1.0f); // they should all match so only check the 1st one for simplicity
-					}
-				}
-
-				SetShaderValue(RHICmdList, VertexShaderParamRef, NeighborSectionLodParameter, CurrentNeighborLOD);
-			}
-		}
-
-		if (XYOffsetTextureParameter.IsBound() && SceneProxy->XYOffsetmapTexture)
-		{
-			SetTextureParameter(RHICmdList, VertexShaderParamRef, XYOffsetTextureParameter, XYOffsetTextureParameterSampler, TStaticSamplerState<SF_Point>::GetRHI(), SceneProxy->XYOffsetmapTexture->Resource->TextureRHI);
-		}
-	}
-
 	virtual void GetElementShaderBindings(
 		const class FSceneInterface* Scene,
 		const FSceneView* InView,
@@ -3738,31 +3602,6 @@ void FLandscapeVertexFactoryPixelShaderParameters::Serialize(FArchive& Ar)
 		<< LocalToWorldNoScalingParameter;
 }
 
-/**
-* Set any shader data specific to this vertex factory
-*/
-void FLandscapeVertexFactoryPixelShaderParameters::SetMesh(FRHICommandList& RHICmdList, FShader* PixelShader, const FVertexFactory* VertexFactory, const FSceneView& View, const FMeshBatchElement& BatchElement, uint32 DataFlags) const
-{
-	SCOPE_CYCLE_COUNTER(STAT_LandscapeVFDrawTimePS);
-
-	const FLandscapeBatchElementParams* BatchElementParams = (const FLandscapeBatchElementParams*)BatchElement.UserData;
-
-	if (LocalToWorldNoScalingParameter.IsBound())
-	{
-		SetShaderValue(RHICmdList, PixelShader->GetPixelShader(), LocalToWorldNoScalingParameter, *BatchElementParams->LocalToWorldNoScalingPtr);
-	}
-
-	if (NormalmapTextureParameter.IsBound())
-	{
-		SetTextureParameter(
-			RHICmdList,
-			PixelShader->GetPixelShader(),
-			NormalmapTextureParameter,
-			NormalmapTextureParameterSampler,
-			BatchElementParams->SceneProxy->NormalmapTexture->Resource);
-	}
-}
-
 void FLandscapeVertexFactoryPixelShaderParameters::GetElementShaderBindings(
 	const class FSceneInterface* Scene,
 	const FSceneView* InView,
diff --git a/Engine/Source/Runtime/Landscape/Private/LandscapeRenderMobile.cpp b/Engine/Source/Runtime/Landscape/Private/LandscapeRenderMobile.cpp
--- a/Engine/Source/Runtime/Landscape/Private/LandscapeRenderMobile.cpp
+++ b/Engine/Source/Runtime/Landscape/Private/LandscapeRenderMobile.cpp
@@ -62,97 +62,7 @@ public:
 		Ar << LodBiasParameter;
 		Ar << SectionLodsParameter;
 	}
-	/**
-	* Set any shader data specific to this vertex factory
-	*/
-	virtual void SetMesh(FRHICommandList& RHICmdList, FShader* VertexShader,const class FVertexFactory* VertexFactory,const class FSceneView& View,const struct FMeshBatchElement& BatchElement,uint32 DataFlags) const override
-	{
-		SCOPE_CYCLE_COUNTER(STAT_LandscapeVFDrawTimeVS);
-
-		const FLandscapeBatchElementParams* BatchElementParams = (const FLandscapeBatchElementParams*)BatchElement.UserData;
-		check(BatchElementParams);
-
-		const FLandscapeComponentSceneProxyMobile* SceneProxy = (const FLandscapeComponentSceneProxyMobile*)BatchElementParams->SceneProxy;
-		SetUniformBufferParameter(RHICmdList, VertexShader->GetVertexShader(),VertexShader->GetUniformBufferParameter<FLandscapeUniformShaderParameters>(),*BatchElementParams->LandscapeUniformShaderParametersResource);
-
-		if (LodValuesParameter.IsBound())
-		{
-			FVector4 LodValues(
-				0.0f, // this is the mesh's LOD, ES2 always uses the LOD0 mesh
-				0.0f, // unused
-				(float)SceneProxy->SubsectionSizeQuads,
-				1.f / (float)SceneProxy->SubsectionSizeQuads);
-
-			SetShaderValue(RHICmdList, VertexShader->GetVertexShader(), LodValuesParameter, LodValues);
-		}
-
-		if (LodBiasParameter.IsBound())
-		{
-			FVector CameraLocalPos3D = SceneProxy->WorldToLocal.TransformPosition(View.ViewMatrices.GetViewOrigin());
-
-			FVector4 LodBias(
-				0.0f, // unused
-				0.0f, // unused
-				CameraLocalPos3D.X + SceneProxy->SectionBase.X,
-				CameraLocalPos3D.Y + SceneProxy->SectionBase.Y
-			);
-			SetShaderValue(RHICmdList, VertexShader->GetVertexShader(), LodBiasParameter, LodBias);
-		}
-
-		FLandscapeComponentSceneProxy::FViewCustomDataLOD* LODData = (FLandscapeComponentSceneProxy::FViewCustomDataLOD*)View.GetCustomData(SceneProxy->GetPrimitiveSceneInfo()->GetIndex());
-		int32 SubSectionIndex = BatchElementParams->SubX + BatchElementParams->SubY * SceneProxy->NumSubsections;
-
-		if (LODData != nullptr)
-		{
-			if (LodTessellationParameter.IsBound())
-			{
-				SetShaderValue(RHICmdList, VertexShader->GetVertexShader(), LodTessellationParameter, LODData->LodTessellationParams);
-			}
-
-			if (SectionLodsParameter.IsBound())
-			{
-				if (LODData->UseCombinedMeshBatch)
-				{
-					SetShaderValue(RHICmdList, VertexShader->GetVertexShader(), SectionLodsParameter, LODData->ShaderCurrentLOD);
-				}
-				else // in non combined, only the one representing us as we'll be called 4 times (once per sub section)
-				{
-					check(SubSectionIndex >= 0);
-					FVector4 ShaderCurrentLOD(ForceInitToZero);
-					ShaderCurrentLOD.Component(SubSectionIndex) = LODData->ShaderCurrentLOD.Component(SubSectionIndex);
-
-					SetShaderValue(RHICmdList, VertexShader->GetVertexShader(), SectionLodsParameter, ShaderCurrentLOD);
-				}
-			}
-
-			if (NeighborSectionLodParameter.IsBound())
-			{
-				FVector4 ShaderCurrentNeighborLOD[FLandscapeComponentSceneProxy::NEIGHBOR_COUNT] = { FVector4(ForceInitToZero), FVector4(ForceInitToZero), FVector4(ForceInitToZero), FVector4(ForceInitToZero) };
-
-				if (LODData->UseCombinedMeshBatch)
-				{
-					int32 SubSectionCount = SceneProxy->NumSubsections == 1 ? 1 : FLandscapeComponentSceneProxy::MAX_SUBSECTION_COUNT;
-
-					for (int32 NeighborSubSectionIndex = 0; NeighborSubSectionIndex < SubSectionCount; ++NeighborSubSectionIndex)
-					{
-						ShaderCurrentNeighborLOD[NeighborSubSectionIndex] = LODData->SubSections[NeighborSubSectionIndex].ShaderCurrentNeighborLOD;
-						check(ShaderCurrentNeighborLOD[NeighborSubSectionIndex].X != -1.0f); // they should all match so only check the 1st one for simplicity
-					}
-
-					SetShaderValue(RHICmdList, VertexShader->GetVertexShader(), NeighborSectionLodParameter, ShaderCurrentNeighborLOD);
-				}
-				else // in non combined, only the one representing us as we'll be called 4 times (once per sub section)
-				{
-					check(SubSectionIndex >= 0);
-					ShaderCurrentNeighborLOD[SubSectionIndex] = LODData->SubSections[SubSectionIndex].ShaderCurrentNeighborLOD;
-					check(ShaderCurrentNeighborLOD[SubSectionIndex].X != -1.0f); // they should all match so only check the 1st one for simplicity
-
-					SetShaderValue(RHICmdList, VertexShader->GetVertexShader(), NeighborSectionLodParameter, ShaderCurrentNeighborLOD);
-				}
-			}
-		}	
-	}
-
+	
 	virtual void GetElementShaderBindings(
 		const class FSceneInterface* Scene,
 		const FSceneView* InView,
@@ -284,29 +194,7 @@ public:
 		FLandscapeVertexFactoryPixelShaderParameters::Serialize(Ar);
 		Ar << BlendableLayerMaskParameter;
 	}
-	/**
-	* Set any shader data specific to this vertex factory
-	*/
-	virtual void SetMesh(FRHICommandList& RHICmdList, FShader* PixelShader, const FVertexFactory* VertexFactory, const FSceneView& View, const FMeshBatchElement& BatchElement, uint32 DataFlags) const override
-	{
-		SCOPE_CYCLE_COUNTER(STAT_LandscapeVFDrawTimePS);
-
-		FLandscapeVertexFactoryPixelShaderParameters::SetMesh(RHICmdList, PixelShader, VertexFactory, View, BatchElement, DataFlags);
-
-		const FLandscapeBatchElementParams* BatchElementParams = (const FLandscapeBatchElementParams*)BatchElement.UserData;
-		check(BatchElementParams);
-
-		const FLandscapeComponentSceneProxyMobile* SceneProxy = (const FLandscapeComponentSceneProxyMobile*)BatchElementParams->SceneProxy;
-
-		if (BlendableLayerMaskParameter.IsBound())
-		{
-			FVector MaskVector;
-			MaskVector[0] = (SceneProxy->BlendableLayerMask & (1 << 0)) ? 1 : 0;
-			MaskVector[1] = (SceneProxy->BlendableLayerMask & (1 << 1)) ? 1 : 0;
-			MaskVector[2] = (SceneProxy->BlendableLayerMask & (1 << 2)) ? 1 : 0;
-			SetShaderValue(RHICmdList, PixelShader->GetPixelShader(), BlendableLayerMaskParameter, MaskVector);
-		}
-	}
+	
 protected:
 	FShaderParameter BlendableLayerMaskParameter;
 };
diff --git a/Engine/Source/Runtime/Landscape/Public/LandscapeRender.h b/Engine/Source/Runtime/Landscape/Public/LandscapeRender.h
--- a/Engine/Source/Runtime/Landscape/Public/LandscapeRender.h
+++ b/Engine/Source/Runtime/Landscape/Public/LandscapeRender.h
@@ -129,11 +129,6 @@ public:
 	*/
 	virtual void Serialize(FArchive& Ar) override;
 
-	/**
-	* Set any shader data specific to this vertex factory
-	*/
-	virtual void SetMesh(FRHICommandList& RHICmdList, FShader* PixelShader, const FVertexFactory* VertexFactory, const FSceneView& View, const FMeshBatchElement& BatchElement, uint32 DataFlags) const override;
-
 	virtual void GetElementShaderBindings(
 		const class FSceneInterface* Scene,
 		const FSceneView* InView,
diff --git a/Engine/Source/Runtime/RHI/Public/RHIResources.h b/Engine/Source/Runtime/RHI/Public/RHIResources.h
--- a/Engine/Source/Runtime/RHI/Public/RHIResources.h
+++ b/Engine/Source/Runtime/RHI/Public/RHIResources.h
@@ -1983,8 +1983,6 @@ public:
 		};
 		uint16						Flags;
 	};
-
-	friend class FMeshDrawingPolicy;
 };
 
 #if RHI_RAYTRACING
diff --git a/Engine/Source/Runtime/RenderCore/Private/RenderCore.cpp b/Engine/Source/Runtime/RenderCore/Private/RenderCore.cpp
--- a/Engine/Source/Runtime/RenderCore/Private/RenderCore.cpp
+++ b/Engine/Source/Runtime/RenderCore/Private/RenderCore.cpp
@@ -71,8 +71,6 @@ DEFINE_STAT(STAT_PresentTime);
 
 DEFINE_STAT(STAT_SceneLights);
 DEFINE_STAT(STAT_MeshDrawCalls);
-DEFINE_STAT(STAT_DynamicPathMeshDrawCalls);
-DEFINE_STAT(STAT_StaticDrawListMeshDrawCalls);
 
 DEFINE_STAT(STAT_SceneDecals);
 DEFINE_STAT(STAT_Decals);
@@ -83,7 +81,6 @@ DEFINE_STAT(STAT_DecalsDrawTime);
 // Not to track all of the allocations, and not to track resource memory (index buffers, vertex buffers, etc).
 
 
-DEFINE_STAT(STAT_StaticDrawListMemory);
 DEFINE_STAT(STAT_PrimitiveInfoMemory);
 DEFINE_STAT(STAT_RenderingSceneMemory);
 DEFINE_STAT(STAT_ViewStateMemory);
@@ -140,9 +137,6 @@ DEFINE_STAT(STAT_BindRayTracingPipeline);
 // Shadow setup is tracked in the InitViews group
 
 DEFINE_STAT(STAT_RenderWholeSceneShadowProjectionsTime);
-DEFINE_STAT(STAT_WholeSceneDynamicShadowDepthsTime);
-DEFINE_STAT(STAT_WholeSceneStaticShadowDepthsTime);
-DEFINE_STAT(STAT_WholeSceneStaticDrawListShadowDepthsTime);
 DEFINE_STAT(STAT_RenderWholeSceneShadowDepthsTime);
 DEFINE_STAT(STAT_RenderPerObjectShadowProjectionsTime);
 DEFINE_STAT(STAT_RenderPerObjectShadowDepthsTime);
@@ -194,7 +188,6 @@ DEFINE_STAT(STAT_UpdateLPVs);
 DEFINE_STAT(STAT_ReflectiveShadowMaps);
 DEFINE_STAT(STAT_ReflectiveShadowMapDrawTime);
 DEFINE_STAT(STAT_NumReflectiveShadowMapLights);
-DEFINE_STAT(STAT_RenderWholeSceneReflectiveShadowMapsTime);
 
 DEFINE_STAT(STAT_ShadowmapAtlasMemory);
 DEFINE_STAT(STAT_CachedShadowmapMemory);
diff --git a/Engine/Source/Runtime/RenderCore/Private/VertexFactory.cpp b/Engine/Source/Runtime/RenderCore/Private/VertexFactory.cpp
--- a/Engine/Source/Runtime/RenderCore/Private/VertexFactory.cpp
+++ b/Engine/Source/Runtime/RenderCore/Private/VertexFactory.cpp
@@ -177,35 +177,6 @@ FVertexFactoryType* FindVertexFactoryType(FName TypeName)
 	return NULL;
 }
 
-void FVertexFactory::SetStreams(ERHIFeatureLevel::Type InFeatureLevel, FRHICommandList& RHICmdList) const
-{
-	bool bSupportsVertexFetch = SupportsManualVertexFetch(InFeatureLevel);
-	check(IsInitialized());
-	for(int32 StreamIndex = 0;StreamIndex < Streams.Num();StreamIndex++)
-	{
-		const FVertexStream& Stream = Streams[StreamIndex];
-		if (!(EnumHasAnyFlags(EVertexStreamUsage::ManualFetch, Stream.VertexStreamUsage) && bSupportsVertexFetch))
-		{
-			if (!Stream.VertexBuffer)
-			{
-				RHICmdList.SetStreamSource(StreamIndex, nullptr, 0);
-			}
-			else
-			{
-				if (EnumHasAnyFlags(EVertexStreamUsage::Overridden, Stream.VertexStreamUsage) && !Stream.VertexBuffer->IsInitialized())
-				{
-					RHICmdList.SetStreamSource(StreamIndex, nullptr, 0);
-				}
-				else
-				{
-					checkf(Stream.VertexBuffer->IsInitialized(), TEXT("Vertex buffer was not initialized! Stream %u, Stride %u, Name %s"), StreamIndex, Stream.Stride, *Stream.VertexBuffer->GetFriendlyName());
-					RHICmdList.SetStreamSource(StreamIndex, Stream.VertexBuffer->VertexBufferRHI, Stream.Offset);
-				}
-			}
-		}
-	}
-}
-
 void FVertexFactory::GetStreams(ERHIFeatureLevel::Type InFeatureLevel, FVertexInputStreamArray& OutVertexStreams) const
 {
 	bool bSupportsVertexFetch = SupportsManualVertexFetch(InFeatureLevel);
@@ -237,18 +208,6 @@ void FVertexFactory::GetStreams(ERHIFeatureLevel::Type InFeatureLevel, FVertexIn
 	}
 }
 
-void FVertexFactory::OffsetInstanceStreams(FRHICommandList& RHICmdList, uint32 FirstVertex) const
-{
-	for(int32 StreamIndex = 0;StreamIndex < Streams.Num();StreamIndex++)
-	{
-		const FVertexStream& Stream = Streams[StreamIndex];
-		if (EnumHasAnyFlags(EVertexStreamUsage::Instancing, Stream.VertexStreamUsage))
-		{
-			RHICmdList.SetStreamSource( StreamIndex, Stream.VertexBuffer->VertexBufferRHI, Stream.Offset + Stream.Stride * FirstVertex);
-		}
-	}
-}
-
 void FVertexFactory::OffsetInstanceStreams(uint32 InstanceOffset, bool bOperateOnPositionOnly, FVertexInputStreamArray& VertexStreams) const
 {
 	const TArray<FVertexStream,TFixedAllocator<MaxVertexElementCount> >& StreamArray = bOperateOnPositionOnly ? PositionStream : Streams;
@@ -270,18 +229,6 @@ void FVertexFactory::OffsetInstanceStreams(uint32 InstanceOffset, bool bOperateO
 	}
 }
 
-void FVertexFactory::SetPositionStream(FRHICommandList& RHICmdList) const
-{
-	check(IsInitialized());
-	// Set the predefined vertex streams.
-	for(int32 StreamIndex = 0;StreamIndex < PositionStream.Num();StreamIndex++)
-	{
-		const FVertexStream& Stream = PositionStream[StreamIndex];
-		check(Stream.VertexBuffer->IsInitialized());
-		RHICmdList.SetStreamSource( StreamIndex, Stream.VertexBuffer->VertexBufferRHI, Stream.Offset);
-	}
-}
-
 void FVertexFactory::GetPositionOnlyStream(FVertexInputStreamArray& OutVertexStreams) const
 {
 	check(IsInitialized());
@@ -294,18 +241,6 @@ void FVertexFactory::GetPositionOnlyStream(FVertexInputStreamArray& OutVertexStr
 	}
 }
 
-void FVertexFactory::OffsetPositionInstanceStreams(FRHICommandList& RHICmdList, uint32 FirstVertex) const
-{
-	for(int32 StreamIndex = 0;StreamIndex < PositionStream.Num();StreamIndex++)
-	{
-		const FVertexStream& Stream = PositionStream[StreamIndex];
-		if (EnumHasAnyFlags(EVertexStreamUsage::Instancing, Stream.VertexStreamUsage))
-		{
-			RHICmdList.SetStreamSource( StreamIndex, Stream.VertexBuffer->VertexBufferRHI, Stream.Offset + Stream.Stride * FirstVertex);
-		}
-	}
-}
-
 void FVertexFactory::ReleaseRHI()
 {
 	Declaration.SafeRelease();
@@ -314,46 +249,6 @@ void FVertexFactory::ReleaseRHI()
 	PositionStream.Empty();
 }
 
-/**
-* Fill in array of strides from this factory's vertex streams without shadow/light maps
-* @param OutStreamStrides - output array of # MaxVertexElementCount stream strides to fill
-*/
-int32 FVertexFactory::GetStreamStrides(uint32 *OutStreamStrides, bool bPadWithZeroes) const
-{
-	int32 StreamIndex;
-	for(StreamIndex = 0;StreamIndex < Streams.Num();++StreamIndex)
-	{
-		OutStreamStrides[StreamIndex] = Streams[StreamIndex].Stride;
-	}
-	if (bPadWithZeroes)
-	{
-		// Pad stream strides with 0's to be safe (they can be used in hashes elsewhere)
-		for (;StreamIndex < MaxVertexElementCount;++StreamIndex)
-		{
-			OutStreamStrides[StreamIndex] = 0;
-		}
-	}
-	return StreamIndex;
-}
-
-/**
-* Fill in array of strides from this factory's position only vertex streams
-* @param OutStreamStrides - output array of # MaxVertexElementCount stream strides to fill
-*/
-void FVertexFactory::GetPositionStreamStride(uint32 *OutStreamStrides) const
-{
-	int32 StreamIndex;
-	for(StreamIndex = 0;StreamIndex < PositionStream.Num();++StreamIndex)
-	{
-		OutStreamStrides[StreamIndex] = PositionStream[StreamIndex].Stride;
-	}
-	// Pad stream strides with 0's to be safe (they can be used in hashes elsewhere)
-	for (;StreamIndex < MaxVertexElementCount;++StreamIndex)
-	{
-		OutStreamStrides[StreamIndex] = 0;
-	}
-}
-
 FVertexElement FVertexFactory::AccessStreamComponent(const FVertexStreamComponent& Component,uint8 AttributeIndex)
 {
 	FVertexStream VertexStream;
diff --git a/Engine/Source/Runtime/RenderCore/Public/RenderCore.h b/Engine/Source/Runtime/RenderCore/Public/RenderCore.h
--- a/Engine/Source/Runtime/RenderCore/Public/RenderCore.h
+++ b/Engine/Source/Runtime/RenderCore/Public/RenderCore.h
@@ -37,8 +37,6 @@ DECLARE_CYCLE_STAT_EXTERN(TEXT("Dynamic Primitive drawing"),STAT_DynamicPrimitiv
 DECLARE_CYCLE_STAT_EXTERN(TEXT("StaticDrawList drawing"),STAT_StaticDrawListDrawTime,STATGROUP_SceneRendering, RENDERCORE_API);
 DECLARE_DWORD_ACCUMULATOR_STAT_EXTERN(TEXT("Lights in scene"),STAT_SceneLights,STATGROUP_SceneRendering, RENDERCORE_API);
 DECLARE_DWORD_COUNTER_STAT_EXTERN(TEXT("Mesh draw calls"),STAT_MeshDrawCalls,STATGROUP_SceneRendering, RENDERCORE_API);
-DECLARE_DWORD_COUNTER_STAT_EXTERN(TEXT("Dynamic path draw calls"),STAT_DynamicPathMeshDrawCalls,STATGROUP_SceneRendering, RENDERCORE_API);
-DECLARE_DWORD_COUNTER_STAT_EXTERN(TEXT("Static list draw calls"),STAT_StaticDrawListMeshDrawCalls,STATGROUP_SceneRendering, RENDERCORE_API);
 DECLARE_CYCLE_STAT_EXTERN(TEXT("Present time"),STAT_PresentTime,STATGROUP_SceneRendering, RENDERCORE_API);
 DECLARE_DWORD_ACCUMULATOR_STAT_EXTERN(TEXT("Decals in scene"),STAT_SceneDecals,STATGROUP_SceneRendering, RENDERCORE_API);
 DECLARE_DWORD_COUNTER_STAT_EXTERN(TEXT("Decals in view"),STAT_Decals,STATGROUP_SceneRendering, RENDERCORE_API);
@@ -92,10 +90,6 @@ DECLARE_CYCLE_STAT_EXTERN(TEXT("WholeScene Shadow Projections"),STAT_RenderWhole
 DECLARE_CYCLE_STAT_EXTERN(TEXT("WholeScene Shadow Depths"),STAT_RenderWholeSceneShadowDepthsTime,STATGROUP_ShadowRendering, RENDERCORE_API);
 DECLARE_CYCLE_STAT_EXTERN(TEXT("PerObject Shadow Projections"),STAT_RenderPerObjectShadowProjectionsTime,STATGROUP_ShadowRendering, RENDERCORE_API);
 DECLARE_CYCLE_STAT_EXTERN(TEXT("PerObject Shadow Depths"),STAT_RenderPerObjectShadowDepthsTime,STATGROUP_ShadowRendering, RENDERCORE_API);
-DECLARE_CYCLE_STAT_EXTERN(TEXT("Static Draw List"),STAT_WholeSceneStaticDrawListShadowDepthsTime,STATGROUP_ShadowRendering, RENDERCORE_API);
-DECLARE_CYCLE_STAT_EXTERN(TEXT("Static Primitives"),STAT_WholeSceneStaticShadowDepthsTime,STATGROUP_ShadowRendering, RENDERCORE_API);
-DECLARE_CYCLE_STAT_EXTERN(TEXT("Dynamic Primitives"),STAT_WholeSceneDynamicShadowDepthsTime,STATGROUP_ShadowRendering, RENDERCORE_API);
-DECLARE_CYCLE_STAT_EXTERN(TEXT("WholeScene Reflective Shadow maps"),STAT_RenderWholeSceneReflectiveShadowMapsTime,STATGROUP_ShadowRendering, RENDERCORE_API); 
 DECLARE_MEMORY_STAT_EXTERN(TEXT("Per-Frame Shadowmap Atlases"),STAT_ShadowmapAtlasMemory,STATGROUP_ShadowRendering, RENDERCORE_API); 
 DECLARE_MEMORY_STAT_EXTERN(TEXT("Cached Shadowmaps"),STAT_CachedShadowmapMemory,STATGROUP_ShadowRendering, RENDERCORE_API); 
 DECLARE_DWORD_COUNTER_STAT_EXTERN(TEXT("Whole Scene shadows"),STAT_WholeSceneShadows,STATGROUP_ShadowRendering, RENDERCORE_API);
@@ -142,7 +136,6 @@ DECLARE_CYCLE_STAT_EXTERN(TEXT("Update Static Draw Lists For Materials (RT)"),ST
 DECLARE_DWORD_COUNTER_STAT_EXTERN( TEXT( "Proxy Total" ), STAT_GameToRendererMallocTotal, STATGROUP_Memory, RENDERCORE_API );
 
 /** Memory stats for tracking virtual allocations used by the renderer to represent the scene. */
-DECLARE_MEMORY_STAT_EXTERN(TEXT("Static draw list memory"),STAT_StaticDrawListMemory,STATGROUP_SceneMemory, RENDERCORE_API);
 DECLARE_MEMORY_STAT_EXTERN(TEXT("Primitive memory"),STAT_PrimitiveInfoMemory,STATGROUP_SceneMemory, RENDERCORE_API);
 DECLARE_MEMORY_STAT_EXTERN(TEXT("Scene memory"),STAT_RenderingSceneMemory,STATGROUP_SceneMemory, RENDERCORE_API);
 DECLARE_MEMORY_STAT_EXTERN(TEXT("ViewState memory"),STAT_ViewStateMemory,STATGROUP_SceneMemory, RENDERCORE_API);
diff --git a/Engine/Source/Runtime/RenderCore/Public/RendererInterface.h b/Engine/Source/Runtime/RenderCore/Public/RendererInterface.h
--- a/Engine/Source/Runtime/RenderCore/Public/RendererInterface.h
+++ b/Engine/Source/Runtime/RenderCore/Public/RendererInterface.h
@@ -28,30 +28,6 @@ struct FSynthBenchmarkResults;
 class IVirtualTextureSpace;
 struct FVirtualTextureSpaceDesc;
 
-enum class EShowMaterialDrawEventTypes
-{
-	None						= 0 << 0,  //0
-	CompositionLighting			= 1 << 0,  //1
-	BasePass					= 1 << 1,  //2
-	DepthPositionOnly			= 1 << 2,  //4
-	Depth						= 1 << 3,  //8
-	DistortionDynamic			= 1 << 4,  //16
-	DistortionStatic			= 1 << 5,  //32
-	MobileBasePass				= 1 << 6,  //64
-	MobileTranslucent			= 1 << 7,  //128
-	MobileTranslucentOpacity	= 1 << 8,  //256
-	ShadowDepth					= 1 << 9,  //512
-	ShadowDepthRsm				= 1 << 10, //1024
-	ShadowDepthStatic			= 1 << 11, //2048
-	StaticDraw					= 1 << 12, //4096
-	StaticDrawStereo			= 1 << 13, //8192
-	TranslucentLighting			= 1 << 14, //16384
-	Translucent					= 1 << 15, //32768
-	Velocity					= 1 << 16, //65536
-	FogVoxelization				= 1 << 17, //131072
-};
-ENUM_CLASS_FLAGS(EShowMaterialDrawEventTypes)
-
 // Shortcut for the allocator used by scene rendering.
 typedef TMemStackAllocator<> SceneRenderingAllocator;
 
diff --git a/Engine/Source/Runtime/RenderCore/Public/VertexFactory.h b/Engine/Source/Runtime/RenderCore/Public/VertexFactory.h
--- a/Engine/Source/Runtime/RenderCore/Public/VertexFactory.h
+++ b/Engine/Source/Runtime/RenderCore/Public/VertexFactory.h
@@ -149,20 +149,19 @@ public:
 	virtual ~FVertexFactoryShaderParameters() {}
 	virtual void Bind(const class FShaderParameterMap& ParameterMap) = 0;
 	virtual void Serialize(FArchive& Ar) = 0;
-	virtual void SetMesh(FRHICommandList& RHICmdList, FShader* VertexShader,const class FVertexFactory* VertexFactory,const class FSceneView& View,const struct FMeshBatchElement& BatchElement,uint32 DataFlags) const = 0;
-	
+
 	/** 
 	 * Gets the vertex factory's shader bindings and vertex streams.
 	 * View can be null when caching mesh draw commands (only for supported vertex factories)
 	 */
 	virtual void GetElementShaderBindings(
 		const class FSceneInterface* Scene,
-		const FSceneView* View,
+		const class FSceneView* View,
 		const class FMeshMaterialShader* Shader,
 		bool bShaderRequiresPositionOnlyStream,
 		ERHIFeatureLevel::Type FeatureLevel,
-		const FVertexFactory* VertexFactory,
-		const FMeshBatchElement& BatchElement,
+		const class FVertexFactory* VertexFactory,
+		const struct FMeshBatchElement& BatchElement,
 		class FMeshDrawSingleShaderBindings& ShaderBindings,
 		FVertexInputStreamArray& VertexStreams) const = 0;
 
@@ -464,32 +463,12 @@ public:
 
 	virtual FVertexFactoryType* GetType() const { return NULL; }
 
-	/**
-	 * Activates the vertex factory.
-	 */
-	void SetStreams(ERHIFeatureLevel::Type InFeatureLevel, FRHICommandList& RHICmdList) const;
-
 	void GetStreams(ERHIFeatureLevel::Type InFeatureLevel, FVertexInputStreamArray& OutVertexStreams) const;
 
-	/**
-	 * Call SetStreamSource on instance streams to offset the read pointer
-	 */
-	void OffsetInstanceStreams(FRHICommandList& RHICmdList, uint32 FirstVertex) const;
-
 	void OffsetInstanceStreams(uint32 InstanceOffset, bool bOperateOnPositionOnly, FVertexInputStreamArray& VertexStreams) const;
 
-	/**
-	* Sets the position stream as the current stream source.
-	*/
-	void SetPositionStream(FRHICommandList& RHICmdList) const;
-
 	void GetPositionOnlyStream(FVertexInputStreamArray& OutVertexStreams) const;
 
-	/**
-	* Call SetStreamSource on instance streams to offset the read pointer
-	 */
-	void OffsetPositionInstanceStreams(FRHICommandList& RHICmdList, uint32 FirstVertex) const;
-
 	/**
 	* Can be overridden by FVertexFactory subclasses to modify their compile environment just before compilation occurs.
 	*/
@@ -525,17 +504,6 @@ public:
 
 	virtual bool RendersPrimitivesAsCameraFacingSprites() const { return false; }
 
-	/**
-	* Fill in array of strides from this factory's vertex streams without shadow/light maps
-	* @param OutStreamStrides - output array of # MaxVertexElementCount stream strides to fill
-	*/
-	int32 GetStreamStrides(uint32 *OutStreamStrides, bool bPadWithZeroes=true) const;
-	/**
-	* Fill in array of strides from this factory's position only vertex streams
-	* @param OutStreamStrides - output array of # MaxVertexElementCount stream strides to fill
-	*/
-	void GetPositionStreamStride(uint32 *OutStreamStrides) const;
-
 	/**
 	 * Get a bitmask representing the visibility of each FMeshBatch element.
 	 * FMeshBatch.bRequiresPerElementVisibility must be set for this to be called.
@@ -647,22 +615,14 @@ public:
 		delete Parameters;
 	}
 
-	void SetMesh(FRHICommandList& RHICmdList, FShader* Shader,const FVertexFactory* VertexFactory,const class FSceneView& View,const struct FMeshBatchElement& BatchElement,uint32 DataFlags) const
-	{
-		if(Parameters)
-		{
-			Parameters->SetMesh(RHICmdList, Shader,VertexFactory,View,BatchElement,DataFlags);
-		}
-	}
-
 	void GetElementShaderBindings(
 		const class FSceneInterface* Scene,
 		const FSceneView* View,
-		const FMeshMaterialShader* Shader,
+		const class FMeshMaterialShader* Shader,
 		bool bShaderRequiresPositionOnlyStream,
 		ERHIFeatureLevel::Type FeatureLevel,
 		const FVertexFactory* VertexFactory,
-		const FMeshBatchElement& BatchElement,
+		const struct FMeshBatchElement& BatchElement,
 		FMeshDrawSingleShaderBindings& ShaderBindings,
 		FVertexInputStreamArray& VertexStreams
 		) const
diff --git a/Engine/Source/Runtime/Renderer/Private/BasePassRendering.cpp b/Engine/Source/Runtime/Renderer/Private/BasePassRendering.cpp
--- a/Engine/Source/Runtime/Renderer/Private/BasePassRendering.cpp
+++ b/Engine/Source/Runtime/Renderer/Private/BasePassRendering.cpp
@@ -5,7 +5,6 @@
 =============================================================================*/
 
 #include "BasePassRendering.h"
-#include "StaticMeshDrawList.h"
 #include "DeferredShadingRenderer.h"
 #include "DynamicPrimitiveDrawing.h"
 #include "ScenePrivate.h"
@@ -346,50 +345,6 @@ FExclusiveDepthStencil::Type GetDefaultBasePassDepthStencilAccess(EShadingPath S
 	}
 }
 
-void FBasePassDrawingPolicy::ApplyDitheredLODTransitionState(FDrawingPolicyRenderState& DrawRenderState, const FViewInfo& ViewInfo, const FStaticMesh& Mesh, const bool InAllowStencilDither)
-{
-	FDepthStencilStateRHIParamRef DepthStencilState = nullptr;
-	DrawRenderState.SetDitheredLODTransitionAlpha(0.0f);
-
-	if (Mesh.bDitheredLODTransition)
-	{
-		if (ViewInfo.StaticMeshFadeOutDitheredLODMap[Mesh.Id])
-		{
-			if (InAllowStencilDither)
-			{
-				DepthStencilState = TStaticDepthStencilState<
-					false, CF_Equal,
-					true, CF_Always, SO_Keep, SO_Keep, SO_Replace,
-					false, CF_Always, SO_Keep, SO_Keep, SO_Keep,
-					0xFF, GET_STENCIL_BIT_MASK(RECEIVE_DECAL, 1) | STENCIL_LIGHTING_CHANNELS_MASK(0x7)
-				>::GetRHI();
-			}
-			else
-			{
-				DrawRenderState.SetDitheredLODTransitionAlpha(ViewInfo.GetTemporalLODTransition());
-			}
-		}
-		else if (ViewInfo.StaticMeshFadeInDitheredLODMap[Mesh.Id])
-		{
-			if (InAllowStencilDither)
-			{
-				DepthStencilState = TStaticDepthStencilState<
-					false, CF_Equal,
-					true, CF_Always, SO_Keep, SO_Keep, SO_Replace,
-					false, CF_Always, SO_Keep, SO_Keep, SO_Keep,
-					0xFF, GET_STENCIL_BIT_MASK(RECEIVE_DECAL, 1) | STENCIL_LIGHTING_CHANNELS_MASK(0x7)
-				>::GetRHI();
-			}
-			else
-			{
-				DrawRenderState.SetDitheredLODTransitionAlpha(ViewInfo.GetTemporalLODTransition() - 1.0f);
-			}
-		}
-	}
-
-	SetDepthStencilStateForBasePass(DrawRenderState, ViewInfo.FeatureLevel, Mesh, Mesh.PrimitiveSceneInfo->Proxy, bEnableReceiveDecalOutput, ViewInfo.Family->UseDebugViewPS(), DepthStencilState);
-}
-
 /**
  * Get shader templates allowing to redirect between compatible shaders.
  */
@@ -809,6 +764,8 @@ void FDeferredShadingSceneRenderer::RenderBasePassViewParallel(FViewInfo& View,
 		BasePassDepthStencilAccess,
 		InDrawRenderState);
 
+	// enqueue RHIThread command that blocks on prereq, lock / unlock vertex buffer upload
+
 	View.ParallelMeshDrawCommandPasses[EMeshPass::BasePass].DispatchDraw(&ParallelSet, ParentCmdList);
 }
 
