Fixed crash when selecting DebugViewMode

diff --git a/Engine/Source/Runtime/Engine/Private/DebugViewModeMaterialProxy.cpp b/Engine/Source/Runtime/Engine/Private/DebugViewModeMaterialProxy.cpp
--- a/Engine/Source/Runtime/Engine/Private/DebugViewModeMaterialProxy.cpp
+++ b/Engine/Source/Runtime/Engine/Private/DebugViewModeMaterialProxy.cpp
@@ -111,6 +111,8 @@ void FDebugViewModeMaterialProxy::ValidateAllShaders(TSet<UMaterialInterface*>&
 {
 	FlushRenderingCommands();
 
+	TArray<FDebugViewModeMaterialProxy*> MaterialsToUpdate;
+
 	for (TMap<FMaterialUsagePair, FDebugViewModeMaterialProxy*>::TIterator It(DebugMaterialShaderMap); It; ++It)
 	{
 		const UMaterialInterface* OriginalMaterialInterface = It.Key().MaterialInterface;
@@ -125,7 +127,11 @@ void FDebugViewModeMaterialProxy::ValidateAllShaders(TSet<UMaterialInterface*>&
 				const FUniformExpressionSet& DebugViewUniformExpressionSet = DebugMaterial->GetGameThreadShaderMap()->GetUniformExpressionSet();
 				const FUniformExpressionSet& OrignialUniformExpressionSet = OriginalMaterial->GetGameThreadShaderMap()->GetUniformExpressionSet();
 
-				if (!(DebugViewUniformExpressionSet == OrignialUniformExpressionSet))
+				if (DebugViewUniformExpressionSet == OrignialUniformExpressionSet)
+				{
+					MaterialsToUpdate.Add(DebugMaterial);
+				}
+				else
 				{
 					// This will happen when the debug shader compiled misses logic. Usually caused by custom features in the original shader compilation not implemented in FDebugViewModeMaterialProxy.
 					UE_LOG(TextureStreamingBuild, Verbose, TEXT("Uniform expression set mismatch for %s, skipping shader"), *DebugMaterial->GetMaterialInterface()->GetName());
@@ -151,6 +157,20 @@ void FDebugViewModeMaterialProxy::ValidateAllShaders(TSet<UMaterialInterface*>&
 			}
 		}
 	}
+
+	ENQUEUE_UNIQUE_RENDER_COMMAND_ONEPARAMETER(
+		UpdateDebugMaterialExpressionCache,
+		TArray<FDebugViewModeMaterialProxy*>, MaterialsToUpdateCopy, MaterialsToUpdate,
+		{
+			for (FDebugViewModeMaterialProxy* MaterialToUpdate : MaterialsToUpdateCopy)
+			{
+				check(MaterialToUpdate);
+				MaterialToUpdate->UpdateUniformExpressionCacheIfNeeded(MaterialToUpdate->FMaterial::GetFeatureLevel());
+			}
+		}
+	)
+
+	FlushRenderingCommands();
 }
 
 FDebugViewModeMaterialProxy::FDebugViewModeMaterialProxy(
diff --git a/Engine/Source/Runtime/Renderer/Private/DebugViewModeRendering.cpp b/Engine/Source/Runtime/Renderer/Private/DebugViewModeRendering.cpp
--- a/Engine/Source/Runtime/Renderer/Private/DebugViewModeRendering.cpp
+++ b/Engine/Source/Runtime/Renderer/Private/DebugViewModeRendering.cpp
@@ -283,8 +283,6 @@ void FDebugViewModeMeshProcessor::AddMeshBatch(const FMeshBatch& RESTRICT MeshBa
 		return;
 	}
 
-	MaterialRenderProxy->UpdateUniformExpressionCacheIfNeeded(FeatureLevel);
-
 	FVertexFactoryType* VertexFactoryType = MeshBatch.VertexFactory->GetType();
 
 	const EMaterialTessellationMode MaterialTessellationMode = Material->GetTessellationMode();
