Fix for accessing texture objects through static switch and other forms of indirection
- Previosuly, TextureSample expression would only look at nodes directly connected to 'Tex' input to determine which texture to use

diff --git a/Engine/Source/Runtime/Engine/Private/Materials/HLSLMaterialTranslator.h b/Engine/Source/Runtime/Engine/Private/Materials/HLSLMaterialTranslator.h
--- a/Engine/Source/Runtime/Engine/Private/Materials/HLSLMaterialTranslator.h
+++ b/Engine/Source/Runtime/Engine/Private/Materials/HLSLMaterialTranslator.h
@@ -4083,6 +4083,11 @@ protected:
 		return Material->GetReferencedTextures().Find(TextureValue);
 	}
 
+	virtual UTexture* GetReferencedTexture(int32 Index)
+	{
+		return Material->GetReferencedTextures()[Index];
+	}
+
 	virtual int32 StaticBool(bool bValue) override
 	{
 		return AddInlinedCodeChunk(MCT_StaticBool,(bValue ? TEXT("true") : TEXT("false")));
diff --git a/Engine/Source/Runtime/Engine/Private/Materials/MaterialExpressions.cpp b/Engine/Source/Runtime/Engine/Private/Materials/MaterialExpressions.cpp
--- a/Engine/Source/Runtime/Engine/Private/Materials/MaterialExpressions.cpp
+++ b/Engine/Source/Runtime/Engine/Private/Materials/MaterialExpressions.cpp
@@ -207,6 +207,7 @@
 #include "Materials/MaterialExpressionAtmosphericLightColor.h"
 #include "Materials/MaterialExpressionMaterialLayerOutput.h"
 #include "Materials/MaterialExpressionCurveAtlasRowParameter.h"
+#include "Materials/MaterialUniformExpressions.h"
 #include "EditorSupportDelegates.h"
 #include "MaterialCompiler.h"
 #if WITH_EDITOR
@@ -1637,47 +1638,42 @@ int32 UMaterialExpressionTextureSample::Compile(class FMaterialCompiler* Compile
 			TextureCodeIndex = Compiler->Texture(Texture, TextureReferenceIndex, SamplerSource, MipValueMode);
 		}
 
+		if (TextureCodeIndex == INDEX_NONE)
+		{
+			// Can't continue without a texture to sample
+			return INDEX_NONE;
+		}
+
 		UTexture* EffectiveTexture = Texture;
 		EMaterialSamplerType EffectiveSamplerType = SamplerType;
 		TOptional<FName> EffectiveParameterName;
 		if (InputExpression)
 		{
-			UMaterialExpressionFunctionInput* FunctionInput = Cast<UMaterialExpressionFunctionInput>(InputExpression);
-			if (FunctionInput)
-			{	
-				UMaterialExpressionFunctionInput* NestedFunctionInput = FunctionInput;
-
-				// Walk the input chain to find the last node in the chain
-				while (true)
+			// If 'InputExpression' is connected, we use need to find the texture object that was passed in
+			// In this case, the texture/sampler assigned on this expression node are not used
+			FMaterialUniformExpression* TextureUniformBase = Compiler->GetParameterUniformExpression(TextureCodeIndex);
+			if (FMaterialUniformExpressionTexture* TextureUniform = TextureUniformBase->GetTextureUniformExpression())
+			{
+				TextureReferenceIndex = TextureUniform->GetTextureIndex();
+				EffectiveTexture = Compiler->GetReferencedTexture(TextureReferenceIndex);
+				// Here we assume that the effective sampler type is correct for the given texture,
+				// with the expectation that the expression that generated this texture should have already triggered an error if the sampler type didn't match
+				EffectiveSamplerType = UMaterialExpressionTextureBase::GetSamplerTypeForTexture(EffectiveTexture);
+				if (FMaterialUniformExpressionTextureParameter* TextureParameterUniform = TextureUniform->GetTextureParameterUniformExpression())
 				{
-					UMaterialExpression* PreviewExpression = NestedFunctionInput->GetEffectivePreviewExpression();
-					if (PreviewExpression && PreviewExpression->IsA(UMaterialExpressionFunctionInput::StaticClass()))
-					{
-						NestedFunctionInput = CastChecked<UMaterialExpressionFunctionInput>(PreviewExpression);
-					}
-					else
-					{
-						break;
-					}
+					EffectiveParameterName = TextureParameterUniform->GetParameterName();
 				}
-				InputExpression = NestedFunctionInput->GetEffectivePreviewExpression();
-			}
-
-			UMaterialExpressionTextureObject* TextureObjectExpression = Cast<UMaterialExpressionTextureObject>(InputExpression);
-			UMaterialExpressionTextureObjectParameter* TextureObjectParameter = Cast<UMaterialExpressionTextureObjectParameter>(InputExpression);
-			if (TextureObjectExpression)
-			{
-				EffectiveTexture = TextureObjectExpression->Texture;
-				EffectiveSamplerType = TextureObjectExpression->SamplerType;
 			}
-			else if (TextureObjectParameter)
+			else if (FMaterialUniformExpressionExternalTexture* ExternalTextureUniform = TextureUniformBase->GetExternalTextureUniformExpression())
 			{
-				EffectiveTexture = TextureObjectParameter->Texture;
-				EffectiveSamplerType = TextureObjectParameter->SamplerType;
-				EffectiveParameterName = TextureObjectParameter->ParameterName;
+				TextureReferenceIndex = ExternalTextureUniform->GetSourceTextureIndex();
+				EffectiveTexture = Compiler->GetReferencedTexture(TextureReferenceIndex);
+				EffectiveSamplerType = SAMPLERTYPE_External;
+				if (FMaterialUniformExpressionExternalTextureParameter* ExternalTextureParameterUniform = ExternalTextureUniform->GetExternalTextureParameterUniformExpression())
+				{
+					EffectiveParameterName = ExternalTextureParameterUniform->GetParameterName();
+				}
 			}
-
-			TextureReferenceIndex = Compiler->GetTextureReferenceIndex(EffectiveTexture);
 		}
 
 		if (EffectiveTexture && VerifySamplerType(Compiler, (Desc.Len() > 0 ? *Desc : TEXT("TextureSample")), EffectiveTexture, EffectiveSamplerType))
@@ -2016,6 +2012,11 @@ int32 UMaterialExpressionTextureObjectParameter::Compile(class FMaterialCompiler
 		return CompilerError(Compiler, GetRequirements());
 	}
 
+	if (!VerifySamplerType(Compiler, (Desc.Len() > 0 ? *Desc : TEXT("TextureObjectParameter")), Texture, SamplerType))
+	{
+		return INDEX_NONE;
+	}
+
 	return SamplerType == SAMPLERTYPE_External ? Compiler->ExternalTextureParameter(ParameterName, Texture) : Compiler->TextureParameter(ParameterName, Texture);
 }
 
@@ -2092,6 +2093,12 @@ int32 UMaterialExpressionTextureObject::Compile(class FMaterialCompiler* Compile
 	{
 		return CompilerError(Compiler, TEXT("Requires valid texture"));
 	}
+
+	if (!VerifySamplerType(Compiler, (Desc.Len() > 0 ? *Desc : TEXT("TextureObject")), Texture, SamplerType))
+	{
+		return INDEX_NONE;
+	}
+
 	return SamplerType == SAMPLERTYPE_External ? Compiler->ExternalTexture(Texture) : Compiler->Texture(Texture);
 }
 
diff --git a/Engine/Source/Runtime/Engine/Private/Materials/MaterialUniformExpressions.h b/Engine/Source/Runtime/Engine/Private/Materials/MaterialUniformExpressions.h
--- a/Engine/Source/Runtime/Engine/Private/Materials/MaterialUniformExpressions.h
+++ b/Engine/Source/Runtime/Engine/Private/Materials/MaterialUniformExpressions.h
@@ -300,6 +300,8 @@ public:
 	{}
 
 	// FMaterialUniformExpression interface.
+	virtual class FMaterialUniformExpressionTextureParameter* GetTextureParameterUniformExpression() override { return this; }
+
 	virtual void Serialize(FArchive& Ar)
 	{
 		Ar << ParameterInfo;
@@ -416,6 +418,12 @@ public:
 	virtual void Serialize(FArchive& Ar) override;
 	virtual bool IsIdentical(const FMaterialUniformExpression* OtherExpression) const override;
 	virtual bool GetExternalTexture(const FMaterialRenderContext& Context, FTextureRHIRef& OutTextureRHI, FSamplerStateRHIRef& OutSamplerStateRHI) const override;
+	virtual FMaterialUniformExpressionExternalTextureParameter* GetExternalTextureParameterUniformExpression() override { return this; }
+
+	FName GetParameterName() const
+	{
+		return ParameterName;
+	}
 
 private:
 	FName ParameterName;
diff --git a/Engine/Source/Runtime/Engine/Public/MaterialCompiler.h b/Engine/Source/Runtime/Engine/Public/MaterialCompiler.h
--- a/Engine/Source/Runtime/Engine/Public/MaterialCompiler.h
+++ b/Engine/Source/Runtime/Engine/Public/MaterialCompiler.h
@@ -177,6 +177,7 @@ public:
 	virtual int32 ExternalTextureCoordinateOffset(const FGuid& ExternalTextureGuid) = 0;
 
 	virtual int32 GetTextureReferenceIndex(UTexture* Texture) { return INDEX_NONE; }
+	virtual UTexture* GetReferencedTexture(int32 Index) { return nullptr; }
 
 	int32 Texture(UTexture* InTexture,ESamplerSourceMode SamplerSource=SSM_FromTextureAsset)
 	{
diff --git a/Engine/Source/Runtime/Engine/Public/MaterialShared.h b/Engine/Source/Runtime/Engine/Public/MaterialShared.h
--- a/Engine/Source/Runtime/Engine/Public/MaterialShared.h
+++ b/Engine/Source/Runtime/Engine/Public/MaterialShared.h
@@ -274,6 +274,7 @@ public:
 	/** Accesses the texture used for rendering this uniform expression. */
 	virtual void GetGameThreadTextureValue(const class UMaterialInterface* MaterialInterface,const FMaterial& Material,UTexture*& OutValue,bool bAllowOverride=true) const;
 	virtual class FMaterialUniformExpressionTexture* GetTextureUniformExpression() { return this; }
+	virtual class FMaterialUniformExpressionTextureParameter* GetTextureParameterUniformExpression() { return nullptr; }
 	void SetTransientOverrideTextureValue( UTexture* InOverrideTexture );
 
 	virtual bool IsConstant() const
@@ -308,6 +309,8 @@ public:
 	virtual bool IsConstant() const override { return false; }
 	virtual bool IsIdentical(const FMaterialUniformExpression* OtherExpression) const override;
 
+	int32 GetSourceTextureIndex() const { return SourceTextureIndex; }
+
 protected:
 
 	/** Resolve the guid that relates to texture information inside FExternalTexture */
@@ -332,6 +335,7 @@ public:
 
 	// FMaterialUniformExpression interface.
 	virtual FMaterialUniformExpressionExternalTexture* GetExternalTextureUniformExpression() override { return this; }
+	virtual class FMaterialUniformExpressionExternalTextureParameter* GetExternalTextureParameterUniformExpression() { return nullptr; }
 
 	// Lookup the external texture if it is set
 	virtual bool GetExternalTexture(const FMaterialRenderContext& Context, FTextureRHIRef& OutTextureRHI, FSamplerStateRHIRef& OutSamplerStateRHI) const;
