Recreate uniform buffer inside uniform expression cache if it's layout changes.

diff --git a/Engine/Source/Runtime/Engine/Private/Materials/MaterialShared.cpp b/Engine/Source/Runtime/Engine/Private/Materials/MaterialShared.cpp
--- a/Engine/Source/Runtime/Engine/Private/Materials/MaterialShared.cpp
+++ b/Engine/Source/Runtime/Engine/Private/Materials/MaterialShared.cpp
@@ -2019,9 +2019,9 @@ void FMaterialRenderProxy::EvaluateUniformExpressions(FUniformExpressionCache& O
 	}
 	else
 	{
-		if (IsValidRef(OutUniformExpressionCache.UniformBuffer))
+		if (IsValidRef(OutUniformExpressionCache.UniformBuffer)
+			&& OutUniformExpressionCache.UniformBuffer->GetLayout() == UniformBufferStruct.GetLayout())
 		{
-			check(OutUniformExpressionCache.UniformBuffer->GetLayout() == UniformBufferStruct.GetLayout());
 			RHIUpdateUniformBuffer(OutUniformExpressionCache.UniformBuffer, TempBuffer);
 		}
 		else
