Fixed crash when saving material

diff --git a/Engine/Source/Runtime/Engine/Private/DebugViewModeMaterialProxy.cpp b/Engine/Source/Runtime/Engine/Private/DebugViewModeMaterialProxy.cpp
--- a/Engine/Source/Runtime/Engine/Private/DebugViewModeMaterialProxy.cpp
+++ b/Engine/Source/Runtime/Engine/Private/DebugViewModeMaterialProxy.cpp
@@ -154,6 +154,11 @@ void FDebugViewModeMaterialProxy::ValidateAllShaders(TSet<UMaterialInterface*>&
 					DebugMaterial->MarkAsInvalid();
 					Materials.Remove(const_cast<UMaterialInterface*>(DebugMaterial->GetMaterialInterface()));
 				}
+				else
+				{
+					MaterialsToUpdate.Add(DebugMaterial);
+				}
+				
 			}
 		}
 	}
