Bumping shader GUIDs. This will hopefully allow CIS to recover from bad shaders inside DDC.

diff --git a/Engine/Plugins/FX/Niagara/Source/NiagaraShader/Private/NiagaraShaderDerivedDataVersion.h b/Engine/Plugins/FX/Niagara/Source/NiagaraShader/Private/NiagaraShaderDerivedDataVersion.h
--- a/Engine/Plugins/FX/Niagara/Source/NiagaraShader/Private/NiagaraShaderDerivedDataVersion.h
+++ b/Engine/Plugins/FX/Niagara/Source/NiagaraShader/Private/NiagaraShaderDerivedDataVersion.h
@@ -11,5 +11,5 @@ NiagaraShaderDerivedDataVersion.h: Shader derived data version for Niagara.
 // In case of merge conflicts with DDC versions, you MUST generate a new GUID and set this new
 // guid as version
 
-#define NIAGARASHADERMAP_DERIVEDDATA_VER		TEXT("1DC4E39C8C2B41729D60563AA19EB231")
+#define NIAGARASHADERMAP_DERIVEDDATA_VER		TEXT("09312EC2BA384FE6B5973CDAD64A62A0")
 													  
diff --git a/Engine/Source/Runtime/Engine/Private/ShaderDerivedDataVersion.h b/Engine/Source/Runtime/Engine/Private/ShaderDerivedDataVersion.h
--- a/Engine/Source/Runtime/Engine/Private/ShaderDerivedDataVersion.h
+++ b/Engine/Source/Runtime/Engine/Private/ShaderDerivedDataVersion.h
@@ -11,5 +11,5 @@
 // In case of merge conflicts with DDC versions, you MUST generate a new GUID and set this new
 // guid as version
 
-#define GLOBALSHADERMAP_DERIVEDDATA_VER			TEXT("4A393A7FDF3048EF8376E103162A8ED6")
-#define MATERIALSHADERMAP_DERIVEDDATA_VER		TEXT("F544D0CDF81844F7BFC5027050D52F68")
\ No newline at end of file
+#define GLOBALSHADERMAP_DERIVEDDATA_VER			TEXT("A97D25A9914A43C193801914248E0F66")
+#define MATERIALSHADERMAP_DERIVEDDATA_VER		TEXT("AD9DE195370D4C7DB112DF9967877D21")
\ No newline at end of file
