Can't allow RHI resources to die on the gamethread. Fixes a crash with skincache and sequencer.
Tracked in UE-68826 for a better fix.

diff --git a/Engine/Source/Runtime/Engine/Private/GPUSkinVertexFactory.cpp b/Engine/Source/Runtime/Engine/Private/GPUSkinVertexFactory.cpp
--- a/Engine/Source/Runtime/Engine/Private/GPUSkinVertexFactory.cpp
+++ b/Engine/Source/Runtime/Engine/Private/GPUSkinVertexFactory.cpp
@@ -701,7 +701,13 @@ void FGPUSkinPassthroughVertexFactory::InternalUpdateVertexDeclaration(FGPUBaseS
 	}
 
 	int32 PrevNumStreams = Streams.Num();
-	UpdateRHI();
+
+	//hack to allow us to release the alias pointers properly in ReleaseRHI.
+	//To be cleaned up in UE-68826
+	FLocalVertexFactory::ReleaseRHI();
+	FLocalVertexFactory::ReleaseDynamicRHI();
+	FLocalVertexFactory::InitDynamicRHI();
+	FLocalVertexFactory::InitRHI();
 
 	// Verify no additional stream was created
 	check(Streams.Num() == PrevNumStreams);
diff --git a/Engine/Source/Runtime/Engine/Public/GPUSkinVertexFactory.h b/Engine/Source/Runtime/Engine/Public/GPUSkinVertexFactory.h
--- a/Engine/Source/Runtime/Engine/Public/GPUSkinVertexFactory.h
+++ b/Engine/Source/Runtime/Engine/Public/GPUSkinVertexFactory.h
@@ -494,6 +494,15 @@ public:
 		TangentStreamIndex = -1;
 	}
 
+	virtual void ReleaseRHI() override
+	{
+		FLocalVertexFactory::ReleaseRHI();
+
+		//when adding anything else to this function be aware of the bypassing code in InternalUpdateVertexDeclaration
+		PositionVBAlias.ReleaseRHI();
+		TangentVBAlias.ReleaseRHI();
+	}
+
 protected:
 	// Vertex buffer required for creating the Vertex Declaration
 	FVertexBuffer PositionVBAlias;
diff --git a/Engine/Source/Runtime/Engine/Public/LocalVertexFactory.h b/Engine/Source/Runtime/Engine/Public/LocalVertexFactory.h
--- a/Engine/Source/Runtime/Engine/Public/LocalVertexFactory.h
+++ b/Engine/Source/Runtime/Engine/Public/LocalVertexFactory.h
@@ -71,7 +71,7 @@ public:
 
 	// FRenderResource interface.
 	virtual void InitRHI() override;
-	virtual void ReleaseRHI()
+	virtual void ReleaseRHI() override
 	{
 		UniformBuffer.SafeRelease();
 		FVertexFactory::ReleaseRHI();
