Add ray tracing support for Alembic geometry cache

diff --git a/Engine/Shaders/Private/GeometryCacheVertexFactory.ush b/Engine/Shaders/Private/GeometryCacheVertexFactory.ush
--- a/Engine/Shaders/Private/GeometryCacheVertexFactory.ush
+++ b/Engine/Shaders/Private/GeometryCacheVertexFactory.ush
@@ -62,6 +62,33 @@ struct FVertexFactoryInput
 #endif
 };
 
+// RHI_RAYTRACING
+#if RAYHITGROUPSHADER
+FVertexFactoryInput LoadVertexFactoryInput(uint TriangleIndex, int VertexIndex)
+{
+	FVertexFactoryInput Input;
+	
+	FTriangleBaseAttributes Triangle = LoadTriangleBaseAttributes(TriangleIndex);
+	uint VertexId = Triangle.Indices[VertexIndex];
+	
+	Input.Position.x = GeomCacheMVF.Position[VertexId * 3 + 0];
+	Input.Position.y = GeomCacheMVF.Position[VertexId * 3 + 1];
+	Input.Position.z = GeomCacheMVF.Position[VertexId * 3 + 2];
+	Input.TangentX = GeomCacheMVF.TangentX[VertexId];
+	Input.TangentZ = GeomCacheMVF.TangentZ[VertexId];
+	Input.Color = GeomCacheMVF.Color[VertexId];
+    Input.MotionBlurData.x = GeomCacheMVF.MotionBlurData[VertexId * 3 + 0];
+    Input.MotionBlurData.y = GeomCacheMVF.MotionBlurData[VertexId * 3 + 1];
+    Input.MotionBlurData.z = GeomCacheMVF.MotionBlurData[VertexId * 3 + 2];
+	
+#if NUM_MATERIAL_TEXCOORDS_VERTEX
+	Input.TexCoords[0].x = GeomCacheMVF.TexCoords[VertexId * 2 + 0];
+	Input.TexCoords[0].y = GeomCacheMVF.TexCoords[VertexId * 2 + 1];
+#endif
+
+	return Input;
+}
+#endif
 
 struct FPositionOnlyVertexFactoryInput
 {
@@ -78,9 +105,9 @@ float4 VertexFactoryGetWorldPosition(FPositionOnlyVertexFactoryInput Input)
 
 struct FVertexFactoryInterpolantsVSToPS
 {
-	TANGENTTOWORLD_INTERPOLATOR_BLOCK
-
+#if INTERPOLATE_VERTEX_COLOR
 	float4 Color : COLOR0;
+#endif
 
 #if NUM_TEX_COORD_INTERPOLATORS
 	float2	TexCoords[NUM_TEX_COORD_INTERPOLATORS]	: TEXCOORD0;
@@ -89,6 +116,8 @@ struct FVertexFactoryInterpolantsVSToPS
 #if INSTANCED_STEREO
 	nointerpolation uint EyeIndex : PACKED_EYE_INDEX;
 #endif
+
+	TANGENTTOWORLD_INTERPOLATOR_BLOCK
 };
 
 struct FVertexFactoryInterpolantsVSToDS
@@ -113,8 +142,15 @@ FMaterialPixelParameters GetMaterialPixelParameters(FVertexFactoryInterpolantsVS
 	half3 TangentToWorld0 = Interpolants.TangentToWorld0.xyz;
 	half4 TangentToWorld2 = Interpolants.TangentToWorld2;
 	Result.TangentToWorld = AssembleTangentToWorld( TangentToWorld0, TangentToWorld2 );
-	Result.UnMirrored = TangentToWorld2.w;	
+#if USE_WORLDVERTEXNORMAL_CENTER_INTERPOLATION
+	Result.WorldVertexNormal_Center = Interpolants.TangentToWorld2_Center.xyz;
+#endif
+	Result.UnMirrored = TangentToWorld2.w;
+#if INTERPOLATE_VERTEX_COLOR
 	Result.VertexColor = Interpolants.Color;
+#else
+	Result.VertexColor = 0;
+#endif
 	Result.TwoSidedSign = 1;
 	return Result;
 }
@@ -160,7 +196,7 @@ FMaterialVertexParameters GetMaterialVertexParameters(FVertexFactoryInput Input,
 */
 float3 UnpackedPosition( FVertexFactoryInput Input )
 {
-		return float3(Input.Position.xyz*MeshExtension + MeshOrigin);
+		return float3(Input.Position.xyz * GeomCache.MeshExtension + GeomCache.MeshOrigin);
 }
 
 /**
@@ -187,9 +223,8 @@ half3x3 CalcTangentToLocal(FVertexFactoryInput Input)
 	return Result;
 }
 
-float3x3 CalcTangentToWorld(FVertexFactoryInput Input, float3x3 TangentToLocal)
+void CalcTangentToWorld(FVertexFactoryInput Input, FVertexFactoryIntermediates Intermediates, out float3 TangentToWorld0, out float4 TangentToWorld2)
 {
-   
 	float3x3 LocalToWorld = GetLocalToWorld3x3();
 
 	// Remove scaling.
@@ -198,7 +233,10 @@ float3x3 CalcTangentToWorld(FVertexFactoryInput Input, float3x3 TangentToLocal)
 	LocalToWorld[1] *= InvScale.y;
 	LocalToWorld[2] *= InvScale.z;
 
-	return mul(TangentToLocal, LocalToWorld);
+	float3x3 TangentToWorld = mul(Intermediates.TangentToLocal, LocalToWorld);
+
+	TangentToWorld0 = TangentToWorld[0];
+	TangentToWorld2 = float4(TangentToWorld[2], Input.TangentZ.w * Primitive.InvNonUniformScaleAndDeterminantSign.w);
 }
 
 FVertexFactoryIntermediates GetVertexFactoryIntermediates(FVertexFactoryInput Input)
@@ -209,7 +247,6 @@ FVertexFactoryIntermediates GetVertexFactoryIntermediates(FVertexFactoryInput In
 
 	// Fill TangentToLocal
 	Intermediates.TangentToLocal = CalcTangentToLocal(Input);
-    Intermediates.TangentToWorld = CalcTangentToWorld(Input,Intermediates.TangentToLocal);
 
 	// Swizzle vertex color.
 	Intermediates.Color = Input.Color FCOLOR_COMPONENT_SWIZZLE;
@@ -260,11 +297,15 @@ FVertexFactoryInterpolantsVSToPS VertexFactoryGetInterpolantsVSToPS(FVertexFacto
 #endif
 
     //Packs the TangentToWorld matrix into the interpolants.    
-	float TangentToWorldSign = TangentBias(Input.TangentZ.w) * Primitive.InvNonUniformScaleAndDeterminantSign.w;
-	Interpolants.TangentToWorld0 = float4(Intermediates.TangentToWorld[0],0);
-	Interpolants.TangentToWorld2 = float4(Intermediates.TangentToWorld[2],TangentToWorldSign);
-    
+	Interpolants.TangentToWorld0.w = 0;
+	CalcTangentToWorld(Input, Intermediates, Interpolants.TangentToWorld0.xyz, Interpolants.TangentToWorld2);
+#if USE_WORLDVERTEXNORMAL_CENTER_INTERPOLATION
+	Interpolants.TangentToWorld2_Center = Interpolants.TangentToWorld2;
+#endif
+
+#if INTERPOLATE_VERTEX_COLOR
 	Interpolants.Color = Intermediates.Color;
+#endif
 
 #if INSTANCED_STEREO
 	Interpolants.EyeIndex = 0;
@@ -276,7 +317,7 @@ FVertexFactoryInterpolantsVSToPS VertexFactoryGetInterpolantsVSToPS(FVertexFacto
 // @return The previous position of the vertex in object space.
 float3 CalcPreviousPosition(FVertexFactoryInput Input, FVertexFactoryIntermediates Intermediates)
 {
-    float3 PreviousPosition =  (Input.MotionBlurData.xyz * MotionBlurDataExtension + MotionBlurDataOrigin) + Intermediates.UnpackedPosition * MotionBlurPositionScale;
+    float3 PreviousPosition =  (Input.MotionBlurData.xyz * GeomCache.MotionBlurDataExtension + GeomCache.MotionBlurDataOrigin) + Intermediates.UnpackedPosition * GeomCache.MotionBlurPositionScale;
 #if 0    
     float ExaggerateMotion = 10.0;
     float3 Delta = Intermediates.UnpackedPosition - PreviousPosition;
@@ -337,7 +378,9 @@ FMaterialTessellationParameters GetMaterialTessellationParameters(FVertexFactory
 	half3 TangentToWorld0 = Interpolants.InterpolantsVSToPS.TangentToWorld0.xyz;
 	half4 TangentToWorld2 = Interpolants.InterpolantsVSToPS.TangentToWorld2;
 
+#if INTERPOLATE_VERTEX_COLOR
 	Result.VertexColor = Interpolants.InterpolantsVSToPS.Color;
+#endif
 
 	Result.TangentToWorld = AssembleTangentToWorld( TangentToWorld0, TangentToWorld2 );
 
@@ -355,7 +398,9 @@ FVertexFactoryInterpolantsVSToDS VertexFactoryInterpolate(FVertexFactoryInterpol
 
 	TESSELLATION_INTERPOLATE_MEMBER(InterpolantsVSToPS.TangentToWorld0.xyz);
 	TESSELLATION_INTERPOLATE_MEMBER(InterpolantsVSToPS.TangentToWorld2);
+#if INTERPOLATE_VERTEX_COLOR
 	TESSELLATION_INTERPOLATE_MEMBER(InterpolantsVSToPS.Color);
+#endif
 
 #if NUM_TEX_COORD_INTERPOLATORS
 	UNROLL
diff --git a/Engine/Source/Runtime/Engine/Private/GeometryCacheVertexFactory.cpp b/Engine/Source/Runtime/Engine/Private/GeometryCacheVertexFactory.cpp
--- a/Engine/Source/Runtime/Engine/Private/GeometryCacheVertexFactory.cpp
+++ b/Engine/Source/Runtime/Engine/Private/GeometryCacheVertexFactory.cpp
@@ -15,6 +15,9 @@
 FGeometryCacheVertexFactoryShaderParameters
 -----------------------------------------------------------------------------*/
 
+IMPLEMENT_GLOBAL_SHADER_PARAMETER_STRUCT(FGeometryCacheVertexFactoryUniformBufferParameters, "GeomCache");
+IMPLEMENT_GLOBAL_SHADER_PARAMETER_STRUCT(FGeometryCacheManualVertexFetchUniformBufferParameters, "GeomCacheMVF");
+
 /** Shader parameters for use with TGPUSkinVertexFactory */
 class FGeometryCacheVertexFactoryShaderParameters : public FVertexFactoryShaderParameters
 {
@@ -75,6 +78,9 @@ public:
 		ShaderBindings.Add(MotionBlurDataOrigin, BatchData->MotionBlurDataOrigin);
 		ShaderBindings.Add(MotionBlurDataExtension, BatchData->MotionBlurDataExtension);
 		ShaderBindings.Add(MotionBlurPositionScale, BatchData->MotionBlurPositionScale);
+
+		ShaderBindings.Add(Shader->GetUniformBufferParameter<FGeometryCacheVertexFactoryUniformBufferParameters>(), BatchData->UniformBuffer);
+		ShaderBindings.Add(Shader->GetUniformBufferParameter<FGeometryCacheManualVertexFetchUniformBufferParameters>(), BatchData->ManualVertexFetchUniformBuffer);
 	}
 
 	virtual uint32 GetSize() const override { return sizeof(*this); }
@@ -116,22 +122,74 @@ void FGeometryCacheVertexVertexFactory::SetData(const FDataType& InData)
 	UpdateRHI();
 }
 
+class FDummyVertexBuffer : public FVertexBuffer
+{
+public:
+	FShaderResourceViewRHIRef SRV;
+
+	virtual void InitRHI() override
+	{
+		FRHIResourceCreateInfo CreateInfo;
+		void* BufferData = nullptr;
+		VertexBufferRHI = RHICreateAndLockVertexBuffer(sizeof(FVector4) * 2, BUF_Static | BUF_ShaderResource, CreateInfo, BufferData);
+		FVector4* DummyContents = (FVector4*)BufferData;
+		DummyContents[0] = FVector4(0.0f, 0.0f, 0.0f, 0.0f);
+		DummyContents[1] = FVector4(1.0f, 1.0f, 1.0f, 1.0f);
+		RHIUnlockVertexBuffer(VertexBufferRHI);
+
+		SRV = RHICreateShaderResourceView(VertexBufferRHI, sizeof(float), PF_R32_FLOAT);
+	}
+
+	virtual void ReleaseRHI() override
+	{
+		SRV.SafeRelease();
+		FVertexBuffer::ReleaseRHI();
+	}
+};
+TGlobalResource<FDummyVertexBuffer> DummyPositionBuffer;
+
+class FDummyTangentBuffer : public FVertexBuffer
+{
+public:
+	FShaderResourceViewRHIRef SRV;
+
+	virtual void InitRHI() override
+	{
+		FRHIResourceCreateInfo CreateInfo;
+		void* BufferData = nullptr;
+		VertexBufferRHI = RHICreateAndLockVertexBuffer(sizeof(FVector4) * 2, BUF_Static | BUF_ShaderResource, CreateInfo, BufferData);
+		FVector4* DummyContents = (FVector4*)BufferData;
+		DummyContents[0] = FVector4(0.0f, 0.0f, 0.0f, 0.0f);
+		DummyContents[1] = FVector4(1.0f, 1.0f, 1.0f, 1.0f);
+		RHIUnlockVertexBuffer(VertexBufferRHI);
+
+		SRV = RHICreateShaderResourceView(VertexBufferRHI, sizeof(FPackedNormal), PF_R8G8B8A8_SNORM);
+	}
+
+	virtual void ReleaseRHI() override
+	{
+		SRV.SafeRelease();
+		FVertexBuffer::ReleaseRHI();
+	}
+};
+TGlobalResource<FDummyTangentBuffer> GDummyTangentBuffer;
+
 void FGeometryCacheVertexVertexFactory::InitRHI()
 {
 	// Position needs to be separate from the rest (we just theck tangents here)
-	check(Data.PositionComponent.VertexBuffer != Data.TangentBasisComponents[0].VertexBuffer)
-		// Motion Blur data also needs to be separate from the rest
-		check(Data.MotionBlurDataComponent.VertexBuffer != Data.TangentBasisComponents[0].VertexBuffer)
-		check(Data.MotionBlurDataComponent.VertexBuffer != Data.PositionComponent.VertexBuffer)
-
-		// If the vertex buffer containing position is not the same vertex buffer containing the rest of the data,
-		// then initialize PositionStream and PositionDeclaration.
-		if (Data.PositionComponent.VertexBuffer != Data.TangentBasisComponents[0].VertexBuffer)
-		{
-			FVertexDeclarationElementList PositionOnlyStreamElements;
-			PositionOnlyStreamElements.Add(AccessPositionStreamComponent(Data.PositionComponent, 0));
-			InitPositionDeclaration(PositionOnlyStreamElements);
-		}
+	check(Data.PositionComponent.VertexBuffer != Data.TangentBasisComponents[0].VertexBuffer);
+	// Motion Blur data also needs to be separate from the rest
+	check(Data.MotionBlurDataComponent.VertexBuffer != Data.TangentBasisComponents[0].VertexBuffer);
+	check(Data.MotionBlurDataComponent.VertexBuffer != Data.PositionComponent.VertexBuffer);
+
+	// If the vertex buffer containing position is not the same vertex buffer containing the rest of the data,
+	// then initialize PositionStream and PositionDeclaration.
+	if (Data.PositionComponent.VertexBuffer != Data.TangentBasisComponents[0].VertexBuffer)
+	{
+		FVertexDeclarationElementList PositionOnlyStreamElements;
+		PositionOnlyStreamElements.Add(AccessPositionStreamComponent(Data.PositionComponent, 0));
+		InitPositionDeclaration(PositionOnlyStreamElements);
+	}
 
 	FVertexDeclarationElementList Elements;
 	if (Data.PositionComponent.VertexBuffer != NULL)
@@ -202,9 +260,87 @@ void FGeometryCacheVertexVertexFactory::InitRHI()
 	check(IsValidRef(GetDeclaration()));
 }
 
+void FGeometryCacheVertexVertexFactory::CreateManualVertexFetchUniformBuffer(
+	const FVertexBuffer* PoistionBuffer, 
+	const FVertexBuffer* MotionBlurBuffer,
+	FGeometryCacheVertexFactoryUserData& OutUserData) const
+{
+	FGeometryCacheManualVertexFetchUniformBufferParameters ManualVertexFetchParameters;
+
+	if (PoistionBuffer != NULL)
+	{
+		OutUserData.PositionSRV = RHICreateShaderResourceView(PoistionBuffer->VertexBufferRHI, sizeof(float), PF_R32_FLOAT);
+		// Position will need per-component fetch since we don't have R32G32B32 pixel format
+		ManualVertexFetchParameters.Position = OutUserData.PositionSRV;
+	}
+	else
+	{
+		ManualVertexFetchParameters.Position = DummyPositionBuffer.SRV;
+	}
+
+	if (Data.TangentBasisComponents[0].VertexBuffer != NULL)
+	{
+		OutUserData.TangentXSRV = RHICreateShaderResourceView(Data.TangentBasisComponents[0].VertexBuffer->VertexBufferRHI, sizeof(FPackedNormal), PF_R8G8B8A8_SNORM);
+		ManualVertexFetchParameters.TangentX = OutUserData.TangentXSRV;
+	}
+	else
+	{
+		ManualVertexFetchParameters.TangentX = GDummyTangentBuffer.SRV;
+	}
+
+	if (Data.TangentBasisComponents[1].VertexBuffer != NULL)
+	{
+		OutUserData.TangentZSRV = RHICreateShaderResourceView(Data.TangentBasisComponents[1].VertexBuffer->VertexBufferRHI, sizeof(FPackedNormal), PF_R8G8B8A8_SNORM);
+		ManualVertexFetchParameters.TangentZ = OutUserData.TangentZSRV;
+	}
+	else
+	{
+		ManualVertexFetchParameters.TangentZ = GDummyTangentBuffer.SRV;
+	}
+
+	if (Data.ColorComponent.VertexBuffer)
+	{
+		OutUserData.ColorSRV = RHICreateShaderResourceView(Data.ColorComponent.VertexBuffer->VertexBufferRHI, sizeof(FColor), PF_B8G8R8A8);
+		ManualVertexFetchParameters.Color = OutUserData.ColorSRV;
+	}
+	else
+	{
+		OutUserData.ColorSRV = GNullColorVertexBuffer.VertexBufferSRV;
+		ManualVertexFetchParameters.Color = OutUserData.ColorSRV;
+	}
+
+	if (MotionBlurBuffer)
+	{
+		OutUserData.MotionBlurDataSRV = RHICreateShaderResourceView(MotionBlurBuffer->VertexBufferRHI, sizeof(float), PF_R32_FLOAT);
+		ManualVertexFetchParameters.MotionBlurData = OutUserData.MotionBlurDataSRV;
+	}
+	else if (PoistionBuffer != NULL)
+	{
+		ManualVertexFetchParameters.MotionBlurData = OutUserData.PositionSRV;
+	}
+	else
+	{
+		ManualVertexFetchParameters.MotionBlurData = DummyPositionBuffer.SRV;
+	}
+
+	if (Data.TextureCoordinates.Num())
+	{
+		checkf(Data.TextureCoordinates.Num() <= 1, TEXT("We're assuming FGeometryCacheSceneProxy uses only one TextureCoordinates vertex buffer"));
+		OutUserData.TexCoordsSRV = RHICreateShaderResourceView(Data.TextureCoordinates[0].VertexBuffer->VertexBufferRHI, sizeof(float), PF_R32_FLOAT);
+		// TexCoords will need per-component fetch since we don't have R32G32 pixel format
+		ManualVertexFetchParameters.TexCoords = OutUserData.TexCoordsSRV;
+	}
+	else
+	{
+		ManualVertexFetchParameters.TexCoords = DummyPositionBuffer.SRV;
+	}
+
+	OutUserData.ManualVertexFetchUniformBuffer = FGeometryCacheManualVertexFetchUniformBufferParametersRef::CreateUniformBufferImmediate(ManualVertexFetchParameters, UniformBuffer_SingleFrame);
+}
+
 FVertexFactoryShaderParameters* FGeometryCacheVertexVertexFactory::ConstructShaderParameters(EShaderFrequency ShaderFrequency)
 {
-	return (ShaderFrequency == SF_Vertex) ? new FGeometryCacheVertexFactoryShaderParameters() : nullptr;
+	return (ShaderFrequency == SF_Vertex || ShaderFrequency == SF_RayHitGroup) ? new FGeometryCacheVertexFactoryShaderParameters() : nullptr;
 }
 
 bool FGeometryCacheVertexVertexFactory::ShouldCompilePermutation(EShaderPlatform Platform, const class FMaterial* Material, const class FShaderType* ShaderType)
diff --git a/Engine/Source/Runtime/Engine/Public/GeometryCacheVertexFactory.h b/Engine/Source/Runtime/Engine/Public/GeometryCacheVertexFactory.h
--- a/Engine/Source/Runtime/Engine/Public/GeometryCacheVertexFactory.h
+++ b/Engine/Source/Runtime/Engine/Public/GeometryCacheVertexFactory.h
@@ -14,6 +14,27 @@
 #include "VertexFactory.h"
 #include "LocalVertexFactory.h"
 
+BEGIN_GLOBAL_SHADER_PARAMETER_STRUCT(FGeometryCacheVertexFactoryUniformBufferParameters, ENGINE_API)
+	SHADER_PARAMETER(FVector, MeshOrigin)
+	SHADER_PARAMETER(FVector, MeshExtension)
+	SHADER_PARAMETER(FVector, MotionBlurDataOrigin)
+	SHADER_PARAMETER(FVector, MotionBlurDataExtension)
+	SHADER_PARAMETER(float, MotionBlurPositionScale)
+END_GLOBAL_SHADER_PARAMETER_STRUCT()
+
+typedef TUniformBufferRef<FGeometryCacheVertexFactoryUniformBufferParameters> FGeometryCacheVertexFactoryUniformBufferParametersRef;
+
+BEGIN_GLOBAL_SHADER_PARAMETER_STRUCT(FGeometryCacheManualVertexFetchUniformBufferParameters, ENGINE_API)
+	SHADER_PARAMETER_SRV(Buffer<float>, Position)
+	SHADER_PARAMETER_SRV(Buffer<float>, MotionBlurData)
+	SHADER_PARAMETER_SRV(Buffer<half4>, TangentX)
+	SHADER_PARAMETER_SRV(Buffer<half4>, TangentZ)
+	SHADER_PARAMETER_SRV(Buffer<float4>, Color)
+	SHADER_PARAMETER_SRV(Buffer<float>, TexCoords)
+END_GLOBAL_SHADER_PARAMETER_STRUCT()
+
+typedef TUniformBufferRef<FGeometryCacheManualVertexFetchUniformBufferParameters> FGeometryCacheManualVertexFetchUniformBufferParametersRef;
+
 /**
  * The mesh batch element user data should point to an instance of this struct
  */
@@ -30,6 +51,17 @@ struct FGeometryCacheVertexFactoryUserData
 	FVector MotionBlurDataOrigin;
 	FVector MotionBlurDataExtension;
 	float MotionBlurPositionScale;
+
+	FGeometryCacheVertexFactoryUniformBufferParametersRef UniformBuffer;
+
+	FShaderResourceViewRHIRef PositionSRV;
+	FShaderResourceViewRHIRef TangentXSRV;
+	FShaderResourceViewRHIRef TangentZSRV;
+	FShaderResourceViewRHIRef ColorSRV;
+	FShaderResourceViewRHIRef MotionBlurDataSRV;
+	FShaderResourceViewRHIRef TexCoordsSRV;
+
+	FUniformBufferRHIRef ManualVertexFetchUniformBuffer;
 };
 
 class FGeometryCacheVertexFactoryShaderParameters;
@@ -77,6 +109,11 @@ public:
 	*/
 	void SetData(const FDataType& InData);
 
+	void CreateManualVertexFetchUniformBuffer(
+		const FVertexBuffer* PoistionBuffer,
+		const FVertexBuffer* MotionBlurBuffer,
+		FGeometryCacheVertexFactoryUserData& OutUserData) const;
+
 	virtual void InitRHI() override;
 
 	// FRenderResource interface.
diff --git a/Engine/Source/Runtime/Renderer/Private/RayTracing/RayTracingMaterialHitShaders.cpp b/Engine/Source/Runtime/Renderer/Private/RayTracing/RayTracingMaterialHitShaders.cpp
--- a/Engine/Source/Runtime/Renderer/Private/RayTracing/RayTracingMaterialHitShaders.cpp
+++ b/Engine/Source/Runtime/Renderer/Private/RayTracing/RayTracingMaterialHitShaders.cpp
@@ -10,10 +10,12 @@ static bool IsSupportedVertexFactoryType(const FVertexFactoryType* VertexFactory
 	static FName LSkinnedVfFname = FName(TEXT("FGPUSkinPassthroughVertexFactory"), FNAME_Find);
 	static FName InstancedVfFname = FName(TEXT("FInstancedStaticMeshVertexFactory"), FNAME_Find);
 	static FName NiagaraSpriteVfFname = FName(TEXT("FNiagaraSpriteVertexFactory"), FNAME_Find);
+	static FName GeometryCacheVfFname = FName(TEXT("FGeometryCacheVertexVertexFactory"), FNAME_Find);
 
 	return (VertexFactoryType == FindVertexFactoryType(LocalVfFname)
 		|| VertexFactoryType == FindVertexFactoryType(LSkinnedVfFname)
-		|| VertexFactoryType == FindVertexFactoryType(NiagaraSpriteVfFname));
+		|| VertexFactoryType == FindVertexFactoryType(NiagaraSpriteVfFname)
+		|| VertexFactoryType == FindVertexFactoryType(GeometryCacheVfFname));
 }
 
 class FMaterialCHS : public FMeshMaterialShader, public FUniformLightMapPolicyShaderParametersType
