DR - Deprecate ENQUEUE_UNIQUE_RENDER_COMMAND_ FIVE and SIX
#rnx

diff --git a/Engine/Plugins/Experimental/GeometryCache/Source/GeometryCache/Private/GeometryCacheComponent.cpp b/Engine/Plugins/Experimental/GeometryCache/Source/GeometryCache/Private/GeometryCacheComponent.cpp
--- a/Engine/Plugins/Experimental/GeometryCache/Source/GeometryCache/Private/GeometryCacheComponent.cpp
+++ b/Engine/Plugins/Experimental/GeometryCache/Source/GeometryCache/Private/GeometryCacheComponent.cpp
@@ -125,15 +125,15 @@ void UGeometryCacheComponent::TickComponent(float DeltaTime, enum ELevelTick Tic
 		// Schedule an update on the render thread
 		if (FGeometryCacheSceneProxy* CastedProxy = static_cast<FGeometryCacheSceneProxy*>(SceneProxy))
 		{
-			ENQUEUE_UNIQUE_RENDER_COMMAND_FIVEPARAMETER(
-				FGeometryCacheUpdateAnimation,
-				FGeometryCacheSceneProxy*, SceneProxy, CastedProxy,
-				float, AnimationTime, GetAnimationTime(),
-				bool, bLooping, IsLooping(),
-				bool, bIsPlayingBackwards, IsPlayingReversed(),
-				float, PlaybackSpeed, ActualPlaybackSpeed,
+			FGeometryCacheSceneProxy* InSceneProxy = CastedProxy;
+			float AnimationTime = GetAnimationTime();
+			bool bInLooping = IsLooping();
+			bool bIsPlayingBackwards = IsPlayingReversed();
+			float InPlaybackSpeed = ActualPlaybackSpeed;
+			ENQUEUE_RENDER_COMMAND(FGeometryCacheUpdateAnimation)(
+				[InSceneProxy, AnimationTime, bInLooping, bIsPlayingBackwards, InPlaybackSpeed, ActualPlaybackSpeed](FRHICommandList& RHICmdList)
 				{
-					SceneProxy->UpdateAnimation(AnimationTime, bLooping, bIsPlayingBackwards,PlaybackSpeed);
+					InSceneProxy->UpdateAnimation(AnimationTime, bInLooping, bIsPlayingBackwards, InPlaybackSpeed);
 				});
 		}
 	}
@@ -178,15 +178,13 @@ void UGeometryCacheComponent::TickAtThisTime(const float Time, bool bInIsRunning
 		// Schedule an update on the render thread
 		if (FGeometryCacheSceneProxy* CastedProxy = static_cast<FGeometryCacheSceneProxy*>(SceneProxy))
 		{
-			ENQUEUE_UNIQUE_RENDER_COMMAND_FIVEPARAMETER(
-				FGeometryCacheUpdateAnimation,
-				FGeometryCacheSceneProxy*, SceneProxy, CastedProxy,
-				float, AnimationTime, Time,
-				bool, bLooping, bInIsLooping,
-				bool, bIsPlayingBackwards, bInBackwards,
-				float, PlaybackSpeed, ActualPlaybackSpeed,
+			FGeometryCacheSceneProxy* InSceneProxy = CastedProxy;
+			float AnimationTime = Time;
+			float InPlaybackSpeed = ActualPlaybackSpeed;
+			ENQUEUE_RENDER_COMMAND(FGeometryCacheUpdateAnimation)(
+				[InSceneProxy, AnimationTime, bInIsLooping, bInBackwards, InPlaybackSpeed](FRHICommandList& RHICmdList)
 				{
-					SceneProxy->UpdateAnimation(AnimationTime, bLooping, bIsPlayingBackwards,PlaybackSpeed);
+					InSceneProxy->UpdateAnimation(AnimationTime, bInIsLooping, bInBackwards, InPlaybackSpeed);
 				});
 		}
 	}
diff --git a/Engine/Source/Runtime/Engine/Private/Components/SkyLightComponent.cpp b/Engine/Source/Runtime/Engine/Private/Components/SkyLightComponent.cpp
--- a/Engine/Source/Runtime/Engine/Private/Components/SkyLightComponent.cpp
+++ b/Engine/Source/Runtime/Engine/Private/Components/SkyLightComponent.cpp
@@ -167,19 +167,19 @@ FSkyLightSceneProxy::FSkyLightSceneProxy(const USkyLightComponent* InLightCompon
 	, IsDirtyImportanceSamplingData(true)
 #endif
 {
-	ENQUEUE_UNIQUE_RENDER_COMMAND_SIXPARAMETER(
-		FInitSkyProxy,
-		const FSHVectorRGB3*,InIrradianceEnvironmentMap,&InLightComponent->IrradianceEnvironmentMap,
-		const FSHVectorRGB3*,BlendDestinationIrradianceEnvironmentMap,&InLightComponent->BlendDestinationIrradianceEnvironmentMap,
-		const float*,InAverageBrightness,&InLightComponent->AverageBrightness,
-		const float*,BlendDestinationAverageBrightness,&InLightComponent->BlendDestinationAverageBrightness,
-		float,InBlendFraction,InLightComponent->BlendFraction,
-		FSkyLightSceneProxy*,LightSceneProxy,this,
-	{
-		// Only access the irradiance maps on the RT, even though they belong to the USkyLightComponent, 
-		// Because FScene::UpdateSkyCaptureContents does not block the RT so the writes could still be in flight
-		LightSceneProxy->Initialize(InBlendFraction, InIrradianceEnvironmentMap, BlendDestinationIrradianceEnvironmentMap, InAverageBrightness, BlendDestinationAverageBrightness);
-	});
+	const FSHVectorRGB3* InIrradianceEnvironmentMap = &InLightComponent->IrradianceEnvironmentMap;
+	const FSHVectorRGB3* BlendDestinationIrradianceEnvironmentMap = &InLightComponent->BlendDestinationIrradianceEnvironmentMap;
+	const float* InAverageBrightness = &InLightComponent->AverageBrightness;
+	const float* BlendDestinationAverageBrightness = &InLightComponent->BlendDestinationAverageBrightness;
+	float InBlendFraction = InLightComponent->BlendFraction;
+	FSkyLightSceneProxy* LightSceneProxy = this;
+	ENQUEUE_RENDER_COMMAND(FInitSkyProxy)(
+		[InIrradianceEnvironmentMap, BlendDestinationIrradianceEnvironmentMap, InAverageBrightness, BlendDestinationAverageBrightness, InBlendFraction, LightSceneProxy](FRHICommandList& RHICmdList)
+		{
+			// Only access the irradiance maps on the RT, even though they belong to the USkyLightComponent, 
+			// Because FScene::UpdateSkyCaptureContents does not block the RT so the writes could still be in flight
+			LightSceneProxy->Initialize(InBlendFraction, InIrradianceEnvironmentMap, BlendDestinationIrradianceEnvironmentMap, InAverageBrightness, BlendDestinationAverageBrightness);
+		});
 }
 
 USkyLightComponent::USkyLightComponent(const FObjectInitializer& ObjectInitializer)
@@ -332,17 +332,17 @@ void USkyLightComponent::UpdateLimitedRenderingStateFast()
 {
 	if (SceneProxy)
 	{
-		ENQUEUE_UNIQUE_RENDER_COMMAND_FOURPARAMETER(
-			FFastUpdateSkyLightCommand,
-			FSkyLightSceneProxy*,LightSceneProxy,SceneProxy,
-			FLinearColor,LightColor,FLinearColor(LightColor) * Intensity,
-			float,IndirectLightingIntensity,IndirectLightingIntensity,
-			float,VolumetricScatteringIntensity,VolumetricScatteringIntensity,
-		{
-			LightSceneProxy->LightColor = LightColor;
-			LightSceneProxy->IndirectLightingIntensity = IndirectLightingIntensity;
-			LightSceneProxy->VolumetricScatteringIntensity = VolumetricScatteringIntensity;
-		});
+		FSkyLightSceneProxy* LightSceneProxy = SceneProxy;
+		FLinearColor InLightColor = FLinearColor(LightColor) * Intensity;
+		float InIndirectLightingIntensity = IndirectLightingIntensity;
+		float InVolumetricScatteringIntensity = VolumetricScatteringIntensity;
+		ENQUEUE_RENDER_COMMAND(FFastUpdateSkyLightCommand)(
+			[LightSceneProxy, InLightColor, InIndirectLightingIntensity, InVolumetricScatteringIntensity](FRHICommandList& RHICmdList)
+			{
+				LightSceneProxy->LightColor = InLightColor;
+				LightSceneProxy->IndirectLightingIntensity = InIndirectLightingIntensity;
+				LightSceneProxy->VolumetricScatteringIntensity = InVolumetricScatteringIntensity;
+			});
 	}
 }
 
@@ -350,19 +350,19 @@ void USkyLightComponent::UpdateOcclusionRenderingStateFast()
 {
 	if (SceneProxy && IsOcclusionSupported())
 	{
-		ENQUEUE_UNIQUE_RENDER_COMMAND_FIVEPARAMETER(
-			FFastUpdateSkyLightOcclusionCommand,
-			FSkyLightSceneProxy*,LightSceneProxy,SceneProxy,
-			float, Contrast, Contrast,
-			float, OcclusionExponent, OcclusionExponent,
-			float, MinOcclusion, MinOcclusion,
-			FColor, OcclusionTint, OcclusionTint,
-		{
-			LightSceneProxy->Contrast = Contrast;
-			LightSceneProxy->OcclusionExponent = OcclusionExponent;
-			LightSceneProxy->MinOcclusion = MinOcclusion;
-			LightSceneProxy->OcclusionTint = OcclusionTint;
-		});
+		FSkyLightSceneProxy* InLightSceneProxy = SceneProxy;
+		float InContrast = Contrast;
+		float InOcclusionExponent = OcclusionExponent;
+		float InMinOcclusion = MinOcclusion;
+		FColor InOcclusionTint = OcclusionTint;
+		ENQUEUE_RENDER_COMMAND(FFastUpdateSkyLightOcclusionCommand)(
+			[InLightSceneProxy, InContrast, InOcclusionExponent, InMinOcclusion, InOcclusionTint](FRHICommandList& RHICmdList)
+			{
+				InLightSceneProxy->Contrast = InContrast;
+				InLightSceneProxy->OcclusionExponent = InOcclusionExponent;
+				InLightSceneProxy->MinOcclusion = InMinOcclusion;
+				InLightSceneProxy->OcclusionTint = InOcclusionTint;
+			});
 	}
 
 }
@@ -772,19 +772,18 @@ void USkyLightComponent::SetCubemapBlend(UTextureCube* SourceCubemap, UTextureCu
 
 			if (SceneProxy)
 			{
-				ENQUEUE_UNIQUE_RENDER_COMMAND_SIXPARAMETER(
-					FUpdateSkyProxy,
-					const FSHVectorRGB3*,InIrradianceEnvironmentMap,&IrradianceEnvironmentMap,
-					const FSHVectorRGB3*,BlendDestinationIrradianceEnvironmentMap,&BlendDestinationIrradianceEnvironmentMap,
-					const float*,InAverageBrightness,&AverageBrightness,
-					const float*,BlendDestinationAverageBrightness,&BlendDestinationAverageBrightness,
-					float,InBlendFraction,BlendFraction,
-					FSkyLightSceneProxy*,LightSceneProxy,SceneProxy,
-				{
-					// Only access the irradiance maps on the RT, even though they belong to the USkyLightComponent, 
-					// Because FScene::UpdateSkyCaptureContents does not block the RT so the writes could still be in flight
-					LightSceneProxy->Initialize(InBlendFraction, InIrradianceEnvironmentMap, BlendDestinationIrradianceEnvironmentMap, InAverageBrightness, BlendDestinationAverageBrightness);
-				});
+				const FSHVectorRGB3* InIrradianceEnvironmentMap = &IrradianceEnvironmentMap;
+				const FSHVectorRGB3* InBlendDestinationIrradianceEnvironmentMap = &BlendDestinationIrradianceEnvironmentMap;
+				const float* InAverageBrightness = &AverageBrightness;
+				const float* InBlendDestinationAverageBrightness = &BlendDestinationAverageBrightness;
+				FSkyLightSceneProxy* LightSceneProxy = SceneProxy;
+				ENQUEUE_RENDER_COMMAND(FUpdateSkyProxy)(
+					[InIrradianceEnvironmentMap, InBlendDestinationIrradianceEnvironmentMap, InAverageBrightness, InBlendDestinationAverageBrightness, InBlendFraction, LightSceneProxy](FRHICommandList& RHICmdList)
+					{
+						// Only access the irradiance maps on the RT, even though they belong to the USkyLightComponent, 
+						// Because FScene::UpdateSkyCaptureContents does not block the RT so the writes could still be in flight
+						LightSceneProxy->Initialize(InBlendFraction, InIrradianceEnvironmentMap, InBlendDestinationIrradianceEnvironmentMap, InAverageBrightness, InBlendDestinationAverageBrightness);
+					});
 			}
 		}
 	}
diff --git a/Engine/Source/Runtime/RenderCore/Public/RenderingThread.h b/Engine/Source/Runtime/RenderCore/Public/RenderingThread.h
--- a/Engine/Source/Runtime/RenderCore/Public/RenderingThread.h
+++ b/Engine/Source/Runtime/RenderCore/Public/RenderingThread.h
@@ -543,7 +543,8 @@ FORCEINLINE_DEBUGGABLE void EnqueueUniqueRenderCommand(LAMBDA& Lambda)
 
 #define ENQUEUE_UNIQUE_RENDER_COMMAND_FIVEPARAMETER(TypeName,ParamType1,ParamName1,ParamValue1,ParamType2,ParamName2,ParamValue2,ParamType3,ParamName3,ParamValue3,ParamType4,ParamName4,ParamValue4,ParamType5,ParamName5,ParamValue5,Code) \
 	ENQUEUE_UNIQUE_RENDER_COMMAND_FIVEPARAMETER_DECLARE(TypeName,ParamType1,ParamName1,ParamValue1,ParamType2,ParamName2,ParamValue2,ParamType3,ParamName3,ParamValue3,ParamType4,ParamName4,ParamValue4,ParamType5,ParamName5,ParamValue5,Code) \
-	ENQUEUE_UNIQUE_RENDER_COMMAND_FIVEPARAMETER_CREATE(TypeName,ParamType1,ParamValue1,ParamType2,ParamValue2,ParamType3,ParamValue3,ParamType4,ParamValue4,ParamType5,ParamValue5)
+	ENQUEUE_UNIQUE_RENDER_COMMAND_FIVEPARAMETER_CREATE(TypeName,ParamType1,ParamValue1,ParamType2,ParamValue2,ParamType3,ParamValue3,ParamType4,ParamValue4,ParamType5,ParamValue5) \
+	DEPRECATED_MACRO(4.22, "Please use ENQUEUE_RENDER_COMMAND instead")
 
 #define ENQUEUE_UNIQUE_RENDER_COMMAND_FIVEPARAMETER_DECLARE_TEMPLATE(TypeName,TemplateParamName,ParamType1,ParamName1,ParamValue1,ParamType2,ParamName2,ParamValue2,ParamType3,ParamName3,ParamValue3,ParamType4,ParamName4,ParamValue4,ParamType5,ParamName5,ParamValue5,Code) \
 	template <typename TemplateParamName> \
@@ -599,7 +600,8 @@ FORCEINLINE_DEBUGGABLE void EnqueueUniqueRenderCommand(LAMBDA& Lambda)
 
 #define ENQUEUE_UNIQUE_RENDER_COMMAND_SIXPARAMETER(TypeName,ParamType1,ParamName1,ParamValue1,ParamType2,ParamName2,ParamValue2,ParamType3,ParamName3,ParamValue3,ParamType4,ParamName4,ParamValue4,ParamType5,ParamName5,ParamValue5,ParamType6,ParamName6,ParamValue6,Code) \
 	ENQUEUE_UNIQUE_RENDER_COMMAND_SIXPARAMETER_DECLARE(TypeName,ParamType1,ParamName1,ParamValue1,ParamType2,ParamName2,ParamValue2,ParamType3,ParamName3,ParamValue3,ParamType4,ParamName4,ParamValue4,ParamType5,ParamName5,ParamValue5,ParamType6,ParamName6,ParamValue6,Code) \
-	ENQUEUE_UNIQUE_RENDER_COMMAND_SIXPARAMETER_CREATE(TypeName,ParamType1,ParamValue1,ParamType2,ParamValue2,ParamType3,ParamValue3,ParamType4,ParamValue4,ParamType5,ParamValue5,ParamType6,ParamValue6)
+	ENQUEUE_UNIQUE_RENDER_COMMAND_SIXPARAMETER_CREATE(TypeName,ParamType1,ParamValue1,ParamType2,ParamValue2,ParamType3,ParamValue3,ParamType4,ParamValue4,ParamType5,ParamValue5,ParamType6,ParamValue6) \
+	DEPRECATED_MACRO(4.22, "Please use ENQUEUE_RENDER_COMMAND instead")
 
 #define ENQUEUE_UNIQUE_RENDER_COMMAND_SIXPARAMETER_DECLARE_TEMPLATE(TypeName,TemplateParamName,ParamType1,ParamName1,ParamValue1,ParamType2,ParamName2,ParamValue2,ParamType3,ParamName3,ParamValue3,ParamType4,ParamName4,ParamValue4,ParamType5,ParamName5,ParamValue5,ParamType6,ParamName6,ParamValue6,Code) \
 	template <typename TemplateParamName> \
diff --git a/Engine/Source/Runtime/Renderer/Private/PlanarReflectionRendering.cpp b/Engine/Source/Runtime/Renderer/Private/PlanarReflectionRendering.cpp
--- a/Engine/Source/Runtime/Renderer/Private/PlanarReflectionRendering.cpp
+++ b/Engine/Source/Runtime/Renderer/Private/PlanarReflectionRendering.cpp
@@ -581,19 +581,17 @@ void FScene::UpdatePlanarReflectionContents(UPlanarReflectionComponent* CaptureC
 				});
 		}
 
-		const FName OwnerName = CaptureComponent->GetOwner() ? CaptureComponent->GetOwner()->GetFName() : NAME_None;
-
-		ENQUEUE_UNIQUE_RENDER_COMMAND_SIXPARAMETER( 
-			CaptureCommand,
-			FSceneRenderer*, MainSceneRenderer, &MainSceneRenderer,
-			FSceneRenderer*, SceneRenderer, SceneRenderer,
-			FPlanarReflectionSceneProxy*, SceneProxy, CaptureComponent->SceneProxy,
-			FPlanarReflectionRenderTarget*, RenderTarget, CaptureComponent->RenderTarget,
-			FPlane, MirrorPlane, MirrorPlane,
-			FName, OwnerName, OwnerName,
 		{
-			UpdatePlanarReflectionContents_RenderThread(RHICmdList, MainSceneRenderer, SceneRenderer, SceneProxy, RenderTarget, RenderTarget, MirrorPlane, OwnerName, FResolveParams(), true);
-		});
+			const FName OwnerName = CaptureComponent->GetOwner() ? CaptureComponent->GetOwner()->GetFName() : NAME_None;
+			FSceneRenderer* MainSceneRendererPtr = &MainSceneRenderer;
+			FPlanarReflectionSceneProxy* SceneProxyPtr = CaptureComponent->SceneProxy;
+			FPlanarReflectionRenderTarget* RenderTargetPtr = CaptureComponent->RenderTarget;
+			ENQUEUE_RENDER_COMMAND(CaptureCommand)(
+				[SceneRenderer, MirrorPlane, OwnerName, MainSceneRendererPtr, SceneProxyPtr, RenderTargetPtr](FRHICommandListImmediate& RHICmdList)
+			{
+				UpdatePlanarReflectionContents_RenderThread(RHICmdList, MainSceneRendererPtr, SceneRenderer, SceneProxyPtr, RenderTargetPtr, RenderTargetPtr, MirrorPlane, OwnerName, FResolveParams(), true);
+			});
+		}
 	}
 }
 
diff --git a/Engine/Source/Runtime/Renderer/Private/ReflectionEnvironmentCapture.cpp b/Engine/Source/Runtime/Renderer/Private/ReflectionEnvironmentCapture.cpp
--- a/Engine/Source/Runtime/Renderer/Private/ReflectionEnvironmentCapture.cpp
+++ b/Engine/Source/Runtime/Renderer/Private/ReflectionEnvironmentCapture.cpp
@@ -1363,22 +1363,16 @@ void CaptureSceneIntoScratchCubemap(
 
 		FSceneRenderer* SceneRenderer = FSceneRenderer::CreateSceneRenderer(&ViewFamily, NULL);
 
-		ENQUEUE_UNIQUE_RENDER_COMMAND_SIXPARAMETER( 
-			CaptureCommand,
-			FSceneRenderer*, SceneRenderer, SceneRenderer,
-			ECubeFace, CubeFace, (ECubeFace)CubeFace,
-			int32, CubemapSize, CubemapSize,
-			bool, bCapturingForSkyLight, bCapturingForSkyLight,
-			bool, bLowerHemisphereIsBlack, bLowerHemisphereIsBlack,
-			FLinearColor, LowerHemisphereColor, LowerHemisphereColor,
-		{
-			CaptureSceneToScratchCubemap(RHICmdList, SceneRenderer, CubeFace, CubemapSize, bCapturingForSkyLight, bLowerHemisphereIsBlack, LowerHemisphereColor);
-
-			if( !bCapturingForSkyLight )
+		ENQUEUE_RENDER_COMMAND(CaptureCommand)(
+			[SceneRenderer, CubeFace, CubemapSize, bCapturingForSkyLight, bLowerHemisphereIsBlack, LowerHemisphereColor](FRHICommandListImmediate& RHICmdList)
 			{
-				RHICmdList.EndFrame();
-			}
-		});
+				CaptureSceneToScratchCubemap(RHICmdList, SceneRenderer, (ECubeFace)CubeFace, CubemapSize, bCapturingForSkyLight, bLowerHemisphereIsBlack, LowerHemisphereColor);
+
+				if (!bCapturingForSkyLight)
+				{
+					RHICmdList.EndFrame();
+				}
+			});
 	}
 }
 
@@ -1468,12 +1462,11 @@ void FScene::CaptureOrUploadReflectionCapture(UReflectionCaptureComponent* Captu
 
 			const int32 ReflectionCaptureSize = UReflectionCaptureComponent::GetReflectionCaptureSize();
 
-			ENQUEUE_UNIQUE_RENDER_COMMAND_ONEPARAMETER( 
-				ClearCommand,
-				int32, ReflectionCaptureSize, ReflectionCaptureSize,
-			{
-				ClearScratchCubemaps(RHICmdList, ReflectionCaptureSize);
-			});
+			ENQUEUE_RENDER_COMMAND(ClearCommand)(
+				[ReflectionCaptureSize](FRHICommandListImmediate& RHICmdList)
+				{
+					ClearScratchCubemaps(RHICmdList, ReflectionCaptureSize);
+				});
 			
 			if (CaptureComponent->ReflectionSourceType == EReflectionSourceType::CapturedScene)
 			{
@@ -1482,14 +1475,13 @@ void FScene::CaptureOrUploadReflectionCapture(UReflectionCaptureComponent* Captu
 			}
 			else if (CaptureComponent->ReflectionSourceType == EReflectionSourceType::SpecifiedCubemap)
 			{
-				ENQUEUE_UNIQUE_RENDER_COMMAND_FOURPARAMETER(
-					CopyCubemapCommand,
-					UTextureCube*, SourceTexture, CaptureComponent->Cubemap,
-					int32, ReflectionCaptureSize, ReflectionCaptureSize,
-					float, SourceCubemapRotation, CaptureComponent->SourceCubemapAngle * (PI / 180.f),
-					ERHIFeatureLevel::Type, FeatureLevel, GetFeatureLevel(),
+				UTextureCube* SourceTexture = CaptureComponent->Cubemap;
+				float SourceCubemapRotation = CaptureComponent->SourceCubemapAngle * (PI / 180.f);
+				ERHIFeatureLevel::Type InFeatureLevel = FeatureLevel;
+				ENQUEUE_RENDER_COMMAND(CopyCubemapCommand)(
+					[SourceTexture, ReflectionCaptureSize, SourceCubemapRotation, InFeatureLevel](FRHICommandList& RHICmdList)
 					{
-						CopyCubemapToScratchCubemap(RHICmdList, FeatureLevel, SourceTexture, ReflectionCaptureSize, false, false, SourceCubemapRotation, FLinearColor());
+						CopyCubemapToScratchCubemap(RHICmdList, InFeatureLevel, SourceTexture, ReflectionCaptureSize, false, false, SourceCubemapRotation, FLinearColor());
 					});
 			}
 			else
@@ -1604,12 +1596,14 @@ void FScene::UpdateSkyCaptureContents(
 				World->SendAllEndOfFrameUpdates();
 			}
 		}
-		ENQUEUE_UNIQUE_RENDER_COMMAND_ONEPARAMETER(
-			ClearCommand,
-			int32, CubemapSize, CaptureComponent->CubemapResolution,
 		{
-			ClearScratchCubemaps(RHICmdList, CubemapSize);
-		});
+			int32 CubemapSize = CaptureComponent->CubemapResolution;
+			ENQUEUE_RENDER_COMMAND(ClearCommand)(
+				[CubemapSize](FRHICommandList& RHICmdList)
+				{
+					ClearScratchCubemaps(RHICmdList, CubemapSize);
+				});
+		}
 
 		if (CaptureComponent->SourceType == SLS_CapturedScene)
 		{
@@ -1618,17 +1612,16 @@ void FScene::UpdateSkyCaptureContents(
 		}
 		else if (CaptureComponent->SourceType == SLS_SpecifiedCubemap)
 		{
-			ENQUEUE_UNIQUE_RENDER_COMMAND_SIXPARAMETER( 
-				CopyCubemapCommand,
-				UTextureCube*, SourceTexture, SourceCubemap,
-				int32, CubemapSize, CaptureComponent->CubemapResolution,
-				bool, bLowerHemisphereIsBlack, CaptureComponent->bLowerHemisphereIsBlack,
-				float, SourceCubemapRotation, CaptureComponent->SourceCubemapAngle * (PI / 180.f),
-				ERHIFeatureLevel::Type, FeatureLevel, GetFeatureLevel(),
-				FLinearColor, LowerHemisphereColor, CaptureComponent->LowerHemisphereColor,
-			{
-				CopyCubemapToScratchCubemap(RHICmdList, FeatureLevel, SourceTexture, CubemapSize, true, bLowerHemisphereIsBlack, SourceCubemapRotation, LowerHemisphereColor);
-			});
+			int32 CubemapSize = CaptureComponent->CubemapResolution;
+			bool bLowerHemisphereIsBlack = CaptureComponent->bLowerHemisphereIsBlack;
+			float SourceCubemapRotation = CaptureComponent->SourceCubemapAngle * (PI / 180.f);
+			ERHIFeatureLevel::Type InnerFeatureLevel = FeatureLevel;
+			FLinearColor LowerHemisphereColor = CaptureComponent->LowerHemisphereColor;
+			ENQUEUE_RENDER_COMMAND(CopyCubemapCommand)(
+				[SourceCubemap, CubemapSize, bLowerHemisphereIsBlack, SourceCubemapRotation, InnerFeatureLevel, LowerHemisphereColor](FRHICommandList& RHICmdList)
+				{
+					CopyCubemapToScratchCubemap(RHICmdList, InnerFeatureLevel, SourceCubemap, CubemapSize, true, bLowerHemisphereIsBlack, SourceCubemapRotation, LowerHemisphereColor);
+				});
 		}
 		else
 		{
@@ -1637,13 +1630,12 @@ void FScene::UpdateSkyCaptureContents(
 
 		if (OutRadianceMap)
 		{
-			ENQUEUE_UNIQUE_RENDER_COMMAND_TWOPARAMETER( 
-				ReadbackCommand,
-				int32, CubemapSize, CaptureComponent->CubemapResolution,
-				TArray<FFloat16Color>*, RadianceMap, OutRadianceMap,
-			{
-				ReadbackRadianceMap(RHICmdList, CubemapSize, *RadianceMap);
-			});
+			int32 CubemapSize = CaptureComponent->CubemapResolution;
+			ENQUEUE_RENDER_COMMAND(ReadbackCommand)(
+				[CubemapSize, OutRadianceMap](FRHICommandListImmediate& RHICmdList)
+				{
+					ReadbackRadianceMap(RHICmdList, CubemapSize, *OutRadianceMap);
+				});
 		}
 
 		ENQUEUE_UNIQUE_RENDER_COMMAND_FOURPARAMETER( 
