DR - Fix static analysis warning
#rnx

diff --git a/Engine/Source/Runtime/Engine/Private/Curves/RichCurve.cpp b/Engine/Source/Runtime/Engine/Private/Curves/RichCurve.cpp
--- a/Engine/Source/Runtime/Engine/Private/Curves/RichCurve.cpp
+++ b/Engine/Source/Runtime/Engine/Private/Curves/RichCurve.cpp
@@ -1378,7 +1378,7 @@ void FRichCurve::CompressCurve(FCompressedRichCurve& OutCurve, float ErrorThresh
 
 			const int32 KeySize = sizeof(uint16);
 
-			if (KeyFormat == RCKTCF_uint16)
+			//if (KeyFormat == RCKTCF_uint16)
 			{
 				WriteOffset = Align(WriteOffset, sizeof(uint16));
 			}
