DR - Fix static analysis warning
#rnx

diff --git a/Engine/Source/Runtime/Engine/Private/SkeletalRenderGPUSkin.cpp b/Engine/Source/Runtime/Engine/Private/SkeletalRenderGPUSkin.cpp
--- a/Engine/Source/Runtime/Engine/Private/SkeletalRenderGPUSkin.cpp
+++ b/Engine/Source/Runtime/Engine/Private/SkeletalRenderGPUSkin.cpp
@@ -293,6 +293,7 @@ void FSkeletalMeshObjectGPUSkin::UpdateDynamicData_RenderThread(FGPUSkinCache* G
 {
 	SCOPE_CYCLE_COUNTER(STAT_GPUSkinUpdateRTTime);
 	check(InDynamicData != nullptr);
+	CA_ASSUME(InDynamicData != nullptr);
 	bool bMorphNeedsUpdate=false;
 	// figure out if the morphing vertex buffer needs to be updated. compare old vs new active morphs
 	bMorphNeedsUpdate = 
@@ -302,8 +303,8 @@ void FSkeletalMeshObjectGPUSkin::UpdateDynamicData_RenderThread(FGPUSkinCache* G
 		: true);
 
 #if RHI_RAYTRACING
-	bRequireRecreatingRayTracingGeometry = (DynamicData == nullptr && InDynamicData != nullptr || // Newly created
-		(DynamicData != nullptr && InDynamicData != nullptr && DynamicData->LODIndex != InDynamicData->LODIndex)); // LOD level changed
+	bRequireRecreatingRayTracingGeometry = (DynamicData == nullptr || // Newly created
+		(DynamicData != nullptr && DynamicData->LODIndex != InDynamicData->LODIndex)); // LOD level changed
 #endif
 
 	WaitForRHIThreadFenceForDynamicData();
