Fix non RHI_RAYTRACING build

diff --git a/Engine/Source/Runtime/Engine/Private/GeometryCacheVertexFactory.cpp b/Engine/Source/Runtime/Engine/Private/GeometryCacheVertexFactory.cpp
--- a/Engine/Source/Runtime/Engine/Private/GeometryCacheVertexFactory.cpp
+++ b/Engine/Source/Runtime/Engine/Private/GeometryCacheVertexFactory.cpp
@@ -340,7 +340,16 @@ void FGeometryCacheVertexVertexFactory::CreateManualVertexFetchUniformBuffer(
 
 FVertexFactoryShaderParameters* FGeometryCacheVertexVertexFactory::ConstructShaderParameters(EShaderFrequency ShaderFrequency)
 {
-	return (ShaderFrequency == SF_Vertex || ShaderFrequency == SF_RayHitGroup) ? new FGeometryCacheVertexFactoryShaderParameters() : nullptr;
+	switch (ShaderFrequency)
+	{
+		case SF_Vertex:
+#if RHI_RAYTRACING
+		case SF_RayHitGroup:
+#endif
+			return new FGeometryCacheVertexFactoryShaderParameters();
+		default:
+			return nullptr;
+	}
 }
 
 bool FGeometryCacheVertexVertexFactory::ShouldCompilePermutation(EShaderPlatform Platform, const class FMaterial* Material, const class FShaderType* ShaderType)
