Renamed FDummyVertexBuffer to FDefaultGeometryCacheVertexBuffer, as global DummyPositionBuffer was conflicting with a different FDummyVertexBuffer from a vector visualization.

diff --git a/Engine/Source/Runtime/Engine/Private/GeometryCacheVertexFactory.cpp b/Engine/Source/Runtime/Engine/Private/GeometryCacheVertexFactory.cpp
--- a/Engine/Source/Runtime/Engine/Private/GeometryCacheVertexFactory.cpp
+++ b/Engine/Source/Runtime/Engine/Private/GeometryCacheVertexFactory.cpp
@@ -122,7 +122,7 @@ void FGeometryCacheVertexVertexFactory::SetData(const FDataType& InData)
 	UpdateRHI();
 }
 
-class FDummyVertexBuffer : public FVertexBuffer
+class FDefaultGeometryCacheVertexBuffer : public FVertexBuffer
 {
 public:
 	FShaderResourceViewRHIRef SRV;
@@ -146,7 +146,7 @@ public:
 		FVertexBuffer::ReleaseRHI();
 	}
 };
-TGlobalResource<FDummyVertexBuffer> DummyPositionBuffer;
+TGlobalResource<FDefaultGeometryCacheVertexBuffer> GDefaultGeometryCacheVertexBuffer;
 
 class FDummyTangentBuffer : public FVertexBuffer
 {
@@ -275,7 +275,7 @@ void FGeometryCacheVertexVertexFactory::CreateManualVertexFetchUniformBuffer(
 	}
 	else
 	{
-		ManualVertexFetchParameters.Position = DummyPositionBuffer.SRV;
+		ManualVertexFetchParameters.Position = GDefaultGeometryCacheVertexBuffer.SRV;
 	}
 
 	if (Data.TangentBasisComponents[0].VertexBuffer != NULL)
@@ -320,7 +320,7 @@ void FGeometryCacheVertexVertexFactory::CreateManualVertexFetchUniformBuffer(
 	}
 	else
 	{
-		ManualVertexFetchParameters.MotionBlurData = DummyPositionBuffer.SRV;
+		ManualVertexFetchParameters.MotionBlurData = GDefaultGeometryCacheVertexBuffer.SRV;
 	}
 
 	if (Data.TextureCoordinates.Num())
@@ -332,7 +332,7 @@ void FGeometryCacheVertexVertexFactory::CreateManualVertexFetchUniformBuffer(
 	}
 	else
 	{
-		ManualVertexFetchParameters.TexCoords = DummyPositionBuffer.SRV;
+		ManualVertexFetchParameters.TexCoords = GDefaultGeometryCacheVertexBuffer.SRV;
 	}
 
 	OutUserData.ManualVertexFetchUniformBuffer = FGeometryCacheManualVertexFetchUniformBufferParametersRef::CreateUniformBufferImmediate(ManualVertexFetchParameters, UniformBuffer_SingleFrame);
