DR - More removal of ENQUEUE_UNIQUE_RENDER_COMMAND_ONEPARAMETER
#rnx

diff --git a/Engine/Source/Editor/MovieSceneTools/Private/TrackEditorThumbnail/TrackEditorThumbnail.cpp b/Engine/Source/Editor/MovieSceneTools/Private/TrackEditorThumbnail/TrackEditorThumbnail.cpp
--- a/Engine/Source/Editor/MovieSceneTools/Private/TrackEditorThumbnail/TrackEditorThumbnail.cpp
+++ b/Engine/Source/Editor/MovieSceneTools/Private/TrackEditorThumbnail/TrackEditorThumbnail.cpp
@@ -79,12 +79,13 @@ void FTrackEditorThumbnail::DestroyTexture()
 {
 	if (Texture)
 	{
-		ENQUEUE_UNIQUE_RENDER_COMMAND_ONEPARAMETER(DestroyTexture,
-			FSlateTexture2DRHIRef*, Texture, Texture,
-		{
-			Texture->ReleaseResource();
-			delete Texture;
-		});
+		FSlateTexture2DRHIRef* InTexture = Texture;
+		ENQUEUE_RENDER_COMMAND(DestroyTexture)(
+			[InTexture](FRHICommandList& RHICmdList)
+			{
+				InTexture->ReleaseResource();
+				delete InTexture;
+			});
 
 		Texture = nullptr;
 	}
@@ -191,12 +192,12 @@ void FTrackEditorThumbnail::CopyTextureIn(FTexture2DRHIRef SourceTexture)
 		*bHasFinishedDrawingPtr = true;
 	};
 
-	ENQUEUE_UNIQUE_RENDER_COMMAND_ONEPARAMETER(
-		ResolveCaptureFrameTexture,
-		TFunction<void(FRHICommandListImmediate&)>, InRenderCommand, RenderCommand,
-	{
-		InRenderCommand(RHICmdList);
-	});
+	TFunction<void(FRHICommandListImmediate&)> InRenderCommand = RenderCommand;
+	ENQUEUE_RENDER_COMMAND(ResolveCaptureFrameTexture)(
+		[InRenderCommand](FRHICommandListImmediate& RHICmdList)
+		{
+			InRenderCommand(RHICmdList);
+		});
 }
 
 
diff --git a/Engine/Source/Editor/WorldBrowser/Private/Tiles/WorldTileThumbnails.cpp b/Engine/Source/Editor/WorldBrowser/Private/Tiles/WorldTileThumbnails.cpp
--- a/Engine/Source/Editor/WorldBrowser/Private/Tiles/WorldTileThumbnails.cpp
+++ b/Engine/Source/Editor/WorldBrowser/Private/Tiles/WorldTileThumbnails.cpp
@@ -217,13 +217,12 @@ void FTileAtlasPage::UpdateSlotImageData(int32 SlotIdx, FSlateTextureDataPtr Ima
 			UpdateRegion
 		};
 			
-		ENQUEUE_UNIQUE_RENDER_COMMAND_ONEPARAMETER(
-			UpdateSlotImageData,
-			FSlotUpdateContext,	Context, Context,
-		{
-			FRHITexture2D* RHITexture2D = (FRHITexture2D*)Context.TextureRHI.GetReference();
-			RHIUpdateTexture2D(RHITexture2D, 0, Context.Region, Context.SourcePitch, Context.ImageData->GetRawBytesPtr());
-		});
+		ENQUEUE_RENDER_COMMAND(UpdateSlotImageData)(
+			[Context](FRHICommandList& RHICmdList)
+			{
+				FRHITexture2D* RHITexture2D = (FRHITexture2D*)Context.TextureRHI.GetReference();
+				RHIUpdateTexture2D(RHITexture2D, 0, Context.Region, Context.SourcePitch, Context.ImageData->GetRawBytesPtr());
+			});
 	}
 }
 
diff --git a/Engine/Source/Runtime/Engine/Private/Components/SceneCaptureComponent.cpp b/Engine/Source/Runtime/Engine/Private/Components/SceneCaptureComponent.cpp
--- a/Engine/Source/Runtime/Engine/Private/Components/SceneCaptureComponent.cpp
+++ b/Engine/Source/Runtime/Engine/Private/Components/SceneCaptureComponent.cpp
@@ -851,14 +851,14 @@ void UPlanarReflectionComponent::DestroyRenderState_Concurrent()
 	{
 		GetWorld()->Scene->RemovePlanarReflection(this);
 
-		ENQUEUE_UNIQUE_RENDER_COMMAND_ONEPARAMETER(
-			FDestroyPlanarReflectionCommand,
-			FPlanarReflectionSceneProxy*,SceneProxy,SceneProxy,
-		{
-			delete SceneProxy;
-		});
+		FPlanarReflectionSceneProxy* InSceneProxy = SceneProxy;
+		ENQUEUE_RENDER_COMMAND(FDestroyPlanarReflectionCommand)(
+			[InSceneProxy](FRHICommandList& RHICmdList)
+			{
+				delete InSceneProxy;
+			});
 
-		SceneProxy = NULL;
+		SceneProxy = nullptr;
 	}
 }
 
diff --git a/Engine/Source/Runtime/Engine/Private/Components/SkyLightComponent.cpp b/Engine/Source/Runtime/Engine/Private/Components/SkyLightComponent.cpp
--- a/Engine/Source/Runtime/Engine/Private/Components/SkyLightComponent.cpp
+++ b/Engine/Source/Runtime/Engine/Private/Components/SkyLightComponent.cpp
@@ -401,14 +401,14 @@ void USkyLightComponent::DestroyRenderState_Concurrent()
 	{
 		GetWorld()->Scene->DisableSkyLight(SceneProxy);
 
-		ENQUEUE_UNIQUE_RENDER_COMMAND_ONEPARAMETER(
-			FDestroySkyLightCommand,
-			FSkyLightSceneProxy*,LightSceneProxy,SceneProxy,
-		{
-			delete LightSceneProxy;
-		});
+		FSkyLightSceneProxy* LightSceneProxy = SceneProxy;
+		ENQUEUE_RENDER_COMMAND(FDestroySkyLightCommand)(
+			[LightSceneProxy](FRHICommandList& RHICmdList)
+			{
+				delete LightSceneProxy;
+			});
 
-		SceneProxy = NULL;
+		SceneProxy = nullptr;
 	}
 }
 
@@ -736,10 +736,11 @@ void USkyLightComponent::SetCubemap(UTextureCube* NewCubemap)
 		SetCaptureIsDirty();
 
 #if RHI_RAYTRACING
-		if (SceneProxy) {
-			ENQUEUE_UNIQUE_RENDER_COMMAND_ONEPARAMETER(
-				SetLightProxyDirty,
-				FSkyLightSceneProxy*, LightSceneProxy, SceneProxy,
+		if (SceneProxy)
+		{
+			FSkyLightSceneProxy* LightSceneProxy = SceneProxy;
+			ENQUEUE_RENDER_COMMAND(SetLightProxyDirty)(
+				[LightSceneProxy](FRHICommandList& RHICmdList)
 				{
 					LightSceneProxy->IsDirtyImportanceSamplingData = true;
 				});
diff --git a/Engine/Source/Runtime/Engine/Private/Components/StaticMeshComponent.cpp b/Engine/Source/Runtime/Engine/Private/Components/StaticMeshComponent.cpp
--- a/Engine/Source/Runtime/Engine/Private/Components/StaticMeshComponent.cpp
+++ b/Engine/Source/Runtime/Engine/Private/Components/StaticMeshComponent.cpp
@@ -1564,9 +1564,9 @@ void UStaticMeshComponent::PostEditUndo()
 	// Debug check command trying to track down undo related uninitialized resource
 	if (GetStaticMesh() != NULL && GetStaticMesh()->RenderData.IsValid() && GetStaticMesh()->RenderData->LODResources.Num() > 0)
 	{
-		ENQUEUE_UNIQUE_RENDER_COMMAND_ONEPARAMETER(
-			ResourceCheckCommand,
-			FRenderResource*, Resource, &GetStaticMesh()->RenderData->LODResources[0].IndexBuffer,
+		FRenderResource* Resource = &GetStaticMesh()->RenderData->LODResources[0].IndexBuffer;
+		ENQUEUE_RENDER_COMMAND(ResourceCheckCommand)(
+			[Resource](FRHICommandList& RHICmdList)
 			{
 				check( Resource->IsInitialized() );
 			}
diff --git a/Engine/Source/Runtime/Engine/Private/DebugViewModeMaterialProxy.cpp b/Engine/Source/Runtime/Engine/Private/DebugViewModeMaterialProxy.cpp
--- a/Engine/Source/Runtime/Engine/Private/DebugViewModeMaterialProxy.cpp
+++ b/Engine/Source/Runtime/Engine/Private/DebugViewModeMaterialProxy.cpp
@@ -163,17 +163,16 @@ void FDebugViewModeMaterialProxy::ValidateAllShaders(TSet<UMaterialInterface*>&
 		}
 	}
 
-	ENQUEUE_UNIQUE_RENDER_COMMAND_ONEPARAMETER(
-		UpdateDebugMaterialExpressionCache,
-		TArray<FDebugViewModeMaterialProxy*>, MaterialsToUpdateCopy, MaterialsToUpdate,
+	TArray<FDebugViewModeMaterialProxy*> MaterialsToUpdateCopy = MaterialsToUpdate;
+	ENQUEUE_RENDER_COMMAND(UpdateDebugMaterialExpressionCache)(
+		[MaterialsToUpdateCopy](FRHICommandList& RHICmdList)
 		{
 			for (FDebugViewModeMaterialProxy* MaterialToUpdate : MaterialsToUpdateCopy)
 			{
 				check(MaterialToUpdate);
 				MaterialToUpdate->UpdateUniformExpressionCacheIfNeeded(MaterialToUpdate->FMaterial::GetFeatureLevel());
 			}
-		}
-	)
+		});
 
 	FlushRenderingCommands();
 }
diff --git a/Engine/Source/Runtime/Engine/Private/DistanceFieldAtlas.cpp b/Engine/Source/Runtime/Engine/Private/DistanceFieldAtlas.cpp
--- a/Engine/Source/Runtime/Engine/Private/DistanceFieldAtlas.cpp
+++ b/Engine/Source/Runtime/Engine/Private/DistanceFieldAtlas.cpp
@@ -538,13 +538,12 @@ void FDistanceFieldVolumeTexture::Initialize(UStaticMesh* InStaticMesh)
 
 		bReferencedByAtlas = true;
 
-		ENQUEUE_UNIQUE_RENDER_COMMAND_ONEPARAMETER(
-			AddAllocation,
-			FDistanceFieldVolumeTexture*, DistanceFieldVolumeTexture, this,
+		FDistanceFieldVolumeTexture* DistanceFieldVolumeTexture = this;
+		ENQUEUE_RENDER_COMMAND(AddAllocation)(
+			[DistanceFieldVolumeTexture](FRHICommandList& RHICmdList)
 			{
 				GDistanceFieldVolumeTextureAtlas.AddAllocation(DistanceFieldVolumeTexture);
-			}
-		);
+			});
 	}
 }
 
@@ -556,13 +555,12 @@ void FDistanceFieldVolumeTexture::Release()
 
 		bReferencedByAtlas = false;
 
-		ENQUEUE_UNIQUE_RENDER_COMMAND_ONEPARAMETER(
-			ReleaseAllocation,
-			FDistanceFieldVolumeTexture*, DistanceFieldVolumeTexture, this,
+		FDistanceFieldVolumeTexture* DistanceFieldVolumeTexture = this;
+		ENQUEUE_RENDER_COMMAND(ReleaseAllocation)(
+			[DistanceFieldVolumeTexture](FRHICommandList& RHICmdList)
 			{
 				GDistanceFieldVolumeTextureAtlas.RemoveAllocation(DistanceFieldVolumeTexture);
-			}
-		);
+			});
 	}
 }
 
diff --git a/Engine/Source/Runtime/Engine/Private/InstancedStaticMesh.cpp b/Engine/Source/Runtime/Engine/Private/InstancedStaticMesh.cpp
--- a/Engine/Source/Runtime/Engine/Private/InstancedStaticMesh.cpp
+++ b/Engine/Source/Runtime/Engine/Private/InstancedStaticMesh.cpp
@@ -1553,14 +1553,14 @@ void UInstancedStaticMeshComponent::ReleasePerInstanceRenderData()
 		FPerInstanceRenderDataPtr* CleanupRenderDataPtr = new FPerInstanceRenderDataPtr(PerInstanceRenderData);
 		PerInstanceRenderData.Reset();
 
-		ENQUEUE_UNIQUE_RENDER_COMMAND_ONEPARAMETER(
-			FReleasePerInstanceRenderData,
-			FPerInstanceRenderDataPtr*, InCleanupRenderDataPtr, CleanupRenderDataPtr,
-		{
-			// Destroy the shared pointer object we allocated on the heap.
-			// Resource will either be released here or by scene proxy on the render thread, whoever gets executed last
-			delete InCleanupRenderDataPtr;
-		});
+		FPerInstanceRenderDataPtr* InCleanupRenderDataPtr = CleanupRenderDataPtr;
+		ENQUEUE_RENDER_COMMAND(FReleasePerInstanceRenderData)(
+			[InCleanupRenderDataPtr](FRHICommandList& RHICmdList)
+			{
+				// Destroy the shared pointer object we allocated on the heap.
+				// Resource will either be released here or by scene proxy on the render thread, whoever gets executed last
+				delete InCleanupRenderDataPtr;
+			});
 	} //-V773
 }
 
diff --git a/Engine/Source/Runtime/Engine/Private/LevelTick.cpp b/Engine/Source/Runtime/Engine/Private/LevelTick.cpp
--- a/Engine/Source/Runtime/Engine/Private/LevelTick.cpp
+++ b/Engine/Source/Runtime/Engine/Private/LevelTick.cpp
@@ -943,16 +943,15 @@ FSendAllEndOfFrameUpdates* BeginSendEndOfFrameUpdatesDrawEvent(FGPUSkinCache* GP
 	SendAllEndOfFrameUpdates->GPUSkinCache = GPUSkinCache;
 
 #if WANTS_DRAW_MESH_EVENTS
-	ENQUEUE_UNIQUE_RENDER_COMMAND_ONEPARAMETER(
-		BeginDrawEventCommand,
-		FSendAllEndOfFrameUpdates*, SendAllEndOfFrameUpdates, SendAllEndOfFrameUpdates,
-	{
-		BEGIN_DRAW_EVENTF(
-			RHICmdList, 
-			SendAllEndOfFrameUpdates, 
-			SendAllEndOfFrameUpdates->DrawEvent,
-			TEXT("SendAllEndOfFrameUpdates"));
-	});
+	ENQUEUE_RENDER_COMMAND(BeginDrawEventCommand)(
+		[SendAllEndOfFrameUpdates](FRHICommandList& RHICmdList)
+		{
+			BEGIN_DRAW_EVENTF(
+				RHICmdList, 
+				SendAllEndOfFrameUpdates, 
+				SendAllEndOfFrameUpdates->DrawEvent,
+				TEXT("SendAllEndOfFrameUpdates"));
+		});
 #endif
 
 	return SendAllEndOfFrameUpdates;
@@ -960,17 +959,15 @@ FSendAllEndOfFrameUpdates* BeginSendEndOfFrameUpdatesDrawEvent(FGPUSkinCache* GP
 
 void EndSendEndOfFrameUpdatesDrawEvent(FSendAllEndOfFrameUpdates* SendAllEndOfFrameUpdates)
 {
-	ENQUEUE_UNIQUE_RENDER_COMMAND_ONEPARAMETER(
-		EndDrawEventCommand,
-		FSendAllEndOfFrameUpdates*, SendAllEndOfFrameUpdates, SendAllEndOfFrameUpdates,
+	ENQUEUE_RENDER_COMMAND(EndDrawEventCommand)(
+		[SendAllEndOfFrameUpdates](FRHICommandList& RHICmdList)
 	{
-	PREPROCESSOR_IF(RHI_RAYTRACING, 
+#if RHI_RAYTRACING
 		if (SendAllEndOfFrameUpdates->GPUSkinCache)
 		{
 			SendAllEndOfFrameUpdates->GPUSkinCache->CommitRayTracingGeometryUpdates(RHICmdList);
 		}
-	, PREPROCESSOR_NOTHING);
-			
+#endif
 		delete SendAllEndOfFrameUpdates;
 	});
 }
@@ -1292,30 +1289,29 @@ DECLARE_CYCLE_STAT(TEXT("TG_LastDemotable"), STAT_TG_LastDemotable, STATGROUP_Ti
 TDrawEvent<FRHICommandList>* BeginTickDrawEvent()
 {
 	TDrawEvent<FRHICommandList>* TickDrawEvent = new TDrawEvent<FRHICommandList>();
-
-	ENQUEUE_UNIQUE_RENDER_COMMAND_ONEPARAMETER(
-		BeginDrawEventCommand,
-		TDrawEvent<FRHICommandList>*,TickDrawEvent,TickDrawEvent,
-	{
-		BEGIN_DRAW_EVENTF(
-			RHICmdList, 
-			WorldTick, 
-			(*TickDrawEvent),
-			TEXT("WorldTick"));
-	});
+	TDrawEvent<FRHICommandList>* InTickDrawEvent = TickDrawEvent;
+	ENQUEUE_RENDER_COMMAND(BeginDrawEventCommand)(
+		[InTickDrawEvent](FRHICommandList& RHICmdList)
+		{
+			BEGIN_DRAW_EVENTF(
+				RHICmdList, 
+				WorldTick, 
+				(*InTickDrawEvent),
+				TEXT("WorldTick"));
+		});
 
 	return TickDrawEvent;
 }
 
 void EndTickDrawEvent(TDrawEvent<FRHICommandList>* TickDrawEvent)
 {
-	ENQUEUE_UNIQUE_RENDER_COMMAND_ONEPARAMETER(
-		EndDrawEventCommand,
-		TDrawEvent<FRHICommandList>*,TickDrawEvent,TickDrawEvent,
-	{
-		STOP_DRAW_EVENT((*TickDrawEvent));
-		delete TickDrawEvent;
-	});
+	TDrawEvent<FRHICommandList>* InTickDrawEvent = TickDrawEvent;
+	ENQUEUE_RENDER_COMMAND(EndDrawEventCommand)(
+		[InTickDrawEvent](FRHICommandList& RHICmdList)
+		{
+			STOP_DRAW_EVENT((*InTickDrawEvent));
+			delete InTickDrawEvent;
+		});
 }
 
 
@@ -1818,17 +1814,16 @@ void UWorld::Tick( ELevelTick TickType, float DeltaSeconds )
 		GEngine->XRSystem->OnEndGameFrame( GEngine->GetWorldContextFromWorldChecked( this ) );
 	}
 
-	ENQUEUE_UNIQUE_RENDER_COMMAND_ONEPARAMETER(
-		TickInGamePerfTrackersRT,
-		UWorld*, WorldParam, this,
-		{
-		//Tick game and other thread trackers.
-		for (int32 Tracker = 0; Tracker < (int32)EInGamePerfTrackers::Num; ++Tracker)
+	UWorld* WorldParam = this;
+	ENQUEUE_RENDER_COMMAND(TickInGamePerfTrackersRT)(
+		[WorldParam](FRHICommandList& RHICmdList)
 		{
-			WorldParam->PerfTrackers->GetInGamePerformanceTracker((EInGamePerfTrackers)Tracker, EInGamePerfTrackerThreads::RenderThread).Tick();
-		}
-	}
-	);
+			//Tick game and other thread trackers.
+			for (int32 Tracker = 0; Tracker < (int32)EInGamePerfTrackers::Num; ++Tracker)
+			{
+				WorldParam->PerfTrackers->GetInGamePerformanceTracker((EInGamePerfTrackers)Tracker, EInGamePerfTrackerThreads::RenderThread).Tick();
+			}
+		});
 
 	EndTickDrawEvent(TickDrawEvent);
 }
diff --git a/Engine/Source/Runtime/Engine/Private/Materials/Material.cpp b/Engine/Source/Runtime/Engine/Private/Materials/Material.cpp
--- a/Engine/Source/Runtime/Engine/Private/Materials/Material.cpp
+++ b/Engine/Source/Runtime/Engine/Private/Materials/Material.cpp
@@ -255,12 +255,12 @@ public:
 	 */
 	void GameThread_Destroy()
 	{
-		ENQUEUE_UNIQUE_RENDER_COMMAND_ONEPARAMETER(
-			FDestroyDefaultMaterialInstanceCommand,
-			FDefaultMaterialInstance*,Resource,this,
-		{
-			delete Resource;
-		});
+		FDefaultMaterialInstance* Resource = this;
+		ENQUEUE_RENDER_COMMAND(FDestroyDefaultMaterialInstanceCommand)(
+			[Resource](FRHICommandList& RHICmdList)
+			{
+				delete Resource;
+			});
 	}
 
 	// FMaterialRenderProxy interface.
diff --git a/Engine/Source/Runtime/Engine/Private/Materials/MaterialInstance.cpp b/Engine/Source/Runtime/Engine/Private/Materials/MaterialInstance.cpp
--- a/Engine/Source/Runtime/Engine/Private/Materials/MaterialInstance.cpp
+++ b/Engine/Source/Runtime/Engine/Private/Materials/MaterialInstance.cpp
@@ -3518,12 +3518,12 @@ void UMaterialInstance::ClearParameterValuesInternal(const bool bAllParameters)
 
 	if (Resource)
 	{
-		ENQUEUE_UNIQUE_RENDER_COMMAND_ONEPARAMETER(
-			FClearMIParametersCommand,
-			FMaterialInstanceResource*,Resource,Resource,
-		{
-			Resource->RenderThread_ClearParameters();
-		});
+		FMaterialInstanceResource* InResource = Resource;
+		ENQUEUE_RENDER_COMMAND(FClearMIParametersCommand)(
+			[InResource](FRHICommandList& RHICmdList)
+			{
+				InResource->RenderThread_ClearParameters();
+			});
 	}
 
 	InitResources();
diff --git a/Engine/Source/Runtime/Engine/Private/Materials/MaterialInstanceSupport.h b/Engine/Source/Runtime/Engine/Private/Materials/MaterialInstanceSupport.h
--- a/Engine/Source/Runtime/Engine/Private/Materials/MaterialInstanceSupport.h
+++ b/Engine/Source/Runtime/Engine/Private/Materials/MaterialInstanceSupport.h
@@ -77,12 +77,12 @@ public:
 	 */
 	void GameThread_Destroy()
 	{
-		ENQUEUE_UNIQUE_RENDER_COMMAND_ONEPARAMETER(
-			FDestroyMaterialInstanceResourceCommand,
-			FMaterialInstanceResource*,Resource,this,
-		{
-			delete Resource;
-		});
+		FMaterialInstanceResource* Resource = this;
+		ENQUEUE_RENDER_COMMAND(FDestroyMaterialInstanceResourceCommand)(
+			[Resource](FRHICommandList& RHICmdList)
+			{
+				delete Resource;
+			});
 	}
 
 	// FRenderResource interface.
diff --git a/Engine/Source/Runtime/Engine/Private/Materials/MaterialShared.cpp b/Engine/Source/Runtime/Engine/Private/Materials/MaterialShared.cpp
--- a/Engine/Source/Runtime/Engine/Private/Materials/MaterialShared.cpp
+++ b/Engine/Source/Runtime/Engine/Private/Materials/MaterialShared.cpp
@@ -957,12 +957,12 @@ void FMaterial::ReleaseShaderMap()
 	{
 		GameThreadShaderMap = nullptr;
 		
-		ENQUEUE_UNIQUE_RENDER_COMMAND_ONEPARAMETER(
-			ReleaseShaderMap,
-			FMaterial*,Material,this,
-		{
-			Material->SetRenderingThreadShaderMap(nullptr);
-		});
+		FMaterial* Material = this;
+		ENQUEUE_RENDER_COMMAND(ReleaseShaderMap)(
+			[Material](FRHICommandList& RHICmdList)
+			{
+				Material->SetRenderingThreadShaderMap(nullptr);
+			});
 	}
 }
 
@@ -3421,10 +3421,10 @@ void SetShaderMapsOnMaterialResources(const TMap<FMaterial*, FMaterialShaderMap*
 		check(!Material->RequiresSynchronousCompilation());
 	}
 
-	ENQUEUE_UNIQUE_RENDER_COMMAND_ONEPARAMETER(
-		FSetShaderMapOnMaterialResources,
-		FMaterialsToUpdateMap, MaterialsToUpdate, MaterialsToUpdate,
+	FMaterialsToUpdateMap InMaterialsToUpdate = MaterialsToUpdate;
+	ENQUEUE_RENDER_COMMAND(FSetShaderMapOnMaterialResources)(
+		[InMaterialsToUpdate](FRHICommandListImmediate& RHICmdList)
 		{
-			SetShaderMapsOnMaterialResources_RenderThread(RHICmdList, MaterialsToUpdate);
+			SetShaderMapsOnMaterialResources_RenderThread(RHICmdList, InMaterialsToUpdate);
 		});
-}
\ No newline at end of file
+}
diff --git a/Engine/Source/Runtime/Engine/Private/Particles/FXSystem.cpp b/Engine/Source/Runtime/Engine/Private/Particles/FXSystem.cpp
--- a/Engine/Source/Runtime/Engine/Private/Particles/FXSystem.cpp
+++ b/Engine/Source/Runtime/Engine/Private/Particles/FXSystem.cpp
@@ -26,12 +26,11 @@ FFXSystemInterface* FFXSystemInterface::Create(ERHIFeatureLevel::Type InFeatureL
 
 void FFXSystemInterface::Destroy( FFXSystemInterface* FXSystem )
 {
-	ENQUEUE_UNIQUE_RENDER_COMMAND_ONEPARAMETER(
-		FDestroyFXSystemCommand,
-		FFXSystemInterface*, FXSystem, FXSystem,
-	{
-		delete FXSystem;
-	});
+	ENQUEUE_RENDER_COMMAND(FDestroyFXSystemCommand)(
+		[FXSystem](FRHICommandList& RHICmdList)
+		{
+			delete FXSystem;
+		});
 }
 
 FFXSystemInterface::~FFXSystemInterface()
diff --git a/Engine/Source/Runtime/Engine/Private/Particles/ParticleGpuSimulation.cpp b/Engine/Source/Runtime/Engine/Private/Particles/ParticleGpuSimulation.cpp
--- a/Engine/Source/Runtime/Engine/Private/Particles/ParticleGpuSimulation.cpp
+++ b/Engine/Source/Runtime/Engine/Private/Particles/ParticleGpuSimulation.cpp
@@ -451,12 +451,12 @@ public:
 	 */
 	void Destroy()
 	{
-		ENQUEUE_UNIQUE_RENDER_COMMAND_ONEPARAMETER(
-			FDestroyParticleSimulationResourcesCommand,
-			FParticleSimulationResources*, ParticleResources, this,
-		{
-			delete ParticleResources;
-		});
+		FParticleSimulationResources* ParticleResources = this;
+		ENQUEUE_RENDER_COMMAND(FDestroyParticleSimulationResourcesCommand)(
+			[ParticleResources](FRHICommandList& RHICmdList)
+			{
+				delete ParticleResources;
+			});
 	}
 
 	/**
@@ -2702,12 +2702,12 @@ public:
 	void Destroy()
 	{
 		bDestroyed_GameThread = true;
-		ENQUEUE_UNIQUE_RENDER_COMMAND_ONEPARAMETER(
-			FReleaseParticleSimulationGPUCommand,
-			FParticleSimulationGPU*, Simulation, this,
-		{
-			Simulation->Destroy_RenderThread();
-		});
+		FParticleSimulationGPU* Simulation = this;
+		ENQUEUE_RENDER_COMMAND(FReleaseParticleSimulationGPUCommand)(
+			[Simulation](FRHICommandList& RHICmdList)
+			{
+				Simulation->Destroy_RenderThread();
+			});
 	}
 
 	/**
@@ -2727,12 +2727,12 @@ public:
 	void BeginReleaseResources()
 	{
 		bReleased_GameThread = true;
-		ENQUEUE_UNIQUE_RENDER_COMMAND_ONEPARAMETER(
-			FReleaseParticleSimulationResourcesGPUCommand,
-			FParticleSimulationGPU*, Simulation, this,
-		{
-			Simulation->ReleaseRenderResources();
-		});
+		FParticleSimulationGPU* Simulation = this;
+		ENQUEUE_RENDER_COMMAND(FReleaseParticleSimulationResourcesGPUCommand)(
+			[Simulation](FRHICommandList& RHICmdList)
+			{
+				Simulation->ReleaseRenderResources();
+			});
 	}
 
 private:
@@ -2811,9 +2811,9 @@ public:
 		{
 			// When all references are released, we need the render thread
 			// to release RHI resources and delete this instance.
-			ENQUEUE_UNIQUE_RENDER_COMMAND_ONEPARAMETER(
-				ReleaseCommand,
-				FRenderResource*, Resource, this,
+			FRenderResource* Resource = this;
+			ENQUEUE_RENDER_COMMAND(ReleaseCommand)(
+				[Resource](FRHICommandList& RHICmdList)
 				{
 					Resource->ReleaseResource();
 					delete Resource;
@@ -3874,18 +3874,17 @@ FGPUSpriteParticleEmitterInstance(FFXSystem* InFXSystem, FGPUSpriteEmitterInfo&
 			return;
 		}
 
-		ENQUEUE_UNIQUE_RENDER_COMMAND_ONEPARAMETER(
-			FComputeGPUSpriteBoundsCommand,
-			FGPUSpriteParticleEmitterInstance*, EmitterInstance, this,
-		{
-			
-			EmitterInstance->ParticleBoundingBox = ComputeParticleBounds(
-				RHICmdList,
-				EmitterInstance->Simulation->VertexBuffer.VertexBufferSRV,
-				EmitterInstance->FXSystem->GetParticleSimulationResources()->GetVisualizeStateTextures().PositionTextureRHI,
-				EmitterInstance->Simulation->VertexBuffer.ParticleCount
-				);
-		});
+		FGPUSpriteParticleEmitterInstance* EmitterInstance = this;
+		ENQUEUE_RENDER_COMMAND(FComputeGPUSpriteBoundsCommand)(
+			[EmitterInstance](FRHICommandListImmediate& RHICmdList)
+			{
+				EmitterInstance->ParticleBoundingBox = ComputeParticleBounds(
+					RHICmdList,
+					EmitterInstance->Simulation->VertexBuffer.VertexBufferSRV,
+					EmitterInstance->FXSystem->GetParticleSimulationResources()->GetVisualizeStateTextures().PositionTextureRHI,
+					EmitterInstance->Simulation->VertexBuffer.ParticleCount
+					);
+			});
 		FlushRenderingCommands();
 
 		// Take the size of sprites in to account.
@@ -3939,9 +3938,9 @@ private:
 		//
 		if (!Simulation->bDestroyed_GameThread)
 		{
-			ENQUEUE_UNIQUE_RENDER_COMMAND_ONEPARAMETER(
-				FResetVectorFieldCommand,
-				FVectorFieldResource*, Resource, Simulation->LocalVectorField.Resource,
+			FVectorFieldResource* Resource = Simulation->LocalVectorField.Resource;
+			ENQUEUE_RENDER_COMMAND(FResetVectorFieldCommand)(
+				[Resource](FRHICommandList& RHICmdList)
 				{
 					if (Resource)
 					{
@@ -4527,17 +4526,16 @@ void FFXSystem::InitGPUResources()
 		ParticleSimulationResources->Init();
 
 		ParticleSimulationResources->ParticleSortBuffers.SetBufferSize(GParticleSimulationTextureSizeX * GParticleSimulationTextureSizeY);
-		ENQUEUE_UNIQUE_RENDER_COMMAND_ONEPARAMETER(
-			FInitParticleSortBuffersCommand,
-			FParticleSimulationResources *, ParticleSimulationResources, ParticleSimulationResources,
+		FParticleSimulationResources* InParticleSimulationResources = ParticleSimulationResources;
+		ENQUEUE_RENDER_COMMAND(FInitParticleSortBuffersCommand)(
+			[InParticleSimulationResources](FRHICommandList& RHICmdList)
 			{
-				ParticleSimulationResources->ParticleSortBuffers.InitRHI();
+				InParticleSimulationResources->ParticleSortBuffers.InitRHI();
 
 				// Initialize SortedVertexBuffer to a valid resource, ensuring it can be used in GetDynamicMeshElementsEmitter() 
-				ParticleSimulationResources->SortedVertexBuffer.VertexBufferRHI = ParticleSimulationResources->ParticleSortBuffers.GetSortedVertexBufferRHI(0);
-				ParticleSimulationResources->SortedVertexBuffer.VertexBufferSRV = ParticleSimulationResources->ParticleSortBuffers.GetSortedVertexBufferSRV(0);
-			}
-		);
+				InParticleSimulationResources->SortedVertexBuffer.VertexBufferRHI = InParticleSimulationResources->ParticleSortBuffers.GetSortedVertexBufferRHI(0);
+				InParticleSimulationResources->SortedVertexBuffer.VertexBufferSRV = InParticleSimulationResources->ParticleSortBuffers.GetSortedVertexBufferSRV(0);
+			});
 	}
 }
 
@@ -4547,13 +4545,12 @@ void FFXSystem::ReleaseGPUResources()
 	{
 		check(ParticleSimulationResources);
 		ParticleSimulationResources->Release();
-		ENQUEUE_UNIQUE_RENDER_COMMAND_ONEPARAMETER(
-			FReleaseParticleSortBuffersCommand,
-			FParticleSimulationResources *, ParticleSimulationResources, ParticleSimulationResources,
+		FParticleSimulationResources* InParticleSimulationResources = ParticleSimulationResources;
+		ENQUEUE_RENDER_COMMAND(FReleaseParticleSortBuffersCommand)(
+			[InParticleSimulationResources](FRHICommandList& RHICmdList)
 			{
-				ParticleSimulationResources->ParticleSortBuffers.ReleaseRHI();
-			}
-		);
+				InParticleSimulationResources->ParticleSortBuffers.ReleaseRHI();
+			});
 	}
 }
 
diff --git a/Engine/Source/Runtime/Engine/Private/Rendering/SubsurfaceProfile.cpp b/Engine/Source/Runtime/Engine/Private/Rendering/SubsurfaceProfile.cpp
--- a/Engine/Source/Runtime/Engine/Private/Rendering/SubsurfaceProfile.cpp
+++ b/Engine/Source/Runtime/Engine/Private/Rendering/SubsurfaceProfile.cpp
@@ -371,9 +371,9 @@ USubsurfaceProfile::USubsurfaceProfile(const FObjectInitializer& ObjectInitializ
 
 void USubsurfaceProfile::BeginDestroy()
 {
-	ENQUEUE_UNIQUE_RENDER_COMMAND_ONEPARAMETER(
-		RemoveSubsurfaceProfile,
-		USubsurfaceProfile*, Ref, this,
+	USubsurfaceProfile* Ref = this;
+	ENQUEUE_RENDER_COMMAND(RemoveSubsurfaceProfile)(
+		[Ref](FRHICommandList& RHICmdList)
 		{
 			GSubsurfaceProfileTextureObject.RemoveProfile(Ref);
 		});
diff --git a/Engine/Source/Runtime/Engine/Private/Slate/DebugCanvas.cpp b/Engine/Source/Runtime/Engine/Private/Slate/DebugCanvas.cpp
--- a/Engine/Source/Runtime/Engine/Private/Slate/DebugCanvas.cpp
+++ b/Engine/Source/Runtime/Engine/Private/Slate/DebugCanvas.cpp
@@ -72,9 +72,8 @@ void FDebugCanvasDrawer::ReleaseResources()
 {
 	FDebugCanvasDrawer* const ReleaseMe = this;
 
-	ENQUEUE_UNIQUE_RENDER_COMMAND_ONEPARAMETER(
-		ReleaseCommand,
-		FDebugCanvasDrawer*, ReleaseMe, ReleaseMe,
+	ENQUEUE_RENDER_COMMAND(ReleaseCommand)(
+		[ReleaseMe](FRHICommandList& RHICmdList)
 		{
 			ReleaseMe->ReleaseTexture();
 		});
diff --git a/Engine/Source/Runtime/Engine/Private/StaticMesh.cpp b/Engine/Source/Runtime/Engine/Private/StaticMesh.cpp
--- a/Engine/Source/Runtime/Engine/Private/StaticMesh.cpp
+++ b/Engine/Source/Runtime/Engine/Private/StaticMesh.cpp
@@ -727,13 +727,13 @@ void FStaticMeshLODResources::InitResources(UStaticMesh* Parent)
 		INC_DWORD_STAT_BY( STAT_StaticMeshDistanceFieldMemory, DistanceFieldData->GetResourceSizeBytes() );
 	}
 
-	ENQUEUE_UNIQUE_RENDER_COMMAND_ONEPARAMETER(
-		UpdateMemoryStats,
-		FStaticMeshLODResources*, This, this,
+	FStaticMeshLODResources* This = this;
+	ENQUEUE_RENDER_COMMAND(UpdateMemoryStats)(
+		[This](FRHICommandList& RHICmdList)
 		{		
 			const uint32 StaticMeshVertexMemory =
-			This->VertexBuffers.StaticMeshVertexBuffer.GetResourceSize() +
-			This->VertexBuffers.PositionVertexBuffer.GetStride() * This->VertexBuffers.PositionVertexBuffer.GetNumVertices();
+				This->VertexBuffers.StaticMeshVertexBuffer.GetResourceSize() +
+				This->VertexBuffers.PositionVertexBuffer.GetStride() * This->VertexBuffers.PositionVertexBuffer.GetNumVertices();
 			const uint32 ResourceVertexColorMemory = This->VertexBuffers.ColorVertexBuffer.GetStride() * This->VertexBuffers.ColorVertexBuffer.GetNumVertices();
 
 			INC_DWORD_STAT_BY( STAT_StaticMeshVertexMemory, StaticMeshVertexMemory );
@@ -1872,9 +1872,9 @@ void UStaticMesh::InitResources()
 	}
 	
 #if	STATS
-	ENQUEUE_UNIQUE_RENDER_COMMAND_ONEPARAMETER(
-		UpdateMemoryStats,
-		UStaticMesh*, This, this,
+	UStaticMesh* This = this;
+	ENQUEUE_RENDER_COMMAND(UpdateMemoryStats)(
+		[This](FRHICommandList& RHICmdList)
 		{
 			const uint32 StaticMeshResourceSize = This->GetResourceSizeBytes( EResourceSizeMode::Exclusive );
 			INC_DWORD_STAT_BY( STAT_StaticMeshTotalMemory, StaticMeshResourceSize );
diff --git a/Engine/Source/Runtime/Engine/Private/Texture2D.cpp b/Engine/Source/Runtime/Engine/Private/Texture2D.cpp
--- a/Engine/Source/Runtime/Engine/Private/Texture2D.cpp
+++ b/Engine/Source/Runtime/Engine/Private/Texture2D.cpp
@@ -1147,12 +1147,12 @@ void UTexture2D::RefreshSamplerStates()
 		return;
 	}
 
-	ENQUEUE_UNIQUE_RENDER_COMMAND_ONEPARAMETER(
-		RefreshSamplerStatesCommand,
-		FTexture2DResource*, Texture2DResource, ((FTexture2DResource*)Resource),
-	{
-		Texture2DResource->RefreshSamplerStates();
-	});
+	FTexture2DResource* Texture2DResource = ((FTexture2DResource*)Resource);
+	ENQUEUE_RENDER_COMMAND(RefreshSamplerStatesCommand)(
+		[Texture2DResource](FRHICommandList& RHICmdList)
+		{
+			Texture2DResource->RefreshSamplerStates();
+		});
 }
 
 /*-----------------------------------------------------------------------------
@@ -1206,11 +1206,11 @@ FTexture2DResource::~FTexture2DResource()
 {
 	// free resource memory that was preallocated
 	// The deletion needs to happen in the rendering thread.
-	ENQUEUE_UNIQUE_RENDER_COMMAND_ONEPARAMETER(
-		DeleteResourceMem,
-		FTexture2DResourceMem*,ResourceMem,ResourceMem,
+	FTexture2DResourceMem* InResourceMem = ResourceMem;
+	ENQUEUE_RENDER_COMMAND(DeleteResourceMem)(
+		[InResourceMem](FRHICommandList& RHICmdList)
 		{
-			delete ResourceMem;
+			delete InResourceMem;
 		});
 
 	// Make sure we're not leaking memory if InitRHI has never been called.
diff --git a/Engine/Source/Runtime/Engine/Private/TextureRenderTargetCube.cpp b/Engine/Source/Runtime/Engine/Private/TextureRenderTargetCube.cpp
--- a/Engine/Source/Runtime/Engine/Private/TextureRenderTargetCube.cpp
+++ b/Engine/Source/Runtime/Engine/Private/TextureRenderTargetCube.cpp
@@ -380,7 +380,7 @@ bool FTextureRenderTargetCubeResource::ReadPixels(TArray< FColor >& OutImageData
 	};
 
 	OutImageData.Reset();
-	FReadSurfaceContext ReadSurfaceContext =
+	FReadSurfaceContext Context =
 	{
 		this,
 		&OutImageData,
@@ -388,17 +388,16 @@ bool FTextureRenderTargetCubeResource::ReadPixels(TArray< FColor >& OutImageData
 		InFlags
 	};
 
-	ENQUEUE_UNIQUE_RENDER_COMMAND_ONEPARAMETER(
-	  ReadSurfaceCommand,
-	  FReadSurfaceContext, Context, ReadSurfaceContext,
-	{
-		RHICmdList.ReadSurfaceData(
-		  Context.SrcRenderTarget->TextureCubeRHI,
-		  Context.Rect,
-		  *Context.OutData,
-		  Context.Flags
-		);
-	});
+	ENQUEUE_RENDER_COMMAND(ReadSurfaceCommand)(
+		[Context](FRHICommandListImmediate& RHICmdList)
+		{
+			RHICmdList.ReadSurfaceData(
+				Context.SrcRenderTarget->TextureCubeRHI,
+				Context.Rect,
+				*Context.OutData,
+				Context.Flags
+			);
+		});
 	FlushRenderingCommands();
 
 	return true;
@@ -426,7 +425,7 @@ bool FTextureRenderTargetCubeResource::ReadPixels(TArray<FFloat16Color>& OutImag
 		ECubeFace CubeFace;
 	};
 
-	FReadSurfaceFloatContext ReadSurfaceFloatContext =
+	FReadSurfaceFloatContext Context =
 	{
 		this,
 		&OutImageData,
@@ -434,18 +433,17 @@ bool FTextureRenderTargetCubeResource::ReadPixels(TArray<FFloat16Color>& OutImag
 		InFlags.GetCubeFace()
 	};
 
-	ENQUEUE_UNIQUE_RENDER_COMMAND_ONEPARAMETER(
-	  ReadSurfaceFloatCommand,
-	  FReadSurfaceFloatContext, Context, ReadSurfaceFloatContext,
-	{
-		RHICmdList.ReadSurfaceFloatData(
-		  Context.SrcRenderTarget->TextureCubeRHI,
-		  Context.Rect,
-		  *Context.OutData,
-		  Context.CubeFace,
-		  0,
-		  0
-		);
+	ENQUEUE_RENDER_COMMAND(ReadSurfaceFloatCommand)(
+		[Context](FRHICommandListImmediate& RHICmdList)
+		{
+			RHICmdList.ReadSurfaceFloatData(
+				Context.SrcRenderTarget->TextureCubeRHI,
+				Context.Rect,
+				*Context.OutData,
+				Context.CubeFace,
+				0,
+				0
+			);
 	});
 
 	FlushRenderingCommands();
diff --git a/Engine/Source/Runtime/Engine/Private/VectorField.cpp b/Engine/Source/Runtime/Engine/Private/VectorField.cpp
--- a/Engine/Source/Runtime/Engine/Private/VectorField.cpp
+++ b/Engine/Source/Runtime/Engine/Private/VectorField.cpp
@@ -57,14 +57,14 @@ FVectorFieldInstance::~FVectorFieldInstance()
 {
 	if (Resource && bInstancedResource)
 	{
-		ENQUEUE_UNIQUE_RENDER_COMMAND_ONEPARAMETER(
-			FDestroyVectorFieldResourceCommand,
-			FVectorFieldResource*,Resource,Resource,
-		{
-			Resource->ReleaseResource();
-			delete Resource;
-		});
-		Resource = NULL;
+		FVectorFieldResource* InResource = Resource;
+		ENQUEUE_RENDER_COMMAND(FDestroyVectorFieldResourceCommand)(
+			[InResource](FRHICommandList& RHICmdList)
+			{
+				InResource->ReleaseResource();
+				delete InResource;
+			});
+		Resource = nullptr;
 	}
 }
 
@@ -380,15 +380,15 @@ void UVectorFieldStatic::ReleaseResource()
 {
 	if ( Resource != NULL )
 	{
-		ENQUEUE_UNIQUE_RENDER_COMMAND_ONEPARAMETER(
-			ReleaseVectorFieldCommand,
-			FRenderResource*,Resource,Resource,
-		{
-			Resource->ReleaseResource();
-			delete Resource;
-		});
+		FRenderResource* InResource = Resource;
+		ENQUEUE_RENDER_COMMAND(ReleaseVectorFieldCommand)(
+			[InResource](FRHICommandList& RHICmdList)
+			{
+				InResource->ReleaseResource();
+				delete InResource;
+			});
 	}
-	Resource = NULL;
+	Resource = nullptr;
 }
 
 void UVectorFieldStatic::Serialize(FArchive& Ar)
@@ -648,13 +648,13 @@ void UVectorFieldComponent::OnUnregister()
 	{
 		if (VectorFieldInstance)
 		{
-			ENQUEUE_UNIQUE_RENDER_COMMAND_ONEPARAMETER(
-				FDestroyVectorFieldInstanceCommand,
-				FVectorFieldInstance*, VectorFieldInstance, VectorFieldInstance,
-			{
-				delete VectorFieldInstance;
-			});
-			VectorFieldInstance = NULL;
+			FVectorFieldInstance* InVectorFieldInstance = VectorFieldInstance;
+			ENQUEUE_RENDER_COMMAND(FDestroyVectorFieldInstanceCommand)(
+				[InVectorFieldInstance](FRHICommandList& RHICmdList)
+				{
+					delete InVectorFieldInstance;
+				});
+			VectorFieldInstance = nullptr;
 		}
 	}
 	else if (VectorFieldInstance)
diff --git a/Engine/Source/Runtime/MediaAssets/Private/Assets/MediaTexture.cpp b/Engine/Source/Runtime/MediaAssets/Private/Assets/MediaTexture.cpp
--- a/Engine/Source/Runtime/MediaAssets/Private/Assets/MediaTexture.cpp
+++ b/Engine/Source/Runtime/MediaAssets/Private/Assets/MediaTexture.cpp
@@ -181,8 +181,9 @@ void UMediaTexture::BeginDestroy()
 
 	if (CurrentGuid.IsValid())
 	{
-		ENQUEUE_UNIQUE_RENDER_COMMAND_ONEPARAMETER(MediaTextureUnregisterGuid,
-			FGuid, Guid, CurrentGuid,
+		FGuid Guid = CurrentGuid;
+		ENQUEUE_RENDER_COMMAND(MediaTextureUnregisterGuid)(
+			[Guid](FRHICommandList& RHICmdList)
 			{
 				FExternalTextureRegistry::Get().UnregisterExternalTexture(Guid);
 			});
diff --git a/Engine/Source/Runtime/RenderCore/Private/RenderResource.cpp b/Engine/Source/Runtime/RenderCore/Private/RenderResource.cpp
--- a/Engine/Source/Runtime/RenderCore/Private/RenderResource.cpp
+++ b/Engine/Source/Runtime/RenderCore/Private/RenderResource.cpp
@@ -198,11 +198,11 @@ struct FBatchedReleaseResources
 	{
 		if (NumBatch)
 		{
-			ENQUEUE_UNIQUE_RENDER_COMMAND_ONEPARAMETER(
-				BatchReleaseCommand,
-				FBatchedReleaseResources, BatchedReleaseResources, *this,
+			const FBatchedReleaseResources BatchedReleaseResources = *this;
+			ENQUEUE_RENDER_COMMAND(BatchReleaseCommand)(
+				[BatchedReleaseResources](FRHICommandList& RHICmdList)
 				{
-					BatchedReleaseResources.Execute();
+					((FBatchedReleaseResources&)BatchedReleaseResources).Execute();
 				});
 			Reset();
 		}
@@ -245,9 +245,8 @@ void BeginReleaseResource(FRenderResource* Resource)
 		GBatchedRelease.Add(Resource);
 		return;
 	}
-	ENQUEUE_UNIQUE_RENDER_COMMAND_ONEPARAMETER(
-		ReleaseCommand,
-		FRenderResource*,Resource,Resource,
+	ENQUEUE_RENDER_COMMAND(ReleaseCommand)(
+		[Resource](FRHICommandList& RHICmdList)
 		{
 			Resource->ReleaseResource();
 		});
@@ -256,9 +255,8 @@ void BeginReleaseResource(FRenderResource* Resource)
 void ReleaseResourceAndFlush(FRenderResource* Resource)
 {
 	// Send the release message.
-	ENQUEUE_UNIQUE_RENDER_COMMAND_ONEPARAMETER(
-		ReleaseCommand,
-		FRenderResource*,Resource,Resource,
+	ENQUEUE_RENDER_COMMAND(ReleaseCommand)(
+		[Resource](FRHICommandList& RHICmdList)
 		{
 			Resource->ReleaseResource();
 		});
diff --git a/Engine/Source/Runtime/RenderCore/Private/StaticBoundShaderState.cpp b/Engine/Source/Runtime/RenderCore/Private/StaticBoundShaderState.cpp
--- a/Engine/Source/Runtime/RenderCore/Private/StaticBoundShaderState.cpp
+++ b/Engine/Source/Runtime/RenderCore/Private/StaticBoundShaderState.cpp
@@ -32,8 +32,9 @@ FGlobalBoundShaderStateResource::FGlobalBoundShaderStateResource()
 	}
 	else
 	{
-		ENQUEUE_UNIQUE_RENDER_COMMAND_ONEPARAMETER(
-			LinkGlobalBoundShaderStateResource,FGlobalBoundShaderStateResource*,Resource,this,
+		FGlobalBoundShaderStateResource* Resource = this;
+		ENQUEUE_RENDER_COMMAND(LinkGlobalBoundShaderStateResource)(
+			[Resource](FRHICommandList& RHICmdList)
 			{
 				Resource->GlobalListLink.LinkHead(GetGlobalBoundShaderStateList());
 			});
