Fixes for skeletal mesh loading

diff --git a/Engine/Source/Runtime/Engine/Private/GPUSkinVertexFactory.cpp b/Engine/Source/Runtime/Engine/Private/GPUSkinVertexFactory.cpp
--- a/Engine/Source/Runtime/Engine/Private/GPUSkinVertexFactory.cpp
+++ b/Engine/Source/Runtime/Engine/Private/GPUSkinVertexFactory.cpp
@@ -621,18 +621,15 @@ public:
 		FVertexInputStreamArray& VertexStreams) const override
 	{
 		check(VertexFactory->GetType() == &FGPUSkinPassthroughVertexFactory::StaticType);
-		// #dxr_todo do we need this call to the base?
-		FLocalVertexFactoryShaderParametersBase::GetElementShaderBindings(Scene, View, Shader, bShaderRequiresPositionOnlyStream, FeatureLevel, VertexFactory, BatchElement, ShaderBindings, VertexStreams);
+		FGPUSkinBatchElementUserData* BatchUserData = (FGPUSkinBatchElementUserData*)BatchElement.VertexFactoryUserData;
+		check(BatchUserData);
 
 		const auto* LocalVertexFactory = static_cast<const FGPUSkinPassthroughVertexFactory*>(VertexFactory);
-		if (LocalVertexFactory->SupportsManualVertexFetch(FeatureLevel))
-		{
-			FUniformBufferRHIParamRef VertexFactoryUniformBuffer = LocalVertexFactory->GetUniformBuffer();
-			ShaderBindings.Add(Shader->GetUniformBufferParameter<FLocalVertexFactoryUniformShaderParameters>(), VertexFactoryUniformBuffer);
-		}
+		FUniformBufferRHIParamRef VertexFactoryUniformBuffer = nullptr;
+		VertexFactoryUniformBuffer = LocalVertexFactory->GetUniformBuffer();
 
-		FGPUSkinBatchElementUserData* BatchUserData = (FGPUSkinBatchElementUserData*)BatchElement.VertexFactoryUserData;
-		check(BatchUserData);
+		// #dxr_todo do we need this call to the base?
+		FLocalVertexFactoryShaderParametersBase::GetElementShaderBindingsBase(Scene, View, Shader, bShaderRequiresPositionOnlyStream, FeatureLevel, VertexFactory, BatchElement, VertexFactoryUniformBuffer, ShaderBindings, VertexStreams);
 		FGPUSkinCache::GetShaderBindings(BatchUserData->Entry, BatchUserData->Section, Shader, (const FGPUSkinPassthroughVertexFactory*)VertexFactory, BatchElement.MinVertexIndex, GPUSkinCachePreviousPositionBuffer, ShaderBindings, VertexStreams);
 	}
 
diff --git a/Engine/Source/Runtime/Engine/Private/InstancedStaticMesh.cpp b/Engine/Source/Runtime/Engine/Private/InstancedStaticMesh.cpp
--- a/Engine/Source/Runtime/Engine/Private/InstancedStaticMesh.cpp
+++ b/Engine/Source/Runtime/Engine/Private/InstancedStaticMesh.cpp
@@ -2479,7 +2479,9 @@ void FInstancedStaticMeshVertexFactoryShaderParameters::GetElementShaderBindings
 	FVertexInputStreamArray& VertexStreams
 	) const
 {
-	FLocalVertexFactoryShaderParametersBase::GetElementShaderBindings(Scene, View, Shader, bShaderRequiresPositionOnlyStream, FeatureLevel, VertexFactory, BatchElement, ShaderBindings, VertexStreams);
+	// Decode VertexFactoryUserData as VertexFactoryUniformBuffer
+	FUniformBufferRHIParamRef VertexFactoryUniformBuffer = static_cast<FUniformBufferRHIParamRef>(BatchElement.VertexFactoryUserData);
+	FLocalVertexFactoryShaderParametersBase::GetElementShaderBindingsBase(Scene, View, Shader, bShaderRequiresPositionOnlyStream, FeatureLevel, VertexFactory, BatchElement, VertexFactoryUniformBuffer, ShaderBindings, VertexStreams);
 
 	const FInstancingUserData* InstancingUserData = (const FInstancingUserData*)BatchElement.UserData;
 	const auto* InstancedVertexFactory = static_cast<const FInstancedStaticMeshVertexFactory*>(VertexFactory);
diff --git a/Engine/Source/Runtime/Engine/Private/LocalVertexFactory.cpp b/Engine/Source/Runtime/Engine/Private/LocalVertexFactory.cpp
--- a/Engine/Source/Runtime/Engine/Private/LocalVertexFactory.cpp
+++ b/Engine/Source/Runtime/Engine/Private/LocalVertexFactory.cpp
@@ -90,7 +90,7 @@ TUniformBufferRef<FLocalVertexFactoryUniformShaderParameters> CreateLocalVFUnifo
 	return TUniformBufferRef<FLocalVertexFactoryUniformShaderParameters>::CreateUniformBufferImmediate(UniformParameters, UniformBuffer_MultiFrame);
 }
 
-void FLocalVertexFactoryShaderParametersBase::GetElementShaderBindings(
+void FLocalVertexFactoryShaderParametersBase::GetElementShaderBindingsBase(
 	const FSceneInterface* Scene,
 	const FSceneView* View,
 	const FMeshMaterialShader* Shader, 
@@ -98,6 +98,7 @@ void FLocalVertexFactoryShaderParametersBase::GetElementShaderBindings(
 	ERHIFeatureLevel::Type FeatureLevel,
 	const FVertexFactory* VertexFactory, 
 	const FMeshBatchElement& BatchElement,
+	FUniformBufferRHIParamRef VertexFactoryUniformBuffer,
 	FMeshDrawSingleShaderBindings& ShaderBindings,
 	FVertexInputStreamArray& VertexStreams
 	) const
@@ -106,8 +107,6 @@ void FLocalVertexFactoryShaderParametersBase::GetElementShaderBindings(
 	
 	if (LocalVertexFactory->SupportsManualVertexFetch(FeatureLevel) || UseGPUScene(GMaxRHIShaderPlatform, FeatureLevel))
 	{
-		FUniformBufferRHIParamRef VertexFactoryUniformBuffer = static_cast<FUniformBufferRHIParamRef>(BatchElement.VertexFactoryUserData);
-
 		if (!VertexFactoryUniformBuffer)
 		{
 			// No batch element override
@@ -149,6 +148,33 @@ void FLocalVertexFactoryShaderParametersBase::GetElementShaderBindings(
 	}
 }
 
+void FLocalVertexFactoryShaderParameters::GetElementShaderBindings(
+	const FSceneInterface* Scene,
+	const FSceneView* View,
+	const FMeshMaterialShader* Shader,
+	bool bShaderRequiresPositionOnlyStream,
+	ERHIFeatureLevel::Type FeatureLevel,
+	const FVertexFactory* VertexFactory,
+	const FMeshBatchElement& BatchElement,
+	FMeshDrawSingleShaderBindings& ShaderBindings,
+	FVertexInputStreamArray& VertexStreams
+) const
+{
+	// Decode VertexFactoryUserData as VertexFactoryUniformBuffer
+	FUniformBufferRHIParamRef VertexFactoryUniformBuffer = static_cast<FUniformBufferRHIParamRef>(BatchElement.VertexFactoryUserData);
+
+	FLocalVertexFactoryShaderParametersBase::GetElementShaderBindingsBase(
+		Scene,
+		View,
+		Shader,
+		bShaderRequiresPositionOnlyStream,
+		FeatureLevel,
+		VertexFactory,
+		BatchElement,
+		VertexFactoryUniformBuffer,
+		ShaderBindings,
+		VertexStreams);
+}
 
 /**
  * Should we cache the material's shadertype on this platform with this vertex factory? 
diff --git a/Engine/Source/Runtime/Engine/Public/LocalVertexFactory.h b/Engine/Source/Runtime/Engine/Public/LocalVertexFactory.h
--- a/Engine/Source/Runtime/Engine/Public/LocalVertexFactory.h
+++ b/Engine/Source/Runtime/Engine/Public/LocalVertexFactory.h
@@ -166,7 +166,8 @@ class FLocalVertexFactoryShaderParametersBase : public FVertexFactoryShaderParam
 public:
 	virtual void Bind(const FShaderParameterMap& ParameterMap) override;
 	virtual void Serialize(FArchive& Ar) override;
-	virtual void GetElementShaderBindings(
+
+	void GetElementShaderBindingsBase(
 		const class FSceneInterface* Scene,
 		const FSceneView* View,
 		const FMeshMaterialShader* Shader,
@@ -174,9 +175,10 @@ public:
 		ERHIFeatureLevel::Type FeatureLevel,
 		const FVertexFactory* VertexFactory,
 		const FMeshBatchElement& BatchElement,
+		FUniformBufferRHIParamRef VertexFactoryUniformBuffer,
 		FMeshDrawSingleShaderBindings& ShaderBindings,
 		FVertexInputStreamArray& VertexStreams
-		) const override;
+		) const;
 
 	FLocalVertexFactoryShaderParametersBase()
 		: bAnySpeedTreeParamIsBound(false)
@@ -194,4 +196,16 @@ public:
 class FLocalVertexFactoryShaderParameters : public FLocalVertexFactoryShaderParametersBase
 {
 public:
+
+	virtual void GetElementShaderBindings(
+		const FSceneInterface* Scene,
+		const FSceneView* View,
+		const FMeshMaterialShader* Shader,
+		bool bShaderRequiresPositionOnlyStream,
+		ERHIFeatureLevel::Type FeatureLevel,
+		const FVertexFactory* VertexFactory,
+		const FMeshBatchElement& BatchElement,
+		FMeshDrawSingleShaderBindings& ShaderBindings,
+		FVertexInputStreamArray& VertexStreams
+	) const override; 
 };
\ No newline at end of file
