Track requested material sampler type on 'FMaterialUniformExpressionTexture'
- This way TextureSample expression can use this to verify correct sampler type, this avoids the need for TextureObject expression to perform this verification
- Technically TextureObject probably should verify this as well, but this breaks certain existing content

diff --git a/Engine/Source/Runtime/Engine/Private/Materials/HLSLMaterialTranslator.h b/Engine/Source/Runtime/Engine/Private/Materials/HLSLMaterialTranslator.h
--- a/Engine/Source/Runtime/Engine/Private/Materials/HLSLMaterialTranslator.h
+++ b/Engine/Source/Runtime/Engine/Private/Materials/HLSLMaterialTranslator.h
@@ -3967,7 +3967,7 @@ protected:
 			);
 	}
 
-	virtual int32 Texture(UTexture* InTexture,int32& TextureReferenceIndex,ESamplerSourceMode SamplerSource=SSM_FromTextureAsset, ETextureMipValueMode MipValueMode = TMVM_None) override
+	virtual int32 Texture(UTexture* InTexture,int32& TextureReferenceIndex, EMaterialSamplerType SamplerType, ESamplerSourceMode SamplerSource=SSM_FromTextureAsset, ETextureMipValueMode MipValueMode = TMVM_None) override
 	{
 		if (FeatureLevel == ERHIFeatureLevel::ES2 && ShaderFrequency == SF_Vertex)
 		{
@@ -4000,10 +4000,10 @@ protected:
 #endif
 		checkf(TextureReferenceIndex != INDEX_NONE, TEXT("Material expression called Compiler->Texture() without implementing UMaterialExpression::GetReferencedTexture properly"));
 
-		return AddUniformExpression(new FMaterialUniformExpressionTexture(TextureReferenceIndex, SamplerSource),ShaderType,TEXT(""));
+		return AddUniformExpression(new FMaterialUniformExpressionTexture(TextureReferenceIndex, SamplerType, SamplerSource),ShaderType,TEXT(""));
 	}
 
-	virtual int32 TextureParameter(FName ParameterName,UTexture* DefaultValue,int32& TextureReferenceIndex,ESamplerSourceMode SamplerSource=SSM_FromTextureAsset) override
+	virtual int32 TextureParameter(FName ParameterName,UTexture* DefaultValue,int32& TextureReferenceIndex, EMaterialSamplerType SamplerType, ESamplerSourceMode SamplerSource=SSM_FromTextureAsset) override
 	{
 		if (ShaderFrequency != SF_Pixel
 			&& ErrorUnlessFeatureLevelSupported(ERHIFeatureLevel::ES3_1) == INDEX_NONE)
@@ -4017,7 +4017,7 @@ protected:
 
 		FMaterialParameterInfo ParameterInfo = GetParameterAssociationInfo();
 		ParameterInfo.Name = ParameterName;
-		return AddUniformExpression(new FMaterialUniformExpressionTextureParameter(ParameterInfo, TextureReferenceIndex, SamplerSource),ShaderType,TEXT(""));
+		return AddUniformExpression(new FMaterialUniformExpressionTextureParameter(ParameterInfo, TextureReferenceIndex, SamplerType, SamplerSource),ShaderType,TEXT(""));
 	}
 
 	virtual int32 ExternalTexture(const FGuid& ExternalTextureGuid) override
@@ -4227,10 +4227,11 @@ protected:
 			else
 			{
 				// Otherwise we sample normally
+				const EMaterialSamplerType SamplerType = SAMPLERTYPE_Masks;
 				FString WeightmapName = FString::Printf(TEXT("Weightmap%d"),WeightmapIndex);
 				int32 TextureReferenceIndex = INDEX_NONE;
-				int32 TextureCodeIndex = TextureParameter(FName(*WeightmapName), GEngine->WeightMapPlaceholderTexture, TextureReferenceIndex);
-				WeightmapCode = TextureSample(TextureCodeIndex, TextureCoordinate(3, false, false), SAMPLERTYPE_Masks);
+				int32 TextureCodeIndex = TextureParameter(FName(*WeightmapName), GEngine->WeightMapPlaceholderTexture, TextureReferenceIndex, SamplerType);
+				WeightmapCode = TextureSample(TextureCodeIndex, TextureCoordinate(3, false, false), SamplerType);
 			}
 
 			FString LayerMaskName = FString::Printf(TEXT("LayerMask_%s"),*ParameterName.ToString());
diff --git a/Engine/Source/Runtime/Engine/Private/Materials/MaterialExpressions.cpp b/Engine/Source/Runtime/Engine/Private/Materials/MaterialExpressions.cpp
--- a/Engine/Source/Runtime/Engine/Private/Materials/MaterialExpressions.cpp
+++ b/Engine/Source/Runtime/Engine/Private/Materials/MaterialExpressions.cpp
@@ -493,7 +493,7 @@ int32 CompileTextureSample(
 	}
 	else
 	{
-		TextureCodeIndex = ParameterName.IsSet() ? Compiler->TextureParameter(ParameterName.GetValue(), Texture, TextureReferenceIndex, SamplerSource) : Compiler->Texture(Texture, TextureReferenceIndex, SamplerSource, MipValueMode);
+		TextureCodeIndex = ParameterName.IsSet() ? Compiler->TextureParameter(ParameterName.GetValue(), Texture, TextureReferenceIndex, SamplerType, SamplerSource) : Compiler->Texture(Texture, TextureReferenceIndex, SamplerType, SamplerSource, MipValueMode);
 	}
 
 	return Compiler->TextureSample(
@@ -1635,7 +1635,7 @@ int32 UMaterialExpressionTextureSample::Compile(class FMaterialCompiler* Compile
 		}
 		else
 		{
-			TextureCodeIndex = Compiler->Texture(Texture, TextureReferenceIndex, SamplerSource, MipValueMode);
+			TextureCodeIndex = Compiler->Texture(Texture, TextureReferenceIndex, SamplerType, SamplerSource, MipValueMode);
 		}
 
 		if (TextureCodeIndex == INDEX_NONE)
@@ -1654,11 +1654,9 @@ int32 UMaterialExpressionTextureSample::Compile(class FMaterialCompiler* Compile
 			FMaterialUniformExpression* TextureUniformBase = Compiler->GetParameterUniformExpression(TextureCodeIndex);
 			if (FMaterialUniformExpressionTexture* TextureUniform = TextureUniformBase->GetTextureUniformExpression())
 			{
+				EffectiveSamplerType = TextureUniform->GetSamplerType();
 				TextureReferenceIndex = TextureUniform->GetTextureIndex();
 				EffectiveTexture = Compiler->GetReferencedTexture(TextureReferenceIndex);
-				// Here we assume that the effective sampler type is correct for the given texture,
-				// with the expectation that the expression that generated this texture should have already triggered an error if the sampler type didn't match
-				EffectiveSamplerType = UMaterialExpressionTextureBase::GetSamplerTypeForTexture(EffectiveTexture);
 				if (FMaterialUniformExpressionTextureParameter* TextureParameterUniform = TextureUniform->GetTextureParameterUniformExpression())
 				{
 					EffectiveParameterName = TextureParameterUniform->GetParameterName();
@@ -2012,12 +2010,13 @@ int32 UMaterialExpressionTextureObjectParameter::Compile(class FMaterialCompiler
 		return CompilerError(Compiler, GetRequirements());
 	}
 
-	if (!VerifySamplerType(Compiler, (Desc.Len() > 0 ? *Desc : TEXT("TextureObjectParameter")), Texture, SamplerType))
+	// It seems like this error should be checked here, but this can break existing materials, see https://jira.it.epicgames.net/browse/UE-68862
+	/*if (!VerifySamplerType(Compiler, (Desc.Len() > 0 ? *Desc : TEXT("TextureObjectParameter")), Texture, SamplerType))
 	{
 		return INDEX_NONE;
-	}
+	}*/
 
-	return SamplerType == SAMPLERTYPE_External ? Compiler->ExternalTextureParameter(ParameterName, Texture) : Compiler->TextureParameter(ParameterName, Texture);
+	return SamplerType == SAMPLERTYPE_External ? Compiler->ExternalTextureParameter(ParameterName, Texture) : Compiler->TextureParameter(ParameterName, Texture, SamplerType);
 }
 
 int32 UMaterialExpressionTextureObjectParameter::CompilePreview(class FMaterialCompiler* Compiler, int32 OutputIndex)
@@ -2094,12 +2093,13 @@ int32 UMaterialExpressionTextureObject::Compile(class FMaterialCompiler* Compile
 		return CompilerError(Compiler, TEXT("Requires valid texture"));
 	}
 
-	if (!VerifySamplerType(Compiler, (Desc.Len() > 0 ? *Desc : TEXT("TextureObject")), Texture, SamplerType))
+	// It seems like this error should be checked here, but this can break existing materials, see https://jira.it.epicgames.net/browse/UE-68862
+	/*if (!VerifySamplerType(Compiler, (Desc.Len() > 0 ? *Desc : TEXT("TextureObject")), Texture, SamplerType))
 	{
 		return INDEX_NONE;
-	}
+	}*/
 
-	return SamplerType == SAMPLERTYPE_External ? Compiler->ExternalTexture(Texture) : Compiler->Texture(Texture);
+	return SamplerType == SAMPLERTYPE_External ? Compiler->ExternalTexture(Texture) : Compiler->Texture(Texture, SamplerType);
 }
 
 int32 UMaterialExpressionTextureObject::CompilePreview(class FMaterialCompiler* Compiler, int32 OutputIndex)
@@ -2524,7 +2524,7 @@ int32 UMaterialExpressionTextureSampleParameterSubUV::Compile(class FMaterialCom
 		return INDEX_NONE;
 	}
 
-	int32 TextureCodeIndex = Compiler->TextureParameter(ParameterName, Texture);
+	int32 TextureCodeIndex = Compiler->TextureParameter(ParameterName, Texture, SamplerType);
 	return ParticleSubUV(Compiler, TextureCodeIndex, Texture, SamplerType, Coordinates, bBlend);
 }
 
@@ -7521,7 +7521,7 @@ int32 UMaterialExpressionParticleSubUV::Compile(class FMaterialCompiler* Compile
 		{
 			return INDEX_NONE;
 		}
-		int32 TextureCodeIndex = Compiler->Texture(Texture);
+		int32 TextureCodeIndex = Compiler->Texture(Texture, SamplerType);
 		return ParticleSubUV(Compiler, TextureCodeIndex, Texture, SamplerType, Coordinates, bBlend);
 	}
 	else
@@ -8856,7 +8856,7 @@ int32 UMaterialExpressionFontSample::Compile(class FMaterialCompiler* Compiler,
 			return INDEX_NONE;
 		}
 
-		int32 TextureCodeIndex = Compiler->Texture(Texture);
+		int32 TextureCodeIndex = Compiler->Texture(Texture, ExpectedSamplerType);
 		Result = Compiler->TextureSample(
 			TextureCodeIndex,
 			Compiler->TextureCoordinate(0, false, false),
@@ -8960,7 +8960,7 @@ int32 UMaterialExpressionFontSampleParameter::Compile(class FMaterialCompiler* C
 		{
 			return INDEX_NONE;
 		}
-		int32 TextureCodeIndex = Compiler->TextureParameter(ParameterName,Texture);
+		int32 TextureCodeIndex = Compiler->TextureParameter(ParameterName,Texture, ExpectedSamplerType);
 		Result = Compiler->TextureSample(
 			TextureCodeIndex,
 			Compiler->TextureCoordinate(0, false, false),
@@ -13514,11 +13514,11 @@ int32 UMaterialExpressionAntialiasedTextureMask::Compile(class FMaterialCompiler
 
 	if (!ParameterName.IsValid() || ParameterName.IsNone())
 	{
-		TextureCodeIndex = Compiler->Texture(Texture);
+		TextureCodeIndex = Compiler->Texture(Texture, SamplerType);
 	}
 	else
 	{
-		TextureCodeIndex = Compiler->TextureParameter(ParameterName, Texture);
+		TextureCodeIndex = Compiler->TextureParameter(ParameterName, Texture, SamplerType);
 	}
 
 	if (!VerifySamplerType(Compiler, (Desc.Len() > 0 ? *Desc : TEXT("AntialiasedTextureMask")), Texture, SamplerType))
diff --git a/Engine/Source/Runtime/Engine/Private/Materials/MaterialUniformExpressions.cpp b/Engine/Source/Runtime/Engine/Private/Materials/MaterialUniformExpressions.cpp
--- a/Engine/Source/Runtime/Engine/Private/Materials/MaterialUniformExpressions.cpp
+++ b/Engine/Source/Runtime/Engine/Private/Materials/MaterialUniformExpressions.cpp
@@ -576,13 +576,19 @@ void FUniformExpressionSet::FillUniformBuffer(const FMaterialRenderContext& Mate
 
 FMaterialUniformExpressionTexture::FMaterialUniformExpressionTexture() :
 	TextureIndex(INDEX_NONE),
+#if WITH_EDITORONLY_DATA
+	SamplerType(SAMPLERTYPE_Color),
+#endif
 	SamplerSource(SSM_FromTextureAsset),
 	TransientOverrideValue_GameThread(NULL),
 	TransientOverrideValue_RenderThread(NULL)
 {}
 
-FMaterialUniformExpressionTexture::FMaterialUniformExpressionTexture(int32 InTextureIndex, ESamplerSourceMode InSamplerSource) :
+FMaterialUniformExpressionTexture::FMaterialUniformExpressionTexture(int32 InTextureIndex, EMaterialSamplerType InSamplerType, ESamplerSourceMode InSamplerSource) :
 	TextureIndex(InTextureIndex),
+#if WITH_EDITORONLY_DATA
+	SamplerType(InSamplerType),
+#endif
 	SamplerSource(InSamplerSource),
 	TransientOverrideValue_GameThread(NULL),
 	TransientOverrideValue_RenderThread(NULL)
diff --git a/Engine/Source/Runtime/Engine/Private/Materials/MaterialUniformExpressions.h b/Engine/Source/Runtime/Engine/Private/Materials/MaterialUniformExpressions.h
--- a/Engine/Source/Runtime/Engine/Private/Materials/MaterialUniformExpressions.h
+++ b/Engine/Source/Runtime/Engine/Private/Materials/MaterialUniformExpressions.h
@@ -294,8 +294,8 @@ public:
 
 	FMaterialUniformExpressionTextureParameter() {}
 
-	FMaterialUniformExpressionTextureParameter(const FMaterialParameterInfo& InParameterInfo, int32 InTextureIndex, ESamplerSourceMode InSourceMode) :
-		Super(InTextureIndex, InSourceMode),
+	FMaterialUniformExpressionTextureParameter(const FMaterialParameterInfo& InParameterInfo, int32 InTextureIndex, EMaterialSamplerType InSamplerType, ESamplerSourceMode InSourceMode) :
+		Super(InTextureIndex, InSamplerType, InSourceMode),
 		ParameterInfo(InParameterInfo)
 	{}
 
diff --git a/Engine/Source/Runtime/Engine/Public/MaterialCompiler.h b/Engine/Source/Runtime/Engine/Public/MaterialCompiler.h
--- a/Engine/Source/Runtime/Engine/Public/MaterialCompiler.h
+++ b/Engine/Source/Runtime/Engine/Public/MaterialCompiler.h
@@ -165,8 +165,8 @@ public:
 	virtual int32 TextureDecalDerivative(bool bDDY) = 0;
 	virtual int32 DecalLifetimeOpacity() = 0;
 
-	virtual int32 Texture(UTexture* Texture,int32& TextureReferenceIndex,ESamplerSourceMode SamplerSource=SSM_FromTextureAsset,ETextureMipValueMode MipValueMode=TMVM_None) = 0;
-	virtual int32 TextureParameter(FName ParameterName,UTexture* DefaultTexture,int32& TextureReferenceIndex,ESamplerSourceMode SamplerSource=SSM_FromTextureAsset) = 0;
+	virtual int32 Texture(UTexture* Texture,int32& TextureReferenceIndex, EMaterialSamplerType SamplerType, ESamplerSourceMode SamplerSource=SSM_FromTextureAsset,ETextureMipValueMode MipValueMode=TMVM_None) = 0;
+	virtual int32 TextureParameter(FName ParameterName,UTexture* DefaultTexture,int32& TextureReferenceIndex, EMaterialSamplerType SamplerType, ESamplerSourceMode SamplerSource=SSM_FromTextureAsset) = 0;
 
 	virtual int32 ExternalTexture(const FGuid& ExternalTextureGuid) = 0;
 	virtual int32 ExternalTexture(UTexture* InTexture, int32& TextureReferenceIndex) = 0;
@@ -179,10 +179,10 @@ public:
 	virtual int32 GetTextureReferenceIndex(UTexture* Texture) { return INDEX_NONE; }
 	virtual UTexture* GetReferencedTexture(int32 Index) { return nullptr; }
 
-	int32 Texture(UTexture* InTexture,ESamplerSourceMode SamplerSource=SSM_FromTextureAsset)
+	int32 Texture(UTexture* InTexture, EMaterialSamplerType SamplerType, ESamplerSourceMode SamplerSource=SSM_FromTextureAsset)
 	{
 		int32 TextureReferenceIndex = INDEX_NONE;
-		return Texture(InTexture, TextureReferenceIndex, SamplerSource);
+		return Texture(InTexture, TextureReferenceIndex, SamplerType, SamplerSource);
 	}
 
 	int32 ExternalTexture(UTexture* DefaultTexture)
@@ -191,10 +191,10 @@ public:
 		return ExternalTexture(DefaultTexture, TextureReferenceIndex);
 	}
 
-	int32 TextureParameter(FName ParameterName,UTexture* DefaultTexture,ESamplerSourceMode SamplerSource=SSM_FromTextureAsset)
+	int32 TextureParameter(FName ParameterName,UTexture* DefaultTexture, EMaterialSamplerType SamplerType, ESamplerSourceMode SamplerSource=SSM_FromTextureAsset)
 	{
 		int32 TextureReferenceIndex = INDEX_NONE;
-		return  TextureParameter(ParameterName, DefaultTexture, TextureReferenceIndex, SamplerSource);
+		return  TextureParameter(ParameterName, DefaultTexture, TextureReferenceIndex, SamplerType, SamplerSource);
 	}
 
 	int32 ExternalTextureParameter(FName ParameterName, UTexture* DefaultTexture)
@@ -411,8 +411,8 @@ public:
 	virtual int32 TextureDecalDerivative(bool bDDY) override { return Compiler->TextureDecalDerivative(bDDY); }
 	virtual int32 DecalLifetimeOpacity() override { return Compiler->DecalLifetimeOpacity(); }
 
-	virtual int32 Texture(UTexture* InTexture,int32& TextureReferenceIndex,ESamplerSourceMode SamplerSource=SSM_FromTextureAsset,ETextureMipValueMode MipValueMode=TMVM_None) override { return Compiler->Texture(InTexture,TextureReferenceIndex,SamplerSource,MipValueMode); }
-	virtual int32 TextureParameter(FName ParameterName,UTexture* DefaultValue,int32& TextureReferenceIndex,ESamplerSourceMode SamplerSource=SSM_FromTextureAsset) override { return Compiler->TextureParameter(ParameterName,DefaultValue,TextureReferenceIndex,SamplerSource); }
+	virtual int32 Texture(UTexture* InTexture,int32& TextureReferenceIndex,EMaterialSamplerType SamplerType,ESamplerSourceMode SamplerSource=SSM_FromTextureAsset,ETextureMipValueMode MipValueMode=TMVM_None) override { return Compiler->Texture(InTexture,TextureReferenceIndex, SamplerType, SamplerSource,MipValueMode); }
+	virtual int32 TextureParameter(FName ParameterName,UTexture* DefaultValue,int32& TextureReferenceIndex, EMaterialSamplerType SamplerType, ESamplerSourceMode SamplerSource=SSM_FromTextureAsset) override { return Compiler->TextureParameter(ParameterName,DefaultValue,TextureReferenceIndex, SamplerType, SamplerSource); }
 
 	virtual int32 ExternalTexture(const FGuid& ExternalTextureGuid) override { return Compiler->ExternalTexture(ExternalTextureGuid); }
 	virtual int32 ExternalTexture(UTexture* InTexture, int32& TextureReferenceIndex) override { return Compiler->ExternalTexture(InTexture, TextureReferenceIndex); }
diff --git a/Engine/Source/Runtime/Engine/Public/MaterialShared.h b/Engine/Source/Runtime/Engine/Public/MaterialShared.h
--- a/Engine/Source/Runtime/Engine/Public/MaterialShared.h
+++ b/Engine/Source/Runtime/Engine/Public/MaterialShared.h
@@ -266,7 +266,7 @@ public:
 
 	FMaterialUniformExpressionTexture();
 
-	FMaterialUniformExpressionTexture(int32 InTextureIndex, ESamplerSourceMode InSamplerSource);
+	FMaterialUniformExpressionTexture(int32 InTextureIndex, EMaterialSamplerType InSamplerType, ESamplerSourceMode InSamplerSource);
 
 	// FMaterialUniformExpression interface.
 	virtual void Serialize(FArchive& Ar);
@@ -287,9 +287,17 @@ public:
 
 	int32 GetTextureIndex() const { return TextureIndex; }
 
+#if WITH_EDITORONLY_DATA
+	EMaterialSamplerType GetSamplerType() const { return SamplerType; }
+#endif
+
 protected:
 	/** Index into FMaterial::GetReferencedTextures */
 	int32 TextureIndex;
+#if WITH_EDITORONLY_DATA
+	/** Sampler type of the expression that generated this */
+	EMaterialSamplerType SamplerType;
+#endif
 	ESamplerSourceMode SamplerSource;
 	/** Texture that may be used in the editor for overriding the texture but never saved to disk, accessible only by the game thread! */
 	UTexture* TransientOverrideValue_GameThread;
