Validate the type that's passed in to TextureSample Tex input before trying to access uniform expression
- In general, the material editor UI attempts to prevent invalid types from being assigned
- It's possible to make invalid connection by routing values through a switch node

diff --git a/Engine/Source/Runtime/Engine/Private/Materials/MaterialExpressions.cpp b/Engine/Source/Runtime/Engine/Private/Materials/MaterialExpressions.cpp
--- a/Engine/Source/Runtime/Engine/Private/Materials/MaterialExpressions.cpp
+++ b/Engine/Source/Runtime/Engine/Private/Materials/MaterialExpressions.cpp
@@ -1649,9 +1649,17 @@ int32 UMaterialExpressionTextureSample::Compile(class FMaterialCompiler* Compile
 		TOptional<FName> EffectiveParameterName;
 		if (InputExpression)
 		{
+			const EMaterialValueType TexInputType = Compiler->GetParameterType(TextureCodeIndex);
+			if (!(TexInputType & MCT_Texture))
+			{
+				return CompilerError(Compiler, TEXT("Tex input requires a texture value"));
+			}
+
 			// If 'InputExpression' is connected, we use need to find the texture object that was passed in
 			// In this case, the texture/sampler assigned on this expression node are not used
 			FMaterialUniformExpression* TextureUniformBase = Compiler->GetParameterUniformExpression(TextureCodeIndex);
+			checkf(TextureUniformBase, TEXT("TexInputType is %d, but missing FMaterialUniformExpression"), TexInputType);
+
 			if (FMaterialUniformExpressionTexture* TextureUniform = TextureUniformBase->GetTextureUniformExpression())
 			{
 				EffectiveSamplerType = TextureUniform->GetSamplerType();
