Fixing lightmaps in DrawTileMesh. It was missing the bindings to the LightmapSceneData because no Scene is present when baking materials. Solved with a similar approach as for the PrimitiveSceneData. Also removed the test in question from the blacklist.

#rnx

diff --git a/Engine/Source/Runtime/Engine/Private/LightmapUniformShaderParameters.cpp b/Engine/Source/Runtime/Engine/Private/LightmapUniformShaderParameters.cpp
new file mode 100644
--- /dev/null
+++ b/Engine/Source/Runtime/Engine/Private/LightmapUniformShaderParameters.cpp
@@ -0,0 +1,110 @@
+// Copyright 1998-2019 Epic Games, Inc. All Rights Reserved.
+
+#include "LightmapUniformShaderParameters.h"
+#include "SceneManagement.h"
+#include "UnrealEngine.h"
+
+IMPLEMENT_GLOBAL_SHADER_PARAMETER_STRUCT(FPrecomputedLightingUniformParameters, "PrecomputedLightingBuffer");
+
+FLightmapSceneShaderData::FLightmapSceneShaderData(const class FLightCacheInterface* LCI, ERHIFeatureLevel::Type FeatureLevel)
+{
+	FPrecomputedLightingUniformParameters Parameters;
+	GetPrecomputedLightingParameters(FeatureLevel, Parameters, LCI);
+	Setup(Parameters);
+}
+
+void FLightmapSceneShaderData::Setup(const FPrecomputedLightingUniformParameters& ShaderParameters)
+{
+	static_assert(sizeof(FPrecomputedLightingUniformParameters) == 128, "The FLightmapSceneShaderData manual layout below and in usf must match FPrecomputedLightingUniformParameters.  Update this assert when adding a new member.");
+	// Note: layout must match GetLightmapData in usf
+
+	Data[0] = ShaderParameters.StaticShadowMapMasks;
+	Data[1] = ShaderParameters.InvUniformPenumbraSizes;
+	Data[2] = ShaderParameters.LightMapCoordinateScaleBias;
+	Data[3] = ShaderParameters.ShadowMapCoordinateScaleBias;
+	Data[4] = ShaderParameters.LightMapScale[0];
+	Data[5] = ShaderParameters.LightMapScale[1];
+	Data[6] = ShaderParameters.LightMapAdd[0];
+	Data[7] = ShaderParameters.LightMapAdd[1];
+}
+
+void GetDefaultPrecomputedLightingParameters(FPrecomputedLightingUniformParameters& Parameters)
+{
+	Parameters.StaticShadowMapMasks = FVector4(1, 1, 1, 1);
+	Parameters.InvUniformPenumbraSizes = FVector4(0, 0, 0, 0);
+	Parameters.LightMapCoordinateScaleBias = FVector4(1, 1, 0, 0);
+	Parameters.ShadowMapCoordinateScaleBias = FVector4(1, 1, 0, 0);
+	 
+	const uint32 NumCoef = FMath::Max<uint32>(NUM_HQ_LIGHTMAP_COEF, NUM_LQ_LIGHTMAP_COEF);
+	for (uint32 CoefIndex = 0; CoefIndex < NumCoef; ++CoefIndex)
+	{
+		Parameters.LightMapScale[CoefIndex] = FVector4(1, 1, 1, 1);
+		Parameters.LightMapAdd[CoefIndex] = FVector4(0, 0, 0, 0);
+	}
+}
+
+void GetPrecomputedLightingParameters(
+	ERHIFeatureLevel::Type FeatureLevel,
+	FPrecomputedLightingUniformParameters& Parameters,
+	const FLightCacheInterface* LCI
+)
+{
+	// TDistanceFieldShadowsAndLightMapPolicy
+	const FShadowMapInteraction ShadowMapInteraction = LCI ? LCI->GetShadowMapInteraction() : FShadowMapInteraction();
+	if (ShadowMapInteraction.GetType() == SMIT_Texture)
+	{
+		Parameters.ShadowMapCoordinateScaleBias = FVector4(ShadowMapInteraction.GetCoordinateScale(), ShadowMapInteraction.GetCoordinateBias());
+		Parameters.StaticShadowMapMasks = FVector4(ShadowMapInteraction.GetChannelValid(0), ShadowMapInteraction.GetChannelValid(1), ShadowMapInteraction.GetChannelValid(2), ShadowMapInteraction.GetChannelValid(3));
+		Parameters.InvUniformPenumbraSizes = ShadowMapInteraction.GetInvUniformPenumbraSize();
+	}
+	else
+	{
+		Parameters.ShadowMapCoordinateScaleBias = FVector4(1, 1, 0, 0);
+		Parameters.StaticShadowMapMasks = FVector4(1, 1, 1, 1);
+		Parameters.InvUniformPenumbraSizes = FVector4(0, 0, 0, 0);
+	}
+
+	// TLightMapPolicy
+	const FLightMapInteraction LightMapInteraction = LCI ? LCI->GetLightMapInteraction(FeatureLevel) : FLightMapInteraction();
+	if (LightMapInteraction.GetType() == LMIT_Texture)
+	{
+		const bool bAllowHighQualityLightMaps = AllowHighQualityLightmaps(FeatureLevel) && LightMapInteraction.AllowsHighQualityLightmaps();
+
+		// Vertex Shader
+		const FVector2D LightmapCoordinateScale = LightMapInteraction.GetCoordinateScale();
+		const FVector2D LightmapCoordinateBias = LightMapInteraction.GetCoordinateBias();
+		Parameters.LightMapCoordinateScaleBias = FVector4(LightmapCoordinateScale.X, LightmapCoordinateScale.Y, LightmapCoordinateBias.X, LightmapCoordinateBias.Y);
+
+		static const auto CVar = IConsoleManager::Get().FindTConsoleVariableDataInt(TEXT("r.VirtualTexturedLightmaps"));
+		if (CVar->GetValueOnRenderThread() == 0)
+		{
+
+		}
+		else
+		{
+			checkf(0, TEXT("VT needs to be implemented with Mesh Draw Command pipeline"));
+		}
+
+		const uint32 NumCoef = bAllowHighQualityLightMaps ? NUM_HQ_LIGHTMAP_COEF : NUM_LQ_LIGHTMAP_COEF;
+		const FVector4* Scales = LightMapInteraction.GetScaleArray();
+		const FVector4* Adds = LightMapInteraction.GetAddArray();
+		for (uint32 CoefIndex = 0; CoefIndex < NumCoef; ++CoefIndex)
+		{
+			Parameters.LightMapScale[CoefIndex] = Scales[CoefIndex];
+			Parameters.LightMapAdd[CoefIndex] = Adds[CoefIndex];
+		}
+	}
+	else
+	{
+		// Vertex Shader
+		Parameters.LightMapCoordinateScaleBias = FVector4(1, 1, 0, 0);
+
+		// Pixel Shader
+		const uint32 NumCoef = FMath::Max<uint32>(NUM_HQ_LIGHTMAP_COEF, NUM_LQ_LIGHTMAP_COEF);
+		for (uint32 CoefIndex = 0; CoefIndex < NumCoef; ++CoefIndex)
+		{
+			Parameters.LightMapScale[CoefIndex] = FVector4(1, 1, 1, 1);
+			Parameters.LightMapAdd[CoefIndex] = FVector4(0, 0, 0, 0);
+		}
+	}
+}
\ No newline at end of file
diff --git a/Engine/Source/Runtime/Engine/Private/PrimitiveUniformShaderParameters.cpp b/Engine/Source/Runtime/Engine/Private/PrimitiveUniformShaderParameters.cpp
--- a/Engine/Source/Runtime/Engine/Private/PrimitiveUniformShaderParameters.cpp
+++ b/Engine/Source/Runtime/Engine/Private/PrimitiveUniformShaderParameters.cpp
@@ -6,16 +6,31 @@
 
 void FSinglePrimitiveStructuredBuffer::InitRHI() 
 {
-	FRHIResourceCreateInfo CreateInfo;
-	BufferRHI = RHICreateStructuredBuffer(sizeof(FVector4), FPrimitiveSceneShaderData::PrimitiveDataStrideInFloat4s * sizeof(FVector4), BUF_Static | BUF_ShaderResource, CreateInfo);
+	{
+		FRHIResourceCreateInfo CreateInfo;
+		PrimitiveSceneDataBufferRHI = RHICreateStructuredBuffer(sizeof(FVector4), FPrimitiveSceneShaderData::PrimitiveDataStrideInFloat4s * sizeof(FVector4), BUF_Static | BUF_ShaderResource, CreateInfo);
 
-	void* LockedData = RHILockStructuredBuffer(BufferRHI, 0, FPrimitiveSceneShaderData::PrimitiveDataStrideInFloat4s * sizeof(FVector4), RLM_WriteOnly);
+		void* LockedData = RHILockStructuredBuffer(PrimitiveSceneDataBufferRHI, 0, FPrimitiveSceneShaderData::PrimitiveDataStrideInFloat4s * sizeof(FVector4), RLM_WriteOnly);
 
-	FPlatformMemory::Memcpy(LockedData, PrimitiveSceneData.Data, FPrimitiveSceneShaderData::PrimitiveDataStrideInFloat4s * sizeof(FVector4));
+		FPlatformMemory::Memcpy(LockedData, PrimitiveSceneData.Data, FPrimitiveSceneShaderData::PrimitiveDataStrideInFloat4s * sizeof(FVector4));
 
-	RHIUnlockStructuredBuffer(BufferRHI);
+		RHIUnlockStructuredBuffer(PrimitiveSceneDataBufferRHI);
 
-	BufferSRV = RHICreateShaderResourceView(BufferRHI);
+		PrimitiveSceneDataBufferSRV = RHICreateShaderResourceView(PrimitiveSceneDataBufferRHI);
+	}
+
+	{
+		FRHIResourceCreateInfo CreateInfo;
+		LightmapSceneDataBufferRHI = RHICreateStructuredBuffer(sizeof(FVector4), FLightmapSceneShaderData::LightmapDataStrideInFloat4s * sizeof(FVector4), BUF_Static | BUF_ShaderResource, CreateInfo);
+
+		void* LockedData = RHILockStructuredBuffer(LightmapSceneDataBufferRHI, 0, FLightmapSceneShaderData::LightmapDataStrideInFloat4s * sizeof(FVector4), RLM_WriteOnly);
+
+		FPlatformMemory::Memcpy(LockedData, LightmapSceneData.Data, FLightmapSceneShaderData::LightmapDataStrideInFloat4s * sizeof(FVector4));
+
+		RHIUnlockStructuredBuffer(LightmapSceneDataBufferRHI);
+
+		LightmapSceneDataBufferSRV = RHICreateShaderResourceView(LightmapSceneDataBufferRHI);
+	}
 }
 
 TGlobalResource<FSinglePrimitiveStructuredBuffer> GIdentityPrimitiveBuffer;
diff --git a/Engine/Source/Runtime/Engine/Private/SceneManagement.cpp b/Engine/Source/Runtime/Engine/Private/SceneManagement.cpp
--- a/Engine/Source/Runtime/Engine/Private/SceneManagement.cpp
+++ b/Engine/Source/Runtime/Engine/Private/SceneManagement.cpp
@@ -713,8 +713,8 @@ FViewUniformShaderParameters::FViewUniformShaderParameters()
 	PreIntegratedBRDF = GWhiteTexture->TextureRHI;
 	PreIntegratedBRDFSampler = TStaticSamplerState<SF_Bilinear, AM_Clamp, AM_Clamp, AM_Clamp>::GetRHI();
 
-	PrimitiveSceneData = GIdentityPrimitiveBuffer.BufferSRV;
-	LightmapSceneData = GIdentityPrimitiveBuffer.BufferSRV;
+	PrimitiveSceneData = GIdentityPrimitiveBuffer.PrimitiveSceneDataBufferSRV;
+	LightmapSceneData = GIdentityPrimitiveBuffer.LightmapSceneDataBufferSRV;
 }
 
 FInstancedViewUniformShaderParameters::FInstancedViewUniformShaderParameters()
@@ -751,95 +751,6 @@ void InitializeSharedSamplerStates()
 	}
 }
 
-IMPLEMENT_GLOBAL_SHADER_PARAMETER_STRUCT(FPrecomputedLightingUniformParameters, "PrecomputedLightingBuffer");
-
-FLightmapSceneShaderData::FLightmapSceneShaderData(const class FLightCacheInterface* LCI, ERHIFeatureLevel::Type FeatureLevel)
-{
-	FPrecomputedLightingUniformParameters Parameters;
-	GetPrecomputedLightingParameters(FeatureLevel, Parameters, LCI);
-	Setup(Parameters);
-}
-
-void FLightmapSceneShaderData::Setup(const FPrecomputedLightingUniformParameters& ShaderParameters)
-{
-	static_assert(sizeof(FPrecomputedLightingUniformParameters) == 128, "The FLightmapSceneShaderData manual layout below and in usf must match FPrecomputedLightingUniformParameters.  Update this assert when adding a new member.");
-	// Note: layout must match GetLightmapData in usf
-
-	Data[0] = ShaderParameters.StaticShadowMapMasks;
-	Data[1] = ShaderParameters.InvUniformPenumbraSizes;
-	Data[2] = ShaderParameters.LightMapCoordinateScaleBias;
-	Data[3] = ShaderParameters.ShadowMapCoordinateScaleBias;
-	Data[4] = ShaderParameters.LightMapScale[0];
-	Data[5] = ShaderParameters.LightMapScale[1];
-	Data[6] = ShaderParameters.LightMapAdd[0];
-	Data[7] = ShaderParameters.LightMapAdd[1];
-}
-
-void GetPrecomputedLightingParameters(
-	ERHIFeatureLevel::Type FeatureLevel,
-	FPrecomputedLightingUniformParameters& Parameters, 
-	const FLightCacheInterface* LCI
-	)
-{
-	// TDistanceFieldShadowsAndLightMapPolicy
-	const FShadowMapInteraction ShadowMapInteraction = LCI ? LCI->GetShadowMapInteraction() : FShadowMapInteraction(); 
-	if (ShadowMapInteraction.GetType() == SMIT_Texture)
-	{
-		Parameters.ShadowMapCoordinateScaleBias = FVector4(ShadowMapInteraction.GetCoordinateScale(), ShadowMapInteraction.GetCoordinateBias());
-		Parameters.StaticShadowMapMasks = FVector4(ShadowMapInteraction.GetChannelValid(0), ShadowMapInteraction.GetChannelValid(1), ShadowMapInteraction.GetChannelValid(2), ShadowMapInteraction.GetChannelValid(3));
-		Parameters.InvUniformPenumbraSizes = ShadowMapInteraction.GetInvUniformPenumbraSize();
-	}
-	else
-	{
-		Parameters.StaticShadowMapMasks = FVector4(1, 1, 1, 1);
-		Parameters.InvUniformPenumbraSizes = FVector4(0, 0, 0, 0);
-	}
-
-	// TLightMapPolicy
-	const FLightMapInteraction LightMapInteraction = LCI ? LCI->GetLightMapInteraction(FeatureLevel) : FLightMapInteraction();
-	if (LightMapInteraction.GetType() == LMIT_Texture)
-	{
-		const bool bAllowHighQualityLightMaps = AllowHighQualityLightmaps(FeatureLevel) && LightMapInteraction.AllowsHighQualityLightmaps();
-
-		// Vertex Shader
-		const FVector2D LightmapCoordinateScale = LightMapInteraction.GetCoordinateScale();
-		const FVector2D LightmapCoordinateBias = LightMapInteraction.GetCoordinateBias();
-		Parameters.LightMapCoordinateScaleBias = FVector4(LightmapCoordinateScale.X, LightmapCoordinateScale.Y, LightmapCoordinateBias.X, LightmapCoordinateBias.Y);
-
-		static const auto CVar = IConsoleManager::Get().FindTConsoleVariableDataInt(TEXT("r.VirtualTexturedLightmaps"));
-		if (CVar->GetValueOnRenderThread() == 0)
-		{
-
-		}
-		else
-		{
-			checkf(0, TEXT("VT needs to be implemented with Mesh Draw Command pipeline"));
-		}
-
-		const uint32 NumCoef = bAllowHighQualityLightMaps ? NUM_HQ_LIGHTMAP_COEF : NUM_LQ_LIGHTMAP_COEF;
-		const FVector4* Scales = LightMapInteraction.GetScaleArray();
-		const FVector4* Adds = LightMapInteraction.GetAddArray();
-		for (uint32 CoefIndex = 0; CoefIndex < NumCoef; ++CoefIndex)
-		{
-			Parameters.LightMapScale[CoefIndex] = Scales[CoefIndex];
-			Parameters.LightMapAdd[CoefIndex] = Adds[CoefIndex];
-		}
-	}
-	else
-	{
-		// Vertex Shader
-		Parameters.LightMapCoordinateScaleBias = FVector4(1, 1, 0, 0);
-
-		// Pixel Shader
-		const uint32 NumCoef = FMath::Max<uint32>(NUM_HQ_LIGHTMAP_COEF, NUM_LQ_LIGHTMAP_COEF);
-		for (uint32 CoefIndex = 0; CoefIndex < NumCoef; ++CoefIndex)
-		{
-			Parameters.LightMapScale[CoefIndex] = FVector4(1, 1, 1, 1);
-			Parameters.LightMapAdd[CoefIndex] = FVector4(0, 0, 0, 0);
-		}
-	}
-}
-
 void FLightCacheInterface::CreatePrecomputedLightingUniformBuffer_RenderingThread(ERHIFeatureLevel::Type FeatureLevel)
 {
 	if (LightMap || ShadowMap)
diff --git a/Engine/Source/Runtime/Engine/Public/LightmapUniformShaderParameters.h b/Engine/Source/Runtime/Engine/Public/LightmapUniformShaderParameters.h
new file mode 100644
--- /dev/null
+++ b/Engine/Source/Runtime/Engine/Public/LightmapUniformShaderParameters.h
@@ -0,0 +1,54 @@
+// Copyright 1998-2019 Epic Games, Inc. All Rights Reserved.
+
+#pragma once
+
+#include "CoreMinimal.h"
+#include "Engine/EngineTypes.h"
+#include "RenderResource.h"
+#include "ShaderParameters.h"
+#include "UniformBuffer.h"
+
+class FLightCacheInterface;
+
+/** The maximum value between NUM_LQ_LIGHTMAP_COEF and NUM_HQ_LIGHTMAP_COEF. */
+static const int32 MAX_NUM_LIGHTMAP_COEF = 2;
+
+BEGIN_GLOBAL_SHADER_PARAMETER_STRUCT(FPrecomputedLightingUniformParameters, ENGINE_API)
+SHADER_PARAMETER(FVector4, StaticShadowMapMasks) // TDistanceFieldShadowsAndLightMapPolicy
+SHADER_PARAMETER(FVector4, InvUniformPenumbraSizes) // TDistanceFieldShadowsAndLightMapPolicy
+SHADER_PARAMETER(FVector4, LightMapCoordinateScaleBias) // TLightMapPolicy
+SHADER_PARAMETER(FVector4, ShadowMapCoordinateScaleBias) // TDistanceFieldShadowsAndLightMapPolicy
+SHADER_PARAMETER_ARRAY_EX(FVector4, LightMapScale, [MAX_NUM_LIGHTMAP_COEF], EShaderPrecisionModifier::Half) // TLightMapPolicy
+SHADER_PARAMETER_ARRAY_EX(FVector4, LightMapAdd, [MAX_NUM_LIGHTMAP_COEF], EShaderPrecisionModifier::Half) // TLightMapPolicy
+END_GLOBAL_SHADER_PARAMETER_STRUCT()
+
+ENGINE_API void GetDefaultPrecomputedLightingParameters(FPrecomputedLightingUniformParameters& Parameters);
+
+ENGINE_API void GetPrecomputedLightingParameters(
+	ERHIFeatureLevel::Type FeatureLevel,
+	FPrecomputedLightingUniformParameters& Parameters,
+	const FLightCacheInterface* LCI);
+
+struct FLightmapSceneShaderData
+{
+	// Must match usf
+	enum { LightmapDataStrideInFloat4s = 8 };
+
+	FVector4 Data[LightmapDataStrideInFloat4s];
+
+	FLightmapSceneShaderData()
+	{
+		FPrecomputedLightingUniformParameters ShaderParameters;
+		GetDefaultPrecomputedLightingParameters(ShaderParameters);
+		Setup(ShaderParameters);
+	}
+
+	explicit FLightmapSceneShaderData(const FPrecomputedLightingUniformParameters& ShaderParameters)
+	{
+		Setup(ShaderParameters);
+	}
+
+	ENGINE_API FLightmapSceneShaderData(const class FLightCacheInterface* LCI, ERHIFeatureLevel::Type FeatureLevel);
+
+	ENGINE_API void Setup(const FPrecomputedLightingUniformParameters& ShaderParameters);
+};
diff --git a/Engine/Source/Runtime/Engine/Public/PrimitiveUniformShaderParameters.h b/Engine/Source/Runtime/Engine/Public/PrimitiveUniformShaderParameters.h
--- a/Engine/Source/Runtime/Engine/Public/PrimitiveUniformShaderParameters.h
+++ b/Engine/Source/Runtime/Engine/Public/PrimitiveUniformShaderParameters.h
@@ -8,7 +8,7 @@
 #include "RenderResource.h"
 #include "ShaderParameters.h"
 #include "UniformBuffer.h"
-
+#include "LightmapUniformShaderParameters.h"
 
 /** 
  * The uniform shader parameters associated with a primitive. 
@@ -181,14 +181,20 @@ public:
 
 	virtual void ReleaseRHI() override
 	{
-		BufferSRV.SafeRelease();
-		BufferRHI.SafeRelease();
+		PrimitiveSceneDataBufferRHI.SafeRelease();
+		PrimitiveSceneDataBufferSRV.SafeRelease();
+		LightmapSceneDataBufferRHI.SafeRelease();
+		LightmapSceneDataBufferSRV.SafeRelease();
 	}
 
 	FPrimitiveSceneShaderData PrimitiveSceneData;
+	FLightmapSceneShaderData LightmapSceneData;
+
+	FStructuredBufferRHIRef PrimitiveSceneDataBufferRHI;
+	FShaderResourceViewRHIRef PrimitiveSceneDataBufferSRV;
 
-	FStructuredBufferRHIRef BufferRHI;
-	FShaderResourceViewRHIRef BufferSRV;
+	FStructuredBufferRHIRef LightmapSceneDataBufferRHI;
+	FShaderResourceViewRHIRef LightmapSceneDataBufferSRV;
 };
 
 /**
diff --git a/Engine/Source/Runtime/Engine/Public/SceneManagement.h b/Engine/Source/Runtime/Engine/Public/SceneManagement.h
--- a/Engine/Source/Runtime/Engine/Public/SceneManagement.h
+++ b/Engine/Source/Runtime/Engine/Public/SceneManagement.h
@@ -25,6 +25,7 @@
 #include "BatchedElements.h"
 #include "MeshBatch.h"
 #include "SceneUtils.h"
+#include "LightmapUniformShaderParameters.h"
 
 #ifndef ENVIRONMENT_TEXTURE_ARRAY_WORKAROUND // RHI_RAYTRACING
 #define ENVIRONMENT_TEXTURE_ARRAY_WORKAROUND	1
@@ -299,9 +300,6 @@ static const int32 NUM_LQ_LIGHTMAP_COEF = 2;
 /** The index at which simple coefficients are stored in any array containing all NUM_STORED_LIGHTMAP_COEF coefficients. */ 
 static const int32 LQ_LIGHTMAP_COEF_INDEX = 2;
 
-/** The maximum value between NUM_LQ_LIGHTMAP_COEF and NUM_HQ_LIGHTMAP_COEF. */ 
-static const int32 MAX_NUM_LIGHTMAP_COEF = 2;
-
 /** Compile out low quality lightmaps to save memory */
 // @todo-mobile: Need to fix this!
 #define ALLOW_LQ_LIGHTMAPS (PLATFORM_DESKTOP || PLATFORM_IOS || PLATFORM_ANDROID || PLATFORM_HTML5 || PLATFORM_SWITCH || PLATFORM_LUMINGL4)
@@ -636,36 +634,7 @@ class FShadowMap;
 // 16bbp is limited to 64*64 pools
 #define LIGHTMAP_VT_16BIT 1
 
-BEGIN_GLOBAL_SHADER_PARAMETER_STRUCT(FPrecomputedLightingUniformParameters,ENGINE_API)
-	SHADER_PARAMETER(FVector4, StaticShadowMapMasks) // TDistanceFieldShadowsAndLightMapPolicy
-	SHADER_PARAMETER(FVector4, InvUniformPenumbraSizes) // TDistanceFieldShadowsAndLightMapPolicy
-	SHADER_PARAMETER(FVector4, LightMapCoordinateScaleBias) // TLightMapPolicy
-	SHADER_PARAMETER(FVector4, ShadowMapCoordinateScaleBias) // TDistanceFieldShadowsAndLightMapPolicy
-	SHADER_PARAMETER_ARRAY_EX(FVector4, LightMapScale, [MAX_NUM_LIGHTMAP_COEF], EShaderPrecisionModifier::Half) // TLightMapPolicy
-	SHADER_PARAMETER_ARRAY_EX(FVector4, LightMapAdd, [MAX_NUM_LIGHTMAP_COEF], EShaderPrecisionModifier::Half) // TLightMapPolicy
-END_GLOBAL_SHADER_PARAMETER_STRUCT()
-
-struct FLightmapSceneShaderData
-{
-	// Must match usf
-	enum { LightmapDataStrideInFloat4s = 8 };
-
-	FVector4 Data[LightmapDataStrideInFloat4s];
-
-	explicit FLightmapSceneShaderData(const FPrecomputedLightingUniformParameters& ShaderParameters)
-	{
-		Setup(ShaderParameters);
-	}
-
-	ENGINE_API FLightmapSceneShaderData(const class FLightCacheInterface* LCI, ERHIFeatureLevel::Type FeatureLevel);
 
-	ENGINE_API void Setup(const FPrecomputedLightingUniformParameters& ShaderParameters);
-};
-
-ENGINE_API void GetPrecomputedLightingParameters(
-	ERHIFeatureLevel::Type FeatureLevel,
-	FPrecomputedLightingUniformParameters& Parameters, 
-	const FLightCacheInterface* LCI);
 
 BEGIN_GLOBAL_SHADER_PARAMETER_STRUCT(FLightmapResourceClusterShaderParameters,ENGINE_API)
 	SHADER_PARAMETER_TEXTURE(Texture2D, LightMapTexture) 
diff --git a/Engine/Source/Runtime/Renderer/Private/Renderer.cpp b/Engine/Source/Runtime/Renderer/Private/Renderer.cpp
--- a/Engine/Source/Runtime/Renderer/Private/Renderer.cpp
+++ b/Engine/Source/Runtime/Renderer/Private/Renderer.cpp
@@ -121,10 +121,22 @@ void FRendererModule::DrawTileMesh(FRHICommandListImmediate& RHICmdList, FMeshPa
 				checkf(MeshElement.NumInstances == 1, TEXT("DrawTileMesh does not currently support instancing"));
 				// Force PrimitiveId to be 0 in the shader
 				MeshElement.PrimitiveIdMode = PrimID_ForceZero;
+				
+				// Set the LightmapID to 0, since that's where our light map data resides for this primitive
+				FPrimitiveUniformShaderParameters PrimitiveParams = *(const FPrimitiveUniformShaderParameters*)MeshElement.PrimitiveUniformBufferResource->GetContents();
+				PrimitiveParams.LightmapDataIndex = 0;
+
 				// Now we just need to fill out the first entry of primitive data in a buffer and bind it
-				SinglePrimitiveStructuredBuffer.PrimitiveSceneData = FPrimitiveSceneShaderData(*(const FPrimitiveUniformShaderParameters*)MeshElement.PrimitiveUniformBufferResource->GetContents());
+				SinglePrimitiveStructuredBuffer.PrimitiveSceneData = FPrimitiveSceneShaderData(PrimitiveParams);
+
+				// Set up the parameters for the LightmapSceneData from the given LCI data 
+				FPrecomputedLightingUniformParameters LightmapParams;
+				GetPrecomputedLightingParameters(FeatureLevel, LightmapParams, Mesh.LCI);
+				SinglePrimitiveStructuredBuffer.LightmapSceneData = FLightmapSceneShaderData(LightmapParams);
+
 				SinglePrimitiveStructuredBuffer.InitResource();
-				View.PrimitiveSceneDataOverrideSRV = SinglePrimitiveStructuredBuffer.BufferSRV;
+				View.PrimitiveSceneDataOverrideSRV = SinglePrimitiveStructuredBuffer.PrimitiveSceneDataBufferSRV;
+				View.LightmapSceneDataOverrideSRV = SinglePrimitiveStructuredBuffer.LightmapSceneDataBufferSRV;
 			}
 		}
 		
diff --git a/Engine/Source/Runtime/Renderer/Private/SceneRendering.cpp b/Engine/Source/Runtime/Renderer/Private/SceneRendering.cpp
--- a/Engine/Source/Runtime/Renderer/Private/SceneRendering.cpp
+++ b/Engine/Source/Runtime/Renderer/Private/SceneRendering.cpp
@@ -886,6 +886,7 @@ void FViewInfo::Init()
 	}
 
 	PrimitiveSceneDataOverrideSRV = nullptr;
+	LightmapSceneDataOverrideSRV = nullptr;
 
 	DitherFadeInUniformBuffer = nullptr;
 	DitherFadeOutUniformBuffer = nullptr;
@@ -1454,7 +1455,11 @@ void FViewInfo::SetupUniformBufferParameters(
 			}
 		}
 
-		if (Scene && Scene->GPUScene.LightmapDataBuffer.SRV)
+		if (LightmapSceneDataOverrideSRV)
+		{
+			ViewUniformShaderParameters.LightmapSceneData = LightmapSceneDataOverrideSRV;
+		}
+		else if (Scene && Scene->GPUScene.LightmapDataBuffer.SRV)
 		{
 			ViewUniformShaderParameters.LightmapSceneData = Scene->GPUScene.LightmapDataBuffer.SRV;
 		}
diff --git a/Engine/Source/Runtime/Renderer/Private/SceneRendering.h b/Engine/Source/Runtime/Renderer/Private/SceneRendering.h
--- a/Engine/Source/Runtime/Renderer/Private/SceneRendering.h
+++ b/Engine/Source/Runtime/Renderer/Private/SceneRendering.h
@@ -1042,6 +1042,7 @@ public:
 	TArray<FPrimitiveSceneInfo*, SceneRenderingAllocator> IndirectShadowPrimitives;
 
 	FShaderResourceViewRHIRef PrimitiveSceneDataOverrideSRV;
+	FShaderResourceViewRHIRef LightmapSceneDataOverrideSRV;
 
 #if RHI_RAYTRACING
 	TArray<FRayTracingGeometryInstance, SceneRenderingAllocator> RayTracingGeometryInstances;
