Fixed GPUSprite crash by storing per view per frame particle uniform buffer inside mesh collector's one frame resources.
* During InitViews FGPUSpriteEmitterDynamicUniformBufferRef was created and passed into a mesh draw command through vertex factory.
* During shadow view init it was destroyed by overriding it with a new uniform buffer, causing a crash when original draw was executed.

diff --git a/Engine/Source/Runtime/Engine/Private/Particles/ParticleGpuSimulation.cpp b/Engine/Source/Runtime/Engine/Private/Particles/ParticleGpuSimulation.cpp
--- a/Engine/Source/Runtime/Engine/Private/Particles/ParticleGpuSimulation.cpp
+++ b/Engine/Source/Runtime/Engine/Private/Particles/ParticleGpuSimulation.cpp
@@ -2866,6 +2866,7 @@ class FGPUSpriteCollectorResources : public FOneFrameResource
 {
 public:
 	FGPUSpriteVertexFactory *VertexFactory;
+	FGPUSpriteEmitterDynamicUniformBufferRef UniformBuffer;
 
 	~FGPUSpriteCollectorResources()
 	{
@@ -3059,13 +3060,6 @@ public:
 				Proxy->GetObjectPositionAndScale(*View,ObjectNDCPosition, ObjectMacroUVScales);
 				PerViewDynamicParameters.MacroUVParameters = FVector4(ObjectNDCPosition.X, ObjectNDCPosition.Y, ObjectMacroUVScales.X, ObjectMacroUVScales.Y); 
 
-				FGPUSpriteEmitterDynamicUniformBufferRef LocalDynamicUniformBuffer;
-				// Do here rather than in CreateRenderThreadResources because in some cases Render can be called before CreateRenderThreadResources
-				{
-					// Create per-emitter uniform buffer for dynamic parameters
-					LocalDynamicUniformBuffer = FGPUSpriteEmitterDynamicUniformBufferRef::CreateUniformBufferImmediate(PerViewDynamicParameters, UniformBuffer_SingleFrame);
-				}
-
 				if (bUseLocalSpace == false)
 				{
 					Proxy->UpdateWorldSpacePrimitiveUniformBuffer();
@@ -3083,6 +3077,10 @@ public:
 				CollectorResources.VertexFactory = static_cast<FGPUSpriteVertexFactory*>(InVertexFactory);
 				FGPUSpriteVertexFactory& VertexFactory = *CollectorResources.VertexFactory;
 
+				// Do here rather than in CreateRenderThreadResources because in some cases Render can be called before CreateRenderThreadResources
+				// Create per-emitter uniform buffer for dynamic parameters
+				CollectorResources.UniformBuffer = FGPUSpriteEmitterDynamicUniformBufferRef::CreateUniformBufferImmediate(PerViewDynamicParameters, UniformBuffer_SingleFrame);
+
 				if (bAllowSorting && SortMode == PSORTMODE_DistanceToView)
 				{
 					// Extensibility TODO: This call to AddSortedGPUSimulation is very awkward. When rendering a frame we need to
@@ -3108,7 +3106,7 @@ public:
 					FParticleStateTextures& StateTextures = ParticleSimulationResources->GetVisualizeStateTextures();
 							
 					VertexFactory.EmitterUniformBuffer = Resources->UniformBuffer;
-					VertexFactory.EmitterDynamicUniformBuffer = LocalDynamicUniformBuffer;
+					VertexFactory.EmitterDynamicUniformBuffer = CollectorResources.UniformBuffer;
 					VertexFactory.PositionTextureRHI = StateTextures.PositionTextureRHI;
 					VertexFactory.VelocityTextureRHI = StateTextures.VelocityTextureRHI;
 					VertexFactory.AttributesTextureRHI = ParticleSimulationResources->RenderAttributesTexture.TextureRHI;
