DR - Missing return caused not yet finished compiling Decal Materials to choose a DefaultMaterial from the Surface Domain. which eventually caused some havok later on.

diff --git a/Engine/Source/Runtime/Engine/Private/Materials/MaterialInstance.cpp b/Engine/Source/Runtime/Engine/Private/Materials/MaterialInstance.cpp
--- a/Engine/Source/Runtime/Engine/Private/Materials/MaterialInstance.cpp
+++ b/Engine/Source/Runtime/Engine/Private/Materials/MaterialInstance.cpp
@@ -139,7 +139,7 @@ const FMaterial& FMaterialInstanceResource::GetMaterialWithFallback(ERHIFeatureL
 					UMaterial* FallbackMaterial = UMaterial::GetDefaultMaterial(Domain);
 					//there was an error, use the default material's resource
 					OutFallbackMaterialRenderProxy = FallbackMaterial->GetRenderProxy();
-					OutFallbackMaterialRenderProxy->GetMaterialWithFallback(InFeatureLevel, OutFallbackMaterialRenderProxy);
+					return OutFallbackMaterialRenderProxy->GetMaterialWithFallback(InFeatureLevel, OutFallbackMaterialRenderProxy);
 				}
 			}
 		}
