Workaround for Xbox One shader compilation error

diff --git a/Engine/Source/Runtime/Engine/Public/GeometryCacheVertexFactory.h b/Engine/Source/Runtime/Engine/Public/GeometryCacheVertexFactory.h
--- a/Engine/Source/Runtime/Engine/Public/GeometryCacheVertexFactory.h
+++ b/Engine/Source/Runtime/Engine/Public/GeometryCacheVertexFactory.h
@@ -25,6 +25,8 @@ END_GLOBAL_SHADER_PARAMETER_STRUCT()
 typedef TUniformBufferRef<FGeometryCacheVertexFactoryUniformBufferParameters> FGeometryCacheVertexFactoryUniformBufferParametersRef;
 
 BEGIN_GLOBAL_SHADER_PARAMETER_STRUCT(FGeometryCacheManualVertexFetchUniformBufferParameters, ENGINE_API)
+	// FIXME: UE-69107 error X3000: syntax error: unexpected token ','
+	SHADER_PARAMETER(float, Dummy)
 	SHADER_PARAMETER_SRV(Buffer<float>, Position)
 	SHADER_PARAMETER_SRV(Buffer<float>, MotionBlurData)
 	SHADER_PARAMETER_SRV(Buffer<half4>, TangentX)
