DR - Fix out of bounds access when cloth is used with the skincache.

diff --git a/Engine/Source/Runtime/Engine/Private/GPUSkinCache.cpp b/Engine/Source/Runtime/Engine/Private/GPUSkinCache.cpp
--- a/Engine/Source/Runtime/Engine/Private/GPUSkinCache.cpp
+++ b/Engine/Source/Runtime/Engine/Private/GPUSkinCache.cpp
@@ -301,7 +301,7 @@ public:
 
 			Data.MorphBufferOffset = Section->BaseVertexIndex;
 		}
-		if (ClothVertexBuffer)
+		if (ClothVertexBuffer && ClothVertexBuffer->GetClothIndexMapping().Num() > SectionIndex)
 		{
 			Data.ClothBufferOffset = (ClothVertexBuffer->GetClothIndexMapping()[SectionIndex] & 0xFFFFFFFF);
 		}
@@ -1099,7 +1099,10 @@ void FGPUSkinCache::ProcessEntry(FRHICommandListImmediate& RHICmdList, FGPUBaseS
         checkSlow(Stride*VertexAndNormalData.GetNumVertices() == sizeof(FClothSimulEntry) * SimData->Positions.Num());
         check(sizeof(FClothSimulEntry) == 6 * sizeof(float));
 		
-		InOutEntry->DispatchData[Section].ClothBufferOffset = (ClothVertexBuffer->GetClothIndexMapping()[Section] & 0xFFFFFFFF);
+		if (ClothVertexBuffer && ClothVertexBuffer->GetClothIndexMapping().Num() > Section)
+		{
+			InOutEntry->DispatchData[Section].ClothBufferOffset = (ClothVertexBuffer->GetClothIndexMapping()[Section] & 0xFFFFFFFF);
+		}
 
         for (int32 Index = 0;Index < SimData->Positions.Num();Index++)
         {
