Skip WPO materials inside ApplyViewModeOverrides (used for e.g. lighting only view mode). This prevents from generating invalid draw calls without tesselation shaders, but with a tesselation primitive type. WPO support needs to reimplemented later (see UE-69220).

diff --git a/Engine/Source/Runtime/Engine/Private/PrimitiveDrawingUtils.cpp b/Engine/Source/Runtime/Engine/Private/PrimitiveDrawingUtils.cpp
--- a/Engine/Source/Runtime/Engine/Private/PrimitiveDrawingUtils.cpp
+++ b/Engine/Source/Runtime/Engine/Private/PrimitiveDrawingUtils.cpp
@@ -1385,8 +1385,10 @@ void ApplyViewModeOverrides(
 {
 #if !(UE_BUILD_SHIPPING || UE_BUILD_TEST)
 
-	// If debug viewmodes are not allowed, skip all of the debug viewmode handling
-	if (!AllowDebugViewmodes())
+	// If debug viewmodes are not allowed, skip all of the debug viewmode handling.
+	// If material has tesselation then also skip it, as we can't just replace it with a simple material.
+	if (!AllowDebugViewmodes()
+		|| Mesh.MaterialRenderProxy->GetMaterial(FeatureLevel)->MaterialModifiesMeshPosition_RenderThread())
 	{
 		return;
 	}
