DR - Always dump SCW jobs to help track down build machine issues
#rnx

diff --git a/Engine/Source/Runtime/Engine/Private/ShaderCompiler/ShaderCompiler.cpp b/Engine/Source/Runtime/Engine/Private/ShaderCompiler/ShaderCompiler.cpp
--- a/Engine/Source/Runtime/Engine/Private/ShaderCompiler/ShaderCompiler.cpp
+++ b/Engine/Source/Runtime/Engine/Private/ShaderCompiler/ShaderCompiler.cpp
@@ -610,7 +610,7 @@ void FShaderCompileUtilities::DoReadTaskResults(const TArray<FShaderCommonCompil
 		default:
 		case SCWErrorCode::GeneralCrash:
 			{
-				if (GDumpSCWJobInfoOnCrash != 0)
+				if (GDumpSCWJobInfoOnCrash != 0 || GIsBuildMachine)
 				{
 					auto DumpSingleJob = [](FShaderCompileJob* SingleJob) -> FString
 					{
