Fixed crash when leaving editor in D3D11.
Fixed crash in debug view shaders when rendering translucent materials that samples scene textures in the  vertex position offset.

diff --git a/Engine/Shaders/Private/DebugViewModeVertexShader.usf b/Engine/Shaders/Private/DebugViewModeVertexShader.usf
--- a/Engine/Shaders/Private/DebugViewModeVertexShader.usf
+++ b/Engine/Shaders/Private/DebugViewModeVertexShader.usf
@@ -5,7 +5,9 @@
 =============================================================================*/
 
 #define TEX_COORD_SCALE_ANALYSIS 1
- 
+
+#define SCENE_TEXTURES_DISABLED 1
+
 #include "Common.ush"
 
 struct FTexCoordScalesParams
diff --git a/Engine/Source/Runtime/Engine/Private/DebugViewModeMaterialProxy.cpp b/Engine/Source/Runtime/Engine/Private/DebugViewModeMaterialProxy.cpp
--- a/Engine/Source/Runtime/Engine/Private/DebugViewModeMaterialProxy.cpp
+++ b/Engine/Source/Runtime/Engine/Private/DebugViewModeMaterialProxy.cpp
@@ -79,12 +79,14 @@ void FDebugViewModeMaterialProxy::ClearAllShaders(UMaterialInterface* InMaterial
 	FlushRenderingCommands();
 	bReentrantCall = true;
 
+	TArray<FDebugViewModeMaterialProxy*, TInlineAllocator<1> > MaterialsToDelete;
+
 	if (!InMaterialInterface)
 	{
 		for (TMap<FMaterialUsagePair, FDebugViewModeMaterialProxy*>::TIterator It(DebugMaterialShaderMap); It; ++It)
 		{
 			FDebugViewModeMaterialProxy* Proxy = It.Value();
-			delete Proxy;
+			MaterialsToDelete.Add(Proxy);
 			It.Value() = nullptr;
 		}
 		DebugMaterialShaderMap.Empty();
@@ -96,7 +98,7 @@ void FDebugViewModeMaterialProxy::ClearAllShaders(UMaterialInterface* InMaterial
 			if (It.Key().MaterialInterface == InMaterialInterface)
 			{
 				FDebugViewModeMaterialProxy* Proxy = It.Value();
-				delete Proxy;
+				MaterialsToDelete.Add(Proxy);
 				It.Value() = nullptr;
 
 				It.RemoveCurrent();
@@ -104,6 +106,22 @@ void FDebugViewModeMaterialProxy::ClearAllShaders(UMaterialInterface* InMaterial
 		}
 	}
 
+	if (MaterialsToDelete.Num())
+	{
+		ENQUEUE_RENDER_COMMAND(DeleteDebugMaterials)([MaterialsToDelete](FRHICommandList& RHICmdList)
+		{
+			for (FDebugViewModeMaterialProxy* MaterialToDelete : MaterialsToDelete)
+			{
+				if (MaterialToDelete)
+				{
+					delete MaterialToDelete;
+				}
+			}
+		});
+
+		FlushRenderingCommands();
+	}
+
 	bReentrantCall = false;
 }
 
@@ -163,18 +181,19 @@ void FDebugViewModeMaterialProxy::ValidateAllShaders(TSet<UMaterialInterface*>&
 		}
 	}
 
-	TArray<FDebugViewModeMaterialProxy*> MaterialsToUpdateCopy = MaterialsToUpdate;
-	ENQUEUE_RENDER_COMMAND(UpdateDebugMaterialExpressionCache)(
-		[MaterialsToUpdateCopy](FRHICommandList& RHICmdList)
+	if (MaterialsToUpdate.Num())
+	{
+		ENQUEUE_RENDER_COMMAND(UpdateDebugMaterialExpressionCache)([MaterialsToUpdate](FRHICommandList& RHICmdList)
 		{
-			for (FDebugViewModeMaterialProxy* MaterialToUpdate : MaterialsToUpdateCopy)
+			for (FDebugViewModeMaterialProxy* MaterialToUpdate : MaterialsToUpdate)
 			{
 				check(MaterialToUpdate);
 				MaterialToUpdate->UpdateUniformExpressionCacheIfNeeded(MaterialToUpdate->FMaterial::GetFeatureLevel());
 			}
 		});
 
-	FlushRenderingCommands();
+		FlushRenderingCommands();
+	}
 }
 
 FDebugViewModeMaterialProxy::FDebugViewModeMaterialProxy(
