Fix ps4 automation failures

diff --git a/Engine/Source/Runtime/Engine/Public/PrimitiveUniformShaderParameters.h b/Engine/Source/Runtime/Engine/Public/PrimitiveUniformShaderParameters.h
--- a/Engine/Source/Runtime/Engine/Public/PrimitiveUniformShaderParameters.h
+++ b/Engine/Source/Runtime/Engine/Public/PrimitiveUniformShaderParameters.h
@@ -115,9 +115,11 @@ inline TUniformBufferRef<FPrimitiveUniformShaderParameters> CreatePrimitiveUnifo
 
 inline FPrimitiveUniformShaderParameters GetIdentityPrimitiveParameters()
 {
+	//don't use FMatrix::Identity here as GetIdentityPrimitiveParameters is used by TGlobalResource<FIdentityPrimitiveUniformBuffer> and because static initialization order is undefined
+	//FMatrix::Identiy might be all 0's or random data the first time this is called.
 	return GetPrimitiveUniformShaderParameters(
-		FMatrix::Identity,
-		FMatrix::Identity,
+		FMatrix(FPlane(1, 0, 0, 0), FPlane(0, 1, 0, 0), FPlane(0, 0, 1, 0), FPlane(0, 0, 0, 1)),
+		FMatrix(FPlane(1, 0, 0, 0), FPlane(0, 1, 0, 0), FPlane(0, 0, 1, 0), FPlane(0, 0, 0, 1)),
 		FVector(0.0f, 0.0f, 0.0f),
 		FBoxSphereBounds(EForceInit::ForceInit),
 		FBoxSphereBounds(EForceInit::ForceInit),
