Possible fix for stream crash and extra logging in vorbis decoder.

diff --git a/Engine/Source/Runtime/AudioMixer/Private/AudioMixerBuffer.cpp b/Engine/Source/Runtime/AudioMixer/Private/AudioMixerBuffer.cpp
--- a/Engine/Source/Runtime/AudioMixer/Private/AudioMixerBuffer.cpp
+++ b/Engine/Source/Runtime/AudioMixer/Private/AudioMixerBuffer.cpp
@@ -343,7 +343,10 @@ namespace Audio
 			InWave->DecompressionType = DTYPE_Invalid;
 			InWave->NumChannels = 0;
 
-			InWave->RemoveAudioResource();
+			InWave->RemoveAudioResource(); 
+
+			delete Buffer;
+			Buffer = nullptr;
 		}
 
 		return Buffer;
@@ -382,6 +385,9 @@ namespace Audio
 			InWave->NumChannels = 0;
 
 			InWave->RemoveAudioResource();
+
+			delete Buffer;
+			Buffer = nullptr;
 		}
 
 		return Buffer;
diff --git a/Engine/Source/Runtime/Engine/Private/VorbisAudioInfo.cpp b/Engine/Source/Runtime/Engine/Private/VorbisAudioInfo.cpp
--- a/Engine/Source/Runtime/Engine/Private/VorbisAudioInfo.cpp
+++ b/Engine/Source/Runtime/Engine/Private/VorbisAudioInfo.cpp
@@ -92,7 +92,7 @@ struct FVorbisFileWrapper
 /*------------------------------------------------------------------------------------
 	FVorbisAudioInfo.
 ------------------------------------------------------------------------------------*/
-FVorbisAudioInfo::FVorbisAudioInfo( void )
+FVorbisAudioInfo::FVorbisAudioInfo()
 	: VFWrapper(new FVorbisFileWrapper())
 	, SrcBufferData(NULL)
 	, SrcBufferDataSize(0)
@@ -103,12 +103,13 @@ FVorbisAudioInfo::FVorbisAudioInfo( void )
 	, CurrentStreamingChunkIndex(INDEX_NONE)
 	, NextStreamingChunkIndex(0)
 	, CurrentStreamingChunksSize(0)
+	, bHeaderParsed(false)
 {
 	// Make sure we have properly allocated a VFWrapper
 	check(VFWrapper != NULL);
 }
 
-FVorbisAudioInfo::~FVorbisAudioInfo( void )
+FVorbisAudioInfo::~FVorbisAudioInfo()
 {
 	FScopeLock ScopeLock(&VorbisCriticalSection);
 	check(VFWrapper != nullptr);
@@ -248,8 +249,12 @@ static size_t OggReadStreaming( void *ptr, size_t size, size_t nmemb, void *data
 {
 	check(ptr);
 	check(datasource);
-	FVorbisAudioInfo* OggInfo = (FVorbisAudioInfo*)datasource;
-	return( OggInfo->ReadStreaming( ptr, size * nmemb ) );
+	if (FVorbisAudioInfo* OggInfo = (FVorbisAudioInfo*)datasource)
+	{
+		return OggInfo->ReadStreaming(ptr, size * nmemb);
+	}
+	UE_LOG(LogAudio, Error, TEXT("OggReadStreaming had null audio info datasource."));
+	return 0;
 }
 
 int FVorbisAudioInfo::CloseStreaming( void )
@@ -267,6 +272,7 @@ bool FVorbisAudioInfo::GetCompressedInfoCommon(void* Callbacks, FSoundQualityInf
 {
 	if (!bDllLoaded)
 	{
+		UE_LOG(LogAudio, Error, TEXT("FVorbisAudioInfo::GetCompressedInfoCommon failed due to vorbis DLL not being loaded."));
 		return false;
 	}
 
@@ -310,17 +316,25 @@ bool FVorbisAudioInfo::ReadCompressedInfo( const uint8* InSrcBufferData, uint32
 {
 	if (!bDllLoaded)
 	{
+		UE_LOG(LogAudio, Error, TEXT("FVorbisAudioInfo::ReadCompressedInfo failed due to vorbis DLL not being loaded."));
+		return false;
+	}
+	
+	if (bHeaderParsed)
+	{
+		UE_LOG(LogAudio, Error, TEXT("FVorbisAudioInfo::ReadCompressedInfo failed due to the header being parsed already."));
 		return false;
 	}
-
-	SCOPE_CYCLE_COUNTER( STAT_VorbisPrepareDecompressionTime );
-
-	FScopeLock ScopeLock(&VorbisCriticalSection);
 
 	if (!VFWrapper)
 	{
+		UE_LOG(LogAudio, Error, TEXT("FVorbisAudioInfo::ReadCompressedInfo failed due to not having vorbis wrapper."));
 		return false;
 	}
+	
+	SCOPE_CYCLE_COUNTER( STAT_VorbisPrepareDecompressionTime );
+
+	FScopeLock ScopeLock(&VorbisCriticalSection);
 
 	ov_callbacks Callbacks;
 
@@ -333,9 +347,15 @@ bool FVorbisAudioInfo::ReadCompressedInfo( const uint8* InSrcBufferData, uint32
 	Callbacks.close_func = OggCloseMemory;
 	Callbacks.tell_func = OggTellMemory;
 
-	bool result = GetCompressedInfoCommon(&Callbacks, QualityInfo);
+	bHeaderParsed = GetCompressedInfoCommon(&Callbacks, QualityInfo);
+
+	if (!bHeaderParsed)
+	{	
+		UE_LOG(LogAudio, Error, TEXT("Failed to parse header for compressed vorbis file."));
+	}
+
 
-	return( result );
+	return bHeaderParsed;
 }
 
 
@@ -346,6 +366,7 @@ void FVorbisAudioInfo::ExpandFile( uint8* DstBuffer, FSoundQualityInfo* QualityI
 {
 	if (!bDllLoaded)
 	{
+		UE_LOG(LogAudio, Error, TEXT("FVorbisAudioInfo::ExpandFile failed due to vorbis DLL not being loaded."));
 		return;
 	}
 
@@ -357,6 +378,12 @@ void FVorbisAudioInfo::ExpandFile( uint8* DstBuffer, FSoundQualityInfo* QualityI
 
 	FScopeLock ScopeLock(&VorbisCriticalSection);
 
+	if (!bHeaderParsed)
+	{
+		UE_LOG(LogAudio, Error, TEXT("Failed to expand vorbis file to not parsing header first."));
+		return;
+	}
+
 	// A zero buffer size means decompress the entire ogg vorbis stream to PCM.
 	TotalBytesRead = 0;
 	BytesToRead = QualityInfo->SampleDataSize;
@@ -394,6 +421,7 @@ bool FVorbisAudioInfo::ReadCompressedData( uint8* InDestination, bool bLooping,
 {
 	if (!bDllLoaded)
 	{
+		UE_LOG(LogAudio, Error, TEXT("FVorbisAudioInfo::ReadCompressedData failed due to vorbis DLL not being loaded."));
 		return true;
 	}
 
@@ -409,6 +437,12 @@ bool FVorbisAudioInfo::ReadCompressedData( uint8* InDestination, bool bLooping,
 
 	FScopeLock ScopeLock(&VorbisCriticalSection);
 
+	if (!bHeaderParsed)
+	{
+		UE_LOG(LogAudio, Error, TEXT("FVorbisAudioInfo::ReadCompressedData failed due to not parsing header first."));
+		return true;
+	}
+
 	bLooped = false;
 
 	// Work out number of samples to read
@@ -460,11 +494,18 @@ void FVorbisAudioInfo::SeekToTime( const float SeekTime )
 {
 	if (!bDllLoaded)
 	{
+		UE_LOG(LogAudio, Error, TEXT("FVorbisAudioInfo::SeekToTime failed due to vorbis DLL not being loaded."));
 		return;
 	}
 
 	FScopeLock ScopeLock(&VorbisCriticalSection);
 
+	if (!bHeaderParsed)
+	{
+		UE_LOG(LogAudio, Error, TEXT("FVorbisAudioInfo::SeekToTime failed due to not parsing header first."));
+		return;
+	}
+
 	const float TargetTime = FMath::Min(SeekTime, (float)ov_time_total(&VFWrapper->vf, -1));
 	ov_time_seek( &VFWrapper->vf, TargetTime );
 }
@@ -473,11 +514,18 @@ void FVorbisAudioInfo::EnableHalfRate( bool HalfRate )
 {
 	if (!bDllLoaded)
 	{
+		UE_LOG(LogAudio, Error, TEXT("FVorbisAudioInfo::EnableHalfRate failed due to vorbis DLL not being loaded."));
 		return;
 	}
 
 	FScopeLock ScopeLock(&VorbisCriticalSection);
 
+	if (!bHeaderParsed)
+	{
+		UE_LOG(LogAudio, Error, TEXT("FVorbisAudioInfo::EnableHalfRate failed due to not parsing header first."));
+		return;
+	}
+
 	ov_halfrate(&VFWrapper->vf, int32(HalfRate));
 }
 
@@ -485,6 +533,7 @@ bool FVorbisAudioInfo::StreamCompressedInfo(USoundWave* Wave, struct FSoundQuali
 {
 	if (!bDllLoaded)
 	{
+		UE_LOG(LogAudio, Error, TEXT("FVorbisAudioInfo::StreamCompressedInfo failed to parse header due to vorbis DLL not being loaded for sound '%s'."), *Wave->GetName());
 		return false;
 	}
 
@@ -494,6 +543,7 @@ bool FVorbisAudioInfo::StreamCompressedInfo(USoundWave* Wave, struct FSoundQuali
 
 	if (!VFWrapper)
 	{
+		UE_LOG(LogAudio, Error, TEXT("FVorbisAudioInfo::StreamCompressedInfo failed due to no vorbis wrapper for sound '%s'."), *Wave->GetName());
 		return false;
 	}
 
@@ -509,16 +559,22 @@ bool FVorbisAudioInfo::StreamCompressedInfo(USoundWave* Wave, struct FSoundQuali
 	Callbacks.seek_func = NULL;	// Force streaming
 	Callbacks.tell_func = NULL;	// Force streaming
 
-	bool Result = GetCompressedInfoCommon(&Callbacks, QualityInfo);
+	bHeaderParsed = GetCompressedInfoCommon(&Callbacks, QualityInfo);
+	if (!bHeaderParsed)
+	{
+		UE_LOG(LogAudio, Error, TEXT("FVorbisAudioInfo::StreamCompressedInfo failed to parse header for '%s'."), *Wave->GetName());
+	}
+	
 
-	return( Result );
+	return bHeaderParsed;
 }
 
 bool FVorbisAudioInfo::StreamCompressedData(uint8* InDestination, bool bLooping, uint32 BufferSize)
 {
 	if (!bDllLoaded)
 	{
-		return false;
+		UE_LOG(LogAudio, Error, TEXT("FVorbisAudioInfo::StreamCompressedData failed due to vorbis DLL not being loaded."));
+		return true;
 	}
 
 	check( VFWrapper != NULL );
@@ -531,6 +587,12 @@ bool FVorbisAudioInfo::StreamCompressedData(uint8* InDestination, bool bLooping,
 #endif
 	FScopeLock ScopeLock(&VorbisCriticalSection);
 
+	if (!bHeaderParsed)
+	{
+		UE_LOG(LogAudio, Error, TEXT("FVorbisAudioInfo::StreamCompressedData failed due to not parsing header first."));
+		return true;
+	}
+
 	bool bLooped = false;
 
 	while( BufferSize > 0 )
@@ -637,6 +699,10 @@ void LoadVorbisLibraries()
 		{
 			UE_LOG(LogAudio, Error, TEXT("Failed to load lib vorbis libraries."));
 		}
+		else
+		{
+			UE_LOG(LogAudio, Warning, TEXT("Lib vorbis DLL was dynamically loaded."));
+		}
 #elif WITH_OGGVORBIS
 		bDllLoaded = true;
 #endif	//PLATFORM_WINDOWS
diff --git a/Engine/Source/Runtime/Engine/Public/VorbisAudioInfo.h b/Engine/Source/Runtime/Engine/Public/VorbisAudioInfo.h
--- a/Engine/Source/Runtime/Engine/Public/VorbisAudioInfo.h
+++ b/Engine/Source/Runtime/Engine/Public/VorbisAudioInfo.h
@@ -113,5 +113,6 @@ private:
 	int32 CurrentStreamingChunkIndex;
 	int32 NextStreamingChunkIndex;
 	uint32 CurrentStreamingChunksSize;
+	bool bHeaderParsed;
 };
 #endif
diff --git a/Engine/Source/Runtime/Windows/XAudio2/Private/XAudio2Buffer.cpp b/Engine/Source/Runtime/Windows/XAudio2/Private/XAudio2Buffer.cpp
--- a/Engine/Source/Runtime/Windows/XAudio2/Private/XAudio2Buffer.cpp
+++ b/Engine/Source/Runtime/Windows/XAudio2/Private/XAudio2Buffer.cpp
@@ -508,6 +508,9 @@ FXAudio2SoundBuffer* FXAudio2SoundBuffer::CreateStreamingBuffer( FXAudio2Device*
 	}
 	else
 	{
+		Wave->DecompressionType = DTYPE_Invalid;
+		Wave->NumChannels = 0;
+
 		delete Buffer;
 		Buffer = nullptr;
 
