Fix reporting of FRepState and FObjectReplicator Memory.
This significantly reduces unreported memory between LLM and Obj List.

diff --git a/Engine/Source/Runtime/Engine/Classes/Engine/DemoNetDriver.h b/Engine/Source/Runtime/Engine/Classes/Engine/DemoNetDriver.h
--- a/Engine/Source/Runtime/Engine/Classes/Engine/DemoNetDriver.h
+++ b/Engine/Source/Runtime/Engine/Classes/Engine/DemoNetDriver.h
@@ -322,7 +322,8 @@ struct FRollbackNetStartupActorInfo
 			
 			if (FRepState const * const LocalRepState = SubObjRepStatePair.Value.Get())
 			{
-				Ar.CountBytes(sizeof(FRepState), sizeof(FRepState));
+				const SIZE_T SizeOfRepState = sizeof(FRepState);
+				Ar.CountBytes(SizeOfRepState, SizeOfRepState);
 				LocalRepState->CountBytes(Ar);
 			}
 		}
diff --git a/Engine/Source/Runtime/Engine/Private/DataChannel.cpp b/Engine/Source/Runtime/Engine/Private/DataChannel.cpp
--- a/Engine/Source/Runtime/Engine/Private/DataChannel.cpp
+++ b/Engine/Source/Runtime/Engine/Private/DataChannel.cpp
@@ -2977,11 +2977,8 @@ void UActorChannel::Serialize(FArchive& Ar)
 
 		GRANULAR_NETWORK_MEMORY_TRACKING_TRACK("ReplicationMap",
 			ReplicationMap.CountBytes(Ar);
-			for (const auto& MapPair : ReplicationMap)
-			{
-				Ar.CountBytes(sizeof(*MapPair.Value), sizeof(*MapPair.Value));
-				MapPair.Value->Serialize(Ar);
-			}
+
+			// ObjectReplicators are going to be counted by UNetDriver::Serialize AllOwnedReplicators.
 		);
 		
 		GRANULAR_NETWORK_MEMORY_TRACKING_TRACK("QueudBunches",	
diff --git a/Engine/Source/Runtime/Engine/Private/DataReplication.cpp b/Engine/Source/Runtime/Engine/Private/DataReplication.cpp
--- a/Engine/Source/Runtime/Engine/Private/DataReplication.cpp
+++ b/Engine/Source/Runtime/Engine/Private/DataReplication.cpp
@@ -1485,59 +1485,66 @@ void FObjectReplicator::Serialize(FArchive& Ar)
 {
 	if (Ar.IsCountingMemory())
 	{
-		Retirement.CountBytes(Ar);
-		RecentCustomDeltaState.CountBytes(Ar);
-		for (const auto& RecentCustomDeltaStatePair : RecentCustomDeltaState)
-		{
-			if (INetDeltaBaseState const * const BaseState = RecentCustomDeltaStatePair.Value.Get())
-			{
-				BaseState->CountBytes(Ar);
-			}
-		}
+		CountBytes(Ar);
+	}
+}
 
-		CDOCustomDeltaState.CountBytes(Ar);
-		for (const auto& CDOCustomDeltaStatePair : CDOCustomDeltaState)
+void FObjectReplicator::CountBytes(FArchive& Ar) const
+{
+	Retirement.CountBytes(Ar);
+	RecentCustomDeltaState.CountBytes(Ar);
+	for (const auto& RecentCustomDeltaStatePair : RecentCustomDeltaState)
+	{
+		if (INetDeltaBaseState const * const BaseState = RecentCustomDeltaStatePair.Value.Get())
 		{
-			if (INetDeltaBaseState const * const BaseState = CDOCustomDeltaStatePair.Value.Get())
-			{
-				BaseState->CountBytes(Ar);
-			}
+			BaseState->CountBytes(Ar);
 		}
+	}
 
-		LifetimeCustomDeltaProperties.CountBytes(Ar);
-		LifetimeCustomDeltaPropertyConditions.CountBytes(Ar);
-		UnmappedCustomProperties.CountBytes(Ar);
-		RepNotifies.CountBytes(Ar);
-		RepNotifyMetaData.CountBytes(Ar);
-		for (const auto& MetaDataPair : RepNotifyMetaData)
+	CDOCustomDeltaState.CountBytes(Ar);
+	for (const auto& CDOCustomDeltaStatePair : CDOCustomDeltaState)
+	{
+		if (INetDeltaBaseState const * const BaseState = CDOCustomDeltaStatePair.Value.Get())
 		{
-			MetaDataPair.Value.CountBytes(Ar);
+			BaseState->CountBytes(Ar);
 		}
+	}
+
+	LifetimeCustomDeltaProperties.CountBytes(Ar);
+	LifetimeCustomDeltaPropertyConditions.CountBytes(Ar);
+	UnmappedCustomProperties.CountBytes(Ar);
+	RepNotifies.CountBytes(Ar);
+	RepNotifyMetaData.CountBytes(Ar);
+	for (const auto& MetaDataPair : RepNotifyMetaData)
+	{
+		MetaDataPair.Value.CountBytes(Ar);
+	}
 
-		// FObjectReplicator has a shared pointer to an FRepLayout, but since it's shared with
-		// the UNetDriver, the memory isn't tracked here.
+	// FObjectReplicator has a shared pointer to an FRepLayout, but since it's shared with
+	// the UNetDriver, the memory isn't tracked here.
 
-		if (RepState.IsValid())
-		{
-			RepState->CountBytes(Ar);
-		}
+	if (RepState.IsValid())
+	{
+		const SIZE_T SizeOfRepState = sizeof(FRepState);
+		Ar.CountBytes(SizeOfRepState, SizeOfRepState);
+		RepState->CountBytes(Ar);
+	}
 
-		ReferencedGuids.CountBytes(Ar);
+	ReferencedGuids.CountBytes(Ar);
 
-		// ChangelistMgr points to a ReplicationChangelistMgr managed by the UNetDriver, so it's not tracked here
+	// ChangelistMgr points to a ReplicationChangelistMgr managed by the UNetDriver, so it's not tracked here
 
-		RemoteFuncInfo.CountBytes(Ar);
-		if (RemoteFunctions)
-		{
-			RemoteFunctions->CountMemory(Ar);
-		}
+	RemoteFuncInfo.CountBytes(Ar);
+	if (RemoteFunctions)
+	{
+		RemoteFunctions->CountMemory(Ar);
+	}
 
-		PendingLocalRPCs.CountBytes(Ar);
-		for (const FRPCPendingLocalCall& PendingRPC : PendingLocalRPCs)
-		{
-			PendingRPC.Buffer.CountBytes(Ar);
-			PendingRPC.UnmappedGuids.CountBytes(Ar);
-		}
+	PendingLocalRPCs.CountBytes(Ar);
+	for (const FRPCPendingLocalCall& PendingRPC : PendingLocalRPCs)
+	{
+		PendingRPC.Buffer.CountBytes(Ar);
+		PendingRPC.UnmappedGuids.CountBytes(Ar);
 	}
 }
 
diff --git a/Engine/Source/Runtime/Engine/Private/NetConnection.cpp b/Engine/Source/Runtime/Engine/Private/NetConnection.cpp
--- a/Engine/Source/Runtime/Engine/Private/NetConnection.cpp
+++ b/Engine/Source/Runtime/Engine/Private/NetConnection.cpp
@@ -529,13 +529,8 @@ void UNetConnection::Serialize( FArchive& Ar )
 			}
 		);
 
-		GRANULAR_NETWORK_MEMORY_TRACKING_TRACK("DormantReplicatorMap",
-			DormantReplicatorMap.CountBytes(Ar);
-			for (auto& DormantReplicatorPair : DormantReplicatorMap)
-			{
-				DormantReplicatorPair.Value->Serialize(Ar);
-			}
-		);
+		// ObjectReplicators are going to be counted by UNetDriver::Serialize AllOwnedReplicators.
+		GRANULAR_NETWORK_MEMORY_TRACKING_TRACK("DormantReplicatorMap", DormantReplicatorMap.CountBytes(Ar));
 
 		GRANULAR_NETWORK_MEMORY_TRACKING_TRACK("ClientVisibleLevelNames", ClientVisibleLevelNames.CountBytes(Ar));
 		GRANULAR_NETWORK_MEMORY_TRACKING_TRACK("ClientVisibileActorOuters", ClientVisibileActorOuters.CountBytes(Ar));
diff --git a/Engine/Source/Runtime/Engine/Private/NetworkDriver.cpp b/Engine/Source/Runtime/Engine/Private/NetworkDriver.cpp
--- a/Engine/Source/Runtime/Engine/Private/NetworkDriver.cpp
+++ b/Engine/Source/Runtime/Engine/Private/NetworkDriver.cpp
@@ -2308,7 +2308,21 @@ void UNetDriver::Serialize( FArchive& Ar )
 
 		GRANULAR_NETWORK_MEMORY_TRACKING_TRACK("GuidToReplicatorMap", GuidToReplicatorMap.CountBytes(Ar));
 		GRANULAR_NETWORK_MEMORY_TRACKING_TRACK("UnmappedReplicators", UnmappedReplicators.CountBytes(Ar));
-		GRANULAR_NETWORK_MEMORY_TRACKING_TRACK("AllOwnedReplicators", AllOwnedReplicators.CountBytes(Ar));
+		GRANULAR_NETWORK_MEMORY_TRACKING_TRACK("AllOwnedReplicators",
+		
+			// AllOwnedReplicators is the superset of all initialized FObjectReplicators,
+			// including DormantReplicators, Replicators on Channels, and UnmappedReplicators.
+			AllOwnedReplicators.CountBytes(Ar);
+			
+			for (FObjectReplicator const * const OwnedReplicator : AllOwnedReplicators)
+			{
+				if (OwnedReplicator)
+				{
+					Ar.CountBytes(sizeof(FObjectReplicator), sizeof(FObjectReplicator));
+					OwnedReplicator->CountBytes(Ar);
+				}
+			}
+		);
 
 		// Replicators are owned by UActorChannels, and so we don't track them here.
 
diff --git a/Engine/Source/Runtime/Engine/Private/RepLayout.cpp b/Engine/Source/Runtime/Engine/Private/RepLayout.cpp
--- a/Engine/Source/Runtime/Engine/Private/RepLayout.cpp
+++ b/Engine/Source/Runtime/Engine/Private/RepLayout.cpp
@@ -5119,6 +5119,9 @@ void FRepLayout::CountBytes(FArchive& Ar) const
 
 void FRepState::CountBytes(FArchive& Ar) const
 {
+	const SIZE_T SizeOfThis = sizeof(FRepState);
+	Ar.CountBytes(SizeOfThis, SizeOfThis);
+
 	StaticBuffer.CountBytes(Ar);
 	GuidReferencesMap.CountBytes(Ar);
 	for (const auto& GuidRefPair : GuidReferencesMap)
diff --git a/Engine/Source/Runtime/Engine/Public/Net/DataReplication.h b/Engine/Source/Runtime/Engine/Public/Net/DataReplication.h
--- a/Engine/Source/Runtime/Engine/Public/Net/DataReplication.h
+++ b/Engine/Source/Runtime/Engine/Public/Net/DataReplication.h
@@ -194,7 +194,10 @@ public:
 	/** Packet was dropped */
 	void	ReceivedNak( int32 NakPacketId );
 
-	void	Serialize(FArchive& Ar);
+	UE_DEPRECATED(4.23, "Use CountBytes instead")
+	void Serialize(FArchive& Ar);
+
+	void CountBytes(FArchive& Ar) const;
 
 	/** Writes dirty properties to bunch */
 	void	ReplicateCustomDeltaProperties( FNetBitWriter & Bunch, FReplicationFlags RepFlags );
