#UE4 Adjust log verbosity and category

diff --git a/Engine/Source/Runtime/Engine/Private/VorbisAudioInfo.cpp b/Engine/Source/Runtime/Engine/Private/VorbisAudioInfo.cpp
--- a/Engine/Source/Runtime/Engine/Private/VorbisAudioInfo.cpp
+++ b/Engine/Source/Runtime/Engine/Private/VorbisAudioInfo.cpp
@@ -701,7 +701,7 @@ void LoadVorbisLibraries()
 		}
 		else
 		{
-			UE_LOG(LogAudio, Warning, TEXT("Lib vorbis DLL was dynamically loaded."));
+			UE_LOG(LogAudioDebug, Display, TEXT("Lib vorbis DLL was dynamically loaded."));
 		}
 #elif WITH_OGGVORBIS
 		bDllLoaded = true;
