Reverted change that caused this regression.

diff --git a/Engine/Source/Runtime/Engine/Classes/GameFramework/Actor.h b/Engine/Source/Runtime/Engine/Classes/GameFramework/Actor.h
--- a/Engine/Source/Runtime/Engine/Classes/GameFramework/Actor.h
+++ b/Engine/Source/Runtime/Engine/Classes/GameFramework/Actor.h
@@ -618,7 +618,7 @@ public:
 	uint8 bIsEditorPreviewActor:1;
 
 	/** Whether this actor is hidden by the layer browser. */
-	UPROPERTY(Transient)
+	UPROPERTY()
 	uint8 bHiddenEdLayer:1;
 
 	/** Whether this actor is hidden by the level browser. */
