Add ISR support for instanced particles.
Improve the fix for UE-39249, instead of changing CheckAndUpdateLastFrame() rendering code is moved outside of its conditional scope.
This reduces redundant buffer allocations/uploads within the same frame.

diff --git a/Engine/Source/Runtime/Engine/Private/Particles/ParticleSystemRender.cpp b/Engine/Source/Runtime/Engine/Private/Particles/ParticleSystemRender.cpp
--- a/Engine/Source/Runtime/Engine/Private/Particles/ParticleSystemRender.cpp
+++ b/Engine/Source/Runtime/Engine/Private/Particles/ParticleSystemRender.cpp
@@ -606,7 +606,7 @@ void ApplyOrbitToPosition(
 	}
 }
 
-bool FDynamicSpriteEmitterData::GetVertexAndIndexData(void* VertexData, void* DynamicParameterVertexData, void* FillIndexData, FParticleOrder* ParticleOrder, const FVector& InCameraPosition, const FMatrix& InLocalToWorld) const
+bool FDynamicSpriteEmitterData::GetVertexAndIndexData(void* VertexData, void* DynamicParameterVertexData, void* FillIndexData, FParticleOrder* ParticleOrder, const FVector& InCameraPosition, const FMatrix& InLocalToWorld, uint32 InstanceFactor) const
 {
 	SCOPE_CYCLE_COUNTER(STAT_ParticlePackingTime);
 	int32 ParticleCount = Source.ActiveParticleCount;
@@ -683,27 +683,31 @@ bool FDynamicSpriteEmitterData::GetVertexAndIndexData(void* VertexData, void* Dy
 			GetDynamicValueFromPayload(Source.DynamicParameterDataOffset, Particle, DynamicParameterValue);
 		}
 
-		FillVertex = (FParticleSpriteVertex*)TempVert;
-		FillVertex->Position	= ParticlePosition;
-		FillVertex->RelativeTime = Particle.RelativeTime;
-		FillVertex->OldPosition	= ParticleOldPosition;
-		// Create a floating point particle ID from the counter, map into approximately 0-1
-		FillVertex->ParticleId = (Particle.Flags & STATE_CounterMask) / 10000.0f;
-		FillVertex->Size = GetParticleSizeWithUVFlipInSign(Particle, Size);
-		FillVertex->Rotation	= Particle.Rotation;
-		FillVertex->SubImageIndex = SubImageIndex;
-		FillVertex->Color		= Particle.Color;
-		if (bUsesDynamicParameter)
+		//@todo - refactor into instance step rate in the RHI
+		for (uint32 Factor = 0; Factor < InstanceFactor; Factor++)
 		{
-			DynFillVertex = (FParticleVertexDynamicParameter*)TempDynamicParameterVert;
-			DynFillVertex->DynamicValue[0] = DynamicParameterValue.X;
-			DynFillVertex->DynamicValue[1] = DynamicParameterValue.Y;
-			DynFillVertex->DynamicValue[2] = DynamicParameterValue.Z;
-			DynFillVertex->DynamicValue[3] = DynamicParameterValue.W;
-			TempDynamicParameterVert += VertexDynamicParameterStride;
-		}
+			FillVertex = (FParticleSpriteVertex*)TempVert;
+			FillVertex->Position = ParticlePosition;
+			FillVertex->RelativeTime = Particle.RelativeTime;
+			FillVertex->OldPosition = ParticleOldPosition;
+			// Create a floating point particle ID from the counter, map into approximately 0-1
+			FillVertex->ParticleId = (Particle.Flags & STATE_CounterMask) / 10000.0f;
+			FillVertex->Size = GetParticleSizeWithUVFlipInSign(Particle, Size);
+			FillVertex->Rotation = Particle.Rotation;
+			FillVertex->SubImageIndex = SubImageIndex;
+			FillVertex->Color = Particle.Color;
+			if (bUsesDynamicParameter)
+			{
+				DynFillVertex = (FParticleVertexDynamicParameter*)TempDynamicParameterVert;
+				DynFillVertex->DynamicValue[0] = DynamicParameterValue.X;
+				DynFillVertex->DynamicValue[1] = DynamicParameterValue.Y;
+				DynFillVertex->DynamicValue[2] = DynamicParameterValue.Z;
+				DynFillVertex->DynamicValue[3] = DynamicParameterValue.W;
+				TempDynamicParameterVert += VertexDynamicParameterStride;
+			}
 
-		TempVert += VertexStride;
+			TempVert += VertexStride;
+		}
 	}
 
 	return true;
@@ -1056,10 +1060,11 @@ void FDynamicSpriteEmitterData::GetDynamicMeshElementsEmitter(const FParticleSys
 			}
 
 			const int32 NumVerticesPerParticleInBuffer = bInstanced ? 1 : NumVerticesPerParticle;
+			const uint32 InstanceFactor = bInstanced && View->IsInstancedStereoPass() ? 2 : 1;
 
 			FParticleSpriteVertexFactory* SpriteVertexFactory = static_cast<FParticleSpriteVertexFactory*>(VertexFactory);
 
-			if (SpriteVertexFactory->CheckAndUpdateLastFrame(ViewFamily, View))
+			if (SpriteVertexFactory->CheckAndUpdateLastFrame(ViewFamily))
 			{
 				check(NumVerticesPerParticle == 4 || NumVerticesPerParticle == 8);
 				const FVertexBuffer* TexCoordBuffer = (NumVerticesPerParticle == 4) ? (const FVertexBuffer*)&GParticleTexCoordVertexBuffer : (const FVertexBuffer*)&GParticleEightTexCoordVertexBuffer;
@@ -1076,19 +1081,19 @@ void FDynamicSpriteEmitterData::GetDynamicMeshElementsEmitter(const FParticleSys
 				FGlobalDynamicVertexBuffer::FAllocation DynamicParameterAllocation;
 
 				// Allocate memory for render data.
-				Allocation = FGlobalDynamicVertexBuffer::Get().Allocate( ParticleCount * VertexSize * NumVerticesPerParticleInBuffer );
+				Allocation = FGlobalDynamicVertexBuffer::Get().Allocate( InstanceFactor * ParticleCount * VertexSize * NumVerticesPerParticleInBuffer );
 				if (FGlobalDynamicVertexBuffer::Get().IsRenderAlarmLoggingEnabled())
 				{
-					UE_LOG(LogParticles, Warning, TEXT("Panic logging.  Allocated %u bytes for Resource: %s, Owner: %s"), ParticleCount * VertexSize * NumVerticesPerParticleInBuffer, *Proxy->GetResourceName().ToString(), *Proxy->GetOwnerName().ToString())
+					UE_LOG(LogParticles, Warning, TEXT("Panic logging.  Allocated %u bytes for Resource: %s, Owner: %s"), InstanceFactor * ParticleCount * VertexSize * NumVerticesPerParticleInBuffer, *Proxy->GetResourceName().ToString(), *Proxy->GetOwnerName().ToString())
 				}
 
 				if (bUsesDynamicParameter)
 				{
-					DynamicParameterAllocation = FGlobalDynamicVertexBuffer::Get().Allocate( ParticleCount * DynamicParameterVertexSize * NumVerticesPerParticleInBuffer );
+					DynamicParameterAllocation = FGlobalDynamicVertexBuffer::Get().Allocate( InstanceFactor * ParticleCount * DynamicParameterVertexSize * NumVerticesPerParticleInBuffer );
 
 					if (FGlobalDynamicVertexBuffer::Get().IsRenderAlarmLoggingEnabled())
 					{
-						UE_LOG(LogParticles, Warning, TEXT("Panic logging.  Allocated %u bytes for Resource: %s, Owner: %s"), ParticleCount * DynamicParameterVertexSize * NumVerticesPerParticleInBuffer, *Proxy->GetResourceName().ToString(), *Proxy->GetOwnerName().ToString())
+						UE_LOG(LogParticles, Warning, TEXT("Panic logging.  Allocated %u bytes for Resource: %s, Owner: %s"), InstanceFactor * ParticleCount * DynamicParameterVertexSize * NumVerticesPerParticleInBuffer, *Proxy->GetResourceName().ToString(), *Proxy->GetOwnerName().ToString())
 					}
 				}
 
@@ -1115,7 +1120,7 @@ void FDynamicSpriteEmitterData::GetDynamicMeshElementsEmitter(const FParticleSys
 						if (Collector.ShouldUseTasks())
 						{
 							Collector.AddTask(
-								[this, SourceData, View, Proxy, Allocation, DynamicParameterAllocation, bInstanced, bSort, ParticleCount, NumVerticesPerParticleInBuffer]()
+								[this, SourceData, View, Proxy, Allocation, DynamicParameterAllocation, bInstanced, bSort, ParticleCount, NumVerticesPerParticleInBuffer, InstanceFactor]()
 								{
 									SCOPE_CYCLE_COUNTER(STAT_FDynamicSpriteEmitterData_GetDynamicMeshElementsEmitter_Task);
 									SCOPE_CYCLE_COUNTER(STAT_ParticlesOverview_RT_CNC);
@@ -1132,7 +1137,7 @@ void FDynamicSpriteEmitterData::GetDynamicMeshElementsEmitter(const FParticleSys
 									// Fill vertex buffers.
 									if (bInstanced)
 									{
-										GetVertexAndIndexData(Allocation.Buffer, DynamicParameterAllocation.Buffer, NULL, ParticleOrder, View->ViewMatrices.GetViewOrigin(), Proxy->GetLocalToWorld());
+										GetVertexAndIndexData(Allocation.Buffer, DynamicParameterAllocation.Buffer, NULL, ParticleOrder, View->ViewMatrices.GetViewOrigin(), Proxy->GetLocalToWorld(), InstanceFactor);
 									}
 									else
 									{
@@ -1156,7 +1161,7 @@ void FDynamicSpriteEmitterData::GetDynamicMeshElementsEmitter(const FParticleSys
 							// Fill vertex buffers.
 							if (bInstanced)
 							{
-								GetVertexAndIndexData(Allocation.Buffer, DynamicParameterAllocation.Buffer, NULL, ParticleOrder, View->ViewMatrices.GetViewOrigin(), Proxy->GetLocalToWorld());
+								GetVertexAndIndexData(Allocation.Buffer, DynamicParameterAllocation.Buffer, NULL, ParticleOrder, View->ViewMatrices.GetViewOrigin(), Proxy->GetLocalToWorld(), InstanceFactor);
 							}
 							else
 							{
@@ -1196,69 +1201,69 @@ void FDynamicSpriteEmitterData::GetDynamicMeshElementsEmitter(const FParticleSys
 					}
 
 				}
+			}
 
-					// Don't render if the material will be ignored
-					const bool bIsWireframe = ViewFamily.EngineShowFlags.Wireframe;
+			// Don't render if the material will be ignored
+			const bool bIsWireframe = ViewFamily.EngineShowFlags.Wireframe;
 
-					// Calculate the number of particles that must be drawn.
-					ParticleCount = Source.ActiveParticleCount;
-					if ((Source.MaxDrawCount >= 0) && (ParticleCount > Source.MaxDrawCount))
-					{
-						ParticleCount = Source.MaxDrawCount;
-					}
+			// Calculate the number of particles that must be drawn.
+			ParticleCount = Source.ActiveParticleCount;
+			if ((Source.MaxDrawCount >= 0) && (ParticleCount > Source.MaxDrawCount))
+			{
+				ParticleCount = Source.MaxDrawCount;
+			}
 
-					// Construct the mesh element to render.
-					FMeshBatch& Mesh = Collector.AllocateMesh();
-					FMeshBatchElement& BatchElement = Mesh.Elements[0];
-					check(NumTrianglesPerParticle == 2 || NumTrianglesPerParticle == 6);
-					BatchElement.IndexBuffer = NumTrianglesPerParticle == 2 ? (const FIndexBuffer*)&GParticleIndexBuffer : (const FIndexBuffer*)&GSixTriangleParticleIndexBuffer;
-					BatchElement.bIsInstancedMesh = bInstanced;
-					if(bInstanced)
-					{
-						BatchElement.NumPrimitives = NumTrianglesPerParticle;
-						BatchElement.NumInstances = ParticleCount;
-					}
-					else
-					{
-						BatchElement.NumPrimitives = NumTrianglesPerParticle * ParticleCount;
-						BatchElement.NumInstances = 1;
-					}
-					BatchElement.FirstIndex = 0;
-					Mesh.VertexFactory = SpriteVertexFactory;
-					// if the particle rendering data is presupplied, use it directly
-					Mesh.LCI = NULL;
-					if (SourceData->bUseLocalSpace == true)
-					{
-						BatchElement.PrimitiveUniformBuffer = Proxy->GetUniformBuffer();
-					}
-					else
-					{
-						BatchElement.PrimitiveUniformBuffer = Proxy->GetWorldSpacePrimitiveUniformBuffer();
-					}
-					BatchElement.MinVertexIndex = 0;
-					BatchElement.MaxVertexIndex = (ParticleCount * NumVerticesPerParticle) - 1;
-					Mesh.CastShadow = Proxy->GetCastShadow();
-					Mesh.DepthPriorityGroup = (ESceneDepthPriorityGroup)Proxy->GetDepthPriorityGroup(View);
+			// Construct the mesh element to render.
+			FMeshBatch& Mesh = Collector.AllocateMesh();
+			FMeshBatchElement& BatchElement = Mesh.Elements[0];
+			check(NumTrianglesPerParticle == 2 || NumTrianglesPerParticle == 6);
+			BatchElement.IndexBuffer = NumTrianglesPerParticle == 2 ? (const FIndexBuffer*)&GParticleIndexBuffer : (const FIndexBuffer*)&GSixTriangleParticleIndexBuffer;
+			BatchElement.bIsInstancedMesh = bInstanced;
+			if(bInstanced)
+			{
+				BatchElement.NumPrimitives = NumTrianglesPerParticle;
+				BatchElement.NumInstances = ParticleCount;
+			}
+			else
+			{
+				BatchElement.NumPrimitives = NumTrianglesPerParticle * ParticleCount;
+				BatchElement.NumInstances = 1;
+			}
+			BatchElement.FirstIndex = 0;
+			Mesh.VertexFactory = SpriteVertexFactory;
+			// if the particle rendering data is presupplied, use it directly
+			Mesh.LCI = NULL;
+			if (SourceData->bUseLocalSpace == true)
+			{
+				BatchElement.PrimitiveUniformBuffer = Proxy->GetUniformBuffer();
+			}
+			else
+			{
+				BatchElement.PrimitiveUniformBuffer = Proxy->GetWorldSpacePrimitiveUniformBuffer();
+			}
+			BatchElement.MinVertexIndex = 0;
+			BatchElement.MaxVertexIndex = (ParticleCount * NumVerticesPerParticle) - 1;
+			Mesh.CastShadow = Proxy->GetCastShadow();
+			Mesh.DepthPriorityGroup = (ESceneDepthPriorityGroup)Proxy->GetDepthPriorityGroup(View);
 
-					if ( bIsWireframe )
-					{
-						Mesh.MaterialRenderProxy = UMaterial::GetDefaultMaterial( MD_Surface )->GetRenderProxy();
-					}
-					else
-					{
-						Mesh.MaterialRenderProxy = MaterialResource;
-					}
-					Mesh.Type = PT_TriangleList;
+			if ( bIsWireframe )
+			{
+				Mesh.MaterialRenderProxy = UMaterial::GetDefaultMaterial( MD_Surface )->GetRenderProxy();
+			}
+			else
+			{
+				Mesh.MaterialRenderProxy = MaterialResource;
+			}
+			Mesh.Type = PT_TriangleList;
 
-					Mesh.bCanApplyViewModeOverrides = true;
-					Mesh.bUseWireframeSelectionColoring = Proxy->IsSelected();
+			Mesh.bCanApplyViewModeOverrides = true;
+			Mesh.bUseWireframeSelectionColoring = Proxy->IsSelected();
 			
-				#if !(UE_BUILD_SHIPPING || UE_BUILD_TEST)
-					Mesh.VisualizeLODIndex = (int8)Proxy->GetVisualizeLODIndex();
-				#endif
+		#if !(UE_BUILD_SHIPPING || UE_BUILD_TEST)
+			Mesh.VisualizeLODIndex = (int8)Proxy->GetVisualizeLODIndex();
+		#endif
 
-					Collector.AddMesh(ViewIndex, Mesh);
-			}
+			Collector.AddMesh(ViewIndex, Mesh);
 		}
 		else if (SourceData->EmitterRenderMode == ERM_Point)
 		{
@@ -1634,7 +1639,7 @@ void FDynamicMeshEmitterData::GetDynamicMeshElementsEmitter(const FParticleSyste
 
 			FMeshParticleInstanceVertices* InstanceVerticesCPU = nullptr;
 
-			if (MeshVertexFactory->CheckAndUpdateLastFrame(ViewFamily, View))
+			if (MeshVertexFactory->CheckAndUpdateLastFrame(ViewFamily))
 			{
 				FDynamicMeshEmitterCollectorResources& CollectorResources = Collector.AllocateOneFrameResource<FDynamicMeshEmitterCollectorResources>();
 				CollectorResources.VertexFactory = MeshVertexFactory;
@@ -1663,17 +1668,18 @@ void FDynamicMeshEmitterData::GetDynamicMeshElementsEmitter(const FParticleSyste
 
 				if (bInstanced)
 				{
-					FGlobalDynamicVertexBuffer::FAllocation Allocation = FGlobalDynamicVertexBuffer::Get().Allocate(ParticleCount * InstanceVertexStride);
+					const uint32 InstanceFactor = View->IsInstancedStereoPass() ? 2 : 1;
+					FGlobalDynamicVertexBuffer::FAllocation Allocation = FGlobalDynamicVertexBuffer::Get().Allocate(InstanceFactor * ParticleCount * InstanceVertexStride);
 					FGlobalDynamicVertexBuffer::FAllocation DynamicParameterAllocation;
 					uint8* PrevTransformBuffer = nullptr;
 
 					if (bUsesDynamicParameter)
 					{
-						DynamicParameterAllocation = FGlobalDynamicVertexBuffer::Get().Allocate(ParticleCount * DynamicParameterVertexStride);
+						DynamicParameterAllocation = FGlobalDynamicVertexBuffer::Get().Allocate(InstanceFactor * ParticleCount * DynamicParameterVertexStride);
 
 						if (FGlobalDynamicVertexBuffer::Get().IsRenderAlarmLoggingEnabled())
 						{
-							UE_LOG(LogParticles, Warning, TEXT("Panic logging.  Allocated %u bytes for Resource: %s, Owner: %s"), ParticleCount * DynamicParameterVertexStride, *Proxy->GetResourceName().ToString(), *Proxy->GetOwnerName().ToString());
+							UE_LOG(LogParticles, Warning, TEXT("Panic logging.  Allocated %u bytes for Resource: %s, Owner: %s"), InstanceFactor * ParticleCount * DynamicParameterVertexStride, *Proxy->GetResourceName().ToString(), *Proxy->GetOwnerName().ToString());
 						}
 					}
 
@@ -1726,15 +1732,15 @@ void FDynamicMeshEmitterData::GetDynamicMeshElementsEmitter(const FParticleSyste
 						if (Collector.ShouldUseTasks())
 						{
 							Collector.AddTask(
-								[this, View, Proxy, Allocation, DynamicParameterAllocation, PrevTransformBuffer]()
+								[this, View, Proxy, Allocation, DynamicParameterAllocation, PrevTransformBuffer, InstanceFactor]()
 								{
-									GetInstanceData(Allocation.Buffer, DynamicParameterAllocation.Buffer, PrevTransformBuffer, Proxy, View);
+									GetInstanceData(Allocation.Buffer, DynamicParameterAllocation.Buffer, PrevTransformBuffer, Proxy, View, InstanceFactor);
 								}
 							);
 						}
 						else
 						{
-							GetInstanceData(Allocation.Buffer, DynamicParameterAllocation.Buffer, PrevTransformBuffer, Proxy, View);
+							GetInstanceData(Allocation.Buffer, DynamicParameterAllocation.Buffer, PrevTransformBuffer, Proxy, View, InstanceFactor);
 						}
 					}
 
@@ -1777,8 +1783,8 @@ void FDynamicMeshEmitterData::GetDynamicMeshElementsEmitter(const FParticleSyste
 									(void*)InstanceVerticesCPU->DynamicParameterDataAllocationsCPU.GetData(), 
 									PrevTransformBuffer, 
 									Proxy, 
-									View
-									);
+									View,
+									1);
 							}
 						);
 					}
@@ -1789,8 +1795,8 @@ void FDynamicMeshEmitterData::GetDynamicMeshElementsEmitter(const FParticleSyste
 							(void*)InstanceVerticesCPU->DynamicParameterDataAllocationsCPU.GetData(), 
 							PrevTransformBuffer, 
 							Proxy, 
-							View
-							);
+							View,
+							1);
 					}
 				}
 
@@ -2391,7 +2397,7 @@ void FDynamicMeshEmitterData::CalculateParticleTransform(
 	}
 }
 
-void FDynamicMeshEmitterData::GetInstanceData(void* InstanceData, void* DynamicParameterData, void* PrevTransformBuffer, const FParticleSystemSceneProxy* Proxy, const FSceneView* View) const
+void FDynamicMeshEmitterData::GetInstanceData(void* InstanceData, void* DynamicParameterData, void* PrevTransformBuffer, const FParticleSystemSceneProxy* Proxy, const FSceneView* View, uint32 InstanceFactor) const
 {
 	SCOPE_CYCLE_COUNTER(STAT_ParticlePackingTime);
 	SCOPE_CYCLE_COUNTER(STAT_ParticlesOverview_RT_CNC);
@@ -2555,7 +2561,13 @@ void FDynamicMeshEmitterData::GetInstanceData(void* InstanceData, void* DynamicP
 		// The particle's relative time
 		CurrentInstanceVertex->RelativeTime = Particle.RelativeTime;
 
-		TempVert += InstanceVertexStride;
+		//@todo - refactor into instance step rate in the RHI
+		for (uint32 Factor = 1; Factor < InstanceFactor; Factor++)
+		{
+			FMemory::Memcpy(TempVert + InstanceVertexStride * Factor, TempVert, InstanceVertexStride);
+		}
+
+		TempVert += InstanceVertexStride * InstanceFactor;
 	}
 }
 
@@ -2747,7 +2759,7 @@ void FDynamicBeam2EmitterData::GetDynamicMeshElementsEmitter(const FParticleSyst
 	FParticleBeamTrailVertexFactory* BeamTrailVertexFactory = static_cast<FParticleBeamTrailVertexFactory*>(VertexFactory);
 	const bool bIsWireframe = ViewFamily.EngineShowFlags.Wireframe;
 
-	if (BeamTrailVertexFactory->CheckAndUpdateLastFrame(ViewFamily, View))
+	if (BeamTrailVertexFactory->CheckAndUpdateLastFrame(ViewFamily))
 	{
 
 	FGlobalDynamicVertexBuffer::FAllocation DynamicVertexAllocation;
@@ -5520,7 +5532,7 @@ void FDynamicTrailsEmitterData::GetDynamicMeshElementsEmitter(const FParticleSys
 	int32 RenderedPrimitiveCount = 0;
 
 
-	if (BeamTrailVertexFactory->CheckAndUpdateLastFrame(ViewFamily, View))
+	if (BeamTrailVertexFactory->CheckAndUpdateLastFrame(ViewFamily))
 	{
 
 	FGlobalDynamicVertexBuffer::FAllocation DynamicVertexAllocation;
diff --git a/Engine/Source/Runtime/Engine/Public/ParticleHelper.h b/Engine/Source/Runtime/Engine/Public/ParticleHelper.h
--- a/Engine/Source/Runtime/Engine/Public/ParticleHelper.h
+++ b/Engine/Source/Runtime/Engine/Public/ParticleHelper.h
@@ -1711,10 +1711,11 @@ struct FDynamicSpriteEmitterData : public FDynamicSpriteEmitterDataBase
 	 *	@param	ParticleOrder		The (optional) particle ordering to use
 	 *	@param	InCameraPosition	The position of the camera in world space.
 	 *	@param	InLocalToWorld		Transform from local to world space.
+	 *	@param	InstanceFactor		The factor to duplicate instances by.
 	 *
 	 *	@return	bool			true if successful, false if failed
 	 */
-	bool GetVertexAndIndexData(void* VertexData, void* DynamicParameterVertexData, void* FillIndexData, FParticleOrder* ParticleOrder, const FVector& InCameraPosition, const FMatrix& InLocalToWorld) const;
+	bool GetVertexAndIndexData(void* VertexData, void* DynamicParameterVertexData, void* FillIndexData, FParticleOrder* ParticleOrder, const FVector& InCameraPosition, const FMatrix& InLocalToWorld, uint32 InstanceFactor) const;
 
 	/**
 	 *	Retrieve the vertex and (optional) index required to render this emitter.
@@ -1862,8 +1863,9 @@ struct FDynamicMeshEmitterData : public FDynamicSpriteEmitterDataBase
 	 *	@param	PrevTransformBuffer     The memory to fill the vertex prev transform data into. May be null
 	 *	@param	Proxy                   The scene proxy for the particle system that owns this emitter
 	 *	@param	View                    The scene view being rendered
+	 *	@param	InstanceFactor			The factor to duplicate instances by
 	 */
-	void GetInstanceData(void* InstanceData, void* DynamicParameterData, void* PrevTransformBuffer, const FParticleSystemSceneProxy* Proxy, const FSceneView* View) const;
+	void GetInstanceData(void* InstanceData, void* DynamicParameterData, void* PrevTransformBuffer, const FParticleSystemSceneProxy* Proxy, const FSceneView* View, uint32 InstanceFactor) const;
 
 	/**
 	 *	Helper function for retrieving the particle transform.
diff --git a/Engine/Source/Runtime/Engine/Public/ParticleVertexFactory.h b/Engine/Source/Runtime/Engine/Public/ParticleVertexFactory.h
--- a/Engine/Source/Runtime/Engine/Public/ParticleVertexFactory.h
+++ b/Engine/Source/Runtime/Engine/Public/ParticleVertexFactory.h
@@ -41,7 +41,6 @@ public:
 		: FVertexFactory(InFeatureLevel)
 		, LastFrameSetup(MAX_uint32)
 		, LastViewFamily(nullptr)
-		, LastView(nullptr)
 		, LastFrameRealTime(-1.0f)
 		, ParticleFactoryType(Type)
 		, bInUse(false)
@@ -80,16 +79,15 @@ public:
 
 	ERHIFeatureLevel::Type GetFeatureLevel() const { check(HasValidFeatureLevel());  return FRenderResource::GetFeatureLevel(); }
 
-	bool CheckAndUpdateLastFrame(const FSceneViewFamily& ViewFamily, const FSceneView *View = nullptr) const
+	bool CheckAndUpdateLastFrame(const FSceneViewFamily& ViewFamily) const
 	{
-		if (LastFrameSetup != MAX_uint32 && (&ViewFamily == LastViewFamily) && (View == LastView) && ViewFamily.FrameNumber == LastFrameSetup && LastFrameRealTime == ViewFamily.CurrentRealTime)
+		if (LastFrameSetup != MAX_uint32 && (&ViewFamily == LastViewFamily) && ViewFamily.FrameNumber == LastFrameSetup && LastFrameRealTime == ViewFamily.CurrentRealTime)
 		{
 			return false;
 		}
 		LastFrameSetup = ViewFamily.FrameNumber;
 		LastFrameRealTime = ViewFamily.CurrentRealTime;
 		LastViewFamily = &ViewFamily;
-		LastView = View;
 		return true;
 	}
 
@@ -101,7 +99,6 @@ private:
 	/** Last state where we set this. We only need to setup these once per frame, so detemine same frame by number, time, and view family. */
 	mutable uint32 LastFrameSetup;
 	mutable const FSceneViewFamily *LastViewFamily;
-	mutable const FSceneView *LastView;
 	mutable float LastFrameRealTime;
 
 	/** The type of the vertex factory. */
