Fixed missing include "FXSystemPrivate.h".

diff --git a/Engine/Source/Runtime/Engine/Private/Particles/ParticleModules.cpp b/Engine/Source/Runtime/Engine/Private/Particles/ParticleModules.cpp
--- a/Engine/Source/Runtime/Engine/Private/Particles/ParticleModules.cpp
+++ b/Engine/Source/Runtime/Engine/Private/Particles/ParticleModules.cpp
@@ -83,6 +83,7 @@
 #include "Particles/TypeData/ParticleModuleTypeDataMesh.h"
 #include "Particles/ParticleLODLevel.h"
 #include "Particles/ParticleModuleRequired.h"
+#include "Particles/FXSystemPrivate.h"
 #include "Distributions/DistributionFloatUniformCurve.h"
 #include "Engine/InterpCurveEdSetup.h"
 #include "Engine/RendererSettings.h"
