Disable NetGuidCache History tracking by default.

diff --git a/Engine/Source/Runtime/Engine/Classes/Engine/PackageMapClient.h b/Engine/Source/Runtime/Engine/Classes/Engine/PackageMapClient.h
--- a/Engine/Source/Runtime/Engine/Classes/Engine/PackageMapClient.h
+++ b/Engine/Source/Runtime/Engine/Classes/Engine/PackageMapClient.h
@@ -270,6 +270,9 @@ private:
 
 #if !(UE_BUILD_SHIPPING || UE_BUILD_TEST)
 public:
+
+	static const bool IsHistoryEnabled();
+
 	// History for debugging entries in the guid cache
 	TMap<FNetworkGUID, FString>						History;
 private:
diff --git a/Engine/Source/Runtime/Engine/Private/DataChannel.cpp b/Engine/Source/Runtime/Engine/Private/DataChannel.cpp
--- a/Engine/Source/Runtime/Engine/Private/DataChannel.cpp
+++ b/Engine/Source/Runtime/Engine/Private/DataChannel.cpp
@@ -3937,6 +3937,12 @@ static void	FindNetGUID( const TArray<FString>& Args, UWorld* InWorld )
 			continue;
 		}
 
+		if (!FNetGUIDCache::IsHistoryEnabled())
+		{
+			UE_LOG(LogNet, Warning, TEXT("FindNetGUID - GuidCacheHistory is not enabled"));
+			return;
+		}
+
 		if (Args.Num() <= 0)
 		{
 			// Display all
diff --git a/Engine/Source/Runtime/Engine/Private/PackageMapClient.cpp b/Engine/Source/Runtime/Engine/Private/PackageMapClient.cpp
--- a/Engine/Source/Runtime/Engine/Private/PackageMapClient.cpp
+++ b/Engine/Source/Runtime/Engine/Private/PackageMapClient.cpp
@@ -1898,11 +1898,14 @@ bool UPackageMapClient::PrintExportBatch()
 
 	// Print the whole thing for reference
 #if !(UE_BUILD_SHIPPING || UE_BUILD_TEST)
-	for (auto It = GuidCache->History.CreateIterator(); It; ++It)
+	if (FNetGUIDCache::IsHistoryEnabled())
 	{
-		FString Str = It.Value();
-		FNetworkGUID NetGUID = It.Key();
-		UE_LOG(LogNetPackageMap, Warning, TEXT("<%s> - %s"), *NetGUID.ToString(), *Str);
+		for (auto It = GuidCache->History.CreateIterator(); It; ++It)
+		{
+			FString Str = It.Value();
+			FNetworkGUID NetGUID = It.Key();
+			UE_LOG(LogNetPackageMap, Warning, TEXT("<%s> - %s"), *NetGUID.ToString(), *Str);
+		}
 	}
 #endif
 
@@ -2274,13 +2277,19 @@ void FNetGUIDCache::RegisterNetGUID_Internal( const FNetworkGUID& NetGUID, const
 		NetGUIDLookup.Add( CacheObject.Object, NetGUID );
 
 #if !(UE_BUILD_SHIPPING || UE_BUILD_TEST)
-		History.Add( NetGUID, CacheObject.Object->GetPathName() );
+		if (IsHistoryEnabled())
+		{
+			History.Add(NetGUID, CacheObject.Object->GetPathName());
+		}
 #endif
 	}
 	else
 	{
 #if !(UE_BUILD_SHIPPING || UE_BUILD_TEST)
-		History.Add( NetGUID, CacheObject.PathName.ToString() );
+		if (IsHistoryEnabled())
+		{
+			History.Add(NetGUID, CacheObject.PathName.ToString());
+		}
 #endif
 	}
 }
@@ -3109,6 +3118,20 @@ void FNetGUIDCache::ResetCacheForDemo()
 	NetFieldExportGroupPathToIndex.Reset();
 }
 
+#if !(UE_BUILD_SHIPPING || UE_BUILD_TEST)
+static int32 bIsNetGuidCacheHistoryEnabled = 0;
+static FAutoConsoleVariableRef CVarIsNetGuidCacheHistoryEnabled(
+	TEXT("Net.NetGuidCacheHistoryEnabled"),
+	bIsNetGuidCacheHistoryEnabled,
+	TEXT("When enabled, allows logging of NetGUIDCache History. Warning, this can eat up a lot of memory, and won't free itself until the Cache is destroyed.")
+);
+
+const bool FNetGUIDCache::IsHistoryEnabled()
+{
+	return !!bIsNetGuidCacheHistoryEnabled;
+}
+#endif
+
 void FNetGUIDCache::CountBytes(FArchive& Ar) const
 {
 	GRANULAR_NETWORK_MEMORY_TRACKING_INIT(Ar, "FNetGuidCache::CountBytes");
