Fix RepState overcounting memory.

diff --git a/Engine/Source/Runtime/Engine/Private/RepLayout.cpp b/Engine/Source/Runtime/Engine/Private/RepLayout.cpp
--- a/Engine/Source/Runtime/Engine/Private/RepLayout.cpp
+++ b/Engine/Source/Runtime/Engine/Private/RepLayout.cpp
@@ -5119,9 +5119,6 @@ void FRepLayout::CountBytes(FArchive& Ar) const
 
 void FRepState::CountBytes(FArchive& Ar) const
 {
-	const SIZE_T SizeOfThis = sizeof(FRepState);
-	Ar.CountBytes(SizeOfThis, SizeOfThis);
-
 	StaticBuffer.CountBytes(Ar);
 	GuidReferencesMap.CountBytes(Ar);
 	for (const auto& GuidRefPair : GuidReferencesMap)
