Fixup some nullptr dereference issues in RepLayout

diff --git a/Engine/Source/Runtime/Engine/Private/RepLayout.cpp b/Engine/Source/Runtime/Engine/Private/RepLayout.cpp
--- a/Engine/Source/Runtime/Engine/Private/RepLayout.cpp
+++ b/Engine/Source/Runtime/Engine/Private/RepLayout.cpp
@@ -1094,7 +1094,7 @@ void FRepLayout::PostReplicate(
 
 void FRepLayout::ReceivedNak(FRepState* RepState, int32 NakPacketId) const
 {
-	if (RepState == NULL)
+	if (RepState == nullptr)
 	{
 		return;		// I'm not 100% certain why this happens, the only think I can think of is this is a bNetTemporary?
 	}
@@ -1106,18 +1106,20 @@ void FRepLayout::ReceivedNak(FRepState* RepState, int32 NakPacketId) const
 	}
 	else if (LayoutState == ERepLayoutState::Normal)
 	{
-		FSendingRepState* SendingRepState = RepState->GetSendingRepState();
-		for (int32 i = SendingRepState->HistoryStart; i < SendingRepState->HistoryEnd; ++i)
+		if (FSendingRepState* SendingRepState = RepState->GetSendingRepState())
 		{
-			const int32 HistoryIndex = i % FSendingRepState::MAX_CHANGE_HISTORY;
+			for (int32 i = SendingRepState->HistoryStart; i < SendingRepState->HistoryEnd; ++i)
+			{
+				const int32 HistoryIndex = i % FSendingRepState::MAX_CHANGE_HISTORY;
 
-			FRepChangedHistory & HistoryItem = SendingRepState->ChangeHistory[HistoryIndex];
+				FRepChangedHistory & HistoryItem = SendingRepState->ChangeHistory[HistoryIndex];
 
-			if (!HistoryItem.Resend && HistoryItem.OutPacketIdRange.InRange(NakPacketId))
-			{
-				check(HistoryItem.Changed.Num() > 0);
-				HistoryItem.Resend = true;
-				++SendingRepState->NumNaks;
+				if (!HistoryItem.Resend && HistoryItem.OutPacketIdRange.InRange(NakPacketId))
+				{
+					check(HistoryItem.Changed.Num() > 0);
+					HistoryItem.Resend = true;
+					++SendingRepState->NumNaks;
+				}
 			}
 		}
 	}
@@ -1125,26 +1127,28 @@ void FRepLayout::ReceivedNak(FRepState* RepState, int32 NakPacketId) const
 
 bool FRepLayout::AllAcked(FRepState* RepState) const
 {
-	FSendingRepState* SendingRepState = RepState->GetSendingRepState();
-	if (SendingRepState->HistoryStart != SendingRepState->HistoryEnd)
+	if (FSendingRepState* SendingRepState = RepState->GetSendingRepState())
 	{
-		// We have change lists that haven't been acked
-		return false;
-	}
+		if (SendingRepState->HistoryStart != SendingRepState->HistoryEnd)
+		{
+			// We have change lists that haven't been acked
+			return false;
+		}
 
-	if (SendingRepState->NumNaks > 0)
-	{
-		return false;
-	}
+		if (SendingRepState->NumNaks > 0)
+		{
+			return false;
+		}
 
-	if (!SendingRepState->bOpenAckedCalled)
-	{
-		return false;
-	}
+		if (!SendingRepState->bOpenAckedCalled)
+		{
+			return false;
+		}
 
-	if (SendingRepState->PreOpenAckHistory.Num()> 0)
-	{
-		return false;
+		if (SendingRepState->PreOpenAckHistory.Num() > 0)
+		{
+			return false;
+		}
 	}
 
 	return true;
@@ -1152,6 +1156,7 @@ bool FRepLayout::AllAcked(FRepState* RepState) const
 
 bool FRepLayout::ReadyForDormancy(FRepState* RepState) const
 {
+	// Clients should never go dormant.
 	if (RepState == nullptr || RepState->GetSendingRepState() == nullptr)
 	{
 		return false;
