Fix crash in DemoNetDriver from using the wrong replicator

diff --git a/Engine/Source/Runtime/Engine/Private/DemoNetDriver.cpp b/Engine/Source/Runtime/Engine/Private/DemoNetDriver.cpp
--- a/Engine/Source/Runtime/Engine/Private/DemoNetDriver.cpp
+++ b/Engine/Source/Runtime/Engine/Private/DemoNetDriver.cpp
@@ -5528,9 +5528,8 @@ void UDemoNetDriver::QueueNetStartupActorForRollbackViaDeletion( AActor* Actor )
 
 	if (GDemoSaveRollbackActorState != 0)
 	{
-		TSharedPtr<FObjectReplicator> NewReplicator = MakeShared<FObjectReplicator>();
-		if (NewReplicator.IsValid())
 		{
+			TSharedPtr<FObjectReplicator> NewReplicator = MakeShared<FObjectReplicator>();
 			NewReplicator->InitWithObject(Actor->GetArchetype(), ServerConnection, false);
 
 			if (NewReplicator->RepLayout.IsValid() && NewReplicator->RepState.IsValid())
@@ -5551,20 +5550,17 @@ void UDemoNetDriver::QueueNetStartupActorForRollbackViaDeletion( AActor* Actor )
 			if (ActorComp)
 			{
 				TSharedPtr<FObjectReplicator> SubObjReplicator = MakeShared<FObjectReplicator>();
-				if (SubObjReplicator.IsValid())
+				SubObjReplicator->InitWithObject(ActorComp->GetArchetype(), ServerConnection, false);
+
+				if (SubObjReplicator->RepLayout.IsValid() && SubObjReplicator->RepState.IsValid())
 				{
-					SubObjReplicator->InitWithObject(ActorComp->GetArchetype(), ServerConnection, false);
+					FReceivingRepState* ReceivingRepState = SubObjReplicator->RepState->GetReceivingRepState();
+					FRepShadowDataBuffer ShadowData(ReceivingRepState->StaticBuffer.GetData());
+					FConstRepObjectDataBuffer ActorCompData(ActorComp);
 
-					if (SubObjReplicator->RepLayout.IsValid() && SubObjReplicator->RepState.IsValid())
+					if (SubObjReplicator->RepLayout->DiffStableProperties(nullptr, &RollbackActor.ObjReferences, ShadowData, ActorCompData))
 					{
-						FReceivingRepState* ReceivingRepState = NewReplicator->RepState->GetReceivingRepState();
-						FRepShadowDataBuffer ShadowData(ReceivingRepState->StaticBuffer.GetData());
-						FConstRepObjectDataBuffer ActorCompData(ActorComp);
-
-						if (SubObjReplicator->RepLayout->DiffStableProperties(nullptr, &RollbackActor.ObjReferences, ShadowData, ActorCompData))
-						{
-							RollbackActor.SubObjRepState.Add(ActorComp->GetFullName(), MakeShareable(SubObjReplicator->RepState.Release()));
-						}
+						RollbackActor.SubObjRepState.Add(ActorComp->GetFullName(), MakeShareable(SubObjReplicator->RepState.Release()));
 					}
 				}
 			}
