Exposing GDoAsyncLoadingWhileWaitingForVSync to be modifiable by code instead of just commandline.

diff --git a/Engine/Source/Runtime/Engine/Public/UnrealEngine.h b/Engine/Source/Runtime/Engine/Public/UnrealEngine.h
--- a/Engine/Source/Runtime/Engine/Public/UnrealEngine.h
+++ b/Engine/Source/Runtime/Engine/Public/UnrealEngine.h
@@ -17,6 +17,7 @@ class UFont;
 class ULocalPlayer;
 
 extern ENGINE_API int32 GShowMaterialDrawEvents;
+extern ENGINE_API int32 GDoAsyncLoadingWhileWaitingForVSync;
 
 ENGINE_API DECLARE_LOG_CATEGORY_EXTERN(LogEngine, Log, All);
 //
