Fixup networking related merge fallout.

diff --git a/Engine/Source/Runtime/Engine/Private/DataReplication.cpp b/Engine/Source/Runtime/Engine/Private/DataReplication.cpp
--- a/Engine/Source/Runtime/Engine/Private/DataReplication.cpp
+++ b/Engine/Source/Runtime/Engine/Private/DataReplication.cpp
@@ -440,28 +440,7 @@ FReplicationChangelistMgr::~FReplicationChangelistMgr()
 {
 }
 
-PRAGMA_DISABLE_DEPRECATION_WARNINGS
-void FReplicationChangelistMgr::Update(const UObject* InObject, const uint32 ReplicationFrame, const int32 LastCompareIndex, const FReplicationFlags& RepFlags, const bool bForceCompare)
-{
-	// See if we can re-use the work already done on a previous connection
-	// Rules:
-	//	1. We always compare once per frame (i.e. check LastReplicationFrame == ReplicationFrame)
-	//	2. We check LastCompareIndex > 1 so we can do at least one pass per connection to compare all properties
-	//		This is necessary due to how RemoteRole is manipulated per connection, so we need to give all connections a chance to see if it changed
-	//	3. We ALWAYS compare on bNetInitial to make sure we have a fresh changelist of net initial properties in this case
-	if (!bForceCompare && GShareShadowState && !RepFlags.bNetInitial && LastCompareIndex > 1 && LastReplicationFrame == ReplicationFrame)
-	{
-		INC_DWORD_STAT_BY(STAT_NetSkippedDynamicProps, 1);
-		return;
-	}
-
-	RepLayout->CompareProperties(nullptr, RepChangelistState.Get(), (const uint8*)InObject, RepFlags);
-
-	LastReplicationFrame = ReplicationFrame;
-}
-PRAGMA_ENABLE_DEPRECATION_WARNINGS
-
-void FReplicationChangelistMgr::Update(FRepState* RESTRICT RepState, const UObject* InObject, const uint32 ReplicationFrame, const FReplicationFlags& RepFlags, const bool bForceCompare)
+void FReplicationChangelistMgr::Update(FSendingRepState* RESTRICT RepState, const UObject* InObject, const uint32 ReplicationFrame, const FReplicationFlags& RepFlags, const bool bForceCompare)
 {
 	// See if we can re-use the work already done on a previous connection
 	// Rules:
@@ -1226,7 +1205,7 @@ void FObjectReplicator::ReplicateCustomDeltaProperties( FNetBitWriter & Bunch, F
 	UNetConnection * OwningChannelConnection = OwningChannel->Connection;
 
 	// Initialize a map of which conditions are valid
-	const TStaticBitArray<COND_Max> ConditionMap = FRepState::BuildConditionMap(RepFlags);
+	const TStaticBitArray<COND_Max> ConditionMap = BuildConditionMapFromRepFlags(RepFlags);
 
 	// Make sure net field export group is registered
 	FNetFieldExportGroup* NetFieldExportGroup = OwningChannel->GetOrCreateNetFieldExportGroupForClassNetCache( Object );
diff --git a/Engine/Source/Runtime/Engine/Private/RepLayout.cpp b/Engine/Source/Runtime/Engine/Private/RepLayout.cpp
--- a/Engine/Source/Runtime/Engine/Private/RepLayout.cpp
+++ b/Engine/Source/Runtime/Engine/Private/RepLayout.cpp
@@ -527,16 +527,50 @@
 	const TArray<FRepLayoutCmd>& Cmds;
 };
 
+TStaticBitArray<COND_Max> BuildConditionMapFromRepFlags(FReplicationFlags RepFlags)
+{
+	TStaticBitArray<COND_Max> ConditionMap;
+
+	// Setup condition map
+	const bool bIsInitial = RepFlags.bNetInitial ? true : false;
+	const bool bIsOwner = RepFlags.bNetOwner ? true : false;
+	const bool bIsSimulated = RepFlags.bNetSimulated ? true : false;
+	const bool bIsPhysics = RepFlags.bRepPhysics ? true : false;
+	const bool bIsReplay = RepFlags.bReplay ? true : false;
+
+	ConditionMap[COND_None] = true;
+	ConditionMap[COND_InitialOnly] = bIsInitial;
+
+	ConditionMap[COND_OwnerOnly] = bIsOwner;
+	ConditionMap[COND_SkipOwner] = !bIsOwner;
+
+	ConditionMap[COND_SimulatedOnly] = bIsSimulated;
+	ConditionMap[COND_SimulatedOnlyNoReplay] = bIsSimulated && !bIsReplay;
+	ConditionMap[COND_AutonomousOnly] = !bIsSimulated;
+
+	ConditionMap[COND_SimulatedOrPhysics] = bIsSimulated || bIsPhysics;
+	ConditionMap[COND_SimulatedOrPhysicsNoReplay] = (bIsSimulated || bIsPhysics) && !bIsReplay;
+
+	ConditionMap[COND_InitialOrOwner] = bIsInitial || bIsOwner;
+	ConditionMap[COND_ReplayOrOwner] = bIsReplay || bIsOwner;
+	ConditionMap[COND_ReplayOnly] = bIsReplay;
+	ConditionMap[COND_SkipReplay] = !bIsReplay;
+
+	ConditionMap[COND_Custom] = true;
+
+	return ConditionMap;
+}
+
 uint16 FRepLayout::CompareProperties_r(
-	FRepState* RESTRICT		RepState,
-	const int32				CmdStart,
-	const int32				CmdEnd,
-	const uint8* RESTRICT	ShadowData,
-	const uint8* RESTRICT	Data,
-	TArray<uint16>&			Changed,
-	uint16					Handle,
-	const bool				bIsInitial,
-	const bool				bForceFail) const
+	FSendingRepState* RESTRICT RepState,
+	const int32 CmdStart,
+	const int32 CmdEnd,
+	const uint8* RESTRICT ShadowData,
+	const uint8* RESTRICT Data,
+	TArray<uint16>& Changed,
+	uint16 Handle,
+	const bool bIsInitial,
+	const bool bForceFail) const
 {
 	check(ShadowData);
 
@@ -610,14 +644,14 @@
 }
 
 void FRepLayout::CompareProperties_Array_r(
-	FRepState* RESTRICT		RepState,
-	const uint8* RESTRICT	ShadowData,
-	const uint8* RESTRICT	Data,
-	TArray<uint16>&			Changed,
-	const uint16			CmdIndex,
-	const uint16			Handle,
-	const bool				bIsInitial,
-	const bool				bForceFail
+	FSendingRepState* RESTRICT RepState,
+	const uint8* RESTRICT ShadowData,
+	const uint8* RESTRICT Data,
+	TArray<uint16>& Changed,
+	const uint16 CmdIndex,
+	const uint16 Handle,
+	const bool bIsInitial,
+	const bool bForceFail
 	) const
 {
 	const FRepLayoutCmd& Cmd = Cmds[CmdIndex];
@@ -668,10 +702,10 @@
 }
 
 bool FRepLayout::CompareProperties(
-	FRepState* RESTRICT				RepState,
-	FRepChangelistState* RESTRICT	RepChangelistState,
-	const uint8* RESTRICT			Data,
-	const FReplicationFlags&		RepFlags) const
+	FSendingRepState* RESTRICT RepState,
+	FRepChangelistState* RESTRICT RepChangelistState,
+	const uint8* RESTRICT Data,
+	const FReplicationFlags& RepFlags) const
 {
 	SCOPE_CYCLE_COUNTER(STAT_NetReplicateDynamicPropCompareTime);
 
@@ -880,7 +914,10 @@
 	if (NewlyActiveChangelist.Num() > 0)
 	{
 		TArray<uint16> Temp = MoveTemp(Changed);
+
+PRAGMA_DISABLE_DEPRECATION_WARNINGS
 		MergeChangeList(Data, NewlyActiveChangelist, Temp, Changed);
+PRAGMA_ENABLE_DEPRECATION_WARNINGS
 	}
 
 	// We're all caught up now
@@ -937,7 +974,10 @@
 	if (NewlyInactiveChangelist.Num() > 1)
 	{
 		TArray<uint16> Temp = MoveTemp(RepState->InactiveChangelist);
+
+PRAGMA_DISABLE_DEPRECATION_WARNINGS
 		MergeChangeList(Data, NewlyInactiveChangelist, Temp, RepState->InactiveChangelist);
+PRAGMA_ENABLE_DEPRECATION_WARNINGS
 	}
 
 	// Send the final merged change list
@@ -2026,14 +2066,14 @@
 }
 
 void FRepLayout::SendAllProperties_BackwardsCompatible_r(
-	FRepState* RESTRICT					RepState,
-	FNetBitWriter&						Writer,
-	const bool							bDoChecksum,
-	UPackageMapClient*					PackageMapClient,
-	FNetFieldExportGroup*				NetFieldExportGroup,
-	const int32							CmdStart,
-	const int32							CmdEnd, 
-	const uint8*						SourceData) const
+	FSendingRepState* RESTRICT RepState,
+	FNetBitWriter& Writer,
+	const bool bDoChecksum,
+	UPackageMapClient* PackageMapClient,
+	FNetFieldExportGroup* NetFieldExportGroup,
+	const int32 CmdStart,
+	const int32 CmdEnd, 
+	const uint8* SourceData) const
 {
 	FNetBitWriter TempWriter(Writer.PackageMap, 0);
 
@@ -4924,81 +4964,13 @@
 	}
 }
 
-TStaticBitArray<COND_Max> FRepState::BuildConditionMap(const FReplicationFlags& RepFlags)
-{
-	TStaticBitArray<COND_Max> ConditionMap;
-
-	// Setup condition map
-	const bool bIsInitial = RepFlags.bNetInitial ? true : false;
-	const bool bIsOwner = RepFlags.bNetOwner ? true : false;
-	const bool bIsSimulated = RepFlags.bNetSimulated ? true : false;
-	const bool bIsPhysics = RepFlags.bRepPhysics ? true : false;
-	const bool bIsReplay = RepFlags.bReplay ? true : false;
-
-	ConditionMap[COND_None] = true;
-	ConditionMap[COND_InitialOnly] = bIsInitial;
-
-	ConditionMap[COND_OwnerOnly] = bIsOwner;
-	ConditionMap[COND_SkipOwner] = !bIsOwner;
-
-	ConditionMap[COND_SimulatedOnly] = bIsSimulated;
-	ConditionMap[COND_SimulatedOnlyNoReplay] = bIsSimulated && !bIsReplay;
-	ConditionMap[COND_AutonomousOnly] = !bIsSimulated;
-
-	ConditionMap[COND_SimulatedOrPhysics] = bIsSimulated || bIsPhysics;
-	ConditionMap[COND_SimulatedOrPhysicsNoReplay] = (bIsSimulated || bIsPhysics) && !bIsReplay;
-
-	ConditionMap[COND_InitialOrOwner] = bIsInitial || bIsOwner;
-	ConditionMap[COND_ReplayOrOwner] = bIsReplay || bIsOwner;
-	ConditionMap[COND_ReplayOnly] = bIsReplay;
-	ConditionMap[COND_SkipReplay] = !bIsReplay;
-
-	ConditionMap[COND_Custom] = true;
-
-	return ConditionMap;
-}
-
-TStaticBitArray<COND_Max> FRepState::BuildConditionMap(const FReplicationFlags& RepFlags)
-{
-	TStaticBitArray<COND_Max> ConditionMap;
-
-	// Setup condition map
-	const bool bIsInitial = RepFlags.bNetInitial ? true : false;
-	const bool bIsOwner = RepFlags.bNetOwner ? true : false;
-	const bool bIsSimulated = RepFlags.bNetSimulated ? true : false;
-	const bool bIsPhysics = RepFlags.bRepPhysics ? true : false;
-	const bool bIsReplay = RepFlags.bReplay ? true : false;
-
-	ConditionMap[COND_None] = true;
-	ConditionMap[COND_InitialOnly] = bIsInitial;
-
-	ConditionMap[COND_OwnerOnly] = bIsOwner;
-	ConditionMap[COND_SkipOwner] = !bIsOwner;
-
-	ConditionMap[COND_SimulatedOnly] = bIsSimulated;
-	ConditionMap[COND_SimulatedOnlyNoReplay] = bIsSimulated && !bIsReplay;
-	ConditionMap[COND_AutonomousOnly] = !bIsSimulated;
-
-	ConditionMap[COND_SimulatedOrPhysics] = bIsSimulated || bIsPhysics;
-	ConditionMap[COND_SimulatedOrPhysicsNoReplay] = (bIsSimulated || bIsPhysics) && !bIsReplay;
-
-	ConditionMap[COND_InitialOrOwner] = bIsInitial || bIsOwner;
-	ConditionMap[COND_ReplayOrOwner] = bIsReplay || bIsOwner;
-	ConditionMap[COND_ReplayOnly] = bIsReplay;
-	ConditionMap[COND_SkipReplay] = !bIsReplay;
-
-	ConditionMap[COND_Custom] = true;
-
-	return ConditionMap;
-}
-
 void FRepLayout::RebuildConditionalProperties(
 	FSendingRepState* RESTRICT RepState,
 	const FReplicationFlags& RepFlags) const
 {
 	SCOPE_CYCLE_COUNTER(STAT_NetRebuildConditionalTime);
 	
-	TStaticBitArray<COND_Max> ConditionMap = FRepState::BuildConditionMap(RepFlags);
+	TStaticBitArray<COND_Max> ConditionMap = BuildConditionMapFromRepFlags(RepFlags);
 	for (auto It = TBitArray<>::FIterator(RepState->InactiveParents); It; ++It)
 	{
 		It.GetValue() = !ConditionMap[Parents[It.GetIndex()].Condition];
@@ -5065,7 +5037,7 @@
 
 		// Start out the conditional props based on a default RepFlags struct
 		// It will rebuild if it ever changes
-		RepState->InactiveParents.Init(false, Parents.Num());
+		RepState->SendingRepState->InactiveParents.Init(false, Parents.Num());
 		RebuildConditionalProperties(RepState->SendingRepState.Get(), FReplicationFlags());
 	}
 	
diff --git a/Engine/Source/Runtime/Engine/Public/Net/DataReplication.h b/Engine/Source/Runtime/Engine/Public/Net/DataReplication.h
--- a/Engine/Source/Runtime/Engine/Public/Net/DataReplication.h
+++ b/Engine/Source/Runtime/Engine/Public/Net/DataReplication.h
@@ -18,6 +18,7 @@ class FOutBunch;
 class FRepChangelistState;
 class FRepLayout;
 class FRepState;
+class FSendingRepState;
 class UNetConnection;
 class UNetDriver;
 class AActor;
@@ -61,9 +62,6 @@ public:
 
 	~FReplicationChangelistMgr();
 
-	UE_DEPRECATED(4.22, "Please use the version of Update that accepts a FRepState pointer.")
-	void Update(const UObject* InObject, const uint32 ReplicationFrame, const int32 LastCompareIndex, const FReplicationFlags& RepFlags, const bool bForceCompare);
-
 	/**
 	 * Updates the shared RepChangelistState for the given object, potentially skipping unnecessary updates.
 	 * See FRepLayout::CompareProperties.
@@ -73,7 +71,7 @@ public:
 	 * @param RepFlags		Replication Flags that will be used if the object needs to be replicated.
 	 * @param bForceCompare	Force the comparison, even if other connections have already done it this frame.
 	 */
-	void Update(FRepState* RESTRICT RepState, const UObject* InObject, const uint32 ReplicationFrame, const FReplicationFlags& RepFlags, const bool bForceCompare);
+	void Update(FSendingRepState* RESTRICT RepState, const UObject* InObject, const uint32 ReplicationFrame, const FReplicationFlags& RepFlags, const bool bForceCompare);
 
 	FRepChangelistState* GetRepChangelistState() const { return RepChangelistState.Get(); }
 
diff --git a/Engine/Source/Runtime/Engine/Public/Net/RepLayout.h b/Engine/Source/Runtime/Engine/Public/Net/RepLayout.h
--- a/Engine/Source/Runtime/Engine/Public/Net/RepLayout.h
+++ b/Engine/Source/Runtime/Engine/Public/Net/RepLayout.h
@@ -49,6 +49,9 @@ enum class EReceivePropertiesFlags : uint32
 
 ENUM_CLASS_FLAGS(EReceivePropertiesFlags);
 
+/** Builds a Bitfield that flags whether or not given replication conditions are met. */
+extern TStaticBitArray<COND_Max> BuildConditionMapFromRepFlags(FReplicationFlags RepFlags);
+
 /** Stores meta data about a given Replicated property. */
 class FRepChangedParent
 {
@@ -1474,10 +1477,10 @@ public:
 	 * @param RepFlags				Flags that will be used if the object is replicated.
 	 */
 	bool CompareProperties(
-		FRepState* RESTRICT				RepState,
-		FRepChangelistState* RESTRICT	RepChangelistState,
-		const uint8* RESTRICT			Data,
-		const FReplicationFlags&		RepFlags) const;
+		FSendingRepState* RESTRICT RepState,
+		FRepChangelistState* RESTRICT RepChangelistState,
+		const uint8* RESTRICT Data,
+		const FReplicationFlags& RepFlags) const;
 
 	//~ Begin FGCObject Interface
 	ENGINE_API virtual void AddReferencedObjects(FReferenceCollector& Collector) override;
@@ -1584,14 +1587,14 @@ private:
 		const uint8* RESTRICT SourceData) const;
 
 	void SendAllProperties_BackwardsCompatible_r(
-		FRepState* RESTRICT					RepState,
-		FNetBitWriter&						Writer,
-		const bool							bDoChecksum,
-		UPackageMapClient*					PackageMapClient,
-		FNetFieldExportGroup*				NetFieldExportGroup,
-		const int32							CmdStart,
-		const int32							CmdEnd,
-		const uint8*						SourceData) const;
+		FSendingRepState* RESTRICT RepState,
+		FNetBitWriter& Writer,
+		const bool bDoChecksum,
+		UPackageMapClient* PackageMapClient,
+		FNetFieldExportGroup* NetFieldExportGroup,
+		const int32 CmdStart,
+		const int32 CmdEnd,
+		const uint8* SourceData) const;
 
 	void SendProperties_r(
 		FSendingRepState* RESTRICT RepState,
