NetPacketNotify now verifies that received acks are within the expected window. Packet will be rejected if it is not
- Removed check that could trigger on the server if we got bad data

diff --git a/Engine/Source/Runtime/Engine/Private/Net/NetPacketNotify.cpp b/Engine/Source/Runtime/Engine/Private/Net/NetPacketNotify.cpp
--- a/Engine/Source/Runtime/Engine/Private/Net/NetPacketNotify.cpp
+++ b/Engine/Source/Runtime/Engine/Private/Net/NetPacketNotify.cpp
@@ -25,28 +25,27 @@ SIZE_T FNetPacketNotify::GetCurrentSequenceHistoryLength() const
 
 FNetPacketNotify::SequenceNumberT FNetPacketNotify::UpdateInAckSeqAck(SequenceNumberT::DifferenceT AckCount, SequenceNumberT AckedSeq)
 {
-	check((SIZE_T)AckCount <= AckRecord.Count());
-
-	if (AckCount > 1)
+	if ((SIZE_T)AckCount <= AckRecord.Count())
 	{
-		AckRecord.PopNoCheck(AckCount - 1);
-	}
+		if (AckCount > 1)
+		{
+			AckRecord.PopNoCheck(AckCount - 1);
+		}
 
-	FSentAckData AckData = AckRecord.PeekNoCheck();
-	AckRecord.PopNoCheck();
+		FSentAckData AckData = AckRecord.PeekNoCheck();
+		AckRecord.PopNoCheck();
 
-	// verify that we have a matching sequence number
-	if (AckData.OutSeq == AckedSeq)
-	{
-		return AckData.InAckSeq;
+		// verify that we have a matching sequence number
+		if (AckData.OutSeq == AckedSeq)
+		{
+			return AckData.InAckSeq;
+		}
 	}
-	else
-	{
-		UE_LOG_PACKET_NOTIFY_WARNING(TEXT("FNetPacketNotify::UpdateInAckSeqAck - Failed to find matching AckRecord for %u, (Found %u)"), AckedSeq.Get(), AckData.OutSeq.Get());
 
-		// Pessimistic view, should never occur
-		return SequenceNumberT(AckedSeq.Get() - MaxSequenceHistoryLength);
-	}
+	// Pessimistic view, should never occur but we do want to know about it if it would
+	ensureMsgf(false, TEXT("FNetPacketNotify::UpdateInAckSeqAck - Failed to find matching AckRecord for %u"), AckedSeq.Get());
+	
+	return SequenceNumberT(AckedSeq.Get() - MaxSequenceHistoryLength);
 }
 
 void FNetPacketNotify::Init(SequenceNumberT InitialInSeq, SequenceNumberT InitialOutSeq)
diff --git a/Engine/Source/Runtime/Engine/Private/Net/Tests/NetPacketNotifyTests.cpp b/Engine/Source/Runtime/Engine/Private/Net/Tests/NetPacketNotifyTests.cpp
--- a/Engine/Source/Runtime/Engine/Private/Net/Tests/NetPacketNotifyTests.cpp
+++ b/Engine/Source/Runtime/Engine/Private/Net/Tests/NetPacketNotifyTests.cpp
@@ -5,7 +5,7 @@
 
 #if WITH_DEV_AUTOMATION_TESTS
 
-IMPLEMENT_SIMPLE_AUTOMATION_TEST(FNetPacketNotifyTest, "Network.PacketNotifyTest", EAutomationTestFlags::ApplicationContextMask | EAutomationTestFlags::EngineFilter)
+IMPLEMENT_SIMPLE_AUTOMATION_TEST(FNetPacketNotifyTest, "Net.PacketNotifyTest", EAutomationTestFlags::ApplicationContextMask | EAutomationTestFlags::EngineFilter)
 
 struct FNetPacketNotifyTestUtil
 {
@@ -257,6 +257,36 @@ bool FNetPacketNotifyTest::RunTest(const FString& Parameters)
 		TestEqual(TEXT("Notifications - Create sequence"), FPlatformMemory::Memcmp(ExpectedAckdPacketIds, RcvdAcks, sizeof(ExpectedAckdPacketIds)), 0u);
 	}
 
+	// test received invalid ack
+	{
+		static const FNetPacketNotify::SequenceNumberT ExpectedAckdPacketIds[] = {3, 7, 12, 14, 17, 18};
+		static const SIZE_T ExpectedCount = sizeof(ExpectedAckdPacketIds)/sizeof((ExpectedAckdPacketIds)[0]);		
+
+		FNetPacketNotify::SequenceNumberT RcvdAcks[ExpectedCount] = { 0 };
+		SIZE_T RcvdCount = 0;
+
+		// Create src data
+		FNetPacketNotify Acks = Util.DefaultNotify;
+
+		// Fill in some data
+		FNetPacketNotify::FNotificationHeader Data;
+		Data.Seq = FNetPacketNotify::SequenceNumberT(0);
+		Data.AckedSeq = FNetPacketNotify::SequenceNumberT(19);
+		Data.History = FNetPacketNotify::SequenceHistoryT(0x8853u);
+		Data.HistoryWordCount = 1;
+
+		// Need to fake ack record as well.
+		for (SIZE_T It=0; It <= 18; ++It)
+		{
+			FNetPacketNotifyTestUtil::PretendSendSeq(Acks, 0);
+		}
+	
+		SIZE_T DeltaSeq = FNetPacketNotifyTestUtil::PretendReceivedPacket(Acks, Data, RcvdAcks);
+
+		TestEqual(TEXT("Notifications - Create sequence delta"), DeltaSeq, 0);
+	}
+
+
 	// test various sequence
 	{
 		static const FNetPacketNotify::SequenceNumberT TestSeqs[] = {3, 7, 12, 14, 17, 18};	
diff --git a/Engine/Source/Runtime/Engine/Public/Net/NetPacketNotify.h b/Engine/Source/Runtime/Engine/Public/Net/NetPacketNotify.h
--- a/Engine/Source/Runtime/Engine/Public/Net/NetPacketNotify.h
+++ b/Engine/Source/Runtime/Engine/Public/Net/NetPacketNotify.h
@@ -138,7 +138,7 @@ private:
 template<class Functor>
 FNetPacketNotify::SequenceNumberT::DifferenceT FNetPacketNotify::Update(const FNotificationHeader& NotificationData, Functor&& InFunc)
 {
-	if (NotificationData.Seq > InSeq && NotificationData.AckedSeq >= OutAckSeq)
+	if (NotificationData.Seq > InSeq && NotificationData.AckedSeq >= OutAckSeq && OutSeq > NotificationData.AckedSeq)
 	{
 		UE_LOG_PACKET_NOTIFY(TEXT("FNetPacketNotify::Update - Seq %u, InSeq %u"), NotificationData.Seq.Get(), InSeq.Get());
 
