Re-organize MeshComponent parameter order to restore packing lost in CL# 4414987
#rnx

diff --git a/Engine/Source/Runtime/Engine/Classes/Components/MeshComponent.h b/Engine/Source/Runtime/Engine/Classes/Components/MeshComponent.h
--- a/Engine/Source/Runtime/Engine/Classes/Components/MeshComponent.h
+++ b/Engine/Source/Runtime/Engine/Classes/Components/MeshComponent.h
@@ -102,9 +102,6 @@ class ENGINE_API UMeshComponent : public UPrimitiveComponent
 public:
 	/** Material parameter setting and caching */
 
-	UPROPERTY(EditAnywhere, AdvancedDisplay, BlueprintReadOnly, Category = MaterialParameters)
-	bool bEnableMaterialParameterCaching;
-
 	/** Set all occurrences of Scalar Material Parameters with ParameterName in the set of materials of the SkeletalMesh to ParameterValue */
 	UFUNCTION(BlueprintCallable, Category = "Rendering|Material")
 	void SetScalarParameterValueOnMaterials(const FName ParameterName, const float ParameterValue);
@@ -146,6 +143,10 @@ protected:
 
 	TSortedMap<FName, FMaterialParameterCache> MaterialParameterCache;
 
+	UPROPERTY(EditAnywhere, AdvancedDisplay, BlueprintReadOnly, Category = MaterialParameters)
+	bool bEnableMaterialParameterCaching;
+
 	/** Flag whether or not the cached material parameter indices map is dirty (defaults to true, and is set from SetMaterial/Set(Skeletal)Mesh */
-	uint32 bCachedMaterialParameterIndicesAreDirty : 1;
+	bool bCachedMaterialParameterIndicesAreDirty;
+
 };
