Fix an editor crash on load of a Blueprint asset after switching a child component node's type to be a non-scene derivative in C++ code.

diff --git a/Engine/Source/Runtime/Engine/Private/SimpleConstructionScript.cpp b/Engine/Source/Runtime/Engine/Private/SimpleConstructionScript.cpp
--- a/Engine/Source/Runtime/Engine/Private/SimpleConstructionScript.cpp
+++ b/Engine/Source/Runtime/Engine/Private/SimpleConstructionScript.cpp
@@ -145,7 +145,9 @@ void USimpleConstructionScript::PostLoad()
 	// This pass is not needed during reinstancing.
 	if (!GIsDuplicatingClassForReinstancing)
 	{
-		for (USCS_Node* Node : GetAllNodes())
+		// Use a copy of the array for iterating, as we might have to reposition nodes in the hierarchy below (which can temporarily modify the array).
+		TArray<USCS_Node*> AllNodes = GetAllNodes();
+		for (USCS_Node* Node : AllNodes)
 		{
 			// Fix up any uninitialized category names
 			if (Node->CategoryName.IsEmpty())
