CIS fix.

#rnx

diff --git a/Engine/Source/Runtime/Engine/Private/SimpleConstructionScript.cpp b/Engine/Source/Runtime/Engine/Private/SimpleConstructionScript.cpp
--- a/Engine/Source/Runtime/Engine/Private/SimpleConstructionScript.cpp
+++ b/Engine/Source/Runtime/Engine/Private/SimpleConstructionScript.cpp
@@ -146,8 +146,8 @@ void USimpleConstructionScript::PostLoad()
 	if (!GIsDuplicatingClassForReinstancing)
 	{
 		// Use a copy of the array for iterating, as we might have to reposition nodes in the hierarchy below (which can temporarily modify the array).
-		TArray<USCS_Node*> AllNodes = GetAllNodes();
-		for (USCS_Node* Node : AllNodes)
+		TArray<USCS_Node*> LocalAllNodes = GetAllNodes();
+		for (USCS_Node* Node : LocalAllNodes)
 		{
 			// Fix up any uninitialized category names
 			if (Node->CategoryName.IsEmpty())
