Fix initialization order warnings

diff --git a/Engine/Source/Runtime/Engine/Public/Net/RepLayout.h b/Engine/Source/Runtime/Engine/Public/Net/RepLayout.h
--- a/Engine/Source/Runtime/Engine/Public/Net/RepLayout.h
+++ b/Engine/Source/Runtime/Engine/Public/Net/RepLayout.h
@@ -1070,11 +1070,11 @@ private:
 	friend class UPackageMapClient;
 
 	FRepLayout():
+		LayoutState(ERepLayoutState::Uninitialized),
 		FirstNonCustomParent(0),
-		RoleIndex(-1),
+		RoleIndex(INDEX_NONE),
 		RemoteRoleIndex(-1),
-		Owner(NULL),
-		LayoutState(ERepLayoutState::Uninitialized)
+		Owner(nullptr)
 	{}
 
 public:
@@ -1832,7 +1832,6 @@ private:
 	void CopyProperties(FRepStateStaticBuffer& ShadowData, const uint8* const Src) const;
 	void DestructProperties(FRepStateStaticBuffer& RepStateStaticBuffer) const;
 
-
 	ERepLayoutState LayoutState;
 
 	int16 FirstNonCustomParent;
