PR #5507: Fix #endif comment (Contributed by [REDACTED_PERSON])

diff --git a/Engine/Source/Runtime/Engine/Public/Animation/AnimCurveTypes.h b/Engine/Source/Runtime/Engine/Public/Animation/AnimCurveTypes.h
--- a/Engine/Source/Runtime/Engine/Public/Animation/AnimCurveTypes.h
+++ b/Engine/Source/Runtime/Engine/Public/Animation/AnimCurveTypes.h
@@ -656,7 +656,7 @@ struct FRawCurveTracks
 	 **/
 	UPROPERTY()
 	TArray<FTransformCurve>	TransformCurves;
-#endif // #if WITH_EDITOR_DATA
+#endif // #if WITH_EDITORONLY_DATA
 
 	/**
 	 * Evaluate curve data at the time CurrentTime, and add to Instance. It only evaluates Float Curve for now
