Remove the old deprecated/unused variables/functions InputConsumeOption, SetLockedAxis, CurrentTag,  RelativeTranslation_DEPRECATED, and bAbsoluteTranslation_DEPRECATED. The properties already had redirectors in 4.0 as they are simple renames
Fix issue where only one of the constructors of AActor would call class count code

diff --git a/Engine/Source/Runtime/Engine/Classes/Components/ActorComponent.h b/Engine/Source/Runtime/Engine/Classes/Components/ActorComponent.h
--- a/Engine/Source/Runtime/Engine/Classes/Components/ActorComponent.h
+++ b/Engine/Source/Runtime/Engine/Classes/Components/ActorComponent.h
@@ -27,17 +27,21 @@ class ULevel;
 class SWidget;
 struct FMinimalViewInfo;
 
-/** Annotation for component selection.  This must be in engine isntead of editor for ::IsSelected to work */
+/** Annotation for component selection.  This must be in engine instead of editor for ::IsSelected to work */
 extern ENGINE_API FUObjectAnnotationSparseBool GSelectedComponentAnnotation;
 #endif
 
-/** Information about how to update transform*/
+/** Information about how to update transform when something is moved */
 enum class EUpdateTransformFlags : int32
 {
+	/** Default options */
 	None = 0x0,
-	SkipPhysicsUpdate = 0x1,		// Don't update the underlying physics
-	PropagateFromParent = 0x2,		// The update is coming as a result of the parent updating (i.e. not called directly)
-	OnlyUpdateIfUsingSocket = 0x4	// Only update child transform if attached to parent via a socket
+	/** Don't update the underlying physics */
+	SkipPhysicsUpdate = 0x1,		
+	/** The update is coming as a result of the parent updating (i.e. not called directly) */
+	PropagateFromParent = 0x2,		
+	/** Only update child transform if attached to parent via a socket */
+	OnlyUpdateIfUsingSocket = 0x4	
 };
 
 CONSTEXPR inline EUpdateTransformFlags operator|(EUpdateTransformFlags Left, EUpdateTransformFlags Right)
@@ -60,7 +64,7 @@ CONSTEXPR inline EUpdateTransformFlags operator ~(EUpdateTransformFlags Value)
 	return static_cast<EUpdateTransformFlags>(~static_cast<int32>(Value));
 }
 
-
+/** Converts legacy bool into the SkipPhysicsUpdate bitflag */
 FORCEINLINE EUpdateTransformFlags SkipPhysicsToEnum(bool bSkipPhysics){ return bSkipPhysics ? EUpdateTransformFlags::SkipPhysicsUpdate : EUpdateTransformFlags::None; }
 
 
@@ -82,8 +86,8 @@ UCLASS(DefaultToInstanced, BlueprintType, abstract, meta=(ShortTooltip="An Actor
 class ENGINE_API UActorComponent : public UObject, public IInterface_AssetUserData
 {
 	GENERATED_BODY()
-public:
 
+public:
 	/** Create component physics state global delegate.*/
 	static FActorComponentGlobalCreatePhysicsSignature GlobalCreatePhysicsDelegate;
 	/** Destroy component physics state global delegate.*/
@@ -94,10 +98,7 @@ public:
 	 */
 	UActorComponent(const FObjectInitializer& ObjectInitializer = FObjectInitializer::Get());
 
-public:
-	virtual void GetLifetimeReplicatedProps(TArray<FLifetimeProperty>& OutLifetimeProps) const override;
-
-	/** Main tick function for the Actor */
+	/** Main tick function for the Component */
 	UPROPERTY(EditDefaultsOnly, Category="ComponentTick")
 	struct FActorComponentTickFunction PrimaryComponentTick;
 
@@ -106,7 +107,6 @@ public:
 	TArray<FName> ComponentTags;
 
 protected:
-
 	/** Array of user data stored with the component */
 	UPROPERTY(EditAnywhere, AdvancedDisplay, Instanced, Category = AssetUserData)
 	TArray<UAssetUserData*> AssetUserData;
@@ -149,11 +149,11 @@ private:
 	uint8 bRoutedPostRename:1;
 
 public:
-
 	/** Does this component automatically register with its owner */
 	uint8 bAutoRegister:1;
 
 protected:
+	/** Check whether the component class allows reregistration during ReregisterAllComponents */
 	uint8 bAllowReregistration:1;
 
 public:
@@ -170,11 +170,11 @@ public:
 	uint8 bAllowAnyoneToDestroyMe:1;
 
 #if WITH_EDITORONLY_DATA
-	/** True if this component was created by a construction script, and will be destroyed by DestroyConstructedComponents */
+	/** @deprecated Replaced by CreationMethod */
 	UPROPERTY()
 	uint8 bCreatedByConstructionScript_DEPRECATED:1;
 
-	/** True if this component was created as an instance component */
+	/** @deprecated Replaced by CreationMethod */
 	UPROPERTY()
 	uint8 bInstanceComponent_DEPRECATED:1;
 #endif
@@ -187,6 +187,7 @@ public:
 	UPROPERTY(transient, ReplicatedUsing=OnRep_IsActive)
 	uint8 bIsActive:1;
 
+	/** True if this component can be modified when it was inherited from a parent actor class */
 	UPROPERTY(EditDefaultsOnly, Category="Variable")
 	uint8 bEditableWhenInherited:1;
 
@@ -212,6 +213,7 @@ public:
 
 #if WITH_EDITORONLY_DATA
 private:
+	/** True if this component is only used for visualization, usually a sprite or text */
 	UPROPERTY()
 	uint8 bIsVisualizationComponent : 1;
 #endif
@@ -242,12 +244,9 @@ private:
 
 	/** Tracks whether the component has been added to one of the world's end of frame update lists */
 	uint8 MarkedForEndOfFrameUpdateState:2;
-	friend struct FMarkComponentEndOfFrameUpdateState;
-
-	friend struct FActorComponentInstanceData;
-	friend class FActorComponentDetails;
 
 public:
+	/** Describes how a component instance will be created */
 	UPROPERTY()
 	EComponentCreationMethod CreationMethod;
 
@@ -256,17 +255,28 @@ private:
 	TArray<FSimpleMemberReference> UCSModifiedProperties;
 
 public:
-
+	/** Tracks whether the component has been added to one of the world's end of frame update lists */
 	uint32 GetMarkedForEndOfFrameUpdateState() const { return MarkedForEndOfFrameUpdateState; }
 
+	/** Initializes the list of properties that are modified by the UserConstructionScript */
 	void DetermineUCSModifiedProperties();
+
+	/** Returns the list of properties that are modified by the UserConstructionScript */
 	void GetUCSModifiedProperties(TSet<const UProperty*>& ModifiedProperties) const;
+
+	/** Removes specified properties from the list of UCS-modified properties */
 	void RemoveUCSModifiedProperties(const TArray<UProperty*>& Properties);
 
+	/** True if this component can be modified when it was inherited from a parent actor class */
 	bool IsEditableWhenInherited() const;
 
+	/** Indicates that OnCreatedComponent has been called, but OnDestroyedComponent has not yet */
 	bool HasBeenCreated() const { return bHasBeenCreated; }
+
+	/** Indicates that InitializeComponent has been called, but UninitializeComponent has not yet */
 	bool HasBeenInitialized() const { return bHasBeenInitialized; }
+
+	/** Indicates that BeginPlay has been called, but EndPlay has not yet */
 	bool HasBegunPlay() const { return bHasBegunPlay; }
 
 	/**
@@ -278,8 +288,10 @@ public:
 		return bIsBeingDestroyed;
 	}
 
+	/** Returns true if instances of this component are created by either the user or simple construction script */
 	bool IsCreatedByConstructionScript() const;
 
+	/** Handles replication of active state, handles ticking by default but should be overridden as needed */
 	UFUNCTION()
 	virtual void OnRep_IsActive();
 
@@ -291,33 +303,27 @@ public:
 	UFUNCTION(BlueprintCallable, Category="Components", meta=(Keywords = "Actor Owning Parent"))
 	AActor* GetOwner() const;
 
+	/** Getter for the cached world pointer, will return null if the component is not actually spawned in a level */
 	virtual UWorld* GetWorld() const override final { return (WorldPrivate ? WorldPrivate : GetWorld_Uncached()); }
 
 	/** See if this component contains the supplied tag */
 	UFUNCTION(BlueprintCallable, Category="Components")
 	bool ComponentHasTag(FName Tag) const;
 
-	/**
-	* Called during saving to determine the load flags to save with the object.
-	*
-	* @return	true if this object should always be loaded for editor game
-	*/
-	virtual bool NeedsLoadForEditorGame() const override
-	{
-		return !IsEditorOnly() && Super::NeedsLoadForEditorGame();
-	}
 
-	//~ Begin Trigger/Activation Interface
+	// Activation System
 
+	/** Called when the component has been activated, with parameter indicating if it was from a reset */
 	UPROPERTY(BlueprintAssignable, Category = "Components|Activation")
 	FActorComponentActivatedSignature OnComponentActivated;
 
+	/** Called when the component has been deactivated */
 	UPROPERTY(BlueprintAssignable, Category = "Components|Activation")
 	FActorComponentDeactivateSignature OnComponentDeactivated;
 
 	/**
-	 * Activates the SceneComponent
-	 * @param bReset - Whether the activation should be forced even if ShouldActivate returns false.
+	 * Activates the SceneComponent, should be overridden by native child classes.
+	 * @param bReset - Whether the activation should happen even if ShouldActivate returns false.
 	 */
 	UFUNCTION(BlueprintCallable, Category="Components|Activation", meta=(UnsafeDuringActorConstruction="true"))
 	virtual void Activate(bool bReset=false);
@@ -331,7 +337,7 @@ public:
 	/**
 	 * Sets whether the component is active or not
 	 * @param bNewActive - The new active state of the component
-	 * @param bReset - Whether the activation should be forced even if ShouldActivate returns false.
+	 * @param bReset - Whether the activation should happen even if ShouldActivate returns false.
 	 */
 	UFUNCTION(BlueprintCallable, Category="Components|Activation", meta=(UnsafeDuringActorConstruction="true"))
 	virtual void SetActive(bool bNewActive, bool bReset=false);
@@ -366,17 +372,12 @@ public:
 	/** Shut down any physics engine structure for this component */
 	void DestroyPhysicsState();
 
+
 	// Networking
 
 	/** This signifies the component can be ID'd by name over the network. This only needs to be called by engine code when constructing blueprint components. */
 	void SetNetAddressable();
 
-	/** IsNameStableForNetworking means an object can be referred to its path name (relative to outer) over the network */
-	virtual bool IsNameStableForNetworking() const override;
-
-	/** IsSupportedForNetworking means an object can be referenced over the network */
-	virtual bool IsSupportedForNetworking() const override;
-
 	/** Enable or disable replication. This is the equivalent of RemoteRole for actors (only a bool is required for components) */
 	UFUNCTION(BlueprintCallable, Category="Components")
 	void SetIsReplicated(bool ShouldReplicate);
@@ -393,12 +394,14 @@ public:
 	/** Called on the component right before replication occurs */
 	virtual void PreReplication(IRepChangedPropertyTracker & ChangedPropertyTracker);
 
+	/** Returns true if this type of component can ever replicate, override to disable the default behavior */
 	virtual bool GetComponentClassCanReplicate() const;
 
 #if WITH_EDITORONLY_DATA
 	/** Returns whether this component is an editor-only object or not */
 	virtual bool IsEditorOnly() const override { return bIsEditorOnly; }
 
+	/** Called during component creation to mark this component as editor only */
 	virtual void MarkAsEditorOnlySubobject() override
 	{
 		bIsEditorOnly = true;
@@ -407,7 +410,10 @@ public:
 		bIsVisualizationComponent = true; 
 	}
 
+	/** Returns true if this component is only used for visualization, usually a sprite or text */
 	bool IsVisualizationComponent() const { return bIsVisualizationComponent; }
+
+	/** Sets if this component is only used for visualization */
 	void SetIsVisualizationComponent(const bool bInIsVisualizationComponent)
 	{
 		bIsVisualizationComponent = bInIsVisualizationComponent;
@@ -418,9 +424,8 @@ public:
 	}
 #endif
 
-	/** Returns net role of the owning actor */
-	/** Returns true if we are replicating and not authorative */
-	bool	IsNetSimulating() const;
+	/** Returns true if we are replicating and this client is not authoritative */
+	bool IsNetSimulating() const;
 
 	/** Get the network role of the Owner, or ROLE_None if there is no owner. */
 	ENetRole GetOwnerRole() const;
@@ -445,7 +450,7 @@ public:
 	void SetIsNetStartupComponent(const bool bInIsNetStartupComponent) { bIsNetStartupComponent = bInIsNetStartupComponent; }
 
 private:
-
+	/** Cached pointer to owning actor */
 	mutable AActor* OwnerPrivate;
 
 	/** 
@@ -454,15 +459,14 @@ private:
 	 */
 	UWorld* WorldPrivate;
 
-	// If WorldPrivate isn't set this will determine the world from outers
+	/** If WorldPrivate isn't set this will determine the world from outers */
 	UWorld* GetWorld_Uncached() const;
 
 	/** Private version without inlining that does *not* check Dedicated server build flags (which should already have been done). */
 	ENetMode InternalGetNetMode() const;
 
 protected:
-
-	/** "Trigger" related function. Return true if it should activate **/
+	/** Return true if this component is in a state where it can be activated normally. */
 	virtual bool ShouldActivate() const;
 
 private:
@@ -472,13 +476,10 @@ private:
 	/** Calls OnRegister, CreateRenderState_Concurrent and OnCreatePhysicsState. */
 	void ExecuteRegisterEvents();
 
-	/* Utility function for each of the PostEditChange variations to call for the same behavior */
+	/** Utility function for each of the PostEditChange variations to call for the same behavior */
 	void ConsolidatedPostEditChange(const FPropertyChangedEvent& PropertyChangedEvent);
 protected:
 
-	friend class FComponentReregisterContextBase;
-	friend class FComponentRecreateRenderStateContext;
-
 	/**
 	 * Called when a component is registered, after Scene is set, but before CreateRenderState_Concurrent or OnCreatePhysicsState are called.
 	 */
@@ -495,25 +496,25 @@ protected:
 		return false;
 	}
 
-	/** Used to create any rendering thread information for this component
-	*
-	* **Caution**, this is called concurrently on multiple threads (but never the same component concurrently)
-	*/
+	/** 
+	 * Used to create any rendering thread information for this component
+	 * @warning This is called concurrently on multiple threads (but never the same component concurrently)
+	 */
 	virtual void CreateRenderState_Concurrent();
 
-	/** Called to send a transform update for this component to the rendering thread
-	*
-	* **Caution**, this is called concurrently on multiple threads (but never the same component concurrently)
-	*/
+	/** 
+	 * Called to send a transform update for this component to the rendering thread
+	 * @warning This is called concurrently on multiple threads (but never the same component concurrently)
+	 */
 	virtual void SendRenderTransform_Concurrent();
 
 	/** Called to send dynamic data for this component to the rendering thread */
 	virtual void SendRenderDynamicData_Concurrent();
 
-	/** Used to shut down any rendering thread structure for this component
-	*
-	* **Caution**, this is called concurrently on multiple threads (but never the same component concurrently)
-	*/
+	/** 
+	 * Used to shut down any rendering thread structure for this component
+	 * @warning This is called concurrently on multiple threads (but never the same component concurrently)
+	 */
 	virtual void DestroyRenderState_Concurrent();
 
 	/** Used to create any physics engine information for this component */
@@ -537,14 +538,14 @@ protected:
 
 public:
 	/**
-	 * Initializes the component.  Occurs at level startup. This is before BeginPlay (Actor or Component).  
+	 * Initializes the component.  Occurs at level startup or actor spawn. This is before BeginPlay (Actor or Component).  
 	 * All Components in the level will be Initialized on load before any Actor/Component gets BeginPlay
 	 * Requires component to be registered, and bWantsInitializeComponent to be true.
 	 */
 	virtual void InitializeComponent();
 
 	/**
-	 * BeginsPlay for the component.  Occurs at level startup. This is before BeginPlay (Actor or Component).  
+	 * BeginsPlay for the component.  Occurs at level startup or actor spawn. This is before BeginPlay (Actor or Component).  
 	 * All Components (that want initialization) in the level will be Initialized on load before any 
 	 * Actor/Component gets BeginPlay.
 	 * Requires component to be registered and initialized.
@@ -552,7 +553,7 @@ public:
 	virtual void BeginPlay();
 
 	/** 
-	 * Blueprint implementable event for when the component is beginning play, called before its Owner's BeginPlay on Actor BeginPlay 
+	 * Blueprint implementable event for when the component is beginning play, called before its owning actor's BeginPlay
 	 * or when the component is dynamically created if the Actor has already BegunPlay. 
 	 */
 	UFUNCTION(BlueprintImplementableEvent, meta=(DisplayName = "Begin Play"))
@@ -590,6 +591,7 @@ public:
 	 * @param ThisTickFunction - Internal tick function struct that caused this to run
 	 */
 	virtual void TickComponent(float DeltaTime, enum ELevelTick TickType, FActorComponentTickFunction *ThisTickFunction);
+	
 	/** 
 	 * Set up a tick function for a component in the standard way. 
 	 * Tick after the actor. Don't tick if the actor is static, or if the actor is a template or if this is a "NeverTick" component.
@@ -606,12 +608,13 @@ public:
 	 */
 	UFUNCTION(BlueprintCallable, Category="Utilities")
 	virtual void SetComponentTickEnabled(bool bEnabled);
+
 	/** 
 	 * Spawns a task on GameThread that will call SetComponentTickEnabled
-	 * 
 	 * @param	bEnabled - Whether it should be enabled or not
 	 */
 	virtual void SetComponentTickEnabledAsync(bool bEnabled);
+	
 	/** 
 	 * Returns whether this component has tick enabled or not
 	 */
@@ -625,33 +628,26 @@ public:
 	UFUNCTION(BlueprintCallable, Category="Utilities")
 	void SetComponentTickInterval(float TickInterval);
 
-	/** 
-	* Returns whether this component has tick enabled or not
-	*/
+	/** Returns the tick interval for this component's primary tick function, which is the frequency in seconds at which it will be executed */
 	UFUNCTION(BlueprintCallable, Category="Utilities")
 	float GetComponentTickInterval() const;
 
 	/**
+	 * Registers a component with a specific world, which creates any visual/physical state
 	 * @param InWorld - The world to register the component with.
 	 */
 	void RegisterComponentWithWorld(UWorld* InWorld);
 
-private:
-	
-	/** Friend for access to ConditionalTickComponent(). */
-	friend struct FActorComponentTickFunction;
-
-public:
-
 	/** Overridable check for a component to indicate to its Owner that it should prevent the Actor from auto destroying when finished */
 	virtual bool IsReadyForOwnerToAutoDestroy() const { return true; }
 
-	/**
-	 * Returns whether the component's owner is selected.
-	 */
+	/** Returns whether the component's owner is selected in the editor */
 	bool IsOwnerSelected() const;
 
+	/** Is this component's transform in need of sending to the renderer? */
 	inline bool IsRenderTransformDirty() const { return bRenderTransformDirty; }
+
+	/** Is this component in need of its whole state being sent to the renderer? */
 	inline bool IsRenderStateDirty() const { return bRenderStateDirty; }
 
 	/** Invalidate lighting cache with default options. */
@@ -693,7 +689,7 @@ public:
 	 * Uses the bRenderStateDirty/bRenderTransformDirty to perform any necessary work on this component.
 	 * Do not call this directly, call MarkRenderStateDirty, MarkRenderDynamicDataDirty, 
 	 *
-	 * **Caution**, this is called concurrently on multiple threads (but never the same component concurrently)
+	 * @warning This is called concurrently on multiple threads (but never the same component concurrently)
 	 */
 	void DoDeferredRenderUpdates_Concurrent();
 
@@ -724,28 +720,28 @@ public:
 	/** return true if this component requires end of frame recreates to happen from the game thread. */
 	virtual bool RequiresGameThreadEndOfFrameRecreate() const;
 
-	/** Recreate the render state right away. Generally you always want to call MarkRenderStateDirty instead. 
-	*
-	* **Caution**, this is called concurrently on multiple threads (but never the same component concurrently)
-	*/
+	/** 
+	 * Recreate the render state right away. Generally you always want to call MarkRenderStateDirty instead. 
+	 * @warning This is called concurrently on multiple threads (but never the same component concurrently)
+	 */
 	void RecreateRenderState_Concurrent();
 
 	/** Recreate the physics state right way. */
 	void RecreatePhysicsState();
 
-	/** @return true if the render 'state' (e.g. scene proxy) is created for this component */
+	/** Returns true if the render 'state' (e.g. scene proxy) is created for this component */
 	bool IsRenderStateCreated() const
 	{
 		return bRenderStateCreated;
 	}
 
-	/** @return true if the physics 'state' (e.g. physx bodies) are created for this component */
+	/** Returns true if the physics 'state' (e.g. physx bodies) are created for this component */
 	bool IsPhysicsStateCreated() const
 	{
 		return bPhysicsStateCreated;
 	}
 
-	/** @name Accessors. */
+	/** Returns the rendering scene associated with this component */
 	class FSceneInterface* GetScene() const;
 
 	/** Return the ULevel that this Component is part of. */
@@ -772,12 +768,6 @@ public:
 		return nullptr;
 	}
 
-	// Always called immediately before properties are received from the remote.
-	virtual void PreNetReceive() override { }
-	
-	// Always called immediately after properties are received from the remote.
-	virtual void PostNetReceive() override { }
-
 	/** Called before we throw away components during RerunConstructionScripts, to cache any data we wish to persist across that operation */
 	virtual TStructOnScope<FActorComponentInstanceData> GetComponentInstanceData() const;
 
@@ -785,6 +775,10 @@ public:
 	virtual void BeginDestroy() override;
 	virtual bool NeedsLoadForClient() const override;
 	virtual bool NeedsLoadForServer() const override;
+	virtual bool NeedsLoadForEditorGame() const override;
+	virtual bool IsNameStableForNetworking() const override;
+	virtual bool IsSupportedForNetworking() const override;
+	virtual void GetLifetimeReplicatedProps(TArray<FLifetimeProperty>& OutLifetimeProps) const override;
 	virtual int32 GetFunctionCallspace( UFunction* Function, void* Parameters, FFrame* Stack ) override;
 	virtual bool CallRemoteFunction( UFunction* Function, void* Parameters, FOutParmRec* OutParms, FFrame* Stack ) override;
 	virtual void PostInitProperties() override;
@@ -815,10 +809,10 @@ public:
 	/** See if this component is currently registered */
 	FORCEINLINE bool IsRegistered() const { return bRegistered; }
 
-	/** Checked whether the component class allows reregistration */
+	/** Check whether the component class allows reregistration during ReregisterAllComponents */
 	FORCEINLINE bool AllowReregistration() const { return bAllowReregistration; }
 
-	/** Register this component, creating any rendering/physics state. Will also adds to outer Actor's Components array, if not already present. */
+	/** Register this component, creating any rendering/physics state. Will also add itself to the outer Actor's Components array, if not already present. */
 	void RegisterComponent();
 
 	/** Unregister this component, destroying any rendering/physics state. */
@@ -827,7 +821,7 @@ public:
 	/** Unregister the component, remove it from its outer Actor's Components array and mark for pending kill. */
 	virtual void DestroyComponent(bool bPromoteChildren = false);
 
-	/** Called when a component is created (not loaded) */
+	/** Called when a component is created (not loaded). This can happen in the editor or during gameplay */
 	virtual void OnComponentCreated();
 
 	/** 
@@ -866,7 +860,7 @@ public:
 	UFUNCTION(BlueprintCallable, Category="Utilities", meta=(Keywords = "dependency"))
 	virtual void RemoveTickPrerequisiteComponent(UActorComponent* PrerequisiteComponent);
 
-	/** Event called every frame */
+	/** Event called every frame if tick is enabled */
 	UFUNCTION(BlueprintImplementableEvent, meta=(DisplayName = "Tick"))
 	void ReceiveTick(float DeltaSeconds);
 	
@@ -885,9 +879,12 @@ public:
 	/** set value of bCanEverAffectNavigation flag and update navigation octree if needed */
 	void SetCanEverAffectNavigation(bool bRelevant);
 
-	/** override to supply actual logic */
+	/** Override to specify that a component is relevant to the navigation system */
 	virtual bool IsNavigationRelevant() const { return false; }
 
+	/** Prefix used to identify template component instances */
+	static const FString ComponentTemplateNameSuffix;
+
 protected:
 	/** Makes sure navigation system has up to date information regarding component's navigation relevancy 
 	 *	and if it can affect navigation at all 
@@ -896,20 +893,19 @@ protected:
 	void HandleCanEverAffectNavigationChange(bool bForceUpdate = false);
 
 private:
-
 #if !(UE_BUILD_SHIPPING || UE_BUILD_TEST)
-
-	// this is the old name of the tick function. We just want to avoid mistakes with an attempt to override this
+	/** This is the old name of the tick function. We just want to avoid mistakes with an attempt to override this */
 	virtual void Tick( float DeltaTime ) final { check(0); }
-
 #endif
 
 	void ClearNeedEndOfFrameUpdate_Internal();
 
-public:
-
-	/** Prefix used to identify template component instances */
-	static const FString ComponentTemplateNameSuffix;
+	friend struct FMarkComponentEndOfFrameUpdateState;
+	friend struct FActorComponentInstanceData;
+	friend class FActorComponentDetails;
+	friend class FComponentReregisterContextBase;
+	friend class FComponentRecreateRenderStateContext;
+	friend struct FActorComponentTickFunction;
 };
 
 //////////////////////////////////////////////////////////////////////////
diff --git a/Engine/Source/Runtime/Engine/Classes/Components/PrimitiveComponent.h b/Engine/Source/Runtime/Engine/Classes/Components/PrimitiveComponent.h
--- a/Engine/Source/Runtime/Engine/Classes/Components/PrimitiveComponent.h
+++ b/Engine/Source/Runtime/Engine/Classes/Components/PrimitiveComponent.h
@@ -47,30 +47,31 @@ enum ECanBeCharacterBase
 	ECB_MAX,
 };
 
+/** Determines if a primitive component contains custom collision for navigation/AI */
 UENUM()
 namespace EHasCustomNavigableGeometry
 {
 	enum Type
 	{
-		// Primitive doesn't have custom navigation geometry, if collision is enabled then its convex/trimesh collision will be used for generating the navmesh
+		/** Primitive doesn't have custom navigation geometry, if collision is enabled then its convex/trimesh collision will be used for generating the navmesh */
 		No,
 
-		// If primitive would normally affect navmesh, DoCustomNavigableGeometryExport() should be called to export this primitive's navigable geometry
+		/** If primitive would normally affect navmesh, DoCustomNavigableGeometryExport() should be called to export this primitive's navigable geometry */
 		Yes,
 
-		// DoCustomNavigableGeometryExport() should be called even if the mesh is non-collidable and wouldn't normally affect the navmesh
+		/** DoCustomNavigableGeometryExport() should be called even if the mesh is non-collidable and wouldn't normally affect the navmesh */
 		EvenIfNotCollidable,
 
-		// Don't export navigable geometry even if primitive is relevant for navigation (can still add modifiers)
+		/** Don't export navigable geometry even if primitive is relevant for navigation (can still add modifiers) */
 		DontExport,
 	};
 }
 
-/** Information about the sprite category */
+/** Information about the sprite category, used for visualization in the editor */
 USTRUCT()
 struct FSpriteCategoryInfo
 {
-	GENERATED_USTRUCT_BODY()
+	GENERATED_BODY()
 
 	/** Sprite category that the component belongs to */
 	UPROPERTY()
@@ -171,8 +172,8 @@ UCLASS(abstract, HideCategories=(Mobility), ShowCategories=(PhysicsVolume))
 class ENGINE_API UPrimitiveComponent : public USceneComponent, public INavRelevantInterface
 {
 	GENERATED_BODY()
-public:
 
+public:
 	/**
 	 * Default UObject constructor.
 	 */
@@ -228,8 +229,6 @@ public:
 	TArray<int32> ExcludeForSpecificHLODLevels;
 #endif 
 
-public:
-
 	/**
 	 * When enabled this object will not be culled by distance. This is ignored if a child of a HLOD.
 	 */
@@ -263,10 +262,10 @@ public:
 	 * @see [Overlap Events](https://docs.unrealengine.com/latest/INT/Engine/Physics/Collision/index.html#overlapandgenerateoverlapevents)
 	 * @see UpdateOverlaps(), BeginComponentOverlap(), EndComponentOverlap()
 	 */
-
 	UFUNCTION(BlueprintGetter)
 	bool GetGenerateOverlapEvents() const;
 
+	/** Modifies value returned by GetGenerateOverlapEvents() */
 	UFUNCTION(BlueprintSetter)
 	void SetGenerateOverlapEvents(bool bInGenerateOverlapEvents);
 
@@ -306,7 +305,7 @@ public:
 	UPROPERTY(EditAnywhere, AdvancedDisplay, BlueprintReadOnly, Category=LOD)
 	uint8 bAllowCullDistanceVolume:1;
 
-	/** true if the primitive has motion blur velocity meshes */
+	/** True if the primitive has motion blur velocity meshes */
 	UPROPERTY()
 	uint8 bHasMotionBlurVelocityMeshes:1;
 	
@@ -427,6 +426,7 @@ public:
 	UPROPERTY(EditAnywhere, AdvancedDisplay, BlueprintReadOnly, Category=Lighting, meta=(EditCondition="CastShadow", DisplayName = "Shadow Two Sided"))
 	uint8 bCastShadowAsTwoSided:1;
 
+	/** @deprecated Replaced by LightmapType */
 	UPROPERTY()
 	uint8 bLightAsIfStatic_DEPRECATED:1;
 
@@ -583,7 +583,6 @@ public:
 	virtual void BeginPlay() override;
 
 protected:
-
 	/** Returns true if all descendant components that we can possibly overlap with use relative location and rotation. */
 	virtual bool AreAllCollideableDescendantsRelative(bool bAllowCachedValue = true) const;
 
@@ -615,6 +614,7 @@ public:
 	UPROPERTY(transient)
 	float LastRenderTime;
 
+	/** Same as LastRenderTimeOnScreen but only updated if the component is on screen. Used by the texture streamer. */
 	UPROPERTY(transient)
 	float LastRenderTimeOnScreen;
 
@@ -796,11 +796,6 @@ public:
 	 */
 	virtual bool UpdateOverlapsImpl(TArray<FOverlapInfo> const* NewPendingOverlaps=nullptr, bool bDoNotifies=true, const TArray<FOverlapInfo>* OverlapsAtEndLocation=nullptr) override;
 
-#if WITH_EDITOR
-	/** Update the Bounds of the component.*/
-	virtual void UpdateBounds() override;
-#endif
-
 	/** Update current physics volume for this component, if bShouldUpdatePhysicsVolume is true. Overridden to use the overlaps to find the physics volume. */
 	virtual void UpdatePhysicsVolume( bool bTriggerNotifiers ) override;
 
@@ -820,12 +815,10 @@ public:
 	bool ComponentOverlapMulti(TArray<struct FOverlapResult>& OutOverlaps, const class UWorld* InWorld, const FVector& Pos, const FRotator& Rot, ECollisionChannel TestChannel, const struct FComponentQueryParams& Params = FComponentQueryParams::DefaultComponentQueryParams, const struct FCollisionObjectQueryParams& ObjectQueryParams = FCollisionObjectQueryParams::DefaultObjectQueryParam) const;
 
 protected:
-
-	// Override this method for custom behavior.
+	/** Override this method for custom behavior for ComponentOverlapMulti() */
 	virtual bool ComponentOverlapMultiImpl(TArray<struct FOverlapResult>& OutOverlaps, const class UWorld* InWorld, const FVector& Pos, const FQuat& Rot, ECollisionChannel TestChannel, const struct FComponentQueryParams& Params, const struct FCollisionObjectQueryParams& ObjectQueryParams = FCollisionObjectQueryParams::DefaultObjectQueryParam) const;
 
 public:
-
 	// Internal physics engine data.
 	
 	/** Physics scene information for this component, holds a single rigid body with multiple shapes. */
@@ -991,14 +984,6 @@ public:
 	 */
 	virtual bool CanEditSimulatePhysics();
 
-	/**
-	* Sets the constraint mode of the component.
-	* @param ConstraintMode	The type of constraint to use.
-	*/
-	UE_DEPRECATED(4.8, "This function is deprecated. Please use SetConstraintMode instead.")
-	UFUNCTION(BlueprintCallable, meta = (DeprecatedFunction, DeprecationMessage = "Use SetConstraintMode instead", Keywords = "set locked axis constraint physics"), Category = Physics)
-	virtual void SetLockedAxis(EDOFMode::Type LockedAxis);
-
 	/**
 	 * Sets the constraint mode of the component.
 	 * @param ConstraintMode	The type of constraint to use.
@@ -1325,9 +1310,7 @@ public:
 	UFUNCTION(BlueprintCallable, Category="Physics", meta=(UnsafeDuringActorConstruction="true"))
 	void PutRigidBodyToSleep(FName BoneName = NAME_None);
 
-	/** Changes the value of bNotifyRigidBodyCollision
-	 * @param bNewNotifyRigidBodyCollision - The value to assign to bNotifyRigidBodyCollision
-	 */
+	/** Changes the value of bNotifyRigidBodyCollision */
 	UFUNCTION(BlueprintCallable, Category="Physics")
 	virtual void SetNotifyRigidBodyCollision(bool bNewNotifyRigidBodyCollision);
 
@@ -1437,9 +1420,6 @@ public:
 	UFUNCTION(BlueprintCallable, Category = "Rendering")
 	void SetRenderInMainPass(bool bValue);
 
-public:
-	static int32 CurrentTag;
-
 	/**
 	 * Count of all component overlap events (begin or end) ever generated for any components.
 	 * Changes to this number within a scope can also be a simple way to know if any events were triggered.
@@ -1453,46 +1433,38 @@ public:
 	/** A fence to track when the primitive is detached from the scene in the rendering thread. */
 	FRenderCommandFence DetachFence;
 
-	// Scene data
 private:
 	/** LOD parent primitive to draw instead of this one (multiple UPrim's will point to the same LODParent ) */
 	UPROPERTY(duplicatetransient)
 	class UPrimitiveComponent* LODParentPrimitive;
 
 public:
+	/** Set LOD Parent component, normally associated with an ALODActor */
 	void SetLODParentPrimitive(UPrimitiveComponent* InLODParentPrimitive);
+
+	/** Gets the LOD Parent, which is used to compute visibility when hierarchical LOD is enabled */
 	UPrimitiveComponent* GetLODParentPrimitive() const;
 
 #if WITH_EDITOR
-	virtual const int32 GetNumUncachedStaticLightingInteractions() const override; // recursive function
 	/** This function is used to create hierarchical LOD for the level. You can decide to opt out if you don't want. */
 	virtual const bool ShouldGenerateAutoLOD(const int32 HierarchicalLevelIndex) const;
 #endif
-
-	//~ Begin UActorComponent Interface
-	virtual void InvalidateLightingCacheDetailed(bool bInvalidateBuildEnqueuedLighting, bool bTranslationOnly) override;
-	virtual bool IsEditorOnly() const override;
-	virtual bool ShouldCreatePhysicsState() const override;
-	virtual bool HasValidPhysicsState() const override;
-	virtual TStructOnScope<FActorComponentInstanceData> GetComponentInstanceData() const override;
-	virtual void OnComponentDestroyed(bool bDestroyingHierarchy) override;
-	//~ End UActorComponent Interface
-
-	/** @return true if the owner is selected and this component is selectable */
+	/** Return true if the owner is selected and this component is selectable */
 	virtual bool ShouldRenderSelected() const;
 
 	/** Component is directly selected in the editor separate from its parent actor */
 	bool IsComponentIndividuallySelected() const;
 
-	/**  @return True if a primitive's parameters as well as its position is static during gameplay, and can thus use static lighting. */
+	/** Return True if a primitive's parameters as well as its position is static during gameplay, and can thus use static lighting. */
 	bool HasStaticLighting() const;
 
+	/** Returns true if the component is static and has the right static mesh setup to support lightmaps. */
 	virtual bool HasValidSettingsForStaticLighting(bool bOverlookInvalidComponents) const 
 	{
 		return HasStaticLighting();
 	}
 
-	/**  @return true if only unlit materials are used for rendering, false otherwise. */
+	/** Returns true if only unlit materials are used for rendering, false otherwise. */
 	virtual bool UsesOnlyUnlitMaterials() const;
 
 	/**
@@ -1523,7 +1495,6 @@ public:
 	 */
 	virtual void GetLightAndShadowMapMemoryUsage( int32& LightMapMemoryUsage, int32& ShadowMapMemoryUsage ) const;
 
-
 #if WITH_EDITOR
 	/**
 	 * Requests the information about the component that the static lighting system needs.
@@ -1535,7 +1506,8 @@ public:
 
 	/** Add the used GUIDs from UMapBuildDataRegistry::MeshBuildData. Used to preserve hidden level data in lighting scenario. */
 	virtual void AddMapBuildDataGUIDs(TSet<FGuid>& InGUIDs) const {}
-#endif
+#endif // WITH_EDITOR
+
 	/**
 	 *	Requests whether the component will use texture, vertex or no lightmaps.
 	 *
@@ -1589,7 +1561,6 @@ public:
 	 */
 	virtual void GetUsedTextures(TArray<UTexture*>& OutTextures, EMaterialQualityLevel::Type QualityLevel);
 
-
 	/** Tick function for physics ticking **/
 	UPROPERTY()
 	struct FPrimitiveComponentPostPhysicsTickFunction PostPhysicsComponentTick;
@@ -1609,8 +1580,6 @@ public:
 	/** Return the BodySetup to use for this PrimitiveComponent (single body case) */
 	virtual class UBodySetup* GetBodySetup() { return NULL; }
 
-
-
 	/** Move this component to match the physics rigid body pose. Note, a warning will be generated if you call this function on a component that is attached to something */
 	void SyncComponentToRBPhysics();
 	
@@ -1620,26 +1589,22 @@ public:
 	 */
 	virtual FMatrix GetRenderMatrix() const;
 
-	/** @return number of material elements in this primitive */
+	/** Return number of material elements in this primitive */
 	UFUNCTION(BlueprintPure, Category="Rendering|Material")
 	virtual int32 GetNumMaterials() const;
 	
-	/** Get a BodyInstance from this component. The supplied name is used in the SkeletalMeshComponent case. A name of NAME_None in the skeletal case gives the root body instance. */
-
-
 	/**
-	* returns BodyInstance of the component.
-	*
-	* @param BoneName				Used to get body associated with specific bone. NAME_None automatically gets the root most body
-	* @param bGetWelded				If the component has been welded to another component and bGetWelded is true we return the single welded BodyInstance that is used in the simulation
-	*
-	* @return		Returns the BodyInstance based on various states (does component have multiple bodies? Is the body welded to another body?)
-	*/
-
+	 * Returns BodyInstance of the component.
+	 *
+	 * @param BoneName		Used to get body associated with specific bone. NAME_None automatically gets the root most body
+	 * @param bGetWelded	If the component has been welded to another component and bGetWelded is true we return the single welded BodyInstance that is used in the simulation
+	 *
+	 * @return		Returns the BodyInstance based on various states (does component have multiple bodies? Is the body welded to another body?)
+	 */
 	virtual FBodyInstance* GetBodyInstance(FName BoneName = NAME_None, bool bGetWelded = true) const;
 
 	/** 
-	 * returns The square of the distance to closest Body Instance surface. 
+	 * Returns The square of the distance to closest Body Instance surface. 
 	 *
 	 * @param Point				World 3D vector
 	 * @param OutSquaredDistance The squared distance to closest Body Instance surface. 0 if inside of the body
@@ -1650,14 +1615,14 @@ public:
 	virtual bool GetSquaredDistanceToCollision(const FVector& Point, float& OutSquaredDistance, FVector& OutClosestPointOnCollision) const;
 
 	/** 
-	* returns Distance to closest Body Instance surface. 
-	*
-	* @param Point				World 3D vector
-	* @param OutPointOnBody	Point on the surface of collision closest to Point
-	* 
-	* @return		Success if returns > 0.f, if returns 0.f, point is inside the geometry
-	*				If returns < 0.f, this primitive does not have collsion or if geometry is not supported
-	*/	
+	 * Returns Distance to closest Body Instance surface. 
+	 *
+	 * @param Point				World 3D vector
+	 * @param OutPointOnBody	Point on the surface of collision closest to Point
+	 * 
+	 * @return		Success if returns > 0.f, if returns 0.f, point is inside the geometry
+	 *				If returns < 0.f, this primitive does not have collsion or if geometry is not supported
+	 */	
 	float GetDistanceToCollision(const FVector& Point, FVector& ClosestPointOnCollision) const 
 	{
 		float DistanceSqr = -1.f;
@@ -1739,6 +1704,17 @@ public:
 
 	/** Whether the component has been welded to another simulating component */
 	bool IsWelded() const;
+	
+	/**
+	 * Called to get the Component To World Transform from the Root BodyInstance
+	 * This needs to be virtual since SkeletalMeshComponent Root has to undo its own transform
+	 * Without this, the root LocalToAtom is overridden by physics simulation, causing kinematic velocity to 
+	 * accelerate simulation
+	 *
+	 * @param : UseBI - root body instance
+	 * @return : New GetComponentTransform() to use
+	 */
+	virtual FTransform GetComponentTransformFromBodyInstance(FBodyInstance* UseBI);	
 
 #if WITH_EDITOR
 	/**
@@ -1773,12 +1749,6 @@ protected:
 	/** Give the static mesh component recreate render state context access to Create/DestroyRenderState_Concurrent(). */
 	friend class FStaticMeshComponentRecreateRenderStateContext;
 
-	//~ Begin USceneComponent Interface
-	virtual void OnUpdateTransform(EUpdateTransformFlags UpdateTransformFlags, ETeleportType Teleport = ETeleportType::None) override;
-
-	/** Event called when AttachParent changes, to allow the scene to update its attachment state. */
-	virtual void OnAttachmentChanged() override;
-
 	/** Whether the component type supports static lighting. */
 	virtual bool SupportsStaticLighting() const 
 	{
@@ -1786,18 +1756,12 @@ protected:
 	}
 
 public:
-	virtual bool IsSimulatingPhysics(FName BoneName = NAME_None) const override;
-
 #if !(UE_BUILD_SHIPPING || UE_BUILD_TEST)
 	/** Updates the renderer with the center of mass data */
 	virtual void SendRenderDebugPhysics(FPrimitiveSceneProxy* OverrideSceneProxy = nullptr);
 #endif
 
-	// End USceneComponentInterface
-
-
 	//~ Begin UActorComponent Interface
-protected:
 	virtual void CreateRenderState_Concurrent() override;
 	virtual void SendRenderTransform_Concurrent() override;
 	virtual void OnRegister()  override;
@@ -1806,20 +1770,13 @@ protected:
 	virtual void OnCreatePhysicsState() override;
 	virtual void OnDestroyPhysicsState() override;
 	virtual void OnActorEnableCollisionChanged() override;
-
-public:
-	/**
-	 * Called to get the Component To World Transform from the Root BodyInstance
-	 * This needs to be virtual since SkeletalMeshComponent Root has to undo its own transform
-	 * Without this, the root LocalToAtom is overridden by physics simulation, causing kinematic velocity to 
-	 * accelerate simulation
-	 *
-	 * @param : UseBI - root body instance
-	 * @return : New GetComponentTransform() to use
-	 */
-	virtual FTransform GetComponentTransformFromBodyInstance(FBodyInstance* UseBI);
-public:
 	virtual void RegisterComponentTickFunctions(bool bRegister) override;
+	virtual void InvalidateLightingCacheDetailed(bool bInvalidateBuildEnqueuedLighting, bool bTranslationOnly) override;
+	virtual bool IsEditorOnly() const override;
+	virtual bool ShouldCreatePhysicsState() const override;
+	virtual bool HasValidPhysicsState() const override;
+	virtual TStructOnScope<FActorComponentInstanceData> GetComponentInstanceData() const override;
+	virtual void OnComponentDestroyed(bool bDestroyingHierarchy) override;
 #if WITH_EDITOR
 	virtual void CheckForErrors() override;
 #endif // WITH_EDITOR	
@@ -1838,40 +1795,25 @@ public:
 
 	//~ Begin UObject Interface.
 	virtual void Serialize(FArchive& Ar) override;
-#if WITH_EDITOR
-	virtual void PostEditChangeProperty(FPropertyChangedEvent& PropertyChangedEvent) override;
-	virtual void PostEditChangeChainProperty(FPropertyChangedChainEvent& PropertyChangedEvent) override;
-	virtual bool CanEditChange(const UProperty* InProperty) const override;
-	virtual void UpdateCollisionProfile();
-#endif // WITH_EDITOR
 	virtual void PostInitProperties() override;
 	virtual void PostLoad() override;
 	virtual void PostDuplicate(bool bDuplicateForPIE) override;
 	virtual void GetResourceSizeEx(FResourceSizeEx& CumulativeResourceSize) override;
-
-#if WITH_EDITOR
-	/**
-	 * Called after importing property values for this object (paste, duplicate or .t3d import)
-	 * Allow the object to perform any cleanup for properties which shouldn't be duplicated or
-	 * are unsupported by the script serialization
-	 */
-	virtual void PostEditImport() override;
-#endif
-
 	virtual void BeginDestroy() override;
 	virtual void FinishDestroy() override;
 	virtual bool IsReadyForFinishDestroy() override;
 	virtual bool NeedsLoadForClient() const override;
 	virtual bool NeedsLoadForServer() const override;
+#if WITH_EDITOR
+	virtual void PostEditChangeProperty(FPropertyChangedEvent& PropertyChangedEvent) override;
+	virtual void PostEditChangeChainProperty(FPropertyChangedChainEvent& PropertyChangedEvent) override;
+	virtual bool CanEditChange(const UProperty* InProperty) const override;
+	virtual void UpdateCollisionProfile();
+	virtual void PostEditImport() override;
+#endif // WITH_EDITOR
 	//~ End UObject Interface.
 
-	//Begin USceneComponent Interface
-
-protected:
-	virtual bool MoveComponentImpl(const FVector& Delta, const FQuat& NewRotation, bool bSweep, FHitResult* OutHit = NULL, EMoveComponentFlags MoveFlags = MOVECOMP_NoFlags, ETeleportType Teleport = ETeleportType::None) override;
-	
-public:
-	virtual bool IsWorldGeometry() const override;
+	//~ Begin USceneComponent Interface
 
 	/** Returns the form of collision for this component */
 	UFUNCTION(BlueprintPure, Category="Collision")
@@ -1896,10 +1838,19 @@ public:
 	/** Gets the collision object type */
 	UFUNCTION(BlueprintPure, Category="Collision")
 	virtual ECollisionChannel GetCollisionObjectType() const override;
-
+	
+	virtual void OnUpdateTransform(EUpdateTransformFlags UpdateTransformFlags, ETeleportType Teleport = ETeleportType::None) override;
+	virtual void OnAttachmentChanged() override;
+	virtual bool IsSimulatingPhysics(FName BoneName = NAME_None) const override;
+	virtual bool MoveComponentImpl(const FVector& Delta, const FQuat& NewRotation, bool bSweep, FHitResult* OutHit = NULL, EMoveComponentFlags MoveFlags = MOVECOMP_NoFlags, ETeleportType Teleport = ETeleportType::None) override;
+	virtual bool IsWorldGeometry() const override;
 	virtual const FCollisionResponseContainer& GetCollisionResponseToChannels() const override;
 	virtual FVector GetComponentVelocity() const override;
-	//End USceneComponent Interface
+#if WITH_EDITOR
+	virtual void UpdateBounds() override;
+	virtual const int32 GetNumUncachedStaticLightingInteractions() const override;
+#endif
+	//~ End USceneComponentInterface
 
 	/**
 	 * Dispatch notifications for the given HitResult.
@@ -1952,6 +1903,7 @@ public:
 	/** Gets the diffuse boost for the primitive component. */
 	virtual float GetDiffuseBoost(int32 ElementIndex) const		{ return 1.0f; };
 	
+	/** Disable dynamic shadow casting if the primitive only casts indirect shadows, since dynamic shadows are always shadowing direct lighting */
 	virtual bool GetShadowIndirectOnly() const { return false; }
 
 #if WITH_EDITOR
@@ -1959,19 +1911,6 @@ public:
 	virtual uint64 GetHiddenEditorViews() const;
 #endif// WITH_EDITOR
 
-	/**
-	 *	Set the angular velocity of all bodies in this component.
-	 *
-	 *	@param NewAngVel		New angular velocity to apply to physics, in degrees per second.
-	 *	@param bAddToCurrent	If true, NewAngVel is added to the existing angular velocity of all bodies.
-	 */
-	UE_DEPRECATED(4.8, "Use SetAllPhysicsAngularVelocityInDegrees instead.")
-	UFUNCTION(BlueprintCallable, Category = "Physics", meta = (UnsafeDuringActorConstruction = "true", DeprecatedFunction, DeprecationMessage="Use SetAllPhysicsAngularVelocityInDegrees instead"))
-	virtual void SetAllPhysicsAngularVelocity(const FVector& NewAngVel, bool bAddToCurrent = false)
-	{
-		SetAllPhysicsAngularVelocityInDegrees(NewAngVel, bAddToCurrent);
-	}
-
 	/**
 	 *	Set the angular velocity of all bodies in this component.
 	 *
@@ -2132,7 +2071,6 @@ public:
 	virtual void SetCollisionResponseToChannels(const FCollisionResponseContainer& NewReponses);
 	
 protected:
-
 	/** Called when the BodyInstance ResponseToChannels, CollisionEnabled or bNotifyRigidBodyCollision changes, in case subclasses want to use that information. */
 	virtual void OnComponentCollisionSettingsChanged();
 
@@ -2140,12 +2078,10 @@ protected:
 	void ClearComponentOverlaps(bool bDoNotifies, bool bSkipNotifySelf);
 
 private:
-	
 	/** Check if mobility is set to non-static. If BodyInstanceRequiresSimulation is non-null we check that it is simulated. Triggers a PIE warning if conditions fails */
 	void WarnInvalidPhysicsOperations_Internal(const FText& ActionText, const FBodyInstance* BodyInstanceRequiresSimulation, FName BoneName) const;
 
 public:
-
 	/**
 	 * Applies RigidBodyState only if it needs to be updated
 	 * NeedsUpdate flag will be removed from UpdatedState after all velocity corrections are finished
@@ -2238,12 +2174,10 @@ public:
 	bool ComponentOverlapComponent(class UPrimitiveComponent* PrimComp, const FVector Pos, const FRotator Rot, const FCollisionQueryParams& Params);
 
 protected:
-
-	// Override this method for custom behavior.
+	/** Override this method for custom behavior for ComponentOverlapComponent() */
 	virtual bool ComponentOverlapComponentImpl(class UPrimitiveComponent* PrimComp, const FVector Pos, const FQuat& Rot, const FCollisionQueryParams& Params);
 
-public:
-	
+public:	
 	/** 
 	 *  Test the collision of the supplied shape at the supplied location, and determine if it overlaps this component.
 	 *
@@ -2262,7 +2196,6 @@ public:
 	 * @param Rot				Rotation of collision shape
 	 * @return true if the computation succeeded - assumes that there is an overlap at the specified position/rotation
 	 */
-
 	virtual bool ComputePenetration(FMTDResult & OutMTD, const FCollisionShape& CollisionShape, const FVector& Pos, const FQuat& Rot);
 
 	/**
@@ -2279,8 +2212,10 @@ public:
 	virtual bool IsNavigationRelevant() const override;
 	//~ End INavRelevantInterface Interface
 
+	/** If true then DoCustomNavigableGeometryExport will be called to collect navigable geometry of this component. */
 	FORCEINLINE EHasCustomNavigableGeometry::Type HasCustomNavigableGeometry() const { return bHasCustomNavigableGeometry; }
 
+	/** Set value of HasCustomNavigableGeometry */
 	void SetCustomNavigableGeometry(const EHasCustomNavigableGeometry::Type InType);
 
 	/** Collects custom navigable geometry of component.
diff --git a/Engine/Source/Runtime/Engine/Classes/Components/SceneComponent.h b/Engine/Source/Runtime/Engine/Classes/Components/SceneComponent.h
--- a/Engine/Source/Runtime/Engine/Classes/Components/SceneComponent.h
+++ b/Engine/Source/Runtime/Engine/Classes/Components/SceneComponent.h
@@ -45,7 +45,7 @@ struct ENGINE_API FOverlapInfo
 // All added members of FOverlapInfo are PODs.
 template<> struct TIsPODType<FOverlapInfo> { enum { Value = TIsPODType<FHitResult>::Value }; };
 
-/** Detail mode for scene component rendering. */
+/** Detail mode for scene component rendering, corresponds with the integer value of UWorld::GetDetailMode() */
 UENUM()
 enum EDetailMode
 {
@@ -69,21 +69,25 @@ enum ERelativeTransformSpace
 	RTS_ParentBoneSpace,
 };
 
-//
-// MoveComponent options.
-//
+/** MoveComponent options, stored as bitflags */
 enum EMoveComponentFlags
 {
-	// Bitflags.
-	MOVECOMP_NoFlags						= 0x0000,	// no flags
-	MOVECOMP_IgnoreBases					= 0x0001,	// ignore collisions with things the Actor is based on
-	MOVECOMP_SkipPhysicsMove				= 0x0002,	// when moving this component, do not move the physics representation. Used internally to avoid looping updates when syncing with physics.
-	MOVECOMP_NeverIgnoreBlockingOverlaps	= 0x0004,	// never ignore initial blocking overlaps during movement, which are usually ignored when moving out of an object. MOVECOMP_IgnoreBases is still respected.
-	MOVECOMP_DisableBlockingOverlapDispatch	= 0x0008,	// avoid dispatching blocking hit events when the hit started in penetration (and is not ignored, see MOVECOMP_NeverIgnoreBlockingOverlaps).
+	/** Default options */
+	MOVECOMP_NoFlags						= 0x0000,	
+	/** Ignore collisions with things the Actor is based on */
+	MOVECOMP_IgnoreBases					= 0x0001,	
+	/** When moving this component, do not move the physics representation. Used internally to avoid looping updates when syncing with physics. */
+	MOVECOMP_SkipPhysicsMove				= 0x0002,	
+	/** Never ignore initial blocking overlaps during movement, which are usually ignored when moving out of an object. MOVECOMP_IgnoreBases is still respected. */
+	MOVECOMP_NeverIgnoreBlockingOverlaps	= 0x0004,	
+	/** avoid dispatching blocking hit events when the hit started in penetration (and is not ignored, see MOVECOMP_NeverIgnoreBlockingOverlaps). */
+	MOVECOMP_DisableBlockingOverlapDispatch	= 0x0008,	
 };
 
-#define SCENECOMPONENT_QUAT_TOLERANCE		(1.e-8f) // Comparison tolerance for checking if two FQuats are the same when moving SceneComponents.
-#define SCENECOMPONENT_ROTATOR_TOLERANCE	(1.e-4f) // Comparison tolerance for checking if two FRotators are the same when moving SceneComponents.
+/** Comparison tolerance for checking if two FQuats are the same when moving SceneComponents. */
+#define SCENECOMPONENT_QUAT_TOLERANCE		(1.e-8f) 
+/** Comparison tolerance for checking if two FRotators are the same when moving SceneComponents. */
+#define SCENECOMPONENT_ROTATOR_TOLERANCE	(1.e-4f) 
 
 FORCEINLINE EMoveComponentFlags operator|(EMoveComponentFlags Arg1,EMoveComponentFlags Arg2)	{ return EMoveComponentFlags(uint32(Arg1) | uint32(Arg2)); }
 FORCEINLINE EMoveComponentFlags operator&(EMoveComponentFlags Arg1,EMoveComponentFlags Arg2)	{ return EMoveComponentFlags(uint32(Arg1) & uint32(Arg2)); }
@@ -104,17 +108,12 @@ class ENGINE_API USceneComponent : public UActorComponent
 	GENERATED_BODY()
 
 public:
-
 	/** The name to use for the default scene root variable */
 	static FName GetDefaultSceneRootVariableName();
 
-	/**
-	 * UObject constructor that takes an ObjectInitializer
-	 */
+	/** UObject constructor that takes an optional ObjectInitializer */
 	USceneComponent(const FObjectInitializer& ObjectInitializer = FObjectInitializer::Get());
 
-	virtual void GetLifetimeReplicatedProps(TArray<FLifetimeProperty>& OutLifetimeProps) const override;
-
 	/** Cached level collection that contains the level this component is registered in, for fast access in IsVisible(). */
 	const FLevelCollection* CachedLevelCollection;
 
@@ -162,12 +161,10 @@ public:
 	FVector RelativeScale3D;
 
 private:
-
 	/** Current transform of the component, relative to the world */
 	FTransform ComponentToWorld;
 
 public:
-
 	/**
 	* Velocity of the component.
 	* @see GetComponentVelocity()
@@ -218,7 +215,6 @@ private:
 	uint8 bShouldUpdatePhysicsVolume:1;
 
 public:
-
 	/** If true, a change in the bounds of the component will call trigger a streaming data rebuild */
 	UPROPERTY()
 	uint8 bBoundsChangeTriggersStreamingDataRebuild:1;
@@ -232,44 +228,39 @@ public:
 	/** Clears the skip update overlaps flag. This should be called any time a change to state would prevent the result of UpdateOverlaps. For example attachment, changing collision settings, etc... */
 	void ClearSkipUpdateOverlaps();
 
+	/** If true, we can use the old computed overlaps */
 	bool ShouldSkipUpdateOverlaps() const
 	{
 		return SkipUpdateOverlapsOptimEnabled == 1 && bSkipUpdateOverlaps;
 	}
 
+	/** Gets whether or not the cached PhysicsVolume this component overlaps should be updated when the component is moved.	*/
 	UFUNCTION(BlueprintGetter)
 	bool GetShouldUpdatePhysicsVolume() const;
 
+	/** Sets whether or not the cached PhysicsVolume this component overlaps should be updated when the component is moved. */
 	UFUNCTION(BlueprintSetter)
 	void SetShouldUpdatePhysicsVolume(bool bInShouldUpdatePhysicsVolume);
 
 protected:
-
-	// Transient flag that temporarily disables UpdateOverlaps within DetachFromParent().
+	/** Transient flag that temporarily disables UpdateOverlaps within DetachFromParent(). */
 	uint8 bDisableDetachmentUpdateOverlaps:1;
 
 	/** If true, OnUpdateTransform virtual will be called each time this component is moved. */
 	uint8 bWantsOnUpdateTransform:1;
 
 private:
-
 	uint8 bNetUpdateTransform : 1;
 	uint8 bNetUpdateAttachment : 1;
 	uint8 bNetHasReceivedRelativeLocation : 1;
 	uint8 bNetHasReceivedRelativeRotation : 1;
 
-	// DEPRECATED
-	UPROPERTY()
-	uint8 bAbsoluteTranslation_DEPRECATED : 1;
-
-	// Appends all descendants (recursively) of this scene component to the list of Children.  NOTE: It does NOT clear the list first.
-	void AppendDescendants(TArray<USceneComponent*>& Children) const;
-
 public:
-
+	/** Global flag to enable/disable overlap optimizations, settable with p.SkipUpdateOverlapsOptimEnabled cvar */ 
 	static int32 SkipUpdateOverlapsOptimEnabled;
 
 #if WITH_EDITORONLY_DATA
+	/** This component is explicitly for visualization in the editor */
 	UPROPERTY()
 	uint8 bVisualizeComponent : 1;
 #endif
@@ -282,6 +273,12 @@ public:
 	UPROPERTY(EditAnywhere, AdvancedDisplay, BlueprintReadOnly, Category = LOD)
 	TEnumAsByte<enum EDetailMode> DetailMode;
 
+#if WITH_EDITORONLY_DATA
+protected:
+	/** Editor only component used to display the sprite so as to be able to see the location of the Audio Component  */
+	class UBillboardComponent* SpriteComponent;
+#endif
+
 private:
 	/** Cache that avoids Quat<->Rotator conversions if possible. Only to be used with GetComponentTransform().GetRotation(). */
 	FRotationConversionCache WorldRotationCache;
@@ -290,9 +287,9 @@ private:
 	FRotationConversionCache RelativeRotationCache;
 
 public:
-
 	/** Sets the RelativeRotationCache. Used to ensure component ends up with the same RelativeRotation after calling SetWorldTransform(). */
 	void SetRelativeRotationCache(const FRotationConversionCache& InCache);
+	
 	/** Get the RelativeRotationCache.  */
 	FORCEINLINE const FRotationConversionCache& GetRelativeRotationCache() const { return RelativeRotationCache; }
 
@@ -300,28 +297,19 @@ public:
 	UPROPERTY(BlueprintAssignable, Category=PhysicsVolume, meta=(DisplayName="Physics Volume Changed"))
 	FPhysicsVolumeChanged PhysicsVolumeChangedDelegate;
 
+	/** Delegate called when this component is moved */
+	FTransformUpdated TransformUpdated;
+
 	/** Returns the current scoped movement update, or NULL if there is none. @see FScopedMovementUpdate */
 	class FScopedMovementUpdate* GetCurrentScopedMovement() const;
 
 private:
-
 	/** Stack of current movement scopes. */
 	TArray<class FScopedMovementUpdate*> ScopedMovementStack;
 
 	void BeginScopedMovementUpdate(class FScopedMovementUpdate& ScopedUpdate);
 	void EndScopedMovementUpdate(class FScopedMovementUpdate& ScopedUpdate);
 
-	friend class FScopedMovementUpdate;
-	friend class FScopedPreventAttachedComponentMove;
-
-public:
-
-#if WITH_EDITORONLY_DATA
-	UPROPERTY()
-	FVector RelativeTranslation_DEPRECATED;
-#endif
-
-private:
 	UFUNCTION()
 	void OnRep_Transform();
 
@@ -343,17 +331,12 @@ private:
 	UFUNCTION()
 	void OnRep_Visibility(bool OldValue);
 
-protected:
-
-	virtual void PreNetReceive() override;
-	virtual void PostNetReceive() override;
-	virtual void PostRepNotifies() override;
-
 public:
-
-	/**  Convenience function to get the relative rotation from the passed in world rotation*/
-	/* @param WorldRotation  World rotation that we want to convert to relative to the components parent*/
-	/* @return Returns the relative rotation*/
+	/**  
+	 * Convenience function to get the relative rotation from the passed in world rotation
+	 * @param WorldRotation  World rotation that we want to convert to relative to the components parent
+	 * @return Returns the relative rotation
+	 */
 	FQuat GetRelativeRotationFromWorld(const FQuat & WorldRotation);
 
 	/**
@@ -369,7 +352,6 @@ public:
 	*/
 	void SetRelativeRotationExact(FRotator NewRotation, bool bSweep = false, FHitResult* OutSweepHitResult = nullptr, ETeleportType Teleport = ETeleportType::None);
 
-public:
 	/**
 	 * Set the location of the component relative to its parent
 	 * @param NewLocation		New location of the component relative to its parent.		
@@ -385,7 +367,6 @@ public:
 	void K2_SetRelativeLocation(FVector NewLocation, bool bSweep, FHitResult& SweepHitResult, bool bTeleport);
 	void SetRelativeLocation(FVector NewLocation, bool bSweep=false, FHitResult* OutSweepHitResult=nullptr, ETeleportType Teleport = ETeleportType::None);
 
-
 	/**
 	 * Set the rotation of the component relative to its parent
 	 * @param NewRotation		New rotation of the component relative to its parent
@@ -469,7 +450,6 @@ public:
 	void K2_AddLocalOffset(FVector DeltaLocation, bool bSweep, FHitResult& SweepHitResult, bool bTeleport);
 	void AddLocalOffset(FVector DeltaLocation, bool bSweep=false, FHitResult* OutSweepHitResult=nullptr, ETeleportType Teleport = ETeleportType::None);
 
-	
 	/**
 	 * Adds a delta to the rotation of the component in its local reference frame
 	 * @param DeltaRotation		Change in rotation of the component in its local reference frame.
@@ -484,7 +464,6 @@ public:
 	void AddLocalRotation(FRotator DeltaRotation, bool bSweep=false, FHitResult* OutSweepHitResult=nullptr, ETeleportType Teleport = ETeleportType::None);
 	void AddLocalRotation(const FQuat& DeltaRotation, bool bSweep=false, FHitResult* OutSweepHitResult=nullptr, ETeleportType Teleport = ETeleportType::None);
 
-
 	/**
 	 * Adds a delta to the transform of the component in its local reference frame. Scale is unchanged.
 	 * @param DeltaTransform	Change in transform of the component in its local reference frame. Scale is unchanged.
@@ -500,7 +479,6 @@ public:
 	void K2_AddLocalTransform(const FTransform& DeltaTransform, bool bSweep, FHitResult& SweepHitResult, bool bTeleport);
 	void AddLocalTransform(const FTransform& DeltaTransform, bool bSweep=false, FHitResult* OutSweepHitResult=nullptr, ETeleportType Teleport = ETeleportType::None);
 
-
 	/**
 	 * Put this component at the specified location in world space. Updates relative location to achieve the final world location.
 	 * @param NewLocation		New location in world space for the component.
@@ -516,7 +494,6 @@ public:
 	void K2_SetWorldLocation(FVector NewLocation, bool bSweep, FHitResult& SweepHitResult, bool bTeleport);
 	void SetWorldLocation(FVector NewLocation, bool bSweep=false, FHitResult* OutSweepHitResult=nullptr, ETeleportType Teleport = ETeleportType::None);
 
-
 	/*
 	 * Put this component at the specified rotation in world space. Updates relative rotation to achieve the final world rotation.
 	 * @param NewRotation		New rotation in world space for the component.
@@ -532,7 +509,6 @@ public:
 	void SetWorldRotation(FRotator NewRotation, bool bSweep=false, FHitResult* OutSweepHitResult=nullptr, ETeleportType Teleport = ETeleportType::None);
 	void SetWorldRotation(const FQuat& NewRotation, bool bSweep=false, FHitResult* OutSweepHitResult=nullptr, ETeleportType Teleport = ETeleportType::None);
 
-
 	/**
 	 * Set the relative scale of the component to put it at the supplied scale in world space.
 	 * @param NewScale		New scale in world space for this component.
@@ -555,7 +531,6 @@ public:
 	void K2_SetWorldTransform(const FTransform& NewTransform, bool bSweep, FHitResult& SweepHitResult, bool bTeleport);
 	void SetWorldTransform(const FTransform& NewTransform, bool bSweep=false, FHitResult* OutSweepHitResult=nullptr, ETeleportType Teleport = ETeleportType::None);
 
-
 	/**
 	 * Adds a delta to the location of the component in world space.
 	 * @param DeltaLocation		Change in location in world space for the component.
@@ -571,7 +546,6 @@ public:
 	void K2_AddWorldOffset(FVector DeltaLocation, bool bSweep, FHitResult& SweepHitResult, bool bTeleport);
 	void AddWorldOffset(FVector DeltaLocation, bool bSweep=false, FHitResult* OutSweepHitResult=nullptr, ETeleportType Teleport = ETeleportType::None);
 
-
 	/**
 	 * Adds a delta to the rotation of the component in world space.
 	 * @param DeltaRotation		Change in rotation in world space for the component.
@@ -638,8 +612,6 @@ public:
 	UFUNCTION(BlueprintCallable, Category="Physics")
 	virtual bool IsAnySimulatingPhysics() const;
 
-public:
-
 	/** Get the SceneComponents that are attached to this component. */
 	const TArray<USceneComponent*>& GetAttachChildren() const;
 
@@ -651,7 +623,7 @@ public:
 	UFUNCTION(BlueprintCallable, Category="Utilities|Transformation")
 	FName GetAttachSocketName() const;
 
-	/** Gets all parent components up to and including the root component */
+	/** Gets all attachment parent components up to and including the root component */
 	UFUNCTION(BlueprintCallable, Category="Components")
 	void GetParentComponents(TArray<USceneComponent*>& Parents) const;
 
@@ -664,7 +636,7 @@ public:
 	USceneComponent* GetChildComponent(int32 ChildIndex) const;
 
 	/** 
-	 * Gets all the attached child components
+	 * Gets all components that are attached to this component, possibly recursively
 	 * @param bIncludeAllDescendants Whether to include all descendants in the list of children (i.e. grandchildren, great grandchildren, etc.)
 	 * @param Children The list of attached child components
 	 */
@@ -683,63 +655,44 @@ public:
 	/** Backwards compatibility: Used to convert old-style EAttachLocation to new-style EAttachmentRules */
 	static void ConvertAttachLocation(EAttachLocation::Type InAttachLocation, EAttachmentRule& InOutLocationRule, EAttachmentRule& InOutRotationRule, EAttachmentRule& InOutScaleRule);
 
-	/**
-	 * Attach this component to another scene component, optionally at a named socket. It is valid to call this on components whether or not they have been Registered.
-	 * @param  InParent				Parent to attach to.
-	 * @param  InSocketName			Optional socket to attach to on the parent.
-	 * @param  AttachType			How to handle transform when attaching (Keep relative offset, keep world position, etc).
-	 * @param  bWeldSimulatedBodies Whether to weld together simulated physics bodies.
-	 * @return True if attachment is successful (or already attached to requested parent/socket), false if attachment is rejected and there is no change in AttachParent.
-	 */
 	UE_DEPRECATED(4.12, "This function is deprecated, please use AttachToComponent instead.")
 	bool AttachTo(USceneComponent* InParent, FName InSocketName = NAME_None, EAttachLocation::Type AttachType = EAttachLocation::KeepRelativeOffset, bool bWeldSimulatedBodies = false);
 
-	/**
-	 * Attach this component to another scene component, optionally at a named socket. It is valid to call this on components whether or not they have been Registered.
-	 * @param  InParent				Parent to attach to.
-	 * @param  InSocketName			Optional socket to attach to on the parent.
-	 * @param  AttachType			How to handle transform when attaching (Keep relative offset, keep world position, etc).
-	 * @param  bWeldSimulatedBodies Whether to weld together simulated physics bodies.
-	 * @return True if attachment is successful (or already attached to requested parent/socket), false if attachment is rejected and there is no change in AttachParent.
-	*/
-	UE_DEPRECATED(4.17, "This function is deprecated, please use AttachToComponent instead.")
-	UFUNCTION(BlueprintCallable, Category = "Utilities|Transformation", meta = (DisplayName = "AttachTo (Deprecated)", AttachType = "KeepRelativeOffset"))
+	/** DEPRECATED - Use AttachToComponent() instead */
+	UE_DEPRECATED(4.17, "This function is deprecated, please use AttachToComponent() instead.")
+	UFUNCTION(BlueprintCallable, Category = "Utilities|Transformation", meta = (DeprecationMessage = "OVERRIDE BAD MESSAGE", DisplayName = "AttachTo (Deprecated)", AttachType = "KeepRelativeOffset"))
 	bool K2_AttachTo(USceneComponent* InParent, FName InSocketName = NAME_None, EAttachLocation::Type AttachType = EAttachLocation::KeepRelativeOffset, bool bWeldSimulatedBodies = true);
 
 	/**
-	* Attach this component to another scene component, optionally at a named socket. It is valid to call this on components whether or not they have been Registered, however from
-	* constructor or when not registered it is preferable to use SetupAttachment.
-	* @param  Parent				Parent to attach to.
-	* @param  AttachmentRules		How to handle transforms & welding when attaching.
-	* @param  SocketName			Optional socket to attach to on the parent.
-	* @return True if attachment is successful (or already attached to requested parent/socket), false if attachment is rejected and there is no change in AttachParent.
-	*/
+	 * Attach this component to another scene component, optionally at a named socket. It is valid to call this on components whether or not they have been Registered, however from
+	 * constructor or when not registered it is preferable to use SetupAttachment.
+	 * @param  Parent				Parent to attach to.
+	 * @param  AttachmentRules		How to handle transforms & welding when attaching.
+	 * @param  SocketName			Optional socket to attach to on the parent.
+	 * @return True if attachment is successful (or already attached to requested parent/socket), false if attachment is rejected and there is no change in AttachParent.
+	 */
 	bool AttachToComponent(USceneComponent* InParent, const FAttachmentTransformRules& AttachmentRules, FName InSocketName = NAME_None );
 
 	/**
-	* Attach this component to another scene component, optionally at a named socket. It is valid to call this on components whether or not they have been Registered.
-	* @param  Parent					Parent to attach to.
-	* @param  SocketName				Optional socket to attach to on the parent.
-	* @param  LocationRule				How to handle translation when attaching.
-	* @param  RotationRule				How to handle rotation when attaching.
-	* @param  ScaleRule					How to handle scale when attaching.
-	* @param  bWeldSimulatedBodies		Whether to weld together simulated physics bodies.
-	* @return True if attachment is successful (or already attached to requested parent/socket), false if attachment is rejected and there is no change in AttachParent.
-	*/
+	 * Attach this component to another scene component, optionally at a named socket. It is valid to call this on components whether or not they have been Registered.
+	 * @param  Parent					Parent to attach to.
+	 * @param  SocketName				Optional socket to attach to on the parent.
+	 * @param  LocationRule				How to handle translation when attaching.
+	 * @param  RotationRule				How to handle rotation when attaching.
+	 * @param  ScaleRule					How to handle scale when attaching.
+	 * @param  bWeldSimulatedBodies		Whether to weld together simulated physics bodies.
+	 * @return True if attachment is successful (or already attached to requested parent/socket), false if attachment is rejected and there is no change in AttachParent.
+	 */
 	UFUNCTION(BlueprintCallable, Category = "Utilities|Transformation", meta = (DisplayName = "AttachToComponent", ScriptName = "AttachToComponent", bWeldSimulatedBodies=true))
 	bool K2_AttachToComponent(USceneComponent* Parent, FName SocketName, EAttachmentRule LocationRule, EAttachmentRule RotationRule, EAttachmentRule ScaleRule, bool bWeldSimulatedBodies);
 
-	/** Zeroes out the relative transform of the component, and calls AttachTo(). Useful for attaching directly to a scene component or socket location  */
-	UE_DEPRECATED(4.17, "Use AttachToComponent instead.")
+	/** DEPRECATED - Use AttachToComponent() instead */
+	UE_DEPRECATED(4.17, "Use AttachToComponent() instead.")
 	UFUNCTION(BlueprintCallable, meta=(DeprecatedFunction, DeprecationMessage = "Use AttachToComponent instead."), Category="Utilities|Transformation")
 	bool SnapTo(USceneComponent* InParent, FName InSocketName = NAME_None);
 
-	/** 
-	 * Detach this component from whatever it is attached to. Automatically unwelds components that are welded together (See WeldTo)
-	 * @param bMaintainWorldPosition	If true, update the relative location of the component to keep its world position the same
-	 * @param bCallModify				If true, call Modify() on the component and the current attach parent component
-	 */
-	UE_DEPRECATED(4.12, "This function is deprecated, please use DetachFromComponent instead.")
+	/** DEPRECATED - Use DetachFromComponent() instead */
+	UE_DEPRECATED(4.12, "This function is deprecated, please use DetachFromComponent() instead.")
 	UFUNCTION(BlueprintCallable, Category="Utilities|Transformation", meta = (DisplayName = "DetachFromParent (Deprecated)"))
 	virtual void DetachFromParent(bool bMaintainWorldPosition = false, bool bCallModify = true);
 
@@ -799,9 +752,8 @@ public:
 	virtual FQuat GetSocketQuaternion(FName InSocketName) const;
 
 	/** 
-	 * return true if socket with the given name exists
+	 * Return true if socket with the given name exists
 	 * @param InSocketName Name of the socket or the bone to get the transform 
-	 * @return true if the socket with the given name exists. Otherwise, return false
 	 */
 	UFUNCTION(BlueprintCallable, Category="Utilities|Transformation", meta=(Keywords="Bone"))
 	virtual bool DoesSocketExist(FName InSocketName) const;
@@ -823,15 +775,11 @@ public:
 	UFUNCTION(BlueprintCallable, Category="Utilities|Transformation")
 	virtual FVector GetComponentVelocity() const;
 
-	/** 
-	 * Is this component visible or not in game
-	 * @return true if visible
-	 */
+	/** Returns true if this component is visible in the current context */
 	UFUNCTION(BlueprintCallable, Category="Rendering")
 	virtual bool IsVisible() const;
 
 protected:
-
 	/**
 	 * Overridable internal function to respond to changes in the visibility of the component.
 	 */
@@ -843,19 +791,18 @@ protected:
 	virtual void OnHiddenInGameChanged();
 
 private:
-
 	/** 
 	 * Enum that dictates what propagation policy to follow when calling SetVisibility or SetHiddenInGame recursively 
 	 */
 	enum class EVisibilityPropagation : uint8
 	{
-		// Only change the visibility if needed
+		/** Only change the visibility if needed */
 		NoPropagation, 
 
-		// If the visibility changed, mark all attached component's render states as dirty
+		/** If the visibility changed, mark all attached component's render states as dirty */
 		DirtyOnly,
 
-		// Call function recursively on attached components and also mark their render state as dirty
+		/** Call function recursively on attached components and also mark their render state as dirty */
 		Propagate
 	};
 
@@ -869,8 +816,10 @@ private:
 	*/
 	void SetHiddenInGame(bool bNewHiddenInGame, EVisibilityPropagation PropagateToChildren);
 
-public:
+	/** Appends all descendants (recursively) of this scene component to the list of Children.  NOTE: It does NOT clear the list first. */
+	void AppendDescendants(TArray<USceneComponent*>& Children) const;
 
+public:
 	/** 
 	 * Set visibility of the component, if during game use this to turn on/off
 	 */
@@ -889,27 +838,18 @@ public:
 		SetVisibility(!bVisible, bPropagateToChildren);
 	}
 
-	/**
-	 * Changes the value of HiddenGame.
-	 * @param NewHidden	- The value to assign to HiddenGame.
-	 */
+	/** Changes the value of bHiddenInGame, if false this will disable Visibility during gameplay */
 	UFUNCTION(BlueprintCallable, Category="Development")
 	void SetHiddenInGame(bool NewHidden, bool bPropagateToChildren=false)
 	{
 		SetHiddenInGame(NewHidden, bPropagateToChildren ? EVisibilityPropagation::Propagate : EVisibilityPropagation::DirtyOnly);
 	}
 
-public:
-	FTransformUpdated TransformUpdated;
-
 	//~ Begin ActorComponent Interface
+	virtual void GetLifetimeReplicatedProps(TArray<FLifetimeProperty>& OutLifetimeProps) const override;
 	virtual void OnRegister() override;
 	virtual void OnUnregister() override;
-	/** Return true if CreateRenderState() should be called */
-	virtual bool ShouldCreateRenderState() const override
-	{
-		return true;
-	}
+	virtual bool ShouldCreateRenderState() const override { return true; }
 	virtual void UpdateComponentToWorld(EUpdateTransformFlags UpdateTransformFlags = EUpdateTransformFlags::None, ETeleportType Teleport = ETeleportType::None) override final
 	{
 		UpdateComponentToWorldWithParent(GetAttachParent(), GetAttachSocketName(), UpdateTransformFlags, RelativeRotationCache.RotatorToQuat(RelativeRotation), Teleport);
@@ -920,15 +860,14 @@ public:
 	virtual TStructOnScope<FActorComponentInstanceData> GetComponentInstanceData() const override;
 	//~ End ActorComponent Interface
 
-	// Call UpdateComponentToWorld if bComponentToWorldUpdated is false.
-	void ConditionalUpdateComponentToWorld();
-
 	//~ Begin UObject Interface
 	virtual void PostInterpChange(UProperty* PropertyThatChanged) override;
 	virtual void BeginDestroy() override;
 	virtual bool IsPostLoadThreadSafe() const override;
+	virtual void PreNetReceive() override;
+	virtual void PostNetReceive() override;
+	virtual void PostRepNotifies() override;
 #if WITH_EDITORONLY_DATA
-	virtual void Serialize(FArchive& Ar) override;
 	static void AddReferencedObjects(UObject* InThis, FReferenceCollector& Collector);
 #endif
 
@@ -936,9 +875,10 @@ public:
 	virtual bool NeedsLoadForTargetPlatform(const ITargetPlatform* TargetPlatform) const;
 	virtual void PostEditChangeProperty(FPropertyChangedEvent& PropertyChangedEvent) override;
 	virtual void PostEditChangeChainProperty(FPropertyChangedChainEvent& PropertyChangedEvent) override;
+	virtual bool CanEditChange(const UProperty* Property) const override;
 #endif
-
 	//~ End UObject Interface
+
 protected:
 	/**
 	 * Internal helper, for use from MoveComponent().  Special codepath since the normal setters call MoveComponent.
@@ -946,15 +886,16 @@ protected:
 	 */
 	bool InternalSetWorldLocationAndRotation(FVector NewLocation, const FQuat& NewQuat, bool bNoPhysics = false, ETeleportType Teleport = ETeleportType::None);
 
+	/** Native callback when this component is moved */
 	virtual void OnUpdateTransform(EUpdateTransformFlags UpdateTransformFlags, ETeleportType Teleport = ETeleportType::None);
 
 	/** Check if mobility is set to non-static. If it's static we trigger a PIE warning and return true*/
 	bool CheckStaticMobilityAndWarn(const FText& ActionText) const;
 
+	/** Internal helper for UpdateOverlaps */
 	virtual bool UpdateOverlapsImpl(TArray<FOverlapInfo> const* PendingOverlaps = nullptr, bool bDoNotifies = true, const TArray<FOverlapInfo>* OverlapsAtEndLocation = nullptr);
 
 private:
-
 	void PropagateTransformUpdate(bool bTransformChanged, EUpdateTransformFlags UpdateTransformFlags = EUpdateTransformFlags::None, ETeleportType Teleport = ETeleportType::None);
 	void UpdateComponentToWorldWithParent(USceneComponent* Parent, FName SocketName, EUpdateTransformFlags UpdateTransformFlags, const FQuat& RelativeRotationQuat, ETeleportType Teleport = ETeleportType::None);
 
@@ -987,24 +928,16 @@ public:
 	bool MoveComponent( const FVector& Delta, const FRotator& NewRotation, bool bSweep, FHitResult* Hit=NULL, EMoveComponentFlags MoveFlags = MOVECOMP_NoFlags, ETeleportType Teleport = ETeleportType::None);
 
 protected:
-
-	// Override this method for custom behavior.
+	/** Override this method for custom behavior for MoveComponent */
 	virtual bool MoveComponentImpl(const FVector& Delta, const FQuat& NewRotation, bool bSweep, FHitResult* Hit = NULL, EMoveComponentFlags MoveFlags = MOVECOMP_NoFlags, ETeleportType Teleport = ETeleportType::None);
 
 public:
-
+	/** Call UpdateComponentToWorld if bComponentToWorldUpdated is false. */
+	void ConditionalUpdateComponentToWorld();
 
 	/** Returns true if movement is currently within the scope of an FScopedMovementUpdate. */
 	bool IsDeferringMovementUpdates() const;
 
-#if WITH_EDITORONLY_DATA
-protected:
-	/** Editor only component used to display the sprite so as to be able to see the location of the Audio Component  */
-	class UBillboardComponent* SpriteComponent;
-#endif
-
-public:
-
 	/** Called when AttachParent changes, to allow the scene to update its attachment state. */
 	virtual void OnAttachmentChanged() {}
 
@@ -1096,7 +1029,6 @@ public:
 	UFUNCTION(BlueprintCallable, Category=PhysicsVolume, meta=(UnsafeDuringActorConstruction="true"))
 	APhysicsVolume* GetPhysicsVolume() const;
 
-
 	/** Return const reference to CollsionResponseContainer */
 	virtual const FCollisionResponseContainer& GetCollisionResponseToChannels() const;
 
@@ -1119,10 +1051,11 @@ public:
 #if WITH_EDITOR
 	/** Called when this component is moved in the editor */
 	virtual void PostEditComponentMove(bool bFinished);
-	virtual bool CanEditChange( const UProperty* Property ) const override;
 
+	/** Returns number of lighting interactions that need to be recalculated */
 	virtual const int32 GetNumUncachedStaticLightingInteractions() const;
 
+	/** Called to update any visuals needed for a feature level change */
 	virtual void PreFeatureLevelChange(ERHIFeatureLevel::Type PendingFeatureLevel) {}
 #endif // WITH_EDITOR
 
@@ -1151,9 +1084,9 @@ protected:
 		}
 	}
 
+	/** Utility function to handle calculating transform with a parent */
 	FTransform CalcNewComponentToWorld_GeneralCase(const FTransform& NewRelativeTransform, const USceneComponent* Parent, FName SocketName) const;
 
-	
 public:
 	/**
 	 * Set the location and rotation of the component relative to its parent
@@ -1198,7 +1131,7 @@ public:
 	/** Special version of SetWorldLocationAndRotation that does not affect physics. */
 	void SetWorldLocationAndRotationNoPhysics(const FVector& NewLocation, const FRotator& NewRotation);
 
-	/** Is this component considered 'world' geometry */
+	/** Is this component considered 'world' geometry, by default checks if this uses the WorldStatic collision channel */
 	virtual bool IsWorldGeometry() const;
 
 	/** Returns the form of collision for this component */
@@ -1293,19 +1226,23 @@ protected:
 	{
 		return (IsOwnerRunningUserConstructionScript()) || !(IsRegistered() && (Mobility == EComponentMobility::Static || (!bIgnoreStationary && Mobility == EComponentMobility::Stationary)));
 	}
+
 public:
 	/** Determines whether or not the component can have its mobility set to static */
 	virtual const bool CanHaveStaticMobility() const { return true; }
 
+	/** Updates any visuals after the lighting has changed */
 	virtual void PropagateLightingScenarioChange() {}
 
+	/** True if our precomputed lighting is up to date */
 	virtual bool IsPrecomputedLightingValid() const
 	{
 		return false;
 	}
 
 private:
-
+	friend class FScopedMovementUpdate;
+	friend class FScopedPreventAttachedComponentMove;
 	friend struct FDirectAttachChildrenAccessor;
 };
 
@@ -1570,17 +1507,13 @@ public:
 	/** Registers that this move is a teleport */
 	void SetHasTeleported(ETeleportType InTeleportType);
 
-	//--------------------------------------------------------------------------------------------------------//
-
 protected:
-
 	/** Fills in the list of overlaps at the end location (in EndOverlaps). Returns pointer to the list, or null if it can't be computed. */
 	const TArray<FOverlapInfo>* GetOverlapsAtEnd(class UPrimitiveComponent& PrimComponent, TArray<FOverlapInfo>& EndOverlaps, bool bTransformChanged) const;
 
 	bool SetWorldLocationAndRotation(FVector NewLocation, const FQuat& NewQuat, bool bNoPhysics = false, ETeleportType Teleport = ETeleportType::None);
 
 private:
-
 	/** Notify this scope that the given inner scope completed its update (ie is going out of scope). Only occurs for deferred updates. */
 	void OnInnerScopeComplete(const FScopedMovementUpdate& InnerScope);
 	
@@ -1592,7 +1525,6 @@ private:
 	void	operator delete[]	(void*);
 
 protected:
-
 	USceneComponent* Owner;
 	FScopedMovementUpdate* OuterDeferredScope;
 
diff --git a/Engine/Source/Runtime/Engine/Classes/Engine/EngineTypes.h b/Engine/Source/Runtime/Engine/Classes/Engine/EngineTypes.h
--- a/Engine/Source/Runtime/Engine/Classes/Engine/EngineTypes.h
+++ b/Engine/Source/Runtime/Engine/Classes/Engine/EngineTypes.h
@@ -27,6 +27,7 @@ class USceneComponent;
  */
 enum { NumInlinedActorComponents = 24 };
 
+/** Enum describing how to constrain perspective view port FOV */
 UENUM()
 enum EAspectRatioAxisConstraint
 {
@@ -36,18 +37,6 @@ enum EAspectRatioAxisConstraint
 	AspectRatio_MAX,
 };
 
-
-/** The type of metric we want about the actor **/
-UENUM()
-enum EActorMetricsType
-{
-	METRICS_VERTS,
-	METRICS_TRIS,
-	METRICS_SECTIONS,
-	METRICS_MAX,
-};
-
-
 /** Return values for UEngine::Browse. */
 namespace EBrowseReturnVal
 {
@@ -164,6 +153,7 @@ struct ENGINE_API FDetachmentTransformRules
 	bool bCallModify;
 };
 
+/** Deprecated rules for setting transform on attachment, new functions should use FAttachmentTransformRules isntead */
 UENUM()
 namespace EAttachLocation
 {
@@ -183,7 +173,6 @@ namespace EAttachLocation
 	};
 }
 
-
 /**
  * A priority for sorting scene elements by depth.
  * Elements with higher priority occlude elements with lower priority, disregarding distance.
@@ -198,7 +187,7 @@ enum ESceneDepthPriorityGroup
 	SDPG_MAX,
 };
 
-
+/** Quality of indirect lighting for Movable primitives. This has a large effect on Indirect Lighting Cache update time. */
 UENUM()
 enum EIndirectLightingCacheQuality
 {
@@ -210,6 +199,7 @@ enum EIndirectLightingCacheQuality
 	ILCQ_Volume
 };
 
+/** Type of lightmap that is used for primitive components */
 UENUM()
 enum class ELightmapType : uint8
 {
@@ -225,6 +215,7 @@ enum class ELightmapType : uint8
 	ForceVolumetric
 };
 
+/** Controls how occlusion from Distance Field Ambient Occlusion is combined with Screen Space Ambient Occlusion. */
 UENUM()
 enum EOcclusionCombineMode
 {
@@ -239,8 +230,11 @@ enum EOcclusionCombineMode
 	OCM_MAX,
 };
 
-/** Note: This is mirrored in Lightmass, be sure to update the blend mode structure and logic there if this changes. */
-// Note: Check UMaterialInstance::Serialize if changed!!
+/**
+ * The blending mode for materials
+ * @warning This is mirrored in Lightmass, be sure to update the blend mode structure and logic there if this changes. 
+ * @warning Check UMaterialInstance::Serialize if changed!!
+ */
 UENUM(BlueprintType)
 enum EBlendMode
 {
@@ -253,7 +247,7 @@ enum EBlendMode
 	BLEND_MAX,
 };
 
-
+/** Controls where the sampler for different texture lookups comes from */
 UENUM()
 enum ESamplerSourceMode
 {
@@ -265,7 +259,7 @@ enum ESamplerSourceMode
 	SSM_Clamp_WorldGroupSettings UMETA(DisplayName="Shared: Clamp")
 };
 
-
+/** Describes how to handle lighting of translucent objets */
 UENUM()
 enum ETranslucencyLightingMode
 {
@@ -349,6 +343,7 @@ namespace ETranslucentSortPolicy
 	};
 }
 
+/** Specifies which component of the scene rendering should be output to the final render target. */
 UENUM()
 enum ESceneCaptureSource 
 { 
@@ -362,6 +357,7 @@ enum ESceneCaptureSource
 	SCS_BaseColor UMETA(DisplayName="BaseColor in RGB (Deferred Renderer only)")
 };
 
+/** Specifies how scene captures are composited into render buffers */
 UENUM()
 enum ESceneCaptureCompositeMode
 { 
@@ -370,12 +366,14 @@ enum ESceneCaptureCompositeMode
 	SCCM_Composite UMETA(DisplayName="Composite")
 };
 
+/** Maximum number of custom lighting channels */
 #define NUM_LIGHTING_CHANNELS 3
 
+/** Specifies which lighting channels are relevant */
 USTRUCT(BlueprintType)
 struct FLightingChannels
 {
-	GENERATED_USTRUCT_BODY()
+	GENERATED_BODY()
 
 	FLightingChannels() :
 		bChannel0(true),
@@ -387,43 +385,51 @@ struct FLightingChannels
 	UPROPERTY(EditAnywhere, BlueprintReadOnly, Category=Channels)
 	uint8 bChannel0:1;
 
+	/** First custom channel */
 	UPROPERTY(EditAnywhere, BlueprintReadOnly, Category=Channels)
 	uint8 bChannel1:1;
 
+	/** Second custom channel */
 	UPROPERTY(EditAnywhere, BlueprintReadOnly, Category=Channels)
 	uint8 bChannel2:1;
 };
 
+/** Converts lighting channels into a bitfield */
 inline uint8 GetLightingChannelMaskForStruct(FLightingChannels Value)
 {
 	// Note: this is packed into 3 bits of a stencil channel
 	return (uint8)((Value.bChannel0 ? 1 : 0) | (Value.bChannel1 << 1) | (Value.bChannel2 << 2));
 }
 
+/** Returns mask for only channel 0  */
 inline uint8 GetDefaultLightingChannelMask()
 {
 	return 1;
 }
 
-// Returns the index of the first lighting channel set, or -1 if no channels are set.
+/** Returns the index of the first lighting channel set, or -1 if no channels are set. */
 inline int32 GetFirstLightingChannelFromMask(uint8 Mask)
 {
 	return Mask ? FPlatformMath::CountTrailingZeros(Mask) : -1;
 }
 
-/*
-* Enumerates available GBufferFormats.
-*/
+/** 
+ * Enumerates available GBufferFormats. 
+ * @warning When this enum is updated please update CVarGBufferFormat comments 
+ */
 UENUM()
 namespace EGBufferFormat
 {
-	// When this enum is updated please update CVarGBufferFormat comments 
 	enum Type
 	{
-		Force8BitsPerChannel = 0 UMETA(DisplayName = "Force 8 Bits Per Channel", ToolTip = "Forces all GBuffers to 8 bits per channel. Intended as profiling for best performance."),
-		Default = 1 UMETA(ToolTip = "See GBuffer allocation function for layout details."),
-		HighPrecisionNormals = 3 UMETA(ToolTip = "Same as Default except normals are encoded at 16 bits per channel."),
-		Force16BitsPerChannel = 5 UMETA(DisplayName = "Force 16 Bits Per Channel", ToolTip = "Forces all GBuffers to 16 bits per channel. Intended as profiling for best quality."),
+		/** Forces all GBuffers to 8 bits per channel. Intended as profiling for best performance. */
+		Force8BitsPerChannel = 0 UMETA(DisplayName = "Force 8 Bits Per Channel"),
+		/** See GBuffer allocation function for layout details. */
+		Default = 1,
+		/** Same as Default except normals are encoded at 16 bits per channel. */
+		HighPrecisionNormals = 3,
+		/** Forces all GBuffers to 16 bits per channel. Intended as profiling for best quality. */
+		Force16BitsPerChannel = 5 UMETA(DisplayName = "Force 16 Bits Per Channel"),
 	};
 }
 
@@ -436,6 +442,7 @@ enum ETrailWidthMode
 	ETrailWidthMode_FromSecond UMETA(DisplayName = "From Second Socket"),
 };
 
+/** Specifies how particle collision is computed for GPU particles */
 UENUM()
 namespace EParticleCollisionMode
 {
@@ -446,9 +453,10 @@ namespace EParticleCollisionMode
 	};
 }
 
-
-// Note: Check UMaterialInstance::Serialize if changed!
-
+/** 
+ * Specifies the overal rendering/shading model for a material
+ * @warning Check UMaterialInstance::Serialize if changed!
+ */
 UENUM()
 enum EMaterialShadingModel
 {
@@ -465,7 +473,6 @@ enum EMaterialShadingModel
 	MSM_MAX,
 };
 
-
 /** This is used by the drawing passes to determine tessellation policy, so changes here need to be supported in native code. */
 UENUM()
 enum EMaterialTessellationMode
@@ -479,7 +486,7 @@ enum EMaterialTessellationMode
 	MTM_MAX,
 };
 
-
+/** Describes how textures are sampled for materials */
 UENUM(BlueprintType)
 enum EMaterialSamplerType
 {
@@ -495,7 +502,6 @@ enum EMaterialSamplerType
 	SAMPLERTYPE_MAX,
 };
 
-
 /**	Lighting build quality enumeration */
 UENUM()
 enum ELightingBuildQuality
@@ -539,7 +545,6 @@ enum EMovementMode
 	MOVE_MAX		UMETA(Hidden),
 };
 
-
 /** Smoothing approach used by network interpolation for Characters. */
 UENUM(BlueprintType)
 enum class ENetworkSmoothingMode : uint8
@@ -648,23 +653,24 @@ enum ECollisionChannel
 
 DECLARE_DELEGATE_OneParam(FOnConstraintBroken, int32 /*ConstraintIndex*/);
 
-
 #define COLLISION_GIZMO ECC_EngineTraceChannel1
 
-
-/** @note, if you change this, change GetCollisionChannelFromOverlapFilter() to match */
+/** 
+ * Specifies what types of objects to return from an overlap physics query
+ * @warning If you change this, change GetCollisionChannelFromOverlapFilter() to match 
+ */
 UENUM(BlueprintType)
 enum EOverlapFilterOption
 {
-	// returns both overlaps with both dynamic and static components
+	/** Returns both overlaps with both dynamic and static components */
 	OverlapFilter_All UMETA(DisplayName="AllObjects"),
-	// returns only overlaps with dynamic actors (far fewer results in practice, much more efficient)
+	/** returns only overlaps with dynamic actors (far fewer results in practice, much more efficient) */
 	OverlapFilter_DynamicOnly UMETA(DisplayName="AllDynamicObjects"),
-	// returns only overlaps with static actors (fewer results, more efficient)
+	/** returns only overlaps with static actors (fewer results, more efficient) */
 	OverlapFilter_StaticOnly UMETA(DisplayName="AllStaticObjects"),
 };
 
-
+/** Specifies custom collision object types, overridable per game */
 UENUM(BlueprintType)
 enum EObjectTypeQuery
 {
@@ -704,7 +710,7 @@ enum EObjectTypeQuery
 	ObjectTypeQuery_MAX	UMETA(Hidden)
 };
 
-
+/** Specifies custom collision trace types, overridable per game */
 UENUM(BlueprintType)
 enum ETraceTypeQuery
 {
@@ -754,7 +760,7 @@ enum ECollisionResponse
 	ECR_MAX,
 };
 
-
+/** Interpolation method used by animation blending */
 UENUM()
 enum EFilterInterpolationType
 {
@@ -764,20 +770,7 @@ enum EFilterInterpolationType
 	BSIT_MAX
 };
 
-
-UENUM()
-enum EInputConsumeOptions
-{
-	/** This component consumes all input and no components lower in the stack are processed. */
-	ICO_ConsumeAll=0,
-	/** This component consumes all events for a keys it has bound (whether or not they are handled successfully).  Components lower in the stack will not receive events from these keys. */
-	ICO_ConsumeBoundKeys,
-	/** All input events will be available to components lower in the stack. */
-	ICO_ConsumeNone,
-	ICO_MAX
-};
-
-
+/** Specifies the goal/source of a UWorld object */
 namespace EWorldType
 {
 	enum Type
@@ -811,19 +804,22 @@ namespace EWorldType
 	const EWorldType::Type Preview = EWorldType::EditorPreview;
 }
 
-
+/** Describes what parts of level streaming should be forcibly handled immediately */
 enum class EFlushLevelStreamingType : uint8
 {
+	/** Do not flush state on change */
 	None,			
-	Full,			// Allow multiple load requests
-	Visibility,		// Flush visibility only, do not allow load requests, flushes async loading as well
+	/** Allow multiple load requests */
+	Full,			
+	/** Flush visibility only, do not allow load requests, flushes async loading as well */
+	Visibility,		
 };
 
-
+/** Describes response for a single collision response channel */
 USTRUCT()
 struct FResponseChannel
 {
-	GENERATED_USTRUCT_BODY()
+	GENERATED_BODY()
 
 	/** This should match DisplayName of ECollisionChannel 
 	 *	Meta data of custom channels can be used as well
@@ -831,6 +827,7 @@ struct FResponseChannel
 	 UPROPERTY(EditAnywhere, Category = FResponseChannel)
 	FName Channel;
 
+	/** Describes how the channel behaves */
 	UPROPERTY(EditAnywhere, Category = FResponseChannel)
 	TEnumAsByte<enum ECollisionResponse> Response;
 
@@ -854,7 +851,7 @@ struct FResponseChannel
 USTRUCT(BlueprintType)
 struct ENGINE_API FCollisionResponseContainer
 {
-	GENERATED_USTRUCT_BODY()
+	GENERATED_BODY()
 
 #if !CPP      //noexport property
 
@@ -1050,7 +1047,8 @@ struct ENGINE_API FCollisionResponseContainer
 	/** Take two response containers and create a new container where each element is the 'min' of the two inputs (ie Ignore and Block results in Ignore) */
 	static FCollisionResponseContainer CreateMinContainer(const FCollisionResponseContainer& A, const FCollisionResponseContainer& B);
 
-	static const struct FCollisionResponseContainer & GetDefaultResponseContainer() { return DefaultResponseContainer; }
+	/** Returns the game-wide default collision response */
+	static const struct FCollisionResponseContainer& GetDefaultResponseContainer() { return DefaultResponseContainer; }
 
 private:
 
@@ -1060,7 +1058,6 @@ private:
 	friend class UCollisionProfile;
 };
 
-
 /** Enum for controlling the falloff of strength of a radial impulse as a function of distance from Origin. */
 UENUM()
 enum ERadialImpulseFalloff
@@ -1072,7 +1069,6 @@ enum ERadialImpulseFalloff
 	RIF_MAX,
 };
 
-
 /** Presets of values used in considering when put this body to sleep. */
 UENUM()
 enum class ESleepFamily : uint8
@@ -1085,7 +1081,6 @@ enum class ESleepFamily : uint8
 	Custom,
 };
 
-
 /** Enum used to indicate what type of timeline signature a function matches. */
 UENUM()
 enum ETimelineSigType
@@ -1098,7 +1093,6 @@ enum ETimelineSigType
 	ETS_MAX,
 };
 
-
 /** Enum used to describe what type of collision is enabled on a body. */
 UENUM(BlueprintType)
 namespace ECollisionEnabled 
@@ -1132,7 +1126,7 @@ FORCEINLINE bool CollisionEnabledHasQuery(ECollisionEnabled::Type CollisionEnabl
 USTRUCT()
 struct FRigidBodyState
 {
-	GENERATED_USTRUCT_BODY()
+	GENERATED_BODY()
 
 	UPROPERTY()
 	FVector_NetQuantize100 Position;
@@ -1158,7 +1152,7 @@ struct FRigidBodyState
 	{ }
 };
 
-
+/** Describes extra state about a specific rigid body */
 namespace ERigidBodyFlags
 {
 	enum Type
@@ -1169,6 +1163,7 @@ namespace ERigidBodyFlags
 	};
 }
 
+/** Describes type of wake/sleep event sent to the physics system */
 enum class ESleepEvent : uint8
 {
 	SET_Wakeup,
@@ -1179,7 +1174,7 @@ enum class ESleepEvent : uint8
 USTRUCT()
 struct FRigidBodyErrorCorrection
 {
-	GENERATED_USTRUCT_BODY()
+	GENERATED_BODY()
 
 	/** Value between 0 and 1 which indicates how much velocity
 		and ping based correction to use */
@@ -1260,24 +1255,27 @@ struct FRigidBodyErrorCorrection
 	{ }
 };
 
-
 /**
  * Information about one contact between a pair of rigid bodies.
  */
 USTRUCT()
 struct FRigidBodyContactInfo
 {
-	GENERATED_USTRUCT_BODY()
+	GENERATED_BODY()
 
+	/** Position of contact, where two shapes intersect */
 	UPROPERTY()
 	FVector ContactPosition;
 
+	/** Normal of contact, points from second shape towards first shape */
 	UPROPERTY()
 	FVector ContactNormal;
 
+	/** How far the two shapes penetrated into each other */
 	UPROPERTY()
 	float ContactPenetration;
 
+	/** The physical material of the two shapes involved in a contact */
 	UPROPERTY()
 	class UPhysicalMaterial* PhysMaterial[2];
 
@@ -1317,17 +1315,17 @@ struct FRigidBodyContactInfo
 USTRUCT()
 struct FCollisionImpactData
 {
-	GENERATED_USTRUCT_BODY()
+	GENERATED_BODY()
 
-	/** all the contact points in the collision*/
+	/** All the contact points in the collision*/
 	UPROPERTY()
 	TArray<struct FRigidBodyContactInfo> ContactInfos;
 
-	/** the total impulse applied as the two objects push against each other*/
+	/** The total impulse applied as the two objects push against each other*/
 	UPROPERTY()
 	FVector TotalNormalImpulse;
 
-	/** the total counterimpulse applied of the two objects sliding against each other*/
+	/** The total counterimpulse applied of the two objects sliding against each other*/
 	UPROPERTY()
 	FVector TotalFrictionImpulse;
 
@@ -1340,12 +1338,11 @@ struct FCollisionImpactData
 	void SwapContactOrders();
 };
 
-
 /** Struct used to hold effects for destructible damage events */
 USTRUCT(BlueprintType)
 struct FFractureEffect
 {
-	GENERATED_USTRUCT_BODY()
+	GENERATED_BODY()
 
 	/** Particle system effect to play at fracture location. */
 	UPROPERTY(EditAnywhere, BlueprintReadWrite, Category=FractureEffect)
@@ -1361,16 +1358,17 @@ struct FFractureEffect
 	{ }
 };
 
-
 /**	Struct for handling positions relative to a base actor, which is potentially moving */
 USTRUCT(BlueprintType)
 struct ENGINE_API FBasedPosition
 {
-	GENERATED_USTRUCT_BODY()
+	GENERATED_BODY()
 
+	/** Actor that is the base */
 	UPROPERTY(EditAnywhere, BlueprintReadWrite, Category=BasedPosition)
 	class AActor* Base;
 
+	/** Position relative to the base actor */
 	UPROPERTY(EditAnywhere, BlueprintReadWrite, Category=BasedPosition)
 	FVector Position;
 
@@ -1385,15 +1383,19 @@ struct ENGINE_API FBasedPosition
 
 	FBasedPosition();
 	explicit FBasedPosition( class AActor *InBase, const FVector& InPosition );
-	// Retrieve world location of this position
+
+	/** Retrieve world location of this position */
 	FVector operator*() const;
+
+	/** Updates base/position */
 	void Set( class AActor* InBase, const FVector& InPosition );
+
+	/** Clear base/position */
 	void Clear();
 
 	friend FArchive& operator<<( FArchive& Ar, FBasedPosition& T );
 };
 
-
 /** Struct for caching Quat<->Rotator conversions. */
 struct ENGINE_API FRotationConversionCache
 {
@@ -1483,13 +1485,11 @@ private:
 	mutable FRotator	CachedRotator;	// FRotator matching CachedQuat such that CachedRotator.Quaternion() == CachedQuat.
 };
 
-
-
 /** A line of subtitle text and the time at which it should be displayed. */
 USTRUCT()
 struct FSubtitleCue
 {
-	GENERATED_USTRUCT_BODY()
+	GENERATED_BODY()
 
 	/** The text to appear in the subtitle. */
 	UPROPERTY(EditAnywhere, Category=SubtitleCue)
@@ -1504,12 +1504,11 @@ struct FSubtitleCue
 	{ }
 };
 
-
 /**	Per-light settings for Lightmass */
 USTRUCT()
 struct FLightmassLightSettings
 {
-	GENERATED_USTRUCT_BODY()
+	GENERATED_BODY()
 
 	/** 0 will be completely desaturated, 1 will be unchanged */
 	UPROPERTY(EditAnywhere, Category=Lightmass, meta=(UIMin = "0.0", UIMax = "4.0"))
@@ -1533,20 +1532,18 @@ struct FLightmassLightSettings
 	{ }
 };
 
-
 /**	Point/spot settings for Lightmass */
 USTRUCT()
 struct FLightmassPointLightSettings : public FLightmassLightSettings
 {
-	GENERATED_USTRUCT_BODY()
+	GENERATED_BODY()
 };
 
-
 /**	Directional light settings for Lightmass */
 USTRUCT()
 struct FLightmassDirectionalLightSettings : public FLightmassLightSettings
 {
-	GENERATED_USTRUCT_BODY()
+	GENERATED_BODY()
 
 	/** Angle that the directional light's emissive surface extends relative to a receiver, affects penumbra sizes. */
 	UPROPERTY(EditAnywhere, Category=Lightmass, meta=(UIMin = ".0001", UIMax = "5"))
@@ -1558,12 +1555,11 @@ struct FLightmassDirectionalLightSettings : public FLightmassLightSettings
 	}
 };
 
-
 /**	Per-object settings for Lightmass */
 USTRUCT()
 struct FLightmassPrimitiveSettings
 {
-	GENERATED_USTRUCT_BODY()
+	GENERATED_BODY()
 
 	/** If true, this object will be lit as if it receives light from both sides of its polygons. */
 	UPROPERTY(EditAnywhere, Category=Lightmass)
@@ -1643,12 +1639,11 @@ struct FLightmassPrimitiveSettings
 	friend FArchive& operator<<(FArchive& Ar, FLightmassPrimitiveSettings& Settings);
 };
 
-
 /**	Debug options for Lightmass */
 USTRUCT()
 struct FLightmassDebugOptions
 {
-	GENERATED_USTRUCT_BODY()
+	GENERATED_BODY()
 
 	/**
 	 *	If false, UnrealLightmass.exe is launched automatically (default)
@@ -1738,14 +1733,13 @@ struct FLightmassDebugOptions
 	ENGINE_API FLightmassDebugOptions();
 };
 
-
 /**
  *	Debug options for Swarm
  */
 USTRUCT()
 struct FSwarmDebugOptions
 {
-	GENERATED_USTRUCT_BODY()
+	GENERATED_BODY()
 
 	/**
 	 *	If true, Swarm will distribute jobs.
@@ -1776,7 +1770,7 @@ struct FSwarmDebugOptions
 	void Touch();
 };
 
-
+/** Method for padding a light map in memory */
 UENUM()
 enum ELightMapPaddingType
 {
@@ -1785,7 +1779,6 @@ enum ELightMapPaddingType
 	LMPT_NoPadding
 };
 
-
 /** Bit-field flags that affects storage (e.g. packing, streaming) and other info about a shadowmap. */
 UENUM()
 enum EShadowMapFlags
@@ -1796,19 +1789,21 @@ enum EShadowMapFlags
 	SMF_Streamed		= 0x00000001
 };
 
-
-/** Reference to a specific material in a PrimitiveComponent. */
+/** Reference to a specific material in a PrimitiveComponent, used by Matinee */
 USTRUCT()
 struct FPrimitiveMaterialRef
 {
-	GENERATED_USTRUCT_BODY()
+	GENERATED_BODY()
 
+	/** Material is on a primitive component */
 	UPROPERTY()
 	class UPrimitiveComponent* Primitive;
 
+	/** Material is on a decal component */
 	UPROPERTY()
 	class UDecalComponent* Decal;
 
+	/** Index into the material on the components data */
 	UPROPERTY()
 	int32 ElementIndex;
 
@@ -1831,14 +1826,13 @@ struct FPrimitiveMaterialRef
 	{ 	}
 };
 
-
 /**
  * Structure containing information about one hit of a trace, such as point of impact and surface normal at that point.
  */
 USTRUCT(BlueprintType, meta = (HasNativeBreak = "Engine.GameplayStatics.BreakHitResult", HasNativeMake = "Engine.GameplayStatics.MakeHitResult"))
 struct ENGINE_API FHitResult
 {
-	GENERATED_USTRUCT_BODY()
+	GENERATED_BODY()
 
 	/** Indicates if this hit was a result of blocking collision. If false, there was no hit or it was an overlap/touch instead. */
 	UPROPERTY()
@@ -2112,7 +2106,7 @@ FORCEINLINE bool TeleportEnumToFlag(ETeleportType Teleport) { return ETeleportTy
 USTRUCT()
 struct ENGINE_API FOverlapResult
 {
-	GENERATED_USTRUCT_BODY()
+	GENERATED_BODY()
 
 	/** Actor that the check hit. */
 	UPROPERTY()
@@ -2150,7 +2144,7 @@ template<> struct TIsPODType<FOverlapResult> { enum { Value = true }; };
 USTRUCT()
 struct ENGINE_API FMTDResult
 {
-	GENERATED_USTRUCT_BODY()
+	GENERATED_BODY()
 
 	/** Normalized direction of the minimum translation required to fix penetration. */
 	UPROPERTY()
@@ -2166,12 +2160,11 @@ struct ENGINE_API FMTDResult
 	}
 };
 
-
 /** Struct used for passing information from Matinee to an Actor for blending animations during a sequence. */
 USTRUCT()
 struct FAnimSlotInfo
 {
-	GENERATED_USTRUCT_BODY()
+	GENERATED_BODY()
 
 	/** Name of slot that we want to play the animtion in. */
 	UPROPERTY()
@@ -2182,12 +2175,11 @@ struct FAnimSlotInfo
 	TArray<float> ChannelWeights;
 };
 
-
 /** Used to indicate each slot name and how many channels they have. */
 USTRUCT()
 struct FAnimSlotDesc
 {
-	GENERATED_USTRUCT_BODY()
+	GENERATED_BODY()
 
 	/** Name of the slot. */
 	UPROPERTY()
@@ -2197,11 +2189,9 @@ struct FAnimSlotDesc
 	UPROPERTY()
 	int32 NumChannels;
 
-
 	FAnimSlotDesc()
 		: NumChannels(0)
 	{ }
-
 };
 
 /** Enum for controlling buckets for update rate optimizations if we need to stagger
@@ -2224,7 +2214,7 @@ enum class EUpdateRateShiftBucket : uint8
 USTRUCT()
 struct FAnimUpdateRateParameters
 {
-	GENERATED_USTRUCT_BODY()
+	GENERATED_BODY()
 
 public:
 	enum EOptimizeMode : uint8
@@ -2269,12 +2259,12 @@ public:
 	UPROPERTY()
 	int32 EvaluationRate;
 
-	UPROPERTY(Transient)
 	/** Track time we have lost via skipping */
+	UPROPERTY(Transient)
 	float TickedPoseOffestTime;
 
-	UPROPERTY(Transient)
 	/** Total time of the last series of skipped updates */
+	UPROPERTY(Transient)
 	float AdditionalTime;
 
 	/** The delta time of the last tick */
@@ -2311,9 +2301,11 @@ public:
 	UPROPERTY()
 	TMap<int32, int32> LODToFrameSkipMap;
 
+	/** Number of update frames that have been skipped in a row */
 	UPROPERTY()
 	int32 SkippedUpdateFrames;
 
+	/** Number of evaluate frames that have been skipped in a row */
 	UPROPERTY()
 	int32 SkippedEvalFrames;
 
@@ -2342,38 +2334,47 @@ public:
 		BaseVisibleDistanceFactorThesholds.Add(0.12f);
 	}
 
-	/** Set parameters and verify inputs for Trail Mode (original behaviour - skip frames, track skipped time and then catch up afterwards).
-	 * @param : UpdateShiftRate. Shift our update frames so that updates across all skinned components are staggered
+	/** 
+	 * Set parameters and verify inputs for Trail Mode (original behaviour - skip frames, track skipped time and then catch up afterwards).
+	 * @param : UpdateRateShift. Shift our update frames so that updates across all skinned components are staggered
 	 * @param : NewUpdateRate. How often animation will be updated/ticked. 1 = every frame, 2 = every 2 frames, etc.
 	 * @param : NewEvaluationRate. How often animation will be evaluated. 1 = every frame, 2 = every 2 frames, etc.
 	 * @param : bNewInterpSkippedFrames. When skipping a frame, should it be interpolated or frozen?
 	 */
 	void SetTrailMode(float DeltaTime, uint8 UpdateRateShift, int32 NewUpdateRate, int32 NewEvaluationRate, bool bNewInterpSkippedFrames);
 
+	/**
+	 * Set parameters and verify inputs for Lookahead mode, which handles Root Motion
+	 * @param : UpdateRateShift. Shift our update frames so that updates across all skinned components are staggered
+	 * @param : LookAheadAmount. Amount of time to look ahead and predict movement 
+	 */
 	void SetLookAheadMode(float DeltaTime, uint8 UpdateRateShift, float LookAheadAmount);
 
+	/** Amount to interpolate bone transforms */
 	float GetInterpolationAlpha() const;
 
+	/** Amount to interpoilate root motion */
 	float GetRootMotionInterp() const;
 
+	/** Return true if evaluation rate should be optimized at all */
 	bool DoEvaluationRateOptimizations() const
 	{
 		return OptimizeMode == LookAheadMode || EvaluationRate > 1;
 	}
 
-	/* Getter for bSkipUpdate */
+	/** Getter for bSkipUpdate */
 	bool ShouldSkipUpdate() const
 	{
 		return bSkipUpdate;
 	}
 
-	/* Getter for bSkipEvaluation */
+	/** Getter for bSkipEvaluation */
 	bool ShouldSkipEvaluation() const
 	{
 		return bSkipEvaluation;
 	}
 
-	/* Getter for bInterpolateSkippedFrames */
+	/** Getter for bInterpolateSkippedFrames */
 	bool ShouldInterpolateSkippedFrames() const
 	{
 		return bInterpolateSkippedFrames;
@@ -2385,6 +2386,7 @@ public:
 		return AdditionalTime;
 	}
 
+	/** Returns color to use for debug UI */
 	FColor GetUpdateRateDebugColor() const
 	{
 		if (OptimizeMode == TrailMode)
@@ -2408,14 +2410,11 @@ public:
 	}
 };
 
-
-/**
- * Point Of View type.
- */
+/** Point Of View structure used in Camera calculations */
 USTRUCT(BlueprintType)
 struct FPOV
 {
-	GENERATED_USTRUCT_BODY()
+	GENERATED_BODY()
 
 	/** Location */
 	UPROPERTY(EditAnywhere, BlueprintReadWrite, Category=POV)
@@ -2444,15 +2443,13 @@ struct FPOV
 	}
 };
 
-
-
 /**
  * Settings applied when building a mesh.
  */
 USTRUCT()
 struct FMeshBuildSettings
 {
-	GENERATED_USTRUCT_BODY()
+	GENERATED_BODY()
 
 	/** If true, degenerate triangles will be removed. */
 	UPROPERTY(EditAnywhere, Category=BuildSettings)
@@ -2579,15 +2576,13 @@ struct FMeshBuildSettings
 	}
 };
 
-
-
+/** Event used by AActor::TakeDamage and related functions */
 USTRUCT(BlueprintType)
 struct ENGINE_API FDamageEvent
 {
-	GENERATED_USTRUCT_BODY()
+	GENERATED_BODY()
 
 public:
-
 	/** Default constructor (no initialization). */
 	FDamageEvent() { }
 
@@ -2615,13 +2610,13 @@ public:
 	virtual void GetBestHitInfo(AActor const* HitActor, AActor const* HitInstigator, struct FHitResult& OutHitInfo, FVector& OutImpulseDir) const;
 };
 
-
+/** Damage subclass that handles damage with a single impact location and source direction */
 USTRUCT()
 struct ENGINE_API FPointDamageEvent : public FDamageEvent
 {
-	GENERATED_USTRUCT_BODY()
+	GENERATED_BODY()
 
-	// skipping ImpulseMag for now
+	/** Actual damage done */
 	UPROPERTY()
 	float Damage;
 	
@@ -2629,6 +2624,7 @@ struct ENGINE_API FPointDamageEvent : public FDamageEvent
 	UPROPERTY()
 	FVector_NetQuantizeNormal ShotDirection;
 	
+	/** Describes the trace/location that caused this damage */
 	UPROPERTY()
 	struct FHitResult HitInfo;
 
@@ -2647,30 +2643,32 @@ struct ENGINE_API FPointDamageEvent : public FDamageEvent
 	virtual void GetBestHitInfo(AActor const* HitActor, AActor const* HitInstigator, struct FHitResult& OutHitInfo, FVector& OutImpulseDir) const override;
 };
 
-
+/** Parameters used to compute radial damage */
 USTRUCT(BlueprintType)
 struct ENGINE_API FRadialDamageParams
 {
-	GENERATED_USTRUCT_BODY()
+	GENERATED_BODY()
 
+	/** Max damage done */
 	UPROPERTY(EditAnywhere, BlueprintReadWrite, Category=RadialDamageParams)
 	float BaseDamage;
 
+	/** Damage will not fall below this if within range */
 	UPROPERTY(EditAnywhere, BlueprintReadWrite, Category=RadialDamageParams)
 	float MinimumDamage;
 	
+	/** Within InnerRadius, do max damage */
 	UPROPERTY(EditAnywhere, BlueprintReadWrite, Category=RadialDamageParams)
 	float InnerRadius;
 		
+	/** Outside OuterRadius, do no damage */
 	UPROPERTY(EditAnywhere, BlueprintReadWrite, Category=RadialDamageParams)
 	float OuterRadius;
 		
+	/** Describes amount of exponential damage falloff */
 	UPROPERTY(EditAnywhere, BlueprintReadWrite, Category=RadialDamageParams)
 	float DamageFalloff;
 
-// 	UPROPERTY(EditAnywhere, BlueprintReadWrite, Category=RadiusDamageParams)
-// 	float BaseImpulseMag;
-
 	FRadialDamageParams()
 		: BaseDamage(0.f), MinimumDamage(0.f), InnerRadius(0.f), OuterRadius(0.f), DamageFalloff(1.f)
 	{}
@@ -2684,25 +2682,28 @@ struct ENGINE_API FRadialDamageParams
 		: BaseDamage(InBaseDamage), MinimumDamage(0.f), InnerRadius(0.f), OuterRadius(InRadius), DamageFalloff(1.f)
 	{}
 
+	/** Returns damage done at a certain distance */
 	float GetDamageScale(float DistanceFromEpicenter) const;
 
 	/** Return outermost radius of the damage area. Protects against malformed data. */
 	float GetMaxRadius() const { return FMath::Max( FMath::Max(InnerRadius, OuterRadius), 0.f ); }
 };
 
-
+/** Damage subclass that handles damage with a source location and falloff radius */
 USTRUCT()
 struct ENGINE_API FRadialDamageEvent : public FDamageEvent
 {
-	GENERATED_USTRUCT_BODY()
+	GENERATED_BODY()
 
+	/** Static parameters describing damage falloff math */
 	UPROPERTY()
 	FRadialDamageParams Params;
 	
+	/** Location of origin point */
 	UPROPERTY()
 	FVector Origin;
 
-	// @fixme, will not replicate properly?  component pointer
+	/** Hit reslts of specific impacts */
 	UPROPERTY()
 	TArray<struct FHitResult> ComponentHits;
 
@@ -2720,7 +2721,7 @@ struct ENGINE_API FRadialDamageEvent : public FDamageEvent
 	{}
 };
 
-
+/** The network role of an actor on a local/remote network context */
 UENUM()
 enum ENetRole
 {
@@ -2735,7 +2736,7 @@ enum ENetRole
 	ROLE_MAX,
 };
 
-
+/** Describes if an actor can enter a low network bandwidth dormant mode */
 UENUM()
 enum ENetDormancy
 {
@@ -2753,7 +2754,7 @@ enum ENetDormancy
 	DORM_MAX UMETA(Hidden),
 };
 
-
+/** Specifies which player index will pass input to this actor/component */
 UENUM()
 namespace EAutoReceiveInput
 {
@@ -2771,7 +2772,7 @@ namespace EAutoReceiveInput
 	};
 }
 
-
+/** Specifies if an AI pawn will automatically be possed by an AI controller */
 UENUM()
 enum class EAutoPossessAI : uint8
 {
@@ -2785,7 +2786,7 @@ enum class EAutoPossessAI : uint8
 	PlacedInWorldOrSpawned,
 };
 
-
+/** Specifies why an actor is being deleted/removed from a level */
 UENUM(BlueprintType)
 namespace EEndPlayReason
 {
@@ -2802,12 +2803,11 @@ namespace EEndPlayReason
 		/** When the application is being exited. */
 		Quit,
 	};
-
 }
 
 DECLARE_DYNAMIC_DELEGATE(FTimerDynamicDelegate);
 
-// Unique handle that can be used to distinguish timers that have identical delegates.
+/** Unique handle that can be used to distinguish timers that have identical delegates. */
 USTRUCT(BlueprintType)
 struct FTimerHandle
 {
@@ -2821,11 +2821,13 @@ struct FTimerHandle
 	{
 	}
 
+	/** True if this handle was ever initialized by the timer manager */
 	bool IsValid() const
 	{
 		return Handle != 0;
 	}
 
+	/** Explicitly clear handle */
 	void Invalidate()
 	{
 		Handle = 0;
@@ -2881,6 +2883,7 @@ private:
 	}
 };
 
+/** Describes rules for network replicating a vector efficiently */
 UENUM()
 enum class EVectorQuantization : uint8
 {
@@ -2892,6 +2895,7 @@ enum class EVectorQuantization : uint8
 	RoundTwoDecimals
 };
 
+/** Describes rules for network replicating a vector efficiently */
 UENUM()
 enum class ERotatorQuantization : uint8
 {
@@ -2908,17 +2912,21 @@ enum class ERotatorQuantization : uint8
 USTRUCT()
 struct ENGINE_API FRepMovement
 {
-	GENERATED_USTRUCT_BODY()
+	GENERATED_BODY()
 
+	/** Velocity of component in world space */
 	UPROPERTY(Transient)
 	FVector LinearVelocity;
 
+	/** Velocity of rotation for component */
 	UPROPERTY(Transient)
 	FVector AngularVelocity;
 	
+	/** Location in world space */
 	UPROPERTY(Transient)
 	FVector Location;
 
+	/** Current rotation */
 	UPROPERTY(Transient)
 	FRotator Rotation;
 
@@ -3069,6 +3077,7 @@ struct ENGINE_API FRepMovement
 		return !(*this == Other);
 	}
 
+	/** True if multiplayer rebasing is enabled, corresponds to p.EnableMultiplayerWorldOriginRebasing console variable */
 	static int32 EnableMultiplayerWorldOriginRebasing;
 
 	/** Rebase zero-origin position onto local world origin value. */
@@ -3091,27 +3100,33 @@ struct ENGINE_API FRepMovement
 };
 
 
-/** Handles attachment replication to clients. Movement replication will not happen while AttachParent is non-nullptr */
+/** Handles attachment replication to clients.  */
 USTRUCT()
 struct FRepAttachment
 {
-	GENERATED_USTRUCT_BODY()
+	GENERATED_BODY()
 
+	/** Actor we are attached to, movement replication will not happen while AttachParent is non-nullptr */
 	UPROPERTY()
 	class AActor* AttachParent;
 
+	/** Location offset from attach parent */
 	UPROPERTY()
 	FVector_NetQuantize100 LocationOffset;
 
+	/** Scale relative to attach parent */
 	UPROPERTY()
 	FVector_NetQuantize100 RelativeScale3D;
 
+	/** Rotation offset from attach parent */
 	UPROPERTY()
 	FRotator RotationOffset;
 
+	/** Specific socket we are attached to */
 	UPROPERTY()
 	FName AttachSocket;
 
+	/** Specific component we are attached to */
 	UPROPERTY()
 	class USceneComponent* AttachComponent;
 
@@ -3125,7 +3140,6 @@ struct FRepAttachment
 	{ }
 };
 
-
 /**
  * Controls behavior of WalkableSlopeOverride, determining how to affect walkability of surfaces for Characters.
  * @see FWalkableSlopeOverride
@@ -3160,12 +3174,11 @@ enum EWalkableSlopeBehavior
 	WalkableSlope_Max		UMETA(Hidden),
 };
 
-
 /** Struct allowing control over "walkable" normals, by allowing a restriction or relaxation of what steepness is normally walkable. */
 USTRUCT(BlueprintType)
 struct FWalkableSlopeOverride
 {
-	GENERATED_USTRUCT_BODY()
+	GENERATED_BODY()
 
 	/**
 	 * Behavior of this surface (whether we affect the walkable slope).
@@ -3205,31 +3218,31 @@ public:
 	{
 	}
 
-	// Gets the slope override behavior.
+	/** Gets the slope override behavior. */
 	FORCEINLINE EWalkableSlopeBehavior GetWalkableSlopeBehavior() const
 	{
 		return WalkableSlopeBehavior;
 	}
 
-	// Gets the slope angle used for the override behavior.
+	/** Gets the slope angle used for the override behavior. */
 	FORCEINLINE float GetWalkableSlopeAngle() const
 	{
 		return WalkableSlopeAngle;
 	}
 
-	// Set the slope override behavior.
+	/** Set the slope override behavior. */
 	FORCEINLINE void SetWalkableSlopeBehavior(EWalkableSlopeBehavior NewSlopeBehavior)
 	{
 		WalkableSlopeBehavior = NewSlopeBehavior;
 	}
 
-	// Set the slope angle used for the override behavior.
+	/** Set the slope angle used for the override behavior. */
 	FORCEINLINE void SetWalkableSlopeAngle(float NewSlopeAngle)
 	{
 		WalkableSlopeAngle = FMath::Clamp(NewSlopeAngle, 0.f, 90.f);
 	}
 
-	// Given a walkable floor normal Z value, either relax or restrict the value if we override such behavior.
+	/** Given a walkable floor normal Z value, either relax or restrict the value if we override such behavior. */
 	float ModifyWalkableFloorZ(float InWalkableFloorZ) const
 	{
 		switch(WalkableSlopeBehavior)
@@ -3289,7 +3302,6 @@ struct TStructOpsTypeTraits<FRepMovement> : public TStructOpsTypeTraitsBase2<FRe
 	};
 };
 
-
 /** Structure to hold and pass around transient flags used during replication. */
 struct FReplicationFlags
 {
@@ -3321,21 +3333,19 @@ struct FReplicationFlags
 	}
 };
 
-
 static_assert(sizeof(FReplicationFlags) == 4, "FReplicationFlags has invalid size.");
 
 /** Struct used to specify the property name of the component to constrain */
 USTRUCT()
 struct FConstrainComponentPropName
 {
-	GENERATED_USTRUCT_BODY()
+	GENERATED_BODY()
 
 	/** Name of property */
 	UPROPERTY(EditAnywhere, Category=Constraint)
 	FName	ComponentName;
 };
 
-
 /** 
  *	Struct that allows for different ways to reference a component. 
  *	If just an Actor is specified, will return RootComponent of that Actor.
@@ -3343,7 +3353,7 @@ struct FConstrainComponentPropName
 USTRUCT()
 struct ENGINE_API FComponentReference
 {
-	GENERATED_USTRUCT_BODY()
+	GENERATED_BODY()
 
 	/** Pointer to a different Actor that owns the Component.  */
 	UPROPERTY(EditInstanceOnly, Category=Component)
@@ -3360,8 +3370,7 @@ struct ENGINE_API FComponentReference
 	class USceneComponent* GetComponent(AActor* OwningActor) const;
 };
 
-
-/** Types of surfaces in the game. */
+/** Types of surfaces in the game, used by Physical Materials */
 UENUM(BlueprintType)
 enum EPhysicalSurface
 {
@@ -3431,7 +3440,6 @@ enum EPhysicalSurface
 	SurfaceType_Max UMETA(Hidden)
 };
 
-
 /** Describes how often this component is allowed to move. */
 UENUM(BlueprintType)
 namespace EComponentMobility
@@ -3464,7 +3472,7 @@ namespace EComponentMobility
 	};
 }
 
-
+/** Utility class for engine types */
 UCLASS(abstract, config=Engine)
 class ENGINE_API UEngineTypes : public UObject
 {
@@ -3484,7 +3492,6 @@ public:
 	static ETraceTypeQuery ConvertToTraceType(ECollisionChannel CollisionChannel);
 };
 
-
 /** Type of a socket on a scene component. */
 UENUM()
 namespace EComponentSocketType
@@ -3502,7 +3509,6 @@ namespace EComponentSocketType
 	};
 }
 
-
 /** Info about a socket on a scene component */
 struct FComponentSocketDescription
 {
@@ -3528,7 +3534,7 @@ struct FComponentSocketDescription
 /** Dynamic delegate to use by components that want to route the broken-event into blueprints */
 DECLARE_DYNAMIC_MULTICAST_DELEGATE_OneParam(FConstraintBrokenSignature, int32, ConstraintIndex);
 
-// ANGULAR DOF
+/** Specifies angular degrees of freedom */
 UENUM()
 enum EAngularConstraintMotion
 {
@@ -3542,7 +3548,7 @@ enum EAngularConstraintMotion
 	ACM_MAX,
 };
 
-/** Enum to indicate which frame we want. */
+/** Enum to indicate which context frame we use for physical constraints */
 UENUM()
 namespace EConstraintFrame
 {
@@ -3553,14 +3559,11 @@ namespace EConstraintFrame
 	};
 }
 
-
-/**
- * Structure for file paths that are displayed in the UI.
- */
+/** Structure for file paths that are displayed in the editor with a picker UI. */
 USTRUCT(BlueprintType)
 struct FFilePath
 {
-	GENERATED_USTRUCT_BODY()
+	GENERATED_BODY()
 
 	/**
 	 * The path to the file.
@@ -3569,14 +3572,11 @@ struct FFilePath
 	FString FilePath;
 };
 
-
-/**
- * Structure for directory paths that are displayed in the UI.
- */
+/** Structure for directory paths that are displayed in the editor with a picker UI. */
 USTRUCT(BlueprintType)
 struct FDirectoryPath
 {
-	GENERATED_USTRUCT_BODY()
+	GENERATED_BODY()
 
 	/**
 	 * The path to the directory.
@@ -3592,7 +3592,7 @@ struct FDirectoryPath
 USTRUCT(BlueprintType)
 struct FCollectionReference
 {
-	GENERATED_USTRUCT_BODY()
+	GENERATED_BODY()
 
 	/**
 	 * Name of the collection
@@ -3602,14 +3602,13 @@ struct FCollectionReference
 };
 
 /**
-* This is used for redirecting old name to new name
-* We use manually parsing array, but that makes harder to modify from property setting
-* So adding this USTRUCT to support it properly
-*/
+ * This is used for redirecting an old name to new name, such as for collision profiles
+ * This is used for better UI in the editor
+ */
 USTRUCT()
 struct ENGINE_API FRedirector
 {
-	GENERATED_USTRUCT_BODY()
+	GENERATED_BODY()
 
 	UPROPERTY()
 	FName OldName;
@@ -3629,14 +3628,13 @@ struct ENGINE_API FRedirector
 	{ }
 };
 
-
 /** 
  * Structure for recording float values and displaying them as an Histogram through DrawDebugFloatHistory.
  */
 USTRUCT(BlueprintType)
 struct FDebugFloatHistory
 {
-	GENERATED_USTRUCT_BODY()
+	GENERATED_BODY()
 
 private:
 	/** Samples */
@@ -3734,28 +3732,29 @@ public:
 	}
 };
 
-
-/** info for glow when using depth field rendering */
+/** Info for glow when using depth field rendering */
 USTRUCT(BlueprintType)
 struct FDepthFieldGlowInfo
 {
-	GENERATED_USTRUCT_BODY()
+	GENERATED_BODY()
 
-	/** whether to turn on the outline glow (depth field fonts only) */
+	/** Whether to turn on the outline glow (depth field fonts only) */
 	UPROPERTY(BlueprintReadWrite, Category="Glow")
 	uint32 bEnableGlow:1;
 
-	/** base color to use for the glow */
+	/** Base color to use for the glow */
 	UPROPERTY(BlueprintReadWrite, Category="Glow")
 	FLinearColor GlowColor;
 
-	/** if bEnableGlow, outline glow outer radius (0 to 1, 0.5 is edge of character silhouette)
+	/** 
+	 * If bEnableGlow, outline glow outer radius (0 to 1, 0.5 is edge of character silhouette)
 	 * glow influence will be 0 at GlowOuterRadius.X and 1 at GlowOuterRadius.Y
-	*/
+	 */
 	UPROPERTY(BlueprintReadWrite, Category="Glow")
 	FVector2D GlowOuterRadius;
 
-	/** if bEnableGlow, outline glow inner radius (0 to 1, 0.5 is edge of character silhouette)
+	/** 
+	 * If bEnableGlow, outline glow inner radius (0 to 1, 0.5 is edge of character silhouette)
 	 * glow influence will be 1 at GlowInnerRadius.X and 0 at GlowInnerRadius.Y
 	 */
 	UPROPERTY(BlueprintReadWrite, Category="Glow")
@@ -3792,22 +3791,21 @@ struct FDepthFieldGlowInfo
 	}	
 };
 
-
-/** information used in font rendering */
+/** Information used in font rendering */
 USTRUCT(BlueprintType)
 struct FFontRenderInfo
 {
-	GENERATED_USTRUCT_BODY()
+	GENERATED_BODY()
 
-	/** whether to clip text */
+	/** Whether to clip text */
 	UPROPERTY(BlueprintReadWrite, Category="FontInfo")
 	uint32 bClipText:1;
 
-	/** whether to turn on shadowing */
+	/** Whether to turn on shadowing */
 	UPROPERTY(BlueprintReadWrite, Category="FontInfo")
 	uint32 bEnableShadow:1;
 
-	/** depth field glow parameters (only usable if font was imported with a depth field) */
+	/** Depth field glow parameters (only usable if font was imported with a depth field) */
 	UPROPERTY(BlueprintReadWrite, Category="FontInfo")
 	struct FDepthFieldGlowInfo GlowInfo;
 
@@ -3816,12 +3814,11 @@ struct FFontRenderInfo
 	{}
 };
 
-
 /** Simple 2d triangle with UVs */
 USTRUCT(BlueprintType)
 struct FCanvasUVTri
 {
-	GENERATED_USTRUCT_BODY()
+	GENERATED_BODY()
 
 	/** Position of first vertex */
 	UPROPERTY(EditAnywhere, BlueprintReadWrite, Category = CanvasUVTri)
@@ -3872,7 +3869,6 @@ struct FCanvasUVTri
 	{ }
 };
 
-
 template <> struct TIsZeroConstructType<FCanvasUVTri> { enum { Value = true }; };
 
 /** Defines available strategies for handling the case where an actor is spawned in such a way that it penetrates blocking collision. */
@@ -3894,8 +3890,11 @@ enum class ESpawnActorCollisionHandlingMethod : uint8
 /** Defines the context of a user activity. Activities triggered in Blueprints will by type Game. Those created in code might choose to set another type. */
 enum class EUserActivityContext : uint8
 {
+	/** Event triggered from gameplay, such as from blueprints */
 	Game,
+	/** Event triggered from the editor UI */
 	Editor,
+	/** Event triggered from some other source */
 	Other
 };
 
@@ -3905,7 +3904,7 @@ enum class EUserActivityContext : uint8
 USTRUCT(BlueprintType)
 struct FUserActivity
 {
-	GENERATED_USTRUCT_BODY()
+	GENERATED_BODY()
 
 	/** Describes the user's activity */
 	UPROPERTY(BlueprintReadWrite, Category = "Activity")
@@ -3951,8 +3950,10 @@ enum class ELevelCollectionType : uint8
 	 * Will contain a world's persistent level and any streaming levels that contain dynamic or replicated gameplay actors.
 	 */
 	DynamicSourceLevels,
+
 	/** Gameplay relevant levels that have been duplicated from DynamicSourceLevels if requested by the game. */
 	DynamicDuplicatedLevels,
+
 	/**
 	 * These levels are shared between the source levels and the duplicated levels, and should contain
 	 * only static geometry and other visuals that are not replicated or affected by gameplay.
diff --git a/Engine/Source/Runtime/Engine/Classes/GameFramework/Actor.h b/Engine/Source/Runtime/Engine/Classes/GameFramework/Actor.h
--- a/Engine/Source/Runtime/Engine/Classes/GameFramework/Actor.h
+++ b/Engine/Source/Runtime/Engine/Classes/GameFramework/Actor.h
@@ -73,29 +73,48 @@ extern ENGINE_API FUObjectAnnotationSparseBool GSelectedActorAnnotation;
  * Actors may contain a collection of ActorComponents, which can be used to control how actors move, how they are rendered, etc.
  * The other main function of an Actor is the replication of properties and function calls across the network during play.
  * 
+ * Actor initialization has multiple steps, here's the order of important virtual functions that get called:
+ * - UObject::PostLoad: For actors statically placed in a level, the normal UObject PostLoad gets called both in the editor and during gameplay.
+ *                      This is not called for newly spawned actors.
+ * - UActorComponent::OnComponentCreated: When an actor is spawned in the editor or during gameplay, this gets called for any native components.
+ *                                        For blueprint-created components, this gets called during construction for that component.
+ *                                        This is not called for components loaded from a level.
+ * - AActor::PreRegisterAllComponents: For statically placed actors and spawned actors that have native root components, this gets called now.
+ *                                     For blueprint actors without a native root component, these registration functions get called later during construction.
+ * - UActorComponent::RegisterComponent: All components are registered in editor and at runtime, this creates their physical/visual representation.
+ *                                       These calls may be distributed over multiple frames, but are always after PreRegisterAllComponents.
+ *                                       This may also get called later on after an UnregisterComponent call removes it from the world.
+ * - AActor::PostRegisterAllComponents: Called for all actors both in the editor and in gameplay, this is the last function that is called in all cases.
+ * - AActor::PostActorCreated: When an actor is created in the editor or during gameplay, this gets called right before construction.
+ *                             This is not called for components loaded from a level.
+ * - AActor::UserConstructionScript: Called for blueprints that implement a construction script.
+ * - AActor::OnConstruction: Called at the end of ExecuteConstruction, which calls the blueprint construction script.
+ *                           This is called after all blueprint-created components are fully created and registered.
+ *                           This is only called during gameplay for spawned actors, and may get rerun in the editor when changing blueprints.
+ * - AActor::PreInitializeComponents: Called before InitializeComponent is called on the actor's components.
+ *                                    This is only called during gameplay and in certain editor preview windows.
+ * - UActorComponent::Activate: This will be called only if the component has bAutoActivate set.
+ *                              It will also got called later on if a component is manually activated.
+ * - UActorComponent::InitializeComponent: This will be called only if the component has bWantsInitializeComponentSet.
+ *                                         This only happens once per gameplay session.
+ * - AActor::PostInitializeComponents: Called after the actor's components have been initialized, only during gameplay and some editor previews.
+ * - AActor::BeginPlay: Called when the level starts ticking, only during actual gameplay.
+ *                      This normally happens right after PostInitializeComponents but can be delayed for networked or child actors.
+ *
  * @see https://docs.unrealengine.com/latest/INT/Programming/UnrealArchitecture/Actors/
+ * @see https://docs.unrealengine.com/en-us/Programming/UnrealArchitecture/Actors/ActorLifecycle
  * @see UActorComponent
  */
 UCLASS(BlueprintType, Blueprintable, config=Engine, meta=(ShortTooltip="An Actor is an object that can be placed or spawned in the world."))
 class ENGINE_API AActor : public UObject
 {
-	/**
-	 * The functions of interest to initialization order for an Actor is roughly as follows:
-	 * PostLoad/PostActorCreated - Do any setup of the actor required for construction. PostLoad for serialized actors, PostActorCreated for spawned.  
-	 * AActor::OnConstruction - The construction of the actor, this is where Blueprint actors have their components created and blueprint variables are initialized
-	 * AActor::PreInitializeComponents - Called before InitializeComponent is called on the actor's components
-	 * UActorComponent::InitializeComponent - Each component in the actor's components array gets an initialize call (if bWantsInitializeComponent is true for that component)
-	 * AActor::PostInitializeComponents - Called after the actor's components have been initialized
-	 * AActor::BeginPlay - Called when the level is started
-	 */
-
 	GENERATED_BODY()
-public:
 
+public:
 	/** Default constructor for AActor */
 	AActor();
 
-	/** Constructor for AActor that takes an ObjectInitializer */
+	/** Constructor for AActor that takes an ObjectInitializer for backward compatibility */
 	AActor(const FObjectInitializer& ObjectInitializer);
 
 private:
@@ -103,7 +122,7 @@ private:
 	void InitializeDefaults();
 
 public:
-	/** Returns the properties used for network replication */
+	/** Returns the properties used for network replication, this needs to be overridden by all actor classes with native replicated properties */
 	void GetLifetimeReplicatedProps(TArray<FLifetimeProperty>& OutLifetimeProps) const override;
 
 	/**
@@ -151,15 +170,11 @@ public:
 	UFUNCTION()
 	virtual void OnRep_ReplicateMovement();
 
-	/**
-	 * If true, this actor is no longer replicated to new clients, and is "torn off" (becomes a ROLE_Authority) on clients to which it was being replicated.
-	 * @see TornOff()
-	 */
 	UE_DEPRECATED(4.20, "Use GetTearOff() or TearOff() functions. This property will become private.")
 	UPROPERTY(Replicated)
 	uint8 bTearOff:1; 
 
-	/** returns TearOff status */
+	/** If true, this actor is no longer replicated to new clients, and is "torn off" (becomes a ROLE_Authority) on clients to which it was being replicated. */
 	PRAGMA_DISABLE_DEPRECATION_WARNINGS
 	bool GetTearOff() const { return bTearOff; }
 	PRAGMA_ENABLE_DEPRECATION_WARNINGS
@@ -189,7 +204,7 @@ public:
 
 	/** 
 	 * If true, this actor's component's bounds will be included in the level's
-	 * bounding box unless the Actor's class has overriden IsLevelBoundsRelevant 
+	 * bounding box unless the Actor's class has overridden IsLevelBoundsRelevant 
 	 */
 	UPROPERTY(EditAnywhere, Category=Collision, AdvancedDisplay)
 	uint8 bRelevantForLevelBounds:1;
@@ -230,7 +245,6 @@ public:
 	UPROPERTY()
 	uint8 bCollideWhenPlacing:1;
 
-
 	/** If true, this actor should search for an owned camera component to view through when used as a view target. */
 	UPROPERTY(EditAnywhere, BlueprintReadWrite, Category=Actor, AdvancedDisplay)
 	uint8 bFindCameraComponentWhenViewTarget:1;
@@ -324,13 +338,14 @@ private:
 	UPROPERTY()
 	uint8 bActorEnableCollision:1;
 
-	/** Set when actor is about to be deleted. Needs to be uproperty so included in transactions. */
+	/** Set when actor is about to be deleted. Needs to be a UProperty so it is included in transactions. */
 	UPROPERTY(Transient, DuplicateTransient)
 	uint8 bActorIsBeingDestroyed:1;
 
 	/** Set if an Actor tries to be destroyed while it is beginning play so that once BeginPlay ends we can issue the destroy call. */
 	uint8 bActorWantsDestroyDuringBeginPlay : 1;
 
+	/** Enum defining if BeginPlay has started or finished */
 	enum class EActorBeginPlayState : uint8
 	{
 		HasNotBegunPlay,
@@ -351,7 +366,6 @@ private:
 	TEnumAsByte<enum ENetRole> RemoteRole;
 
 public:
-
 	/**
 	 * Set whether this actor replicates to network clients. When this actor is spawned on the server it will be sent to clients as well.
 	 * Properties flagged for replication will update on clients if they change on the server.
@@ -451,11 +465,6 @@ public:
 	UPROPERTY(EditAnywhere, Category=Input)
 	TEnumAsByte<EAutoReceiveInput::Type> AutoReceiveInput;
 
-#if WITH_EDITORONLY_DATA
-	UPROPERTY()
-	TEnumAsByte<enum EInputConsumeOptions> InputConsumeOption_DEPRECATED;
-#endif
-
 	/** The priority of this input component when pushed in to the stack. */
 	UPROPERTY(EditAnywhere, Category=Input)
 	int32 InputPriority;
@@ -485,7 +494,6 @@ public:
 	float NetPriority;
 
 private:
-
 	/** Caches the most recent last render time we've looked at for this actor */
 	mutable float CachedLastRenderTime;
 
@@ -496,7 +504,7 @@ public:
 	 */
 	void SetNetDriverName(FName NewNetDriverName);
 
-	/** @return name of the net driver associated with this actor (all RPCs will go out via this connection) */
+	/** Returns name of the net driver associated with this actor (all RPCs will go out via this connection) */
 	FName GetNetDriverName() const { return NetDriverName; }
 
 	/** Method that allows an actor to replicate subobjects on its actor channel */
@@ -529,7 +537,7 @@ public:
 	/** Called by the networking system to call PreReplication on this actor and its components using the given NetDriver to find or create RepChangedPropertyTrackers. */
 	void CallPreReplication(UNetDriver* NetDriver);	
 	
-	/** Pawn responsible for damage caused by this actor. */
+	/** Pawn responsible for damage and other gameplay events caused by this actor. */
 	UPROPERTY(BlueprintReadWrite, ReplicatedUsing=OnRep_Instigator, meta=(ExposeOnSpawn=true), Category=Actor)
 	class APawn* Instigator;
 
@@ -537,17 +545,17 @@ public:
 	UFUNCTION()
 	virtual void OnRep_Instigator();
 
-	/** Array of Actors whose Owner is this actor */
+	/** Array of all Actors whose Owner is this actor, these are not necessarily spawned by UChildActorComponent */
 	UPROPERTY(Transient)
 	TArray<AActor*> Children;
 
 protected:
-	/** Collision primitive that defines the transform (location, rotation, scale) of this Actor. */
+	/** The component that defines the transform (location, rotation, scale) of this Actor in the world, all other components must be attached to this one somehow */
 	UPROPERTY(BlueprintGetter=K2_GetRootComponent, Category="Utilities|Transformation")
 	USceneComponent* RootComponent;
 
 #if WITH_EDITORONLY_DATA
-	/** Local space pivot offset for the actor */
+	/** Local space pivot offset for the actor, only used in the editor */
 	UPROPERTY(EditAnywhere, BlueprintReadOnly, AdvancedDisplay, Category=Actor)
 	FVector PivotOffset;
 #endif
@@ -572,7 +580,7 @@ public:
 
 private:
 #if WITH_EDITORONLY_DATA
-	/** The Actor that owns the UChildActorComponent that owns this Actor. */
+	/** @deprecated Use ParentComponent instead */
 	UPROPERTY()
 	TWeakObjectPtr<AActor> ParentComponentActor_DEPRECATED;	
 #endif
@@ -582,9 +590,8 @@ private:
 	TWeakObjectPtr<UChildActorComponent> ParentComponent;	
 
 #if WITH_EDITORONLY_DATA
-
 public:
-	/** The group this actor is a part of. */
+	/** The editor-only group this actor is a part of. */
 	UPROPERTY(Transient)
 	AActor* GroupActor;
 
@@ -653,7 +660,6 @@ private:
 	/** Whether this actor is temporarily hidden within the editor; used for show/hide/etc functionality w/o dirtying the actor. */
 	UPROPERTY(Transient)
 	uint8 bHiddenEdTemporary:1;
-
 #endif // WITH_EDITORONLY_DATA
 
 public:
@@ -764,7 +770,7 @@ public:
 
 	/**
 	 * Get the instigator, cast as a specific class.
-	 * @return The instigator for this weapon if it is the specified type, nullptr otherwise.
+	 * @return The instigator for this actor if it is the specified type, nullptr otherwise.
 	 */
 	template <class T>
 	T* GetInstigator() const { return Cast<T>(Instigator); };
@@ -831,6 +837,8 @@ public:
 	/**
 	 * Returns the bounding box of all components that make up this Actor (excluding ChildActorComponents).
 	 * @param	bOnlyCollidingComponents	If true, will only return the bounding box for components with collision enabled.
+	 * @param	Origin						Set to the center of the actor in world space
+	 * @param	BoxExtent					Set to half the actor's size in 3d space
 	 */
 	UFUNCTION(BlueprintCallable, Category="Collision", meta=(DisplayName = "GetActorBounds"))
 	void GetActorBounds(bool bOnlyCollidingComponents, FVector& Origin, FVector& BoxExtent) const;
@@ -962,19 +970,6 @@ public:
 	 */
 	UFUNCTION(BlueprintCallable, Category="Utilities|Transformation", meta=(DisplayName="AddActorWorldOffset", ScriptName="AddActorWorldOffset", Keywords="location position"))
 	void K2_AddActorWorldOffset(FVector DeltaLocation, bool bSweep, FHitResult& SweepHitResult, bool bTeleport);
-
-	/**
-	 * Adds a delta to the location of this actor in world space.
-	 * 
-	 * @param DeltaLocation		The change in location.
-	 * @param bSweep			Whether we sweep to the destination location, triggering overlaps along the way and stopping short of the target if blocked by something.
-	 *							Only the root component is swept and checked for blocking collision, child components move without sweeping. If collision is off, this has no effect.
-	 * @param Teleport			Whether we teleport the physics state (if physics collision is enabled for this object).
-	 *							If TeleportPhysics, physics velocity for this object is unchanged (so ragdoll parts are not affected by change in location).
-	 *							If None, physics velocity is updated based on the change in position (affecting ragdoll parts).
-	 *							If CCD is on and not teleporting, this will affect objects along the entire swept volume.
-	 * @param SweepHitResult	The hit result from the move if swept.
-	 */
 	void AddActorWorldOffset(FVector DeltaLocation, bool bSweep=false, FHitResult* OutSweepHitResult=nullptr, ETeleportType Teleport = ETeleportType::None);
 
 	/**
@@ -993,13 +988,11 @@ public:
 	void AddActorWorldRotation(FRotator DeltaRotation, bool bSweep=false, FHitResult* OutSweepHitResult=nullptr, ETeleportType Teleport = ETeleportType::None);
 	void AddActorWorldRotation(const FQuat& DeltaRotation, bool bSweep=false, FHitResult* OutSweepHitResult=nullptr, ETeleportType Teleport = ETeleportType::None);
 
-
 	/** Adds a delta to the transform of this actor in world space. Scale is unchanged. */
 	UFUNCTION(BlueprintCallable, Category="Utilities|Transformation", meta=(DisplayName="AddActorWorldTransform", ScriptName="AddActorWorldTransform"))
 	void K2_AddActorWorldTransform(const FTransform& DeltaTransform, bool bSweep, FHitResult& SweepHitResult, bool bTeleport);
 	void AddActorWorldTransform(const FTransform& DeltaTransform, bool bSweep=false, FHitResult* OutSweepHitResult=nullptr, ETeleportType Teleport = ETeleportType::None);
 
-
 	/** 
 	 * Set the Actors transform to the specified one.
 	 * @param NewTransform		The new transform.
@@ -1014,7 +1007,6 @@ public:
 	bool K2_SetActorTransform(const FTransform& NewTransform, bool bSweep, FHitResult& SweepHitResult, bool bTeleport);
 	bool SetActorTransform(const FTransform& NewTransform, bool bSweep=false, FHitResult* OutSweepHitResult=nullptr, ETeleportType Teleport = ETeleportType::None);
 
-
 	/** 
 	 * Adds a delta to the location of this component in its local reference frame.
 	 * @param DelatLocation		The change in location in local space.
@@ -1029,7 +1021,6 @@ public:
 	void K2_AddActorLocalOffset(FVector DeltaLocation, bool bSweep, FHitResult& SweepHitResult, bool bTeleport);
 	void AddActorLocalOffset(FVector DeltaLocation, bool bSweep=false, FHitResult* OutSweepHitResult=nullptr, ETeleportType Teleport = ETeleportType::None);
 
-
 	/**
 	 * Adds a delta to the rotation of this component in its local reference frame
 	 * @param DeltaRotation		The change in rotation in local space.
@@ -1045,7 +1036,6 @@ public:
 	void AddActorLocalRotation(FRotator DeltaRotation, bool bSweep=false, FHitResult* OutSweepHitResult=nullptr, ETeleportType Teleport = ETeleportType::None);
 	void AddActorLocalRotation(const FQuat& DeltaRotation, bool bSweep=false, FHitResult* OutSweepHitResult=nullptr, ETeleportType Teleport = ETeleportType::None);
 
-
 	/**
 	 * Adds a delta to the transform of this component in its local reference frame
 	 * @param NewTransform		The change in transform in local space.
@@ -1060,7 +1050,6 @@ public:
 	void K2_AddActorLocalTransform(const FTransform& NewTransform, bool bSweep, FHitResult& SweepHitResult, bool bTeleport);
 	void AddActorLocalTransform(const FTransform& NewTransform, bool bSweep=false, FHitResult* OutSweepHitResult=nullptr, ETeleportType Teleport = ETeleportType::None);
 
-
 	/**
 	 * Set the actor's RootComponent to the specified relative location.
 	 * @param NewRelativeLocation	New relative location of the actor's root component
@@ -1156,25 +1145,23 @@ public:
 	UFUNCTION(BlueprintCallable, meta=(ScriptNoExport, BlueprintInternalUseOnly = "true", DefaultToSelf="ComponentTemplateContext", InternalUseParam="ComponentTemplateContext"))
 	class UActorComponent* AddComponent(FName TemplateName, bool bManualAttachment, const FTransform& RelativeTransform, const UObject* ComponentTemplateContext);
 
-	/** DEPRECATED - Use Component::DestroyComponent */
-	UE_DEPRECATED(4.17, "Use Component.DestroyComponent instead")
+	UE_DEPRECATED(4.17, "Use UActorComponent::DestroyComponent() instead")
 	UFUNCTION(BlueprintCallable, meta=(DeprecatedFunction, DeprecationMessage = "Use Component.DestroyComponent instead", BlueprintProtected = "true", DisplayName = "DestroyComponent", ScriptName = "DestroyComponent"))
 	void K2_DestroyComponent(UActorComponent* Component);
 
-	/**
-	 * Attaches the RootComponent of this Actor to the supplied component, optionally at a named socket. It is not valid to call this on components that are not Registered.
-	 * @param AttachLocationType	Type of attachment, AbsoluteWorld to keep its world position, RelativeOffset to keep the object's relative offset and SnapTo to snap to the new parent.
-	 */
-	UE_DEPRECATED(4.17, "Use AttachToComponent instead.")
+	/** DEPRECATED - Use AttachToComponent() instead */
+	UE_DEPRECATED(4.17, "Use AttachToComponent() instead.")
 	UFUNCTION(BlueprintCallable, meta = (DisplayName = "AttachActorToComponent (Deprecated)", ScriptNoExport, AttachLocationType = "KeepRelativeOffset"), Category = "Utilities|Transformation")
 	void K2_AttachRootComponentTo(USceneComponent* InParent, FName InSocketName = NAME_None, EAttachLocation::Type AttachLocationType = EAttachLocation::KeepRelativeOffset, bool bWeldSimulatedBodies = true);
 
 	/**
 	 * Attaches the RootComponent of this Actor to the supplied component, optionally at a named socket. It is not valid to call this on components that are not Registered.
-	 * @param  Parent					Parent to attach to.
-	 * @param  SocketName				Optional socket to attach to on the parent.
-	 * @param  AttachmentRules			How to handle transforms when attaching.
-	 * @param  bWeldSimulatedBodies		Whether to weld together simulated physics bodies.
+	 * @param Parent					Parent to attach to.
+	 * @param SocketName				Optional socket to attach to on the parent.
+	 * @param LocationRule				How to handle translation when attaching.
+	 * @param RotationRule				How to handle rotation when attaching.
+	 * @param ScaleRule					How to handle scale when attaching.
+	 * @param bWeldSimulatedBodies		Whether to weld together simulated physics bodies.
 	 */
 	UFUNCTION(BlueprintCallable, meta = (DisplayName = "AttachToComponent", ScriptName = "AttachToComponent", bWeldSimulatedBodies = true), Category = "Utilities|Transformation")
 	void K2_AttachToComponent(USceneComponent* Parent, FName SocketName, EAttachmentRule LocationRule, EAttachmentRule RotationRule, EAttachmentRule ScaleRule, bool bWeldSimulatedBodies);
@@ -1187,11 +1174,8 @@ public:
 	 */
 	void AttachToComponent(USceneComponent* Parent, const FAttachmentTransformRules& AttachmentRules, FName SocketName = NAME_None);
 
-	/**
-	 *  Attaches the RootComponent of this Actor to the supplied component, optionally at a named socket. It is not valid to call this on components that are not Registered.
-	 *   @param AttachLocationType	Type of attachment, AbsoluteWorld to keep its world position, RelativeOffset to keep the object's relative offset and SnapTo to snap to the new parent.
-	 */
-	UE_DEPRECATED(4.17, "Use AttachToActor instead.")
+	/** DEPRECATED - Use AttachToActor() instead */
+	UE_DEPRECATED(4.17, "Use AttachToActor() instead.")
 	UFUNCTION(BlueprintCallable, meta = (DisplayName = "AttachActorToActor (Deprecated)", ScriptNoExport, AttachLocationType = "KeepRelativeOffset"), Category = "Utilities|Transformation")
 	void K2_AttachRootComponentToActor(AActor* InParentActor, FName InSocketName = NAME_None, EAttachLocation::Type AttachLocationType = EAttachLocation::KeepRelativeOffset, bool bWeldSimulatedBodies = true);
 
@@ -1215,19 +1199,12 @@ public:
 	UFUNCTION(BlueprintCallable, meta = (DisplayName = "AttachToActor", ScriptName = "AttachToActor", bWeldSimulatedBodies=true), Category = "Utilities|Transformation")
 	void K2_AttachToActor(AActor* ParentActor, FName SocketName, EAttachmentRule LocationRule, EAttachmentRule RotationRule, EAttachmentRule ScaleRule, bool bWeldSimulatedBodies);
 
-	/** 
-	 * Snap the RootComponent of this Actor to the supplied Actor's root component, optionally at a named socket. It is not valid to call this on components that are not Registered. 
-	 * If InSocketName == NAME_None, it will attach to origin of the InParentActor. 
-	 */
-	UE_DEPRECATED(4.17, "Use AttachRootComponentTo with EAttachLocation::SnapToTarget option instead")
+	UE_DEPRECATED(4.17, "Use AttachToComponent() with EAttachLocation::SnapToTarget option instead")
 	UFUNCTION(BlueprintCallable, meta=(DeprecatedFunction, DeprecationMessage = "Use AttachRootComponentTo with EAttachLocation::SnapToTarget option instead", DisplayName = "SnapActorTo", ScriptNoExport), Category="Utilities|Transformation")
 	void SnapRootComponentTo(AActor* InParentActor, FName InSocketName);
 
-	/** 
-	 * Detaches the RootComponent of this Actor from any SceneComponent it is currently attached to. 
-	 * @param bMaintainWorldTransform	If true, update the relative location/rotation of this component to keep its world position the same
-	 */
-	UE_DEPRECATED(4.17, "Use DetachFromActor instead")
+	/** DEPRECATED - Use DetachFromActor() instead */
+	UE_DEPRECATED(4.17, "Use DetachFromActor() instead")
 	UFUNCTION(BlueprintCallable, meta=(DisplayName = "DetachActorFromActor (Deprecated)", ScriptNoExport), Category="Utilities|Transformation")
 	void DetachRootComponentFromParent(bool bMaintainWorldPosition = true);
 
@@ -1253,7 +1230,7 @@ public:
 	 */
 	void DetachAllSceneComponents(class USceneComponent* InParentComponent, const FDetachmentTransformRules& DetachmentRules);
 
-	/** See if this actor contains the supplied tag */
+	/** See if this actor's Tags array contains the supplied name tag */
 	UFUNCTION(BlueprintCallable, Category="Utilities")
 	bool ActorHasTag(FName Tag) const;
 
@@ -1297,8 +1274,7 @@ public:
 	UFUNCTION(BlueprintCallable, Category="Utilities")
 	void SetTickableWhenPaused(bool bTickableWhenPaused);
 
-	/** Allocate a MID for a given parent material. */
-	UE_DEPRECATED(4.17, "Use PrimitiveComponent.CreateAndSetMaterialInstanceDynamic instead.")
+	UE_DEPRECATED(4.17, "Use UPrimitiveComponent::CreateAndSetMaterialInstanceDynamic() instead.")
 	UFUNCTION(BlueprintCallable, meta=(DeprecatedFunction, DeprecationMessage="Use PrimitiveComponent.CreateAndSetMaterialInstanceDynamic instead.", BlueprintProtected = "true"), Category="Rendering|Material")
 	class UMaterialInstanceDynamic* MakeMIDForMaterial(class UMaterialInterface* Parent);
 
@@ -1318,7 +1294,7 @@ public:
 	/** Initiate a begin play call on this Actor, will handle calling in the correct order. */
 	void DispatchBeginPlay();
 
-	/** Returns whether an actor has been initialized */
+	/** Returns whether an actor has been initialized for gameplay */
 	bool IsActorInitialized() const { return bActorInitialized; }
 
 	/** Returns whether an actor is in the process of beginning play */
@@ -1327,6 +1303,7 @@ public:
 	/** Returns whether an actor has had BeginPlay called on it (and not subsequently had EndPlay called) */
 	bool HasActorBegunPlay() const { return ActorHasBegunPlay == EActorBeginPlayState::HasBegunPlay; }
 
+	/** Returns true if this actor is currently being destroyed, some gameplay events may be unsafe */
 	UFUNCTION(BlueprintCallable, Category="Game")
 	bool IsActorBeingDestroyed() const 
 	{
@@ -1337,10 +1314,7 @@ public:
 	UFUNCTION(BlueprintImplementableEvent, BlueprintAuthorityOnly, meta=(DisplayName = "AnyDamage"), Category="Game|Damage")
 	void ReceiveAnyDamage(float Damage, const class UDamageType* DamageType, class AController* InstigatedBy, AActor* DamageCauser);
 	
-	/** 
-	 * Event when this actor takes RADIAL damage 
-	 * @todo Pass it the full array of hits instead of just one?
-	 */
+	/** Event when this actor takes RADIAL damage */
 	UFUNCTION(BlueprintImplementableEvent, BlueprintAuthorityOnly, meta=(DisplayName = "RadialDamage"), Category="Game|Damage")
 	void ReceiveRadialDamage(float DamageReceived, const class UDamageType* DamageType, FVector Origin, const struct FHitResult& HitInfo, class AController* InstigatedBy, AActor* DamageCauser);
 
@@ -1348,7 +1322,7 @@ public:
 	UFUNCTION(BlueprintImplementableEvent, BlueprintAuthorityOnly, meta=(DisplayName = "PointDamage"), Category="Game|Damage")
 	void ReceivePointDamage(float Damage, const class UDamageType* DamageType, FVector HitLocation, FVector HitNormal, class UPrimitiveComponent* HitComponent, FName BoneName, FVector ShotFromDirection, class AController* InstigatedBy, AActor* DamageCauser, const FHitResult& HitInfo);
 
-	/** Event called every frame */
+	/** Event called every frame, if ticking is enabled */
 	UFUNCTION(BlueprintImplementableEvent, meta=(DisplayName = "Tick"))
 	void ReceiveTick(float DeltaSeconds);
 
@@ -1482,8 +1456,6 @@ public:
 	/**
 	 * Construction script, the place to spawn components and do other setup.
 	 * @note Name used in CreateBlueprint function
-	 * @param	Location	The location.
-	 * @param	Rotation	The rotation.
 	 */
 	UFUNCTION(BlueprintImplementableEvent, meta=(BlueprintInternalUseOnly = "true", DisplayName = "Construction Script"))
 	void UserConstructionScript();
@@ -1493,22 +1465,23 @@ public:
 	 * Destruction is latent. It occurs at the end of the tick.
 	 * @param	bNetForce				[opt] Ignored unless called during play.  Default is false.
 	 * @param	bShouldModifyLevel		[opt] If true, Modify() the level before removing the actor.  Default is true.	
-	 * returns	true if destroyed or already marked for destruction, false if indestructible.
+	 * @returns	true if destroyed or already marked for destruction, false if indestructible.
 	 */
 	bool Destroy(bool bNetForce = false, bool bShouldModifyLevel = true );
 
+	/** Called when the actor has been explicitly destroyed. */
 	UFUNCTION(BlueprintImplementableEvent, meta = (Keywords = "delete", DisplayName = "Destroyed"))
 	void ReceiveDestroyed();
 
-	/** Event triggered when the actor is destroyed. */
+	/** Event triggered when the actor has been explicitly destroyed. */
 	UPROPERTY(BlueprintAssignable, Category="Game")
 	FActorDestroyedSignature OnDestroyed;
 
-	/** Event to notify blueprints this actor is about to be deleted. */
+	/** Event to notify blueprints this actor is being deleted or removed from a level. */
 	UFUNCTION(BlueprintImplementableEvent, meta=(Keywords = "delete", DisplayName = "End Play"))
 	void ReceiveEndPlay(EEndPlayReason::Type EndPlayReason);
 
-	/** Event triggered when the actor is being removed from a level. */
+	/** Event triggered when the actor is being deleted or removed from a level. */
 	UPROPERTY(BlueprintAssignable, Category="Game")
 	FActorEndPlaySignature OnEndPlay;
 	
@@ -1541,9 +1514,10 @@ public:
 	/** When selected can this actor be deleted? */
 	virtual bool CanDeleteSelectedActor(FText& OutReason) const { return true; }
 
+	/** Internal struct used to store information about an actor's components during reconstruction */
 	struct FActorRootComponentReconstructionData
 	{
-		// Struct to store info about attached actors
+		/** Struct to store info about attached actors */
 		struct FAttachedActorInfo
 		{
 			TWeakObjectPtr<AActor> Actor;
@@ -1555,21 +1529,22 @@ public:
 			friend FArchive& operator<<(FArchive& Ar, FAttachedActorInfo& ActorInfo);
 		};
 
-		// The RootComponent's transform
+		/** The RootComponent's transform */
 		FTransform Transform;
 
-		// The RootComponent's relative rotation cache (enforces using the same rotator)
+		/** The RootComponent's relative rotation cache (enforces using the same rotator) */
 		FRotationConversionCache TransformRotationCache;
 
-		// The Actor the RootComponent is attached to
+		/** The Actor the RootComponent is attached to */
 		FAttachedActorInfo AttachedParentInfo;
 
-		// Actors that are attached to this RootComponent
+		/** Actors that are attached to this RootComponent */
 		TArray<FAttachedActorInfo> AttachedToInfo;
 
 		friend FArchive& operator<<(FArchive& Ar, FActorRootComponentReconstructionData& RootComponentData);
 	};
 
+	/** Internal struct to track currently active transactions */
 	class FActorTransactionAnnotation : public ITransactionObjectAnnotation
 	{
 	public:
@@ -1588,11 +1563,9 @@ public:
 
 		bool HasInstanceData() const;
 
-		// Actor and component data
 		TWeakObjectPtr<const AActor> Actor;
 		FComponentInstanceDataCache ComponentInstanceData;
 
-		// Root component reconstruction data
 		bool bRootComponentDataCached;
 		FActorRootComponentReconstructionData RootComponentData;
 
@@ -1607,7 +1580,7 @@ public:
 	virtual TSharedPtr<ITransactionObjectAnnotation> FactoryTransactionAnnotation(const ETransactionAnnotationCreationMode InCreationMode) const override;
 	virtual void PostEditUndo(TSharedPtr<ITransactionObjectAnnotation> TransactionAnnotation) override;
 
-	/** @return true if the component is allowed to re-register its components when modified.  False for CDOs or PIE instances. */
+	/** Returns true if the component is allowed to re-register its components when modified. False for CDOs or PIE instances. */
 	bool ReregisterComponentsWhenModified() const;
 
 	/** Called after an actor has been moved in the editor */
@@ -1699,22 +1672,22 @@ public:
 	/** 
 	 * Called by owning level to shift an actor location and all relevant data structures by specified delta
 	 *  
-	 * @param InWorldOffset	 Offset vector to shift actor location
-	 * @param bWorldShift	 Whether this call is part of whole world shifting
+	 * @param InOffset		Offset vector to shift actor location
+	 * @param bWorldShift	Whether this call is part of whole world shifting
 	 */
 	virtual void ApplyWorldOffset(const FVector& InOffset, bool bWorldShift);
 
-	/** 
-	 * Indicates whether this actor should participate in level bounds calculations
-	 */
+	/** Indicates whether this actor should participate in level bounds calculations */
 	virtual bool IsLevelBoundsRelevant() const { return bRelevantForLevelBounds; }
 
-	/** Set LOD Parent Primitive*/
+	/** 
+	 * Set LOD Parent component for all of our components, normally associated with an ALODActor. 
+	 * @param InLODParent			This component used to compute visibility when hierarchical LOD is enabled. 
+	 * @param InParentDrawDistance	Updates the MinDrawDistances of the LODParent
+	 */
 	void SetLODParent(class UPrimitiveComponent* InLODParent, float InParentDrawDistance);
 
 #if WITH_EDITOR
-	// Editor specific
-
 	/** @todo: Remove this flag once it is decided that additive interactive scaling is what we want */
 	static bool bUsePercentageBasedScaling;
 
@@ -1739,47 +1712,43 @@ public:
 	/** Called by MirrorActors to perform a mirroring operation on the actor */
 	virtual void EditorApplyMirror(const FVector& MirrorScale, const FVector& PivotLocation);	
 
-	/**
-	 * Simple accessor to check if the actor is hidden upon editor startup
-	 * @return	true if the actor is hidden upon editor startup; false if it is not
-	 */
+	/** Returns true if the actor is hidden upon editor startup/by default, false if it is not */
 	UFUNCTION(BlueprintCallable, Category = "Editor Scripting | Actor Editing")
 	bool IsHiddenEdAtStartup() const
 	{
 		return bHiddenEd;
 	}
 
-	/** Returns true if this actor is hidden in the editor viewports. */
+	/** Returns true if this actor is hidden in the editor viewports, also checking temporary flags. */
 	UFUNCTION(BlueprintCallable, Category = "Editor Scripting | Actor Editing")
 	bool IsHiddenEd() const;
 
 	/**
-	 * Sets whether or not this actor is hidden in the editor for the duration of the current editor session
-	 *
+	 * Explicitly sets whether or not this actor is hidden in the editor for the duration of the current editor session
 	 * @param bIsHidden	True if the actor is hidden
 	 */
 	UFUNCTION(BlueprintCallable, Category="Editor Scripting | Actor Editing")
 	virtual void SetIsTemporarilyHiddenInEditor( bool bIsHidden );
 
 	/**
-	 * @param  bIncludeParent - Whether to recurse up child actor hierarchy or not
-	 * @return Whether or not this actor is hidden in the editor for the duration of the current editor session
+	 * Returns whether or not this actor was explicitly hidden in the editor for the duration of the current editor session
+	 * @param bIncludeParent - Whether to recurse up child actor hierarchy or not
 	 */
 	UFUNCTION(BlueprintCallable, Category = "Editor Scripting | Actor Editing")
 	bool IsTemporarilyHiddenInEditor(bool bIncludeParent = false) const;
 
-	/** @return	Returns true if this actor is allowed to be displayed, selected and manipulated by the editor. */
+	/** Returns true if this actor is allowed to be displayed, selected and manipulated by the editor. */
 	UFUNCTION(BlueprintCallable, Category = "Editor Scripting | Actor Editing")
 	bool IsEditable() const;
 
-	/** @return	Returns true if this actor can EVER be selected in a level in the editor.  Can be overridden by specific actors to make them unselectable. */
+	/** Returns true if this actor can EVER be selected in a level in the editor.  Can be overridden by specific actors to make them unselectable. */
 	UFUNCTION(BlueprintCallable, Category = "Editor Scripting | Actor Editing")
 	virtual bool IsSelectable() const { return true; }
 
-	/** @return	Returns true if this actor should be shown in the scene outliner */
+	/** Returns true if this actor should be shown in the scene outliner */
 	virtual bool IsListedInSceneOutliner() const;
 
-	/** @return	Returns true if this actor is allowed to be attached to the given actor */
+	/** Returns true if this actor is allowed to be attached to the given actor */
 	virtual bool EditorCanAttachTo(const AActor* InParent, FText& OutReason) const;
 
 	/** Called before editor copy, true allow export */
@@ -1806,10 +1775,7 @@ public:
 	 */
 	virtual void CheckForDeprecated();
 
-	/**
-	 * Returns this actor's current label.  Actor labels are only available in development builds.
-	 * @return	The label text
-	 */
+	/** Returns this actor's current label.  Actor labels are only available in development builds. */
 	UFUNCTION(BlueprintCallable, Category = "Editor Scripting | Actor Editing")
 	const FString& GetActorLabel() const;
 
@@ -1830,10 +1796,7 @@ public:
 	 */
 	bool IsActorLabelEditable() const;
 
-	/**
-	 * Returns this actor's folder path. Actor folder paths are only available in development builds.
-	 * @return	The folder path
-	 */
+	/** Returns this actor's folder path. Actor folder paths are only available in development builds. */
 	UFUNCTION(BlueprintCallable, Category = "Editor Scripting | Actor Editing")
 	const FName& GetFolderPath() const;
 
@@ -1868,10 +1831,10 @@ public:
 	{
 		return bOptimizeBPComponentData;
 	}
-
 #endif		// WITH_EDITOR
 
 	/**
+	 * Function used to prioritize actors when deciding which to replicate
 	 * @param ViewPos		Position of the viewer
 	 * @param ViewDir		Vector direction of viewer
 	 * @param Viewer		"net object" owned by the client for whom net priority is being determined (typically player controller)
@@ -1879,19 +1842,19 @@ public:
 	 * @param InChannel		Channel on which this actor is being replicated.
 	 * @param Time			Time since actor was last replicated
 	 * @param bLowBandwidth True if low bandwidth of viewer
-	 * @return				Priority of this actor for replication
+	 * @return				Priority of this actor for replication, higher is more important
 	 */
 	virtual float GetNetPriority(const FVector& ViewPos, const FVector& ViewDir, class AActor* Viewer, AActor* ViewTarget, UActorChannel* InChannel, float Time, bool bLowBandwidth);
 
 	/**
 	 * Similar to GetNetPriority, but will only be used for prioritizing actors while recording a replay.
-	 *
 	 * @param ViewPos		Position of the viewer
 	 * @param ViewDir		Vector direction of viewer
 	 * @param Viewer		"net object" owned by the client for whom net priority is being determined (typically player controller)
 	 * @param ViewTarget	The actor that is currently being viewed/controlled by Viewer, usually a pawn
 	 * @param InChannel		Channel on which this actor is being replicated.
 	 * @param Time			Time since actor was last replicated
+	 * @return				Priority of this actor for replays, higher is more important
 	 */
 	virtual float GetReplayPriority(const FVector& ViewPos, const FVector& ViewDir, class AActor* Viewer, AActor* ViewTarget, UActorChannel* const InChannel, float Time);
 
@@ -1927,7 +1890,6 @@ public:
 	/** Swaps Role and RemoteRole if client */
 	void ExchangeNetRoles(bool bRemoteOwner);
 
-	/** The replay system calls this to hack the Role and RemoteRole while recording replays on a client. Only call this if you know what you're doing! */
 	UE_DEPRECATED(4.22, "Renamed to SwapRoles")
 	void SwapRolesForReplay() { SwapRoles(); }
 
@@ -1966,17 +1928,17 @@ public:
 	float GetActorTickInterval() const;
 
 	/**
-	 *	ticks the actor
-	 *	@param	DeltaTime			The time slice of this tick
-	 *	@param	TickType			The type of tick that is happening
-	 *	@param	ThisTickFunction	The tick function that is firing, useful for getting the completion handle
+	 * Dispatches the once-per frame Tick() function for this actor
+	 * @param	DeltaTime			The time slice of this tick
+	 * @param	TickType			The type of tick that is happening
+	 * @param	ThisTickFunction	The tick function that is firing, useful for getting the completion handle
 	 */
 	virtual void TickActor( float DeltaTime, enum ELevelTick TickType, FActorTickFunction& ThisTickFunction );
 
 	/**
-	 * Called when an actor is done spawning into the world (from UWorld::SpawnActor).
+	 * Called when an actor is done spawning into the world (from UWorld::SpawnActor), both in the editor and during gameplay
 	 * For actors with a root component, the location and rotation will have already been set.
-	 * Takes place after any construction scripts have been called
+	 * This is called before calling construction scripts, but after native components have been created
 	 */
 	virtual void PostActorCreated();
 
@@ -2022,7 +1984,6 @@ protected:
 	void SyncReplicatedPhysicsSimulation();
 
 public:
-
 	/** 
 	 * Set the owner of this Actor, used primarily for network replication. 
 	 * @param NewOwner	The Actor whom takes over ownership of this Actor
@@ -2030,10 +1991,7 @@ public:
 	UFUNCTION(BlueprintCallable, Category=Actor)
 	virtual void SetOwner( AActor* NewOwner );
 
-	/**
-	 * Get the owner of this Actor, used primarily for network replication.
-	 * @return Actor that owns this Actor
-	 */
+	/** Get the owner of this Actor, used primarily for network replication. */
 	UFUNCTION(BlueprintCallable, Category=Actor)
 	AActor* GetOwner() const;
 
@@ -2118,21 +2076,25 @@ protected:
 	
 public:	
 	/** 
-	  * @param RealViewer - is the "controlling net object" associated with the client for which network relevancy is being checked (typically player controller)
-	  * @param ViewTarget - is the Actor being used as the point of view for the RealViewer
-	  * @param SrcLocation - is the viewing location
-	  *
-	  * @return bool - true if this actor is network relevant to the client associated with RealViewer 
-	  */
+	 * Checks to see if this actor is relevant for a specific network connection
+	 *
+	 * @param RealViewer - is the "controlling net object" associated with the client for which network relevancy is being checked (typically player controller)
+	 * @param ViewTarget - is the Actor being used as the point of view for the RealViewer
+	 * @param SrcLocation - is the viewing location
+	 *
+	 * @return bool - true if this actor is network relevant to the client associated with RealViewer 
+	 */
 	virtual bool IsNetRelevantFor(const AActor* RealViewer, const AActor* ViewTarget, const FVector& SrcLocation) const;
 
 	/**
-	* @param RealViewer - is the "controlling net object" associated with the client for which network relevancy is being checked (typically player controller)
-	* @param ViewTarget - is the Actor being used as the point of view for the RealViewer
-	* @param SrcLocation - is the viewing location
-	*
-	* @return bool - true if this actor is replay relevant to the client associated with RealViewer
-	*/
+	 * Checks to see if this actor is relevant for a recorded replay
+	 *
+	 * @param RealViewer - is the "controlling net object" associated with the client for which network relevancy is being checked (typically player controller)
+	 * @param ViewTarget - is the Actor being used as the point of view for the RealViewer
+	 * @param SrcLocation - is the viewing location
+	 *
+	 * @return bool - true if this actor is replay relevant to the client associated with RealViewer
+	 */
 	virtual bool IsReplayRelevantFor(const AActor* RealViewer, const AActor* ViewTarget, const FVector& SrcLocation, const float CullDistanceSquared) const;
 
 	/**
@@ -2156,10 +2118,10 @@ public:
 	void PostActorConstruction();
 
 public:
-	/** Called immediately before gameplay begins. */
+	/** Called right before components are initialized, only called during gameplay */
 	virtual void PreInitializeComponents();
 
-	/** Allow actors to initialize themselves on the C++ side */
+	/** Allow actors to initialize themselves on the C++ side after all of their components have been initialized, only called during gameplay */
 	virtual void PostInitializeComponents();
 
 	/**
@@ -2177,10 +2139,10 @@ public:
 	/** Dispatches ReceiveHit virtual and OnComponentHit delegate */
 	virtual void DispatchPhysicsCollisionHit(const struct FRigidBodyCollisionInfo& MyInfo, const struct FRigidBodyCollisionInfo& OtherInfo, const FCollisionImpactData& RigidCollisionData);
 	
-	/** @return the actor responsible for replication, if any. Typically the player controller */
+	/** Return the actor responsible for replication, if any. Typically the player controller */
 	virtual const AActor* GetNetOwner() const;
 
-	/** @return the owning UPlayer (if any) of this actor. This will be a local player, a net connection, or nullptr. */
+	/** Return the owning UPlayer (if any) of this actor. This will be a local player, a net connection, or nullptr. */
 	virtual class UPlayer* GetNetOwningPlayer();
 
 	/**
@@ -2228,7 +2190,10 @@ public:
 	UFUNCTION(BlueprintCallable, Category="Actor")
 	bool IsChildActor() const;
 
-	/** Returns a list of all child actors, including children of children */
+	/** 
+	 * Returns a list of all actors spawned by our Child Actor Components, including children of children. 
+	 * This does not return the contents of the Children array
+	 */
 	UFUNCTION(BlueprintCallable, Category="Actor")
 	void GetAllChildActors(TArray<AActor*>& ChildActors, bool bIncludeDescendants = true) const;
 
@@ -2243,10 +2208,10 @@ public:
 	/** Ensure that all the components in the Components array are registered */
 	virtual void RegisterAllComponents();
 
-	/** Called before all the components in the Components array are registered */
+	/** Called before all the components in the Components array are registered, called both in editor and during gameplay */
 	virtual void PreRegisterAllComponents();
 
-	/** Called after all the components in the Components array are registered */
+	/** Called after all the components in the Components array are registered, called both in editor and during gameplay */
 	virtual void PostRegisterAllComponents();
 
 	/** Returns true if Actor has deferred the RegisterAllComponents() call at spawn time (e.g. pending Blueprint SCS execution to set up a scene root component). */
@@ -2255,7 +2220,10 @@ public:
 	/** Returns true if Actor has a registered root component */
 	bool HasValidRootComponent();
 
-	/** Unregister all currently registered components */
+	/** 
+	 * Unregister all currently registered components
+	 * @param bForReregister If true, RegisterAllComponents will be called immediately after this so some slow operations can be avoided
+	 */
 	virtual void UnregisterAllComponents(bool bForReregister = false);
 
 	/** Called after all currently registered components are cleared */
@@ -2265,7 +2233,7 @@ public:
 	virtual void ReregisterAllComponents();
 
 	/**
-	 * Incrementally registers components associated with this actor
+	 * Incrementally registers components associated with this actor, used during level streaming
 	 *
 	 * @param NumComponentsToRegister  Number of components to register in this run, 0 for all
 	 * @return true when all components were registered for this actor
@@ -2278,10 +2246,10 @@ public:
 	/** Update all components transforms */
 	void UpdateComponentTransforms();
 
-	/** Iterate over components array and call InitializeComponent */
+	/** Iterate over components array and call InitializeComponent, which happens once per actor */
 	void InitializeComponents();
 
-	/** Iterate over components array and call UninitializeComponent */
+	/** Iterate over components array and call UninitializeComponent, called when the actor is ending play */
 	void UninitializeComponents();
 
 	/** Debug rendering to visualize the component tree for this actor. */
@@ -2379,14 +2347,14 @@ public:
 	virtual void EndPlay(const EEndPlayReason::Type EndPlayReason);
 
 	/**
-	 * Iterates up the Base chain to see whether or not this Actor is based on the given Actor
+	 * Iterates up the movement base chain to see whether or not this Actor is based on the given Actor, defaults to checking attachment
 	 * @param Other the Actor to test for
 	 * @return true if this Actor is based on Other Actor
 	 */
 	virtual bool IsBasedOnActor(const AActor* Other) const;
 	
 	/**
-	 * Iterates up the Base chain to see whether or not this Actor is attached to the given Actor
+	 * Iterates up the attachment chain to see whether or not this Actor is attached to the given Actor
 	 * @param Other the Actor to test for
 	 * @return true if this Actor is attached on Other Actor
 	 */
@@ -2400,7 +2368,7 @@ public:
 	// Blueprint
 
 #if WITH_EDITOR
-	/** Find all FRandomStream structs in this ACtor and generate new random seeds for them. */
+	/** Find all FRandomStream structs in this Actor and generate new random seeds for them. */
 	void SeedAllRandomStreams();
 #endif // WITH_EDITOR
 
@@ -2412,14 +2380,12 @@ public:
 
 	/** 
 	 * Debug helper to show the component hierarchy of this actor.
-	 * @param Info	Optional String to display at top of info
+	 * @param Info			Optional String to display at top of info
+	 * @param bShowPosition	If true, will display component's position in world space
 	 */
 	void DebugShowComponentHierarchy( const TCHAR* Info, bool bShowPosition  = true);
 	
-	/** 
-	 * Debug helper for showing the component hierarchy of one component
-	 * @param Info	Optional String to display at top of info
-	 */
+	/** Debug helper for showing the component hierarchy of one component */
 	void DebugShowOneComponentHierarchy( USceneComponent* SceneComp, int32& NestLevel, bool bShowPosition );
 
 	/**
@@ -2469,7 +2435,6 @@ protected:
 	void PostCreateBlueprintComponent(UActorComponent* NewActorComp);
 
 public:
-
 	/** Checks for and resolve any name conflicts prior to instancing a new Blueprint Component. */
 	void CheckComponentInstanceName(const FName InName);
 
@@ -2492,7 +2457,7 @@ public:
 	UFUNCTION(BlueprintCallable, Category="Utilities", meta=(Keywords = "dependency"))
 	void SetTickGroup(ETickingGroup NewTickGroup);
 
-	/** Called once this actor has been deleted */
+	/** Called when this actor is explicitly being destroyed during gameplay or in the editor, not called during level streaming or gameplay ending */
 	virtual void Destroyed();
 
 	/** Call ReceiveHit, as well as delegates on Actor and Component */
@@ -2525,16 +2490,16 @@ public:
 	 */
 	virtual void GetSimpleCollisionCylinder(float& CollisionRadius, float& CollisionHalfHeight) const;
 
-	/** @returns the radius of the collision cylinder from GetSimpleCollisionCylinder(). */
+	/** Returns the radius of the collision cylinder from GetSimpleCollisionCylinder(). */
 	float GetSimpleCollisionRadius() const;
 
-	/** @returns the half height of the collision cylinder from GetSimpleCollisionCylinder(). */
+	/** Returns the half height of the collision cylinder from GetSimpleCollisionCylinder(). */
 	float GetSimpleCollisionHalfHeight() const;
 
-	/** @returns collision extents vector for this Actor, based on GetSimpleCollisionCylinder(). */
+	/** Returns collision extents vector for this Actor, based on GetSimpleCollisionCylinder(). */
 	FVector GetSimpleCollisionCylinderExtent() const;
 
-	/** @returns true if the root component is registered and has collision enabled.  */
+	/** Returns true if the root component is registered and has collision enabled.  */
 	virtual bool IsRootComponentCollisionRegistered() const;
 
 	/**
@@ -2548,11 +2513,9 @@ public:
 	// Collision/Physics functions.
  
 	/** 
-	 * Get Collision Response to the Channel that entered for all components
+	 * Get Collision Response to the passed in Channel for all components
 	 * It returns Max of state - i.e. if Component A overlaps, but if Component B blocks, it will return block as response
 	 * if Component A ignores, but if Component B overlaps, it will return overlap
-	 *
-	 * @param Channel - The channel to change the response of
 	 */
 	virtual ECollisionResponse GetComponentsCollisionResponseToChannel(ECollisionChannel Channel) const;
 
@@ -2560,8 +2523,8 @@ public:
 	void DisableComponentsSimulatePhysics();
 
 	/**
-	 * @return WorldSettings for the World the actor is in
-	 * If you'd like to know UWorld this placed actor (not dynamic spawned actor) belong to, use GetTypedOuter<UWorld>()
+	 * Returns the WorldSettings for the World the actor is in
+	 * If you'd like to know what UWorld this placed actor (not dynamic spawned actor) belong to, use GetTypedOuter<UWorld>()
 	 */
 	class AWorldSettings* GetWorldSettings() const;
 
@@ -2585,12 +2548,12 @@ public:
 protected:
 	virtual float InternalTakeRadialDamage(float Damage, struct FRadialDamageEvent const& RadialDamageEvent, class AController* EventInstigator, AActor* DamageCauser);
 	virtual float InternalTakePointDamage(float Damage, struct FPointDamageEvent const& PointDamageEvent, class AController* EventInstigator, AActor* DamageCauser);
-public:
 
+public:
 	/** Called when this actor becomes the given PlayerController's ViewTarget. Triggers the Blueprint event K2_OnBecomeViewTarget. */
 	virtual void BecomeViewTarget( class APlayerController* PC );
 
-	/*( Called when this actor is no longer the given PlayerController's ViewTarget. Also triggers the Blueprint event K2_OnEndViewTarget. */
+	/** Called when this actor is no longer the given PlayerController's ViewTarget. Also triggers the Blueprint event K2_OnEndViewTarget. */
 	virtual void EndViewTarget( class APlayerController* PC );
 
 	/** Event called when this Actor becomes the view target for the given PlayerController. */
@@ -2704,10 +2667,10 @@ public:
 	 */
 	virtual void PostRenderFor(class APlayerController* PC, class UCanvas* Canvas, FVector CameraPosition, FVector CameraDir);
 
-	/** whether this Actor is in the persistent level, i.e. not a sublevel */
+	/** Returns whether this Actor is in the persistent level, i.e. not a sublevel */
 	bool IsInPersistentLevel(bool bIncludeLevelStreamingPersistent = false) const;
 
-	/** Getter for the cached world pointer */
+	/** Getter for the cached world pointer, will return null if the actor is not actually spawned in a level */
 	virtual UWorld* GetWorld() const override;
 
 	/** Get the timer instance from the actors world */
@@ -2729,10 +2692,10 @@ public:
 	/** Returns true if this is a replicated actor that was placed in the map */
 	bool IsNetStartupActor() const;
 
-	/** Searches components array and returns first encountered component of the specified class. */
+	/** Searches components array and returns first encountered component of the specified class, native version of GetComponentByClass */
 	virtual UActorComponent* FindComponentByClass(const TSubclassOf<UActorComponent> ComponentClass) const;
 	
-	/** Script exposed version of FindComponentByClass */
+	/** Searches components array and returns first encountered component of the specified class */
 	UFUNCTION(BlueprintCallable, Category = "Actor", meta = (ComponentClass = "ActorComponent"), meta = (DeterminesOutputType = "ComponentClass"))
 	UActorComponent* GetComponentByClass(TSubclassOf<UActorComponent> ComponentClass) const;
 
@@ -2747,7 +2710,7 @@ public:
 	UFUNCTION(BlueprintCallable, Category = "Actor", meta = (ComponentClass = "ActorComponent"), meta = (DeterminesOutputType = "ComponentClass"))
 	TArray<UActorComponent*> GetComponentsByTag(TSubclassOf<UActorComponent> ComponentClass, FName Tag) const;
 
-	/** Templatized version for syntactic nicety. */
+	/** Templatized version of FindComponentByClass that handles casting for you */
 	template<class T>
 	T* FindComponentByClass() const
 	{
@@ -2910,7 +2873,6 @@ public:
 	}
 
 protected:
-
 	/** Set of replicated components, stored as an array to save space as this is generally not very large */
 	TArray<UActorComponent*> ReplicatedComponents;
 
@@ -2980,6 +2942,7 @@ public:
 private:
 	static FMakeNoiseDelegate MakeNoiseDelegate;
 
+
 public:
 	//~=============================================================================
 	// Debugging functions
@@ -3030,72 +2993,48 @@ private:
 	// Static helpers for accessing functions on SceneComponent.
 	// These are templates for no other reason than to delay compilation until USceneComponent is defined.
 
-	/*~
-	 * Returns transform of the RootComponent 
-	 */ 
 	template<class T>
 	static FORCEINLINE FTransform TemplateGetActorTransform(const T* RootComponent)
 	{
 		return (RootComponent != nullptr) ? RootComponent->GetComponentTransform() : FTransform();
 	}
 
-	/*~
-	 * Returns location of the RootComponent 
-	 */ 
 	template<class T>
 	static FORCEINLINE FVector TemplateGetActorLocation(const T* RootComponent)
 	{
 		return (RootComponent != nullptr) ? RootComponent->GetComponentLocation() : FVector::ZeroVector;
 	}
 
-	/*~
-	 * Returns rotation of the RootComponent 
-	 */ 
 	template<class T>
 	static FORCEINLINE FRotator TemplateGetActorRotation(const T* RootComponent)
 	{
 		return (RootComponent != nullptr) ? RootComponent->GetComponentRotation() : FRotator::ZeroRotator;
 	}
 
-	/*~
-	 * Returns scale of the RootComponent 
-	 */ 
 	template<class T>
 	static FORCEINLINE FVector TemplateGetActorScale(const T* RootComponent)
 	{
 		return (RootComponent != nullptr) ? RootComponent->GetComponentScale() : FVector(1.f,1.f,1.f);
 	}
 
-	/*~
-	 * Returns quaternion of the RootComponent
-	 */ 
 	template<class T>
 	static FORCEINLINE FQuat TemplateGetActorQuat(const T* RootComponent)
 	{
 		return (RootComponent != nullptr) ? RootComponent->GetComponentQuat() : FQuat(ForceInit);
 	}
 
-	/*~
-	 * Returns the forward (X) vector (length 1.0) of the RootComponent, in world space.
-	 */ 
 	template<class T>
 	static FORCEINLINE FVector TemplateGetActorForwardVector(const T* RootComponent)
 	{
 		return (RootComponent != nullptr) ? RootComponent->GetForwardVector() : FVector::ForwardVector;
 	}
 
-	/*~
-	 * Returns the Up (Z) vector (length 1.0) of the RootComponent, in world space.
-	 */ 
 	template<class T>
 	static FORCEINLINE FVector TemplateGetActorUpVector(const T* RootComponent)
 	{
 		return (RootComponent != nullptr) ? RootComponent->GetUpVector() : FVector::UpVector;
 	}
 
-	/*~
-	 * Returns the Right (Y) vector (length 1.0) of the RootComponent, in world space.	
-	 */ 
 	template<class T>
 	static FORCEINLINE FVector TemplateGetActorRightVector(const T* RootComponent)
 	{
@@ -3103,7 +3042,7 @@ private:
 	}
 };
 
-
+/** Internal struct used by level code to mark actors as destroyed */
 struct FMarkActorIsBeingDestroyed
 {
 private:
@@ -3128,7 +3067,6 @@ private:
 	friend UWorld;
 };
 
-
 /**
  * TInlineComponentArray is simply a TArray that reserves a fixed amount of space on the stack
  * to try to avoid heap allocation when there are fewer than a specified number of elements expected in the result.
diff --git a/Engine/Source/Runtime/Engine/Private/Actor.cpp b/Engine/Source/Runtime/Engine/Private/Actor.cpp
--- a/Engine/Source/Runtime/Engine/Private/Actor.cpp
+++ b/Engine/Source/Runtime/Engine/Private/Actor.cpp
@@ -83,28 +83,10 @@ AActor::AActor()
 	InitializeDefaults();
 }
 
-
 AActor::AActor(const FObjectInitializer& ObjectInitializer)
 {
 	// Forward to default constructor (we don't use ObjectInitializer for anything, this is for compatibility with inherited classes that call Super( ObjectInitializer )
 	InitializeDefaults();
-
-#if (CSV_PROFILER && !UE_BUILD_SHIPPING)
-	// Increment actor class count
-	{
-		if (!HasAnyFlags(RF_ArchetypeObject | RF_ClassDefaultObject))
-		{
-			FScopeLock Lock(&CSVActorClassNameToCountMapLock);
-
-			const UClass* ParentNativeClass = GetParentNativeClass(GetClass());
-			FName NativeClassName = ParentNativeClass ? ParentNativeClass->GetFName() : NAME_None;
-			int32& CurrentCount = CSVActorClassNameToCountMap.FindOrAdd(NativeClassName);
-			CurrentCount++;
-			CSVActorTotalCount++;
-		}
-	}
-#endif // (CSV_PROFILER && !UE_BUILD_SHIPPING)
-
 }
 
 void AActor::InitializeDefaults()
@@ -133,8 +115,7 @@ void AActor::InitializeDefaults()
 	bHiddenEdLevel = false;
 	bActorLabelEditable = true;
 	SpriteScale = 1.0f;
-	bEnableAutoLODGeneration = true;	
-	InputConsumeOption_DEPRECATED = ICO_ConsumeBoundKeys;
+	bEnableAutoLODGeneration = true;
 	bOptimizeBPComponentData = false;
 #endif // WITH_EDITORONLY_DATA
 	NetCullDistanceSquared = 225000000.0f;
@@ -155,6 +136,22 @@ void AActor::InitializeDefaults()
 	PivotOffset = FVector::ZeroVector;
 #endif
 	SpawnCollisionHandlingMethod = ESpawnActorCollisionHandlingMethod::AlwaysSpawn;
+
+#if (CSV_PROFILER && !UE_BUILD_SHIPPING)
+	// Increment actor class count
+	{
+		if (!HasAnyFlags(RF_ArchetypeObject | RF_ClassDefaultObject))
+		{
+			FScopeLock Lock(&CSVActorClassNameToCountMapLock);
+
+			const UClass* ParentNativeClass = GetParentNativeClass(GetClass());
+			FName NativeClassName = ParentNativeClass ? ParentNativeClass->GetFName() : NAME_None;
+			int32& CurrentCount = CSVActorClassNameToCountMap.FindOrAdd(NativeClassName);
+			CurrentCount++;
+			CSVActorTotalCount++;
+		}
+	}
+#endif // (CSV_PROFILER && !UE_BUILD_SHIPPING)
 }
 
 void FActorTickFunction::ExecuteTick(float DeltaTime, enum ELevelTick TickType, ENamedThreads::Type CurrentThread, const FGraphEventRef& MyCompletionGraphEvent)
@@ -676,11 +673,6 @@ void AActor::PostLoad()
 	}
 
 #if WITH_EDITORONLY_DATA
-	if (GetLinkerUE4Version() < VER_UE4_CONSUME_INPUT_PER_BIND)
-	{
-		bBlockInput = (InputConsumeOption_DEPRECATED == ICO_ConsumeAll);
-	}
-
 	if (AActor* ParentActor = ParentComponentActor_DEPRECATED.Get())
 	{
 		TInlineComponentArray<UChildActorComponent*> ParentChildActorComponents(ParentActor);
@@ -1644,7 +1636,7 @@ bool AActor::HasNetOwner() const
 void AActor::K2_AttachRootComponentTo(USceneComponent* InParent, FName InSocketName, EAttachLocation::Type AttachLocationType /*= EAttachLocation::KeepRelativeOffset */, bool bWeldSimulatedBodies /*=true*/)
 {
 PRAGMA_DISABLE_DEPRECATION_WARNINGS
-if (RootComponent && InParent)
+	if (RootComponent && InParent)
 	{
 		RootComponent->AttachTo(InParent, InSocketName, AttachLocationType, bWeldSimulatedBodies);
 	}
diff --git a/Engine/Source/Runtime/Engine/Private/Components/ActorComponent.cpp b/Engine/Source/Runtime/Engine/Private/Components/ActorComponent.cpp
--- a/Engine/Source/Runtime/Engine/Private/Components/ActorComponent.cpp
+++ b/Engine/Source/Runtime/Engine/Private/Components/ActorComponent.cpp
@@ -501,6 +501,11 @@ bool UActorComponent::NeedsLoadForServer() const
 	return (!IsEditorOnly() && bNeedsLoadOuter && Super::NeedsLoadForServer());
 }
 
+bool UActorComponent::NeedsLoadForEditorGame() const
+{
+	return !IsEditorOnly() && Super::NeedsLoadForEditorGame();
+}
+
 int32 UActorComponent::GetFunctionCallspace( UFunction* Function, void* Parameters, FFrame* Stack )
 {
 	AActor* MyOwner = GetOwner();
diff --git a/Engine/Source/Runtime/Engine/Private/Components/PrimitiveComponent.cpp b/Engine/Source/Runtime/Engine/Private/Components/PrimitiveComponent.cpp
--- a/Engine/Source/Runtime/Engine/Private/Components/PrimitiveComponent.cpp
+++ b/Engine/Source/Runtime/Engine/Private/Components/PrimitiveComponent.cpp
@@ -213,7 +213,6 @@ private:
 // PRIMITIVE COMPONENT
 ///////////////////////////////////////////////////////////////////////////////
 
-int32 UPrimitiveComponent::CurrentTag = 2147483647 / 4;
 uint32 UPrimitiveComponent::GlobalOverlapEventsCounter = 0;
 
 // 0 is reserved to mean invalid
diff --git a/Engine/Source/Runtime/Engine/Private/Components/SceneComponent.cpp b/Engine/Source/Runtime/Engine/Private/Components/SceneComponent.cpp
--- a/Engine/Source/Runtime/Engine/Private/Components/SceneComponent.cpp
+++ b/Engine/Source/Runtime/Engine/Private/Components/SceneComponent.cpp
@@ -2340,23 +2340,6 @@ void USceneComponent::UpdateChildTransforms(EUpdateTransformFlags UpdateTransfor
 	}
 }
 
-#if WITH_EDITORONLY_DATA
-void USceneComponent::Serialize(FArchive& Ar)
-{
-	Super::Serialize(Ar);
-
-#if WITH_EDITORONLY_DATA
-	// Copy from deprecated properties
-	if(Ar.UE4Ver() < VER_UE4_SCENECOMP_TRANSLATION_TO_LOCATION)
-	{
-		RelativeLocation = RelativeTranslation_DEPRECATED;
-		bAbsoluteLocation = bAbsoluteTranslation_DEPRECATED;
-	}
-#endif
-}
-#endif
-
-
 void USceneComponent::PostInterpChange(UProperty* PropertyThatChanged)
 {
 	Super::PostInterpChange(PropertyThatChanged);
diff --git a/Engine/Source/Runtime/Engine/Private/PrimitiveComponentPhysics.cpp b/Engine/Source/Runtime/Engine/Private/PrimitiveComponentPhysics.cpp
--- a/Engine/Source/Runtime/Engine/Private/PrimitiveComponentPhysics.cpp
+++ b/Engine/Source/Runtime/Engine/Private/PrimitiveComponentPhysics.cpp
@@ -116,11 +116,6 @@ void UPrimitiveComponent::SetConstraintMode(EDOFMode::Type ConstraintMode)
 	RootBI->SetDOFLock(ConstraintMode);
 }
 
-void UPrimitiveComponent::SetLockedAxis(EDOFMode::Type LockedAxis)
-{
-	SetConstraintMode(LockedAxis);
-}
-
 void UPrimitiveComponent::AddImpulse(FVector Impulse, FName BoneName, bool bVelChange)
 {
 	if (FBodyInstance* BI = GetBodyInstance(BoneName))
