Need to recompile all materials using MPC when parameters are added/removed
This is basically reverting previous CL 2612044 which was created in response to UE-17798
That CL was attempting to avoid recompiling materials if they didn't use particular parameter that was removed, but this is not correct...
Adding/removing any MPC parameter will change uniform buffer layout, which requires all materials using MPC to be recompiled

diff --git a/Engine/Source/Runtime/Engine/Private/Materials/ParameterCollection.cpp b/Engine/Source/Runtime/Engine/Private/Materials/ParameterCollection.cpp
--- a/Engine/Source/Runtime/Engine/Private/Materials/ParameterCollection.cpp
+++ b/Engine/Source/Runtime/Engine/Private/Materials/ParameterCollection.cpp
@@ -141,15 +141,15 @@ void SanitizeParameters(TArray<ParameterType>& Parameters)
 	}
 }
 
-TArray<FCollectionScalarParameter> PreviousScalarParameters;
-TArray<FCollectionVectorParameter> PreviousVectorParameters;
+int32 PreviousNumScalarParameters = 0;
+int32 PreviousNumVectorParameters = 0;
 
 void UMaterialParameterCollection::PreEditChange(UProperty* PropertyThatWillChange)
 {
 	Super::PreEditChange(PropertyThatWillChange);
 
-	PreviousScalarParameters = ScalarParameters;
-	PreviousVectorParameters = VectorParameters;
+	PreviousNumScalarParameters = ScalarParameters.Num();
+	PreviousNumVectorParameters = VectorParameters.Num();
 }
 
 void UMaterialParameterCollection::PostEditChangeProperty(FPropertyChangedEvent& PropertyChangedEvent)
@@ -159,8 +159,7 @@ void UMaterialParameterCollection::PostEditChangeProperty(FPropertyChangedEvent&
 
 	// If the array counts have changed, an element has been added or removed, and we need to update the uniform buffer layout,
 	// Which also requires recompiling any referencing materials
-	if (ScalarParameters.Num() != PreviousScalarParameters.Num()
-		|| VectorParameters.Num() != PreviousVectorParameters.Num())
+	if (ScalarParameters.Num() != PreviousNumScalarParameters || VectorParameters.Num() != PreviousNumVectorParameters)
 	{
 		// Limit the count of parameters to fit within uniform buffer limits
 		const uint32 MaxScalarParameters = 1024;
@@ -190,28 +189,6 @@ void UMaterialParameterCollection::PostEditChangeProperty(FPropertyChangedEvent&
 			CurrentWorld->AddParameterCollectionInstance(this, false);
 		}
 
-		// Build set of changed parameter names
-		TSet<FName> ParameterNames;
-		for (const FCollectionVectorParameter& Param : PreviousVectorParameters)
-		{
-			ParameterNames.Add(Param.ParameterName);
-		}
-
-		for (const FCollectionScalarParameter& Param : PreviousScalarParameters)
-		{
-			ParameterNames.Add(Param.ParameterName);
-		}
-
-		for (const FCollectionVectorParameter& Param : VectorParameters)
-		{
-			ParameterNames.Remove(Param.ParameterName);
-		}
-
-		for (const FCollectionScalarParameter& Param : ScalarParameters)
-		{
-			ParameterNames.Remove(Param.ParameterName);
-		}
-
 		// Create a material update context so we can safely update materials using this parameter collection.
 		{
 			FMaterialUpdateContext UpdateContext;
@@ -235,16 +212,8 @@ void UMaterialParameterCollection::PostEditChangeProperty(FPropertyChangedEvent&
 					{
 						if (CurrentMaterial->MaterialParameterCollectionInfos[FunctionIndex].ParameterCollection == this)
 						{
-							TArray<UMaterialExpressionCollectionParameter*> CollectionParameters;
-							CurrentMaterial->GetAllExpressionsInMaterialAndFunctionsOfType(CollectionParameters);
-							for (UMaterialExpressionCollectionParameter* CollectionParameter : CollectionParameters)
-							{
-								if (ParameterNames.Contains(CollectionParameter->ParameterName))
-								{
-									bRecompile = true;
-									break;
-								}
-							}
+							bRecompile = true;
+							break;
 						}
 					}
 				}
@@ -271,9 +240,6 @@ void UMaterialParameterCollection::PostEditChangeProperty(FPropertyChangedEvent&
 
 	UpdateDefaultResource();
 
-	PreviousScalarParameters.Empty();
-	PreviousVectorParameters.Empty();
-
 	Super::PostEditChangeProperty(PropertyChangedEvent);
 }
 
