PR #5516: UE-22698: Prevent player controller memory leak (Contributed by [REDACTED_PERSON])

diff --git a/Engine/Source/Runtime/Engine/Private/GameModeBase.cpp b/Engine/Source/Runtime/Engine/Private/GameModeBase.cpp
--- a/Engine/Source/Runtime/Engine/Private/GameModeBase.cpp
+++ b/Engine/Source/Runtime/Engine/Private/GameModeBase.cpp
@@ -671,6 +671,7 @@ APlayerController* AGameModeBase::Login(UPlayer* NewPlayer, ENetRole InRemoteRol
 	ErrorMessage = InitNewPlayer(NewPlayerController, UniqueId, Options, Portal);
 	if (!ErrorMessage.IsEmpty())
 	{
+		NewPlayerController->Destroy();
 		return nullptr;
 	}
 
