Fix UKismetMathLibrary::Hypotenuse returning nan for 0x0 triangle

diff --git a/Engine/Source/Runtime/Engine/Private/KismetMathLibrary.cpp b/Engine/Source/Runtime/Engine/Private/KismetMathLibrary.cpp
--- a/Engine/Source/Runtime/Engine/Private/KismetMathLibrary.cpp
+++ b/Engine/Source/Runtime/Engine/Private/KismetMathLibrary.cpp
@@ -199,13 +199,15 @@ bool UKismetMathLibrary::InRange_FloatFloat(float Value, float Min, float Max, b
 float UKismetMathLibrary::Hypotenuse(float Width, float Height)
 {
 	// This implementation avoids overflow/underflow caused by squaring width and height:
-	Width = FMath::Abs(Width);
-	Height = FMath::Abs(Height);
+	float Min = FMath::Abs(Width);
+	float Max = FMath::Abs(Height);
 
-	float Min = FGenericPlatformMath::Min(Width, Height);
-	float Max = FGenericPlatformMath::Max(Width, Height);
-	float Ratio = Min / Max;
-	return Max * FMath::Sqrt(1.f + Ratio * Ratio);
+	if (Min > Max)
+	{
+		Swap(Min, Max);
+	}
+
+	return (FMath::IsNearlyZero(Max) ? 0.f : Max * FMath::Sqrt(1.f + FMath::Square(Min/Max)));
 }
 
 float UKismetMathLibrary::Log(float A, float Base)
