Copy CL 4884877 from Dev-Audio:

- Push up when ActiveSounds are removed from the concurrency manager to when a sound stops, not when it's deleted/removed
   - Fixes bug with sounds not playing due to ActiveSounds still considered to be playing from the perspective of the concurrency manager even though they're queued to stop
- As an extra precaution, safeguard against attempting to add/remove ActiveSounds to concurrency manager multiple times and log when attempt is made

diff --git a/Engine/Source/Runtime/Engine/Classes/Sound/SoundConcurrency.h b/Engine/Source/Runtime/Engine/Classes/Sound/SoundConcurrency.h
--- a/Engine/Source/Runtime/Engine/Classes/Sound/SoundConcurrency.h
+++ b/Engine/Source/Runtime/Engine/Classes/Sound/SoundConcurrency.h
@@ -221,11 +221,11 @@ public:
 	/** Returns a newly allocated active sound given the input active sound struct. Will return nullptr if the active sound concurrency evaluation doesn't allow for it. */
 	FActiveSound* CreateNewActiveSound(const FActiveSound& NewActiveSound);
 
-	/** Removes the active sound from concurrency tracking. */
-	void RemoveActiveSound(FActiveSound* ActiveSound);
+	/** Removes the active sound from concurrency tracking when active sound is stopped. */
+	void StopActiveSound(FActiveSound* ActiveSound);
 
 	/** Stops any active sounds due to max concurrency quietest sound resolution rule */
-	void StopQuietSoundsDueToMaxConcurrency();
+	void UpdateQuietSoundsToStop();
 
 private: // Methods
 	/** Evaluates whether or not the sound can play given the concurrency group's rules. Appends permissible
diff --git a/Engine/Source/Runtime/Engine/Private/AudioDevice.cpp b/Engine/Source/Runtime/Engine/Private/AudioDevice.cpp
--- a/Engine/Source/Runtime/Engine/Private/AudioDevice.cpp
+++ b/Engine/Source/Runtime/Engine/Private/AudioDevice.cpp
@@ -2052,7 +2052,7 @@ void FAudioDevice::StopQuietSoundsDueToMaxConcurrency(TArray<FWaveInstance*>& Wa
 	// Now stop any sounds that are active that are in concurrency resolution groups that resolve by stopping quietest
 	{
 		SCOPE_CYCLE_COUNTER(STAT_AudioEvaluateConcurrency);
-		ConcurrencyManager.StopQuietSoundsDueToMaxConcurrency();
+		ConcurrencyManager.UpdateQuietSoundsToStop();
 	}
 
 	// Remove all wave instances from the wave instance list that are stopping due to max concurrency
@@ -2071,6 +2071,7 @@ void FAudioDevice::StopQuietSoundsDueToMaxConcurrency(TArray<FWaveInstance*>& Wa
 			if (ActiveSound->bShouldStopDueToMaxConcurrency)
 			{
 				ActiveSound->Stop(false);
+				ConcurrencyManager.StopActiveSound(ActiveSound);
 			}
 		}
 	}
@@ -4272,6 +4273,7 @@ void FAudioDevice::ProcessingPendingActiveSoundStops(bool bForceDelete)
 void FAudioDevice::AddSoundToStop(FActiveSound* SoundToStop)
 {
 	check(IsInAudioThread());
+	check(SoundToStop);
 
 	const uint64 AudioComponentID = SoundToStop->GetAudioComponentID();
 	if (AudioComponentID > 0)
@@ -4279,8 +4281,13 @@ void FAudioDevice::AddSoundToStop(FActiveSound* SoundToStop)
 		AudioComponentIDToActiveSoundMap.Remove(AudioComponentID);
 	}
 
-	check(SoundToStop);
-	PendingSoundsToStop.Add(SoundToStop);
+	bool bAlreadyPending = false;
+	PendingSoundsToStop.Add(SoundToStop, &bAlreadyPending);
+
+	if (!bAlreadyPending)
+	{
+		ConcurrencyManager.StopActiveSound(SoundToStop);
+	}
 }
 
 void FAudioDevice::StopActiveSound(const uint64 AudioComponentID)
@@ -4344,8 +4351,6 @@ void FAudioDevice::RemoveActiveSound(FActiveSound* ActiveSound)
 {
 	check(IsInAudioThread());
 
-	ConcurrencyManager.RemoveActiveSound(ActiveSound);
-
 	// Perform the notification
 	if (ActiveSound->GetAudioComponentID() > 0)
 	{
diff --git a/Engine/Source/Runtime/Engine/Private/SoundConcurrency.cpp b/Engine/Source/Runtime/Engine/Private/SoundConcurrency.cpp
--- a/Engine/Source/Runtime/Engine/Private/SoundConcurrency.cpp
+++ b/Engine/Source/Runtime/Engine/Private/SoundConcurrency.cpp
@@ -89,6 +89,11 @@ void FConcurrencyGroup::RemoveActiveSound(FActiveSound* ActiveSound)
 
 	// Remove from array
 	const int32 NumRemoved = ActiveSounds.RemoveSwap(ActiveSound);
+	if (NumRemoved == 0)
+	{
+		return;
+	}
+
 	check(NumRemoved == 1);
 
 	// Rebase generations due to removal of a member
@@ -220,7 +225,6 @@ void FSoundConcurrencyManager::CreateNewGroupsFromHandles(
 					OwnerPerSoundConcurrencyMap.Emplace(OwnerObjectID, FSoundInstanceEntry(ConcurrencyHandle.ObjectID, ConcurrencyGroup.GetGroupID()));
 					OutGroupsToApply.Add(&ConcurrencyGroup);
 				}
-				break;
 			}
 			break;
 
@@ -233,7 +237,6 @@ void FSoundConcurrencyManager::CreateNewGroupsFromHandles(
 					SoundObjectToConcurrencyGroup.Add(SoundObjectID, ConcurrencyGroup.GetGroupID());
 					OutGroupsToApply.Add(&ConcurrencyGroup);
 				}
-				break;
 			}
 			break;
 		}
@@ -278,19 +281,26 @@ FConcurrencyGroup& FSoundConcurrencyManager::CreateNewConcurrencyGroup(const FCo
 FConcurrencyGroup* FSoundConcurrencyManager::CanPlaySound(const FActiveSound& NewActiveSound, const FConcurrencyGroupID GroupID, TArray<FActiveSound*>& OutSoundsToEvict)
 {
 	check(GroupID != 0);
-	FConcurrencyGroup& ConcurrencyGroup = ConcurrencyGroups.FindChecked(GroupID);
+	FConcurrencyGroup* ConcurrencyGroup = ConcurrencyGroups.Find(GroupID);
+	if (!ConcurrencyGroup)
+	{
+		UE_LOG(LogAudio, Warning, TEXT("Attempting to add active sound '%s' (owner '%s') to invalid concurrency group."),
+			NewActiveSound.GetSound() ? *NewActiveSound.GetSound()->GetFullName() : TEXT("Unset"),
+			*NewActiveSound.GetOwnerName());
+		return nullptr;
+	}
 
 	// StopQuietest doesn't evict, it culls once we instantiate the sound.  This
 	// is because it is not possible to evaluate sound volumes *before* they play.
-	if (ConcurrencyGroup.GetSettings().ResolutionRule == EMaxConcurrentResolutionRule::StopQuietest)
+	if (ConcurrencyGroup->GetSettings().ResolutionRule == EMaxConcurrentResolutionRule::StopQuietest)
 	{
-		return &ConcurrencyGroup;
+		return ConcurrencyGroup;
 	}
 
-	if (ConcurrencyGroup.IsFull())
+	if (ConcurrencyGroup->IsFull())
 	{
 		// If no room for new sound, early out
-		if (FActiveSound* SoundToEvict = GetEvictableSound(NewActiveSound, ConcurrencyGroup))
+		if (FActiveSound* SoundToEvict = GetEvictableSound(NewActiveSound, *ConcurrencyGroup))
 		{
 			OutSoundsToEvict.Add(SoundToEvict);
 		}
@@ -300,7 +310,7 @@ FConcurrencyGroup* FSoundConcurrencyManager::CanPlaySound(const FActiveSound& Ne
 		}
 	}
 
-	return &ConcurrencyGroup;
+	return ConcurrencyGroup;
 }
 
 FActiveSound* FSoundConcurrencyManager::GetEvictableSound(const FActiveSound& NewActiveSound, const FConcurrencyGroup& ConcurrencyGroup)
@@ -549,7 +559,7 @@ FActiveSound* FSoundConcurrencyManager::CreateAndEvictActiveSounds(const FActive
 
 		// Remove the active sound from the concurrency manager immediately so it doesn't count towards
 		// subsequent concurrency resolution checks (i.e. if sounds are triggered multiple times in this frame)
-		RemoveActiveSound(SoundToEvict);
+		StopActiveSound(SoundToEvict);
 
 		// Add this sound to list of sounds that need to stop but don't stop it immediately
 		AudioDevice->AddSoundToStop(SoundToEvict);
@@ -559,42 +569,42 @@ FActiveSound* FSoundConcurrencyManager::CreateAndEvictActiveSounds(const FActive
 	return ActiveSound;
 }
 
-void FSoundConcurrencyManager::RemoveActiveSound(FActiveSound* ActiveSound)
+void FSoundConcurrencyManager::StopActiveSound(FActiveSound* ActiveSound)
 {
 	// Remove this sound from it's concurrency list
 	for (const FConcurrencyGroupID ConcurrencyGroupID : ActiveSound->ConcurrencyGroupIDs)
 	{
-		FConcurrencyGroup& ConcurrencyGroup = ConcurrencyGroups.FindChecked(ConcurrencyGroupID);
+		FConcurrencyGroup* ConcurrencyGroup = ConcurrencyGroups.Find(ConcurrencyGroupID);
+		if (!ConcurrencyGroup)
+		{
+			UE_LOG(LogAudio, Fatal, TEXT("Attempting to remove stopped sound '%s' from inactive concurrency group."),
+				ActiveSound->GetSound() ? *ActiveSound->GetSound()->GetName(): TEXT("Unset"));
+			return;
+		}
 
-		check(!ConcurrencyGroup.IsEmpty());
-		ConcurrencyGroup.RemoveActiveSound(ActiveSound);
+		check(!ConcurrencyGroup->IsEmpty());
+		ConcurrencyGroup->RemoveActiveSound(ActiveSound);
 
-		if (ConcurrencyGroup.IsEmpty())
+		if (ConcurrencyGroup->IsEmpty())
 		{
 			// Get the object ID prior to removing from groups collection to avoid reading
 			// from the object after its destroyed.
-			const FConcurrencyObjectID ConcurrencyObjectID = ConcurrencyGroup.GetObjectID();
+			const FConcurrencyObjectID ConcurrencyObjectID = ConcurrencyGroup->GetObjectID();
 
 			// Remove the object from the map.
 			ConcurrencyGroups.Remove(ConcurrencyGroupID);
 
-			// Try to remove from global group map. If not present, check next map.
-			if (ConcurrencyMap.Remove(ConcurrencyObjectID))
-			{
-				continue;
-			}
+			// Remove from global group map if present.
+			ConcurrencyMap.Remove(ConcurrencyObjectID);
 
-			// Try to remove from sound object map, If not present, check next map.
+			// Remove from sound object map if present.
 			if (USoundBase* Sound = ActiveSound->GetSound())
 			{
 				const FSoundOwnerObjectID ObjectID = Sound->GetUniqueID();
-				if (SoundObjectToConcurrencyGroup.Remove(ObjectID))
-				{
-					continue;
-				}
+				SoundObjectToConcurrencyGroup.Remove(ObjectID);
 			}
 
-			// Try to remove from owner map. If not present, check next map.
+			// Remove from owner map if present.
 			const uint32 OwnerID = ActiveSound->GetOwnerID();
 			if (FOwnerConcurrencyMapEntry* OwnerEntry = OwnerConcurrencyMap.Find(OwnerID))
 			{
@@ -604,11 +614,10 @@ void FSoundConcurrencyManager::RemoveActiveSound(FActiveSound* ActiveSound)
 					{
 						OwnerConcurrencyMap.Remove(OwnerID);
 					}
-					continue;
 				}
 			}
 
-			// Lastly, try to remove from owner per sound map.
+			// Remove from owner per sound map if present.
 			if (FSoundInstanceEntry* InstanceEntry = OwnerPerSoundConcurrencyMap.Find(OwnerID))
 			{
 				if (USoundBase* Sound = ActiveSound->GetSound())
@@ -618,7 +627,6 @@ void FSoundConcurrencyManager::RemoveActiveSound(FActiveSound* ActiveSound)
 						if (InstanceEntry->SoundInstanceToConcurrencyGroup.Num() == 0)
 						{
 							OwnerPerSoundConcurrencyMap.Remove(OwnerID);
-							continue;
 						}
 					}
 				}
@@ -628,7 +636,7 @@ void FSoundConcurrencyManager::RemoveActiveSound(FActiveSound* ActiveSound)
 	ActiveSound->ConcurrencyGroupIDs.Reset();
 }
 
-void FSoundConcurrencyManager::StopQuietSoundsDueToMaxConcurrency()
+void FSoundConcurrencyManager::UpdateQuietSoundsToStop()
 {
 	for (auto& ConcurrenyGroupEntry : ConcurrencyGroups)
 	{
diff --git a/Engine/Source/Runtime/Engine/Public/ActiveSound.h b/Engine/Source/Runtime/Engine/Public/ActiveSound.h
--- a/Engine/Source/Runtime/Engine/Public/ActiveSound.h
+++ b/Engine/Source/Runtime/Engine/Public/ActiveSound.h
@@ -599,9 +599,6 @@ private:
 	/** Whether or not the active sound is stopping. */
 	bool IsStopping() const { return bIsStopping; }
 
-	/** Ensures that the sound has finished stopping. Waits until the sound finishes stopping naturally. */
-	void EnsureStopped();
-
 	/** Called when an active sound has been stopped but needs to update it's stopping sounds. Returns true when stopping sources have finished stopping. */
 	bool UpdateStoppingSources(uint64 CurrentTick, bool bEnsureStopped);
 
