Rename material expression 'ShadowReplace' to 'ShadowPassSwitch', add tooltips (copy of CL 4871825)

diff --git a/Engine/Source/Runtime/Engine/Classes/Materials/MaterialExpressionShadowReplace.h b/Engine/Source/Runtime/Engine/Classes/Materials/MaterialExpressionShadowReplace.h
--- a/Engine/Source/Runtime/Engine/Classes/Materials/MaterialExpressionShadowReplace.h
+++ b/Engine/Source/Runtime/Engine/Classes/Materials/MaterialExpressionShadowReplace.h
@@ -9,23 +9,22 @@
 #include "Materials/MaterialExpression.h"
 #include "MaterialExpressionShadowReplace.generated.h"
 
-UCLASS()
+UCLASS(meta = (DisplayName = "ShadowPassSwitch"))
 class UMaterialExpressionShadowReplace : public UMaterialExpression
 {
 	GENERATED_UCLASS_BODY()
 
-	UPROPERTY()
+	UPROPERTY(meta = (RequiredInput = "true", ToolTip = "Input will be used when rendering into non-ShadowMap passes"))
 	FExpressionInput Default;
 
-	UPROPERTY()
+	UPROPERTY(meta = (RequiredInput = "true", ToolTip = "Input will be used when rendering into ShadowMap"))
 	FExpressionInput Shadow;
 
 	//~ Begin UMaterialExpression Interface
 #if WITH_EDITOR
 	virtual int32 Compile(class FMaterialCompiler* Compiler, int32 OutputIndex) override;
 	virtual void GetCaption(TArray<FString>& OutCaptions) const override;
+	virtual void GetExpressionToolTip(TArray<FString>& OutToolTip) override;
 #endif
 	//~ End UMaterialExpression Interface
 };
-
-
diff --git a/Engine/Source/Runtime/Engine/Private/Materials/MaterialExpressions.cpp b/Engine/Source/Runtime/Engine/Private/Materials/MaterialExpressions.cpp
--- a/Engine/Source/Runtime/Engine/Private/Materials/MaterialExpressions.cpp
+++ b/Engine/Source/Runtime/Engine/Private/Materials/MaterialExpressions.cpp
@@ -12539,11 +12539,11 @@ int32 UMaterialExpressionShadowReplace::Compile(class FMaterialCompiler* Compile
 {
 	if (!Default.GetTracedInput().Expression)
 	{
-		return Compiler->Errorf(TEXT("Missing ShadowReplace input Default"));
+		return Compiler->Errorf(TEXT("Missing input Default"));
 	}
 	else if (!Shadow.GetTracedInput().Expression)
 	{
-		return Compiler->Errorf(TEXT("Missing ShadowReplace input Shadow"));
+		return Compiler->Errorf(TEXT("Missing input Shadow"));
 	}
 	else
 	{
@@ -12555,7 +12555,12 @@ int32 UMaterialExpressionShadowReplace::Compile(class FMaterialCompiler* Compile
 
 void UMaterialExpressionShadowReplace::GetCaption(TArray<FString>& OutCaptions) const
 {
-	OutCaptions.Add(TEXT("ShadowReplace"));
+	OutCaptions.Add(TEXT("Shadow Pass Switch"));
+}
+
+void UMaterialExpressionShadowReplace::GetExpressionToolTip(TArray<FString>& OutToolTip)
+{
+	ConvertToMultilineToolTip(TEXT("Allows material to define specialized behavior when being rendered into ShadowMap."), 40, OutToolTip);
 }
 #endif // WITH_EDITOR
 
