Explicitly disallow camera/reflection vector access from domain shader, as these aren't availiable in FMaterialTessellationParameters
- Gives better error message about why material is failing to compile

diff --git a/Engine/Source/Runtime/Engine/Private/Materials/HLSLMaterialTranslator.h b/Engine/Source/Runtime/Engine/Private/Materials/HLSLMaterialTranslator.h
--- a/Engine/Source/Runtime/Engine/Private/Materials/HLSLMaterialTranslator.h
+++ b/Engine/Source/Runtime/Engine/Private/Materials/HLSLMaterialTranslator.h
@@ -2846,7 +2846,7 @@ protected:
 
 	virtual int32 ReflectionVector() override
 	{
-		if (ShaderFrequency != SF_Pixel && ShaderFrequency != SF_Compute && ShaderFrequency != SF_Domain)
+		if (ShaderFrequency != SF_Pixel && ShaderFrequency != SF_Compute)
 		{
 			return NonPixelShaderExpressionError();
 		}
@@ -2859,7 +2859,7 @@ protected:
 
 	virtual int32 ReflectionAboutCustomWorldNormal(int32 CustomWorldNormal, int32 bNormalizeCustomWorldNormal) override
 	{
-		if (ShaderFrequency != SF_Pixel && ShaderFrequency != SF_Compute && ShaderFrequency != SF_Domain)
+		if (ShaderFrequency != SF_Pixel && ShaderFrequency != SF_Compute)
 		{
 			return NonPixelShaderExpressionError();
 		}
@@ -2881,7 +2881,7 @@ protected:
 
 	virtual int32 CameraVector() override
 	{
-		if (ShaderFrequency != SF_Pixel && ShaderFrequency != SF_Compute && ShaderFrequency != SF_Domain)
+		if (ShaderFrequency != SF_Pixel && ShaderFrequency != SF_Compute)
 		{
 			return NonPixelShaderExpressionError();
 		}
