Fix viewport initialisation size in SceneViewport.

diff --git a/Engine/Source/Runtime/Engine/Private/Slate/SceneViewport.cpp b/Engine/Source/Runtime/Engine/Private/Slate/SceneViewport.cpp
--- a/Engine/Source/Runtime/Engine/Private/Slate/SceneViewport.cpp
+++ b/Engine/Source/Runtime/Engine/Private/Slate/SceneViewport.cpp
@@ -1373,7 +1373,10 @@ void FSceneViewport::ResizeFrame(uint32 NewWindowSizeX, uint32 NewWindowSizeY, E
 			}
 
 			// Resize viewport
-			FVector2D ViewportSize = WindowToResize->GetWindowSizeFromClientSize(FVector2D(SizeX, SizeY));
+			uint32 CurrentSizeX = SizeX ? SizeX : NewWindowSizeX;
+			uint32 CurrentSizeY = SizeY ? SizeY : NewWindowSizeY;
+
+			FVector2D ViewportSize = WindowToResize->GetWindowSizeFromClientSize(FVector2D(CurrentSizeX, CurrentSizeY));
 			FVector2D NewViewportSize = WindowToResize->GetViewportSize();
 
 			if (NewViewportSize != ViewportSize || NewWindowMode != OldWindowMode)
