Fix for crash in FRepLayout::InitFromObjectClass()

diff --git a/Engine/Source/Runtime/Engine/Private/RepLayout.cpp b/Engine/Source/Runtime/Engine/Private/RepLayout.cpp
--- a/Engine/Source/Runtime/Engine/Private/RepLayout.cpp
+++ b/Engine/Source/Runtime/Engine/Private/RepLayout.cpp
@@ -4114,12 +4114,13 @@ void FRepLayout::InitFromClass(UClass* InObjectClass, const UNetConnection* Serv
 	SCOPE_CYCLE_COUNTER(STAT_RepLayout_InitFromObjectClass);
 	SCOPE_CYCLE_UOBJECT(ObjectClass, InObjectClass);
 
-	RoleIndex = -1;
-	RemoteRoleIndex = -1;
-	FirstNonCustomParent = -1;
+	const bool bIsObjectActor = InObjectClass->IsChildOf(AActor::StaticClass());
+	RoleIndex                 = -1;
+	RemoteRoleIndex           = -1;
+	FirstNonCustomParent      = -1;
 
-	int32 RelativeHandle = 0;
-	int32 LastOffset = -1;
+	int32 RelativeHandle      = 0;
+	int32 LastOffset          = -1;
 
 	InObjectClass->SetUpRuntimeReplicationData();
 	Parents.Empty(InObjectClass->ClassReps.Num());
@@ -4161,19 +4162,22 @@ void FRepLayout::InitFromClass(UClass* InObjectClass, const UNetConnection* Serv
 			FirstNonCustomParent = ParentHandle;
 		}
 
-		// Find Role/RemoteRole property indexes so we can swap them on the client
-		if (Property->GetFName() == NAME_Role)
+		if (bIsObjectActor)
 		{
-			check(RoleIndex == -1);
-			check(Parents[ParentHandle].CmdEnd == Parents[ParentHandle].CmdStart + 1);
-			RoleIndex = ParentHandle;
-		}
+			// Find Role/RemoteRole property indexes so we can swap them on the client
+			if (Property->GetFName() == NAME_Role)
+			{
+				check(RoleIndex == -1);
+				check(Parents[ParentHandle].CmdEnd == Parents[ParentHandle].CmdStart + 1);
+				RoleIndex = ParentHandle;
+			}
 
-		if (Property->GetFName() == NAME_RemoteRole)
-		{
-			check(RemoteRoleIndex == -1);
-			check(Parents[ParentHandle].CmdEnd == Parents[ParentHandle].CmdStart + 1);
-			RemoteRoleIndex = ParentHandle;
+			if (Property->GetFName() == NAME_RemoteRole)
+			{
+				check(RemoteRoleIndex == -1);
+				check(Parents[ParentHandle].CmdEnd == Parents[ParentHandle].CmdStart + 1);
+				RemoteRoleIndex = ParentHandle;
+			}
 		}
 	}
 
