Remove redundent bCanEverTick sets from deprecated PostPhysicsComponentTick
#rnx

diff --git a/Engine/Plugins/Runtime/ApexDestruction/Source/ApexDestruction/Private/DestructibleComponent.cpp b/Engine/Plugins/Runtime/ApexDestruction/Source/ApexDestruction/Private/DestructibleComponent.cpp
--- a/Engine/Plugins/Runtime/ApexDestruction/Source/ApexDestruction/Private/DestructibleComponent.cpp
+++ b/Engine/Plugins/Runtime/ApexDestruction/Source/ApexDestruction/Private/DestructibleComponent.cpp
@@ -38,7 +38,6 @@ UDestructibleComponent::UDestructibleComponent(const FObjectInitializer& ObjectI
 {
 
 	PrimaryComponentTick.bCanEverTick = false;
-	PostPhysicsComponentTick.bCanEverTick = false;
 
 	bHasCustomNavigableGeometry = EHasCustomNavigableGeometry::Yes;
 
diff --git a/Engine/Plugins/Runtime/MixedRealityCaptureFramework/Source/MixedRealityCaptureFramework/Private/MixedRealityCaptureActor.cpp b/Engine/Plugins/Runtime/MixedRealityCaptureFramework/Source/MixedRealityCaptureFramework/Private/MixedRealityCaptureActor.cpp
--- a/Engine/Plugins/Runtime/MixedRealityCaptureFramework/Source/MixedRealityCaptureFramework/Private/MixedRealityCaptureActor.cpp
+++ b/Engine/Plugins/Runtime/MixedRealityCaptureFramework/Source/MixedRealityCaptureFramework/Private/MixedRealityCaptureActor.cpp
@@ -245,7 +245,6 @@ void AMixedRealityCaptureActor::BeginPlay()
 		CamVizualizer->SetRelativeTransform(FTransform(FRotator(90.f, 0.f, 0.f), FVector(7.5, 0.f, 0.f), FVector(0.15f)));
 		CamVizualizer->SetCollisionProfileName(UCollisionProfile::NoCollision_ProfileName);
 		CamVizualizer->CastShadow = false;
-		CamVizualizer->PostPhysicsComponentTick.bCanEverTick = false;
 		CamVizualizer->RegisterComponent();
 	}
 #endif
diff --git a/Engine/Plugins/Runtime/MixedRealityCaptureFramework/Source/MixedRealityCaptureFramework/Private/MixedRealityCaptureComponent.cpp b/Engine/Plugins/Runtime/MixedRealityCaptureFramework/Source/MixedRealityCaptureFramework/Private/MixedRealityCaptureComponent.cpp
--- a/Engine/Plugins/Runtime/MixedRealityCaptureFramework/Source/MixedRealityCaptureFramework/Private/MixedRealityCaptureComponent.cpp
+++ b/Engine/Plugins/Runtime/MixedRealityCaptureFramework/Source/MixedRealityCaptureFramework/Private/MixedRealityCaptureComponent.cpp
@@ -395,7 +395,6 @@ void UMixedRealityCaptureComponent::OnRegister()
 			ProxyMeshComponent->SetCollisionProfileName(UCollisionProfile::NoCollision_ProfileName);
 			ProxyMeshComponent->bHiddenInGame = true;
 			ProxyMeshComponent->CastShadow    = false;
-			ProxyMeshComponent->PostPhysicsComponentTick.bCanEverTick = false;
 			ProxyMeshComponent->CreationMethod = CreationMethod;
 			ProxyMeshComponent->RegisterComponent();
 		}
diff --git a/Engine/Source/Developer/FunctionalTesting/Private/FunctionalTest.cpp b/Engine/Source/Developer/FunctionalTesting/Private/FunctionalTest.cpp
--- a/Engine/Source/Developer/FunctionalTesting/Private/FunctionalTest.cpp
+++ b/Engine/Source/Developer/FunctionalTesting/Private/FunctionalTest.cpp
@@ -140,7 +140,6 @@ AFunctionalTest::AFunctionalTest( const FObjectInitializer& ObjectInitializer )
 
 #if WITH_EDITORONLY_DATA
 	RenderComp = CreateDefaultSubobject<UFuncTestRenderingComponent>(TEXT("RenderComp"));
-	RenderComp->PostPhysicsComponentTick.bCanEverTick = false;
 	RenderComp->SetupAttachment(RootComponent);
 #endif // WITH_EDITORONLY_DATA
 
@@ -161,7 +160,6 @@ AFunctionalTest::AFunctionalTest( const FObjectInitializer& ObjectInitializer )
 		TestName->SetHorizontalAlignment(EHTA_Center);
 		TestName->SetRelativeLocation(FVector(0, 0, 80));
 		TestName->SetRelativeRotation(FRotator(0, 0, 0));
-		TestName->PostPhysicsComponentTick.bCanEverTick = false;
 		TestName->SetupAttachment(RootComponent);
 	}
 #endif
diff --git a/Engine/Source/Runtime/AIModule/Private/Navigation/NavLinkProxy.cpp b/Engine/Source/Runtime/AIModule/Private/Navigation/NavLinkProxy.cpp
--- a/Engine/Source/Runtime/AIModule/Private/Navigation/NavLinkProxy.cpp
+++ b/Engine/Source/Runtime/AIModule/Private/Navigation/NavLinkProxy.cpp
@@ -26,7 +26,6 @@ ANavLinkProxy::ANavLinkProxy(const FObjectInitializer& ObjectInitializer) : Supe
 
 #if WITH_EDITORONLY_DATA
 	EdRenderComp = CreateDefaultSubobject<UNavLinkRenderingComponent>(TEXT("EdRenderComp"));
-	EdRenderComp->PostPhysicsComponentTick.bCanEverTick = false;
 	EdRenderComp->SetupAttachment(RootComponent);
 #endif // WITH_EDITORONLY_DATA
 
diff --git a/Engine/Source/Runtime/CinematicCamera/Private/CameraRig_Crane.cpp b/Engine/Source/Runtime/CinematicCamera/Private/CameraRig_Crane.cpp
--- a/Engine/Source/Runtime/CinematicCamera/Private/CameraRig_Crane.cpp
+++ b/Engine/Source/Runtime/CinematicCamera/Private/CameraRig_Crane.cpp
@@ -51,7 +51,6 @@ ACameraRig_Crane::ACameraRig_Crane(const FObjectInitializer& ObjectInitializer)
 			PreviewMesh_CraneBase->SetCollisionProfileName(UCollisionProfile::NoCollision_ProfileName);
 			PreviewMesh_CraneBase->bHiddenInGame = true;
 			PreviewMesh_CraneBase->CastShadow = false;
-			PreviewMesh_CraneBase->PostPhysicsComponentTick.bCanEverTick = false;
 
 			PreviewMesh_CraneBase->SetupAttachment(TransformComponent);		// sibling of yawcontrol
 		}
@@ -65,7 +64,6 @@ ACameraRig_Crane::ACameraRig_Crane(const FObjectInitializer& ObjectInitializer)
 			PreviewMesh_CraneArm->SetCollisionProfileName(UCollisionProfile::NoCollision_ProfileName);
 			PreviewMesh_CraneArm->bHiddenInGame = true;
 			PreviewMesh_CraneArm->CastShadow = false;
-			PreviewMesh_CraneArm->PostPhysicsComponentTick.bCanEverTick = false;
 			
 			PreviewMesh_CraneArm->SetupAttachment(CranePitchControl);		// sibling of the mount
 			PreviewMesh_CraneArm->RelativeRotation = FRotator(0.f, 90.f, 0.f);
@@ -82,7 +80,6 @@ ACameraRig_Crane::ACameraRig_Crane(const FObjectInitializer& ObjectInitializer)
 			PreviewMesh_CraneMount->SetCollisionProfileName(UCollisionProfile::NoCollision_ProfileName);
 			PreviewMesh_CraneMount->bHiddenInGame = true;
 			PreviewMesh_CraneMount->CastShadow = false;
-			PreviewMesh_CraneMount->PostPhysicsComponentTick.bCanEverTick = false;
 
 			PreviewMesh_CraneMount->SetupAttachment(CranePitchControl);
 			PreviewMesh_CraneMount->RelativeRotation = FRotator(0.f, 90.f, 0.f);
@@ -99,7 +96,6 @@ ACameraRig_Crane::ACameraRig_Crane(const FObjectInitializer& ObjectInitializer)
 			PreviewMesh_CraneCounterWeight->SetCollisionProfileName(UCollisionProfile::NoCollision_ProfileName);
 			PreviewMesh_CraneCounterWeight->bHiddenInGame = true;
 			PreviewMesh_CraneCounterWeight->CastShadow = false;
-			PreviewMesh_CraneCounterWeight->PostPhysicsComponentTick.bCanEverTick = false;
 
 			PreviewMesh_CraneCounterWeight->SetupAttachment(CranePitchControl);
 			PreviewMesh_CraneCounterWeight->RelativeRotation = FRotator(0.f, 90.f, 0.f);
diff --git a/Engine/Source/Runtime/CinematicCamera/Private/CameraRig_Rail.cpp b/Engine/Source/Runtime/CinematicCamera/Private/CameraRig_Rail.cpp
--- a/Engine/Source/Runtime/CinematicCamera/Private/CameraRig_Rail.cpp
+++ b/Engine/Source/Runtime/CinematicCamera/Private/CameraRig_Rail.cpp
@@ -46,7 +46,6 @@ ACameraRig_Rail::ACameraRig_Rail(const FObjectInitializer& ObjectInitializer)
 			PreviewMesh_Mount->SetCollisionProfileName(UCollisionProfile::NoCollision_ProfileName);
 			PreviewMesh_Mount->bHiddenInGame = true;
 			PreviewMesh_Mount->CastShadow = false;
-			PreviewMesh_Mount->PostPhysicsComponentTick.bCanEverTick = false;
 
 			PreviewMesh_Mount->SetupAttachment(RailCameraMount);
 		}
@@ -66,7 +65,6 @@ USplineMeshComponent* ACameraRig_Rail::CreateSplinePreviewSegment()
 		Segment->SetCollisionProfileName(UCollisionProfile::NoCollision_ProfileName);
 		Segment->bHiddenInGame = true;
 		Segment->CastShadow = false;
-		Segment->PostPhysicsComponentTick.bCanEverTick = false;
 	 
 		Segment->SetupAttachment(TransformComponent);
 		Segment->RegisterComponent();
diff --git a/Engine/Source/Runtime/CinematicCamera/Private/CineCameraComponent.cpp b/Engine/Source/Runtime/CinematicCamera/Private/CineCameraComponent.cpp
--- a/Engine/Source/Runtime/CinematicCamera/Private/CineCameraComponent.cpp
+++ b/Engine/Source/Runtime/CinematicCamera/Private/CineCameraComponent.cpp
@@ -483,7 +483,6 @@ void UCineCameraComponent::CreateDebugFocusPlane()
 			DebugFocusPlaneComponent->SetCollisionProfileName(UCollisionProfile::NoCollision_ProfileName);
 			DebugFocusPlaneComponent->bHiddenInGame = false;
 			DebugFocusPlaneComponent->CastShadow = false;
-			DebugFocusPlaneComponent->PostPhysicsComponentTick.bCanEverTick = false;
 			DebugFocusPlaneComponent->CreationMethod = CreationMethod;
 			DebugFocusPlaneComponent->bSelectable = false;
 
diff --git a/Engine/Source/Runtime/Engine/Private/Camera/CameraComponent.cpp b/Engine/Source/Runtime/Engine/Private/Camera/CameraComponent.cpp
--- a/Engine/Source/Runtime/Engine/Private/Camera/CameraComponent.cpp
+++ b/Engine/Source/Runtime/Engine/Private/Camera/CameraComponent.cpp
@@ -108,7 +108,6 @@ void UCameraComponent::OnRegister()
 			ProxyMeshComponent->SetCollisionProfileName(UCollisionProfile::NoCollision_ProfileName);
 			ProxyMeshComponent->bHiddenInGame = true;
 			ProxyMeshComponent->CastShadow = false;
-			ProxyMeshComponent->PostPhysicsComponentTick.bCanEverTick = false;
 			ProxyMeshComponent->CreationMethod = CreationMethod;
 			ProxyMeshComponent->RegisterComponentWithWorld(GetWorld());
 		}
diff --git a/Engine/Source/Runtime/Engine/Private/Components/SceneCaptureComponent.cpp b/Engine/Source/Runtime/Engine/Private/Components/SceneCaptureComponent.cpp
--- a/Engine/Source/Runtime/Engine/Private/Components/SceneCaptureComponent.cpp
+++ b/Engine/Source/Runtime/Engine/Private/Components/SceneCaptureComponent.cpp
@@ -53,29 +53,29 @@ void ASceneCapture::PostLoad()
 		if (IsTemplate())
 		{
 			if (UBlueprintGeneratedClass* BPClass = Cast<UBlueprintGeneratedClass>(GetClass()))
-			{
+{
 				for (USCS_Node* RootNode : BPClass->SimpleConstructionScript->GetRootNodes())
-				{
+{
 					static const FName OldMeshName(TEXT("CamMesh0"));
 					static const FName OldFrustumName(TEXT("DrawFrust0"));
 					static const FName NewRootName(TEXT("SceneComponent"));
 					if (RootNode->ParentComponentOrVariableName == OldMeshName || RootNode->ParentComponentOrVariableName == OldFrustumName)
-					{
+	{
 						RootNode->ParentComponentOrVariableName = NewRootName;
 					}
 				}
-			}
-		}
+	}
+}
 
 		PRAGMA_DISABLE_DEPRECATION_WARNINGS
 		if (MeshComp_DEPRECATED)
 		{
 			MeshComp_DEPRECATED->SetStaticMesh(nullptr);
-		}
+			}
 		PRAGMA_ENABLE_DEPRECATION_WARNINGS
-	}
+		}
 #endif
-}
+	}
 
 void ASceneCapture::Serialize(FArchive& Ar)
 {
@@ -143,7 +143,6 @@ void USceneCaptureComponent::OnRegister()
 			ProxyMeshComponent->SetCollisionProfileName(UCollisionProfile::NoCollision_ProfileName);
 			ProxyMeshComponent->bHiddenInGame = true;
 			ProxyMeshComponent->CastShadow = false;
-			ProxyMeshComponent->PostPhysicsComponentTick.bCanEverTick = false;
 			ProxyMeshComponent->CreationMethod = CreationMethod;
 			ProxyMeshComponent->RegisterComponentWithWorld(GetWorld());
 		}
@@ -712,15 +711,15 @@ void APlanarReflection::PostLoad()
 	Super::PostLoad();
 
 	if (GetLinkerCustomVersion(FEditorObjectVersion::GUID) < FEditorObjectVersion::ChangeSceneCaptureRootComponent)
-	{
+{
 		if (PlanarReflectionComponent)
-		{
+	{
 			PRAGMA_DISABLE_DEPRECATION_WARNINGS
 			PlanarReflectionComponent->bShowPreviewPlane = bShowPreviewPlane_DEPRECATED;
 			PRAGMA_ENABLE_DEPRECATION_WARNINGS
 		}
-	}
-}
+			}
+		}
 
 void APlanarReflection::OnInterpToggle(bool bEnable)
 {
@@ -854,9 +853,9 @@ void UPlanarReflectionComponent::DestroyRenderState_Concurrent()
 		FPlanarReflectionSceneProxy* InSceneProxy = SceneProxy;
 		ENQUEUE_RENDER_COMMAND(FDestroyPlanarReflectionCommand)(
 			[InSceneProxy](FRHICommandList& RHICmdList)
-			{
+		{
 				delete InSceneProxy;
-			});
+		});
 
 		SceneProxy = nullptr;
 	}
diff --git a/Engine/Source/Runtime/NavigationSystem/Private/NavigationTestingActor.cpp b/Engine/Source/Runtime/NavigationSystem/Private/NavigationTestingActor.cpp
--- a/Engine/Source/Runtime/NavigationSystem/Private/NavigationTestingActor.cpp
+++ b/Engine/Source/Runtime/NavigationSystem/Private/NavigationTestingActor.cpp
@@ -32,7 +32,6 @@ ANavigationTestingActor::ANavigationTestingActor(const FObjectInitializer& Objec
 {
 #if WITH_EDITORONLY_DATA
 	EdRenderComp = CreateDefaultSubobject<UNavTestRenderingComponent>(TEXT("EdRenderComp"));
-	EdRenderComp->PostPhysicsComponentTick.bCanEverTick = false;
 
 #if WITH_RECAST
 	TickHelper = NULL;
