Add -usedcomponents to liststaticmeshes.  Dumps the staticmeshcomponent usage count for ustaticmeshes, and also all used components by path.

Strips Mobile data when running in builds that can't accurately count mobile data.

diff --git a/Engine/Source/Runtime/Engine/Private/UnrealEngine.cpp b/Engine/Source/Runtime/Engine/Private/UnrealEngine.cpp
--- a/Engine/Source/Runtime/Engine/Private/UnrealEngine.cpp
+++ b/Engine/Source/Runtime/Engine/Private/UnrealEngine.cpp
@@ -3261,24 +3261,26 @@ struct FCompareFSortedTexture
 */
 struct FSortedStaticMesh
 {
-	int32		NumKB;
-	int32		MaxKB;
-	int32		ResKBExc;
-	int32		ResKBInc;
-	int32		ResKBIncMobile;
-	int32		LodCount;
-	int32		MobileMinLOD;
-	int32		VertexCountLod0;
-	int32		VertexCountLod1;
-	int32		VertexCountLod2;
-	int32		VertexCountTotal;
-	int32		VertexCountTotalMobile;
-	int32		VertexCountCollision;
-	int32		ShapeCountCollision;
-	FString		Name;
+	int32			NumKB;
+	int32			MaxKB;
+	int32			ResKBExc;
+	int32			ResKBInc;
+	int32			ResKBIncMobile;
+	int32			LodCount;
+	int32			MobileMinLOD;
+	int32			VertexCountLod0;
+	int32			VertexCountLod1;
+	int32			VertexCountLod2;
+	int32			VertexCountTotal;
+	int32			VertexCountTotalMobile;
+	int32			VertexCountCollision;
+	int32			ShapeCountCollision;
+	int32			UsageCount;
+	UStaticMesh*	Mesh;
+	FString			Name;
 
 	/** Constructor, initializing every member variable with passed in values. */
-	FSortedStaticMesh(int32 InNumKB, int32 InMaxKB, int32 InResKBExc, int32 InResKBInc, int32 InResKBIncMobile, int32 InLodCount, int32 InMobileMinLOD, int32 InVertexCountLod0, int32 InVertexCountLod1, int32 InVertexCountLod2, int32 InVertexCountTotal, int32 InVertexCountTotalMobile, int32 InVertexCountCollision, int32 InShapeCountCollision, FString InName)
+	FSortedStaticMesh(UStaticMesh* InMesh, int32 InNumKB, int32 InMaxKB, int32 InResKBExc, int32 InResKBInc, int32 InResKBIncMobile, int32 InLodCount, int32 InMobileMinLOD, int32 InVertexCountLod0, int32 InVertexCountLod1, int32 InVertexCountLod2, int32 InVertexCountTotal, int32 InVertexCountTotalMobile, int32 InVertexCountCollision, int32 InShapeCountCollision, int32 InUsageCount, FString InName)
 		: NumKB(InNumKB)
 		, MaxKB(InMaxKB)
 		, ResKBExc(InResKBExc)
@@ -3293,6 +3295,8 @@ struct FSortedStaticMesh
 		, VertexCountTotalMobile(InVertexCountTotalMobile)
 		, VertexCountCollision(InVertexCountCollision)
 		, ShapeCountCollision(InShapeCountCollision)
+		, UsageCount(InUsageCount)
+		, Mesh(InMesh)		
 		, Name(InName)
 	{}
 };
@@ -4991,10 +4995,23 @@ bool UEngine::HandleListTexturesCommand( const TCHAR* Cmd, FOutputDevice& Ar )
 bool UEngine::HandleListStaticMeshesCommand(const TCHAR* Cmd, FOutputDevice& Ar)
 {
 	const bool bAlphaSort = FParse::Param(Cmd, TEXT("ALPHASORT"));
-	const bool bMobileSort = FParse::Param(Cmd, TEXT("MOBILESORT"));
-	const bool bCSV = FParse::Param(Cmd, TEXT("CSV"));
+	const bool bMobileSort = FParse::Param(Cmd, TEXT("MOBILESORT"));	
+	const bool bUsedComponents = FParse::Param(Cmd, TEXT("usedcomponents"));
+
+	//non-editor builds literally don't have the data to determine mobile lods or vert data.  The data prints out incorrectly and is confusing, just remove it.
+	const bool bHasMobileColumns = (bool)WITH_EDITORONLY_DATA;
+	Ar.Logf(TEXT("Listing all static meshes. Optional params: \n-alphasort: sort alphabetically \n-mobilesort: sort by mobile verts \n-usedcomponents: print the all components used by each mesh"));
 
-	Ar.Logf(TEXT("Listing all static meshes."));
+	//Collect usage counts
+	TMap<UStaticMesh*, TArray<UStaticMeshComponent*>> UsageList;
+	for (TObjectIterator<UStaticMeshComponent> It; It; ++It)
+	{
+		UStaticMeshComponent* MeshComponent = *It;
+		UStaticMesh* Mesh = MeshComponent->GetStaticMesh();
+
+		TArray<UStaticMeshComponent*>& MeshUsageArray = UsageList.FindOrAdd(Mesh);
+		MeshUsageArray.Add(MeshComponent);
+	}
 
 	// Collect meshes.
 	TArray<FSortedStaticMesh> SortedMeshes;
@@ -5002,6 +5019,9 @@ bool UEngine::HandleListStaticMeshesCommand(const TCHAR* Cmd, FOutputDevice& Ar)
 	{
 		UStaticMesh*		Mesh = *It;
 
+		const TArray<UStaticMeshComponent*>* MeshUsageList = UsageList.Find(Mesh);
+		int32 UsageCount = MeshUsageList ? MeshUsageList->Num() : 0;
+
 		FArchiveCountMem Count(Mesh);
 		FResourceSizeEx ResourceSizeExc = FResourceSizeEx(EResourceSizeMode::Exclusive);
 		Mesh->GetResourceSizeEx(ResourceSizeExc);
@@ -5055,6 +5075,7 @@ bool UEngine::HandleListStaticMeshesCommand(const TCHAR* Cmd, FOutputDevice& Ar)
 		FString				Name = Mesh->GetFullName();
 
 		new(SortedMeshes) FSortedStaticMesh(
+			Mesh,
 			NumKB,
 			MaxKB,
 			ResKBExc,
@@ -5069,6 +5090,7 @@ bool UEngine::HandleListStaticMeshesCommand(const TCHAR* Cmd, FOutputDevice& Ar)
 			VertexCountTotalMobile,
 			VertexCountCollision,
 			CollisionShapeCount,
+			UsageCount,
 			Name);
 	}
 
@@ -5084,34 +5106,53 @@ bool UEngine::HandleListStaticMeshesCommand(const TCHAR* Cmd, FOutputDevice& Ar)
 	int32 TotalVertexCount = 0;
 	int32 TotalVertexCountMobile = 0;
 
-	if(bCSV)
-	{
-		Ar.Logf(TEXT(",NumKB, MaxKB, ResKBExc, ResKBInc, ResKBIncMobile, LOD Count, MobileMinLOD, Verts LOD0, Verts LOD1, Verts LOD2, Verts Total, Verts Total Mobile, Verts Collision, Collision Shapes, Name"));
-	}
-	else
-	{
-		Ar.Logf(TEXT("       NumKB       MaxKB    ResKBExc    ResKBInc	ResKBIncMobile    NumLODs   MobileMinLOD  VertsLOD0   VertsLOD1   VertsLOD2   VertsTotal  VertsTotalMobile  VertsColl  CollisionShapes Name"));
-	}
+	FString HeaderString(TEXT(",       NumKB,       MaxKB,    ResKBExc,    ResKBInc,    LODCount,   VertsLOD0,   VertsLOD1,   VertsLOD2, Verts Total,  Verts Coll, Coll Shapes,     NumUsed"));
+	FString MobileHeaderString = bHasMobileColumns ? FString(TEXT(", ResKBIncMob,Verts Mobile,MobileMinLOD")) : FString();
+	
+	Ar.Logf(TEXT("%s%s, Name"), *HeaderString, *MobileHeaderString);	
 
 	for(int32 MeshIndex = 0; MeshIndex<SortedMeshes.Num(); MeshIndex++)
 	{
 		const FSortedStaticMesh& SortedMesh = SortedMeshes[MeshIndex];
 
-		if (bCSV)
-		{
-			Ar.Logf(TEXT(",%i, %i, %i, %i, %i, %i, %i, %i, %i, %i, %i, %i, %i, %i, %s"),
-				SortedMesh.NumKB, SortedMesh.MaxKB, SortedMesh.ResKBExc, SortedMesh.ResKBInc, SortedMesh.ResKBIncMobile, SortedMesh.LodCount, SortedMesh.MobileMinLOD,
-				SortedMesh.VertexCountLod0, SortedMesh.VertexCountLod1, SortedMesh.VertexCountLod2, SortedMesh.VertexCountTotal, SortedMesh.VertexCountTotalMobile, SortedMesh.VertexCountCollision, SortedMesh.ShapeCountCollision,
+		if (bHasMobileColumns)
+		{
+			Ar.Logf(TEXT(", %11i, %11i, %11i, %11i, %11i, %11i, %11i, %11i, %11i, %11i, %11i, %11i, %11i, %11i, %11i, %s"),
+				SortedMesh.NumKB,
+				SortedMesh.MaxKB,
+				SortedMesh.ResKBExc,
+				SortedMesh.ResKBInc,
+				SortedMesh.LodCount,
+				SortedMesh.VertexCountLod0,
+				SortedMesh.VertexCountLod1,
+				SortedMesh.VertexCountLod2,
+				SortedMesh.VertexCountTotal,
+				SortedMesh.VertexCountCollision,
+				SortedMesh.ShapeCountCollision,
+				SortedMesh.ResKBIncMobile,
+				SortedMesh.MobileMinLOD,
+				SortedMesh.VertexCountTotalMobile,
+				SortedMesh.UsageCount,
 				*SortedMesh.Name);
 		}
 		else
 		{
-			Ar.Logf(TEXT("%9i KB %8i KB %8i KB %8i KB  %11i KB %10i %14i %10i %11i %11i %12i %17i %10i %4i %s"),
-				SortedMesh.NumKB, SortedMesh.MaxKB, SortedMesh.ResKBExc, SortedMesh.ResKBInc, SortedMesh.ResKBIncMobile, SortedMesh.LodCount, SortedMesh.MobileMinLOD,
-				SortedMesh.VertexCountLod0, SortedMesh.VertexCountLod1, SortedMesh.VertexCountLod2, SortedMesh.VertexCountTotal, SortedMesh.VertexCountTotalMobile, SortedMesh.VertexCountCollision, SortedMesh.ShapeCountCollision,
+			Ar.Logf(TEXT(" ,%11i, %11i, %11i, %11i, %11i, %11i, %11i, %11i, %11i, %11i, %11i, %11i, %s"),
+				SortedMesh.NumKB,
+				SortedMesh.MaxKB,
+				SortedMesh.ResKBExc,
+				SortedMesh.ResKBInc,
+				SortedMesh.LodCount,
+				SortedMesh.VertexCountLod0,
+				SortedMesh.VertexCountLod1,
+				SortedMesh.VertexCountLod2,
+				SortedMesh.VertexCountTotal,
+				SortedMesh.VertexCountCollision,
+				SortedMesh.ShapeCountCollision,
+				SortedMesh.UsageCount,
 				*SortedMesh.Name);
 		}
-
+		
 		TotalNumKB += SortedMesh.NumKB;
 		TotalMaxKB += SortedMesh.MaxKB;
 		TotalResKBExc += SortedMesh.ResKBExc;
@@ -5123,6 +5164,28 @@ bool UEngine::HandleListStaticMeshesCommand(const TCHAR* Cmd, FOutputDevice& Ar)
 
 	Ar.Logf(TEXT("Total NumKB: %lld KB, Total MaxKB: %lld KB, Total ResKB Exc: %lld KB, Total ResKB Inc %lld KB,  Total ResKB Inc Mobile %lld KB, Total Vertex Count: %i, Total Vertex Count Mobile: %i, Static Mesh Count=%d"), TotalNumKB, TotalMaxKB, TotalResKBExc, TotalResKBInc, TotalResKBIncMobile, TotalVertexCount, TotalVertexCountMobile, SortedMeshes.Num());
 
+	if (bUsedComponents)
+	{
+		for (int32 MeshIndex = 0; MeshIndex < SortedMeshes.Num(); MeshIndex++)
+		{
+			const FSortedStaticMesh& SortedMesh = SortedMeshes[MeshIndex];
+			const TArray<UStaticMeshComponent*>* MeshUsageList = UsageList.Find(SortedMesh.Mesh);
+			if (MeshUsageList)
+			{
+				Ar.Logf(TEXT("%s mesh has %i UStaticMeshComponents referencing"), *SortedMesh.Name, MeshUsageList->Num());
+				for (int32 MeshComponentIndex = 0; MeshComponentIndex < MeshUsageList->Num(); ++MeshComponentIndex)
+				{
+					const UStaticMeshComponent* MeshComponent = (*MeshUsageList)[MeshComponentIndex];
+					Ar.Logf(TEXT("\t\t%s"), *MeshComponent->GetPathName());
+				}
+			}
+			else
+			{
+				Ar.Logf(TEXT("%s mesh has no UStaticMeshComponents referencing"), *SortedMesh.Name);
+			}
+		}
+	}
+
 	return true;
 }
 
