Resolve robomerge conflicts

diff --git a/Engine/Config/BaseEngine.ini b/Engine/Config/BaseEngine.ini
--- a/Engine/Config/BaseEngine.ini
+++ b/Engine/Config/BaseEngine.ini
@@ -912,7 +912,7 @@ SpawnPrioritySeconds=60.0
 +ChannelDefinitions=(ChannelName=Actor, ClassName=/Script/Engine.ActorChannel, StaticChannelIndex=-1, bTickOnCreate=false, bServerOpen=true, bClientOpen=false, bInitialServer=false, bInitialClient=false)
 
 [TextureStreaming]
-NeverStreamOutTextures=False
+NeverStreamOutRenderAssets=False
 MinTextureResidentMipCount=7
 PoolSize=160
 MemoryMargin=5
@@ -1437,13 +1437,13 @@ BeaconConnectionInitialTimeout=90.0
 BeaconConnectionTimeout=45.0
 
 [StaticMeshLODSettings]
-LevelArchitecture=(NumLODs=4,LightMapResolution=32,LODPercentTriangles=50,PixelError=12,SilhouetteImportance=4,Name=LOCTEXT("LevelArchitectureLOD","Level Architecture"))
-SmallProp=(NumLODs=4,LODPercentTriangles=50,PixelError=10,Name=LOCTEXT("SmallPropLOD","Small Prop"))
-LargeProp=(NumLODs=4,LODPercentTriangles=50,PixelError=10,Name=LOCTEXT("LargePropLOD","Large Prop"))
-Deco=(NumLODs=4,LODPercentTriangles=50,PixelError=10,Name=LOCTEXT("DecoLOD","Deco"))
-Vista=(NumLODs=1,Name=LOCTEXT("VistaLOD","Vista"))
-Foliage=(NumLODs=1,Name=LOCTEXT("FoliageLOD","Foliage"))
-HighDetail=(NumLODs=6,LODPercentTriangles=50,PixelError=6,Name=LOCTEXT("HighDetailLOD","High Detail"))
+LevelArchitecture=(NumLODs=4,MaxNumStreamedLODs=0,bSupportLODStreaming=0,LightMapResolution=32,LODPercentTriangles=50,PixelError=12,SilhouetteImportance=4,Name=LOCTEXT("LevelArchitectureLOD","Level Architecture"))
+SmallProp=(NumLODs=4,MaxNumStreamedLODs=0,bSupportLODStreaming=0,LODPercentTriangles=50,PixelError=10,Name=LOCTEXT("SmallPropLOD","Small Prop"))
+LargeProp=(NumLODs=4,MaxNumStreamedLODs=0,bSupportLODStreaming=0,LODPercentTriangles=50,PixelError=10,Name=LOCTEXT("LargePropLOD","Large Prop"))
+Deco=(NumLODs=4,MaxNumStreamedLODs=0,bSupportLODStreaming=0,LODPercentTriangles=50,PixelError=10,Name=LOCTEXT("DecoLOD","Deco"))
+Vista=(NumLODs=1,MaxNumStreamedLODs=0,bSupportLODStreaming=0,Name=LOCTEXT("VistaLOD","Vista"))
+Foliage=(NumLODs=1,MaxNumStreamedLODs=0,bSupportLODStreaming=0,Name=LOCTEXT("FoliageLOD","Foliage"))
+HighDetail=(NumLODs=6,MaxNumStreamedLODs=0,bSupportLODStreaming=0,LODPercentTriangles=50,PixelError=6,Name=LOCTEXT("HighDetailLOD","High Detail"))
 
 [TextureTracking]
 ;TextureName=T_GD_Traffic_Crosswalk_01
diff --git a/Engine/Plugins/2D/Paper2D/Source/Paper2D/Classes/PaperFlipbookComponent.h b/Engine/Plugins/2D/Paper2D/Source/Paper2D/Classes/PaperFlipbookComponent.h
--- a/Engine/Plugins/2D/Paper2D/Source/Paper2D/Classes/PaperFlipbookComponent.h
+++ b/Engine/Plugins/2D/Paper2D/Source/Paper2D/Classes/PaperFlipbookComponent.h
@@ -196,7 +196,7 @@ public:
 	virtual void GetUsedTextures(TArray<UTexture*>& OutTextures, EMaterialQualityLevel::Type QualityLevel) override;
 	virtual UMaterialInterface* GetMaterial(int32 MaterialIndex) const override;
 	virtual void GetUsedMaterials(TArray<UMaterialInterface*>& OutMaterials, bool bGetDebugMaterials = false) const override;
-	virtual void GetStreamingTextureInfo(FStreamingTextureLevelContext& LevelContext, TArray<FStreamingTexturePrimitiveInfo>& OutStreamingTextures) const override;
+	virtual void GetStreamingRenderAssetInfo(FStreamingTextureLevelContext& LevelContext, TArray<FStreamingRenderAssetPrimitiveInfo>& OutStreamingTextures) const override;
 	virtual int32 GetNumMaterials() const override;
 	virtual UBodySetup* GetBodySetup() override;
 	// End of UPrimitiveComponent interface
diff --git a/Engine/Plugins/2D/Paper2D/Source/Paper2D/Classes/PaperSpriteComponent.h b/Engine/Plugins/2D/Paper2D/Source/Paper2D/Classes/PaperSpriteComponent.h
--- a/Engine/Plugins/2D/Paper2D/Source/Paper2D/Classes/PaperSpriteComponent.h
+++ b/Engine/Plugins/2D/Paper2D/Source/Paper2D/Classes/PaperSpriteComponent.h
@@ -91,7 +91,7 @@ public:
 	virtual void GetUsedTextures(TArray<UTexture*>& OutTextures, EMaterialQualityLevel::Type QualityLevel) override;
 	virtual UMaterialInterface* GetMaterial(int32 MaterialIndex) const override;
 	virtual void GetUsedMaterials(TArray<UMaterialInterface*>& OutMaterials, bool bGetDebugMaterials = false) const override;
-	virtual void GetStreamingTextureInfo(FStreamingTextureLevelContext& LevelContext, TArray<FStreamingTexturePrimitiveInfo>& OutStreamingTextures) const override;
+	virtual void GetStreamingRenderAssetInfo(FStreamingTextureLevelContext& LevelContext, TArray<FStreamingRenderAssetPrimitiveInfo>& OutStreamingRenderAssets) const override;
 	virtual int32 GetNumMaterials() const override;
 	// End of UPrimitiveComponent interface
 
diff --git a/Engine/Plugins/2D/Paper2D/Source/Paper2D/Private/PaperFlipbookComponent.cpp b/Engine/Plugins/2D/Paper2D/Source/Paper2D/Private/PaperFlipbookComponent.cpp
--- a/Engine/Plugins/2D/Paper2D/Source/Paper2D/Private/PaperFlipbookComponent.cpp
+++ b/Engine/Plugins/2D/Paper2D/Source/Paper2D/Private/PaperFlipbookComponent.cpp
@@ -171,10 +171,10 @@ void UPaperFlipbookComponent::GetUsedMaterials(TArray<UMaterialInterface*>& OutM
 	return Super::GetUsedMaterials(OutMaterials, bGetDebugMaterials);
 }
 
-void UPaperFlipbookComponent::GetStreamingTextureInfo(FStreamingTextureLevelContext& LevelContext, TArray<FStreamingTexturePrimitiveInfo>& OutStreamingTextures) const
+void UPaperFlipbookComponent::GetStreamingRenderAssetInfo(FStreamingTextureLevelContext& LevelContext, TArray<FStreamingRenderAssetPrimitiveInfo>& OutStreamingRenderAssets) const
 {
 	//@TODO: PAPER2D: Need to support this for proper texture streaming
-	return Super::GetStreamingTextureInfo(LevelContext, OutStreamingTextures);
+	return Super::GetStreamingRenderAssetInfo(LevelContext, OutStreamingRenderAssets);
 }
 
 int32 UPaperFlipbookComponent::GetNumMaterials() const
diff --git a/Engine/Plugins/2D/Paper2D/Source/Paper2D/Private/PaperSpriteComponent.cpp b/Engine/Plugins/2D/Paper2D/Source/Paper2D/Private/PaperSpriteComponent.cpp
--- a/Engine/Plugins/2D/Paper2D/Source/Paper2D/Private/PaperSpriteComponent.cpp
+++ b/Engine/Plugins/2D/Paper2D/Source/Paper2D/Private/PaperSpriteComponent.cpp
@@ -274,10 +274,10 @@ void UPaperSpriteComponent::GetUsedMaterials(TArray<UMaterialInterface*>& OutMat
 	return Super::GetUsedMaterials(OutMaterials, bGetDebugMaterials);
 }
 
-void UPaperSpriteComponent::GetStreamingTextureInfo(FStreamingTextureLevelContext& LevelContext, TArray<FStreamingTexturePrimitiveInfo>& OutStreamingTextures) const
+void UPaperSpriteComponent::GetStreamingRenderAssetInfo(FStreamingTextureLevelContext& LevelContext, TArray<FStreamingRenderAssetPrimitiveInfo>& OutStreamingRenderAssets) const
 {
 	//@TODO: PAPER2D: Need to support this for proper texture streaming
-	return Super::GetStreamingTextureInfo(LevelContext, OutStreamingTextures);
+	return Super::GetStreamingRenderAssetInfo(LevelContext, OutStreamingRenderAssets);
 }
 
 int32 UPaperSpriteComponent::GetNumMaterials() const
diff --git a/Engine/Source/Developer/TargetPlatform/Public/Common/TargetPlatformBase.h b/Engine/Source/Developer/TargetPlatform/Public/Common/TargetPlatformBase.h
--- a/Engine/Source/Developer/TargetPlatform/Public/Common/TargetPlatformBase.h
+++ b/Engine/Source/Developer/TargetPlatform/Public/Common/TargetPlatformBase.h
@@ -243,6 +243,8 @@ public:
 
 		case ETargetPlatformFeatures::TextureStreaming:
 			return TPlatformProperties::SupportsTextureStreaming();
+		case ETargetPlatformFeatures::MeshLODStreaming:
+			return TPlatformProperties::SupportsMeshLODStreaming();
 
 		case ETargetPlatformFeatures::SdkConnectDisconnect:
 		case ETargetPlatformFeatures::UserCredentials:
diff --git a/Engine/Source/Developer/TargetPlatform/Public/Interfaces/ITargetPlatform.h b/Engine/Source/Developer/TargetPlatform/Public/Interfaces/ITargetPlatform.h
--- a/Engine/Source/Developer/TargetPlatform/Public/Interfaces/ITargetPlatform.h
+++ b/Engine/Source/Developer/TargetPlatform/Public/Interfaces/ITargetPlatform.h
@@ -52,6 +52,9 @@ enum class ETargetPlatformFeatures
 	/** Texture streaming. */
 	TextureStreaming,
 
+	/** Mesh LOD streaming. */
+	MeshLODStreaming,
+
 	/** User credentials are required to use the device. */
 	UserCredentials,
 
diff --git a/Engine/Source/Editor/StaticMeshEditor/Private/StaticMeshEditorTools.cpp b/Engine/Source/Editor/StaticMeshEditor/Private/StaticMeshEditorTools.cpp
--- a/Engine/Source/Editor/StaticMeshEditor/Private/StaticMeshEditorTools.cpp
+++ b/Engine/Source/Editor/StaticMeshEditor/Private/StaticMeshEditorTools.cpp
@@ -3070,6 +3070,24 @@ void FLevelOfDetailSettingsLayout::AddToDetailsPanel( IDetailLayoutBuilder& Deta
 		.PlatformOverrideNames(this, &FLevelOfDetailSettingsLayout::GetMinLODPlatformOverrideNames)
 	];
 
+	LODSettingsCategory.AddCustomRow(LOCTEXT("NumStreamedLODs", "Num Streamed LODs"))
+	.NameContent()
+	[
+		SNew(STextBlock)
+		.Font(IDetailLayoutBuilder::GetDetailFont())
+		.Text(LOCTEXT("NumStreamdLODs", "Num Streamed LODs"))
+	]
+	.ValueContent()
+	.MinDesiredWidth((float)(StaticMesh->NumStreamedLODs.PerPlatform.Num() + 1)*125.0f)
+	.MaxDesiredWidth((float)(PlatformNumber + 1)*125.0f)
+	[
+		SNew(SPerPlatformPropertiesWidget)
+		.OnGenerateWidget(this, &FLevelOfDetailSettingsLayout::GetNumStreamedLODsWidget)
+		.OnAddPlatform(this, &FLevelOfDetailSettingsLayout::AddNumStreamedLODsPlatformOverride)
+		.OnRemovePlatform(this, &FLevelOfDetailSettingsLayout::RemoveNumStreamedLODsPlatformOverride)
+		.PlatformOverrideNames(this, &FLevelOfDetailSettingsLayout::GetNumStreamedLODsPlatformOverrideNames)
+	];
+
 	// Add Number of LODs slider.
 	const int32 MinAllowedLOD = 1;
 	LODSettingsCategory.AddCustomRow( LOCTEXT("NumberOfLODs", "Number of LODs") )
@@ -4043,6 +4061,118 @@ TArray<FName> FLevelOfDetailSettingsLayout::GetMinLODPlatformOverrideNames() con
 	return KeyArray;
 }
 
+/** @return - whether value was different */
+static bool UpdateStaticMeshNumStreamedLODsHelper(UStaticMesh* StaticMesh, int32 NewValue, FName Platform)
+{
+	bool bWasDifferent = false;
+	StaticMesh->Modify();
+	{
+		FStaticMeshComponentRecreateRenderStateContext ReregisterContext(StaticMesh, false);
+		NewValue = FMath::Clamp<int32>(NewValue, -1, MAX_STATIC_MESH_LODS);
+		if (Platform == NAME_None)
+		{
+			bWasDifferent = StaticMesh->NumStreamedLODs.Default != NewValue;
+			StaticMesh->NumStreamedLODs.Default = NewValue;
+		}
+		else
+		{
+			int32* ValuePtr = StaticMesh->NumStreamedLODs.PerPlatform.Find(Platform);
+			if (ValuePtr != nullptr)
+			{
+				bWasDifferent = *ValuePtr != NewValue;
+				*ValuePtr = NewValue;
+			}
+		}
+	}
+	return bWasDifferent;
+}
+
+void FLevelOfDetailSettingsLayout::OnNumStreamedLODsChanged(int32 NewValue, FName Platform)
+{
+	UStaticMesh* StaticMesh = StaticMeshEditor.GetStaticMesh();
+	check(StaticMesh);
+	UpdateStaticMeshNumStreamedLODsHelper(StaticMesh, NewValue, Platform);
+	StaticMeshEditor.RefreshViewport();
+}
+
+void FLevelOfDetailSettingsLayout::OnNumStreamedLODsCommitted(int32 InValue, ETextCommit::Type CommitInfo, FName Platform)
+{
+	UStaticMesh* StaticMesh = StaticMeshEditor.GetStaticMesh();
+	check(StaticMesh);
+	if (UpdateStaticMeshNumStreamedLODsHelper(StaticMesh, InValue, Platform))
+	{
+		ApplyChanges();
+	}
+	StaticMeshEditor.RefreshViewport();
+}
+
+int32 FLevelOfDetailSettingsLayout::GetNumStreamedLODs(FName Platform) const
+{
+	UStaticMesh* StaticMesh = StaticMeshEditor.GetStaticMesh();
+	check(StaticMesh);
+
+	int32* ValuePtr = (Platform == NAME_None) ? nullptr : StaticMesh->NumStreamedLODs.PerPlatform.Find(Platform);
+	return (ValuePtr != nullptr) ? *ValuePtr : StaticMesh->NumStreamedLODs.Default;
+}
+
+TSharedRef<SWidget> FLevelOfDetailSettingsLayout::GetNumStreamedLODsWidget(FName PlatformGroupName) const
+{
+	return SNew(SSpinBox<int32>)
+		.Font(IDetailLayoutBuilder::GetDetailFont())
+		.Value(this, &FLevelOfDetailSettingsLayout::GetNumStreamedLODs, PlatformGroupName)
+		.OnValueChanged(this, &FLevelOfDetailSettingsLayout::OnNumStreamedLODsChanged, PlatformGroupName)
+		.OnValueCommitted(this, &FLevelOfDetailSettingsLayout::OnNumStreamedLODsCommitted, PlatformGroupName)
+		.MinValue(-1)
+		.MaxValue(MAX_STATIC_MESH_LODS)
+		.ToolTipText(this, &FLevelOfDetailSettingsLayout::GetNumStreamedLODsTooltip)
+		.IsEnabled(FLevelOfDetailSettingsLayout::GetLODCount() > 1);
+}
+
+bool FLevelOfDetailSettingsLayout::AddNumStreamedLODsPlatformOverride(FName PlatformGroupName)
+{
+	FScopedTransaction Transaction(LOCTEXT("AddNumStreamedLODsPlatformOverride", "Add NumStreamdLODs Platform Override"));
+	UStaticMesh* StaticMesh = StaticMeshEditor.GetStaticMesh();
+	check(StaticMesh);
+	StaticMesh->Modify();
+	if (StaticMesh->NumStreamedLODs.PerPlatform.Find(PlatformGroupName) == nullptr)
+	{
+		float Value = StaticMesh->NumStreamedLODs.Default;
+		StaticMesh->NumStreamedLODs.PerPlatform.Add(PlatformGroupName, Value);
+		OnNumStreamedLODsChanged(Value, PlatformGroupName);
+		return true;
+	}
+	return false;
+}
+
+bool FLevelOfDetailSettingsLayout::RemoveNumStreamedLODsPlatformOverride(FName PlatformGroupName)
+{
+	FScopedTransaction Transaction(LOCTEXT("RemoveNumStreamedLODsPlatformOverride", "Remove NumStreamedLODs Platform Override"));
+	UStaticMesh* StaticMesh = StaticMeshEditor.GetStaticMesh();
+	check(StaticMesh);
+	StaticMesh->Modify();
+	if (StaticMesh->NumStreamedLODs.PerPlatform.Remove(PlatformGroupName) != 0)
+	{
+		OnNumStreamedLODsChanged(StaticMesh->NumStreamedLODs.Default, PlatformGroupName);
+		return true;
+	}
+	return false;
+}
+
+TArray<FName> FLevelOfDetailSettingsLayout::GetNumStreamedLODsPlatformOverrideNames() const
+{
+	UStaticMesh* StaticMesh = StaticMeshEditor.GetStaticMesh();
+	check(StaticMesh);
+	TArray<FName> KeyArray;
+	StaticMesh->NumStreamedLODs.PerPlatform.GenerateKeyArray(KeyArray);
+	KeyArray.Sort();
+	return KeyArray;
+}
+
+FText FLevelOfDetailSettingsLayout::GetNumStreamedLODsTooltip() const
+{
+	return LOCTEXT("NumStreamedLODsTooltip", "If non-negative, the number of LODs that can be streamed. Only has effect if mesh LOD streaming is enabled on the target platform.");
+}
+
 FText FLevelOfDetailSettingsLayout::GetLODCustomModeNameContent(int32 LODIndex) const
 {
 	int32 CurrentLodIndex = 0;
diff --git a/Engine/Source/Editor/StaticMeshEditor/Private/StaticMeshEditorTools.h b/Engine/Source/Editor/StaticMeshEditor/Private/StaticMeshEditorTools.h
--- a/Engine/Source/Editor/StaticMeshEditor/Private/StaticMeshEditorTools.h
+++ b/Engine/Source/Editor/StaticMeshEditor/Private/StaticMeshEditorTools.h
@@ -523,6 +523,14 @@ private:
 	bool RemoveMinLODPlatformOverride(FName PlatformGroupName);
 	TArray<FName> GetMinLODPlatformOverrideNames() const;
 
+	void OnNumStreamedLODsChanged(int32 NewValue, FName Platform);
+	void OnNumStreamedLODsCommitted(int32 InValue, ETextCommit::Type CommitInfo, FName Platform);
+	int32 GetNumStreamedLODs(FName Platform) const;
+	TSharedRef<SWidget> GetNumStreamedLODsWidget(FName PlatformGroupName) const;
+	bool AddNumStreamedLODsPlatformOverride(FName PlatformGroupName);
+	bool RemoveNumStreamedLODsPlatformOverride(FName PlatformGroupName);
+	TArray<FName> GetNumStreamedLODsPlatformOverrideNames() const;
+
 	bool CanRemoveLOD(int32 LODIndex) const;
 	FReply OnRemoveLOD(int32 LODIndex);
 
@@ -551,6 +559,7 @@ private:
 	void UpdateLODNames();
 	FText GetLODCountTooltip() const;
 	FText GetMinLODTooltip() const;
+	FText GetNumStreamedLODsTooltip() const;
 
 	FText GetLODCustomModeNameContent(int32 LODIndex) const;
 	ECheckBoxState IsLODCustomModeCheck(int32 LODIndex) const;
diff --git a/Engine/Source/Editor/UnrealEd/Private/TextureStreamingNotification.cpp b/Engine/Source/Editor/UnrealEd/Private/TextureStreamingNotification.cpp
--- a/Engine/Source/Editor/UnrealEd/Private/TextureStreamingNotification.cpp
+++ b/Engine/Source/Editor/UnrealEd/Private/TextureStreamingNotification.cpp
@@ -50,7 +50,7 @@ int32 FTextureStreamingNotificationImpl::GetNumStreamingTextures()
 
 	if (StreamingManagerCollection.IsTextureStreamingEnabled())
 	{
-		ITextureStreamingManager& TextureStreamingManager = StreamingManagerCollection.GetTextureStreamingManager();
+		IRenderAssetStreamingManager& TextureStreamingManager = StreamingManagerCollection.GetTextureStreamingManager();
 		return TextureStreamingManager.GetNumWantingResources();
 	}
 
diff --git a/Engine/Source/Runtime/Core/Public/GenericPlatform/GenericPlatformProperties.h b/Engine/Source/Runtime/Core/Public/GenericPlatform/GenericPlatformProperties.h
--- a/Engine/Source/Runtime/Core/Public/GenericPlatform/GenericPlatformProperties.h
+++ b/Engine/Source/Runtime/Core/Public/GenericPlatform/GenericPlatformProperties.h
@@ -239,6 +239,11 @@ struct FGenericPlatformProperties
 		return true;
 	}
 
+	static FORCEINLINE bool SupportsMeshLODStreaming()
+	{
+		return true;
+	}
+
 	/**
 	 * Gets whether user settings should override the resolution or not
 	 */
diff --git a/Engine/Source/Runtime/D3D12RHI/Private/D3D12Buffer.cpp b/Engine/Source/Runtime/D3D12RHI/Private/D3D12Buffer.cpp
--- a/Engine/Source/Runtime/D3D12RHI/Private/D3D12Buffer.cpp
+++ b/Engine/Source/Runtime/D3D12RHI/Private/D3D12Buffer.cpp
@@ -222,15 +222,43 @@ BufferType* FD3D12Adapter::CreateRHIBuffer(FRHICommandListImmediate* RHICmdList,
 				}
 			};
 			
-			if (RHICmdList && !RHICmdList->Bypass())
+			if (!RHICmdList)
 			{
-				ALLOC_COMMAND_CL(*RHICmdList, FD3D12RHICommandInitializeBuffer)(BufferOut, SrcResourceLoc, Size);
+				// Need to update buffer content on RHI thread (immediate context) because the buffer can be a
+				// sub-allocation and its backing resource may be in a state incompatible with the copy queue.
+				// TODO:
+				// Create static buffers in COMMON state, rely on state promotion/decay to avoid transition barriers,
+				// and initialize them asynchronously on the copy queue. D3D12 buffers always allow simultaneous acess
+				// so it is legal to write to a region on the copy queue while other non-overlapping regions are
+				// being read on the graphics/compute queue. Currently, d3ddebug throws error for such usage.
+				// Once Microsoft (via Windows update) fix the debug layer, async static buffer initialization should
+				// be done on the copy queue.
+				FD3D12ResourceLocation* SrcResourceLoc_Heap = new FD3D12ResourceLocation(SrcResourceLoc.GetParentDevice());
+				FD3D12ResourceLocation::TransferOwnership(*SrcResourceLoc_Heap, SrcResourceLoc);
+				ENQUEUE_RENDER_COMMAND(CmdD3D12InitializeBuffer)(
+					[BufferOut, SrcResourceLoc_Heap, Size](FRHICommandListImmediate& RHICmdList)
+				{
+					if (RHICmdList.Bypass())
+					{
+						FD3D12RHICommandInitializeBuffer Command(BufferOut, *SrcResourceLoc_Heap, Size);
+						Command.ExecuteNoCmdList();
+					}
+					else
+					{
+						new (RHICmdList.AllocCommand<FD3D12RHICommandInitializeBuffer>()) FD3D12RHICommandInitializeBuffer(BufferOut, *SrcResourceLoc_Heap, Size);
+					}
+					delete SrcResourceLoc_Heap;
+				});
 			}
-			else
+			else if (RHICmdList->Bypass())
 			{
 				FD3D12RHICommandInitializeBuffer Command(BufferOut, SrcResourceLoc, Size);
 				Command.ExecuteNoCmdList();
 			}
+			else
+			{
+				new (RHICmdList->AllocCommand<FD3D12RHICommandInitializeBuffer>()) FD3D12RHICommandInitializeBuffer(BufferOut, SrcResourceLoc, Size);
+			}
 		}
 
 		// Discard the resource array's contents.
diff --git a/Engine/Source/Runtime/D3D12RHI/Private/D3D12RHICommon.h b/Engine/Source/Runtime/D3D12RHI/Private/D3D12RHICommon.h
--- a/Engine/Source/Runtime/D3D12RHI/Private/D3D12RHICommon.h
+++ b/Engine/Source/Runtime/D3D12RHI/Private/D3D12RHICommon.h
@@ -64,6 +64,11 @@ public:
 		return Parent;
 	}
 
+	FD3D12Device* GetParentDevice_Unsafe() const
+	{
+		return Parent;
+	}
+
 	// To be used with delayed setup
 	inline void SetParentDevice(FD3D12Device* InParent)
 	{
diff --git a/Engine/Source/Runtime/D3D12RHI/Private/D3D12RHIPrivate.h b/Engine/Source/Runtime/D3D12RHI/Private/D3D12RHIPrivate.h
--- a/Engine/Source/Runtime/D3D12RHI/Private/D3D12RHIPrivate.h
+++ b/Engine/Source/Runtime/D3D12RHI/Private/D3D12RHIPrivate.h
@@ -290,6 +290,7 @@ public:
 	virtual FShaderResourceViewRHIRef RHICreateShaderResourceView(FStructuredBufferRHIParamRef StructuredBuffer) final override;
 	virtual FShaderResourceViewRHIRef RHICreateShaderResourceView(FVertexBufferRHIParamRef VertexBuffer, uint32 Stride, uint8 Format) final override;
 	virtual FShaderResourceViewRHIRef RHICreateShaderResourceView(FIndexBufferRHIParamRef Buffer) final override;
+	virtual void RHIUpdateShaderResourceView(FShaderResourceViewRHIParamRef SRV, FVertexBufferRHIParamRef VertexBuffer, uint32 Stride, uint8 Format) final override;
 	virtual uint64 RHICalcTexture2DPlatformSize(uint32 SizeX, uint32 SizeY, uint8 Format, uint32 NumMips, uint32 NumSamples, uint32 Flags, uint32& OutAlign) override;
 	virtual uint64 RHICalcTexture3DPlatformSize(uint32 SizeX, uint32 SizeY, uint32 SizeZ, uint8 Format, uint32 NumMips, uint32 Flags, uint32& OutAlign) final override;
 	virtual uint64 RHICalcTextureCubePlatformSize(uint32 Size, uint8 Format, uint32 NumMips, uint32 Flags, uint32& OutAlign) final override;
diff --git a/Engine/Source/Runtime/D3D12RHI/Private/D3D12View.cpp b/Engine/Source/Runtime/D3D12RHI/Private/D3D12View.cpp
--- a/Engine/Source/Runtime/D3D12RHI/Private/D3D12View.cpp
+++ b/Engine/Source/Runtime/D3D12RHI/Private/D3D12View.cpp
@@ -243,6 +243,11 @@ FShaderResourceViewRHIRef FD3D12DynamicRHI::RHICreateShaderResourceView(FVertexB
 		}
 	};
 
+	if (!VertexBufferRHI)
+	{
+		return new FD3D12ShaderResourceView(nullptr);
+	}
+
 	FD3D12VertexBuffer* VertexBuffer = FD3D12DynamicRHI::ResourceCast(VertexBufferRHI);
 	return GetAdapter().CreateLinkedViews<FD3D12VertexBuffer, FD3D12ShaderResourceView>(VertexBuffer,
 	[Stride, Format](FD3D12VertexBuffer* VertexBuffer)
@@ -321,6 +326,47 @@ FShaderResourceViewRHIRef FD3D12DynamicRHI::RHICreateShaderResourceView(FIndexBu
 	});
 }
 
+void FD3D12DynamicRHI::RHIUpdateShaderResourceView(FShaderResourceViewRHIParamRef SRV, FVertexBufferRHIParamRef VertexBuffer, uint32 Stride, uint8 Format)
+{
+	check(SRV);
+	if (VertexBuffer)
+	{
+		FD3D12VertexBuffer* VBD3D12 = ResourceCast(VertexBuffer);
+		FD3D12ShaderResourceView* SRVD3D12 = ResourceCast(SRV);
+		D3D12_SHADER_RESOURCE_VIEW_DESC SRVDesc = {};
+		SRVDesc.Shader4ComponentMapping = D3D12_DEFAULT_SHADER_4_COMPONENT_MAPPING;
+		SRVDesc.ViewDimension = D3D12_SRV_DIMENSION_BUFFER;
+		if (VBD3D12->GetUsage() & BUF_ByteAddressBuffer)
+		{
+			SRVDesc.Format = DXGI_FORMAT_R32_TYPELESS;
+			SRVDesc.Buffer.NumElements = VBD3D12->GetSize() / 4;
+			SRVDesc.Buffer.Flags = D3D12_BUFFER_SRV_FLAG_RAW;
+			Stride = 4;
+		}
+		else
+		{
+			SRVDesc.Format = FindShaderResourceDXGIFormat((DXGI_FORMAT)GPixelFormats[Format].PlatformFormat, false);
+			SRVDesc.Buffer.NumElements = VBD3D12->GetSize() / Stride;
+		}
+		SRVDesc.Buffer.StructureByteStride = 0;
+		SRVDesc.Buffer.FirstElement = VBD3D12->ResourceLocation.GetOffsetFromBaseOfResource() / Stride;
+
+		// Rename the SRV to view on the new vertex buffer
+		while (VBD3D12)
+		{
+			FD3D12Device* ParentDevice = VBD3D12->GetParentDevice();
+			SRVD3D12->Initialize(ParentDevice, SRVDesc, VBD3D12->ResourceLocation, Stride);
+			VBD3D12->SetDynamicSRV(SRVD3D12);
+			VBD3D12 = VBD3D12->GetNextObject();
+			if (VBD3D12 && !SRVD3D12->GetNextObject())
+			{
+				SRVD3D12->SetNextObject(new FD3D12ShaderResourceView(ParentDevice));
+			}
+			SRVD3D12 = SRVD3D12->GetNextObject();
+		}
+	}
+}
+
 FShaderResourceViewRHIRef FD3D12DynamicRHI::RHICreateShaderResourceView_RenderThread(FRHICommandListImmediate& RHICmdList, FTexture2DRHIParamRef Texture2DRHI, uint8 MipLevel)
 {
 	return RHICreateShaderResourceView(Texture2DRHI, MipLevel);
diff --git a/Engine/Source/Runtime/D3D12RHI/Private/D3D12View.h b/Engine/Source/Runtime/D3D12RHI/Private/D3D12View.h
--- a/Engine/Source/Runtime/D3D12RHI/Private/D3D12View.h
+++ b/Engine/Source/Runtime/D3D12RHI/Private/D3D12View.h
@@ -648,20 +648,20 @@ public:
 	TD3D12ViewDescriptorHandle(FD3D12Device* InParentDevice)
 		: FD3D12DeviceChild(InParentDevice)
 	{
-		// Allocate descriptor slot
-		FD3D12Device* Device = GetParentDevice();
-		FD3D12OfflineDescriptorManager& DescriptorAllocator = Device->template GetViewDescriptorAllocator<TDesc>();
-		Handle = DescriptorAllocator.AllocateHeapSlot(Index);
-		check(Handle.ptr != 0);
+		Handle.ptr = 0;
+		AllocateDescriptorSlot();
 	}
 
 	~TD3D12ViewDescriptorHandle()
 	{
-		// Free descriptor slot
-		FD3D12Device* Device = GetParentDevice();
-		FD3D12OfflineDescriptorManager& DescriptorAllocator = Device->template GetViewDescriptorAllocator<TDesc>();
-		DescriptorAllocator.FreeHeapSlot(Handle, Index);
-		Handle.ptr = 0;
+		FreeDescriptorSlot();
+	}
+
+	void SetParentDevice(FD3D12Device* InParent)
+	{
+		check(!Parent && !Handle.ptr);
+		FD3D12DeviceChild::SetParentDevice(InParent);
+		AllocateDescriptorSlot();
 	}
 
 	void CreateView(const TDesc& Desc, ID3D12Resource* Resource)
@@ -684,6 +684,30 @@ public:
 
 	inline const CD3DX12_CPU_DESCRIPTOR_HANDLE& GetHandle() const { return Handle; }
 	inline uint32 GetIndex() const { return Index; }
+
+private:
+	void AllocateDescriptorSlot()
+	{
+		if (Parent)
+		{
+			FD3D12Device* Device = GetParentDevice();
+			FD3D12OfflineDescriptorManager& DescriptorAllocator = Device->template GetViewDescriptorAllocator<TDesc>();
+			Handle = DescriptorAllocator.AllocateHeapSlot(Index);
+			check(Handle.ptr != 0);
+		}
+	}
+
+	void FreeDescriptorSlot()
+	{
+		if (Parent)
+		{
+			FD3D12Device* Device = GetParentDevice();
+			FD3D12OfflineDescriptorManager& DescriptorAllocator = Device->template GetViewDescriptorAllocator<TDesc>();
+			DescriptorAllocator.FreeHeapSlot(Handle, Index);
+			Handle.ptr = 0;
+		}
+		check(!Handle.ptr);
+	}
 };
 
 typedef TD3D12ViewDescriptorHandle<D3D12_SHADER_RESOURCE_VIEW_DESC>		FD3D12DescriptorHandleSRV;
@@ -765,6 +789,7 @@ protected:
 
 public:
 	inline FD3D12Device*					GetParentDevice()			const { return Descriptor.GetParentDevice(); }
+	inline FD3D12Device*					GetParentDevice_Unsafe()	const { return Descriptor.GetParentDevice_Unsafe(); }
 	inline const TDesc&						GetDesc()					const { check(bInitialized); return Desc; }
 	inline CD3DX12_CPU_DESCRIPTOR_HANDLE	GetView()					const { check(bInitialized); return Descriptor.GetHandle(); }
 	inline uint32							GetDescriptorHeapIndex()	const { check(bInitialized); return Descriptor.GetIndex(); }
@@ -773,6 +798,11 @@ public:
 	inline FD3D12ResidencyHandle*			GetResidencyHandle()		const { check(bInitialized); return ResidencyHandle; }
 	inline const CViewSubresourceSubset&	GetViewSubresourceSubset()	const { check(bInitialized); return ViewSubresourceSubset; }
 
+	void SetParentDevice(FD3D12Device* InParent)
+	{
+		Descriptor.SetParentDevice(InParent);
+	}
+
 	template< class T >
 	inline bool DoesNotOverlap(const FD3D12View< T >& Other) const
 	{
@@ -826,6 +856,18 @@ public:
 		CreateView(InDesc, InResourceLocation);
 	}
 
+	void Initialize(FD3D12Device* InParent, D3D12_SHADER_RESOURCE_VIEW_DESC& InDesc, FD3D12ResourceLocation& InResourceLocation, uint32 InStride, bool InSkipFastClearFinalize = false)
+	{
+		if (!this->GetParentDevice_Unsafe())
+		{
+			// This is a null SRV created without viewing on any resource
+			// We need to set its device and allocate a descriptor slot before moving forward
+			this->SetParentDevice(InParent);
+		}
+		check(GetParentDevice() == InParent);
+		Initialize(InDesc, InResourceLocation, InStride, InSkipFastClearFinalize);
+	}
+
 	void Rename(FD3D12ResourceLocation& InResourceLocation)
 	{
 		// Update the first element index, then reinitialize the SRV
diff --git a/Engine/Source/Runtime/Engine/Classes/Components/InstancedStaticMeshComponent.h b/Engine/Source/Runtime/Engine/Classes/Components/InstancedStaticMeshComponent.h
--- a/Engine/Source/Runtime/Engine/Classes/Components/InstancedStaticMeshComponent.h
+++ b/Engine/Source/Runtime/Engine/Classes/Components/InstancedStaticMeshComponent.h
@@ -174,7 +174,7 @@ class ENGINE_API UInstancedStaticMeshComponent : public UStaticMeshComponent
 	/** Build the data to compute accuracte StreaminTexture data. */
 	virtual bool BuildTextureStreamingData(ETextureStreamingBuildType BuildType, EMaterialQualityLevel::Type QualityLevel, ERHIFeatureLevel::Type FeatureLevel, TSet<FGuid>& DependentResources) override;
 	/** Get the StreaminTexture data. */
-	virtual void GetStreamingTextureInfo(FStreamingTextureLevelContext& LevelContext, TArray<FStreamingTexturePrimitiveInfo>& OutStreamingTextures) const override;
+	virtual void GetStreamingRenderAssetInfo(FStreamingTextureLevelContext& LevelContext, TArray<FStreamingRenderAssetPrimitiveInfo>& OutStreamingRenderAssets) const override;
 
 	/**
 	* Update the transform for the instance specified.
diff --git a/Engine/Source/Runtime/Engine/Classes/Components/MeshComponent.h b/Engine/Source/Runtime/Engine/Classes/Components/MeshComponent.h
--- a/Engine/Source/Runtime/Engine/Classes/Components/MeshComponent.h
+++ b/Engine/Source/Runtime/Engine/Classes/Components/MeshComponent.h
@@ -89,7 +89,7 @@ class ENGINE_API UMeshComponent : public UPrimitiveComponent
 	virtual bool GetMaterialStreamingData(int32 MaterialIndex, FPrimitiveMaterialInfo& MaterialData) const { return false; }
 
 	/** Generate streaming data for all materials. */
-	void GetStreamingTextureInfoInner(FStreamingTextureLevelContext& LevelContext, const TArray<FStreamingTextureBuildInfo>* PreBuiltData, float ComponentScaling, TArray<FStreamingTexturePrimitiveInfo>& OutStreamingTextures) const;
+	void GetStreamingTextureInfoInner(FStreamingTextureLevelContext& LevelContext, const TArray<FStreamingTextureBuildInfo>* PreBuiltData, float ComponentScaling, TArray<FStreamingRenderAssetPrimitiveInfo>& OutStreamingTextures) const;
 
 #if !(UE_BUILD_SHIPPING || UE_BUILD_TEST)
 	/**
diff --git a/Engine/Source/Runtime/Engine/Classes/Components/PrimitiveComponent.h b/Engine/Source/Runtime/Engine/Classes/Components/PrimitiveComponent.h
--- a/Engine/Source/Runtime/Engine/Classes/Components/PrimitiveComponent.h
+++ b/Engine/Source/Runtime/Engine/Classes/Components/PrimitiveComponent.h
@@ -235,9 +235,9 @@ public:
 	UPROPERTY(EditAnywhere, AdvancedDisplay, BlueprintReadOnly, Category=LOD)
 	uint8 bNeverDistanceCull:1;
 
-	/** Whether this primitive is referenced by a FLevelTextureManager  */
+	/** Whether this primitive is referenced by a FLevelRenderAssetManager  */
 	mutable uint8 bAttachedToStreamingManagerAsStatic : 1;
-	/** Whether this primitive is referenced by a FDynamicTextureInstanceManager */
+	/** Whether this primitive is referenced by a FDynamicRenderAssetInstanceManager */
 	mutable uint8 bAttachedToStreamingManagerAsDynamic : 1;
 	/** Whether this primitive is handled as dynamic, although it could have no references */
 	mutable uint8 bHandledByStreamingManagerAsDynamic : 1;
@@ -1516,17 +1516,17 @@ public:
 	virtual ELightMapInteractionType GetStaticLightingType() const	{ return LMIT_None;	}
 
 	/**
-	 * Enumerates the streaming textures used by the primitive.
+	 * Enumerates the streaming textures/meshes used by the primitive.
 	 * @param LevelContext - Level scope context used to process texture streaming build data.
-	 * @param OutStreamingTextures - Upon return, contains a list of the streaming textures used by the primitive.
+	 * @param OutStreamingRenderAssets - Upon return, contains a list of the streaming textures/meshes used by the primitive.
 	 */
-	virtual void GetStreamingTextureInfo(FStreamingTextureLevelContext& LevelContext, TArray<FStreamingTexturePrimitiveInfo>& OutStreamingTextures) const;
+	virtual void GetStreamingRenderAssetInfo(FStreamingTextureLevelContext& LevelContext, TArray<FStreamingRenderAssetPrimitiveInfo>& OutStreamingRenderAssets) const;
 
 	/**
-	 * Call GetStreamingTextureInfo and remove the elements with a NULL texture
-	 * @param OutStreamingTextures - Upon return, contains a list of the non-null streaming textures used by the primitive.
+	 * Call GetStreamingRenderAssetInfo and remove the elements with a NULL texture
+	 * @param OutStreamingRenderAssets - Upon return, contains a list of the non-null streaming textures or meshes used by the primitive.
 	 */
-	void GetStreamingTextureInfoWithNULLRemoval(FStreamingTextureLevelContext& LevelContext, TArray<FStreamingTexturePrimitiveInfo>& OutStreamingTextures) const;
+	void GetStreamingRenderAssetInfoWithNULLRemoval(FStreamingTextureLevelContext& LevelContext, TArray<FStreamingRenderAssetPrimitiveInfo>& OutStreamingRenderAssets) const;
 
 	/**
 	 *	Update the streaming data of this component.
diff --git a/Engine/Source/Runtime/Engine/Classes/Engine/StaticMesh.h b/Engine/Source/Runtime/Engine/Classes/Engine/StaticMesh.h
--- a/Engine/Source/Runtime/Engine/Classes/Engine/StaticMesh.h
+++ b/Engine/Source/Runtime/Engine/Classes/Engine/StaticMesh.h
@@ -16,6 +16,7 @@
 #include "Components.h"
 #include "Interfaces/Interface_CollisionDataProvider.h"
 #include "Engine/MeshMerging.h"
+#include "Engine/StreamableRenderAsset.h"
 #include "Templates/UniquePtr.h"
 #include "StaticMeshResources.h"
 #include "PerPlatformProperties.h"
@@ -30,6 +31,7 @@ class FSpeedTreeWind;
 class UAssetUserData;
 class UMaterialInterface;
 class UNavCollisionBase;
+class FStaticMeshUpdate;
 struct FMeshDescriptionBulkData;
 struct FStaticMeshLODResources;
 
@@ -547,7 +549,7 @@ struct FStaticMeshDescriptionAttributeGetter
  * @see AStaticMeshActor, UStaticMeshComponent
  */
 UCLASS(hidecategories=Object, customconstructor, MinimalAPI, BlueprintType, config=Engine)
-class UStaticMesh : public UObject, public IInterface_CollisionDataProvider, public IInterface_AssetUserData
+class UStaticMesh : public UStreamableRenderAsset, public IInterface_CollisionDataProvider, public IInterface_AssetUserData
 {
 	GENERATED_UCLASS_BODY()
 
@@ -591,6 +593,13 @@ class UStaticMesh : public UObject, public IInterface_CollisionDataProvider, pub
 	UPROPERTY(EditAnywhere, AssetRegistrySearchable, Category=LodSettings)
 	FName LODGroup;
 
+	/**
+	 * If non-negative, specify the maximum number of streamed LODs. Only has effect if
+	 * mesh LOD streaming is enabled for the target platform.
+	 */
+	UPROPERTY()
+	FPerPlatformInt NumStreamedLODs;
+
 	/* The last import version */
 	UPROPERTY()
 	int32 ImportVersion;
@@ -778,6 +787,11 @@ protected:
 	UPROPERTY(EditAnywhere, AdvancedDisplay, Instanced, Category = StaticMesh)
 	TArray<UAssetUserData*> AssetUserData;
 
+	FStaticMeshUpdate* PendingUpdate;
+
+	friend struct FStaticMeshUpdateContext;
+	friend class FStaticMeshUpdate;
+
 public:
 	/** The editable mesh representation of this static mesh */
 	// @todo: Maybe we don't want this visible in the details panel in the end; for now, this might aid debugging.
@@ -911,6 +925,32 @@ public:
 	ENGINE_API virtual bool CanBeClusterRoot() const override;
 	//~ End UObject Interface.
 
+	//~ Begin UStreamableRenderAsset Interface
+	virtual int32 GetLODGroupForStreaming() const final override;
+	virtual int32 GetNumMipsForStreaming() const final override;
+	virtual int32 GetNumNonStreamingMips() const final override;
+	virtual int32 CalcNumOptionalMips() const final override;
+	virtual int32 CalcCumulativeLODSize(int32 NumLODs) const final override;
+	virtual bool GetMipDataFilename(const int32 MipIndex, FString& BulkDataFilename) const final override;
+	virtual bool IsReadyForStreaming() const final override;
+	virtual int32 GetNumResidentMips() const final override;
+	virtual int32 GetNumRequestedMips() const final override;
+	virtual bool CancelPendingMipChangeRequest() final override;
+	virtual bool HasPendingUpdate() const final override;
+	virtual bool StreamOut(int32 NewMipCount) final override;
+	virtual bool StreamIn(int32 NewMipCount, bool bHighPrio) final override;
+	virtual bool UpdateStreamingStatus(bool bWaitForMipFading = false) final override;
+	//~ End UStreamableRenderAsset Interface
+
+	void LinkStreaming();
+	void UnlinkStreaming();
+
+	/**
+	* Cancels any pending static mesh streaming actions if possible.
+	* Returns when no more async loading requests are in flight.
+	*/
+	ENGINE_API static void CancelAllPendingStreamingActions();
+
 	/**
 	 * Rebuilds renderable data for this static mesh.
 	 * @param bSilent - If true will not popup a progress dialog.
diff --git a/Engine/Source/Runtime/Engine/Classes/Engine/StreamableRenderAsset.h b/Engine/Source/Runtime/Engine/Classes/Engine/StreamableRenderAsset.h
new file mode 100644
--- /dev/null
+++ b/Engine/Source/Runtime/Engine/Classes/Engine/StreamableRenderAsset.h
@@ -0,0 +1,264 @@
+// Copyright 1998-2018 Epic Games, Inc. All Rights Reserved.
+
+#pragma once
+
+#include "CoreMinimal.h"
+#include "UObject/ObjectMacros.h"
+#include "UObject/Object.h"
+#include "Misc/App.h"
+#include "StreamableRenderAsset.generated.h"
+
+#define STREAMABLERENDERASSET_NODEFAULT(FuncName) LowLevelFatalError(TEXT("UStreamableRenderAsset::%s has no default implementation"), TEXT(#FuncName))
+
+UCLASS(Abstract, MinimalAPI)
+class UStreamableRenderAsset : public UObject
+{
+	GENERATED_UCLASS_BODY()
+
+public:
+	/** Get an integer representation of the LOD group */
+	virtual int32 GetLODGroupForStreaming() const
+	{
+		STREAMABLERENDERASSET_NODEFAULT(GetLODGroupForStreaming);
+		return -1;
+	}
+
+	/** Get the total number of LODs including non-streamable and optional LODs */
+	virtual int32 GetNumMipsForStreaming() const
+	{
+		STREAMABLERENDERASSET_NODEFAULT(GetNumMipsForStreaming);
+		return -1;
+	}
+
+	/** Returns the number of LODs in this asset that are not able to be streamed */
+	virtual int32 GetNumNonStreamingMips() const
+	{
+		STREAMABLERENDERASSET_NODEFAULT(GetNumNonStreamingMips);
+		return -1;
+	}
+
+	virtual int32 CalcNumOptionalMips() const
+	{
+		STREAMABLERENDERASSET_NODEFAULT(CalcNumOptionalMips);
+		return -1;
+	}
+
+	virtual int32 CalcCumulativeLODSize(int32 NumLODs) const
+	{
+		STREAMABLERENDERASSET_NODEFAULT(CalcCumulativeLODSize);
+		return -1;
+	}
+
+	virtual bool GetMipDataFilename(const int32 MipIndex, FString& OutBulkDataFilename) const
+	{
+		STREAMABLERENDERASSET_NODEFAULT(GetMipDataFilename);
+		return false;
+	}
+
+	virtual bool IsReadyForStreaming() const
+	{
+		STREAMABLERENDERASSET_NODEFAULT(IsReadyForStreaming);
+		return false;
+	}
+
+	/** The number of LODs currently in memory */
+	virtual int32 GetNumResidentMips() const
+	{
+		STREAMABLERENDERASSET_NODEFAULT(GetNumResidentMips);
+		return -1;
+	}
+
+	/** When the asset is being updated from StreamIn() or StreamOut(), returns the number of LODs requested */
+	virtual int32 GetNumRequestedMips() const
+	{
+		STREAMABLERENDERASSET_NODEFAULT(GetNumRequestedMips);
+		return -1;
+	}
+
+	/**
+	* Tries to cancel a pending LOD change request. Requests cannot be canceled if they are in the
+	* finalization phase.
+	*
+	* @param	true if cancelation was successful, false otherwise
+	*/
+	virtual bool CancelPendingMipChangeRequest()
+	{
+		STREAMABLERENDERASSET_NODEFAULT(CancelPendingMipChangeRequest);
+		return false;
+	}
+
+	virtual bool HasPendingUpdate() const
+	{
+		STREAMABLERENDERASSET_NODEFAULT(HasPendingUpdate);
+		return false;
+	}
+
+	/**
+	* Unload some mips from memory. Only usable if the asset is streamable.
+	*
+	* @param NewMipCount - The desired mip count after the mips are unloaded.
+	* @return Whether any mips were requested to be unloaded.
+	*/
+	virtual bool StreamOut(int32 NewMipCount)
+	{
+		STREAMABLERENDERASSET_NODEFAULT(StreamOut);
+		return false;
+	}
+
+	/**
+	* Loads mips from disk to memory. Only usable if the asset is streamable.
+	*
+	* @param NewMipCount - The desired mip count after the mips are loaded.
+	* @param bHighPrio   - true if the load request is of high priority and must be issued before other asset requests.
+	* @return Whether any mips were resquested to be loaded.
+	*/
+	virtual bool StreamIn(int32 NewMipCount, bool bHighPrio)
+	{
+		STREAMABLERENDERASSET_NODEFAULT(StreamIn);
+		return false;
+	}
+
+	/**
+	* Updates the streaming status of the asset and performs finalization when appropriate. The function returns
+	* true while there are pending requests in flight and updating needs to continue.
+	*
+	* @param bWaitForMipFading	Whether to wait for Mip Fading to complete before finalizing.
+	* @return					true if there are requests in flight, false otherwise
+	*/
+	virtual bool UpdateStreamingStatus(bool bWaitForMipFading = false)
+	{
+		STREAMABLERENDERASSET_NODEFAULT(UpdateStreamingStatus);
+		return false;
+	}
+
+	/**
+	* Invalidates per-asset last render time. Mainly used to opt in UnknownRefHeuristic
+	* during LOD index calculation. See FStreamingRenderAsset::bUseUnknownRefHeuristic
+	*/
+	virtual void InvalidateLastRenderTimeForStreaming() {}
+
+	/**
+	* Get the per-asset last render time. FLT_MAX means never use UnknownRefHeuristic
+	* and the asset will only keep non-streamable LODs when there is no instance/reference
+	* in the scene
+	*/
+	virtual float GetLastRenderTimeForStreaming() const { return FLT_MAX; }
+
+	/**
+	* Returns whether miplevels should be forced resident.
+	*
+	* @return true if either transient or serialized override requests miplevels to be resident, false otherwise
+	*/
+	virtual bool ShouldMipLevelsBeForcedResident() const
+	{
+		return bGlobalForceMipLevelsToBeResident
+			|| bForceMiplevelsToBeResident
+			|| ForceMipLevelsToBeResidentTimestamp >= FApp::GetCurrentTime();
+	}
+
+	/**
+	* Tells the streaming system that it should force all mip-levels to be resident for a number of seconds.
+	* @param Seconds					Duration in seconds
+	* @param CinematicTextureGroups	Bitfield indicating which texture groups that use extra high-resolution mips
+	*/
+	ENGINE_API void SetForceMipLevelsToBeResident(float Seconds, int32 CinematicLODGroupMask = 0);
+
+	/**
+	* Returns the cached combined LOD bias based on texture LOD group and LOD bias.
+	* @return	LOD bias
+	*/
+	ENGINE_API int32 GetCachedLODBias() const { return CachedCombinedLODBias; }
+
+	FORCEINLINE void SetCachedNumResidentLODs(uint8 NewVal)
+	{
+#if !(WITH_EDITOR)
+		CachedNumResidentLODs = NewVal;
+#endif
+	}
+
+	FORCEINLINE void SetCachedReadyForStreaming(bool NewVal)
+	{
+#if !(WITH_EDITOR)
+		bCachedReadyForStreaming = NewVal;
+#endif
+	}
+
+	FORCEINLINE uint8 GetCachedNumResidentLODs() const
+	{
+#if WITH_EDITOR
+		return GetNumResidentMips();
+#else
+		check(CachedNumResidentLODs == GetNumResidentMips());
+		return CachedNumResidentLODs;
+#endif
+	}
+
+	FORCEINLINE bool GetCachedReadyForStreaming() const
+	{
+#if WITH_EDITOR
+		return IsReadyForStreaming();
+#else
+		check(bCachedReadyForStreaming == IsReadyForStreaming());
+		return bCachedReadyForStreaming;
+#endif
+	}
+
+protected:
+	/** WorldSettings timestamp that tells the streamer to force all miplevels to be resident up until that time. */
+	UPROPERTY(transient)
+	double ForceMipLevelsToBeResidentTimestamp;
+
+public:
+	/** Number of mip-levels to use for cinematic quality. */
+	UPROPERTY(EditAnywhere, BlueprintReadWrite, Category = LevelOfDetail, AdvancedDisplay)
+	int32 NumCinematicMipLevels;
+
+protected:
+	/** FStreamingRenderAsset index used by the texture streaming system. */
+	UPROPERTY(transient, duplicatetransient, NonTransactional)
+	int32 StreamingIndex;
+
+	/** Cached combined group and texture LOD bias to use.	*/
+	UPROPERTY(transient)
+	int32 CachedCombinedLODBias;
+
+	/** Cached value of GetNumResidentMips(). Used to reduce cache misses */
+	UPROPERTY(transient)
+	uint8 CachedNumResidentLODs;
+
+	/** Cached value of IsReadyForStreaming(). Used to reduce cache misses */
+	UPROPERTY(transient)
+	uint8 bCachedReadyForStreaming : 1;
+
+public:
+	UPROPERTY(EditAnywhere, BlueprintReadWrite, Category = LevelOfDetail, AssetRegistrySearchable, AdvancedDisplay)
+	uint8 NeverStream : 1;
+
+	/** Global and serialized version of ForceMiplevelsToBeResident.				*/
+	UPROPERTY(EditAnywhere, BlueprintReadOnly, Category = LevelOfDetail, meta = (DisplayName = "Global Force Resident Mip Levels"), AdvancedDisplay)
+	uint8 bGlobalForceMipLevelsToBeResident : 1;
+
+	/** Whether the texture is currently streamable or not.						*/
+	UPROPERTY(transient, NonTransactional)
+	uint8 bIsStreamable : 1;
+
+	/** Whether some mips might be streamed soon. If false, the texture is not planned resolution will be stable. */
+	UPROPERTY(transient, NonTransactional)
+	uint8 bHasStreamingUpdatePending : 1;
+
+	/** Override whether to fully stream even if texture hasn't been rendered.	*/
+	UPROPERTY(transient)
+	uint8 bForceMiplevelsToBeResident : 1;
+
+	/** Ignores the streaming mip bias used to accommodate memory constraints. */
+	UPROPERTY(transient)
+	uint8 bIgnoreStreamingMipBias : 1;
+
+protected:
+	/** Whether to use the extra cinematic quality mip-levels, when we're forcing mip-levels to be resident. */
+	UPROPERTY(transient)
+	uint8 bUseCinematicMipLevels : 1;
+
+	friend struct FRenderAssetStreamingManager;
+	friend struct FStreamingRenderAsset;
+};
diff --git a/Engine/Source/Runtime/Engine/Classes/Engine/Texture.h b/Engine/Source/Runtime/Engine/Classes/Engine/Texture.h
--- a/Engine/Source/Runtime/Engine/Classes/Engine/Texture.h
+++ b/Engine/Source/Runtime/Engine/Classes/Engine/Texture.h
@@ -15,6 +15,7 @@
 #include "Engine/TextureDefines.h"
 #include "MaterialShared.h"
 #include "TextureResource.h"
+#include "Engine/StreamableRenderAsset.h"
 #include "Texture.generated.h"
 
 class ITargetPlatform;
@@ -368,7 +369,7 @@ struct FTexturePlatformData
 };
 
 UCLASS(abstract, MinimalAPI, BlueprintType)
-class UTexture : public UObject, public IInterface_AssetUserData
+class UTexture : public UStreamableRenderAsset, public IInterface_AssetUserData
 {
 	GENERATED_UCLASS_BODY()
 
@@ -520,10 +521,6 @@ public:
 	UPROPERTY(EditAnywhere, BlueprintReadWrite, Category=LevelOfDetail, meta=(DisplayName="LOD Bias"), AssetRegistrySearchable)
 	int32 LODBias;
 
-	/** Number of mip-levels to use for cinematic quality. */
-	UPROPERTY(EditAnywhere, BlueprintReadWrite, Category=LevelOfDetail, AdvancedDisplay)
-	int32 NumCinematicMipLevels;
-
 	/** Compression settings to use when building the texture. */
 	UPROPERTY(EditAnywhere, BlueprintReadWrite, Category=Compression, AssetRegistrySearchable)
 	TEnumAsByte<enum TextureCompressionSettings> CompressionSettings;
@@ -548,26 +545,15 @@ public:
 
 #endif // WITH_EDITORONLY_DATA
 
-	UPROPERTY(EditAnywhere, BlueprintReadWrite, Category=Texture, AssetRegistrySearchable, AdvancedDisplay)
-	uint8 NeverStream:1;
-
 	/** If true, the RHI texture will be created using TexCreate_NoTiling */
 	UPROPERTY()
 	uint8 bNoTiling:1;
 
-	/** Whether to use the extra cinematic quality mip-levels, when we're forcing mip-levels to be resident. */
-	UPROPERTY(transient)
-	uint8 bUseCinematicMipLevels:1;
-
 private:
 	/** Whether the async resource release process has already been kicked off or not */
 	UPROPERTY(transient)
 	uint8 bAsyncResourceReleaseHasBeenStarted : 1;
 
-	/** Cached combined group and texture LOD bias to use.	*/
-	UPROPERTY(transient)
-	int32 CachedCombinedLODBias;
-
 protected:
 	/** Array of user data stored with the asset */
 	UPROPERTY(EditAnywhere, AdvancedDisplay, Instanced, Category = Texture)
@@ -603,13 +589,6 @@ public:
 	 */
 	virtual class FTextureResource* CreateResource() PURE_VIRTUAL(UTexture::CreateResource,return NULL;);
 
-	/**
-	 * Returns the cached combined LOD bias based on texture LOD group and LOD bias.
-	 *
-	 * @return	LOD bias
-	 */
-	ENGINE_API int32 GetCachedLODBias() const;
-
 	/** Cache the combined LOD bias based on texture LOD group and LOD bias. */
 	ENGINE_API void UpdateCachedLODBias();
 
@@ -624,18 +603,6 @@ public:
 	virtual void WaitForStreaming()
 	{
 	}
-	
-	/**
-	 * Updates the streaming status of the texture and performs finalization when appropriate. The function returns
-	 * true while there are pending requests in flight and updating needs to continue.
-	 *
-	 * @param bWaitForMipFading	Whether to wait for Mip Fading to complete before finalizing.
-	 * @return					true if there are requests in flight, false otherwise
-	 */
-	virtual bool UpdateStreamingStatus( bool bWaitForMipFading = false )
-	{
-		return false;
-	}
 
 	/**
 	 * Textures that use the derived data cache must override this function and
@@ -761,6 +728,11 @@ public:
 	ENGINE_API virtual bool IsPostLoadThreadSafe() const override;
 	//~ End UObject Interface.
 
+	//~ Begin UStreamableRenderAsset Interface
+	virtual int32 GetLODGroupForStreaming() const final override { return static_cast<int32>(LODGroup); }
+	virtual bool UpdateStreamingStatus(bool bWaitForMipFading = false) override { return false; }
+	//~ End UStreamableRenderAsset Interface
+
 	/**
 	 *	Gets the average brightness of the texture (in linear space)
 	 *
diff --git a/Engine/Source/Runtime/Engine/Classes/Engine/Texture2D.h b/Engine/Source/Runtime/Engine/Classes/Engine/Texture2D.h
--- a/Engine/Source/Runtime/Engine/Classes/Engine/Texture2D.h
+++ b/Engine/Source/Runtime/Engine/Classes/Engine/Texture2D.h
@@ -17,11 +17,6 @@ class UTexture2D : public UTexture
 {
 	GENERATED_UCLASS_BODY()
 
-private:
-	/** FStreamingTexture index used by the texture streaming system. */
-	UPROPERTY(transient, duplicatetransient, NonTransactional)
-	int32 StreamingIndex;
-
 public:
 
 	/*
@@ -59,26 +54,6 @@ private:
 	uint8 bTemporarilyDisableStreaming:1;
 
 public:
-	/** Whether the texture is currently streamable or not.						*/
-	UPROPERTY(transient, NonTransactional)
-	uint8 bIsStreamable:1;
-
-	/** Whether some mips might be streamed soon. If false, the texture is not planned resolution will be stable. */
-	UPROPERTY(transient, NonTransactional)
-	uint8 bHasStreamingUpdatePending:1;
-
-	/** Override whether to fully stream even if texture hasn't been rendered.	*/
-	UPROPERTY(transient)
-	uint8 bForceMiplevelsToBeResident:1;
-
-	/** Ignores the streaming mip bias used to accommodate memory constraints. */
-	UPROPERTY(transient)
-	uint8 bIgnoreStreamingMipBias:1;
-
-	/** Global and serialized version of ForceMiplevelsToBeResident.				*/
-	UPROPERTY(EditAnywhere, BlueprintReadOnly, Category=LevelOfDetail, meta=(DisplayName="Global Force Resident Mip Levels"), AdvancedDisplay)
-	uint8 bGlobalForceMipLevelsToBeResident:1;
-
 #if WITH_EDITORONLY_DATA
 	/** Whether the texture has been painted in the editor.						*/
 	UPROPERTY()
@@ -101,10 +76,6 @@ private:
 	UPROPERTY()
 	FIntPoint ImportedSize;
 
-	/** WorldSettings timestamp that tells the streamer to force all miplevels to be resident up until that time. */
-	UPROPERTY(transient)
-	double ForceMipLevelsToBeResidentTimestamp;
-
 public:
 	/** The derived data for this texture on this platform. */
 	FTexturePlatformData *PlatformData;
@@ -116,26 +87,6 @@ public:
 	/** memory used for directly loading bulk mip data */
 	FTexture2DResourceMem*		ResourceMem;
 
-	/**
-	 * Loads mips from disk to memory. Only usable if the texture is streamable.
-	 *
-	 * @param NewMipCount - The desired mip count after the mips are loaded.
-	 * @param bHighPrio   - true if the load request is of high priority and must be issued before other texture requests.
-	 * @return Whether any mips were resquested to be loaded.
-	 */
-	ENGINE_API bool StreamIn(int32 NewMipCount, bool bHighPrio);
-
-	/**
-	 * Unload some mips from memory. Only usable if the texture is streamable.
-	 *
-	 * @param NewMipCount - The desired mip count after the mips are unloaded.
-	 * @return Whether any mips were requested to be unloaded.
-	 */
-	ENGINE_API bool StreamOut(int32 NewMipCount);
-
-	/** true if the texture is currently being updated through StreamIn() or StreamOut(). */
-	FORCEINLINE bool HasPendingUpdate() const { return PendingUpdate != nullptr; }
-
 protected:
 
 	/** Helper to manage the current pending update following a call to StreamIn() or StreamOut(). */
@@ -170,6 +121,51 @@ public:
 #endif
 	//~ End UTexture Interface.
 
+	//~ Begin UStreamableRenderAsset Interface
+	virtual int32 GetNumMipsForStreaming() const final override { return GetNumMips(); }
+	virtual int32 GetNumNonStreamingMips() const final override;
+	virtual int32 CalcNumOptionalMips() const final override;
+	virtual int32 CalcCumulativeLODSize(int32 NumLODs) const final override { return CalcTextureMemorySize(NumLODs); }
+
+	virtual bool GetMipDataFilename(const int32 MipIndex, FString& OutBulkDataFilename) const final override
+	{
+		if (PlatformData)
+		{
+			if (MipIndex < PlatformData->Mips.Num() && MipIndex >= 0)
+			{
+				OutBulkDataFilename = PlatformData->Mips[MipIndex].BulkData.GetFilename();
+				return true;
+			}
+		}
+		return false;
+	}
+
+	/**
+	* Returns whether the texture is ready for streaming aka whether it has had InitRHI called on it.
+	*
+	* @return true if initialized and ready for streaming, false otherwise
+	*/
+	virtual bool IsReadyForStreaming() const final override
+	{
+		FTexture2DResource* Texture2DResource = (FTexture2DResource*)Resource;
+		return Texture2DResource && Texture2DResource->bReadyForStreaming;
+	}
+
+	ENGINE_API virtual int32 GetNumResidentMips() const final override;
+	virtual int32 GetNumRequestedMips() const final override;
+	virtual bool CancelPendingMipChangeRequest() final override;
+
+	/** true if the texture is currently being updated through StreamIn() or StreamOut(). */
+	virtual bool HasPendingUpdate() const final override { return PendingUpdate != nullptr; }
+
+	virtual bool StreamOut(int32 NewMipCount) final override;
+	virtual bool StreamIn(int32 NewMipCount, bool bHighPrio) final override;
+	virtual bool UpdateStreamingStatus(bool bWaitForMipFading = false) final override;
+	virtual void InvalidateLastRenderTimeForStreaming() final override;
+	virtual float GetLastRenderTimeForStreaming() const final override;
+	virtual bool ShouldMipLevelsBeForcedResident() const final override;
+	//~ End UStreamableRenderAsset Interface
+
 	/** Trivial accessors. */
 	FORCEINLINE int32 GetSizeX() const
 	{
@@ -196,21 +192,6 @@ public:
 		return 0;
 	}
 
-	int32 CalcNumOptionalMips() const;
-
-	FORCEINLINE bool GetMipDataFilename(const int32 MipIndex, const FString*& BulkDataFilename)
-	{
-		if (PlatformData)
-		{
-			if ( MipIndex < PlatformData->Mips.Num() && MipIndex >= 0)
-			{
-				BulkDataFilename = &PlatformData->Mips[MipIndex].BulkData.GetFilename();
-				return true;
-			}
-		}
-		return false;
-	}
-
 	FORCEINLINE EPixelFormat GetPixelFormat() const
 	{
 		if (PlatformData)
@@ -234,12 +215,6 @@ public:
 	}
 
 	FORCEINLINE int32 GetStreamingIndex() const { return StreamingIndex; }
-	
-	/** The number of mips currently in memory. */
-	ENGINE_API int32 GetNumResidentMips() const;
-
-	/** When the texture is being updated from StreamIn() or StreamOut(), returns the number of mips requested. */
-	int32 GetNumRequestedMips() const;
 
 	/**
 	 * Calculates the maximum number of mips the engine allows to be loaded for this texture. 
@@ -266,12 +241,6 @@ public:
 	/** Sets the minimum number of mips that must be resident in memory (cannot be streamed). */
 	static void SetMinTextureResidentMipCount(int32 InMinTextureResidentMipCount);
 
-	/** Retrieves the last render time used for streaming. */
-	float GetLastRenderTimeForStreaming();
-
-	/** Invalidates the last render time used for streaming. */
-	void InvalidateLastRenderTimeForStreaming();
-
 	/**
 	 * Get mip data starting with the specified mip index.
 	 * @param FirstMipToLoad - The first mip index to cache.
@@ -281,11 +250,6 @@ public:
 	 */
 	ENGINE_API void GetMipData(int32 FirstMipToLoad, void** OutMipData);
 
-	/**
-	 * Returns the number of mips in this texture that are not able to be streamed.
-	 */
-	ENGINE_API int32 GetNumNonStreamingMips() const;
-
 	/**
 	 * Computes the minimum and maximum allowed mips for a texture.
 	 * @param MipCount - The number of mip levels in the texture.
@@ -335,38 +299,10 @@ public:
 	 */
 	ENGINE_API bool HasAlphaChannel() const;
 
-	/**
-	 * Returns whether the texture is ready for streaming aka whether it has had InitRHI called on it.
-	 *
-	 * @return true if initialized and ready for streaming, false otherwise
-	 */
-	bool FORCEINLINE IsReadyForStreaming() const
-	{
-		FTexture2DResource* Texture2DResource = (FTexture2DResource*)Resource;
-		return Texture2DResource && Texture2DResource->bReadyForStreaming;
-	}
-
 	/**
 	 * Waits until all streaming requests for this texture has been fully processed.
 	 */
 	virtual void WaitForStreaming() override;
-	
-	/**
-	 * Updates the streaming status of the texture and performs finalization when appropriate. The function returns
-	 * true while there are pending requests in flight and updating needs to continue.
-	 *
-	 * @param bWaitForMipFading	Whether to wait for Mip Fading to complete before finalizing.
-	 * @return					true if there are requests in flight, false otherwise
-	 */
-	virtual bool UpdateStreamingStatus( bool bWaitForMipFading = false ) override;
-
-	/**
-	 * Tries to cancel a pending mip change request. Requests cannot be canceled if they are in the
-	 * finalization phase.
-	 *
-	 * @param	true if cancelation was successful, false otherwise
-	 */
-	bool CancelPendingMipChangeRequest();
 
 	/**
 	 * Returns the size of the object/ resource for display to artists/ LDs in the Editor.
@@ -375,13 +311,6 @@ public:
 	 */
 	virtual void GetResourceSizeEx(FResourceSizeEx& CumulativeResourceSize) override;
 
-	/**
-	 * Returns whether miplevels should be forced resident.
-	 *
-	 * @return true if either transient or serialized override requests miplevels to be resident, false otherwise
-	 */
-	bool ShouldMipLevelsBeForcedResident() const;
-
 	/**
 	 * Whether all miplevels of this texture have been fully streamed in, LOD settings permitting.
 	 */
@@ -446,15 +375,8 @@ public:
 
 #endif // WITH_EDITOR
 
-	friend struct FStreamingManagerTexture;
-	friend struct FStreamingTexture;
-	
-	/**
-	 * Tells the streaming system that it should force all mip-levels to be resident for a number of seconds.
-	 * @param Seconds					Duration in seconds
-	 * @param CinematicTextureGroups	Bitfield indicating which texture groups that use extra high-resolution mips
-	 */
-	virtual void SetForceMipLevelsToBeResident( float Seconds, int32 CinematicTextureGroups = 0 );
+	friend struct FRenderAssetStreamingManager;
+	friend struct FStreamingRenderAsset;
 	
 	/** creates and initializes a new Texture2D with the requested settings */
 	ENGINE_API static class UTexture2D* CreateTransient(int32 InSizeX, int32 InSizeY, EPixelFormat InFormat = PF_B8G8R8A8);
diff --git a/Engine/Source/Runtime/Engine/Classes/Engine/TextureStreamingTypes.h b/Engine/Source/Runtime/Engine/Classes/Engine/TextureStreamingTypes.h
--- a/Engine/Source/Runtime/Engine/Classes/Engine/TextureStreamingTypes.h
+++ b/Engine/Source/Runtime/Engine/Classes/Engine/TextureStreamingTypes.h
@@ -25,23 +25,24 @@ ENGINE_API DECLARE_LOG_CATEGORY_EXTERN(TextureStreamingBuild, Log, All);
 
 class UTexture;
 class UTexture2D;
+class UStreamableRenderAsset;
 struct FStreamingTextureBuildInfo;
 struct FMaterialTextureInfo;
 
 // The PackedRelativeBox value that return the bound unaltered
 static const uint32 PackedRelativeBox_Identity = 0xffff0000;
 
-/** Information about a streaming texture that a primitive uses for rendering. */
+/** Information about a streaming texture/mesh that a primitive uses for rendering. */
 USTRUCT()
-struct FStreamingTexturePrimitiveInfo
+struct FStreamingRenderAssetPrimitiveInfo
 {
 	GENERATED_USTRUCT_BODY()
 
 	UPROPERTY()
-	UTexture2D* Texture;
+	UStreamableRenderAsset* RenderAsset;
 
 	/** 
-	 * The streaming bounds of the texture, usually the component material bounds. 
+	 * The streaming bounds of the texture/mesh, usually the component material bounds. 
 	 * Usually only valid for registered component, as component bounds are only updated when the components are registered.
 	 * otherwise only PackedRelativeBox can be used.Irrelevant when the component is not registered, as the component could be moved by ULevel::ApplyWorldOffset()
 	 * In that case, only PackedRelativeBox is meaningful.
@@ -54,23 +55,23 @@ struct FStreamingTexturePrimitiveInfo
 
 	/** 
 	 * When non zero, this represents the relative box used to compute Bounds, using the component bounds as reference.
-	 * If available, this allows the texture streamer to generate the level streaming data before the level gets visible.
+	 * If available, this allows the renderable asset streamer to generate the level streaming data before the level gets visible.
 	 * At that point, the component are not yet registered, and the bounds are unknown, but the precompiled build data is still available.
 	 * Also allows to update the relative bounds after a level get moved around from ApplyWorldOffset.
 	 */
 	UPROPERTY()
 	uint32 PackedRelativeBox;
 
-	FStreamingTexturePrimitiveInfo() : 
-		Texture(nullptr), 
+	FStreamingRenderAssetPrimitiveInfo() : 
+		RenderAsset(nullptr),
 		Bounds(ForceInit), 
 		TexelFactor(1.0f),
 		PackedRelativeBox(0)
 	{
 	}
 
-	FStreamingTexturePrimitiveInfo(UTexture2D* InTexture, const FBoxSphereBounds& InBounds, float InTexelFactor, uint32 InPackedRelativeBox = 0) :
-		Texture(InTexture), 
+	FStreamingRenderAssetPrimitiveInfo(UStreamableRenderAsset* InAsset, const FBoxSphereBounds& InBounds, float InTexelFactor, uint32 InPackedRelativeBox = 0) :
+		RenderAsset(InAsset),
 		Bounds(InBounds), 
 		TexelFactor(InTexelFactor),
 		PackedRelativeBox(InPackedRelativeBox)
@@ -128,7 +129,7 @@ struct FStreamingTextureBuildInfo
 	 *	@param	RefBounds		[in]		The reference bounds used to compute the packed relative box.
 	 *	@param	Info			[in]		The unpacked params.
 	 */
-	ENGINE_API void PackFrom(ULevel* Level, const FBoxSphereBounds& RefBounds, const FStreamingTexturePrimitiveInfo& Info);
+	ENGINE_API void PackFrom(ULevel* Level, const FBoxSphereBounds& RefBounds, const FStreamingRenderAssetPrimitiveInfo& Info);
 
 };
 
@@ -165,7 +166,7 @@ enum ETextureStreamingBuildType
 };
 
 /** 
- * Context used to resolve FStreamingTextureBuildInfo to FStreamingTexturePrimitiveInfo
+ * Context used to resolve FStreamingTextureBuildInfo to FStreamingRenderAssetPrimitiveInfo
  * The context make sure that build data and each texture is only processed once per component (with constant time).
  * It manage internally structures used to accelerate the binding between precomputed data and textures,
  * so that there is only one map lookup per texture per level. 
@@ -223,7 +224,7 @@ public:
 	~FStreamingTextureLevelContext();
 
 	void BindBuildData(const TArray<FStreamingTextureBuildInfo>* PreBuiltData);
-	void ProcessMaterial(const FBoxSphereBounds& ComponentBounds, const FPrimitiveMaterialInfo& MaterialData, float ComponentScaling, TArray<FStreamingTexturePrimitiveInfo>& OutStreamingTextures);
+	void ProcessMaterial(const FBoxSphereBounds& ComponentBounds, const FPrimitiveMaterialInfo& MaterialData, float ComponentScaling, TArray<FStreamingRenderAssetPrimitiveInfo>& OutStreamingTextures);
 
 	EMaterialQualityLevel::Type GetQualityLevel() { return QualityLevel; }
 	ERHIFeatureLevel::Type GetFeatureLevel() { return FeatureLevel; }
@@ -249,11 +250,11 @@ ENGINE_API bool BuildTextureStreamingComponentData(UWorld* InWorld, EMaterialQua
 ENGINE_API void CheckTextureStreamingBuildValidity(UWorld* InWorld);
 
 /**
- * Checks whether a UTexture2D is a texture with streamable mips
- * @param Texture	Texture to check
- * @return			true if the UTexture2D is supposed to be streaming
+ * Checks whether a UStreamableRenderAsset is a texture/mesh with streamable mips
+ * @param Asset		Asset to check
+ * @return			true if the UStreamableRenderAsset is supposed to be streaming
  */
-ENGINE_API bool IsStreamingTexture( const UTexture2D* Texture2D );
+ENGINE_API bool IsStreamingRenderAsset( const UStreamableRenderAsset* Asset );
 
 ENGINE_API uint32 PackRelativeBox(const FVector& RefOrigin, const FVector& RefExtent, const FVector& Origin, const FVector& Extent);
 ENGINE_API uint32 PackRelativeBox(const FBox& RefBox, const FBox& Box);
diff --git a/Engine/Source/Runtime/Engine/Private/Components/MeshComponent.cpp b/Engine/Source/Runtime/Engine/Private/Components/MeshComponent.cpp
--- a/Engine/Source/Runtime/Engine/Private/Components/MeshComponent.cpp
+++ b/Engine/Source/Runtime/Engine/Private/Components/MeshComponent.cpp
@@ -408,7 +408,7 @@ void UMeshComponent::CacheMaterialParameterNameIndices()
 	bCachedMaterialParameterIndicesAreDirty = false;
 }
 
-void UMeshComponent::GetStreamingTextureInfoInner(FStreamingTextureLevelContext& LevelContext, const TArray<FStreamingTextureBuildInfo>* PreBuiltData, float ComponentScaling, TArray<FStreamingTexturePrimitiveInfo>& OutStreamingTextures) const
+void UMeshComponent::GetStreamingTextureInfoInner(FStreamingTextureLevelContext& LevelContext, const TArray<FStreamingTextureBuildInfo>* PreBuiltData, float ComponentScaling, TArray<FStreamingRenderAssetPrimitiveInfo>& OutStreamingTextures) const
 {
 	LevelContext.BindBuildData(PreBuiltData);
 
diff --git a/Engine/Source/Runtime/Engine/Private/Components/ModelComponent.cpp b/Engine/Source/Runtime/Engine/Private/Components/ModelComponent.cpp
--- a/Engine/Source/Runtime/Engine/Private/Components/ModelComponent.cpp
+++ b/Engine/Source/Runtime/Engine/Private/Components/ModelComponent.cpp
@@ -413,7 +413,7 @@ void UModelComponent::SelectAllSurfaces()
 }
 #endif // WITH_EDITOR
 
-void UModelComponent::GetStreamingTextureInfo(FStreamingTextureLevelContext& LevelContext, TArray<FStreamingTexturePrimitiveInfo>& OutStreamingTextures) const
+void UModelComponent::GetStreamingRenderAssetInfo(FStreamingTextureLevelContext& LevelContext, TArray<FStreamingRenderAssetPrimitiveInfo>& OutStreamingRenderAssets) const
 {
 	if( Model )
 	{
@@ -472,10 +472,10 @@ void UModelComponent::GetStreamingTextureInfo(FStreamingTextureLevelContext& Lev
 					UTexture2D* Texture2D = Cast<UTexture2D>(Textures[TextureIndex]);
 					if (!Texture2D) continue;
 
-					FStreamingTexturePrimitiveInfo& StreamingTexture = *new(OutStreamingTextures) FStreamingTexturePrimitiveInfo;
+					FStreamingRenderAssetPrimitiveInfo& StreamingTexture = *new(OutStreamingRenderAssets) FStreamingRenderAssetPrimitiveInfo;
 					StreamingTexture.Bounds = SurfaceBoundingSphere;
 					StreamingTexture.TexelFactor = TexelFactor;
-					StreamingTexture.Texture = Texture2D;
+					StreamingTexture.RenderAsset = Texture2D;
 				}
 			}
 		}
diff --git a/Engine/Source/Runtime/Engine/Private/Components/PrimitiveComponent.cpp b/Engine/Source/Runtime/Engine/Private/Components/PrimitiveComponent.cpp
--- a/Engine/Source/Runtime/Engine/Private/Components/PrimitiveComponent.cpp
+++ b/Engine/Source/Runtime/Engine/Private/Components/PrimitiveComponent.cpp
@@ -38,6 +38,10 @@
 #include "Engine/LODActor.h"
 #endif // WITH_EDITOR
 
+#if DO_CHECK
+#include "Engine/StaticMesh.h"
+#endif
+
 #define LOCTEXT_NAMESPACE "PrimitiveComponent"
 
 //////////////////////////////////////////////////////////////////////////
@@ -330,7 +334,7 @@ bool UPrimitiveComponent::HasStaticLighting() const
 	return ((Mobility == EComponentMobility::Static) || LightmapType == ELightmapType::ForceSurface) && SupportsStaticLighting();
 }
 
-void UPrimitiveComponent::GetStreamingTextureInfo(FStreamingTextureLevelContext& LevelContext, TArray<FStreamingTexturePrimitiveInfo>& OutStreamingTextures) const
+void UPrimitiveComponent::GetStreamingRenderAssetInfo(FStreamingTextureLevelContext& LevelContext, TArray<FStreamingRenderAssetPrimitiveInfo>& OutStreamingRenderAssets) const
 {
 	if (CVarStreamingUseNewMetrics.GetValueOnGameThread() != 0)
 	{
@@ -357,7 +361,7 @@ void UPrimitiveComponent::GetStreamingTextureInfo(FStreamingTextureLevelContext&
 				if (MaterialInterface)
 				{
 					MaterialData.Material = MaterialInterface;
-					LevelContext.ProcessMaterial(Bounds, MaterialData, Bounds.SphereRadius, OutStreamingTextures);
+					LevelContext.ProcessMaterial(Bounds, MaterialData, Bounds.SphereRadius, OutStreamingRenderAssets);
 				}
 				// Remove all instances of this material in case there were duplicates.
 				UsedMaterials.RemoveSwap(MaterialInterface);
@@ -367,26 +371,33 @@ void UPrimitiveComponent::GetStreamingTextureInfo(FStreamingTextureLevelContext&
 }
 
 
-void UPrimitiveComponent::GetStreamingTextureInfoWithNULLRemoval(FStreamingTextureLevelContext& LevelContext, TArray<struct FStreamingTexturePrimitiveInfo>& OutStreamingTextures) const
+void UPrimitiveComponent::GetStreamingRenderAssetInfoWithNULLRemoval(FStreamingTextureLevelContext& LevelContext, TArray<struct FStreamingRenderAssetPrimitiveInfo>& OutStreamingRenderAssets) const
 {
 	// Ignore components that are fully initialized but have no scene proxy (hidden primitive or non game primitive)
 	if (!IsRegistered() || !IsRenderStateCreated() || SceneProxy)
 	{
-		GetStreamingTextureInfo(LevelContext, OutStreamingTextures);
-		for (int32 Index = 0; Index < OutStreamingTextures.Num(); Index++)
+		GetStreamingRenderAssetInfo(LevelContext, OutStreamingRenderAssets);
+		for (int32 Index = 0; Index < OutStreamingRenderAssets.Num(); Index++)
 		{
-			const FStreamingTexturePrimitiveInfo& Info = OutStreamingTextures[Index];
-			if (!IsStreamingTexture(Info.Texture))
+			const FStreamingRenderAssetPrimitiveInfo& Info = OutStreamingRenderAssets[Index];
+			if (!IsStreamingRenderAsset(Info.RenderAsset))
 			{
-				OutStreamingTextures.RemoveAtSwap(Index--);
+				OutStreamingRenderAssets.RemoveAtSwap(Index--);
 			}
 			else
 			{
+#if DO_CHECK
+				ensure(Info.TexelFactor >= 0.f
+					|| Info.RenderAsset->IsA<UStaticMesh>()
+					|| Info.RenderAsset->IsA<USkeletalMesh>()
+					|| (Info.RenderAsset->IsA<UTexture>() && Info.RenderAsset->GetLODGroupForStreaming() == TEXTUREGROUP_Terrain_Heightmap));
+#endif
+
 				// Other wise check that everything is setup right. If the component is not yet registered, then the bound data is irrelevant.
 				const bool bCanBeStreamedByDistance = Info.TexelFactor > SMALL_NUMBER && (Info.Bounds.SphereRadius > SMALL_NUMBER || !IsRegistered()) && ensure(FMath::IsFinite(Info.TexelFactor));
-				if (!bForceMipStreaming && !bCanBeStreamedByDistance && !(Info.TexelFactor < 0 && Info.Texture->LODGroup == TEXTUREGROUP_Terrain_Heightmap))
+				if (!bForceMipStreaming && !bCanBeStreamedByDistance && Info.TexelFactor >= 0.f)
 				{
-					OutStreamingTextures.RemoveAtSwap(Index--);
+					OutStreamingRenderAssets.RemoveAtSwap(Index--);
 				}
 			}
 		}
diff --git a/Engine/Source/Runtime/Engine/Private/Components/SkinnedMeshComponent.cpp b/Engine/Source/Runtime/Engine/Private/Components/SkinnedMeshComponent.cpp
--- a/Engine/Source/Runtime/Engine/Private/Components/SkinnedMeshComponent.cpp
+++ b/Engine/Source/Runtime/Engine/Private/Components/SkinnedMeshComponent.cpp
@@ -900,9 +900,9 @@ bool USkinnedMeshComponent::GetMaterialStreamingData(int32 MaterialIndex, FPrimi
 	return MaterialData.IsValid();
 }
 
-void USkinnedMeshComponent::GetStreamingTextureInfo(FStreamingTextureLevelContext& LevelContext, TArray<FStreamingTexturePrimitiveInfo>& OutStreamingTextures) const
+void USkinnedMeshComponent::GetStreamingRenderAssetInfo(FStreamingTextureLevelContext& LevelContext, TArray<FStreamingRenderAssetPrimitiveInfo>& OutStreamingRenderAssets) const
 {
-	GetStreamingTextureInfoInner(LevelContext, nullptr, GetComponentTransform().GetMaximumAxisScale() * StreamingDistanceMultiplier, OutStreamingTextures);
+	GetStreamingTextureInfoInner(LevelContext, nullptr, GetComponentTransform().GetMaximumAxisScale() * StreamingDistanceMultiplier, OutStreamingRenderAssets);
 }
 
 bool USkinnedMeshComponent::ShouldUpdateBoneVisibility() const
diff --git a/Engine/Source/Runtime/Engine/Private/Components/StaticMeshComponent.cpp b/Engine/Source/Runtime/Engine/Private/Components/StaticMeshComponent.cpp
--- a/Engine/Source/Runtime/Engine/Private/Components/StaticMeshComponent.cpp
+++ b/Engine/Source/Runtime/Engine/Private/Components/StaticMeshComponent.cpp
@@ -638,21 +638,21 @@ static int32 GetNumberOfElements(const TIndirectArray<FStaticMeshLODResources>&
  *	@param	StreamingTextureData	[out]		The resulting packed data.
  *	@param	RefBounds				[in]		The reference bounds used to packed the relative bounds.
  */
-static void PackStreamingTextureData(ULevel* Level, TArray<FStreamingTexturePrimitiveInfo>& UnpackedData, TArray<FStreamingTextureBuildInfo>& StreamingTextureData, const FBoxSphereBounds& RefBounds)
+static void PackStreamingTextureData(ULevel* Level, TArray<FStreamingRenderAssetPrimitiveInfo>& UnpackedData, TArray<FStreamingTextureBuildInfo>& StreamingTextureData, const FBoxSphereBounds& RefBounds)
 {
 	StreamingTextureData.Empty();
 
 	while (UnpackedData.Num())
 	{
-		FStreamingTexturePrimitiveInfo Info = UnpackedData[0];
+		FStreamingRenderAssetPrimitiveInfo Info = UnpackedData[0];
 		UnpackedData.RemoveAtSwap(0);
 
 		// Merge with any other lod section using the same texture.
 		for (int32 Index = 0; Index < UnpackedData.Num(); ++Index)
 		{
-			const FStreamingTexturePrimitiveInfo& CurrInfo = UnpackedData[Index];
+			const FStreamingRenderAssetPrimitiveInfo& CurrInfo = UnpackedData[Index];
 
-			if (CurrInfo.Texture == Info.Texture)
+			if (CurrInfo.RenderAsset == Info.RenderAsset)
 			{
 				Info.Bounds = Union(Info.Bounds, CurrInfo.Bounds);
 				// Take the max scale since it relates to higher texture resolution.
@@ -726,7 +726,7 @@ bool UStaticMeshComponent::BuildTextureStreamingData(ETextureStreamingBuildType
 				{
 					// Get the data without any component scaling as the built data does not include scale.
 					FStreamingTextureLevelContext LevelContext(QualityLevel, FeatureLevel, true); // Use the boxes that were just computed!
-					TArray<FStreamingTexturePrimitiveInfo> UnpackedData;
+					TArray<FStreamingRenderAssetPrimitiveInfo> UnpackedData;
 					GetStreamingTextureInfoInner(LevelContext, nullptr, 1.f, UnpackedData);
 					PackStreamingTextureData(Level, UnpackedData, StreamingTextureData, Bounds);
 				}
@@ -801,7 +801,7 @@ float UStaticMeshComponent::GetTextureStreamingTransformScale() const
 	return GetComponentTransform().GetMaximumAxisScale();
 }
 
-void UStaticMeshComponent::GetStreamingTextureInfo(FStreamingTextureLevelContext& LevelContext, TArray<FStreamingTexturePrimitiveInfo>& OutStreamingTextures) const
+void UStaticMeshComponent::GetStreamingRenderAssetInfo(FStreamingTextureLevelContext& LevelContext, TArray<FStreamingRenderAssetPrimitiveInfo>& OutStreamingRenderAssets) const
 {
 	if (bIgnoreInstanceForTextureStreaming || !GetStaticMesh() || !GetStaticMesh()->HasValidRenderData())
 	{
@@ -809,7 +809,7 @@ void UStaticMeshComponent::GetStreamingTextureInfo(FStreamingTextureLevelContext
 	}
 
 	const float TransformScale = GetTextureStreamingTransformScale();
-	GetStreamingTextureInfoInner(LevelContext, Mobility == EComponentMobility::Static ? &StreamingTextureData : nullptr, TransformScale * StreamingDistanceMultiplier, OutStreamingTextures);
+	GetStreamingTextureInfoInner(LevelContext, Mobility == EComponentMobility::Static ? &StreamingTextureData : nullptr, TransformScale * StreamingDistanceMultiplier, OutStreamingRenderAssets);
 
 	// Process the lightmaps and shadowmaps entries.
 	for (int32 LODIndex = 0; LODIndex < LODData.Num(); ++LODIndex)
@@ -824,9 +824,9 @@ void UStaticMeshComponent::GetStreamingTextureInfo(FStreamingTextureLevelContext
 			if (Scale.X > SMALL_NUMBER && Scale.Y > SMALL_NUMBER)
 			{
 				const float TexelFactor = GetStaticMesh()->LightmapUVDensity * TransformScale / FMath::Min(Scale.X, Scale.Y);
-				new (OutStreamingTextures) FStreamingTexturePrimitiveInfo(Lightmap->GetTexture(LightmapIndex), Bounds, TexelFactor, PackedRelativeBox_Identity);
-				new (OutStreamingTextures) FStreamingTexturePrimitiveInfo(Lightmap->GetAOMaterialMaskTexture(), Bounds, TexelFactor, PackedRelativeBox_Identity);
-				new (OutStreamingTextures) FStreamingTexturePrimitiveInfo(Lightmap->GetSkyOcclusionTexture(), Bounds, TexelFactor, PackedRelativeBox_Identity);
+				new (OutStreamingRenderAssets) FStreamingRenderAssetPrimitiveInfo(Lightmap->GetTexture(LightmapIndex), Bounds, TexelFactor, PackedRelativeBox_Identity);
+				new (OutStreamingRenderAssets) FStreamingRenderAssetPrimitiveInfo(Lightmap->GetAOMaterialMaskTexture(), Bounds, TexelFactor, PackedRelativeBox_Identity);
+				new (OutStreamingRenderAssets) FStreamingRenderAssetPrimitiveInfo(Lightmap->GetSkyOcclusionTexture(), Bounds, TexelFactor, PackedRelativeBox_Identity);
 			}
 		}
 
@@ -837,10 +837,16 @@ void UStaticMeshComponent::GetStreamingTextureInfo(FStreamingTextureLevelContext
 			if (Scale.X > SMALL_NUMBER && Scale.Y > SMALL_NUMBER)
 			{
 				const float TexelFactor = GetStaticMesh()->LightmapUVDensity * TransformScale / FMath::Min(Scale.X, Scale.Y);
-				new (OutStreamingTextures) FStreamingTexturePrimitiveInfo(Shadowmap->GetTexture(), Bounds, TexelFactor, PackedRelativeBox_Identity);
+				new (OutStreamingRenderAssets) FStreamingRenderAssetPrimitiveInfo(Shadowmap->GetTexture(), Bounds, TexelFactor, PackedRelativeBox_Identity);
 			}
 		}
 	}
+
+	if (IsStreamingRenderAsset(GetStaticMesh()))
+	{
+		const float TexelFactor = ForcedLodModel > 0 ? -(GetStaticMesh()->RenderData->LODResources.Num() - ForcedLodModel + 1) : Bounds.SphereRadius * 2.f;
+		new (OutStreamingRenderAssets) FStreamingRenderAssetPrimitiveInfo(GetStaticMesh(), Bounds, TexelFactor, PackedRelativeBox_Identity);
+	}
 }
 
 UBodySetup* UStaticMeshComponent::GetBodySetup()
diff --git a/Engine/Source/Runtime/Engine/Private/ContentStreaming.cpp b/Engine/Source/Runtime/Engine/Private/ContentStreaming.cpp
--- a/Engine/Source/Runtime/Engine/Private/ContentStreaming.cpp
+++ b/Engine/Source/Runtime/Engine/Private/ContentStreaming.cpp
@@ -6,6 +6,7 @@
 
 #include "ContentStreaming.h"
 #include "Engine/Texture2D.h"
+#include "Engine/StaticMesh.h"
 #include "Misc/CommandLine.h"
 #include "Misc/ConfigCacheIni.h"
 #include "UObject/UObjectHash.h"
@@ -49,14 +50,14 @@ void FlushResourceStreaming()
 	Texture tracking.
 -----------------------------------------------------------------------------*/
 
-/** Turn on ENABLE_TEXTURE_TRACKING and setup GTrackedTextures to track specific textures through the streaming system. */
-#define ENABLE_TEXTURE_TRACKING !(UE_BUILD_SHIPPING || UE_BUILD_TEST)
-#define ENABLE_TEXTURE_LOGGING 1
-#if ENABLE_TEXTURE_TRACKING
-struct FTrackedTextureEvent
+/** Turn on ENABLE_RENDER_ASSET_TRACKING and setup GTrackedTextures to track specific textures/meshes through the streaming system. */
+#define ENABLE_RENDER_ASSET_TRACKING !(UE_BUILD_SHIPPING || UE_BUILD_TEST)
+#define ENABLE_RENDER_ASSET_LOGGING 1
+#if ENABLE_RENDER_ASSET_TRACKING
+struct FTrackedRenderAssetEvent
 {
-	FTrackedTextureEvent( TCHAR* InTextureName=NULL )
-	:	TextureName(InTextureName)
+	FTrackedRenderAssetEvent( TCHAR* InAssetName=NULL )
+	:	RenderAssetName(InAssetName)
 	,	NumResidentMips(0)
 	,	NumRequestedMips(0)
 	,	WantedMips(0)
@@ -64,18 +65,17 @@ struct FTrackedTextureEvent
 	,	BoostFactor(1.0f)
 	{
 	}
-	FTrackedTextureEvent( const FString& InTextureNameString )
-	:	NumResidentMips(0)
+	FTrackedRenderAssetEvent( const FString& InAssetNameString )
+	:	RenderAssetName(InAssetNameString)
+	,	NumResidentMips(0)
 	,	NumRequestedMips(0)
 	,	WantedMips(0)
 	,	Timestamp(0.0f)
 	,	BoostFactor(1.0f)
 	{
-		TextureName = new TCHAR[InTextureNameString.Len() + 1];
-		FMemory::Memcpy( TextureName, *InTextureNameString, (InTextureNameString.Len() + 1)*sizeof(TCHAR) );
 	}
-	/** Partial name of the texture (not case-sensitive). */
-	TCHAR*			TextureName;
+	/** Partial name of the texture/mesh (not case-sensitive). */
+	FString			RenderAssetName;
 	/** Number of mip-levels currently in memory. */
 	int32			NumResidentMips;
 	/** Number of mip-levels requested. */
@@ -87,77 +87,96 @@ struct FTrackedTextureEvent
 	/** Currently used boost factor for the streaming distance. */
 	float			BoostFactor;
 };
-/** List of textures to track (using stristr for name comparison). */
-TArray<FString> GTrackedTextureNames;
-bool GTrackedTexturesInitialized = false;
-#define NUM_TRACKEDTEXTUREEVENTS 512
-FTrackedTextureEvent GTrackedTextureEvents[NUM_TRACKEDTEXTUREEVENTS];
-int32 GTrackedTextureEventIndex = -1;
-TArray<FTrackedTextureEvent> GTrackedTextures;
+/** List of textures/meshes to track (using stristr for name comparison). */
+TArray<FString> GTrackedRenderAssetNames;
+bool GTrackedRenderAssetsInitialized = false;
+#define NUM_TRACKEDRENDERASSETEVENTS 512
+FTrackedRenderAssetEvent GTrackedRenderAssetEvents[NUM_TRACKEDRENDERASSETEVENTS];
+int32 GTrackedRenderAssetEventIndex = -1;
+TArray<FTrackedRenderAssetEvent> GTrackedRenderAssets;
 
 /**
- * Initializes the texture tracking. Called when GTrackedTexturesInitialized is false.
+ * Initializes the texture/mesh tracking. Called when GTrackedRenderAssetsInitialized is false.
  */
-void TrackTextureInit()
+void TrackRenderAssetInit()
 {
 	if ( GConfig && GConfig->Num() > 0 )
 	{
-		GTrackedTextureNames.Empty();
-		GTrackedTexturesInitialized = true;
-		GConfig->GetArray( TEXT("TextureTracking"), TEXT("TextureName"), GTrackedTextureNames, GEngineIni );
+		GTrackedRenderAssetNames.Empty();
+		GTrackedRenderAssetsInitialized = true;
+		GConfig->GetArray(TEXT("RenderAssetTracking"), TEXT("RenderAssetName"), GTrackedRenderAssetNames, GEngineIni);
+		TArray<FString> Tmp;
+		GConfig->GetArray( TEXT("TextureTracking"), TEXT("TextureName"), Tmp, GEngineIni );
+		if (Tmp.Num() > 0)
+		{
+			GTrackedRenderAssetNames.Append(MoveTemp(Tmp));
+			check(!Tmp.Num());
+			GConfig->SetArray(TEXT("TextureTracking"), TEXT("TextureName"), Tmp, GEngineIni);
+			GConfig->SetArray(TEXT("RenderAssetTracking"), TEXT("RenderAssetName"), GTrackedRenderAssetNames, GEngineIni);
+		}
 	}
 }
 
+void TrackTextureInit()
+{
+	TrackRenderAssetInit();
+}
+
 /**
- * Adds a (partial) texture name to track in the streaming system and updates the .ini setting.
+ * Adds a (partial) texture/mesh name to track in the streaming system and updates the .ini setting.
  *
- * @param TextureName	Partial name of a new texture to track (not case-sensitive)
+ * @param AssetName		Partial name of a new texture/mesh to track (not case-sensitive)
  * @return				true if the name was added
  */
-bool TrackTexture( const FString& TextureName )
+bool TrackRenderAsset( const FString& AssetName )
 {
-	if ( GConfig && TextureName.Len() > 0 )
+	if ( GConfig && AssetName.Len() > 0 )
 	{
-		for ( int32 TrackedTextureIndex=0; TrackedTextureIndex < GTrackedTextureNames.Num(); ++TrackedTextureIndex )
+		for ( int32 TrackedAssetIndex=0; TrackedAssetIndex < GTrackedRenderAssetNames.Num(); ++TrackedAssetIndex)
 		{
-			const FString& TrackedTextureName = GTrackedTextureNames[TrackedTextureIndex];
-			if ( FCString::Stricmp(*TextureName, *TrackedTextureName) == 0 )
+			const FString& TrackedAssetName = GTrackedRenderAssetNames[TrackedAssetIndex];
+			if ( FCString::Stricmp(*AssetName, *TrackedAssetName) == 0 )
 			{
 				return false;	
 			}
 		}
 
-		GTrackedTextureNames.Add( *TextureName );
-		GConfig->SetArray( TEXT("TextureTracking"), TEXT("TextureName"), GTrackedTextureNames, GEngineIni );
+		GTrackedRenderAssetNames.Add( *AssetName );
+		GConfig->SetArray( TEXT("RenderAssetTracking"), TEXT("RenderAssetName"), GTrackedRenderAssetNames, GEngineIni );
 		return true;
 	}
 	return false;
 }
 
+bool TrackTexture(const FString& TextureName)
+{
+	return TrackRenderAsset(TextureName);
+}
+
 /**
- * Removes a texture name from being tracked in the streaming system and updates the .ini setting.
+ * Removes a texture/mesh name from being tracked in the streaming system and updates the .ini setting.
  * The name must match an existing tracking name, but isn't case-sensitive.
  *
- * @param TextureName	Name of a new texture to stop tracking (not case-sensitive)
- * @return				true if the name was added
+ * @param AssetName		Name of a texture/mesh to stop tracking (not case-sensitive)
+ * @return				true if the name was removed
  */
-bool UntrackTexture( const FString& TextureName )
+bool UntrackRenderAsset( const FString& AssetName )
 {
-	if ( GConfig && TextureName.Len() > 0 )
+	if ( GConfig && AssetName.Len() > 0 )
 	{
-		int32 TrackedTextureIndex = 0;
-		for ( ; TrackedTextureIndex < GTrackedTextureNames.Num(); ++TrackedTextureIndex )
+		int32 TrackedAssetIndex = 0;
+		for ( ; TrackedAssetIndex < GTrackedRenderAssetNames.Num(); ++TrackedAssetIndex)
 		{
-			const FString& TrackedTextureName = GTrackedTextureNames[TrackedTextureIndex];
-			if ( FCString::Stricmp(*TextureName, *TrackedTextureName) == 0 )
+			const FString& TrackedAssetName = GTrackedRenderAssetNames[TrackedAssetIndex];
+			if ( FCString::Stricmp(*AssetName, *TrackedAssetName) == 0 )
 			{
 				break;
 			}
 		}
-		if ( TrackedTextureIndex < GTrackedTextureNames.Num() )
+		if ( TrackedAssetIndex < GTrackedRenderAssetNames.Num() )
 		{
-			GTrackedTextureNames.RemoveAt( TrackedTextureIndex );
-			GConfig->SetArray( TEXT("TextureTracking"), TEXT("TextureName"), GTrackedTextureNames, GEngineIni );
+			GTrackedRenderAssetNames.RemoveAt( TrackedAssetIndex );
+			GConfig->SetArray( TEXT("RenderAssetTracking"), TEXT("RenderAssetName"), GTrackedRenderAssetNames, GEngineIni );
 
 			return true;
 		}
@@ -165,61 +184,71 @@ bool UntrackTexture( const FString& TextureName )
 	return false;
 }
 
+bool UntrackTexture(const FString& TextureName)
+{
+	return UntrackRenderAsset(TextureName);
+}
+
 /**
- * Lists all currently tracked texture names in the specified log.
+ * Lists all currently tracked texture/mesh names in the specified log.
  *
  * @param Ar			Desired output log
- * @param NumTextures	Maximum number of tracked texture names to output. Outputs all if NumTextures <= 0.
+ * @param NumAssets		Maximum number of tracked texture/mesh names to output. Outputs all if NumAssets <= 0.
  */
-void ListTrackedTextures( FOutputDevice& Ar, int32 NumTextures )
+void ListTrackedRenderAssets( FOutputDevice& Ar, int32 NumAssets )
 {
-	NumTextures = (NumTextures > 0) ? FMath::Min(NumTextures, GTrackedTextureNames.Num()) : GTrackedTextureNames.Num();
-	for ( int32 TrackedTextureIndex=0; TrackedTextureIndex < NumTextures; ++TrackedTextureIndex )
+	NumAssets = (NumAssets > 0) ? FMath::Min(NumAssets, GTrackedRenderAssetNames.Num()) : GTrackedRenderAssetNames.Num();
+	for ( int32 TrackedAssetIndex=0; TrackedAssetIndex < NumAssets; ++TrackedAssetIndex )
 	{
-		const FString& TrackedTextureName = GTrackedTextureNames[TrackedTextureIndex];
-		Ar.Log( TrackedTextureName );
+		const FString& TrackedAssetName = GTrackedRenderAssetNames[TrackedAssetIndex];
+		Ar.Log( TrackedAssetName );
 	}
-	Ar.Logf( TEXT("%d textures are being tracked."), NumTextures );
+	Ar.Logf( TEXT("%d render assets are being tracked."), NumAssets );
+}
+
+void ListTrackedTextures(FOutputDevice& Ar, int32 NumTextures)
+{
+	ListTrackedRenderAssets(Ar, NumTextures);
 }
 
 /**
- * Checks a texture and tracks it if its name contains any of the tracked textures names (GTrackedTextureNames).
+ * Checks a texture/mesh and tracks it if its name contains any of the tracked render asset names (GTrackedRenderAssetNames).
  *
- * @param Texture						Texture to check
- * @param bForceMipLEvelsToBeResident	Whether all mip-levels in the texture are forced to be resident
- * @param Manager                       can be 0
+ * @param RenderAsset					Texture/mesh to check
+ * @param bForceMipLEvelsToBeResident	Whether all mip-levels in the texture/mesh are forced to be resident
+ * @param Manager                       can be null
  */
-bool TrackTextureEvent( FStreamingTexture* StreamingTexture, UTexture2D* Texture, bool bForceMipLevelsToBeResident, const FStreamingManagerTexture* Manager)
+bool TrackRenderAssetEvent(FStreamingRenderAsset* StreamingRenderAsset, UStreamableRenderAsset* RenderAsset, bool bForceMipLevelsToBeResident, const FRenderAssetStreamingManager* Manager)
 {
-	// Whether the texture is currently being destroyed
-	const bool bIsDestroying = (StreamingTexture == 0);
-	const bool bEnableLogging = ENABLE_TEXTURE_LOGGING;
+	// Whether the texture/mesh is currently being destroyed
+	const bool bIsDestroying = !StreamingRenderAsset;
+	const bool bEnableLogging = ENABLE_RENDER_ASSET_LOGGING;
 
 	// Initialize the tracking system, if necessary.
-	if ( GTrackedTexturesInitialized == false )
+	if ( !GTrackedRenderAssetsInitialized )
 	{
-		TrackTextureInit();
+		TrackRenderAssetInit();
 	}
 
-	int32 NumTrackedTextures = GTrackedTextureNames.Num();
-	if ( NumTrackedTextures )
+	int32 NumTrackedAssets = GTrackedRenderAssetNames.Num();
+	if ( NumTrackedAssets )
 	{
-		// See if it matches any of the texture names that we're tracking.
-		FString TextureNameString = Texture->GetFullName();
-		const TCHAR* TextureName = *TextureNameString;
-		for ( int32 TrackedTextureIndex=0; TrackedTextureIndex < NumTrackedTextures; ++TrackedTextureIndex )
+		// See if it matches any of the texture/mesh names that we're tracking.
+		FString AssetNameString = RenderAsset->GetFullName();
+		const TCHAR* AssetName = *AssetNameString;
+		for ( int32 TrackedAssetIndex=0; TrackedAssetIndex < NumTrackedAssets; ++TrackedAssetIndex )
 		{
-			const FString& TrackedTextureName = GTrackedTextureNames[TrackedTextureIndex];
-			if ( FCString::Stristr(TextureName, *TrackedTextureName) != NULL )
+			const FString& TrackedAssetName = GTrackedRenderAssetNames[TrackedAssetIndex];
+			if ( FCString::Stristr(AssetName, *TrackedAssetName) != NULL )
 			{
 				if ( bEnableLogging )
 				{
-					// Find the last event for this particular texture.
-					FTrackedTextureEvent* LastEvent = NULL;
-					for ( int32 LastEventIndex=0; LastEventIndex < GTrackedTextures.Num(); ++LastEventIndex )
+					// Find the last event for this particular texture/mesh.
+					FTrackedRenderAssetEvent* LastEvent = NULL;
+					for ( int32 LastEventIndex=0; LastEventIndex < GTrackedRenderAssets.Num(); ++LastEventIndex )
 					{
-						FTrackedTextureEvent* Event = &GTrackedTextures[LastEventIndex];
-						if ( FCString::Strcmp(TextureName, Event->TextureName) == 0 )
+						FTrackedRenderAssetEvent* Event = &GTrackedRenderAssets[LastEventIndex];
+						if ( FCString::Strcmp(AssetName, *Event->RenderAssetName) == 0 )
 						{
 							LastEvent = Event;
 							break;
@@ -228,34 +257,35 @@ bool TrackTextureEvent( FStreamingTexture* StreamingTexture, UTexture2D* Texture
 					// Didn't find any recorded event?
 					if ( LastEvent == NULL )
 					{
-						int32 NewIndex = GTrackedTextures.Add(FTrackedTextureEvent(TextureNameString));
-						LastEvent = &GTrackedTextures[NewIndex];
+						int32 NewIndex = GTrackedRenderAssets.Add(FTrackedRenderAssetEvent(AssetNameString));
+						LastEvent = &GTrackedRenderAssets[NewIndex];
 					}
 
-					int32 WantedMips		= Texture->GetNumRequestedMips();
+					int32 WantedMips		= RenderAsset->GetNumRequestedMips();
 					float BoostFactor		= 1.0f;
-					if ( StreamingTexture )
+					if ( StreamingRenderAsset )
 					{
-						WantedMips			= StreamingTexture->WantedMips;
-						BoostFactor			= StreamingTexture->BoostFactor;
+						WantedMips			= StreamingRenderAsset->WantedMips;
+						BoostFactor			= StreamingRenderAsset->BoostFactor;
 					}
 
-					if ( LastEvent->NumResidentMips != Texture->GetNumResidentMips() ||
-						 LastEvent->NumRequestedMips != Texture->GetNumRequestedMips() ||
+					if ( LastEvent->NumResidentMips != RenderAsset->GetNumResidentMips() ||
+						 LastEvent->NumRequestedMips != RenderAsset->GetNumRequestedMips() ||
 						 LastEvent->WantedMips != WantedMips ||
 						 LastEvent->BoostFactor != BoostFactor ||
 						 bIsDestroying )
 					{
-						GTrackedTextureEventIndex		= (GTrackedTextureEventIndex + 1) % NUM_TRACKEDTEXTUREEVENTS;
-						FTrackedTextureEvent& NewEvent	= GTrackedTextureEvents[GTrackedTextureEventIndex];
-						NewEvent.TextureName			= LastEvent->TextureName;
-						NewEvent.NumResidentMips		= LastEvent->NumResidentMips	= Texture->GetNumResidentMips();
-						NewEvent.NumRequestedMips		= LastEvent->NumRequestedMips	= Texture->GetNumRequestedMips();
+						GTrackedRenderAssetEventIndex	= (GTrackedRenderAssetEventIndex + 1) % NUM_TRACKEDRENDERASSETEVENTS;
+						FTrackedRenderAssetEvent& NewEvent	= GTrackedRenderAssetEvents[GTrackedRenderAssetEventIndex];
+						NewEvent.RenderAssetName		= LastEvent->RenderAssetName;
+						NewEvent.NumResidentMips		= LastEvent->NumResidentMips	= RenderAsset->GetNumResidentMips();
+						NewEvent.NumRequestedMips		= LastEvent->NumRequestedMips	= RenderAsset->GetNumRequestedMips();
 						NewEvent.WantedMips				= LastEvent->WantedMips			= WantedMips;
 						NewEvent.Timestamp				= LastEvent->Timestamp			= float(FPlatformTime::Seconds() - GStartTime);
 						NewEvent.BoostFactor			= LastEvent->BoostFactor		= BoostFactor;
-						UE_LOG(LogContentStreaming, Log, TEXT("Texture: \"%s\", ResidentMips: %d/%d, RequestedMips: %d, WantedMips: %d, Boost: %.1f (%s)"),
-							TextureName, LastEvent->NumResidentMips, Texture->GetNumMips(), bIsDestroying ? 0 : LastEvent->NumRequestedMips, LastEvent->WantedMips, 
+						UE_LOG(LogContentStreaming, Log, TEXT("%s: \"%s\", ResidentMips: %d/%d, RequestedMips: %d, WantedMips: %d, Boost: %.1f (%s)"),
+							FStreamingRenderAsset::GetStreamingAssetTypeStr(StreamingRenderAsset->RenderAssetType),
+							AssetName, LastEvent->NumResidentMips, RenderAsset->GetNumMipsForStreaming(), bIsDestroying ? 0 : LastEvent->NumRequestedMips, LastEvent->WantedMips, 
 							BoostFactor, bIsDestroying ? TEXT("DESTROYED") : TEXT("updated") );
 					}
 				}
@@ -265,22 +295,43 @@ bool TrackTextureEvent( FStreamingTexture* StreamingTexture, UTexture2D* Texture
 	}
 	return false;
 }
+
+bool TrackTextureEvent(FStreamingRenderAsset* StreamingTexture, UStreamableRenderAsset* Texture, bool bForceMipLevelsToBeResident, const FRenderAssetStreamingManager* Manager)
+{
+	return TrackRenderAssetEvent(StreamingTexture, Texture, bForceMipLevelsToBeResident, Manager);
+}
 #else
-bool TrackTexture( const FString& /*TextureName*/ )
+bool TrackRenderAsset(const FString& AssetName)
 {
 	return false;
 }
-bool UntrackTexture( const FString& /*TextureName*/ )
+bool TrackTexture( const FString& TextureName )
+{
+	return TrackRenderAsset(TextureName);
+}
+bool UntrackRenderAsset(const FString& AssetName)
 {
 	return false;
 }
-void ListTrackedTextures( FOutputDevice& /*Ar*/, int32 /*NumTextures*/ )
+bool UntrackTexture( const FString& TextureName )
+{
+	return UntrackRenderAsset(TextureName);
+}
+void ListTrackedRenderAssets(FOutputDevice& Ar, int32 NumAssets)
+{
+}
+void ListTrackedTextures( FOutputDevice& Ar, int32 NumTextures )
 {
+	ListTrackedRenderAssets(Ar, NumTextures);
 }
-bool TrackTextureEvent( FStreamingTexture* StreamingTexture, UTexture2D* Texture, bool bForceMipLevelsToBeResident, const FStreamingManagerTexture* Manager )
+bool TrackRenderAssetEvent(FStreamingRenderAsset* StreamingRenderAsset, UStreamableRenderAsset* RenderAsset, bool bForceMipLevelsToBeResident, const FRenderAssetStreamingManager* Manager)
 {
 	return false;
 }
+bool TrackTextureEvent( FStreamingRenderAsset* StreamingTexture, UStreamableRenderAsset* Texture, bool bForceMipLevelsToBeResident, const FRenderAssetStreamingManager* Manager )
+{
+	return TrackRenderAssetEvent(StreamingTexture, Texture, bForceMipLevelsToBeResident, Manager);
+}
 #endif
 
 /*-----------------------------------------------------------------------------
@@ -859,7 +910,7 @@ bool FStreamingManagerCollection::IsTextureStreamingEnabled() const
 	return TextureStreamingManager != 0;
 }
 
-ITextureStreamingManager& FStreamingManagerCollection::GetTextureStreamingManager() const
+IRenderAssetStreamingManager& FStreamingManagerCollection::GetTextureStreamingManager() const
 {
 	check(TextureStreamingManager != 0);
 	return *TextureStreamingManager;
@@ -1063,7 +1114,7 @@ void FStreamingManagerCollection::AddOrRemoveTextureStreamingManagerIfNeeded(boo
 		{
 			GConfig->GetFloat( TEXT("TextureStreaming"), TEXT("LoadMapTimeLimit"), LoadMapTimeLimit, GEngineIni );
 			// Create the streaming manager and add the default streamers.
-			TextureStreamingManager = new FStreamingManagerTexture();
+			TextureStreamingManager = new FRenderAssetStreamingManager();
 			AddStreamingManager( TextureStreamingManager );		
 				
 			//Need to work out if all textures should be streamable and added to the texture streaming manager.
@@ -1074,6 +1125,13 @@ void FStreamingManagerCollection::AddOrRemoveTextureStreamingManagerIfNeeded(boo
 				{
 					It->UpdateResource();
 				}
+				for (TObjectIterator<UStaticMesh> It; It; ++It)
+				{
+					if (It->bIsStreamable)
+					{
+						It->LinkStreaming();
+					}
+				}
 			}
 		}
 	}
@@ -1095,6 +1153,14 @@ void FStreamingManagerCollection::AddOrRemoveTextureStreamingManagerIfNeeded(boo
 					It->UpdateResource();
 				}
 			}
+			for (TObjectIterator<UStaticMesh> It; It; ++It)
+			{
+				if (It->bIsStreamable)
+				{
+					// This will clear StreamingIndex
+					It->LinkStreaming();
+				}
+			}
 		}
 	}
 }
diff --git a/Engine/Source/Runtime/Engine/Private/HierarchicalInstancedStaticMesh.cpp b/Engine/Source/Runtime/Engine/Private/HierarchicalInstancedStaticMesh.cpp
--- a/Engine/Source/Runtime/Engine/Private/HierarchicalInstancedStaticMesh.cpp
+++ b/Engine/Source/Runtime/Engine/Private/HierarchicalInstancedStaticMesh.cpp
@@ -1251,7 +1251,7 @@ void FHierarchicalStaticMeshSceneProxy::FillDynamicMeshElements(FMeshElementColl
 	int64 TotalTriangles = 0;
 
 	int32 OnlyLOD = FMath::Min<int32>(CVarOnlyLOD.GetValueOnRenderThread(),InstancedRenderData.VertexFactories.Num() - 1);
-	int32 FirstLOD = (OnlyLOD < 0) ? 0 : OnlyLOD;
+	int32 FirstLOD = FMath::Max((OnlyLOD < 0) ? 0 : OnlyLOD, static_cast<int32>(this->GetCurrentFirstLODIdx_Internal()));
 	int32 LastLODPlusOne = (OnlyLOD < 0) ? InstancedRenderData.VertexFactories.Num() : (OnlyLOD+1);
 
 	for (int32 LODIndex = FirstLOD; LODIndex < LastLODPlusOne; LODIndex++)
diff --git a/Engine/Source/Runtime/Engine/Private/InstancedStaticMesh.cpp b/Engine/Source/Runtime/Engine/Private/InstancedStaticMesh.cpp
--- a/Engine/Source/Runtime/Engine/Private/InstancedStaticMesh.cpp
+++ b/Engine/Source/Runtime/Engine/Private/InstancedStaticMesh.cpp
@@ -2021,12 +2021,12 @@ bool UInstancedStaticMeshComponent::BuildTextureStreamingData(ETextureStreamingB
 	return true;
 }
 
-void UInstancedStaticMeshComponent::GetStreamingTextureInfo(FStreamingTextureLevelContext& LevelContext, TArray<FStreamingTexturePrimitiveInfo>& OutStreamingTextures) const
+void UInstancedStaticMeshComponent::GetStreamingRenderAssetInfo(FStreamingTextureLevelContext& LevelContext, TArray<FStreamingRenderAssetPrimitiveInfo>& OutStreamingRenderAssets) const
 {
 	// Don't only look the instance count but also if the bound is valid, as derived classes might not set PerInstanceSMData.
 	if (GetInstanceCount() > 0 || Bounds.SphereRadius > 0)
 	{
-		return Super::GetStreamingTextureInfo(LevelContext, OutStreamingTextures);
+		return Super::GetStreamingRenderAssetInfo(LevelContext, OutStreamingRenderAssets);
 	}
 }
 
diff --git a/Engine/Source/Runtime/Engine/Private/Materials/MaterialInterface.cpp b/Engine/Source/Runtime/Engine/Private/Materials/MaterialInterface.cpp
--- a/Engine/Source/Runtime/Engine/Private/Materials/MaterialInterface.cpp
+++ b/Engine/Source/Runtime/Engine/Private/Materials/MaterialInterface.cpp
@@ -622,7 +622,7 @@ bool UMaterialInterface::UseAnyStreamingTexture() const
 
 	for (UTexture* Texture : Textures)
 	{
-		if (IsStreamingTexture(Cast<UTexture2D>(Texture)))
+		if (IsStreamingRenderAsset(Texture))
 		{
 			return true;
 		}
