Resolve robomerge conflicts

diff --git a/Engine/Source/Runtime/Engine/Private/Particles/ParticleComponents.cpp b/Engine/Source/Runtime/Engine/Private/Particles/ParticleComponents.cpp
--- a/Engine/Source/Runtime/Engine/Private/Particles/ParticleComponents.cpp
+++ b/Engine/Source/Runtime/Engine/Private/Particles/ParticleComponents.cpp
@@ -7530,7 +7530,7 @@ void AddMaterials(TArray<FMaterialWithScale, TInlineAllocator<12> >& OutMaterial
 	}
 }
 
-void UParticleSystemComponent::GetStreamingTextureInfo(FStreamingTextureLevelContext& LevelContext, TArray<FStreamingTexturePrimitiveInfo>& OutStreamingTextures) const
+void UParticleSystemComponent::GetStreamingRenderAssetInfo(FStreamingTextureLevelContext& LevelContext, TArray<FStreamingRenderAssetPrimitiveInfo>& OutStreamingRenderAssets) const
 {
 	TArray<FMaterialWithScale, TInlineAllocator<12> > MaterialWithScales;
 
@@ -7573,7 +7573,7 @@ void UParticleSystemComponent::GetStreamingTextureInfo(FStreamingTextureLevelCon
 			for (const FMaterialWithScale& MaterialWithScale : MaterialWithScales)
 			{
 				MaterialData.Material = MaterialWithScale.Key;
-				LevelContext.ProcessMaterial(Bounds, MaterialData, Bounds.SphereRadius * MaterialWithScale.Value, OutStreamingTextures);
+				LevelContext.ProcessMaterial(Bounds, MaterialData, Bounds.SphereRadius * MaterialWithScale.Value, OutStreamingRenderAssets);
 			}
 		}
 	}
diff --git a/Engine/Source/Runtime/Engine/Private/Particles/ParticleSystemRender.cpp b/Engine/Source/Runtime/Engine/Private/Particles/ParticleSystemRender.cpp
--- a/Engine/Source/Runtime/Engine/Private/Particles/ParticleSystemRender.cpp
+++ b/Engine/Source/Runtime/Engine/Private/Particles/ParticleSystemRender.cpp
@@ -1575,7 +1575,8 @@ public:
 
 uint32 FDynamicMeshEmitterData::GetMeshLODIndexFromProxy(const FParticleSystemSceneProxy *InOwnerProxy) const
 {
-	int FirstAvailableLOD = 0;
+	check(IsInRenderingThread());
+	int FirstAvailableLOD = StaticMesh->RenderData->CurrentFirstLODIdx;
 	for (; FirstAvailableLOD < StaticMesh->RenderData->LODResources.Num(); FirstAvailableLOD++)
 	{
 		if (StaticMesh->RenderData->LODResources[FirstAvailableLOD].GetNumVertices() > 0)
diff --git a/Engine/Source/Runtime/Engine/Private/PositionVertexBuffer.cpp b/Engine/Source/Runtime/Engine/Private/PositionVertexBuffer.cpp
--- a/Engine/Source/Runtime/Engine/Private/PositionVertexBuffer.cpp
+++ b/Engine/Source/Runtime/Engine/Private/PositionVertexBuffer.cpp
@@ -29,7 +29,8 @@ FPositionVertexBuffer::FPositionVertexBuffer():
 	VertexData(NULL),
 	Data(NULL),
 	Stride(0),
-	NumVertices(0)
+	NumVertices(0),
+	bStreamed(false)
 {}
 
 FPositionVertexBuffer::~FPositionVertexBuffer()
@@ -151,7 +152,7 @@ void FPositionVertexBuffer::Serialize( FArchive& Ar, bool bInNeedsCPUAccess )
 {
 	bNeedsCPUAccess = bInNeedsCPUAccess;
 
-	Ar << Stride << NumVertices;
+	SerializeMetaData(Ar);
 
 	if(Ar.IsLoading())
 	{
@@ -169,6 +170,11 @@ void FPositionVertexBuffer::Serialize( FArchive& Ar, bool bInNeedsCPUAccess )
 	}
 }
 
+void FPositionVertexBuffer::SerializeMetaData(FArchive& Ar)
+{
+	Ar << Stride << NumVertices;
+}
+
 /**
 * Specialized assignment operator, only used when importing LOD's.  
 */
@@ -178,7 +184,8 @@ void FPositionVertexBuffer::operator=(const FPositionVertexBuffer &Other)
 	VertexData = NULL;
 }
 
-void FPositionVertexBuffer::InitRHI()
+template <bool bRenderThread>
+FVertexBufferRHIRef FPositionVertexBuffer::CreateRHIBuffer_Internal()
 {
 	check(VertexData);
 	FResourceArrayInterface* ResourceArray = VertexData->GetResourceArray();
@@ -186,20 +193,44 @@ void FPositionVertexBuffer::InitRHI()
 	{
 		// Create the vertex buffer.
 		FRHIResourceCreateInfo CreateInfo(ResourceArray);
-		VertexBufferRHI = RHICreateVertexBuffer(ResourceArray->GetResourceDataSize(), BUF_Static | BUF_ShaderResource, CreateInfo);
-
-		// we have decide to create the SRV based on GMaxRHIShaderPlatform because this is created once and shared between feature levels for editor preview.
-		if (RHISupportsManualVertexFetch(GMaxRHIShaderPlatform) || IsGPUSkinCacheAvailable())
+		if (bRenderThread)
+		{
+			return RHICreateVertexBuffer(ResourceArray->GetResourceDataSize(), BUF_Static | BUF_ShaderResource, CreateInfo);
+		}
+		else
 		{
-			PositionComponentSRV = RHICreateShaderResourceView(VertexBufferRHI, 4, PF_R32_FLOAT);
+			return RHIAsyncCreateVertexBuffer(ResourceArray->GetResourceDataSize(), BUF_Static | BUF_ShaderResource, CreateInfo);
 		}
 	}
+	return nullptr;
+}
+
+FVertexBufferRHIRef FPositionVertexBuffer::CreateRHIBuffer_RenderThread()
+{
+	return CreateRHIBuffer_Internal<true>();
+}
+
+FVertexBufferRHIRef FPositionVertexBuffer::CreateRHIBuffer_Async()
+{
+	return CreateRHIBuffer_Internal<false>();
+}
+
+void FPositionVertexBuffer::InitRHI()
+{
+	if (!bStreamed)
+	{
+		VertexBufferRHI = CreateRHIBuffer_RenderThread();
+	}
+	// we have decide to create the SRV based on GMaxRHIShaderPlatform because this is created once and shared between feature levels for editor preview.
+	if ((bStreamed || VertexBufferRHI) && (RHISupportsManualVertexFetch(GMaxRHIShaderPlatform) || IsGPUSkinCacheAvailable()))
+	{
+		PositionComponentSRV = RHICreateShaderResourceView(VertexBufferRHI, 4, PF_R32_FLOAT);
+	}
 }
 
 void FPositionVertexBuffer::ReleaseRHI()
 {
 	PositionComponentSRV.SafeRelease();
-
 	FVertexBuffer::ReleaseRHI();
 }
 
diff --git a/Engine/Source/Runtime/Engine/Private/RawIndexBuffer.cpp b/Engine/Source/Runtime/Engine/Private/RawIndexBuffer.cpp
--- a/Engine/Source/Runtime/Engine/Private/RawIndexBuffer.cpp
+++ b/Engine/Source/Runtime/Engine/Private/RawIndexBuffer.cpp
@@ -140,7 +140,9 @@ FArchive& operator<<(FArchive& Ar,FRawIndexBuffer16or32& I)
 
 FRawStaticIndexBuffer::FRawStaticIndexBuffer(bool InNeedsCPUAccess)
 	: IndexStorage(InNeedsCPUAccess)
+	, CachedNumIndices(-1)
 	, b32Bit(false)
+	, bStreamed(false)
 {
 }
 
@@ -168,6 +170,7 @@ void FRawStaticIndexBuffer::SetIndices(const TArray<uint32>& InIndices, EIndexBu
 	int32 IndexStride = bShouldUse32Bit ? sizeof(uint32) : sizeof(uint16);
 	IndexStorage.Empty(IndexStride * NumIndices);
 	IndexStorage.AddUninitialized(IndexStride * NumIndices);
+	CachedNumIndices = NumIndices;
 
 	// Store them!
 	if (bShouldUse32Bit)
@@ -197,6 +200,7 @@ void FRawStaticIndexBuffer::InsertIndices( const uint32 At, const uint32* Indice
 		const uint32 IndexStride = b32Bit ? sizeof( uint32 ) : sizeof( uint16 );
 
 		IndexStorage.InsertUninitialized( At * IndexStride, NumIndicesToAppend * IndexStride );
+		CachedNumIndices = IndexStorage.Num() / static_cast<int32>(IndexStride);
 		uint8* const DestIndices = &IndexStorage[ At * IndexStride ];
 
 		if( IndicesToAppend )
@@ -235,6 +239,7 @@ void FRawStaticIndexBuffer::RemoveIndicesAt( const uint32 At, const uint32 NumIn
 	{
 		const int32 IndexStride = b32Bit ? sizeof( uint32 ) : sizeof( uint16 );
 		IndexStorage.RemoveAt( At * IndexStride, NumIndicesToRemove * IndexStride );
+		CachedNumIndices = IndexStorage.Num() / IndexStride;
 	}
 }
 
@@ -277,17 +282,50 @@ FIndexArrayView FRawStaticIndexBuffer::GetArrayView() const
 	return FIndexArrayView(IndexStorage.GetData(),NumIndices,b32Bit);
 }
 
-void FRawStaticIndexBuffer::InitRHI()
+template <bool bRenderThread>
+FIndexBufferRHIRef FRawStaticIndexBuffer::CreateRHIBuffer_Internal()
 {
-	uint32 IndexStride = b32Bit ? sizeof(uint32) : sizeof(uint16);
-	uint32 SizeInBytes = IndexStorage.Num();
+	const uint32 IndexStride = b32Bit ? sizeof(uint32) : sizeof(uint16);
+	const uint32 SizeInBytes = IndexStorage.Num();
 
 	if (SizeInBytes > 0)
 	{
 		// Create the index buffer.
 		FRHIResourceCreateInfo CreateInfo(&IndexStorage);
-		IndexBufferRHI = RHICreateIndexBuffer(IndexStride,SizeInBytes,BUF_Static,CreateInfo);
-	}    
+		if (bRenderThread)
+		{
+			return RHICreateIndexBuffer(IndexStride, SizeInBytes, BUF_Static, CreateInfo);
+		}
+		else
+		{
+			return RHIAsyncCreateIndexBuffer(IndexStride, SizeInBytes, BUF_Static, CreateInfo);
+		}
+	}
+	return nullptr;
+}
+
+FIndexBufferRHIRef FRawStaticIndexBuffer::CreateRHIBuffer_RenderThread()
+{
+	return CreateRHIBuffer_Internal<true>();
+}
+
+FIndexBufferRHIRef FRawStaticIndexBuffer::CreateRHIBuffer_Async()
+{
+	return CreateRHIBuffer_Internal<false>();
+}
+
+void FRawStaticIndexBuffer::InitRHIForStreaming(FIndexBufferRHIParamRef IntermediateBuffer)
+{
+	check(!IndexBufferRHI);
+	IndexBufferRHI = IntermediateBuffer;
+}
+
+void FRawStaticIndexBuffer::InitRHI()
+{
+	if (!bStreamed)
+	{
+		IndexBufferRHI = CreateRHIBuffer_RenderThread();
+	}
 }
 
 void FRawStaticIndexBuffer::Serialize(FArchive& Ar, bool bNeedsCPUAccess)
@@ -311,12 +349,19 @@ void FRawStaticIndexBuffer::Serialize(FArchive& Ar, bool bNeedsCPUAccess)
 		Ar << b32Bit;
 		IndexStorage.BulkSerialize(Ar);
 	}
+	CachedNumIndices = IndexStorage.Num() / (b32Bit ? 4 : 2);
+}
+
+void FRawStaticIndexBuffer::SerializeMetaData(FArchive& Ar)
+{
+	Ar << CachedNumIndices << b32Bit;
 }
 
 void FRawStaticIndexBuffer::Discard()
 {
     IndexStorage.SetAllowCPUAccess(false);
     IndexStorage.Discard();
+	CachedNumIndices = IndexStorage.Num() / (b32Bit ? 4 : 2);
 }
 
 /*-----------------------------------------------------------------------------
diff --git a/Engine/Source/Runtime/Engine/Private/Rendering/ColorVertexBuffer.cpp b/Engine/Source/Runtime/Engine/Private/Rendering/ColorVertexBuffer.cpp
--- a/Engine/Source/Runtime/Engine/Private/Rendering/ColorVertexBuffer.cpp
+++ b/Engine/Source/Runtime/Engine/Private/Rendering/ColorVertexBuffer.cpp
@@ -27,7 +27,8 @@ FColorVertexBuffer::FColorVertexBuffer():
 	VertexData(NULL),
 	Data(NULL),
 	Stride(0),
-	NumVertices(0)
+	NumVertices(0),
+	bStreamed(false)
 {
 }
 
@@ -212,7 +213,7 @@ void FColorVertexBuffer::Serialize( FArchive& Ar, bool bNeedsCPUAccess )
 	}
 	else
 	{
-		Ar << Stride << NumVertices;
+		SerializeMetaData(Ar);
 
 		if (Ar.IsLoading() && NumVertices > 0)
 		{
@@ -237,6 +238,11 @@ void FColorVertexBuffer::Serialize( FArchive& Ar, bool bNeedsCPUAccess )
 	}
 }
 
+void FColorVertexBuffer::SerializeMetaData(FArchive& Ar)
+{
+	Ar << Stride << NumVertices;
+}
+
 
 /** Export the data to a string, used for editor Copy&Paste. */
 void FColorVertexBuffer::ExportText(FString &ValueStr) const
@@ -375,25 +381,54 @@ uint32 FColorVertexBuffer::GetAllocatedSize() const
 	}
 }
 
-void FColorVertexBuffer::InitRHI()
+template <bool bRenderThread>
+FVertexBufferRHIRef FColorVertexBuffer::CreateRHIBuffer_Internal()
 {
-	if( VertexData != NULL )
+	if (VertexData != NULL)
 	{
 		FResourceArrayInterface* ResourceArray = VertexData->GetResourceArray();
-		if(ResourceArray->GetResourceDataSize())
+		if (ResourceArray->GetResourceDataSize())
 		{
 			// Create the vertex buffer.
 			FRHIResourceCreateInfo CreateInfo(ResourceArray);
-			VertexBufferRHI = RHICreateVertexBuffer(ResourceArray->GetResourceDataSize(), BUF_Static | BUF_ShaderResource, CreateInfo);
-			ColorComponentsSRV = RHICreateShaderResourceView(VertexBufferRHI, 4,  PF_R8G8B8A8);
+			if (bRenderThread)
+			{
+				return RHICreateVertexBuffer(ResourceArray->GetResourceDataSize(), BUF_Static | BUF_ShaderResource, CreateInfo);
+			}
+			else
+			{
+				return RHIAsyncCreateVertexBuffer(ResourceArray->GetResourceDataSize(), BUF_Static | BUF_ShaderResource, CreateInfo);
+			}
 		}
 	}
+	return nullptr;
+}
+
+FVertexBufferRHIRef FColorVertexBuffer::CreateRHIBuffer_RenderThread()
+{
+	return CreateRHIBuffer_Internal<true>();
+}
+
+FVertexBufferRHIRef FColorVertexBuffer::CreateRHIBuffer_Async()
+{
+	return CreateRHIBuffer_Internal<false>();
+}
+
+void FColorVertexBuffer::InitRHI()
+{
+	if (!bStreamed)
+	{
+		VertexBufferRHI = CreateRHIBuffer_RenderThread();
+	}
+	if (bStreamed || VertexBufferRHI)
+	{
+		ColorComponentsSRV = RHICreateShaderResourceView(VertexBufferRHI, 4, PF_R8G8B8A8);
+	}
 }
 
 void FColorVertexBuffer::ReleaseRHI()
 {
 	ColorComponentsSRV.SafeRelease();
-
 	FVertexBuffer::ReleaseRHI();
 }
 
diff --git a/Engine/Source/Runtime/Engine/Private/Rendering/StaticMeshVertexBuffer.cpp b/Engine/Source/Runtime/Engine/Private/Rendering/StaticMeshVertexBuffer.cpp
--- a/Engine/Source/Runtime/Engine/Private/Rendering/StaticMeshVertexBuffer.cpp
+++ b/Engine/Source/Runtime/Engine/Private/Rendering/StaticMeshVertexBuffer.cpp
@@ -13,7 +13,8 @@ FStaticMeshVertexBuffer::FStaticMeshVertexBuffer() :
 	NumTexCoords(0),
 	NumVertices(0),
 	bUseFullPrecisionUVs(!GVertexElementTypeSupport.IsSupported(VET_Half2)),
-	bUseHighPrecisionTangentBasis(false)
+	bUseHighPrecisionTangentBasis(false),
+	bStreamed(false)
 {}
 
 FStaticMeshVertexBuffer::~FStaticMeshVertexBuffer()
@@ -200,9 +201,7 @@ void FStaticMeshVertexBuffer::Serialize(FArchive& Ar, bool bNeedsCPUAccess)
 
 	FStripDataFlags StripFlags(Ar, 0, VER_UE4_STATIC_SKELETAL_MESH_SERIALIZATION_FIX);
 
-	Ar << NumTexCoords << NumVertices;
-	Ar << bUseFullPrecisionUVs;
-	Ar << bUseHighPrecisionTangentBasis;
+	SerializeMetaData(Ar);
 
 	if (Ar.IsLoading())
 	{
@@ -238,6 +237,15 @@ void FStaticMeshVertexBuffer::Serialize(FArchive& Ar, bool bNeedsCPUAccess)
 	}
 }
 
+void FStaticMeshVertexBuffer::SerializeMetaData(FArchive& Ar)
+{
+	Ar << NumTexCoords << NumVertices;
+	Ar << bUseFullPrecisionUVs;
+	Ar << bUseHighPrecisionTangentBasis;
+
+	InitTangentAndTexCoordStrides();
+}
+
 
 /**
 * Specialized assignment operator, only used when importing LOD's.
@@ -250,35 +258,88 @@ void FStaticMeshVertexBuffer::operator=(const FStaticMeshVertexBuffer &Other)
 	bUseHighPrecisionTangentBasis = Other.bUseHighPrecisionTangentBasis;
 }
 
-void FStaticMeshVertexBuffer::InitRHI()
+template <bool bRenderThread>
+FVertexBufferRHIRef FStaticMeshVertexBuffer::CreateTangentsRHIBuffer_Internal()
 {
+	check(TangentsData);
+	FResourceArrayInterface* ResourceArray = TangentsData->GetResourceArray();
+	if (ResourceArray->GetResourceDataSize())
 	{
-		check(TangentsData);
-		FResourceArrayInterface* ResourceArray = TangentsData->GetResourceArray();
-		if (ResourceArray->GetResourceDataSize())
+		// Create the vertex buffer.
+		FRHIResourceCreateInfo CreateInfo(ResourceArray);
+		if (bRenderThread)
 		{
-			// Create the vertex buffer.
-			FRHIResourceCreateInfo CreateInfo(ResourceArray);
-			TangentsVertexBuffer.VertexBufferRHI = RHICreateVertexBuffer(ResourceArray->GetResourceDataSize(), BUF_Static | BUF_ShaderResource, CreateInfo);
-			if (RHISupportsManualVertexFetch(GMaxRHIShaderPlatform) || IsGPUSkinCacheAvailable())
-			{
-				TangentsSRV = RHICreateShaderResourceView(TangentsVertexBuffer.VertexBufferRHI, GetUseHighPrecisionTangentBasis() ? 8 : 4, GetUseHighPrecisionTangentBasis() ? PF_R16G16B16A16_SNORM : PF_R8G8B8A8_SNORM);
-			}
+			return RHICreateVertexBuffer(ResourceArray->GetResourceDataSize(), BUF_Static | BUF_ShaderResource, CreateInfo);
+		}
+		else
+		{
+			return RHIAsyncCreateVertexBuffer(ResourceArray->GetResourceDataSize(), BUF_Static | BUF_ShaderResource, CreateInfo);
 		}
 	}
+	return nullptr;
+}
+
+FVertexBufferRHIRef FStaticMeshVertexBuffer::CreateTangentsRHIBuffer_RenderThread()
+{
+	return CreateTangentsRHIBuffer_Internal<true>();
+}
+
+FVertexBufferRHIRef FStaticMeshVertexBuffer::CreateTangentsRHIBuffer_Async()
+{
+	return CreateTangentsRHIBuffer_Internal<false>();
+}
+
+template <bool bRenderThread>
+FVertexBufferRHIRef FStaticMeshVertexBuffer::CreateTexCoordRHIBuffer_Internal()
+{
+	check(TexcoordData);
+	FResourceArrayInterface* ResourceArray = TexcoordData->GetResourceArray();
+	if (ResourceArray->GetResourceDataSize())
 	{
-		check(TexcoordData);
-		FResourceArrayInterface* ResourceArray = TexcoordData->GetResourceArray();
-		if (ResourceArray->GetResourceDataSize())
+		// Create the vertex buffer.
+		FRHIResourceCreateInfo CreateInfo(ResourceArray);
+		if (bRenderThread)
 		{
-			// Create the vertex buffer.
-			FRHIResourceCreateInfo CreateInfo(ResourceArray);
-			TexCoordVertexBuffer.VertexBufferRHI = RHICreateVertexBuffer(ResourceArray->GetResourceDataSize(), BUF_Static | BUF_ShaderResource, CreateInfo);
-			if (RHISupportsManualVertexFetch(GMaxRHIShaderPlatform))
-			{
-				TextureCoordinatesSRV = RHICreateShaderResourceView(TexCoordVertexBuffer.VertexBufferRHI, GetUseFullPrecisionUVs() ? 8 : 4, GetUseFullPrecisionUVs() ? PF_G32R32F : PF_G16R16F);
-			}
+			return RHICreateVertexBuffer(ResourceArray->GetResourceDataSize(), BUF_Static | BUF_ShaderResource, CreateInfo);
 		}
+		else
+		{
+			return RHIAsyncCreateVertexBuffer(ResourceArray->GetResourceDataSize(), BUF_Static | BUF_ShaderResource, CreateInfo);
+		}
+	}
+	return nullptr;
+}
+
+FVertexBufferRHIRef FStaticMeshVertexBuffer::CreateTexCoordRHIBuffer_RenderThread()
+{
+	return CreateTexCoordRHIBuffer_Internal<true>();
+}
+
+FVertexBufferRHIRef FStaticMeshVertexBuffer::CreateTexCoordRHIBuffer_Async()
+{
+	return CreateTexCoordRHIBuffer_Internal<false>();
+}
+
+void FStaticMeshVertexBuffer::InitRHI()
+{
+	if (!bStreamed)
+	{
+		TangentsVertexBuffer.VertexBufferRHI = CreateTangentsRHIBuffer_RenderThread();
+		TexCoordVertexBuffer.VertexBufferRHI = CreateTexCoordRHIBuffer_RenderThread();
+	}
+	if ((bStreamed || TangentsVertexBuffer.VertexBufferRHI) && (RHISupportsManualVertexFetch(GMaxRHIShaderPlatform) || IsGPUSkinCacheAvailable()))
+	{
+		TangentsSRV = RHICreateShaderResourceView(
+			TangentsVertexBuffer.VertexBufferRHI,
+			GetUseHighPrecisionTangentBasis() ? 8 : 4,
+			GetUseHighPrecisionTangentBasis() ? PF_R16G16B16A16_SNORM : PF_R8G8B8A8_SNORM);
+	}
+	if ((bStreamed || TexCoordVertexBuffer.VertexBufferRHI) && RHISupportsManualVertexFetch(GMaxRHIShaderPlatform))
+	{
+		TextureCoordinatesSRV = RHICreateShaderResourceView(
+			TexCoordVertexBuffer.VertexBufferRHI,
+			GetUseFullPrecisionUVs() ? 8 : 4,
+			GetUseFullPrecisionUVs() ? PF_G32R32F : PF_G16R16F);
 	}
 }
 
@@ -286,7 +347,6 @@ void FStaticMeshVertexBuffer::ReleaseRHI()
 {
 	TangentsSRV.SafeRelease();
 	TextureCoordinatesSRV.SafeRelease();
-
 	TangentsVertexBuffer.ReleaseRHI();
 	TexCoordVertexBuffer.ReleaseRHI();
 }
@@ -370,6 +430,17 @@ int FStaticMeshVertexBuffer::GetTexCoordSize()
 	}
 }
 
+void FStaticMeshVertexBuffer::InitTangentAndTexCoordStrides()
+{
+	typedef TStaticMeshVertexTangentDatum<typename TStaticMeshVertexTangentTypeSelector<EStaticMeshVertexTangentBasisType::HighPrecision>::TangentTypeT> HighPrecTangentType;
+	typedef TStaticMeshVertexTangentDatum<typename TStaticMeshVertexTangentTypeSelector<EStaticMeshVertexTangentBasisType::Default>::TangentTypeT> DefaultTangentType;
+	typedef TStaticMeshVertexUVsDatum<typename TStaticMeshVertexUVsTypeSelector<EStaticMeshVertexUVType::HighPrecision>::UVsTypeT> HighPrecUVType;
+	typedef TStaticMeshVertexUVsDatum<typename TStaticMeshVertexUVsTypeSelector<EStaticMeshVertexUVType::Default>::UVsTypeT> DefaultUVType;
+
+	TangentsStride = GetUseHighPrecisionTangentBasis() ? sizeof(HighPrecTangentType) : sizeof(DefaultTangentType);
+	TexcoordStride = GetUseFullPrecisionUVs() ? sizeof(HighPrecUVType) : sizeof(DefaultUVType);
+}
+
 void FStaticMeshVertexBuffer::BindTangentVertexBuffer(const FVertexFactory* VertexFactory, FStaticMeshDataType& Data) const
 {
 	{
diff --git a/Engine/Source/Runtime/Engine/Private/SceneManagement.cpp b/Engine/Source/Runtime/Engine/Private/SceneManagement.cpp
--- a/Engine/Source/Runtime/Engine/Private/SceneManagement.cpp
+++ b/Engine/Source/Runtime/Engine/Private/SceneManagement.cpp
@@ -541,7 +541,7 @@ int8 ComputeStaticMeshLOD( const FStaticMeshRenderData* RenderData, const FVecto
 	return MinLOD;
 }
 
-FLODMask ComputeLODForMeshes(const TArray<class FStaticMeshBatchRelevance>& StaticMeshRelevances, const FSceneView& View, const FVector4& Origin, float SphereRadius, int32 ForcedLODLevel, float& OutScreenRadiusSquared, float ScreenSizeScale, bool bDitheredLODTransition)
+FLODMask ComputeLODForMeshes(const TArray<class FStaticMeshBatchRelevance>& StaticMeshRelevances, const FSceneView& View, const FVector4& Origin, float SphereRadius, int32 ForcedLODLevel, float& OutScreenRadiusSquared, int8 CurFirstLODIdx, float ScreenSizeScale, bool bDitheredLODTransition)
 {
 	FLODMask LODToRender;
 	const FSceneView& LODView = GetLODView(View);
@@ -563,6 +563,7 @@ FLODMask ComputeLODForMeshes(const TArray<class FStaticMeshBatchRelevance>& Stat
 				MaxLOD = FMath::Max(MaxLOD, (int32)Mesh.LODIndex);
 			}
 		}
+		MinLOD = FMath::Max(MinLOD, (int32)CurFirstLODIdx);
 		LODToRender.SetLOD(FMath::Clamp(ForcedLODLevel, MinLOD, MaxLOD));
 	}
 	else if (LODView.Family->EngineShowFlags.LOD && NumMeshes)
@@ -626,6 +627,7 @@ FLODMask ComputeLODForMeshes(const TArray<class FStaticMeshBatchRelevance>& Stat
 				LODToRender.SetLOD(MinLODFound);
 			}
 		}
+		LODToRender.ClampToFirstLOD(CurFirstLODIdx);
 	}
 	return LODToRender;
 }
diff --git a/Engine/Source/Runtime/Engine/Private/StaticMesh.cpp b/Engine/Source/Runtime/Engine/Private/StaticMesh.cpp
--- a/Engine/Source/Runtime/Engine/Private/StaticMesh.cpp
+++ b/Engine/Source/Runtime/Engine/Private/StaticMesh.cpp
@@ -72,6 +72,10 @@
 #include "Streaming/UVChannelDensity.h"
 #include "Logging/MessageLog.h"
 #include "Misc/UObjectToken.h"
+#include "UObject/CoreRedirects.h"
+#include "HAL/FileManager.h"
+#include "ContentStreaming.h"
+#include "Streaming/StaticMeshUpdate.h"
 
 #define LOCTEXT_NAMESPACE "StaticMesh"
 DEFINE_LOG_CATEGORY(LogStaticMesh);	
@@ -125,6 +129,8 @@ static TAutoConsoleVariable<int32> CVarStripDistanceFieldDataDuringLoad(
 	TEXT("If non-zero, data for distance fields will be discarded on load. TODO: change to discard during cook!."),
 	ECVF_ReadOnly | ECVF_RenderThreadSafe);
 
+extern bool TrackRenderAssetEvent(struct FStreamingRenderAsset* StreamingRenderAsset, UStreamableRenderAsset* RenderAsset, bool bForceMipLevelsToBeResident, const FRenderAssetStreamingManager* Manager);
+
 #if ENABLE_COOK_STATS
 namespace StaticMeshCookStats
 {
@@ -183,15 +189,131 @@ FArchive& operator<<(FArchive& Ar, FStaticMeshSection& Section)
 	return Ar;
 }
 
-void FStaticMeshLODResources::Serialize(FArchive& Ar, UObject* Owner, int32 Index)
+uint8 FStaticMeshLODResources::GenerateClassStripFlags(FArchive& Ar, UStaticMesh* OwnerStaticMesh, int32 Index)
 {
-	DECLARE_SCOPE_CYCLE_COUNTER( TEXT("FStaticMeshLODResources::Serialize"), STAT_StaticMeshLODResources_Serialize, STATGROUP_LoadTime );
+#if WITH_EDITOR
+	// Defined class flags for possible stripping
+	const uint8 AdjacencyDataStripFlag = CDSF_AdjacencyData;
+	const uint8 MinLodDataStripFlag = CDSF_MinLodData;
+	const uint8 ReversedIndexBufferStripFlag = CDSF_ReversedIndexBuffer;
+
+	const bool bWantToStripTessellation = Ar.IsCooking()
+		&& ((GForceStripMeshAdjacencyDataDuringCooking != 0) || !Ar.CookingTarget()->SupportsFeature(ETargetPlatformFeatures::Tessellation));
+	const bool bWantToStripLOD = Ar.IsCooking()
+		&& (CVarStripMinLodDataDuringCooking.GetValueOnAnyThread() != 0)
+		&& OwnerStaticMesh
+		&& OwnerStaticMesh->MinLOD.GetValueForPlatformIdentifiers(
+			Ar.CookingTarget()->GetPlatformInfo().PlatformGroupName,
+			Ar.CookingTarget()->GetPlatformInfo().VanillaPlatformName) > Index;
+
+	return (bWantToStripTessellation ? AdjacencyDataStripFlag : 0) |
+		(bWantToStripLOD ? MinLodDataStripFlag : 0);
+#else
+	return 0;
+#endif
+}
+
+bool FStaticMeshLODResources::IsLODCookedOut(const ITargetPlatform* TargetPlatform, UStaticMesh* StaticMesh, bool bIsBelowMinLOD)
+{
+	check(StaticMesh);
+#if WITH_EDITOR
+	if (!bIsBelowMinLOD)
+	{
+		return false;
+	}
+
+	if (!TargetPlatform)
+	{
+		TargetPlatform = GetTargetPlatformManagerRef().GetRunningTargetPlatform();
+	}
+	check(TargetPlatform);
+
+	// If LOD streaming is supported, LODs below MinLOD are stored to optional paks and thus never cooked out
+	const FStaticMeshLODGroup& LODGroupSettings = TargetPlatform->GetStaticMeshLODSettings().GetLODGroup(StaticMesh->LODGroup);
+	return !TargetPlatform->SupportsFeature(ETargetPlatformFeatures::MeshLODStreaming) || !LODGroupSettings.IsLODStreamingSupported();
+#else
+	return false;
+#endif
+}
+
+bool FStaticMeshLODResources::IsLODInlined(const ITargetPlatform* TargetPlatform, UStaticMesh* StaticMesh, int32 LODIdx, bool bIsBelowMinLOD)
+{
+	check(StaticMesh);
+#if WITH_EDITOR
+	if (!TargetPlatform)
+	{
+		TargetPlatform = GetTargetPlatformManagerRef().GetRunningTargetPlatform();
+	}
+	check(TargetPlatform);
+
+	const FStaticMeshLODGroup& LODGroupSettings = TargetPlatform->GetStaticMeshLODSettings().GetLODGroup(StaticMesh->LODGroup);
+	if (!TargetPlatform->SupportsFeature(ETargetPlatformFeatures::MeshLODStreaming) || !LODGroupSettings.IsLODStreamingSupported())
+	{
+		return true;
+	}
+
+	if (bIsBelowMinLOD)
+	{
+		return false;
+	}
+
+	int32 MaxNumStreamedLODs = 0;
+	const int32 NumStreamedLODsOverride = StaticMesh->NumStreamedLODs.GetValueForPlatformIdentifiers(
+		TargetPlatform->GetPlatformInfo().PlatformGroupName,
+		TargetPlatform->GetPlatformInfo().VanillaPlatformName);
+	if (NumStreamedLODsOverride >= 0)
+	{
+		MaxNumStreamedLODs = NumStreamedLODsOverride;
+	}
+	else
+	{
+		MaxNumStreamedLODs = LODGroupSettings.GetDefaultMaxNumStreamedLODs();
+	}
+	
+	const int32 NumLODs = StaticMesh->GetNumLODs();
+	const int32 NumStreamedLODs = FMath::Min(MaxNumStreamedLODs, NumLODs - 1);
+	const int32 InlinedLODStartIdx = NumStreamedLODs;
+	return LODIdx >= InlinedLODStartIdx;
+#else
+	return false;
+#endif
+}
+
+void FStaticMeshLODResources::AccumVertexBuffersSize(const FStaticMeshVertexBuffers& VertexBuffers, uint32& OutSize)
+{
+#if (WITH_EDITOR || DO_CHECK)
+	const FPositionVertexBuffer& Pos = VertexBuffers.PositionVertexBuffer;
+	const FStaticMeshVertexBuffer& TanTex = VertexBuffers.StaticMeshVertexBuffer;
+	const FColorVertexBuffer& Color = VertexBuffers.ColorVertexBuffer;
+	OutSize += Pos.GetNumVertices() * Pos.GetStride();
+	OutSize += TanTex.GetResourceSize();
+	OutSize += Color.GetNumVertices() * Color.GetStride();
+#endif
+}
 
+void FStaticMeshLODResources::AccumIndexBufferSize(const FRawStaticIndexBuffer& IndexBuffer, uint32& OutSize)
+{
+#if (WITH_EDITOR || DO_CHECK)
+	OutSize += IndexBuffer.GetIndexDataSize();
+#endif
+}
+
+uint32 FStaticMeshLODResources::FStaticMeshBuffersSize::CalcBuffersSize() const
+{
+	// Assumes these two cvars don't change at runtime
+	const bool bEnableDepthOnlyIndexBuffer = !!CVarSupportDepthOnlyIndexBuffers.GetValueOnAnyThread();
+	const bool bEnableReversedIndexBuffer = !!CVarSupportReversedIndexBuffers.GetValueOnAnyThread();
+	return SerializedBuffersSize
+		- (bEnableDepthOnlyIndexBuffer ? 0 : DepthOnlyIBSize)
+		- (bEnableReversedIndexBuffer ? 0 : ReversedIBsSize);
+}
+
+void FStaticMeshLODResources::SerializeBuffers(FArchive& Ar, UStaticMesh* OwnerStaticMesh, uint8 InStripFlags, FStaticMeshBuffersSize& OutBuffersSize)
+{
 	bool bEnableDepthOnlyIndexBuffer = (CVarSupportDepthOnlyIndexBuffers.GetValueOnAnyThread() == 1);
 	bool bEnableReversedIndexBuffer = (CVarSupportReversedIndexBuffers.GetValueOnAnyThread() == 1);
 
 	// See if the mesh wants to keep resources CPU accessible
-	UStaticMesh* OwnerStaticMesh = Cast<UStaticMesh>(Owner);
 	bool bMeshCPUAcces = OwnerStaticMesh ? OwnerStaticMesh->bAllowCPUAccess : false;
 
 	// Note: this is all derived data, native versioning is not needed, but be sure to bump STATICMESH_DERIVEDDATA_VER when modifying!
@@ -200,86 +322,226 @@ void FStaticMeshLODResources::Serialize(FArchive& Ar, UObject* Owner, int32 Inde
 	// TODO: Not needed in uncooked games either after PostLoad!
 	bool bNeedsCPUAccess = !FPlatformProperties::RequiresCookedData() || bMeshCPUAcces;
 
+	bHasWireframeIndices = false;
 	bHasAdjacencyInfo = false;
 	bHasDepthOnlyIndices = false;
 	bHasReversedIndices = false;
 	bHasReversedDepthOnlyIndices = false;
+	bHasColorVertexData = false;
 	DepthOnlyNumTriangles = 0;
 
-	// Defined class flags for possible stripping
-	const uint8 AdjacencyDataStripFlag = 1;
-	const uint8 MinLodDataStripFlag = 2;
-	const uint8 ReversedIndexBufferStripFlag = 4;
+	FStripDataFlags StripFlags(Ar, InStripFlags);
 
-	// Actual flags used during serialization
-	uint8 ClassDataStripFlags = 0;
+	VertexBuffers.PositionVertexBuffer.Serialize(Ar, bNeedsCPUAccess);
+	VertexBuffers.StaticMeshVertexBuffer.Serialize(Ar, bNeedsCPUAccess);
+	VertexBuffers.ColorVertexBuffer.Serialize(Ar, bNeedsCPUAccess);
+	OutBuffersSize.Clear();
+	AccumVertexBuffersSize(VertexBuffers, OutBuffersSize.SerializedBuffersSize);
 
-#if WITH_EDITOR
-	const bool bWantToStripTessellation = Ar.IsCooking() && ((GForceStripMeshAdjacencyDataDuringCooking != 0) || !Ar.CookingTarget()->SupportsFeature(ETargetPlatformFeatures::Tessellation));
-	const bool bWantToStripLOD = Ar.IsCooking() && (CVarStripMinLodDataDuringCooking.GetValueOnAnyThread() != 0) && OwnerStaticMesh && OwnerStaticMesh->MinLOD.GetValueForPlatformIdentifiers(Ar.CookingTarget()->GetPlatformInfo().PlatformGroupName, Ar.CookingTarget()->GetPlatformInfo().VanillaPlatformName) > Index;
+	IndexBuffer.Serialize(Ar, bNeedsCPUAccess);
+	AccumIndexBufferSize(IndexBuffer, OutBuffersSize.SerializedBuffersSize);
+
+	const bool bSerailizeReversedIndexBuffer = !StripFlags.IsClassDataStripped(CDSF_ReversedIndexBuffer);
+	if (bSerailizeReversedIndexBuffer)
+	{
+		ReversedIndexBuffer.Serialize(Ar, bNeedsCPUAccess);
+		AccumIndexBufferSize(ReversedIndexBuffer, OutBuffersSize.ReversedIBsSize);
+		AccumIndexBufferSize(ReversedIndexBuffer, OutBuffersSize.SerializedBuffersSize);
+		if (!bEnableReversedIndexBuffer)
+		{
+			ReversedIndexBuffer.Discard();
+		}
+	}
+
+	DepthOnlyIndexBuffer.Serialize(Ar, bNeedsCPUAccess);
+	AccumIndexBufferSize(DepthOnlyIndexBuffer, OutBuffersSize.DepthOnlyIBSize);
+	AccumIndexBufferSize(DepthOnlyIndexBuffer, OutBuffersSize.SerializedBuffersSize);
+	if (!bEnableDepthOnlyIndexBuffer)
+	{
+		DepthOnlyIndexBuffer.Discard();
+	}
+
+	if (bSerailizeReversedIndexBuffer)
+	{
+		ReversedDepthOnlyIndexBuffer.Serialize(Ar, bNeedsCPUAccess);
+		AccumIndexBufferSize(ReversedDepthOnlyIndexBuffer, OutBuffersSize.ReversedIBsSize);
+		AccumIndexBufferSize(ReversedDepthOnlyIndexBuffer, OutBuffersSize.SerializedBuffersSize);
+		if (!bEnableReversedIndexBuffer)
+		{
+			ReversedDepthOnlyIndexBuffer.Discard();
+		}
+	}
+
+	if (!StripFlags.IsEditorDataStripped())
+	{
+		WireframeIndexBuffer.Serialize(Ar, bNeedsCPUAccess);
+		AccumIndexBufferSize(WireframeIndexBuffer, OutBuffersSize.SerializedBuffersSize);
+		bHasWireframeIndices = WireframeIndexBuffer.GetNumIndices() != 0;
+	}
+
+	if (!StripFlags.IsClassDataStripped(CDSF_AdjacencyData))
+	{
+		AdjacencyIndexBuffer.Serialize(Ar, bNeedsCPUAccess);
+		AccumIndexBufferSize(AdjacencyIndexBuffer, OutBuffersSize.SerializedBuffersSize);
+		bHasAdjacencyInfo = AdjacencyIndexBuffer.GetNumIndices() != 0;
+	}
+
+	// Needs to be done now because on cooked platform, indices are discarded after RHIInit.
+	bHasDepthOnlyIndices = DepthOnlyIndexBuffer.GetNumIndices() != 0;
+	bHasReversedIndices = bSerailizeReversedIndexBuffer && ReversedIndexBuffer.GetNumIndices() != 0;
+	bHasReversedDepthOnlyIndices = bSerailizeReversedIndexBuffer && ReversedDepthOnlyIndexBuffer.GetNumIndices() != 0;
+	bHasColorVertexData = VertexBuffers.ColorVertexBuffer.GetNumVertices() > 0;
+	DepthOnlyNumTriangles = DepthOnlyIndexBuffer.GetNumIndices() / 3;
+
+	AreaWeightedSectionSamplers.SetNum(Sections.Num());
+	for (FStaticMeshSectionAreaWeightedTriangleSampler& Sampler : AreaWeightedSectionSamplers)
+	{
+		Sampler.Serialize(Ar);
+	}
+	AreaWeightedSampler.Serialize(Ar);
+}
 
-	ClassDataStripFlags |=	(bWantToStripTessellation ? AdjacencyDataStripFlag : 0)	|
-							(bWantToStripLOD ? MinLodDataStripFlag : 0);
+void FStaticMeshLODResources::SerializeAvailabilityInfo(FArchive& Ar)
+{
+	Ar << DepthOnlyNumTriangles;
+	uint32 Packed;
+#if WITH_EDITOR
+	if (Ar.IsSaving())
+	{
+		Packed = bHasAdjacencyInfo
+			| (bHasDepthOnlyIndices << 1u)
+			| (bHasReversedIndices << 2u)
+			| (bHasReversedDepthOnlyIndices << 3u)
+			| (bHasColorVertexData << 4u)
+			| (bHasWireframeIndices << 5u);
+		Ar << Packed;
+	}
+	else
 #endif
+	{
+		const bool bEnableDepthOnlyIndexBuffer = !!CVarSupportDepthOnlyIndexBuffers.GetValueOnAnyThread();
+		const bool bEnableReversedIndexBuffer = !!CVarSupportReversedIndexBuffers.GetValueOnAnyThread();
+		Ar << Packed;
+		DepthOnlyNumTriangles *= static_cast<uint32>(bEnableDepthOnlyIndexBuffer);
+		bHasAdjacencyInfo = Packed & 1u;
+		bHasDepthOnlyIndices = bEnableDepthOnlyIndexBuffer && !!(Packed & 2u);
+		bHasReversedIndices = bEnableReversedIndexBuffer && !!(Packed & 4u);
+		bHasReversedDepthOnlyIndices = bEnableReversedIndexBuffer && !!(Packed & 8u);
+		bHasColorVertexData = (Packed >> 4u) & 1u;
+		bHasWireframeIndices = (Packed >> 5u) & 1u;
+	}
 
-	FStripDataFlags StripFlags( Ar, ClassDataStripFlags );
+	VertexBuffers.StaticMeshVertexBuffer.SerializeMetaData(Ar);
+	VertexBuffers.PositionVertexBuffer.SerializeMetaData(Ar);
+	VertexBuffers.ColorVertexBuffer.SerializeMetaData(Ar);
+	IndexBuffer.SerializeMetaData(Ar);
+	ReversedIndexBuffer.SerializeMetaData(Ar);
+	if (!bHasReversedIndices)
+	{
+		// Reversed indices are either stripped during cook or will be stripped on load.
+		// In either case, clear CachedNumIndices to show that the buffer will be empty after actual loading
+		ReversedIndexBuffer.Discard();
+	}
+	DepthOnlyIndexBuffer.SerializeMetaData(Ar);
+	if (!bHasDepthOnlyIndices)
+	{
+		DepthOnlyIndexBuffer.Discard();
+	}
+	ReversedDepthOnlyIndexBuffer.SerializeMetaData(Ar);
+	if (!bHasReversedDepthOnlyIndices)
+	{
+		ReversedDepthOnlyIndexBuffer.Discard();
+	}
+	WireframeIndexBuffer.SerializeMetaData(Ar);
+	if (!bHasWireframeIndices)
+	{
+		WireframeIndexBuffer.Discard();
+	}
+	AdjacencyIndexBuffer.SerializeMetaData(Ar);
+	if (!bHasAdjacencyInfo)
+	{
+		AdjacencyIndexBuffer.Discard();
+	}
+}
+
+void FStaticMeshLODResources::Serialize(FArchive& Ar, UObject* Owner, int32 Index)
+{
+	DECLARE_SCOPE_CYCLE_COUNTER(TEXT("FStaticMeshLODResources::Serialize"), STAT_StaticMeshLODResources_Serialize, STATGROUP_LoadTime);
+
+	UStaticMesh* OwnerStaticMesh = Cast<UStaticMesh>(Owner);
+	// Actual flags used during serialization
+	const uint8 ClassDataStripFlags = GenerateClassStripFlags(Ar, OwnerStaticMesh, Index);
+	FStripDataFlags StripFlags(Ar, ClassDataStripFlags);
 
 	Ar << Sections;
 	Ar << MaxDeviation;
 
-	if( !StripFlags.IsDataStrippedForServer() && !StripFlags.IsClassDataStripped(MinLodDataStripFlag) )
-	{
-		VertexBuffers.PositionVertexBuffer.Serialize( Ar, bNeedsCPUAccess );
-		VertexBuffers.StaticMeshVertexBuffer.Serialize( Ar, bNeedsCPUAccess );
-		VertexBuffers.ColorVertexBuffer.Serialize( Ar, bNeedsCPUAccess );
-		IndexBuffer.Serialize(Ar, bNeedsCPUAccess);
+	const bool bIsBelowMinLOD = StripFlags.IsClassDataStripped(CDSF_MinLodData);
+	bool bIsLODCookedOut = IsLODCookedOut(Ar.CookingTarget(), OwnerStaticMesh, bIsBelowMinLOD);
+	Ar << bIsLODCookedOut;
 
-		const bool bSerailizeReversedIndexBuffer = !StripFlags.IsClassDataStripped(ReversedIndexBufferStripFlag);
-		if (bSerailizeReversedIndexBuffer)
-		{
-			ReversedIndexBuffer.Serialize(Ar, bNeedsCPUAccess);
-			if (!bEnableReversedIndexBuffer)
-			{
-				ReversedIndexBuffer.Discard();
-			}
-		}
-		DepthOnlyIndexBuffer.Serialize(Ar, bNeedsCPUAccess);
-		if (!bEnableDepthOnlyIndexBuffer)
+	bool bInlined = bIsLODCookedOut || IsLODInlined(Ar.CookingTarget(), OwnerStaticMesh, Index, bIsBelowMinLOD);
+	Ar << bInlined;
+	bBuffersInlined = bInlined;
+
+	VertexBuffers.StaticMeshVertexBuffer.SetIsStreamed(!bBuffersInlined);
+	VertexBuffers.PositionVertexBuffer.SetIsStreamed(!bBuffersInlined);
+	VertexBuffers.ColorVertexBuffer.SetIsStreamed(!bBuffersInlined);
+	IndexBuffer.SetIsStreamed(!bBuffersInlined);
+	ReversedIndexBuffer.SetIsStreamed(!bBuffersInlined);
+	DepthOnlyIndexBuffer.SetIsStreamed(!bBuffersInlined);
+	ReversedDepthOnlyIndexBuffer.SetIsStreamed(!bBuffersInlined);
+	WireframeIndexBuffer.SetIsStreamed(!bBuffersInlined);
+	AdjacencyIndexBuffer.SetIsStreamed(!bBuffersInlined);
+
+	if (!StripFlags.IsDataStrippedForServer() && !bIsLODCookedOut)
+	{
+		FStaticMeshBuffersSize TmpBuffersSize;
+		if (bInlined)
 		{
-			DepthOnlyIndexBuffer.Discard();
+			SerializeBuffers(Ar, OwnerStaticMesh, ClassDataStripFlags, TmpBuffersSize);
+			Ar << TmpBuffersSize;
+			BuffersSize = TmpBuffersSize.CalcBuffersSize();
 		}
-		if (bSerailizeReversedIndexBuffer)
+		else if (FPlatformProperties::RequiresCookedData() || Ar.IsCooking())
 		{
-			ReversedDepthOnlyIndexBuffer.Serialize(Ar, bNeedsCPUAccess);
-			if (!bEnableReversedIndexBuffer)
+#if WITH_EDITOR
+			if (Ar.IsSaving())
 			{
-				ReversedDepthOnlyIndexBuffer.Discard();
+				TArray<uint8> TmpBuff;
+				FMemoryWriter MemWriter(TmpBuff, true);
+				MemWriter.SetCookingTarget(Ar.CookingTarget());
+				MemWriter.SetByteSwapping(Ar.IsByteSwapping());
+				SerializeBuffers(MemWriter, OwnerStaticMesh, ClassDataStripFlags, TmpBuffersSize);
+
+				bIsOptionalLOD = bIsBelowMinLOD;
+				const uint32 BulkDataFlags = BULKDATA_Force_NOT_InlinePayload
+					| (bIsOptionalLOD ? BULKDATA_OptionalPayload : 0);
+				BulkData.SetBulkDataFlags(BulkDataFlags);
+				BulkData.Lock(LOCK_READ_WRITE);
+				void* BulkDataMem = BulkData.Realloc(TmpBuff.Num());
+				FMemory::Memcpy(BulkDataMem, TmpBuff.GetData(), TmpBuff.Num());
+				BulkData.Unlock();
+				BulkData.Serialize(Ar, Owner, Index);
+			}
+			else
+#endif
+			{
+				FByteBulkData TmpBulkData;
+				TmpBulkData.Serialize(Ar, Owner, Index);
+				bIsOptionalLOD = !!(TmpBulkData.GetBulkDataFlags() & BULKDATA_OptionalPayload);
+				int64 Tmp = TmpBulkData.GetBulkDataOffsetInFile();
+				check(Tmp >= 0 && Tmp <= 0xffffffff);
+				OffsetInFile = static_cast<uint32>(Tmp);
+				check(TmpBulkData.GetBulkDataSize() >= 0);
+				BulkDataSize = static_cast<uint32>(TmpBulkData.GetBulkDataSize());
 			}
-		}
-
-		if( !StripFlags.IsEditorDataStripped() )
-		{
-			WireframeIndexBuffer.Serialize(Ar, bNeedsCPUAccess);
-		}
-
-		if ( !StripFlags.IsClassDataStripped( AdjacencyDataStripFlag ) )
-		{
-			AdjacencyIndexBuffer.Serialize( Ar, bNeedsCPUAccess );
-			bHasAdjacencyInfo = AdjacencyIndexBuffer.GetNumIndices() != 0;
-		}
 
-		// Needs to be done now because on cooked platform, indices are discarded after RHIInit.
-		bHasDepthOnlyIndices = DepthOnlyIndexBuffer.GetNumIndices() != 0;
-		bHasReversedIndices = bSerailizeReversedIndexBuffer && ReversedIndexBuffer.GetNumIndices() != 0;
-		bHasReversedDepthOnlyIndices = bSerailizeReversedIndexBuffer && ReversedDepthOnlyIndexBuffer.GetNumIndices() != 0;
-		DepthOnlyNumTriangles = DepthOnlyIndexBuffer.GetNumIndices() / 3;
+			SerializeAvailabilityInfo(Ar);
 
-		AreaWeightedSectionSamplers.SetNum(Sections.Num());
-		for (FStaticMeshSectionAreaWeightedTriangleSampler& Sampler : AreaWeightedSectionSamplers)
-		{
-			Sampler.Serialize(Ar);
+			Ar << TmpBuffersSize;
+			BuffersSize = TmpBuffersSize.CalcBuffersSize();
 		}
-		AreaWeightedSampler.Serialize(Ar);
 	}
 }
 
@@ -614,7 +876,12 @@ FStaticMeshLODResources::FStaticMeshLODResources()
 	, bHasDepthOnlyIndices(false)
 	, bHasReversedIndices(false)
 	, bHasReversedDepthOnlyIndices(false)
+	, bHasColorVertexData(false)
+	, bHasWireframeIndices(false)
+	, bBuffersInlined(false)
+	, bIsOptionalLOD(false)
 	, DepthOnlyNumTriangles(0)
+	, BuffersSize(0)
 #if STATS
 	, StaticMeshIndexMemory(0)
 #endif
@@ -626,10 +893,8 @@ FStaticMeshLODResources::~FStaticMeshLODResources()
 	delete DistanceFieldData;
 }
 
-void FStaticMeshLODResources::InitResources(UStaticMesh* Parent)
+void FStaticMeshLODResources::ConditionalForce16BitIndexBuffer(EShaderPlatform MaxShaderPlatform, UStaticMesh* Parent)
 {
-	const auto MaxShaderPlatform = GShaderPlatformForFeatureLevel[GMaxRHIFeatureLevel];
-
 	// Initialize the vertex and index buffers.
 	// All platforms supporting Metal also support 32-bit indices.
 	if (IsES2Platform(MaxShaderPlatform) && !IsMetalPlatform(MaxShaderPlatform))
@@ -640,49 +905,88 @@ void FStaticMeshLODResources::InitResources(UStaticMesh* Parent)
 			TArray<uint32> Indices;
 			IndexBuffer.GetCopy(Indices);
 			IndexBuffer.SetIndices(Indices, EIndexBufferStride::Force16Bit);
-			UE_LOG(LogStaticMesh, Warning, TEXT("[%s] Mesh has more that 65535 vertices, incompatible with mobile; forcing 16-bit (will probably cause rendering issues)." ), *Parent->GetName());
+			UE_LOG(LogStaticMesh, Warning, TEXT("[%s] Mesh has more that 65535 vertices, incompatible with mobile; forcing 16-bit (will probably cause rendering issues)."), *Parent->GetName());
 		}
 	}
+}
+
+template <bool bIncrement>
+void FStaticMeshLODResources::UpdateIndexMemoryStats()
+{
+#if STATS
+	if (bIncrement)
+	{
+		uint32 iMem = IndexBuffer.GetAllocatedSize();
+		uint32 wiMem = WireframeIndexBuffer.GetAllocatedSize();
+		uint32 riMem = ReversedIndexBuffer.GetAllocatedSize();
+		uint32 doiMem = DepthOnlyIndexBuffer.GetAllocatedSize();
+		uint32 rdoiMem = ReversedDepthOnlyIndexBuffer.GetAllocatedSize();
+		uint32 aiMem = AdjacencyIndexBuffer.GetAllocatedSize();
+		StaticMeshIndexMemory = iMem + wiMem + riMem + doiMem + rdoiMem + aiMem;
+		INC_DWORD_STAT_BY(STAT_StaticMeshIndexMemory, StaticMeshIndexMemory);
+	}
+	else
+	{
+		DEC_DWORD_STAT_BY(STAT_StaticMeshIndexMemory, StaticMeshIndexMemory);
+	}
+#endif
+}
 
+template <bool bIncrement>
+void FStaticMeshLODResources::UpdateVertexMemoryStats() const
+{
 #if STATS
-	uint32 iMem = IndexBuffer.GetAllocatedSize();
-	uint32 wiMem = WireframeIndexBuffer.GetAllocatedSize();
-	uint32 riMem = ReversedIndexBuffer.GetAllocatedSize();
-	uint32 doiMem = DepthOnlyIndexBuffer.GetAllocatedSize();
-	uint32 rdoiMem = ReversedDepthOnlyIndexBuffer.GetAllocatedSize();
-	uint32 aiMem = AdjacencyIndexBuffer.GetAllocatedSize();
-	StaticMeshIndexMemory = iMem + wiMem + riMem + doiMem + rdoiMem + aiMem;
-	INC_DWORD_STAT_BY(STAT_StaticMeshIndexMemory, StaticMeshIndexMemory);
+	const uint32 StaticMeshVertexMemory =
+		VertexBuffers.StaticMeshVertexBuffer.GetResourceSize() +
+		VertexBuffers.PositionVertexBuffer.GetStride() * VertexBuffers.PositionVertexBuffer.GetNumVertices();
+	const uint32 ResourceVertexColorMemory = VertexBuffers.ColorVertexBuffer.GetStride() * VertexBuffers.ColorVertexBuffer.GetNumVertices();
+
+	if (bIncrement)
+	{
+		INC_DWORD_STAT_BY(STAT_StaticMeshVertexMemory, StaticMeshVertexMemory);
+		INC_DWORD_STAT_BY(STAT_ResourceVertexColorMemory, ResourceVertexColorMemory);
+	}
+	else
+	{
+		DEC_DWORD_STAT_BY(STAT_StaticMeshVertexMemory, StaticMeshVertexMemory);
+		DEC_DWORD_STAT_BY(STAT_ResourceVertexColorMemory, ResourceVertexColorMemory);
+	}
 #endif
+}
+
+void FStaticMeshLODResources::InitResources(UStaticMesh* Parent)
+{
+	ConditionalForce16BitIndexBuffer(GMaxRHIShaderPlatform, Parent);
+	UpdateIndexMemoryStats<true>();
 
 	BeginInitResource(&IndexBuffer);
-	if( WireframeIndexBuffer.GetNumIndices() > 0 )
+	if(bHasWireframeIndices)
 	{
 		BeginInitResource(&WireframeIndexBuffer);
-	}	
+	}
 	BeginInitResource(&VertexBuffers.StaticMeshVertexBuffer);
 	BeginInitResource(&VertexBuffers.PositionVertexBuffer);
-	if( VertexBuffers.ColorVertexBuffer.GetNumVertices() > 0 )
+	if(bHasColorVertexData)
 	{
 		BeginInitResource(&VertexBuffers.ColorVertexBuffer);
 	}
 
-	if (ReversedIndexBuffer.GetNumIndices() > 0)
+	if (bHasReversedIndices)
 	{
 		BeginInitResource(&ReversedIndexBuffer);
 	}
 
-	if (DepthOnlyIndexBuffer.GetNumIndices() > 0)
+	if (bHasDepthOnlyIndices)
 	{
 		BeginInitResource(&DepthOnlyIndexBuffer);
 	}
 
-	if (ReversedDepthOnlyIndexBuffer.GetNumIndices() > 0)
+	if (bHasReversedDepthOnlyIndices)
 	{
 		BeginInitResource(&ReversedDepthOnlyIndexBuffer);
 	}
 
-	if (RHISupportsTessellation(MaxShaderPlatform))
+	if (bHasAdjacencyInfo && RHISupportsTessellation(GMaxRHIShaderPlatform))
 	{
 		BeginInitResource(&AdjacencyIndexBuffer);
 	}
@@ -729,30 +1033,19 @@ void FStaticMeshLODResources::InitResources(UStaticMesh* Parent)
 		INC_DWORD_STAT_BY( STAT_StaticMeshDistanceFieldMemory, DistanceFieldData->GetResourceSizeBytes() );
 	}
 
-	FStaticMeshLODResources* This = this;
+#if STATS
 	ENQUEUE_RENDER_COMMAND(UpdateMemoryStats)(
-		[This](FRHICommandList& RHICmdList)
-		{		
-			const uint32 StaticMeshVertexMemory =
-				This->VertexBuffers.StaticMeshVertexBuffer.GetResourceSize() +
-				This->VertexBuffers.PositionVertexBuffer.GetStride() * This->VertexBuffers.PositionVertexBuffer.GetNumVertices();
-			const uint32 ResourceVertexColorMemory = This->VertexBuffers.ColorVertexBuffer.GetStride() * This->VertexBuffers.ColorVertexBuffer.GetNumVertices();
-
-			INC_DWORD_STAT_BY( STAT_StaticMeshVertexMemory, StaticMeshVertexMemory );
-			INC_DWORD_STAT_BY( STAT_ResourceVertexColorMemory, ResourceVertexColorMemory );
-		});
+		[this](FRHICommandListImmediate&)
+	{
+		UpdateVertexMemoryStats<true>();
+	});
+#endif
 }
 
 void FStaticMeshLODResources::ReleaseResources()
 {
-	const uint32 StaticMeshVertexMemory = 
-		VertexBuffers.StaticMeshVertexBuffer.GetResourceSize() +
-		VertexBuffers.PositionVertexBuffer.GetStride() * VertexBuffers.PositionVertexBuffer.GetNumVertices();
-	const uint32 ResourceVertexColorMemory = VertexBuffers.ColorVertexBuffer.GetStride() * VertexBuffers.ColorVertexBuffer.GetNumVertices();
-
-	DEC_DWORD_STAT_BY( STAT_StaticMeshVertexMemory, StaticMeshVertexMemory );
-	DEC_DWORD_STAT_BY( STAT_ResourceVertexColorMemory, ResourceVertexColorMemory );
-	DEC_DWORD_STAT_BY( STAT_StaticMeshIndexMemory, StaticMeshIndexMemory );
+	UpdateVertexMemoryStats<false>();
+	UpdateIndexMemoryStats<false>();
 
 	// Release the vertex and index buffers.
 	
@@ -780,12 +1073,40 @@ void FStaticMeshLODResources::ReleaseResources()
 	}
 }
 
+void FStaticMeshLODResources::IncrementMemoryStats()
+{
+	UpdateIndexMemoryStats<true>();
+	UpdateVertexMemoryStats<true>();
+}
+
+void FStaticMeshLODResources::DecrementMemoryStats()
+{
+	UpdateVertexMemoryStats<false>();
+	UpdateIndexMemoryStats<false>();
+}
+
+void FStaticMeshLODResources::DiscardCPUData()
+{
+	VertexBuffers.StaticMeshVertexBuffer.CleanUp();
+	VertexBuffers.PositionVertexBuffer.CleanUp();
+	VertexBuffers.ColorVertexBuffer.CleanUp();
+	IndexBuffer.Discard();
+	ReversedIndexBuffer.Discard();
+	DepthOnlyIndexBuffer.Discard();
+	ReversedDepthOnlyIndexBuffer.Discard();
+	WireframeIndexBuffer.Discard();
+	AdjacencyIndexBuffer.Discard();
+}
+
 /*------------------------------------------------------------------------------
 	FStaticMeshRenderData
 ------------------------------------------------------------------------------*/
 
 FStaticMeshRenderData::FStaticMeshRenderData()
 	: bLODsShareStaticLighting(false)
+	, bReadyForStreaming(false)
+	, NumInlinedLODs(0)
+	, CurrentFirstLODIdx(0)
 {
 	for (int32 LODIndex = 0; LODIndex < MAX_STATIC_MESH_LODS; ++LODIndex)
 	{
@@ -815,6 +1136,25 @@ void FStaticMeshRenderData::Serialize(FArchive& Ar, UStaticMesh* Owner, bool bCo
 #endif // #if WITH_EDITORONLY_DATA
 
 	LODResources.Serialize(Ar, Owner);
+#if WITH_EDITOR
+	if (Ar.IsSaving())
+	{
+		int32 Tmp = 0;
+		for (int32 Idx = LODResources.Num() - 1; Idx >= 0; --Idx)
+		{
+			if (!LODResources[Idx].bBuffersInlined)
+			{
+				break;
+			}
+			++Tmp;
+		}
+		NumInlinedLODs = Tmp;
+	}
+#endif
+	Ar << NumInlinedLODs;
+	CurrentFirstLODIdx = LODResources.Num() - NumInlinedLODs;
+	Owner->SetCachedNumResidentLODs(NumInlinedLODs);
+
 	if (Ar.IsLoading())
 	{
 		LODVertexFactories.Empty(LODResources.Num());
@@ -922,6 +1262,13 @@ void FStaticMeshRenderData::InitResources(ERHIFeatureLevel::Type InFeatureLevel,
 			LODVertexFactories[LODIndex].InitResources(LODResources[LODIndex], LODIndex, Owner);
 		}
 	}
+
+	ENQUEUE_RENDER_COMMAND(CmdSetStaticMeshReadyForStreaming)(
+		[this, Owner](FRHICommandListImmediate&)
+	{
+		bReadyForStreaming = true;
+		Owner->SetCachedReadyForStreaming(true);
+	});
 	bIsInitialized = true;
 }
 
@@ -1155,6 +1502,7 @@ void FStaticMeshRenderData::SyncUVChannelData(const TArray<FStaticMaterial>& Obj
 
 void FStaticMeshLODSettings::Initialize(const FConfigFile& IniFile)
 {
+	check(!Groups.Num());
 	// Ensure there is a default LOD group.
 	Groups.FindOrAdd(NAME_None);
 
@@ -1171,6 +1519,17 @@ void FStaticMeshLODSettings::Initialize(const FConfigFile& IniFile)
 		};
 	}
 
+	Groups.KeySort(TLess<FName>());
+	GroupName2Index.Empty(Groups.Num());
+	{
+		int32 GroupIdx = 0;
+		TMap<FName, FStaticMeshLODGroup>::TConstIterator It(Groups);
+		for (; It; ++It, ++GroupIdx)
+		{
+			GroupName2Index.Add(It.Key(), GroupIdx);
+		}
+	}
+
 	// Do some per-group initialization.
 	for (TMap<FName,FStaticMeshLODGroup>::TIterator It(Groups); It; ++It)
 	{
@@ -1205,6 +1564,17 @@ void FStaticMeshLODSettings::ReadEntry(FStaticMeshLODGroup& Group, FString Entry
 		Group.DefaultNumLODs = FMath::Clamp<int32>(Group.DefaultNumLODs, 1, MAX_STATIC_MESH_LODS);
 	}
 
+	if (FParse::Value(*Entry, TEXT("MaxNumStreamedLODs="), Group.DefaultMaxNumStreamedLODs))
+	{
+		Group.DefaultMaxNumStreamedLODs = FMath::Max(Group.DefaultMaxNumStreamedLODs, 0);
+	}
+
+	{
+		int32 LocalSupportLODStreaming = 0;
+		FParse::Value(*Entry, TEXT("bSupportLODStreaming="), LocalSupportLODStreaming);
+		Group.bSupportLODStreaming = !!LocalSupportLODStreaming;
+	}
+
 	if (FParse::Value(*Entry, TEXT("LightMapResolution="), Group.DefaultLightMapResolution))
 	{
 		Group.DefaultLightMapResolution = FMath::Max<int32>(Group.DefaultLightMapResolution, 0);
@@ -1498,7 +1868,7 @@ namespace StaticMeshDerivedDataTimings
 		);
 }
 
-static FString BuildStaticMeshDerivedDataKey(UStaticMesh* Mesh, const FStaticMeshLODGroup& LODGroup)
+static FString BuildStaticMeshDerivedDataKeySuffix(UStaticMesh* Mesh, const FStaticMeshLODGroup& LODGroup)
 {
 	FString KeySuffix(TEXT(""));
 	TArray<uint8> TempBytes;
@@ -1561,6 +1931,31 @@ static FString BuildStaticMeshDerivedDataKey(UStaticMesh* Mesh, const FStaticMes
 		}
 	}
 
+	// Mesh LOD streaming settings that need to trigger recache when changed
+	const ITargetPlatform* RunningPlatform = GetTargetPlatformManagerRef().GetRunningTargetPlatform();
+	check(RunningPlatform);
+	const bool bAllowLODStreaming = RunningPlatform->SupportsFeature(ETargetPlatformFeatures::MeshLODStreaming) && LODGroup.IsLODStreamingSupported();
+	KeySuffix.AppendChar(bAllowLODStreaming ? TEXT('1') : TEXT('0'));
+	if (bAllowLODStreaming)
+	{
+		int32 MaxNumStreamedLODs = Mesh->NumStreamedLODs.GetValueForPlatformIdentifiers(
+				RunningPlatform->GetPlatformInfo().PlatformGroupName,
+				RunningPlatform->GetPlatformInfo().VanillaPlatformName);
+		if (MaxNumStreamedLODs < 0)
+		{
+			MaxNumStreamedLODs = LODGroup.GetDefaultMaxNumStreamedLODs();
+		}
+		for (int32 Idx = 0; Idx < 4; ++Idx)
+		{
+			ByteToHex((MaxNumStreamedLODs & 0xff000000) >> 24, KeySuffix);
+			MaxNumStreamedLODs <<= 8;
+		}
+	}
+	else
+	{
+		KeySuffix += TEXT("zzzzzzzz");
+	}
+
 	KeySuffix.AppendChar(Mesh->bSupportUniformlyDistributedSampling ? TEXT('1') : TEXT('0'));
 
 	// Value of this CVar affects index buffer <-> painted vertex color correspondence (see UE-51421).
@@ -1591,12 +1986,23 @@ static FString BuildStaticMeshDerivedDataKey(UStaticMesh* Mesh, const FStaticMes
 				break;
 		}
 	}
+	return KeySuffix;
+}
 
+static FString BuildStaticMeshDerivedDataKey(const FString& KeySuffix)
+{
 	return FDerivedDataCacheInterface::BuildCacheKey(
 		TEXT("STATICMESH"),
 		*GetStaticMeshDerivedDataVersion(),
-		*KeySuffix
-		);
+		*KeySuffix);
+}
+
+static FString BuildStaticMeshLODDerivedDataKey(const FString& KeySuffix, int32 LODIdx)
+{
+	return FDerivedDataCacheInterface::BuildCacheKey(
+		TEXT("STATICMESH"),
+		*GetStaticMeshDerivedDataVersion(),
+		*FString::Printf(TEXT("%s_LOD%d"), *KeySuffix, LODIdx));
 }
 
 void FStaticMeshRenderData::ComputeUVDensities()
@@ -1691,7 +2097,8 @@ void FStaticMeshRenderData::Cache(UStaticMesh* Owner, const FStaticMeshLODSettin
 		int32 T0 = FPlatformTime::Cycles();
 		int32 NumLODs = Owner->SourceModels.Num();
 		const FStaticMeshLODGroup& LODGroup = LODSettings.GetLODGroup(Owner->LODGroup);
-		DerivedDataKey = BuildStaticMeshDerivedDataKey(Owner, LODGroup);
+		const FString KeySuffix = BuildStaticMeshDerivedDataKeySuffix(Owner, LODGroup);
+		DerivedDataKey = BuildStaticMeshDerivedDataKey(KeySuffix);
 
 		TArray<uint8> DerivedData;
 		if (GetDerivedDataCacheRef().GetSynchronous(*DerivedDataKey, DerivedData))
@@ -1700,6 +2107,26 @@ void FStaticMeshRenderData::Cache(UStaticMesh* Owner, const FStaticMeshLODSettin
 			FMemoryReader Ar(DerivedData, /*bIsPersistent=*/ true);
 			Serialize(Ar, Owner, /*bCooked=*/ false);
 
+			for (int32 LODIdx = 0; LODIdx < LODResources.Num(); ++LODIdx)
+			{
+				FStaticMeshLODResources& LODResource = LODResources[LODIdx];
+				if (LODResource.bBuffersInlined)
+				{
+					break;
+				}
+				// TODO: can we postpone the loading to streaming time?
+				LODResource.DerivedDataKey = BuildStaticMeshLODDerivedDataKey(KeySuffix, LODIdx);
+				TArray<uint8> LODDerivedData;
+				verify(GetDerivedDataCacheRef().GetSynchronous(*LODResource.DerivedDataKey, LODDerivedData));
+				FMemoryReader LODAr(LODDerivedData, true);
+				typename FStaticMeshLODResources::FStaticMeshBuffersSize DummyBuffersSize;
+				LODResource.SerializeBuffers(LODAr, Owner, 0, DummyBuffersSize);
+				typename FStaticMeshLODResources::FStaticMeshBuffersSize LODBuffersSize;
+				Ar << LODBuffersSize;
+				LODResource.BuffersSize = LODBuffersSize.CalcBuffersSize();
+				check(LODResource.BuffersSize == DummyBuffersSize.CalcBuffersSize());
+			}
+
 			int32 T1 = FPlatformTime::Cycles();
 			UE_LOG(LogStaticMesh,Verbose,TEXT("Static mesh found in DDC [%fms] %s"),
 				FPlatformTime::ToMilliseconds(T1-T0),
@@ -1730,6 +2157,24 @@ void FStaticMeshRenderData::Cache(UStaticMesh* Owner, const FStaticMeshLODSettin
 			bLODsShareStaticLighting = Owner->CanLODsShareStaticLighting();
 			FMemoryWriter Ar(DerivedData, /*bIsPersistent=*/ true);
 			Serialize(Ar, Owner, /*bCooked=*/ false);
+
+			for (int32 LODIdx = 0; LODIdx < LODResources.Num(); ++LODIdx)
+			{
+				FStaticMeshLODResources& LODResource = LODResources[LODIdx];
+				if (LODResource.bBuffersInlined)
+				{
+					break;
+				}
+				TArray<uint8> LODDerivedData;
+				FMemoryWriter LODAr(LODDerivedData, true);
+				typename FStaticMeshLODResources::FStaticMeshBuffersSize LODBuffersSize;
+				const uint8 LODStripFlags = FStaticMeshLODResources::GenerateClassStripFlags(LODAr, Owner, LODIdx);
+				LODResource.SerializeBuffers(LODAr, Owner, LODStripFlags, LODBuffersSize);
+				Ar << LODBuffersSize;
+				LODResource.DerivedDataKey = BuildStaticMeshLODDerivedDataKey(KeySuffix, LODIdx);
+				GetDerivedDataCacheRef().Put(*LODResource.DerivedDataKey, LODDerivedData);
+			}
+
 			GetDerivedDataCacheRef().Put(*DerivedDataKey, DerivedData);
 
 			int32 T1 = FPlatformTime::Cycles();
@@ -1822,7 +2267,7 @@ const float UStaticMesh::MinimumAutoLODPixelError = SMALL_NUMBER;
 #endif	//#if WITH_EDITORONLY_DATA
 
 UStaticMesh::UStaticMesh(const FObjectInitializer& ObjectInitializer)
-	: UObject(ObjectInitializer)
+	: UStreamableRenderAsset(ObjectInitializer)
 {
 	ElementToIgnoreForTexFactor = -1;
 	bHasNavigationData=true;
@@ -1830,6 +2275,7 @@ UStaticMesh::UStaticMesh(const FObjectInitializer& ObjectInitializer)
 	bAutoComputeLODScreenSize=true;
 	ImportVersion = EImportStaticMeshVersion::BeforeImportStaticMeshVersionWasAdded;
 	LODForOccluderMesh = -1;
+	NumStreamedLODs.Default = -1;
 #endif // #if WITH_EDITORONLY_DATA
 	LightMapResolution = 4;
 	LpvBiasMultiplier = 1.0f;
@@ -1872,6 +2318,30 @@ void UStaticMesh::InitResources()
 	{
 		INC_DWORD_STAT_BY( STAT_StaticMeshOccluderMemory, OccluderData->GetResourceSizeBytes() );
 	}
+
+	// Determine whether or not this mesh can be streamed.
+	const int32 NumLODs = GetNumLODs();
+	bIsStreamable = !NeverStream
+		&& NumLODs > 1
+		&& !RenderData->LODResources[0].bBuffersInlined;
+		//&& !bTemporarilyDisableStreaming;
+
+#if (WITH_EDITOR && DO_CHECK)
+	if (bIsStreamable)
+	{
+		for (int32 LODIdx = 0; LODIdx < NumLODs; ++LODIdx)
+		{
+			const FStaticMeshLODResources& LODResource = RenderData->LODResources[LODIdx];
+			check(LODResource.bBuffersInlined || !LODResource.DerivedDataKey.IsEmpty());
+		}
+	}
+#endif
+
+	UnlinkStreaming();
+	if (bIsStreamable)
+	{
+		LinkStreaming();
+	}
 	
 #if	STATS
 	UStaticMesh* This = this;
@@ -2531,6 +3001,15 @@ void UStaticMesh::BeginDestroy()
 {
 	Super::BeginDestroy();
 
+	// Cancel any in flight IO requests
+	CancelPendingMipChangeRequest();
+
+	// Safely unlink mesh from list of streamable ones.
+	UnlinkStreaming();
+
+	// Remove from the list of tracked assets if necessary
+	TrackRenderAssetEvent(nullptr, this, false, nullptr);
+
 	if (FApp::CanEverRender() && !HasAnyFlags(RF_ClassDefaultObject))
 	{
 		ReleaseResources();
@@ -2955,7 +3434,8 @@ static FStaticMeshRenderData& GetPlatformStaticMeshRenderData(UStaticMesh* Mesh,
 {
 	check(Mesh && Mesh->RenderData);
 	const FStaticMeshLODSettings& PlatformLODSettings = Platform->GetStaticMeshLODSettings();
-	FString PlatformDerivedDataKey = BuildStaticMeshDerivedDataKey(Mesh, PlatformLODSettings.GetLODGroup(Mesh->LODGroup));
+	FString PlatformDerivedDataKey = BuildStaticMeshDerivedDataKey(
+		BuildStaticMeshDerivedDataKeySuffix(Mesh, PlatformLODSettings.GetLODGroup(Mesh->LODGroup)));
 	FStaticMeshRenderData* PlatformRenderData = Mesh->RenderData.Get();
 
 	if (Mesh->GetOutermost()->HasAnyPackageFlags(PKG_FilterEditorOnly))
@@ -4095,6 +4575,248 @@ bool UStaticMesh::CanBeClusterRoot() const
 	return false;
 }
 
+int32 UStaticMesh::GetLODGroupForStreaming() const
+{
+	// TODO: mesh LOD streaming may need to know LOD group settings
+	return 0;
+}
+
+int32 UStaticMesh::GetNumMipsForStreaming() const
+{
+	check(RenderData);
+	return GetNumLODs();
+}
+
+int32 UStaticMesh::GetNumNonStreamingMips() const
+{
+	check(RenderData);
+	return RenderData->NumInlinedLODs;
+}
+
+int32 UStaticMesh::CalcNumOptionalMips() const
+{
+#if !WITH_EDITOR
+	return MinLOD.Default;
+#else
+	return 0;
+#endif
+}
+
+int32 UStaticMesh::CalcCumulativeLODSize(int32 NumLODs) const
+{
+	uint32 Accum = 0;
+	const int32 LODCount = GetNumLODs();
+	const int32 LastLODIdx = LODCount - NumLODs;
+	for (int32 Idx = LODCount - 1; Idx >= LastLODIdx; --Idx)
+	{
+		Accum += RenderData->LODResources[Idx].BuffersSize;
+	}
+	return Accum;
+}
+
+bool UStaticMesh::GetMipDataFilename(const int32 MipIndex, FString& OutBulkDataFilename) const
+{
+#if !WITH_EDITOR
+	FString PackageName = GetOutermost()->FileName.ToString();
+	// Handle name redirection and localization
+	const FCoreRedirectObjectName RedirectedName =
+		FCoreRedirects::GetRedirectedName(
+			ECoreRedirectFlags::Type_Package,
+			FCoreRedirectObjectName(NAME_None, NAME_None, *PackageName));
+	FString LocalizedName;
+	LocalizedName = FPackageName::GetDelegateResolvedPackagePath(RedirectedName.PackageName.ToString());
+	LocalizedName = FPackageName::GetLocalizedPackagePath(LocalizedName);
+	bool bSucceed = FPackageName::DoesPackageExist(LocalizedName, nullptr, &OutBulkDataFilename);
+	check(bSucceed);
+	OutBulkDataFilename = FPaths::ChangeExtension(OutBulkDataFilename, MipIndex < MinLOD.Default ? TEXT(".uptnl") : TEXT(".ubulk"));
+	check(MipIndex < MinLOD.Default || IFileManager::Get().FileExists(*OutBulkDataFilename));
+	return true;
+#else
+	return false;
+#endif
+}
+
+bool UStaticMesh::IsReadyForStreaming() const
+{
+	return RenderData && RenderData->bReadyForStreaming;
+}
+
+int32 UStaticMesh::GetNumResidentMips() const
+{
+	check(RenderData);
+	return GetNumLODs() - RenderData->CurrentFirstLODIdx;
+}
+
+int32 UStaticMesh::GetNumRequestedMips() const
+{
+	if (PendingUpdate && !PendingUpdate->IsCancelled())
+	{
+		return PendingUpdate->GetNumRequestedMips();
+	}
+	else
+	{
+		return GetCachedNumResidentLODs();
+	}
+}
+
+bool UStaticMesh::CancelPendingMipChangeRequest()
+{
+	if (PendingUpdate)
+	{
+		if (!PendingUpdate->IsCancelled())
+		{
+			PendingUpdate->Abort();
+		}
+		return true;
+	}
+	return false;
+}
+
+bool UStaticMesh::HasPendingUpdate() const
+{
+	return !!PendingUpdate;
+}
+
+bool UStaticMesh::StreamOut(int32 NewMipCount)
+{
+	check(IsInGameThread());
+	if (bIsStreamable && !PendingUpdate && RenderData.IsValid() && RenderData->bReadyForStreaming && NewMipCount < GetNumResidentMips())
+	{
+		PendingUpdate = new FStaticMeshStreamOut(this, NewMipCount);
+
+		// The object starts in the locked state while it is being initialized.
+		PendingUpdate->DoUnlock();
+
+		return !PendingUpdate->IsCancelled();
+	}
+	return false;
+}
+
+bool UStaticMesh::StreamIn(int32 NewMipCount, bool bHighPrio)
+{
+	check(IsInGameThread());
+	if (bIsStreamable && !PendingUpdate && RenderData.IsValid() && RenderData->bReadyForStreaming && NewMipCount > GetNumResidentMips())
+	{
+#if WITH_EDITOR
+		if (FPlatformProperties::HasEditorOnlyData())
+		{
+			if (GRHISupportsAsyncTextureCreation)
+			{
+				PendingUpdate = new FStaticMeshStreamIn_DDC_Async(this, NewMipCount);
+			}
+			else
+			{
+				PendingUpdate = new FStaticMeshStreamIn_DDC_RenderThread(this, NewMipCount);
+			}
+		}
+		else
+#endif
+		{
+			if (GRHISupportsAsyncTextureCreation)
+			{
+				PendingUpdate = new FStaticMeshStreamIn_IO_Async(this, NewMipCount, bHighPrio);
+			}
+			else
+			{
+				PendingUpdate = new FStaticMeshStreamIn_IO_RenderThread(this, NewMipCount, bHighPrio);
+			}
+		}
+
+		// The object starts in the locked state while it is being initialized.
+		PendingUpdate->DoUnlock();
+		return !PendingUpdate->IsCancelled();
+	}
+	return false;
+}
+
+bool UStaticMesh::UpdateStreamingStatus(bool bWaitForMipFading)
+{
+	// if resident and requested mip counts match then no pending request is in flight
+	if (PendingUpdate)
+	{
+		if (GIsRequestingExit || !RenderData)
+		{
+			PendingUpdate->Abort();
+		}
+
+		// When the renderthread is the gamethread, allow the Tick to execute rendercommands.
+		PendingUpdate->Tick(this, GIsThreadedRendering ? FStaticMeshUpdate::TT_None : FStaticMeshUpdate::TT_Render);
+		if (!PendingUpdate->IsCompleted())
+		{
+			return true;
+		}
+
+#if WITH_EDITOR
+		const bool bRebuildPlatformData = PendingUpdate->DDCIsInvalid() && !IsPendingKillOrUnreachable();
+#endif
+
+		delete PendingUpdate;
+		PendingUpdate = nullptr;
+
+#if WITH_EDITOR
+		if (GIsEditor)
+		{
+			// When all the requested mips are streamed in, generate an empty property changed event, to force the
+			// ResourceSize asset registry tag to be recalculated.
+			FPropertyChangedEvent EmptyPropertyChangedEvent(nullptr);
+			FCoreUObjectDelegates::OnObjectPropertyChanged.Broadcast(this, EmptyPropertyChangedEvent);
+
+			// We can't load the source art from a bulk data object if the mesh itself is pending kill because the linker will have been detached.
+			// In this case we don't rebuild the data and instead let the streaming request be cancelled. This will let the garbage collector finish
+			// destroying the object.
+			if (bRebuildPlatformData)
+			{
+				// TODO: force rebuild even if DDC keys match
+				ITargetPlatformManagerModule& TargetPlatformManager = GetTargetPlatformManagerRef();
+				ITargetPlatform* TargetPlatform = TargetPlatformManager.GetRunningTargetPlatform();
+				check(TargetPlatform);
+				const FStaticMeshLODSettings& LODSettings = TargetPlatform->GetStaticMeshLODSettings();
+				RenderData->Cache(this, LODSettings);
+				// @TODO this can not be called from this callstack since the entry needs to be removed completely from the streamer.
+				// UpdateResource();
+			}
+		}
+#endif
+	}
+
+	// TODO: LOD fading?
+
+	return false;
+}
+
+void UStaticMesh::LinkStreaming()
+{
+	if (!IsTemplate() && IStreamingManager::Get().IsTextureStreamingEnabled() && IsStreamingRenderAsset(this))
+	{
+		IStreamingManager::Get().GetTextureStreamingManager().AddStreamingRenderAsset(this);
+	}
+	else
+	{
+		StreamingIndex = INDEX_NONE;
+	}
+}
+
+void UStaticMesh::UnlinkStreaming()
+{
+	if (!IsTemplate() && IStreamingManager::Get().IsTextureStreamingEnabled())
+	{
+		IStreamingManager::Get().GetTextureStreamingManager().RemoveStreamingRenderAsset(this);
+	}
+}
+
+void UStaticMesh::CancelAllPendingStreamingActions()
+{
+	FlushRenderingCommands();
+
+	for (TObjectIterator<UStaticMesh> It; It; ++It)
+	{
+		UStaticMesh* StaticMesh = *It;
+		StaticMesh->CancelPendingMipChangeRequest();
+	}
+
+	FlushRenderingCommands();
+}
+
 //
 //	UStaticMesh::GetDesc
 //
diff --git a/Engine/Source/Runtime/Engine/Private/StaticMeshRender.cpp b/Engine/Source/Runtime/Engine/Private/StaticMeshRender.cpp
--- a/Engine/Source/Runtime/Engine/Private/StaticMeshRender.cpp
+++ b/Engine/Source/Runtime/Engine/Private/StaticMeshRender.cpp
@@ -363,6 +363,11 @@ bool FStaticMeshSceneProxy::GetShadowMeshElement(int32 LODIndex, int32 BatchInde
 
 	const bool bUseReversedIndices = GUseReversedIndexBuffer && IsLocalToWorldDeterminantNegative() && LOD.bHasReversedDepthOnlyIndices;
 
+	if ((bUseReversedIndices && !LOD.bHasReversedDepthOnlyIndices) || (!bUseReversedIndices && !LOD.bHasDepthOnlyIndices))
+	{
+		return false;
+	}
+
 	FMeshBatchElement& OutBatchElement = OutMeshBatch.Elements[0];
 	OutBatchElement.VertexFactoryUserData = ProxyLODInfo.OverrideColorVertexBuffer ? ProxyLODInfo.OverrideColorVFUniformBuffer.GetReference() : VFs.VertexFactory.GetUniformBuffer();
 	OutMeshBatch.MaterialRenderProxy = UMaterial::GetDefaultMaterial(MD_Surface)->GetRenderProxy();
@@ -1817,6 +1822,10 @@ FLODMask FStaticMeshSceneProxy::GetLODMask(const FSceneView* View) const
 			}
 		}
 	}
+
+	const int8 CurFirstLODIdx = GetCurrentFirstLODIdx_Internal();
+	check(CurFirstLODIdx >= 0);
+	Result.ClampToFirstLOD(CurFirstLODIdx);
 	
 	return Result;
 }
diff --git a/Engine/Source/Runtime/Engine/Private/StreamableRenderAsset.cpp b/Engine/Source/Runtime/Engine/Private/StreamableRenderAsset.cpp
new file mode 100644
--- /dev/null
+++ b/Engine/Source/Runtime/Engine/Private/StreamableRenderAsset.cpp
@@ -0,0 +1,19 @@
+// Copyright 1998-2018 Epic Games, Inc. All Rights Reserved.
+
+#include "Engine/StreamableRenderAsset.h"
+#include "Misc/App.h"
+
+
+UStreamableRenderAsset::UStreamableRenderAsset(const FObjectInitializer& ObjectInitializer)
+	: Super(ObjectInitializer)
+	, StreamingIndex(INDEX_NONE)
+{}
+
+void UStreamableRenderAsset::SetForceMipLevelsToBeResident(float Seconds, int32 CinematicLODGroupMask)
+{
+	uint32 TextureGroupBitfield = (uint32)CinematicLODGroupMask;
+	int32 LODGroup = GetLODGroupForStreaming();
+	uint32 MyTextureGroup = FMath::BitFlag[LODGroup];
+	bUseCinematicMipLevels = (TextureGroupBitfield & MyTextureGroup) ? true : false;
+	ForceMipLevelsToBeResidentTimestamp = FApp::GetCurrentTime() + Seconds;
+}
