Resolve robomerge conflicts

diff --git a/Engine/Source/Runtime/Engine/Private/Streaming/AsyncTextureStreaming.cpp b/Engine/Source/Runtime/Engine/Private/Streaming/AsyncTextureStreaming.cpp
--- a/Engine/Source/Runtime/Engine/Private/Streaming/AsyncTextureStreaming.cpp
+++ b/Engine/Source/Runtime/Engine/Private/Streaming/AsyncTextureStreaming.cpp
@@ -9,7 +9,11 @@ AsyncTextureStreaming.cpp: Definitions of classes used for texture streaming asy
 #include "Streaming/StreamingManagerTexture.h"
 #include "Engine/World.h"
 
-void FAsyncTextureStreamingData::Init(TArray<FStreamingViewInfo> InViewInfos, float InLastUpdateTime, TIndirectArray<FLevelTextureManager>& LevelTextureManagers, FDynamicTextureInstanceManager& DynamicComponentManager)
+void FAsyncRenderAssetStreamingData::Init(
+	TArray<FStreamingViewInfo> InViewInfos,
+	float InLastUpdateTime,
+	TIndirectArray<FLevelRenderAssetManager>& LevelStaticInstanceManagers,
+	FDynamicRenderAssetInstanceManager& DynamicComponentManager)
 {
 	ViewInfos = InViewInfos;
 	LastUpdateTime = InLastUpdateTime;
@@ -20,24 +24,54 @@ void FAsyncTextureStreamingData::Init(TArray<FStreamingViewInfo> InViewInfos, fl
 	StaticInstancesViewIndices.Reset();
 	CulledStaticInstancesViewIndices.Reset();
 
-	for (FLevelTextureManager& LevelManager : LevelTextureManagers)
+	for (FLevelRenderAssetManager& LevelManager : LevelStaticInstanceManagers)
 	{
-		if (LevelManager.IsInitialized() && LevelManager.GetLevel()->bIsVisible && LevelManager.HasTextureReferences())
+		if (LevelManager.IsInitialized() && LevelManager.GetLevel()->bIsVisible && LevelManager.HasRenderAssetReferences())
 		{
 			StaticInstancesViews.Push(LevelManager.GetAsyncView());
 		}
 	}
 }
 
-void FAsyncTextureStreamingData::UpdateBoundSizes_Async(const FTextureStreamingSettings& Settings)
+void FAsyncRenderAssetStreamingData::ComputeViewInfoExtras(const FRenderAssetStreamingSettings& Settings)
+{
+	const float OneOverMaxHiddenPrimitiveViewBoost = 1.f / Settings.MaxHiddenPrimitiveViewBoost;
+	const int32 NumViews = ViewInfos.Num();
+	ViewInfoExtras.Empty(NumViews);
+	ViewInfoExtras.AddUninitialized(NumViews);
+	MaxScreenSizeOverAllViews = 0.f;
+	
+	for (int32 ViewIndex = 0; ViewIndex < NumViews; ++ViewIndex)
+	{
+		const FStreamingViewInfo& ViewInfo = ViewInfos[ViewIndex];
+		FStreamingViewInfoExtra& ViewInfoExtra = ViewInfoExtras[ViewIndex];
+
+		const float EffectiveScreenSize = (Settings.MaxEffectiveScreenSize > 0.0f) ? FMath::Min(Settings.MaxEffectiveScreenSize, ViewInfo.ScreenSize) : ViewInfo.ScreenSize;
+		MaxScreenSizeOverAllViews = FMath::Max(MaxScreenSizeOverAllViews, EffectiveScreenSize);
+		ViewInfoExtra.ScreenSizeFloat = EffectiveScreenSize * .5f; // Multiply by half since the ratio factors map to half the screen only
+		ViewInfoExtra.ExtraBoostForVisiblePrimitiveFloat = 1.f;
+
+		if (ViewInfo.BoostFactor > Settings.MaxHiddenPrimitiveViewBoost)
+		{
+			ViewInfoExtra.ScreenSizeFloat *= Settings.MaxHiddenPrimitiveViewBoost;
+			ViewInfoExtra.ExtraBoostForVisiblePrimitiveFloat = ViewInfo.BoostFactor * OneOverMaxHiddenPrimitiveViewBoost;
+		}
+		else
+		{
+			ViewInfoExtra.ScreenSizeFloat *= ViewInfo.BoostFactor;
+		}
+	}
+}
+
+void FAsyncRenderAssetStreamingData::UpdateBoundSizes_Async(const FRenderAssetStreamingSettings& Settings)
 {
 	for (int32 StaticViewIndex = 0; StaticViewIndex < StaticInstancesViews.Num(); ++StaticViewIndex)
 	{
-		FTextureInstanceAsyncView& StaticInstancesView = StaticInstancesViews[StaticViewIndex];
-		StaticInstancesView.UpdateBoundSizes_Async(ViewInfos, LastUpdateTime, Settings);
+		FRenderAssetInstanceAsyncView& StaticInstancesView = StaticInstancesViews[StaticViewIndex];
+		StaticInstancesView.UpdateBoundSizes_Async(ViewInfos, ViewInfoExtras, LastUpdateTime, Settings);
 
 		// Skip levels that can not contribute to resolution.
-		if (StaticInstancesView.GetMaxLevelTextureScreenSize() > Settings.MinLevelTextureScreenSize)
+		if (StaticInstancesView.GetMaxLevelRenderAssetScreenSize() > Settings.MinLevelRenderAssetScreenSize)
 		{
 			StaticInstancesViewIndices.Add(StaticViewIndex);
 		}
@@ -48,70 +82,79 @@ void FAsyncTextureStreamingData::UpdateBoundSizes_Async(const FTextureStreamingS
 	}
 	
 	// Sort by max possible size, this allows early exit when iteration on many levels.
-	if (Settings.MinLevelTextureScreenSize > 0)
+	if (Settings.MinLevelRenderAssetScreenSize > 0)
 	{
-		StaticInstancesViewIndices.Sort([&](int32 LHS, int32 RHS) { return StaticInstancesViews[LHS].GetMaxLevelTextureScreenSize() > StaticInstancesViews[RHS].GetMaxLevelTextureScreenSize(); });
+		StaticInstancesViewIndices.Sort([&](int32 LHS, int32 RHS) { return StaticInstancesViews[LHS].GetMaxLevelRenderAssetScreenSize() > StaticInstancesViews[RHS].GetMaxLevelRenderAssetScreenSize(); });
 	}
 
-	DynamicInstancesView.UpdateBoundSizes_Async(ViewInfos, LastUpdateTime, Settings);
+	DynamicInstancesView.UpdateBoundSizes_Async(ViewInfos, ViewInfoExtras, LastUpdateTime, Settings);
 }
 
-void FAsyncTextureStreamingData::UpdatePerfectWantedMips_Async(FStreamingTexture& StreamingTexture, const FTextureStreamingSettings& Settings, bool bOutputToLog) const
+void FAsyncRenderAssetStreamingData::UpdatePerfectWantedMips_Async(FStreamingRenderAsset& StreamingRenderAsset, const FRenderAssetStreamingSettings& Settings, bool bOutputToLog) const
 {
 #if UE_BUILD_SHIPPING
 	bOutputToLog = false;
 #endif
 
-	// Cache Texture on the stack as it could be nullified on the gamethread.
-	const UTexture2D* Texture = StreamingTexture.Texture;
-	if (!Texture) return;
+	// Cache RenderAsset on the stack as it could be nullified on the gamethread.
+	const UStreamableRenderAsset* RenderAsset = StreamingRenderAsset.RenderAsset;
+	if (!RenderAsset) return;
 
 	float MaxSize = 0;
 	float MaxSize_VisibleOnly = 0;
+	int32 MaxNumForcedLODs = 0;
 	bool bLooksLowRes = false;
 
-	const float MaxAllowedSize = StreamingTexture.GetMaxAllowedSize();
+	const float MaxAllowedSize = StreamingRenderAsset.GetMaxAllowedSize(MaxScreenSizeOverAllViews);
 
 #if !UE_BUILD_SHIPPING
 	if (Settings.bStressTest)
 	{
 		// In stress test, we choose between the allowed mips. Combined with "r.Streaming.DropMips=2" this can also generate cancel requests.
-		MaxSize_VisibleOnly = MaxSize = (float)(0x1 << (FMath::RandRange(StreamingTexture.MinAllowedMips, StreamingTexture.MaxAllowedMips) - 1));
+		const int32 NumMips = FMath::RandRange(StreamingRenderAsset.MinAllowedMips, StreamingRenderAsset.MaxAllowedMips);
+		MaxSize_VisibleOnly = MaxSize = StreamingRenderAsset.GetLODScreenSize(NumMips, MaxScreenSizeOverAllViews);
 	}
 	else
 #endif
 	if (Settings.bFullyLoadUsedTextures)
 	{
-		if (StreamingTexture.LastRenderTime < 300)
+		if (StreamingRenderAsset.LastRenderTime < 300)
 		{
 			MaxSize_VisibleOnly = FLT_MAX;
 		}
 	}
-	else if (StreamingTexture.MinAllowedMips == StreamingTexture.MaxAllowedMips)
+	else if (StreamingRenderAsset.MinAllowedMips == StreamingRenderAsset.MaxAllowedMips)
 	{
 		MaxSize_VisibleOnly = MaxSize = MaxAllowedSize;
 	}
 	else
 	{
-		DynamicInstancesView.GetTexelSize(Texture, MaxSize, MaxSize_VisibleOnly, bOutputToLog ? TEXT("Dynamic") : nullptr);
+		const typename FStreamingRenderAsset::EAssetType AssetType = StreamingRenderAsset.RenderAssetType;
+		DynamicInstancesView.GetRenderAssetScreenSize(AssetType, RenderAsset, MaxSize, MaxSize_VisibleOnly, MaxNumForcedLODs, bOutputToLog ? TEXT("Dynamic") : nullptr);
 
 		for (int32 StaticViewIndex : StaticInstancesViewIndices)
 		{
-			const FTextureInstanceAsyncView& StaticInstancesView = StaticInstancesViews[StaticViewIndex];
+			const FRenderAssetInstanceAsyncView& StaticInstancesView = StaticInstancesViews[StaticViewIndex];
 			
 			// No need to iterate more if texture is already at maximum resolution.
-			if ((MaxSize_VisibleOnly >= MAX_TEXTURE_SIZE || (MaxSize_VisibleOnly > StaticInstancesView.GetMaxLevelTextureScreenSize() && Settings.MinLevelTextureScreenSize > 0)) && !bOutputToLog)
+			if ((MaxNumForcedLODs >= StreamingRenderAsset.MaxAllowedMips
+				|| MaxSize_VisibleOnly >= MAX_TEXTURE_SIZE
+				|| (MaxSize_VisibleOnly > StaticInstancesView.GetMaxLevelRenderAssetScreenSize() && Settings.MinLevelRenderAssetScreenSize > 0))
+				&& !bOutputToLog)
 			{
 				break;
 			}
 
-			StaticInstancesView.GetTexelSize(Texture, MaxSize, MaxSize_VisibleOnly, bOutputToLog ? TEXT("Static") : nullptr);
+			StaticInstancesView.GetRenderAssetScreenSize(AssetType, RenderAsset, MaxSize, MaxSize_VisibleOnly, MaxNumForcedLODs, bOutputToLog ? TEXT("Static") : nullptr);
 		}
 
-		// Don't apply to FLT_MAX since it is used as forced streaming. BoostFactor as only meaning for texture instances since the other heuristics are based on max resolution.
-		if ((MaxSize > 0 || MaxSize_VisibleOnly > 0) && MaxSize != FLT_MAX && MaxSize_VisibleOnly != FLT_MAX)
+		// Don't apply to FLT_MAX since it is used as forced streaming. BoostFactor as only meaning for texture/mesh instances since the other heuristics are based on max resolution.
+		if (MaxNumForcedLODs < StreamingRenderAsset.MaxAllowedMips
+			&& (MaxSize > 0 || MaxSize_VisibleOnly > 0)
+			&& MaxSize != FLT_MAX
+			&& MaxSize_VisibleOnly != FLT_MAX)
 		{
-			const float CumBoostFactor = StreamingTexture.BoostFactor * StreamingTexture.DynamicBoostFactor;
+			const float CumBoostFactor = StreamingRenderAsset.BoostFactor * StreamingRenderAsset.DynamicBoostFactor;
 
 			// If there is not enough resolution in the texture to fix the required quality, save this information to prevent degrading this texture before other ones.
 			bLooksLowRes = FMath::Max3<int32>(MaxSize_VisibleOnly, MaxSize, MaxAllowedSize) / MaxAllowedSize >= CumBoostFactor * 2.f;
@@ -121,38 +164,44 @@ void FAsyncTextureStreamingData::UpdatePerfectWantedMips_Async(FStreamingTexture
 		}
 
 		// Last part checks that it has been used since the last reference was removed.
-		const float TimeSinceRemoved = (float)(FApp::GetCurrentTime() - StreamingTexture.InstanceRemovedTimestamp);
-		StreamingTexture.bUseUnkownRefHeuristic = MaxSize == 0 && MaxSize_VisibleOnly == 0 && StreamingTexture.LastRenderTime < TimeSinceRemoved - 5.f;
-		if (StreamingTexture.bUseUnkownRefHeuristic)
+		const float TimeSinceRemoved = (float)(FApp::GetCurrentTime() - StreamingRenderAsset.InstanceRemovedTimestamp);
+		StreamingRenderAsset.bUseUnkownRefHeuristic = MaxSize == 0 && MaxSize_VisibleOnly == 0 && !MaxNumForcedLODs && StreamingRenderAsset.LastRenderTime < TimeSinceRemoved - 5.f;
+		if (StreamingRenderAsset.bUseUnkownRefHeuristic)
 		{
 			// Check that it's not simply culled
 			for (int32 StaticViewIndex : CulledStaticInstancesViewIndices)
 			{
-				const FTextureInstanceAsyncView& StaticInstancesView = StaticInstancesViews[StaticViewIndex];
-				if (StaticInstancesView.HasTextureReferences(Texture))
+				const FRenderAssetInstanceAsyncView& StaticInstancesView = StaticInstancesViews[StaticViewIndex];
+				if (StaticInstancesView.HasRenderAssetReferences(RenderAsset))
 				{
-					StreamingTexture.bUseUnkownRefHeuristic = false;
+					StreamingRenderAsset.bUseUnkownRefHeuristic = false;
 					break;
 				}
 			}
 
-			if (StreamingTexture.bUseUnkownRefHeuristic)
+			if (StreamingRenderAsset.bUseUnkownRefHeuristic)
 			{
 				if (bOutputToLog) UE_LOG(LogContentStreaming, Log,  TEXT("  UnkownRef"));
 				MaxSize = FMath::Max<int32>(MaxSize, MaxAllowedSize); // affected by HiddenPrimitiveScale
-				if (StreamingTexture.LastRenderTime < 5.0f)
+				if (StreamingRenderAsset.LastRenderTime < 5.0f)
 				{
 					MaxSize_VisibleOnly = FMath::Max<int32>(MaxSize_VisibleOnly, MaxAllowedSize);
 				}
 			}
 		}
 
-		if (StreamingTexture.bForceFullyLoad || (StreamingTexture.LODGroup == TEXTUREGROUP_HierarchicalLOD && Settings.HLODStrategy == 2))
+		// TODO: for meshes, how to determine whether they are HLOD?
+		if (StreamingRenderAsset.bForceFullyLoad
+			|| (AssetType == FStreamingRenderAsset::AT_Texture
+				&& StreamingRenderAsset.LODGroup == TEXTUREGROUP_HierarchicalLOD
+				&& Settings.HLODStrategy == 2))
 		{
 			if (bOutputToLog) UE_LOG(LogContentStreaming, Log,  TEXT("  Forced FullyLoad"));
 			MaxSize = FLT_MAX; // Forced load ensure the texture gets fully loaded but after what is visible/required by the other logic.
 		}
-		else if (StreamingTexture.LODGroup == TEXTUREGROUP_HierarchicalLOD && Settings.HLODStrategy == 1)
+		else if (AssetType == FStreamingRenderAsset::AT_Texture
+			&& StreamingRenderAsset.LODGroup == TEXTUREGROUP_HierarchicalLOD
+			&& Settings.HLODStrategy == 1)
 		{
 			if (bOutputToLog) UE_LOG(LogContentStreaming, Log,  TEXT("  HLOD Strategy"));
 
@@ -173,10 +222,10 @@ void FAsyncTextureStreamingData::UpdatePerfectWantedMips_Async(FStreamingTexture
 		MaxSize_VisibleOnly = MaxSize = FMath::Max<float>(MaxSize_VisibleOnly, MaxSize);
 	}
 
-	StreamingTexture.SetPerfectWantedMips_Async(MaxSize, MaxSize_VisibleOnly, bLooksLowRes, Settings);
+	StreamingRenderAsset.SetPerfectWantedMips_Async(MaxSize, MaxSize_VisibleOnly, MaxScreenSizeOverAllViews, MaxNumForcedLODs, bLooksLowRes, Settings);
 }
 
-bool FAsyncTextureStreamingTask::AllowPerTextureMipBiasChanges() const
+bool FRenderAssetStreamingMipCalcTask::AllowPerRenderAssetMipBiasChanges() const
 {
 	const TArray<FStreamingViewInfo>& ViewInfos = StreamingData.GetViewInfos();
 	for (int32 ViewIndex = 0; ViewIndex < ViewInfos.Num(); ++ViewIndex)
@@ -190,31 +239,31 @@ bool FAsyncTextureStreamingTask::AllowPerTextureMipBiasChanges() const
 	return true;
 }
 
-void FAsyncTextureStreamingTask::UpdateBudgetedMips_Async(int64& MemoryUsed, int64& TempMemoryUsed)
+void FRenderAssetStreamingMipCalcTask::UpdateBudgetedMips_Async(int64& MemoryUsed, int64& TempMemoryUsed)
 {
 	//*************************************
 	// Update Budget
 	//*************************************
 
-	TArray<FStreamingTexture>& StreamingTextures = StreamingManager.StreamingTextures;
-	const FTextureStreamingSettings& Settings = StreamingManager.Settings;
+	TArray<FStreamingRenderAsset>& StreamingRenderAssets = StreamingManager.StreamingRenderAssets;
+	const FRenderAssetStreamingSettings& Settings = StreamingManager.Settings;
 
-	TArray<int32> PrioritizedTextures;
+	TArray<int32> PrioritizedRenderAssets;
 
 	int64 MemoryBudgeted = 0;
 	MemoryUsed = 0;
 	TempMemoryUsed = 0;
 
-	for (FStreamingTexture& StreamingTexture : StreamingTextures)
+	for (FStreamingRenderAsset& StreamingRenderAsset : StreamingRenderAssets)
 	{
 		if (IsAborted()) break;
 
-		MemoryBudgeted += StreamingTexture.UpdateRetentionPriority_Async();
-		MemoryUsed += StreamingTexture.GetSize(StreamingTexture.ResidentMips);
+		MemoryBudgeted += StreamingRenderAsset.UpdateRetentionPriority_Async();
+		MemoryUsed += StreamingRenderAsset.GetSize(StreamingRenderAsset.ResidentMips);
 
-		if (StreamingTexture.ResidentMips != StreamingTexture.RequestedMips)
+		if (StreamingRenderAsset.ResidentMips != StreamingRenderAsset.RequestedMips)
 		{
-			TempMemoryUsed += StreamingTexture.GetSize(StreamingTexture.RequestedMips);
+			TempMemoryUsed += StreamingRenderAsset.GetSize(StreamingRenderAsset.RequestedMips);
 		}
 	}
 
@@ -237,16 +286,16 @@ void FAsyncTextureStreamingTask::UpdateBudgetedMips_Async(int64& MemoryUsed, int
 	}
 
 
-	const int64 NonStreamingTextureMemory =  AllocatedMemory - MemoryUsed;
-	int64 AvailableMemoryForStreaming = PoolSize - NonStreamingTextureMemory - MemoryMargin;
+	const int64 NonStreamingRenderAssetMemory =  AllocatedMemory - MemoryUsed;
+	int64 AvailableMemoryForStreaming = PoolSize - NonStreamingRenderAssetMemory - MemoryMargin;
 
 	// If the platform defines a max VRAM usage, check if the pool size must be reduced,
-	// but also check if it would be safe to some of the NonStreamingTextureMemory from the pool size computation.
-	// The later helps significantly in low budget settings, where NonStreamingTextureMemory would take too much of the texture pool.
+	// but also check if it would be safe to some of the NonStreamingRenderAssetMemory from the pool size computation.
+	// The later helps significantly in low budget settings, where NonStreamingRenderAssetMemory would take too much of the pool.
 	if (GPoolSizeVRAMPercentage > 0 && TotalGraphicsMemory > 0)
 	{
 		const int64 UsableVRAM = TotalGraphicsMemory * GPoolSizeVRAMPercentage / 100 - (int64)GCurrentRendertargetMemorySize * 1024ll; // Add any other...
-		const int64 AvailableVRAMForStreaming = FMath::Min<int64>(UsableVRAM - NonStreamingTextureMemory - MemoryMargin, PoolSize);
+		const int64 AvailableVRAMForStreaming = FMath::Min<int64>(UsableVRAM - NonStreamingRenderAssetMemory - MemoryMargin, PoolSize);
 		if (Settings.bLimitPoolSizeToVRAM || AvailableVRAMForStreaming > AvailableMemoryForStreaming)
 		{
 			AvailableMemoryForStreaming = AvailableVRAMForStreaming;
@@ -273,17 +322,21 @@ void FAsyncTextureStreamingTask::UpdateBudgetedMips_Async(int64& MemoryUsed, int
 	// Reset per mip bias if not required anymore.
 	//*******************************************
 
-	// When using mip per texture, the BudgetMipBias gets reset when the required resolution does not get affected anymore by the BudgetMipBias.
-	// This allows texture to reset their bias when the viewpoint gets far enough, or the primitive is not visible anymore.
+	// When using mip per texture/mesh, the BudgetMipBias gets reset when the required resolution does not get affected anymore by the BudgetMipBias.
+	// This allows texture/mesh to reset their bias when the viewpoint gets far enough, or the primitive is not visible anymore.
 	if (Settings.bUsePerTextureBias)
 	{
-		for (FStreamingTexture& StreamingTexture : StreamingTextures)
+		for (FStreamingRenderAsset& StreamingRenderAsset : StreamingRenderAssets)
 		{
 			if (IsAborted()) break;
 
-			if ((bResetMipBias || FMath::Max<int32>(StreamingTexture.VisibleWantedMips, StreamingTexture.HiddenWantedMips + StreamingTexture.NumMissingMips) < StreamingTexture.MaxAllowedMips) && StreamingTexture.BudgetMipBias > 0)
+			if (StreamingRenderAsset.BudgetMipBias > 0
+				&& (bResetMipBias
+					|| FMath::Max<int32>(
+						StreamingRenderAsset.VisibleWantedMips,
+						StreamingRenderAsset.HiddenWantedMips + StreamingRenderAsset.NumMissingMips) < StreamingRenderAsset.MaxAllowedMips))
 			{
-				StreamingTexture.BudgetMipBias = 0;
+				StreamingRenderAsset.BudgetMipBias = 0;
 			}
 		}
 	}
@@ -296,62 +349,63 @@ void FAsyncTextureStreamingTask::UpdateBudgetedMips_Async(int64& MemoryUsed, int
 	if (MemoryBudgeted > MemoryBudget && !IsAborted())
 	{
 		//*************************************
-		// Get texture list in order of reduction
+		// Get texture/mesh list in order of reduction
 		//*************************************
 
-		PrioritizedTextures.Empty(StreamingTextures.Num());
+		PrioritizedRenderAssets.Empty(StreamingRenderAssets.Num());
 
-		for (int32 TextureIndex = 0; TextureIndex < StreamingTextures.Num() && !IsAborted(); ++TextureIndex)
+		for (int32 AssetIndex = 0; AssetIndex < StreamingRenderAssets.Num() && !IsAborted(); ++AssetIndex)
 		{
-			FStreamingTexture& StreamingTexture = StreamingTextures[TextureIndex];
-			// Only consider non deleted textures (can change any time).
-			if (!StreamingTexture.Texture) continue;
+			FStreamingRenderAsset& StreamingRenderAsset = StreamingRenderAssets[AssetIndex];
+			// Only consider non deleted textures/meshes (can change any time).
+			if (!StreamingRenderAsset.RenderAsset) continue;
 
-			// Ignore textures for which we are not allowed to reduce resolution.
-			if (!StreamingTexture.IsMaxResolutionAffectedByGlobalBias()) continue;
+			// Ignore textures/meshes for which we are not allowed to reduce resolution.
+			if (!StreamingRenderAsset.IsMaxResolutionAffectedByGlobalBias()) continue;
 
-			// Ignore texture that can't drop any mips
-			if (StreamingTexture.BudgetedMips > StreamingTexture.MinAllowedMips)
+			// Ignore texture/mesh that can't drop any mips
+			const int32 MinAllowedMips = FMath::Max(StreamingRenderAsset.MinAllowedMips, StreamingRenderAsset.NumForcedMips);
+			if (StreamingRenderAsset.BudgetedMips > MinAllowedMips)
 			{
-				PrioritizedTextures.Add(TextureIndex);
+				PrioritizedRenderAssets.Add(AssetIndex);
 			}
 		}
 
-		// Sort texture, having those that should be dropped first.
-		PrioritizedTextures.Sort(FCompareTextureByRetentionPriority(StreamingTextures));
+		// Sort texture/mesh, having those that should be dropped first.
+		PrioritizedRenderAssets.Sort(FCompareRenderAssetByRetentionPriority(StreamingRenderAssets));
 
 
-		if (Settings.bUsePerTextureBias && AllowPerTextureMipBiasChanges())
+		if (Settings.bUsePerTextureBias && AllowPerRenderAssetMipBiasChanges())
 		{
 			//*************************************
 			// Drop Max Resolution until in budget.
 			//*************************************
 
-			// When using mip bias per texture, we first reduce the maximum resolutions (if used) in order to fit.
+			// When using mip bias per texture/mesh, we first reduce the maximum resolutions (if used) in order to fit.
 			for (int32 NumDroppedMips = 0; NumDroppedMips < Settings.GlobalMipBias && MemoryBudgeted > MemoryBudget && !IsAborted(); ++NumDroppedMips)
 			{
 				const int64 PreviousMemoryBudgeted = MemoryBudgeted;
 
-				for (int32 PriorityIndex = PrioritizedTextures.Num() - 1; PriorityIndex >= 0 && MemoryBudgeted > MemoryBudget && !IsAborted(); --PriorityIndex)
+				for (int32 PriorityIndex = PrioritizedRenderAssets.Num() - 1; PriorityIndex >= 0 && MemoryBudgeted > MemoryBudget && !IsAborted(); --PriorityIndex)
 				{
-					int32 TextureIndex = PrioritizedTextures[PriorityIndex];
-					if (TextureIndex == INDEX_NONE) continue;
-
-					FStreamingTexture& StreamingTexture = StreamingTextures[TextureIndex];
+					int32 AssetIndex = PrioritizedRenderAssets[PriorityIndex];
+					if (AssetIndex == INDEX_NONE) continue;
 
-					if (StreamingTexture.BudgetedMips <= StreamingTexture.MinAllowedMips)
+					FStreamingRenderAsset& StreamingRenderAsset = StreamingRenderAssets[AssetIndex];
+					const int32 MinAllowedMips = FMath::Max(StreamingRenderAsset.MinAllowedMips, StreamingRenderAsset.NumForcedMips);
+					if (StreamingRenderAsset.BudgetedMips <= MinAllowedMips)
 					{
 						// Don't try this one again.
-						PrioritizedTextures[PriorityIndex] = INDEX_NONE;
+						PrioritizedRenderAssets[PriorityIndex] = INDEX_NONE;
 						continue;
 					}
 
-					// If the texture requires a high resolution mip, consider dropping it. 
-					// When considering dropping the first mip, only textures using the first mip will drop their resolution, 
-					// But when considering dropping the second mip, textures using their first and second mips will loose it.
-					if (StreamingTexture.MaxAllowedMips + StreamingTexture.BudgetMipBias - NumDroppedMips <= StreamingTexture.BudgetedMips)
+					// If the texture/mesh requires a high resolution mip, consider dropping it. 
+					// When considering dropping the first mip, only textures/meshes using the first mip will drop their resolution, 
+					// But when considering dropping the second mip, textures/meshes using their first and second mips will loose it.
+					if (StreamingRenderAsset.MaxAllowedMips + StreamingRenderAsset.BudgetMipBias - NumDroppedMips <= StreamingRenderAsset.BudgetedMips)
 					{
-						MemoryBudgeted -= StreamingTexture.DropMaxResolution_Async(NumDroppedMips + 1 - StreamingTexture.BudgetMipBias);
+						MemoryBudgeted -= StreamingRenderAsset.DropMaxResolution_Async(NumDroppedMips + 1 - StreamingRenderAsset.BudgetMipBias);
 					}
 				}
 
@@ -372,28 +426,28 @@ void FAsyncTextureStreamingTask::UpdateBudgetedMips_Async(int64& MemoryUsed, int
 			const int64 PreviousMemoryBudgeted = MemoryBudgeted;
 
 			// Drop from the lowest priority first (starting with last elements)
-			for (int32 PriorityIndex = PrioritizedTextures.Num() - 1; PriorityIndex >= 0 && MemoryBudgeted > MemoryBudget && !IsAborted(); --PriorityIndex)
+			for (int32 PriorityIndex = PrioritizedRenderAssets.Num() - 1; PriorityIndex >= 0 && MemoryBudgeted > MemoryBudget && !IsAborted(); --PriorityIndex)
 			{
-				int32 TextureIndex = PrioritizedTextures[PriorityIndex];
-				if (TextureIndex == INDEX_NONE) continue;
+				int32 AssetIndex = PrioritizedRenderAssets[PriorityIndex];
+				if (AssetIndex == INDEX_NONE) continue;
 
-				FStreamingTexture& StreamingTexture = StreamingTextures[TextureIndex];
-
-				if (StreamingTexture.BudgetedMips <= StreamingTexture.MinAllowedMips)
+				FStreamingRenderAsset& StreamingRenderAsset = StreamingRenderAssets[AssetIndex];
+				const int32 MinAllowedMips = FMath::Max(StreamingRenderAsset.MinAllowedMips, StreamingRenderAsset.NumForcedMips);
+				if (StreamingRenderAsset.BudgetedMips <= MinAllowedMips)
 				{
 					// Don't try this one again.
-					PrioritizedTextures[PriorityIndex] = INDEX_NONE;
+					PrioritizedRenderAssets[PriorityIndex] = INDEX_NONE;
 					continue;
 				}
 
-				// If this texture has already missing mips for its normal quality, don't drop more than required..
-				if (StreamingTexture.NumMissingMips > 0)
+				// If this texture/mesh has already missing mips for its normal quality, don't drop more than required..
+				if (StreamingRenderAsset.NumMissingMips > 0)
 				{
-					--StreamingTexture.NumMissingMips;
+					--StreamingRenderAsset.NumMissingMips;
 					continue;
 				}
 
-				MemoryBudgeted -= StreamingTexture.DropOneMip_Async();
+				MemoryBudgeted -= StreamingRenderAsset.DropOneMip_Async();
 			}
 
 			// Break when memory does not change anymore
@@ -412,24 +466,24 @@ void FAsyncTextureStreamingTask::UpdateBudgetedMips_Async(int64& MemoryUsed, int
 	// This will run even after sacrificing to fit in budget since some small unwanted mips could still be kept.
 	if (MemoryBudgeted < MemoryBudget && !IsAborted())
 	{
-		PrioritizedTextures.Empty(StreamingTextures.Num());
+		PrioritizedRenderAssets.Empty(StreamingRenderAssets.Num());
 		const int64 MaxMipSize = MemoryBudget - MemoryBudgeted;
-		for (int32 TextureIndex = 0; TextureIndex < StreamingTextures.Num() && !IsAborted(); ++TextureIndex)
+		for (int32 AssetIndex = 0; AssetIndex < StreamingRenderAssets.Num() && !IsAborted(); ++AssetIndex)
 		{
-			FStreamingTexture& StreamingTexture = StreamingTextures[TextureIndex];
-			// Only consider non deleted textures (can change any time).
-			if (!StreamingTexture.Texture) continue;
+			FStreamingRenderAsset& StreamingRenderAsset = StreamingRenderAssets[AssetIndex];
+			// Only consider non deleted textures/meshes (can change any time).
+			if (!StreamingRenderAsset.RenderAsset) continue;
 
-			// Ignore texture that can't drop any mips
-			if (StreamingTexture.BudgetedMips < StreamingTexture.ResidentMips && 
-				StreamingTexture.GetSize(StreamingTexture.BudgetedMips + 1) - StreamingTexture.GetSize(StreamingTexture.BudgetedMips) <= MaxMipSize)
+			// Only consider textures/meshes that won't bust budget nor generate new I/O requests
+			if (StreamingRenderAsset.BudgetedMips < StreamingRenderAsset.ResidentMips &&
+				StreamingRenderAsset.GetSize(StreamingRenderAsset.BudgetedMips + 1) - StreamingRenderAsset.GetSize(StreamingRenderAsset.BudgetedMips) <= MaxMipSize)
 			{
-				PrioritizedTextures.Add(TextureIndex);
+				PrioritizedRenderAssets.Add(AssetIndex);
 			}
 		}
 
-		// Sort texture, having those that should be dropped first.
-		PrioritizedTextures.Sort(FCompareTextureByRetentionPriority(StreamingTextures));
+		// Sort texture/mesh, having those that should be dropped first.
+		PrioritizedRenderAssets.Sort(FCompareRenderAssetByRetentionPriority(StreamingRenderAssets));
 
 		bool bBudgetIsChanging = true;
 		while (MemoryBudgeted < MemoryBudget && bBudgetIsChanging && !IsAborted())
@@ -437,13 +491,13 @@ void FAsyncTextureStreamingTask::UpdateBudgetedMips_Async(int64& MemoryUsed, int
 			bBudgetIsChanging = false;
 
 			// Keep from highest priority first.
-			for (int32 PriorityIndex = 0; PriorityIndex < PrioritizedTextures.Num() && MemoryBudgeted < MemoryBudget && !IsAborted(); ++PriorityIndex)
+			for (int32 PriorityIndex = 0; PriorityIndex < PrioritizedRenderAssets.Num() && MemoryBudgeted < MemoryBudget && !IsAborted(); ++PriorityIndex)
 			{
-				int32 TextureIndex = PrioritizedTextures[PriorityIndex];
-				if (TextureIndex == INDEX_NONE) continue;
+				int32 AssetIndex = PrioritizedRenderAssets[PriorityIndex];
+				if (AssetIndex == INDEX_NONE) continue;
 
-				FStreamingTexture& StreamingTexture = StreamingTextures[TextureIndex];
-				int64 TakenMemory = StreamingTexture.KeepOneMip_Async();
+				FStreamingRenderAsset& StreamingRenderAsset = StreamingRenderAssets[AssetIndex];
+				int64 TakenMemory = StreamingRenderAsset.KeepOneMip_Async();
 
 				if (TakenMemory > 0)
 				{
@@ -454,13 +508,13 @@ void FAsyncTextureStreamingTask::UpdateBudgetedMips_Async(int64& MemoryUsed, int
 					}
 					else // Cancel keeping this mip
 					{
-						StreamingTexture.DropOneMip_Async();
-						PrioritizedTextures[PriorityIndex] = INDEX_NONE; // Don't try this one again.
+						StreamingRenderAsset.DropOneMip_Async();
+						PrioritizedRenderAssets[PriorityIndex] = INDEX_NONE; // Don't try this one again.
 					}
 				}
 				else // No other mips to keep.
 				{
-					PrioritizedTextures[PriorityIndex] = INDEX_NONE; // Don't try this one again.
+					PrioritizedRenderAssets[PriorityIndex] = INDEX_NONE; // Don't try this one again.
 				}
 			}
 		}
@@ -472,83 +526,86 @@ void FAsyncTextureStreamingTask::UpdateBudgetedMips_Async(int64& MemoryUsed, int
 #if !UE_BUILD_SHIPPING
 	if (Settings.DropMips > 0)
 	{
-		for (FStreamingTexture& StreamingTexture : StreamingTextures)
+		for (FStreamingRenderAsset& StreamingRenderAsset : StreamingRenderAssets)
 		{
 			if (IsAborted()) break;
 
 			if (Settings.DropMips == 1)
 			{
-				StreamingTexture.BudgetedMips = FMath::Min<int32>(StreamingTexture.BudgetedMips, StreamingTexture.GetPerfectWantedMips());
+				StreamingRenderAsset.BudgetedMips = FMath::Min<int32>(StreamingRenderAsset.BudgetedMips, StreamingRenderAsset.GetPerfectWantedMips());
 			}
 			else
 			{
-				StreamingTexture.BudgetedMips = FMath::Min<int32>(StreamingTexture.BudgetedMips, StreamingTexture.VisibleWantedMips);
+				StreamingRenderAsset.BudgetedMips = FMath::Min<int32>(StreamingRenderAsset.BudgetedMips, StreamingRenderAsset.VisibleWantedMips);
 			}
 		}
 	}
 #endif
 }
 
-void FAsyncTextureStreamingTask::UpdateLoadAndCancelationRequests_Async(int64 MemoryUsed, int64 TempMemoryUsed)
+void FRenderAssetStreamingMipCalcTask::UpdateLoadAndCancelationRequests_Async(int64 MemoryUsed, int64 TempMemoryUsed)
 {
-	TArray<FStreamingTexture>& StreamingTextures = StreamingManager.StreamingTextures;
-	const FTextureStreamingSettings& Settings = StreamingManager.Settings;
+	TArray<FStreamingRenderAsset>& StreamingRenderAssets = StreamingManager.StreamingRenderAssets;
+	const FRenderAssetStreamingSettings& Settings = StreamingManager.Settings;
 
-	TArray<int32> PrioritizedTextures;
-	PrioritizedTextures.Empty(StreamingTextures.Num());
-	for (int32 TextureIndex = 0; TextureIndex < StreamingTextures.Num() && !IsAborted(); ++TextureIndex)
+	TArray<int32> PrioritizedRenderAssets;
+	PrioritizedRenderAssets.Empty(StreamingRenderAssets.Num());
+	for (int32 AssetIndex = 0; AssetIndex < StreamingRenderAssets.Num() && !IsAborted(); ++AssetIndex)
 	{
-		FStreamingTexture& StreamingTexture = StreamingTextures[TextureIndex];
+		FStreamingRenderAsset& StreamingTexture = StreamingRenderAssets[AssetIndex];
 		if (StreamingTexture.UpdateLoadOrderPriority_Async(Settings.MinMipForSplitRequest))
 		{
-			PrioritizedTextures.Add(TextureIndex);
+			PrioritizedRenderAssets.Add(AssetIndex);
 		}
 	}
-	PrioritizedTextures.Sort(FCompareTextureByLoadOrderPriority(StreamingTextures));
+	PrioritizedRenderAssets.Sort(FCompareRenderAssetByLoadOrderPriority(StreamingRenderAssets));
 
 	LoadRequests.Empty();
 	CancelationRequests.Empty();
 
 	// Now fill in the LoadRequest and CancelationRequests
-	for (int32 PriorityIndex = 0; PriorityIndex < PrioritizedTextures.Num() && !IsAborted(); ++PriorityIndex)
+	for (int32 PriorityIndex = 0; PriorityIndex < PrioritizedRenderAssets.Num() && !IsAborted(); ++PriorityIndex)
 	{
-		int32 TextureIndex = PrioritizedTextures[PriorityIndex];
-		FStreamingTexture& StreamingTexture = StreamingTextures[TextureIndex];
+		int32 AssetIndex = PrioritizedRenderAssets[PriorityIndex];
+		FStreamingRenderAsset& StreamingRenderAsset = StreamingRenderAssets[AssetIndex];
 
 		// If there is a pending update with no cancelation request
-		if (StreamingTexture.bInFlight && StreamingTexture.RequestedMips != StreamingTexture.ResidentMips)
+		if (StreamingRenderAsset.bInFlight && StreamingRenderAsset.RequestedMips != StreamingRenderAsset.ResidentMips)
 		{
 			// If there is a pending load that attempts to load unrequired data (by at least 2 mips), 
 			// or if there is a pending unload that attempts to unload required data, try to cancel it.
-			if (StreamingTexture.RequestedMips > FMath::Max<int32>(StreamingTexture.ResidentMips, StreamingTexture.WantedMips + 1 ) || 
-				StreamingTexture.RequestedMips < FMath::Min<int32>(StreamingTexture.ResidentMips, StreamingTexture.WantedMips ))
+			if (StreamingRenderAsset.RequestedMips > FMath::Max<int32>(StreamingRenderAsset.ResidentMips, StreamingRenderAsset.WantedMips + 1 ) ||
+				StreamingRenderAsset.RequestedMips < FMath::Min<int32>(StreamingRenderAsset.ResidentMips, StreamingRenderAsset.WantedMips ))
 			{
-				CancelationRequests.Add(TextureIndex);
+				CancelationRequests.Add(AssetIndex);
 			}
 		}
-		else if (StreamingTexture.WantedMips < StreamingTexture.ResidentMips && TempMemoryUsed < TempMemoryBudget)
+		else if (StreamingRenderAsset.WantedMips < StreamingRenderAsset.ResidentMips && TempMemoryUsed < TempMemoryBudget)
 		{
-			const int64 TempMemoryRequired = StreamingTexture.GetSize(StreamingTexture.WantedMips);
-			const int64 UsedMemoryRequired = StreamingTexture.GetSize(StreamingTexture.WantedMips) - StreamingTexture.GetSize(StreamingTexture.ResidentMips);
+			const int64 TempMemoryRequired = StreamingRenderAsset.GetSize(StreamingRenderAsset.WantedMips);
+			const int64 UsedMemoryRequired = StreamingRenderAsset.GetSize(StreamingRenderAsset.WantedMips) - StreamingRenderAsset.GetSize(StreamingRenderAsset.ResidentMips);
 
 			// Respect the temporary budget unless this is the first unload request. This allows a single mip update of any size.
 			if (TempMemoryUsed + TempMemoryRequired <= TempMemoryBudget || !LoadRequests.Num())
 			{
-				LoadRequests.Add(TextureIndex);
+				LoadRequests.Add(AssetIndex);
 	
 				MemoryUsed -= UsedMemoryRequired;
 				TempMemoryUsed += TempMemoryRequired;
 			}
 		}
-		else if (StreamingTexture.WantedMips > StreamingTexture.ResidentMips && MemoryUsed < MemoryBudget && TempMemoryUsed < TempMemoryBudget)
+		else if (StreamingRenderAsset.WantedMips > StreamingRenderAsset.ResidentMips
+			&& (MemoryUsed < MemoryBudget || StreamingRenderAsset.NumForcedMips > StreamingRenderAsset.ResidentMips)
+			&& TempMemoryUsed < TempMemoryBudget)
 		{
-			const int64 UsedMemoryRequired = StreamingTexture.GetSize(StreamingTexture.WantedMips) - StreamingTexture.GetSize(StreamingTexture.ResidentMips);
-			const int64 TempMemoryRequired = StreamingTexture.GetSize(StreamingTexture.WantedMips);
+			const int64 UsedMemoryRequired = StreamingRenderAsset.GetSize(StreamingRenderAsset.WantedMips) - StreamingRenderAsset.GetSize(StreamingRenderAsset.ResidentMips);
+			const int64 TempMemoryRequired = StreamingRenderAsset.GetSize(StreamingRenderAsset.WantedMips);
 
 			// Respect the temporary budget unless this is the first load request. This allows a single mip update of any size.
-			if (MemoryUsed + UsedMemoryRequired <= MemoryBudget && (TempMemoryUsed + TempMemoryRequired <= TempMemoryBudget || !LoadRequests.Num()))
+			if ((StreamingRenderAsset.NumForcedMips > StreamingRenderAsset.ResidentMips || MemoryUsed + UsedMemoryRequired <= MemoryBudget)
+				&& (TempMemoryUsed + TempMemoryRequired <= TempMemoryBudget || !LoadRequests.Num()))
 			{
-				LoadRequests.Add(TextureIndex);
+				LoadRequests.Add(AssetIndex);
 	
 				MemoryUsed += UsedMemoryRequired;
 				TempMemoryUsed += TempMemoryRequired;
@@ -557,55 +614,57 @@ void FAsyncTextureStreamingTask::UpdateLoadAndCancelationRequests_Async(int64 Me
 	}
 }
 
-void FAsyncTextureStreamingTask::UpdatePendingStreamingStatus_Async()
+void FRenderAssetStreamingMipCalcTask::UpdatePendingStreamingStatus_Async()
 {
-	TArray<FStreamingTexture>& StreamingTextures = StreamingManager.StreamingTextures;
-	const bool bIsStreamingPaused = StreamingManager.bPauseTextureStreaming;
+	TArray<FStreamingRenderAsset>& StreamingRenderAssets = StreamingManager.StreamingRenderAssets;
+	const bool bIsStreamingPaused = StreamingManager.bPauseRenderAssetStreaming;
 
 	PendingUpdateDirties.Empty();
 
-	for (int32 TextureIndex = 0; TextureIndex < StreamingTextures.Num() && !IsAborted(); ++TextureIndex)
+	for (int32 AssetIndex = 0; AssetIndex < StreamingRenderAssets.Num() && !IsAborted(); ++AssetIndex)
 	{
-		const FStreamingTexture& StreamingTexture = StreamingTextures[TextureIndex];
+		const FStreamingRenderAsset& StreamingTexture = StreamingRenderAssets[AssetIndex];
 		if (StreamingTexture.bHasUpdatePending != StreamingTexture.HasUpdatePending(bIsStreamingPaused, HasAnyView()))
 		{
-			// The texture state are only updated on the gamethread, where we can make sure the UTextre is in sync.
-			PendingUpdateDirties.Add(TextureIndex);
+			// The texture/mesh state are only updated on the gamethread, where we can make sure the UStreamableRenderAsset is in sync.
+			PendingUpdateDirties.Add(AssetIndex);
 		}
 	}
 }
 
-void FAsyncTextureStreamingTask::DoWork()
+void FRenderAssetStreamingMipCalcTask::DoWork()
 {
-	SCOPED_NAMED_EVENT(FAsyncTextureStreamingTask_DoWork, FColor::Turquoise);
-	DECLARE_SCOPE_CYCLE_COUNTER(TEXT("FAsyncTextureStreamingTask::DoWork"), STAT_AsyncTextureStreaming_DoWork, STATGROUP_StreamingDetails);
+	SCOPED_NAMED_EVENT(FRenderAssetStreamingMipCalcTask_DoWork, FColor::Turquoise);
+	DECLARE_SCOPE_CYCLE_COUNTER(TEXT("FRenderAssetStreamingMipCalcTask::DoWork"), STAT_FRenderAssetStreamingMipCalcTask_DoWork, STATGROUP_StreamingDetails);
 
-	// While the async task is runnning, the StreamingTextures are guarantied not to be reallocated.
+	// While the async task is runnning, the StreamingRenderAssets are guarantied not to be reallocated.
 	// 2 things can happen : a texture can be removed, in which case the texture will be set to null
 	// or some members can be updated following calls to UpdateDynamicData().
-	TArray<FStreamingTexture>& StreamingTextures = StreamingManager.StreamingTextures;
-	const FTextureStreamingSettings& Settings = StreamingManager.Settings;
+	TArray<FStreamingRenderAsset>& StreamingRenderAssets = StreamingManager.StreamingRenderAssets;
+	const FRenderAssetStreamingSettings& Settings = StreamingManager.Settings;
+
+	StreamingData.ComputeViewInfoExtras(Settings);
 
 	// Update the distance and size for each bounds.
 	StreamingData.UpdateBoundSizes_Async(Settings);
 	
 	if (StreamingManager.GetAndResetNewFilesHaveLoaded())
 	{
-		for (FStreamingTexture& StreamingTexture : StreamingTextures)
+		for (FStreamingRenderAsset& StreamingRenderAsset : StreamingRenderAssets)
 		{
 			if (IsAborted()) break;
-			StreamingTexture.ClearCachedOptionalMipsState_Async();
+			StreamingRenderAsset.ClearCachedOptionalMipsState_Async();
 		}
 	}
 
-	for (FStreamingTexture& StreamingTexture : StreamingTextures)
+	for (FStreamingRenderAsset& StreamingRenderAsset : StreamingRenderAssets)
 	{
 		if (IsAborted()) break;
 
-		StreamingTexture.UpdateOptionalMipsState_Async();
+		StreamingRenderAsset.UpdateOptionalMipsState_Async();
 		
-		StreamingData.UpdatePerfectWantedMips_Async(StreamingTexture, Settings);
-		StreamingTexture.DynamicBoostFactor = 1.f; // Reset after every computation.
+		StreamingData.UpdatePerfectWantedMips_Async(StreamingRenderAsset, Settings);
+		StreamingRenderAsset.DynamicBoostFactor = 1.f; // Reset after every computation.
 	}
 
 	int64 MemoryUsed, TempMemoryUsed;
@@ -623,14 +682,14 @@ void FAsyncTextureStreamingTask::DoWork()
 	STAT(UpdateStats_Async());
 }
 
-void FAsyncTextureStreamingTask::UpdateStats_Async()
+void FRenderAssetStreamingMipCalcTask::UpdateStats_Async()
 {
 #if STATS
-	FTextureStreamingStats& Stats = StreamingManager.GatheredStats;
-	FTextureStreamingSettings& Settings = StreamingManager.Settings;
-	TArray<FStreamingTexture>& StreamingTextures = StreamingManager.StreamingTextures;
+	FRenderAssetStreamingStats& Stats = StreamingManager.GatheredStats;
+	FRenderAssetStreamingSettings& Settings = StreamingManager.Settings;
+	TArray<FStreamingRenderAsset>& StreamingRenderAssets = StreamingManager.StreamingRenderAssets;
 
-	Stats.TexturePool = PoolSize;
+	Stats.RenderAssetPool = PoolSize;
 	// Stats.StreamingPool = MemoryBudget;
 	Stats.UsedStreamingPool = 0;
 
@@ -653,16 +712,16 @@ void FAsyncTextureStreamingTask::UpdateStats_Async()
 
 	Stats.OverBudget = 0;
 
-	for (FStreamingTexture& StreamingTexture : StreamingTextures)
+	for (FStreamingRenderAsset& StreamingRenderAsset : StreamingRenderAssets)
 	{
 		if (IsAborted()) break;
-		if (!StreamingTexture.Texture) continue;
+		if (!StreamingRenderAsset.RenderAsset) continue;
 
-		const int64 ResidentSize = StreamingTexture.GetSize(StreamingTexture.ResidentMips);
-		const int64 RequiredSize = StreamingTexture.GetSize(StreamingTexture.GetPerfectWantedMips());
-		const int64 BudgetedSize = StreamingTexture.GetSize(StreamingTexture.BudgetedMips);
-		const int64 MaxSize = StreamingTexture.GetSize(StreamingTexture.MaxAllowedMips);
-		const int64 VisibleWantedSize = StreamingTexture.GetSize(StreamingTexture.VisibleWantedMips);
+		const int64 ResidentSize = StreamingRenderAsset.GetSize(StreamingRenderAsset.ResidentMips);
+		const int64 RequiredSize = StreamingRenderAsset.GetSize(StreamingRenderAsset.GetPerfectWantedMips());
+		const int64 BudgetedSize = StreamingRenderAsset.GetSize(StreamingRenderAsset.BudgetedMips);
+		const int64 MaxSize = StreamingRenderAsset.GetSize(StreamingRenderAsset.MaxAllowedMips);
+		const int64 VisibleWantedSize = StreamingRenderAsset.GetSize(StreamingRenderAsset.VisibleWantedMips);
 
 		// How much the streamer would use if there was no limit.
 		Stats.RequiredPool += RequiredSize;
@@ -674,19 +733,19 @@ void FAsyncTextureStreamingTask::UpdateStats_Async()
 		Stats.NonStreamingMips -= ResidentSize;
 
 		// All persistent mip bias bigger than the expected is considered overbudget.
-		const int32 OverBudgetBias = FMath::Max<int32>(0, StreamingTexture.BudgetMipBias - Settings.GlobalMipBias);
-		Stats.OverBudget += StreamingTexture.GetSize(StreamingTexture.MaxAllowedMips + OverBudgetBias) - MaxSize;
+		const int32 OverBudgetBias = FMath::Max<int32>(0, StreamingRenderAsset.BudgetMipBias - Settings.GlobalMipBias);
+		Stats.OverBudget += StreamingRenderAsset.GetSize(StreamingRenderAsset.MaxAllowedMips + OverBudgetBias) - MaxSize;
 
 		const int64 UsedSize = FMath::Min3<int64>(RequiredSize, BudgetedSize, ResidentSize);
 
 		Stats.WantedMips += UsedSize;
 		Stats.CachedMips += FMath::Max<int64>(ResidentSize - UsedSize, 0);
 
-		if (GIsEditor && StreamingTexture.bForceFullyLoadHeuristic)
+		if (GIsEditor && StreamingRenderAsset.bForceFullyLoadHeuristic)
 		{
 			Stats.ForcedMips += UsedSize;
 		}
-		else if (StreamingTexture.bUseUnkownRefHeuristic)
+		else if (StreamingRenderAsset.bUseUnkownRefHeuristic)
 		{
 			Stats.UnkownRefMips += UsedSize;
 		}
@@ -701,7 +760,10 @@ void FAsyncTextureStreamingTask::UpdateStats_Async()
 				Stats.VisibleMips += VisibleWantedSize;
 
 				// Forced mips are not the same as hidden mips as they are loaded because the user wants them absolutly
-				if (StreamingTexture.bForceFullyLoadHeuristic || (StreamingTexture.LODGroup == TEXTUREGROUP_HierarchicalLOD && Settings.HLODStrategy > 0))
+				if (StreamingRenderAsset.bForceFullyLoadHeuristic
+					|| (StreamingRenderAsset.RenderAssetType == FStreamingRenderAsset::AT_Texture
+						&& StreamingRenderAsset.LODGroup == TEXTUREGROUP_HierarchicalLOD
+						&& Settings.HLODStrategy > 0))
 				{
 					Stats.ForcedMips += UsedSize - VisibleWantedSize;
 				}
@@ -712,9 +774,9 @@ void FAsyncTextureStreamingTask::UpdateStats_Async()
 			}
 		}
 
-		if (StreamingTexture.RequestedMips > StreamingTexture.ResidentMips)
+		if (StreamingRenderAsset.RequestedMips > StreamingRenderAsset.ResidentMips)
 		{
-			Stats.PendingRequests += StreamingTexture.GetSize(StreamingTexture.RequestedMips) - ResidentSize;
+			Stats.PendingRequests += StreamingRenderAsset.GetSize(StreamingRenderAsset.RequestedMips) - ResidentSize;
 		}
 	}
 
diff --git a/Engine/Source/Runtime/Engine/Private/Streaming/AsyncTextureStreaming.h b/Engine/Source/Runtime/Engine/Private/Streaming/AsyncTextureStreaming.h
--- a/Engine/Source/Runtime/Engine/Private/Streaming/AsyncTextureStreaming.h
+++ b/Engine/Source/Runtime/Engine/Private/Streaming/AsyncTextureStreaming.h
@@ -13,27 +13,33 @@ AsyncTextureStreaming.h: Definitions of classes used for texture streaming async
 #include "Streaming/StreamingTexture.h"
 #include "Streaming/TextureInstanceView.h"
 
-class FLevelTextureManager;
-class FDynamicTextureInstanceManager;
-struct FStreamingManagerTexture;
+class FLevelRenderAssetManager;
+class FDynamicRenderAssetInstanceManager;
+struct FRenderAssetStreamingManager;
 
 /** Thread-safe helper struct for streaming information. */
-class FAsyncTextureStreamingData
+class FAsyncRenderAssetStreamingData
 {
 public:
 
 	/** Set the data but do as little as possible, called from the game thread. */
-	void Init(TArray<FStreamingViewInfo> InViewInfos, float InWorldTime, TIndirectArray<FLevelTextureManager>& LevelTextureManagers, FDynamicTextureInstanceManager& DynamicComponentManager);
+	void Init(
+		TArray<FStreamingViewInfo> InViewInfos,
+		float InWorldTime,
+		TIndirectArray<FLevelRenderAssetManager>& LevelStaticInstanceManagers,
+		FDynamicRenderAssetInstanceManager& DynamicComponentManager);
+
+	void ComputeViewInfoExtras(const FRenderAssetStreamingSettings& Settings);
 
 	/** Update everything internally so to allow calls to CalcWantedMips */
-	void UpdateBoundSizes_Async(const FTextureStreamingSettings& Settings);
+	void UpdateBoundSizes_Async(const FRenderAssetStreamingSettings& Settings);
 
-	void UpdatePerfectWantedMips_Async(FStreamingTexture& StreamingTexture, const FTextureStreamingSettings& Settings, bool bOutputToLog = false) const;
+	void UpdatePerfectWantedMips_Async(FStreamingRenderAsset& StreamingRenderAsset, const FRenderAssetStreamingSettings& Settings, bool bOutputToLog = false) const;
 
 	uint32 GetAllocatedSize() const { return ViewInfos.GetAllocatedSize() + StaticInstancesViews.GetAllocatedSize(); }
 
-	FORCEINLINE const FTextureInstanceAsyncView& GetDynamicInstancesView() const { return DynamicInstancesView; }
-	FORCEINLINE const TArray<FTextureInstanceAsyncView>& GetStaticInstancesViews() const { return StaticInstancesViews; }
+	FORCEINLINE const FRenderAssetInstanceAsyncView& GetDynamicInstancesView() const { return DynamicInstancesView; }
+	FORCEINLINE const TArray<FRenderAssetInstanceAsyncView>& GetStaticInstancesViews() const { return StaticInstancesViews; }
 	FORCEINLINE const TArray<FStreamingViewInfo>& GetViewInfos() const { return ViewInfos; }
 
 	FORCEINLINE bool HasAnyView() const { return ViewInfos.Num() > 0; }
@@ -42,14 +48,14 @@ public:
 	// This must be done on the gamethread as the refcount are not threadsafe.
 	void ReleaseViews()
 	{
-		DynamicInstancesView = FTextureInstanceAsyncView();
+		DynamicInstancesView = FRenderAssetInstanceAsyncView();
 		StaticInstancesViews.Reset();
 	}
 
 	void OnTaskDone_Async() 
 	{ 
 		DynamicInstancesView.OnTaskDone();
-		for (FTextureInstanceAsyncView& StaticView : StaticInstancesViews)
+		for (FRenderAssetInstanceAsyncView& StaticView : StaticInstancesViews)
 		{
 			StaticView.OnTaskDone();
 		}
@@ -60,14 +66,19 @@ private:
 	/** Cached from FStreamingManagerBase. */
 	TArray<FStreamingViewInfo> ViewInfos;
 	
-	FTextureInstanceAsyncView DynamicInstancesView;
+	/** View related data taking view boost into account */
+	FStreamingViewInfoExtraArray ViewInfoExtras;
+
+	FRenderAssetInstanceAsyncView DynamicInstancesView;
 
 	/** Cached from each ULevel. */
-	TArray<FTextureInstanceAsyncView> StaticInstancesViews;
+	TArray<FRenderAssetInstanceAsyncView> StaticInstancesViews;
 
 	/** Time since last full update. Used to know if something is immediately visible. */
 	float LastUpdateTime;
 
+	float MaxScreenSizeOverAllViews;
+
 	/** Sorted list of all static instances. The sorting is based on MaxLevelTextureScreenSize. */
 	TArray<int32> StaticInstancesViewIndices;
 
@@ -75,43 +86,43 @@ private:
 	TArray<int32> CulledStaticInstancesViewIndices;
 };
 
-struct FCompareTextureByRetentionPriority // Bigger retention priority first.
+struct FCompareRenderAssetByRetentionPriority // Bigger retention priority first.
 {
-	FCompareTextureByRetentionPriority(const TArray<FStreamingTexture>& InStreamingTextures) : StreamingTextures(InStreamingTextures) {}
-	const TArray<FStreamingTexture>& StreamingTextures;
+	FCompareRenderAssetByRetentionPriority(const TArray<FStreamingRenderAsset>& InStreamingRenderAssets) : StreamingRenderAssets(InStreamingRenderAssets) {}
+	const TArray<FStreamingRenderAsset>& StreamingRenderAssets;
 
 	FORCEINLINE bool operator()( int32 IndexA, int32 IndexB ) const
 	{
-		const int32 PrioA = StreamingTextures[IndexA].RetentionPriority;
-		const int32 PrioB = StreamingTextures[IndexB].RetentionPriority;
+		const int32 PrioA = StreamingRenderAssets[IndexA].RetentionPriority;
+		const int32 PrioB = StreamingRenderAssets[IndexB].RetentionPriority;
 		if ( PrioA > PrioB )  return true;
 		if ( PrioA == PrioB ) return IndexA > IndexB;  // Sorting by index so that it gets deterministic.
 		return false;
 	}
 };
 
-struct FCompareTextureByLoadOrderPriority // Bigger load order priority first.
+struct FCompareRenderAssetByLoadOrderPriority // Bigger load order priority first.
 {
-	FCompareTextureByLoadOrderPriority(const TArray<FStreamingTexture>& InStreamingTextures) : StreamingTextures(InStreamingTextures) {}
-	const TArray<FStreamingTexture>& StreamingTextures;
+	FCompareRenderAssetByLoadOrderPriority(const TArray<FStreamingRenderAsset>& InStreamingRenderAssets) : StreamingRenderAssets(InStreamingRenderAssets) {}
+	const TArray<FStreamingRenderAsset>& StreamingRenderAssets;
 
 	FORCEINLINE bool operator()( int32 IndexA, int32 IndexB ) const
 	{
-		const int32 PrioA = StreamingTextures[IndexA].LoadOrderPriority;
-		const int32 PrioB = StreamingTextures[IndexB].LoadOrderPriority;
+		const int32 PrioA = StreamingRenderAssets[IndexA].LoadOrderPriority;
+		const int32 PrioB = StreamingRenderAssets[IndexB].LoadOrderPriority;
 		if ( PrioA > PrioB )  return true;
 		if ( PrioA == PrioB ) return IndexA > IndexB;  // Sorting by index so that it gets deterministic.
 		return false;
 	}
 };
 
-/** Async work class for calculating priorities for all textures. */
+/** Async work for calculating priorities and target number of mips for all textures/meshes. */
 // this could implement a better abandon, but give how it is used, it does that anyway via the abort mechanism
-class FAsyncTextureStreamingTask : public FNonAbandonableTask
+class FRenderAssetStreamingMipCalcTask : public FNonAbandonableTask
 {
 public:
 
-	FAsyncTextureStreamingTask( FStreamingManagerTexture* InStreamingManager )
+	FRenderAssetStreamingMipCalcTask( FRenderAssetStreamingManager* InStreamingManager )
 	:	StreamingManager( *InStreamingManager )
 	,	bAbort( false )
 	{
@@ -138,12 +149,12 @@ public:
 	/** Whether the async work is being aborted. Can be used in conjunction with IsDone() to see if it has finished. */
 	bool IsAborted() const { return bAbort;	}
 
-	/** Returns the resulting priorities, matching the FStreamingManagerTexture::StreamingTextures array. */
+	/** Returns the resulting priorities, matching the FRenderAssetStreamingManager::StreamingRenderAssets array. */
 	const TArray<int32>& GetLoadRequests() const { return LoadRequests; }
 	const TArray<int32>& GetCancelationRequests() const { return CancelationRequests; }
 	const TArray<int32>& GetPendingUpdateDirties() const { return PendingUpdateDirties; }
 
-	FAsyncTextureStreamingData StreamingData;
+	FAsyncRenderAssetStreamingData StreamingData;
 	
 	/** Performs the async work. */
 	void DoWork();
@@ -157,12 +168,12 @@ public:
 
 protected:
 
-	/** Ensures that no temporary streaming boost are active which could interfere with texture streaming bias in undesirable ways. */
-	bool AllowPerTextureMipBiasChanges() const;
+	/** Ensures that no temporary streaming boost are active which could interfere with render asset streaming bias in undesirable ways. */
+	bool AllowPerRenderAssetMipBiasChanges() const;
 
 private:
 
-	friend class FAsyncTask<FAsyncTextureStreamingTask>;
+	friend class FAsyncTask<FRenderAssetStreamingMipCalcTask>;
 
 	void UpdateBudgetedMips_Async(int64& OutMemoryUsed, int64& OutTempMemoryUsed);
 
@@ -174,11 +185,11 @@ private:
 
 	FORCEINLINE TStatId GetStatId() const
 	{
-		RETURN_QUICK_DECLARE_CYCLE_STAT(FAsyncTextureStreamingTask, STATGROUP_ThreadPoolAsyncTasks);
+		RETURN_QUICK_DECLARE_CYCLE_STAT(FRenderAssetStreamingMipCalcTask, STATGROUP_ThreadPoolAsyncTasks);
 	}
 
 	/** Reference to the owning streaming manager, for accessing the thread-safe data. */
-	FStreamingManagerTexture& StreamingManager;
+	FRenderAssetStreamingManager& StreamingManager;
 
 	/** Indices for load/unload requests, sorted by load order. */
 	TArray<int32>	LoadRequests;
@@ -193,7 +204,7 @@ private:
 	/** How much VRAM the hardware has. */
 	int64 TotalGraphicsMemory;
 
-	/** How much gpu resources are currently allocated in the texture pool (all category). */
+	/** How much gpu resources are currently allocated in the texture/mesh pool (all category). */
 	int64 AllocatedMemory; 
 
 	/** Size of the pool once non streaming data is removed and value is stabilized */
@@ -205,7 +216,7 @@ private:
 	/** How much temp memory is allowed (temp memory is taken when changing mip count). */
 	int64 MemoryMargin;
 
-	/** How much memory is available for textures. */
+	/** How much memory is available for textures/meshes. */
 	int64 MemoryBudget;
 
 	/**
diff --git a/Engine/Source/Runtime/Engine/Private/Streaming/DynamicTextureInstanceManager.cpp b/Engine/Source/Runtime/Engine/Private/Streaming/DynamicTextureInstanceManager.cpp
--- a/Engine/Source/Runtime/Engine/Private/Streaming/DynamicTextureInstanceManager.cpp
+++ b/Engine/Source/Runtime/Engine/Private/Streaming/DynamicTextureInstanceManager.cpp
@@ -7,7 +7,7 @@
 #include "Streaming/DynamicTextureInstanceManager.h"
 #include "Components/PrimitiveComponent.h"
 
-void FDynamicTextureInstanceManager::FTasks::SyncResults()
+void FDynamicRenderAssetInstanceManager::FTasks::SyncResults()
 {
 	// Update the bounds first as we want the async view to be fully up-to-date.
 	RefreshFullTask->TryWork(false);
@@ -17,31 +17,31 @@ void FDynamicTextureInstanceManager::FTasks::SyncResults()
 	CreateViewTask->TrySync();
 }
 
-FDynamicTextureInstanceManager::FDynamicTextureInstanceManager()
+FDynamicRenderAssetInstanceManager::FDynamicRenderAssetInstanceManager()
 	: DirtyIndex(0)
 	, PendingDefragSrcBoundIndex(INDEX_NONE)
 	, PendingDefragDstBoundIndex(INDEX_NONE)
 {
 	FTasks& Tasks = StateSync.GetTasks();
-	Tasks.RefreshFullTask = new FRefreshFullTask(TextureInstanceTask::FRefreshFull::FOnWorkDone::CreateLambda([this](int32 InBeginIndex, int32 InEndIndex, const TArray<int32>& SkippedIndices, int32 FirstFreeBound, int32 LastUsedBound){ this->OnRefreshVisibilityDone(InBeginIndex, InEndIndex, SkippedIndices, FirstFreeBound, LastUsedBound); }));
-	Tasks.CreateViewTask = new FCreateViewTask(TextureInstanceTask::FCreateViewWithUninitializedBounds::FOnWorkDone::CreateLambda([this](FTextureInstanceView* InView){ this->OnCreateViewDone(InView); }));
+	Tasks.RefreshFullTask = new FRefreshFullTask(RenderAssetInstanceTask::FRefreshFull::FOnWorkDone::CreateLambda([this](int32 InBeginIndex, int32 InEndIndex, const TArray<int32>& SkippedIndices, int32 FirstFreeBound, int32 LastUsedBound){ this->OnRefreshVisibilityDone(InBeginIndex, InEndIndex, SkippedIndices, FirstFreeBound, LastUsedBound); }));
+	Tasks.CreateViewTask = new FCreateViewTask(RenderAssetInstanceTask::FCreateViewWithUninitializedBounds::FOnWorkDone::CreateLambda([this](FRenderAssetInstanceView* InView){ this->OnCreateViewDone(InView); }));
 }
 
-bool FDynamicTextureInstanceManager::IsReferenced(const UPrimitiveComponent* Component) const
+bool FDynamicRenderAssetInstanceManager::IsReferenced(const UPrimitiveComponent* Component) const
 { 
 	return Component && Component->bAttachedToStreamingManagerAsDynamic && (StateSync.GetState()->HasComponentReferences(Component) || PendingComponents.Contains(Component));
 }
 
-void FDynamicTextureInstanceManager::RegisterTasks(TextureInstanceTask::FDoWorkTask& AsyncTask)
+void FDynamicRenderAssetInstanceManager::RegisterTasks(RenderAssetInstanceTask::FDoWorkTask& AsyncTask)
 {
 	FTasks& Tasks = StateSync.GetTasks();
 	AsyncTask.Add(Tasks.RefreshFullTask.GetReference());
 	AsyncTask.Add(Tasks.CreateViewTask.GetReference());
 }
 
-void FDynamicTextureInstanceManager::IncrementalUpdate(FRemovedTextureArray& RemovedTextures, float Percentage)
+void FDynamicRenderAssetInstanceManager::IncrementalUpdate(FRemovedRenderAssetArray& RemovedRenderAssets, float Percentage)
 {
-	FTextureInstanceState* State = StateSync.SyncAndGetState();
+	FRenderAssetInstanceState* State = StateSync.SyncAndGetState();
 
 	// First try to apply the pending defrag.
 	if (PendingDefragSrcBoundIndex != INDEX_NONE && PendingDefragDstBoundIndex != INDEX_NONE)
@@ -59,7 +59,7 @@ void FDynamicTextureInstanceManager::IncrementalUpdate(FRemovedTextureArray& Rem
 	{
 		check(Component);
 
-		State->RemoveComponent(Component, &RemovedTextures);
+		State->RemoveComponent(Component, &RemovedRenderAssets);
 		Component->bAttachedToStreamingManagerAsDynamic = false;
 		// Re-enable updates now that the component is out of the pending list.
 		Component->bIgnoreStreamingManagerUpdate = false;
@@ -87,23 +87,23 @@ void FDynamicTextureInstanceManager::IncrementalUpdate(FRemovedTextureArray& Rem
 	Refresh(Percentage);
 }
 
-void FDynamicTextureInstanceManager::OnCreateViewDone(FTextureInstanceView* InView)
+void FDynamicRenderAssetInstanceManager::OnCreateViewDone(FRenderAssetInstanceView* InView)
 {
 	// Don't call get state here to prevent recursion as this is a task callback.
-	FTextureInstanceState* State = StateSync.GetStateUnsafe();
+	FRenderAssetInstanceState* State = StateSync.GetStateUnsafe();
 	check(DirtyIndex >= State->NumBounds()); // Must be fully valid or the swap will be destructive.
 
 	// The tasks creates dirty bounds, so after sync, move the valid bounds to the view and mark dirty all current bounds.
 	// The incremental update will refresh the bounds within the next update loop.
-	FTextureInstanceView::SwapData(InView, State);
+	FRenderAssetInstanceView::SwapData(InView, State);
 
-	AsyncView = TRefCountPtr<const FTextureInstanceView>(InView);
+	AsyncView = TRefCountPtr<const FRenderAssetInstanceView>(InView);
 }
 
-void FDynamicTextureInstanceManager::OnRefreshVisibilityDone(int32 BeginIndex, int32 EndIndex, const TArray<int32>& SkippedIndices, int32 FirstFreeBound, int32 LastUsedBound)
+void FDynamicRenderAssetInstanceManager::OnRefreshVisibilityDone(int32 BeginIndex, int32 EndIndex, const TArray<int32>& SkippedIndices, int32 FirstFreeBound, int32 LastUsedBound)
 {
 	// Don't call get state here to prevent recursion as this is a task callback.
-	FTextureInstanceState* State = StateSync.GetStateUnsafe();
+	FRenderAssetInstanceState* State = StateSync.GetStateUnsafe();
 	check(DirtyIndex == BeginIndex);
 
 	for (int32 SkippedIndex : SkippedIndices)
@@ -123,7 +123,7 @@ void FDynamicTextureInstanceManager::OnRefreshVisibilityDone(int32 BeginIndex, i
 	}
 }
 
-void FDynamicTextureInstanceManager::OnPreGarbageCollect(FRemovedTextureArray& RemovedTextures)
+void FDynamicRenderAssetInstanceManager::OnPreGarbageCollect(FRemovedRenderAssetArray& RemovedRenderAssets)
 {
 	for (int32 Index = 0; Index < PendingComponents.Num(); ++Index)
 	{
@@ -131,7 +131,7 @@ void FDynamicTextureInstanceManager::OnPreGarbageCollect(FRemovedTextureArray& R
 		check(Primitive);
 
 		// If the component is not registered anymore, remove it. If it gets registered again it will be reinserted in the pending list.
-		// This allows to remove all unregistered components at once without having to handle each of them in FDynamicTextureInstanceManager::Remove().
+		// This allows to remove all unregistered components at once without having to handle each of them in FDynamicRenderAssetInstanceManager::Remove().
 		// The goal here is to bypass the possibly slow search in PendingComponents.
 		if (!Primitive->IsRegistered() || Primitive->IsPendingKill() || Primitive->HasAnyFlags(RF_BeginDestroyed|RF_FinishDestroyed))
 		{
@@ -141,7 +141,7 @@ void FDynamicTextureInstanceManager::OnPreGarbageCollect(FRemovedTextureArray& R
 
 			if (StateSync.GetState()->HasComponentReferences(Primitive))
 			{
-				StateSync.SyncAndGetState()->RemoveComponent(Primitive, &RemovedTextures);
+				StateSync.SyncAndGetState()->RemoveComponent(Primitive, &RemovedRenderAssets);
 			}
 			Primitive->bAttachedToStreamingManagerAsDynamic = false;
 		}
@@ -150,22 +150,22 @@ void FDynamicTextureInstanceManager::OnPreGarbageCollect(FRemovedTextureArray& R
 
 
 /*-----------------------------------
------- ITextureInstanceManager ------
+------ IRenderAssetInstanceManager ------
 -----------------------------------*/
 
-bool FDynamicTextureInstanceManager::CanManage(const UPrimitiveComponent* Component) const
+bool FDynamicRenderAssetInstanceManager::CanManage(const UPrimitiveComponent* Component) const
 {
 	return Component && !Component->IsPendingKill() && !Component->HasAnyFlags(RF_BeginDestroyed|RF_FinishDestroyed);
 }
 
-void FDynamicTextureInstanceManager::Refresh(float Percentage)
+void FDynamicRenderAssetInstanceManager::Refresh(float Percentage)
 {
 	static bool bUseBackgroupTask = true;
 
-	QUICK_SCOPE_CYCLE_COUNTER(FDynamicTextureInstanceManager_Refresh);
+	QUICK_SCOPE_CYCLE_COUNTER(FDynamicRenderAssetInstanceManager_Refresh);
 
 	// Even if the incremental update does not do any allocation, the tasks must be completed in order to udpate the new dirty arrays - see FTasks::SyncResults.
-	FTextureInstanceState* State = StateSync.SyncAndGetState();
+	FRenderAssetInstanceState* State = StateSync.SyncAndGetState();
 
 	if (DirtyIndex < State->NumBounds())
 	{
@@ -174,7 +174,7 @@ void FDynamicTextureInstanceManager::Refresh(float Percentage)
 	}
 }
 
-EAddComponentResult FDynamicTextureInstanceManager::Add(const UPrimitiveComponent* Component, FStreamingTextureLevelContext& LevelContext, float MaxAllowedUIDensity)
+EAddComponentResult FDynamicRenderAssetInstanceManager::Add(const UPrimitiveComponent* Component, FStreamingTextureLevelContext& LevelContext, float MaxAllowedUIDensity)
 {
 	// Don't cull out primitives with no SceneProxy because they need to be removed first (ex: if the primitive got hidden).
 	if (CanManage(Component))
@@ -202,7 +202,7 @@ EAddComponentResult FDynamicTextureInstanceManager::Add(const UPrimitiveComponen
 	return EAddComponentResult::Fail;
 }
 
-void FDynamicTextureInstanceManager::Remove(const UPrimitiveComponent* Component, FRemovedTextureArray* RemovedTextures)
+void FDynamicRenderAssetInstanceManager::Remove(const UPrimitiveComponent* Component, FRemovedRenderAssetArray* RemovedRenderAssets)
 {
 	check(!Component || Component->IsValidLowLevelFast());
 	if (Component && Component->bAttachedToStreamingManagerAsDynamic)
@@ -213,19 +213,19 @@ void FDynamicTextureInstanceManager::Remove(const UPrimitiveComponent* Component
 		// If the component is used, stop any task possibly indirecting it, and clear references.
 		if (StateSync.GetState()->HasComponentReferences(Component))
 		{
-			StateSync.SyncAndGetState()->RemoveComponent(Component, RemovedTextures);
+			StateSync.SyncAndGetState()->RemoveComponent(Component, RemovedRenderAssets);
 		}
 		Component->bAttachedToStreamingManagerAsDynamic = false;
 	}
 }
 
-void FDynamicTextureInstanceManager::PrepareAsyncView()
+void FDynamicRenderAssetInstanceManager::PrepareAsyncView()
 {
 	// Terminate any pending work as we are about to create a new task.
-	const FTextureInstanceState* State = StateSync.SyncAndGetState();
+	const FRenderAssetInstanceState* State = StateSync.SyncAndGetState();
 
 	// Update any dirty bounds as everything must be up-do-date before making the view.
-	// Bounds can be dirty even after calling FDynamicTextureInstanceManager::IncrementalUpdate as newly added bounds do not update DirtyIndex.
+	// Bounds can be dirty even after calling FDynamicRenderAssetInstanceManager::IncrementalUpdate as newly added bounds do not update DirtyIndex.
 	Refresh(1.f);
 
 	 // Can only release the old view if no one else is referring it, as the refcount is not threadsafe.
@@ -235,19 +235,19 @@ void FDynamicTextureInstanceManager::PrepareAsyncView()
 	AsyncView.SafeRelease();
 }
 
-const FTextureInstanceView* FDynamicTextureInstanceManager::GetAsyncView(bool bCreateIfNull)
+const FRenderAssetInstanceView* FDynamicRenderAssetInstanceManager::GetAsyncView(bool bCreateIfNull)
 {
-	FTextureInstanceState* State = StateSync.SyncAndGetState();
+	FRenderAssetInstanceState* State = StateSync.SyncAndGetState();
 	if (!AsyncView && bCreateIfNull)
 	{
-		AsyncView = FTextureInstanceView::CreateView(State);
+		AsyncView = FRenderAssetInstanceView::CreateView(State);
 	}
 	DirtyIndex = 0; // Force a full refresh!
 	return AsyncView.GetReference();
 }
 
-uint32 FDynamicTextureInstanceManager::GetAllocatedSize() const
+uint32 FDynamicRenderAssetInstanceManager::GetAllocatedSize() const
 { 
-	const FTextureInstanceState* State = StateSync.GetState();
-	return  State ? (sizeof(FTextureInstanceState) + State->GetAllocatedSize()) : 0;
+	const FRenderAssetInstanceState* State = StateSync.GetState();
+	return  State ? (sizeof(FRenderAssetInstanceState) + State->GetAllocatedSize()) : 0;
 }
diff --git a/Engine/Source/Runtime/Engine/Private/Streaming/DynamicTextureInstanceManager.h b/Engine/Source/Runtime/Engine/Private/Streaming/DynamicTextureInstanceManager.h
--- a/Engine/Source/Runtime/Engine/Private/Streaming/DynamicTextureInstanceManager.h
+++ b/Engine/Source/Runtime/Engine/Private/Streaming/DynamicTextureInstanceManager.h
@@ -12,29 +12,29 @@ DynamicTextureInstanceManager.h: Definitions of classes used for texture streami
 #include "ContentStreaming.h"
 
 /** 
- * A texture instance manager to manage dynamic components. 
+ * A texture/mesh instance manager to manage dynamic components. 
  * The async view generated is duplicated so that the state can change freely.
  */
-class FDynamicTextureInstanceManager : public ITextureInstanceManager
+class FDynamicRenderAssetInstanceManager : public IRenderAssetInstanceManager
 {
 public:
 
 	/** Contructor. */
-	FDynamicTextureInstanceManager();
-	~FDynamicTextureInstanceManager() { StateSync.Sync(); }
+	FDynamicRenderAssetInstanceManager();
+	~FDynamicRenderAssetInstanceManager() { StateSync.Sync(); }
 
-	void RegisterTasks(TextureInstanceTask::FDoWorkTask& AsyncTask);
+	void RegisterTasks(RenderAssetInstanceTask::FDoWorkTask& AsyncTask);
 
-	void IncrementalUpdate(FRemovedTextureArray& RemovedTextures, float Percentage);
+	void IncrementalUpdate(FRemovedRenderAssetArray& RemovedRenderAssets, float Percentage);
 
 	// Get all (non removed) components refered by the manager. Debug only.
 	void GetReferencedComponents(TArray<const UPrimitiveComponent*>& Components) { StateSync.SyncAndGetState()->GetReferencedComponents(Components); }
 
 	/** Remove all pending components that are marked for delete. This prevents searching in the pending list for each entry. */
-	void OnPreGarbageCollect(FRemovedTextureArray& RemovedTextures);
+	void OnPreGarbageCollect(FRemovedRenderAssetArray& RemovedRenderAssets);
 
 	/*-----------------------------------
-	------ ITextureInstanceManager ------
+	------ IRenderAssetInstanceManager ------
 	-----------------------------------*/
 
 	/** Return whether this component can be managed by this manager. */
@@ -49,27 +49,27 @@ public:
 	/** Add a component streaming data, the LevelContext gives support for precompiled data. */
 	EAddComponentResult Add(const UPrimitiveComponent* Component, FStreamingTextureLevelContext& LevelContext, float MaxAllowedUIDensity = 0) final override;
 
-	/** Remove a component, the RemoveTextures is the list of textures not referred anymore. */
-	 void Remove(const UPrimitiveComponent* Component, FRemovedTextureArray* RemovedTextures) final override;
+	/** Remove a component, the RemovedRenderAssets is the list of textures or meshes not referred anymore. */
+	 void Remove(const UPrimitiveComponent* Component, FRemovedRenderAssetArray* RemovedRenderAssets) final override;
 
 	/** Notify the manager that an async view will be requested on the next frame. */
 	void PrepareAsyncView() final override;
 
 	/** Return a view of the data that has to be 100% thread safe. The content is allowed to be updated, but not memory must be reallocated. */
-	const FTextureInstanceView* GetAsyncView(bool bCreateIfNull) final override;
+	const FRenderAssetInstanceView* GetAsyncView(bool bCreateIfNull) final override;
 
 	/** Return the size taken for sub-allocation. */
 	uint32 GetAllocatedSize() const final override;
 
 protected:
 
-	void OnCreateViewDone(FTextureInstanceView* InView);
+	void OnCreateViewDone(FRenderAssetInstanceView* InView);
 	void OnRefreshVisibilityDone(int32 BeginIndex, int32 EndIndex, const TArray<int32>& SkippedIndices, int32 FirstFreeBound, int32 LastUsedBound);
 
 private:
 
-	typedef TextureInstanceTask::FCreateViewWithUninitializedBoundsTask FCreateViewTask;
-	typedef TextureInstanceTask::FRefreshFullTask FRefreshFullTask;
+	typedef RenderAssetInstanceTask::FCreateViewWithUninitializedBoundsTask FCreateViewTask;
+	typedef RenderAssetInstanceTask::FRefreshFullTask FRefreshFullTask;
 
 	struct FTasks
 	{
@@ -79,11 +79,11 @@ private:
 		TRefCountPtr<FRefreshFullTask> RefreshFullTask;
 	};
 
-	/** The texture instances. Shared with the async task. */
-	FTextureInstanceStateTaskSync<FTasks> StateSync;
+	/** The texture/mesh instances. Shared with the async task. */
+	FRenderAssetInstanceStateTaskSync<FTasks> StateSync;
 
 	/** A duplicate view for the async streaming task. */
-	TRefCountPtr<const FTextureInstanceView> AsyncView;
+	TRefCountPtr<const FRenderAssetInstanceView> AsyncView;
 
 	/** Ranges from 0 to Bounds4Components.Num(). Used in the incremental update to update bounds and visibility. */
 	int32 DirtyIndex;
diff --git a/Engine/Source/Runtime/Engine/Private/Streaming/LevelTextureManager.cpp b/Engine/Source/Runtime/Engine/Private/Streaming/LevelTextureManager.cpp
--- a/Engine/Source/Runtime/Engine/Private/Streaming/LevelTextureManager.cpp
+++ b/Engine/Source/Runtime/Engine/Private/Streaming/LevelTextureManager.cpp
@@ -9,7 +9,7 @@
 #include "Engine/Texture2D.h"
 #include "Engine/World.h"
 
-FLevelTextureManager::FLevelTextureManager(ULevel* InLevel, TextureInstanceTask::FDoWorkTask& AsyncTask)
+FLevelRenderAssetManager::FLevelRenderAssetManager(ULevel* InLevel, RenderAssetInstanceTask::FDoWorkTask& AsyncTask)
 	: Level(InLevel)
 	, bIsInitialized(false)
 	, StaticInstances(AsyncTask)
@@ -22,7 +22,7 @@ FLevelTextureManager::FLevelTextureManager(ULevel* InLevel, TextureInstanceTask:
 }
 
 
-void FLevelTextureManager::Remove(FRemovedTextureArray* RemovedTextures)
+void FLevelRenderAssetManager::Remove(FRemovedRenderAssetArray* RemovedRenderAssets)
 { 
 	TArray<const UPrimitiveComponent*> ReferencedComponents;
 	StaticInstances.GetReferencedComponents(ReferencedComponents);
@@ -41,12 +41,12 @@ void FLevelTextureManager::Remove(FRemovedTextureArray* RemovedTextures)
 		}
 	}
 
-	// Mark all static textures for removal.
-	if (RemovedTextures)
+	// Mark all static textures/meshes for removal.
+	if (RemovedRenderAssets)
 	{
-		for (auto It = StaticInstances.GetTextureIterator(); It; ++It)
+		for (auto It = StaticInstances.GetRenderAssetIterator(); It; ++It)
 		{
-			RemovedTextures->Push(*It);
+			RemovedRenderAssets->Push(*It);
 		}
 	}
 
@@ -62,7 +62,7 @@ void FLevelTextureManager::Remove(FRemovedTextureArray* RemovedTextures)
 	}
 }
 
-float FLevelTextureManager::GetWorldTime() const
+float FLevelRenderAssetManager::GetWorldTime() const
 {
 	if (Level)
 	{
@@ -81,7 +81,7 @@ float FLevelTextureManager::GetWorldTime() const
 	return 0;
 }
 
-void FLevelTextureManager::SetAsStatic(FDynamicTextureInstanceManager& DynamicComponentManager, const UPrimitiveComponent* Primitive)
+void FLevelRenderAssetManager::SetAsStatic(FDynamicRenderAssetInstanceManager& DynamicComponentManager, const UPrimitiveComponent* Primitive)
 {
 	check(Primitive);
 	Primitive->bAttachedToStreamingManagerAsStatic = true;
@@ -93,7 +93,7 @@ void FLevelTextureManager::SetAsStatic(FDynamicTextureInstanceManager& DynamicCo
 }
 
 
-void FLevelTextureManager::SetAsDynamic(FDynamicTextureInstanceManager& DynamicComponentManager, FStreamingTextureLevelContext& LevelContext, const UPrimitiveComponent* Primitive)
+void FLevelRenderAssetManager::SetAsDynamic(FDynamicRenderAssetInstanceManager& DynamicComponentManager, FStreamingTextureLevelContext& LevelContext, const UPrimitiveComponent* Primitive)
 {
 	check(Primitive);
 	Primitive->bAttachedToStreamingManagerAsStatic = false;
@@ -103,7 +103,7 @@ void FLevelTextureManager::SetAsDynamic(FDynamicTextureInstanceManager& DynamicC
 	}
 }
 
-void FLevelTextureManager::IncrementalBuild(FDynamicTextureInstanceManager& DynamicComponentManager, FStreamingTextureLevelContext& LevelContext, bool bForceCompletion, int64& NumStepsLeft)
+void FLevelRenderAssetManager::IncrementalBuild(FDynamicRenderAssetInstanceManager& DynamicComponentManager, FStreamingTextureLevelContext& LevelContext, bool bForceCompletion, int64& NumStepsLeft)
 {
 	check(Level);
 
@@ -243,7 +243,7 @@ void FLevelTextureManager::IncrementalBuild(FDynamicTextureInstanceManager& Dyna
 	}
 }
 
-bool FLevelTextureManager::NeedsIncrementalBuild(int32 NumStepsLeftForIncrementalBuild) const
+bool FLevelRenderAssetManager::NeedsIncrementalBuild(int32 NumStepsLeftForIncrementalBuild) const
 {
 	check(Level);
 
@@ -261,14 +261,14 @@ bool FLevelTextureManager::NeedsIncrementalBuild(int32 NumStepsLeftForIncrementa
 	}
 }
 
-void FLevelTextureManager::IncrementalUpdate(
-	FDynamicTextureInstanceManager& DynamicComponentManager, 
-	FRemovedTextureArray& RemovedTextures, 
+void FLevelRenderAssetManager::IncrementalUpdate(
+	FDynamicRenderAssetInstanceManager& DynamicComponentManager, 
+	FRemovedRenderAssetArray& RemovedRenderAssets,
 	int64& NumStepsLeftForIncrementalBuild, 
 	float Percentage, 
 	bool bUseDynamicStreaming) 
 {
-	QUICK_SCOPE_CYCLE_COUNTER(FStaticComponentTextureManager_IncrementalUpdate);
+	QUICK_SCOPE_CYCLE_COUNTER(FStaticComponentRenderAssetManager_IncrementalUpdate);
 
 	check(Level);
 
@@ -293,16 +293,16 @@ void FLevelTextureManager::IncrementalUpdate(
 		else if (bIsInitialized)
 		{
 			// Mark all static textures for removal.
-			for (auto It = StaticInstances.GetTextureIterator(); It; ++It)
+			for (auto It = StaticInstances.GetRenderAssetIterator(); It; ++It)
 			{
-				RemovedTextures.Push(*It);
+				RemovedRenderAssets.Push(*It);
 			}
 			bIsInitialized = false;
 		}
 	}
 }
 
-void FLevelTextureManager::NotifyLevelOffset(const FVector& Offset)
+void FLevelRenderAssetManager::NotifyLevelOffset(const FVector& Offset)
 {
 	if (BuildStep == EStaticBuildStep::Done)
 	{
@@ -311,7 +311,7 @@ void FLevelTextureManager::NotifyLevelOffset(const FVector& Offset)
 	}
 }
 
-uint32 FLevelTextureManager::GetAllocatedSize() const
+uint32 FLevelRenderAssetManager::GetAllocatedSize() const
 {
 	return StaticInstances.GetAllocatedSize() + 
 		UnprocessedComponents.GetAllocatedSize() + 
diff --git a/Engine/Source/Runtime/Engine/Private/Streaming/LevelTextureManager.h b/Engine/Source/Runtime/Engine/Private/Streaming/LevelTextureManager.h
--- a/Engine/Source/Runtime/Engine/Private/Streaming/LevelTextureManager.h
+++ b/Engine/Source/Runtime/Engine/Private/Streaming/LevelTextureManager.h
@@ -16,39 +16,39 @@ TextureInstanceManager.h: Definitions of classes used for texture streaming.
 #include "UObject/UObjectHash.h"
 
 // The streaming data of a level.
-class FLevelTextureManager
+class FLevelRenderAssetManager
 {
 public:
 
-	FLevelTextureManager(ULevel* InLevel, TextureInstanceTask::FDoWorkTask& AsyncTask);
+	FLevelRenderAssetManager(ULevel* InLevel, RenderAssetInstanceTask::FDoWorkTask& AsyncTask);
 
 	ULevel* GetLevel() const { return Level; }
 
-	FORCEINLINE bool HasTextureReferences() const { return StaticInstances.HasTextureReferences(); }
+	FORCEINLINE bool HasRenderAssetReferences() const { return StaticInstances.HasRenderAssetReferences(); }
 
-	// Remove the whole level. Optional list of textures referenced
-	void Remove(FRemovedTextureArray* RemovedTextures);
+	// Remove the whole level. Optional list of textures or meshes referenced
+	void Remove(FRemovedRenderAssetArray* RemovedRenderAssets);
 
 	// Invalidate a component reference.
 
 	FORCEINLINE void RemoveActorReferences(const AActor* Actor) {}
 
-	void RemoveComponentReferences(const UPrimitiveComponent* Component, FRemovedTextureArray& RemovedTextures) 
+	void RemoveComponentReferences(const UPrimitiveComponent* Component, FRemovedRenderAssetArray& RemovedRenderAssets)
 	{ 
 		// Check everywhere as the mobility can change in game.
-		StaticInstances.Remove(Component, &RemovedTextures); 
+		StaticInstances.Remove(Component, &RemovedRenderAssets);
 		UnprocessedComponents.RemoveSingleSwap(Component); 
 		PendingComponents.RemoveSingleSwap(Component); 
 	}
 
-	const FStaticTextureInstanceManager& GetStaticInstances() const { return StaticInstances; }
+	const FStaticRenderAssetInstanceManager& GetStaticInstances() const { return StaticInstances; }
 
 	float GetWorldTime() const;
 
-	FORCEINLINE FTextureInstanceAsyncView GetAsyncView() { return FTextureInstanceAsyncView(StaticInstances.GetAsyncView(true)); }
-	FORCEINLINE const FTextureInstanceView* GetRawAsyncView() { return StaticInstances.GetAsyncView(false); }
+	FORCEINLINE FRenderAssetInstanceAsyncView GetAsyncView() { return FRenderAssetInstanceAsyncView(StaticInstances.GetAsyncView(true)); }
+	FORCEINLINE const FRenderAssetInstanceView* GetRawAsyncView() { return StaticInstances.GetAsyncView(false); }
 
-	void IncrementalUpdate(FDynamicTextureInstanceManager& DynamicManager, FRemovedTextureArray& RemovedTextures, int64& NumStepsLeftForIncrementalBuild, float Percentage, bool bUseDynamicStreaming);
+	void IncrementalUpdate(FDynamicRenderAssetInstanceManager& DynamicManager, FRemovedRenderAssetArray& RemovedRenderAssets, int64& NumStepsLeftForIncrementalBuild, float Percentage, bool bUseDynamicStreaming);
 
 	uint32 GetAllocatedSize() const;
 
@@ -62,7 +62,7 @@ private:
 
 	bool bIsInitialized;
 
-	FStaticTextureInstanceManager StaticInstances;
+	FStaticRenderAssetInstanceManager StaticInstances;
 
 	/** Incremental build implementation. */
 
@@ -87,8 +87,8 @@ private:
 	TMap<FGuid, int32> TextureGuidToLevelIndex;
 
 	bool NeedsIncrementalBuild(int32 NumStepsLeftForIncrementalBuild) const;
-	void IncrementalBuild(FDynamicTextureInstanceManager& DynamicComponentManager, FStreamingTextureLevelContext& LevelContext, bool bForceCompletion, int64& NumStepsLeft);
+	void IncrementalBuild(FDynamicRenderAssetInstanceManager& DynamicComponentManager, FStreamingTextureLevelContext& LevelContext, bool bForceCompletion, int64& NumStepsLeft);
 
-	FORCEINLINE_DEBUGGABLE void SetAsStatic(FDynamicTextureInstanceManager& DynamicComponentManager, const UPrimitiveComponent* Primitive);
-	FORCEINLINE_DEBUGGABLE void SetAsDynamic(FDynamicTextureInstanceManager& DynamicComponentManager, FStreamingTextureLevelContext& LevelContext, const UPrimitiveComponent* Primitive);
+	FORCEINLINE_DEBUGGABLE void SetAsStatic(FDynamicRenderAssetInstanceManager& DynamicComponentManager, const UPrimitiveComponent* Primitive);
+	FORCEINLINE_DEBUGGABLE void SetAsDynamic(FDynamicRenderAssetInstanceManager& DynamicComponentManager, FStreamingTextureLevelContext& LevelContext, const UPrimitiveComponent* Primitive);
 };
diff --git a/Engine/Source/Runtime/Engine/Private/Streaming/RenderAssetUpdate.cpp b/Engine/Source/Runtime/Engine/Private/Streaming/RenderAssetUpdate.cpp
new file mode 100644
--- /dev/null
+++ b/Engine/Source/Runtime/Engine/Private/Streaming/RenderAssetUpdate.cpp
@@ -0,0 +1,22 @@
+// Copyright 1998-2018 Epic Games, Inc. All Rights Reserved.
+
+/*=============================================================================
+RenderAssetUpdate.cpp: Base class of helpers to stream in and out texture/mesh LODs.
+=============================================================================*/
+
+#include "Streaming/RenderAssetUpdate.h"
+
+volatile int32 GRenderAssetStreamingSuspendRenderThreadTasks;
+
+void SuspendRenderAssetStreamingRenderTasksInternal()
+{
+	// This doesn't prevent from having a task pushed immediately after as some threads could already be deep in FTexture2DUpdate::PushTask.
+	// This is why GSuspendRenderThreadTasks also gets checked in FTexture2DUpdate::Tick. The goal being to avoid accessing the RHI more
+	// than avoiding pushing new render commands. The reason behind is that some code paths access the RHI outside the render thread.
+	FPlatformAtomics::InterlockedIncrement(&GRenderAssetStreamingSuspendRenderThreadTasks);
+}
+
+void ResumeRenderAssetStreamingRenderTasksInternal()
+{
+	FPlatformAtomics::InterlockedDecrement(&GRenderAssetStreamingSuspendRenderThreadTasks);
+}
