Resolve robomerge conflicts

diff --git a/Engine/Source/Runtime/Engine/Private/Streaming/RenderAssetUpdate.h b/Engine/Source/Runtime/Engine/Private/Streaming/RenderAssetUpdate.h
new file mode 100644
--- /dev/null
+++ b/Engine/Source/Runtime/Engine/Private/Streaming/RenderAssetUpdate.h
@@ -0,0 +1,442 @@
+// Copyright 1998-2018 Epic Games, Inc. All Rights Reserved.
+
+/*=============================================================================
+RenderAssetUpdate.h: Base class of helpers to stream in and out texture/mesh LODs
+=============================================================================*/
+
+#pragma once
+
+#include "CoreMinimal.h"
+#include "HAL/ThreadSafeCounter.h"
+#include "Async/AsyncWork.h"
+#include "Engine/StreamableRenderAsset.h"
+#include "RenderingThread.h"
+
+/** SRA stands for StreamableRenderAsset */
+#define SRA_UPDATE_CALLBACK(FunctionName) [this](const FContext& C){ FunctionName(C); }
+
+/** A counter defining whether render thread tasks should be postponed The callback handling the cancellation of the update (only if the update gets cancelled). */
+extern volatile int32 GRenderAssetStreamingSuspendRenderThreadTasks;
+
+/**
+* This class provides a framework for loading and unloading the texture/mesh LODs.
+* Each thread essentially calls Tick() until the job is done.
+* The object can be safely deleted when IsCompleted() returns true.
+*/
+template <typename FUpdateContext>
+class TRenderAssetUpdate
+{
+public:
+	typedef FUpdateContext FContext;
+
+	/** A thread type used for doing a part of the update process.  */
+	enum EThreadType
+	{
+		TT_None,	// No thread.
+		TT_Render,	// The render thread.
+		TT_Async	// An async work thread.
+	};
+
+	/**  The state of scheduled work for the update process. */
+	enum ETaskState
+	{
+		TS_None,		// Nothing to do.
+		TS_Pending,		// The next task (or update step) is configured, but a callback has not been scheduled yet.
+		TS_Scheduled,	// The next task (or update step) is configured and a callback has been scheduled on either the renderthread or async thread.
+		TS_Locked,		// The object is locked, and no one is allowed to process or look at the next task.
+	};
+
+	TRenderAssetUpdate(UStreamableRenderAsset* InAsset, int32 InRequestedMips)
+		: PendingFirstMip(INDEX_NONE)
+		, RequestedMips(INDEX_NONE)
+		, ScheduledTaskCount(0)
+		, bIsCancelled(false)
+		, TaskState(TS_Locked) // The object is created in the locked state to follow the Tick path
+		, PendingTaskState(TS_None)
+		, TaskThread(TT_None)
+		, CancelationThread(TT_None)
+	{
+		check(InAsset);
+
+		const int32 NonStreamingMipCount = InAsset->GetNumNonStreamingMips();
+		const int32 MaxMipCount = InAsset->GetNumMipsForStreaming();
+		InRequestedMips = FMath::Clamp<int32>(InRequestedMips, NonStreamingMipCount, MaxMipCount);
+
+		if (InRequestedMips > 0 && InRequestedMips != InAsset->GetNumResidentMips() && InAsset->bIsStreamable)
+		{
+			RequestedMips = InRequestedMips;
+			PendingFirstMip = MaxMipCount - RequestedMips;
+		}
+		else // This shouldn't happen but if it does, then the update is canceled
+		{
+			bIsCancelled = true;
+		}
+	}
+
+	~TRenderAssetUpdate()
+	{
+		// Work must be done here because derived destructors have been called now and so derived members are invalid.
+		ensure(ScheduledTaskCount <= 0);
+	}
+
+	/**
+	* Do or schedule any pending work for a given texture.
+	*
+	* @param InAsset - the texture/mesh being updated, this must be the same texture/mesh as the texture/mesh used to create this object.
+	* @param InCurrentThread - the thread from which the tick is being called. Using TT_None ensures that no work will be immediately performed.
+	*/
+	void Tick(UStreamableRenderAsset* InAsset, EThreadType InCurrentThread)
+	{
+		if (TaskState == TS_None || (TaskSynchronization.GetValue() > 0 && InCurrentThread == TT_None))
+		{
+			// Early exit if the task is not ready to execute and we are ticking from non executing thread.
+			// Executing thread must not early exit in order to make sure that tasks are correctly scheduled.
+			// Otherwise, this assumes that the game thread regularly ticks.
+			return;
+		}
+
+		// Acquire the lock if there is work to do and if it is allowed to wait for the lock
+		if (DoConditionalLock(InCurrentThread))
+		{
+			ensure(PendingTaskState == TS_Scheduled || PendingTaskState == TS_Pending);
+
+			// If the task is ready to execute. 
+			// If this is the renderthread and we shouldn't run renderthread task, mark as pending.
+			// This will require a tick from the game thread to reschedule the task.
+			if (TaskSynchronization.GetValue() <= 0 && !(GRenderAssetStreamingSuspendRenderThreadTasks > 0 && InCurrentThread == TT_Render))
+			{
+				FContext Context(InAsset, InCurrentThread);
+
+				// The task params can not change at this point but the bIsCancelled state could change.
+				// To ensure coherency, the cancel state is cached as it affects which thread is relevant.
+				const bool bCachedIsCancelled = bIsCancelled;
+				const EThreadType RelevantThread = !bCachedIsCancelled ? TaskThread : CancelationThread;
+
+				if (RelevantThread == TT_None)
+				{
+					ClearTask();
+				}
+				else if (InCurrentThread == RelevantThread)
+				{
+					FCallback CachedCallback = !bCachedIsCancelled ? TaskCallback : CancelationCallback;
+					ClearTask();
+					CachedCallback(Context); // Valid if the thread is valid.
+				}
+				else if (PendingTaskState != TS_Scheduled || InCurrentThread != TT_None)
+				{
+					// If the task was never scheduled (because synchro was not ready) schedule now.
+					// We also reschedule if this is an executing thread and it end up not being the good thread.
+					// This can happen when a task gets cancelled between scheduling and execution.
+					// We enforce that executing threads either execute or reschedule to prevent a possible stalls
+					// since game thread will not reschedule after the first time.
+					// It's completely safe to schedule several time as task execution can only ever happen once.
+					// we also keep track of how many callbacks are scheduled through ScheduledTaskCount to prevent
+					// deleting the object while another thread is scheduled to access it.
+					ScheduleTick(Context, RelevantThread);
+				}
+				else // Otherwise unlock the task for the executing thread to process it.
+				{
+					PendingTaskState = TS_Scheduled;
+				}
+			}
+			else // If synchro is not ready, mark the task as pending to be scheduled.
+			{
+				PendingTaskState = TS_Pending;
+			}
+
+			DoUnlock();
+		}
+	}
+
+	/** Returns whether the task has finished executing and there is no other thread possibly accessing it. */
+	bool IsCompleted() const
+	{
+		return ScheduledTaskCount <= 0 && TaskState == TS_None;
+	}
+
+	/** Cancel the current update. Will also attempt to cancel pending IO requests, see FTexture2DStreamIn_IO::Abort(). */
+	void Abort()
+	{
+		MarkAsCancelled();
+	}
+
+	/** Returns whether the task was aborted through Abort() or cancelled.  */
+	bool IsCancelled() const
+	{
+		return bIsCancelled;
+	}
+
+	/**
+	* Perform a lock on the object, preventing any other thread from processing a pending task in Tick().
+	* The lock will stall if there was another thread already processing the task data.
+	*/
+	void DoLock()
+	{
+		DECLARE_SCOPE_CYCLE_COUNTER(TEXT("FRenderAssetUpdate::DoLock"), STAT_FRenderAssetUpdate_DoLock, STATGROUP_StreamingDetails);
+
+		// Acquire the lock
+		int32 CachedTaskState = TS_None;
+		do
+		{
+			// Sleep in between retries.
+			if (CachedTaskState != TS_None)
+			{
+				FPlatformProcess::Sleep(0);
+			}
+			// Cache the task state.
+			CachedTaskState = TaskState;
+
+		} while (CachedTaskState == TS_Locked || FPlatformAtomics::InterlockedCompareExchange(&TaskState, TS_Locked, CachedTaskState) != CachedTaskState);
+
+		// If we just acquired the lock, nothing should be in the process.
+		ensure(PendingTaskState == TS_None);
+		PendingTaskState = (ETaskState)CachedTaskState;
+	}
+
+	/** Release any lock on the object, allowing other thread to modify it. */
+	void DoUnlock()
+	{
+		ensure(TaskState == TS_Locked && PendingTaskState != TS_Locked);
+
+		ETaskState CachedPendingTaskState = PendingTaskState;
+
+		// Reset the pending task state first to prevent racing condition that could fail ensure(PendingTaskState == TS_None) in DoLock()
+		PendingTaskState = TS_None;
+		TaskState = CachedPendingTaskState;
+	}
+
+	/** Get the number of requested mips for this update, ignoring cancellation attempts. */
+	int32 GetNumRequestedMips() const
+	{
+		return RequestedMips;
+	}
+
+protected:
+	/** A callback used to perform a task in the update process. Each task must be executed on a specific thread. */
+	typedef TFunction<void(const FContext& Context)> FCallback;
+
+	/** Set the task state as cancelled. This is internally called in Abort() and when any critical conditions are not met when performing the update. */
+	void MarkAsCancelled()
+	{
+		bIsCancelled = true;
+	}
+
+	/**
+	* Defines the next step to be executed. The next step will be executed by calling the callback on the specified thread.
+	* The callback (for both success and cancelation) will only be executed if TaskSynchronization reaches 0.
+	* If all requirements are immediately satisfied when calling the PushTask the relevant callback will be called immediately.
+	*
+	* @param Context - The context defining which texture is being updated and on which thread this is being called.
+	* @param InTaskThread - The thread on which to call the next step of the update, being TaskCallback.
+	* @param InTaskCallback - The callback that will perform the next step of the update.
+	* @param InCancelationThread - The thread on which to call the cancellation of the update (only if the update gets cancelled).
+	* @param InCancelationCallback - The callback handling the cancellation of the update (only if the update gets cancelled).
+	*/
+	void PushTask(const FContext& Context, EThreadType InTaskThread, const FCallback& InTaskCallback, EThreadType InCancelationThread, const FCallback& InCancelationCallback)
+	{
+		// PushTask can only be called by the one thread/callback that is doing the processing. 
+		// This means we don't need to check whether other threads could be trying to push tasks.
+		check(TaskState == TS_Locked);
+		checkSlow((bool)InTaskCallback == (InTaskThread != TT_None));
+		checkSlow((bool)InCancelationCallback == (InCancelationThread != TT_None));
+
+		// To ensure coherency, the cancel state is cached as it affects which thread is relevant.
+		const bool bCachedIsCancelled = bIsCancelled;
+		const EThreadType RelevantThread = !bCachedIsCancelled ? InTaskThread : InCancelationThread;
+
+		// TaskSynchronization is expected to be set before call this.
+		// If the rendering thread is suspended, delay the scheduling until not suspended anymore.
+		const bool bCanExecuteNow = TaskSynchronization.GetValue() <= 0 && !(GRenderAssetStreamingSuspendRenderThreadTasks > 0 && RelevantThread == TT_Render);
+
+		if (RelevantThread == TT_None)
+		{
+			// Nothing to do.
+		}
+		else if (bCanExecuteNow && Context.GetCurrentThread() == RelevantThread)
+		{
+			FCallback CachedCallback = !bCachedIsCancelled ? InTaskCallback : InCancelationCallback;
+			// Not null otherwise relevant thread would be TT_None.
+			CachedCallback(Context);
+		}
+		else
+		{
+			TaskThread = InTaskThread;
+			TaskCallback = InTaskCallback;
+			CancelationThread = InCancelationThread;
+			CancelationCallback = InCancelationCallback;
+
+			if (bCanExecuteNow)
+			{
+				ScheduleTick(Context, RelevantThread);
+			}
+			else
+			{
+				PendingTaskState = TS_Pending;
+			}
+		}
+	}
+
+private:
+	/**
+	* Locks the object only if there is work to do. If the calling thread is has no capability to actually perform any work,
+	* the lock attempt will also fails if the object is already locked. This is used essentially to prevent the gamethread
+	* from being blocked when ticking the update while the update is already being ticked on another thread.
+	*
+	* @param InCurrentThread - the thread trying to lock the object. Passing TS_None may end up preventing to acquired the object.
+	* @return Whether the lock succeeded.
+	*/
+	bool DoConditionalLock(EThreadType InCurrentThread)
+	{
+		DECLARE_SCOPE_CYCLE_COUNTER(TEXT("FRenderAssetUpdate::DoConditionalLock"), STAT_FRenderAssetUpdate_DoConditionalLock, STATGROUP_StreamingDetails);
+
+		// Acquire the lock if there is work to do and if it is allowed to wait for the lock
+		int32 CachedTaskState = TS_None;
+		do
+		{
+			// Sleep in between retries.
+			if (CachedTaskState != TS_None)
+			{
+				FPlatformProcess::Sleep(0);
+			}
+			// Cache the task state.
+			CachedTaskState = TaskState;
+
+			// Return immediately if there is no work to do or if it is locked and we are not on an executing thread.
+			// When the renderthread is the gamethread, don't lock if this is the renderthread to prevent stalling on low priority async tasks.
+			if (CachedTaskState == TS_None || (CachedTaskState == TS_Locked && (InCurrentThread == TT_None || (InCurrentThread == TT_Render && !GIsThreadedRendering))))
+			{
+				return false;
+			}
+		} while (CachedTaskState == TS_Locked || FPlatformAtomics::InterlockedCompareExchange(&TaskState, TS_Locked, CachedTaskState) != CachedTaskState);
+
+		ensure(PendingTaskState == TS_None);
+		PendingTaskState = (ETaskState)CachedTaskState;
+
+		return true;
+	}
+
+	/**
+	* Pushes a callback on either the gamethread or an async work thread.
+	*
+	* @param Context - the context from the caller thread.
+	* @param InThread - the thread on which to schedule a tick.
+	*/
+	void ScheduleTick(const FContext& Context, EThreadType InThread)
+	{
+		// Task can only be scheduled once the synchronization is completed.
+		check(TaskSynchronization.GetValue() <= 0);
+
+		// The pending update needs to be cached because the scheduling can happen in the constructor, before the assignment.
+
+		// When not having many threads, async task should never schedule tasks since they would wake threads with higher priority while still holding the lock.
+		if (Context.GetCurrentThread() == TT_Async && !FApp::ShouldUseThreadingForPerformance())
+		{
+			PendingTaskState = TS_Pending;
+		}
+		else if (InThread == TT_Render)
+		{
+			FPlatformAtomics::InterlockedIncrement(&ScheduledTaskCount);
+			PendingTaskState = TS_Scheduled;
+
+			UStreamableRenderAsset* RenderAsset = Context.GetRenderAsset();
+			TRenderAssetUpdate* CachedPendingUpdate = this;
+			ENQUEUE_RENDER_COMMAND(RenderAssetUpdateCommand)(
+				[RenderAsset, CachedPendingUpdate](FRHICommandListImmediate&)
+			{
+				check(RenderAsset && CachedPendingUpdate);
+
+				// Recompute the context has things might have changed!
+				CachedPendingUpdate->Tick(RenderAsset, TT_Render);
+
+				FPlatformMisc::MemoryBarrier();
+				FPlatformAtomics::InterlockedDecrement(&CachedPendingUpdate->ScheduledTaskCount);
+			});
+		}
+		else // InThread == TT_Async
+		{
+			check(InThread == TT_Async);
+
+			FPlatformAtomics::InterlockedIncrement(&ScheduledTaskCount);
+			PendingTaskState = TS_Scheduled;
+
+			(new FAsyncMipUpdateTask(Context.GetRenderAsset(), this))->StartBackgroundTask();
+		}
+	}
+
+	/** Clears any pending work. */
+	void ClearTask()
+	{
+		// Reset everything so that the callback setup a new callback
+		PendingTaskState = TS_None;
+		TaskThread = TT_None;
+		TaskCallback = nullptr;
+		CancelationThread = TT_None;
+		CancelationCallback = nullptr;
+		TaskSynchronization.Set(0);
+	}
+
+	/** An async task used to call tick on the pending update. */
+	class FMipUpdateTask : public FNonAbandonableTask
+	{
+	public:
+		FMipUpdateTask(UStreamableRenderAsset* InAsset, TRenderAssetUpdate* InPendingUpdate)
+			: RenderAsset(InAsset)
+			, CachedPendingUpdate(InPendingUpdate)
+		{}
+
+		void DoWork()
+		{
+			check(RenderAsset && CachedPendingUpdate);
+
+			// Recompute the context has things might have changed!
+			CachedPendingUpdate->Tick(RenderAsset, TRenderAssetUpdate::TT_Async);
+
+			FPlatformMisc::MemoryBarrier();
+			FPlatformAtomics::InterlockedDecrement(&CachedPendingUpdate->ScheduledTaskCount);
+		}
+
+		FORCEINLINE TStatId GetStatId() const
+		{
+			RETURN_QUICK_DECLARE_CYCLE_STAT(FMipUpdateTask, STATGROUP_ThreadPoolAsyncTasks);
+		}
+
+	protected:
+		UStreamableRenderAsset* RenderAsset;
+		TRenderAssetUpdate* CachedPendingUpdate;
+	};
+
+	/** The async task to update this object, only one can be active at anytime. It just calls Tick(). */
+	typedef FAutoDeleteAsyncTask<FMipUpdateTask> FAsyncMipUpdateTask;
+
+protected:
+	/** The index of mip that will end as being the first mip of the intermediate (future) texture/mesh. */
+	int32 PendingFirstMip;
+	/** The total number of mips of the intermediate (future) texture/mesh. */
+	int32 RequestedMips;
+
+	/** Synchronization used for trigger the task next step execution. */
+	FThreadSafeCounter	TaskSynchronization;
+
+	/** The number of scheduled ticks (and exceptionally other work from FCancelIORequestsTask) from the renderthread and async thread. Used to prevent deleting the object while it could be accessed. */
+	volatile int32 ScheduledTaskCount;
+
+	/** Whether the task has been cancelled because the update could not proceed or because the user called Abort(). */
+	bool bIsCancelled;
+
+	/** The state of the work yet to be performed to complete the update or cancelation. */
+	volatile int32 TaskState;
+	/** The pending state of future work. Used because the object is in locked state (TS_Locked) when being updated, until the call to Unlock() */
+	ETaskState PendingTaskState;
+
+	/** The thread on which to call the next step of the update, being TaskCallback. */
+	EThreadType TaskThread;
+	/** The callback that will perform the next step of the update. */
+	FCallback TaskCallback;
+	/** The thread on which to call the cancellation of the update (only if the update gets cancelled). */
+	EThreadType CancelationThread; // The thread on which the callbacks should be called.
+								   /** The callback handling the cancellation of the update (only if the update gets cancelled). */
+	FCallback CancelationCallback;
+};
+
+void SuspendRenderAssetStreamingRenderTasksInternal();
+void ResumeRenderAssetStreamingRenderTasksInternal();
diff --git a/Engine/Source/Runtime/Engine/Private/Streaming/StaticMeshUpdate.cpp b/Engine/Source/Runtime/Engine/Private/Streaming/StaticMeshUpdate.cpp
new file mode 100644
--- /dev/null
+++ b/Engine/Source/Runtime/Engine/Private/Streaming/StaticMeshUpdate.cpp
@@ -0,0 +1,549 @@
+// Copyright 1998-2018 Epic Games, Inc. All Rights Reserved.
+
+/*=============================================================================
+StaticMeshUpdate.cpp: Helpers to stream in and out static mesh LODs.
+=============================================================================*/
+
+#include "Streaming/StaticMeshUpdate.h"
+#include "RenderUtils.h"
+#include "Containers/ResourceArray.h"
+#include "Streaming/TextureStreamingHelpers.h"
+#include "HAL/PlatformFilemanager.h"
+#include "Serialization/MemoryReader.h"
+
+FStaticMeshUpdateContext::FStaticMeshUpdateContext(UStaticMesh* InMesh, EThreadType InCurrentThread)
+	: Mesh(InMesh)
+	, CurrentThread(InCurrentThread)
+{
+	check(InMesh);
+	checkSlow(InCurrentThread != FStaticMeshUpdate::TT_Render || IsInRenderingThread());
+	RenderData = Mesh->RenderData.Get();
+}
+
+FStaticMeshUpdateContext::FStaticMeshUpdateContext(UStreamableRenderAsset* InMesh, EThreadType InCurrentThread)
+#if UE_BUILD_SHIPPING
+	: FStaticMeshUpdateContext(static_cast<UStaticMesh*>(InMesh), InCurrentThread)
+#else
+	: FStaticMeshUpdateContext(Cast<UStaticMesh>(InMesh), InCurrentThread)
+#endif
+{}
+
+FStaticMeshUpdate::FStaticMeshUpdate(UStaticMesh* InMesh, int32 InRequestedMips)
+	: TRenderAssetUpdate<FStaticMeshUpdateContext>(InMesh, InRequestedMips)
+{
+	if (InMesh->RenderData)
+	{
+		CurrentFirstLODIdx = InMesh->RenderData->CurrentFirstLODIdx;
+		check(CurrentFirstLODIdx >= 0 && CurrentFirstLODIdx < MAX_MESH_LOD_COUNT);
+	}
+	else
+	{
+		RequestedMips = INDEX_NONE;
+		PendingFirstMip = INDEX_NONE;
+		bIsCancelled = true;
+		CurrentFirstLODIdx = INDEX_NONE;
+	}
+}
+
+void FStaticMeshStreamIn::FIntermediateBuffers::CreateFromCPUData_RenderThread(UStaticMesh* Mesh, FStaticMeshLODResources& LODResource)
+{
+	LODResource.ConditionalForce16BitIndexBuffer(GMaxRHIShaderPlatform, Mesh);
+	FStaticMeshVertexBuffers& VBs = LODResource.VertexBuffers;
+	TangentsVertexBuffer = VBs.StaticMeshVertexBuffer.CreateTangentsRHIBuffer_RenderThread();
+	TexCoordVertexBuffer = VBs.StaticMeshVertexBuffer.CreateTexCoordRHIBuffer_RenderThread();
+	PositionVertexBuffer = VBs.PositionVertexBuffer.CreateRHIBuffer_RenderThread();
+	ColorVertexBuffer = VBs.ColorVertexBuffer.CreateRHIBuffer_RenderThread();
+	IndexBuffer = LODResource.IndexBuffer.CreateRHIBuffer_RenderThread();
+	ReversedIndexBuffer = LODResource.ReversedIndexBuffer.CreateRHIBuffer_RenderThread();
+	DepthOnlyIndexBuffer = LODResource.DepthOnlyIndexBuffer.CreateRHIBuffer_RenderThread();
+	ReversedDepthOnlyIndexBuffer = LODResource.ReversedDepthOnlyIndexBuffer.CreateRHIBuffer_RenderThread();
+	WireframeIndexBuffer = LODResource.WireframeIndexBuffer.CreateRHIBuffer_RenderThread();
+	AdjacencyIndexBuffer = LODResource.AdjacencyIndexBuffer.CreateRHIBuffer_RenderThread();
+}
+
+void FStaticMeshStreamIn::FIntermediateBuffers::CreateFromCPUData_Async(UStaticMesh* Mesh, FStaticMeshLODResources& LODResource)
+{
+	LODResource.ConditionalForce16BitIndexBuffer(GMaxRHIShaderPlatform, Mesh);
+	FStaticMeshVertexBuffers& VBs = LODResource.VertexBuffers;
+	TangentsVertexBuffer = VBs.StaticMeshVertexBuffer.CreateTangentsRHIBuffer_Async();
+	TexCoordVertexBuffer = VBs.StaticMeshVertexBuffer.CreateTexCoordRHIBuffer_Async();
+	PositionVertexBuffer = VBs.PositionVertexBuffer.CreateRHIBuffer_Async();
+	ColorVertexBuffer = VBs.ColorVertexBuffer.CreateRHIBuffer_Async();
+	IndexBuffer = LODResource.IndexBuffer.CreateRHIBuffer_Async();
+	ReversedIndexBuffer = LODResource.ReversedIndexBuffer.CreateRHIBuffer_Async();
+	DepthOnlyIndexBuffer = LODResource.DepthOnlyIndexBuffer.CreateRHIBuffer_Async();
+	ReversedDepthOnlyIndexBuffer = LODResource.ReversedDepthOnlyIndexBuffer.CreateRHIBuffer_Async();
+	WireframeIndexBuffer = LODResource.WireframeIndexBuffer.CreateRHIBuffer_Async();
+	AdjacencyIndexBuffer = LODResource.AdjacencyIndexBuffer.CreateRHIBuffer_Async();
+}
+
+void FStaticMeshStreamIn::FIntermediateBuffers::SafeRelease()
+{
+	TangentsVertexBuffer.SafeRelease();
+	TexCoordVertexBuffer.SafeRelease();
+	PositionVertexBuffer.SafeRelease();
+	ColorVertexBuffer.SafeRelease();
+	IndexBuffer.SafeRelease();
+	ReversedIndexBuffer.SafeRelease();
+	DepthOnlyIndexBuffer.SafeRelease();
+	ReversedDepthOnlyIndexBuffer.SafeRelease();
+	WireframeIndexBuffer.SafeRelease();
+	AdjacencyIndexBuffer.SafeRelease();
+}
+
+template <int32 MaxNumUpdates>
+void FStaticMeshStreamIn::FIntermediateBuffers::TransferBuffers(FStaticMeshLODResources& LODResource, TRHIResourceUpdateBatcher<MaxNumUpdates>& Batcher)
+{
+	FStaticMeshVertexBuffers& VBs = LODResource.VertexBuffers;
+	VBs.StaticMeshVertexBuffer.InitRHIForStreaming(TangentsVertexBuffer, TexCoordVertexBuffer, Batcher);
+	VBs.PositionVertexBuffer.InitRHIForStreaming(PositionVertexBuffer, Batcher);
+	VBs.ColorVertexBuffer.InitRHIForStreaming(ColorVertexBuffer, Batcher);
+	LODResource.IndexBuffer.InitRHIForStreaming(IndexBuffer);
+	LODResource.ReversedIndexBuffer.InitRHIForStreaming(ReversedIndexBuffer);
+	LODResource.DepthOnlyIndexBuffer.InitRHIForStreaming(DepthOnlyIndexBuffer);
+	LODResource.ReversedDepthOnlyIndexBuffer.InitRHIForStreaming(ReversedDepthOnlyIndexBuffer);
+	LODResource.WireframeIndexBuffer.InitRHIForStreaming(WireframeIndexBuffer);
+	LODResource.AdjacencyIndexBuffer.InitRHIForStreaming(AdjacencyIndexBuffer);
+	SafeRelease();
+}
+
+void FStaticMeshStreamIn::FIntermediateBuffers::CheckIsNull() const
+{
+	check(!TangentsVertexBuffer
+		&& !TexCoordVertexBuffer
+		&& !PositionVertexBuffer
+		&& !ColorVertexBuffer
+		&& !IndexBuffer
+		&& !ReversedIndexBuffer
+		&& !DepthOnlyIndexBuffer
+		&& !ReversedDepthOnlyIndexBuffer
+		&& !WireframeIndexBuffer
+		&& !AdjacencyIndexBuffer);
+}
+
+FStaticMeshStreamIn::FStaticMeshStreamIn(UStaticMesh* InMesh, int32 InRequestedMips)
+	: FStaticMeshUpdate(InMesh, InRequestedMips)
+{}
+
+FStaticMeshStreamIn::~FStaticMeshStreamIn()
+{
+#if DO_CHECK
+	for (int32 Idx = 0; Idx < MAX_MESH_LOD_COUNT; ++Idx)
+	{
+		IntermediateBuffersArray[Idx].CheckIsNull();
+	}
+#endif
+}
+
+template <bool bRenderThread>
+void FStaticMeshStreamIn::CreateBuffers_Internal(const FContext& Context)
+{
+	UStaticMesh* Mesh = Context.Mesh;
+	FStaticMeshRenderData* RenderData = Context.RenderData;
+	if (!IsCancelled() && Mesh && RenderData)
+	{
+		check(CurrentFirstLODIdx == RenderData->CurrentFirstLODIdx && PendingFirstMip < CurrentFirstLODIdx);
+
+		for (int32 LODIdx = PendingFirstMip; LODIdx < CurrentFirstLODIdx; ++LODIdx)
+		{
+			FStaticMeshLODResources& LODResource = RenderData->LODResources[LODIdx];
+			if (bRenderThread)
+			{
+				IntermediateBuffersArray[LODIdx].CreateFromCPUData_RenderThread(Mesh, LODResource);
+			}
+			else
+			{
+				IntermediateBuffersArray[LODIdx].CreateFromCPUData_Async(Mesh, LODResource);
+			}
+		}
+	}
+}
+
+void FStaticMeshStreamIn::CreateBuffers_RenderThread(const FContext& Context)
+{
+	check(Context.CurrentThread == TT_Render);
+	CreateBuffers_Internal<true>(Context);
+}
+
+void FStaticMeshStreamIn::CreateBuffers_Async(const FContext& Context)
+{
+	check(Context.CurrentThread == TT_Async);
+	CreateBuffers_Internal<false>(Context);
+}
+
+void FStaticMeshStreamIn::DiscardNewLODs(const FContext& Context)
+{
+	FStaticMeshRenderData* RenderData = Context.RenderData;
+	if (RenderData)
+	{
+		check(CurrentFirstLODIdx == RenderData->CurrentFirstLODIdx && PendingFirstMip < CurrentFirstLODIdx);
+
+		for (int32 LODIdx = PendingFirstMip; LODIdx < CurrentFirstLODIdx; ++LODIdx)
+		{
+			FStaticMeshLODResources& LODResource = RenderData->LODResources[LODIdx];
+			LODResource.DiscardCPUData();
+		}
+	}
+}
+
+void FStaticMeshStreamIn::DoFinishUpdate(const FContext& Context)
+{
+	UStaticMesh* Mesh = Context.Mesh;
+	FStaticMeshRenderData* RenderData = Context.RenderData;
+	if (!IsCancelled() && Mesh && RenderData)
+	{
+		check(Context.CurrentThread == TT_Render
+			&& CurrentFirstLODIdx == RenderData->CurrentFirstLODIdx
+			&& PendingFirstMip < CurrentFirstLODIdx);
+		// Use a scope to flush the batcher before updating CurrentFirstLODIdx
+		{
+			constexpr int32 MaxNumUpdatesPerLOD = 4;
+			constexpr int32 MaxNumUpdates = (MAX_STATIC_MESH_LODS - 1) * MaxNumUpdatesPerLOD;
+			TRHIResourceUpdateBatcher<MaxNumUpdates> Batcher;
+
+			for (int32 LODIdx = PendingFirstMip; LODIdx < CurrentFirstLODIdx; ++LODIdx)
+			{
+				FStaticMeshLODResources& LODResource = RenderData->LODResources[LODIdx];
+				LODResource.IncrementMemoryStats();
+				IntermediateBuffersArray[LODIdx].TransferBuffers(LODResource, Batcher);
+			}
+		}
+		RenderData->CurrentFirstLODIdx = PendingFirstMip;
+		Mesh->SetCachedNumResidentLODs(static_cast<uint8>(RenderData->LODResources.Num() - PendingFirstMip));
+	}
+	else
+	{
+		for (int32 LODIdx = PendingFirstMip; LODIdx < CurrentFirstLODIdx; ++LODIdx)
+		{
+			IntermediateBuffersArray[LODIdx].SafeRelease();
+		}
+	}
+}
+
+void FStaticMeshStreamIn::DoCancel(const FContext& Context)
+{
+	// TODO: support streaming CPU data for editor builds
+	if (!GIsEditor)
+	{
+		DiscardNewLODs(Context);
+	}
+	DoFinishUpdate(Context);
+}
+
+FStaticMeshStreamOut::FStaticMeshStreamOut(UStaticMesh* InMesh, int32 InRequestedMips)
+	: FStaticMeshUpdate(InMesh, InRequestedMips)
+{
+	PushTask(FContext(InMesh, TT_None), TT_Render, SRA_UPDATE_CALLBACK(DoReleaseBuffers), TT_None, nullptr);
+}
+
+void FStaticMeshStreamOut::DoReleaseBuffers(const FContext& Context)
+{
+	check(Context.CurrentThread == TT_Render);
+	UStaticMesh* Mesh = Context.Mesh;
+	FStaticMeshRenderData* RenderData = Context.RenderData;
+	if (!IsCancelled() && Mesh && RenderData)
+	{
+		check(CurrentFirstLODIdx == RenderData->CurrentFirstLODIdx && PendingFirstMip > CurrentFirstLODIdx);
+
+		RenderData->CurrentFirstLODIdx = PendingFirstMip;
+		Mesh->SetCachedNumResidentLODs(static_cast<uint8>(RenderData->LODResources.Num() - PendingFirstMip));
+		
+		constexpr int32 MaxNumUpdatesPerLOD = 4;
+		constexpr int32 MaxNumUpdates = (MAX_STATIC_MESH_LODS - 1) * MaxNumUpdatesPerLOD;
+		TRHIResourceUpdateBatcher<MaxNumUpdates> Batcher;
+
+		for (int32 LODIdx = CurrentFirstLODIdx; LODIdx < PendingFirstMip; ++LODIdx)
+		{
+			FStaticMeshLODResources& LODResource = RenderData->LODResources[LODIdx];
+			FStaticMeshVertexBuffers& VBs = LODResource.VertexBuffers;
+			LODResource.DecrementMemoryStats();
+			VBs.StaticMeshVertexBuffer.ReleaseRHIForStreaming(Batcher);
+			VBs.PositionVertexBuffer.ReleaseRHIForStreaming(Batcher);
+			VBs.ColorVertexBuffer.ReleaseRHIForStreaming(Batcher);
+			// Index buffers don't need to update SRV so we can reuse ReleaseRHI
+			LODResource.IndexBuffer.ReleaseRHI();
+			LODResource.ReversedIndexBuffer.ReleaseRHI();
+			LODResource.DepthOnlyIndexBuffer.ReleaseRHI();
+			LODResource.ReversedDepthOnlyIndexBuffer.ReleaseRHI();
+			LODResource.WireframeIndexBuffer.ReleaseRHI();
+			LODResource.AdjacencyIndexBuffer.ReleaseRHI();
+		}
+	}
+}
+
+void FStaticMeshStreamIn_IO::FCancelIORequestsTask::DoWork()
+{
+	check(PendingUpdate);
+	// Acquire the lock of this object in order to cancel any pending IO.
+	// If the object is currently being ticked, wait.
+	PendingUpdate->DoLock();
+	PendingUpdate->CancelIORequest();
+	PendingUpdate->DoUnlock();
+	FPlatformAtomics::InterlockedDecrement(&PendingUpdate->ScheduledTaskCount);
+}
+
+FStaticMeshStreamIn_IO::FStaticMeshStreamIn_IO(UStaticMesh* InMesh, int32 InRequestedMips, bool bHighPrio)
+	: FStaticMeshStreamIn(InMesh, InRequestedMips)
+	, IOFileHandle(nullptr)
+	, IORequest(nullptr)
+	, bHighPrioIORequest(bHighPrio)
+{}
+
+void FStaticMeshStreamIn_IO::Abort()
+{
+	if (!IsCancelled() && !IsCompleted())
+	{
+		FStaticMeshStreamIn::Abort();
+
+		// IO requests can only exist in the lifetime of IOFileHandle.
+		if (IOFileHandle)
+		{
+			// Prevent the update from being considered done before this is finished.
+			// By checking that it was not already cancelled, we make sure this doesn't get called twice.
+			FPlatformAtomics::InterlockedIncrement(&ScheduledTaskCount);
+			(new FAsyncCancelIORequestsTask(this))->StartBackgroundTask();
+		}
+	}
+}
+
+FString FStaticMeshStreamIn_IO::GetIOFilename(const FContext& Context)
+{
+	UStaticMesh* Mesh = Context.Mesh;
+	if (!IsCancelled() && Mesh)
+	{
+		FString Filename;
+		verify(Mesh->GetMipDataFilename(PendingFirstMip, Filename));
+		return Filename;
+	}
+	MarkAsCancelled();
+	return FString();
+}
+
+void FStaticMeshStreamIn_IO::SetAsyncFileCallback(const FContext& Context)
+{
+	AsyncFileCallback = [this, Context](bool bWasCancelled, IAsyncReadRequest* Req)
+	{
+		// At this point task synchronization would hold the number of pending requests.
+		TaskSynchronization.Decrement();
+
+		if (bWasCancelled)
+		{
+			MarkAsCancelled();
+		}
+
+#if !UE_BUILD_SHIPPING
+		// On some platforms the IO is too fast to test cancelation requests timing issues.
+		if (FRenderAssetStreamingSettings::ExtraIOLatency > 0 && TaskSynchronization.GetValue() == 0)
+		{
+			FPlatformProcess::Sleep(FRenderAssetStreamingSettings::ExtraIOLatency * .001f); // Slow down the streaming.
+		}
+#endif
+
+		// The tick here is intended to schedule the success or cancel callback.
+		// Using TT_None ensure gets which could create a dead lock.
+		Tick(Context.Mesh, FStaticMeshUpdate::TT_None);
+	};
+}
+
+void FStaticMeshStreamIn_IO::SetIORequest(const FContext& Context, const FString& IOFilename)
+{
+	if (IsCancelled())
+	{
+		return;
+	}
+	check(!IOFileHandle && PendingFirstMip < CurrentFirstLODIdx);
+	IOFileHandle = FPlatformFileManager::Get().GetPlatformFile().OpenAsyncRead(*IOFilename);
+	FStaticMeshRenderData* RenderData = Context.RenderData;
+	if (IOFileHandle && RenderData)
+	{
+		SetAsyncFileCallback(Context);
+
+		int64 ReadOffset;
+		int64 ReadSize;
+		{
+			FStaticMeshLODResources& FirstLOD = RenderData->LODResources[PendingFirstMip];
+			FStaticMeshLODResources& LastLOD = RenderData->LODResources[CurrentFirstLODIdx - 1];
+			ReadOffset = FirstLOD.OffsetInFile;
+			ReadSize = LastLOD.OffsetInFile + LastLOD.BulkDataSize - ReadOffset;
+#if DO_CHECK
+			int64 TestOffset = ReadOffset;
+			for (int32 LODIdx = PendingFirstMip; LODIdx < CurrentFirstLODIdx; ++LODIdx)
+			{
+				FStaticMeshLODResources& LODResource = RenderData->LODResources[LODIdx];
+				const int64 CurOffset = (int64)LODResource.OffsetInFile;
+				check(ReadOffset <= CurOffset);
+				check(TestOffset == CurOffset);
+				TestOffset += LODResource.BulkDataSize;
+			}
+			check(TestOffset == ReadOffset + ReadSize);
+#endif
+		}
+
+		// Increment as we push the request. If a request complete immediately, then it will call the callback
+		// but that won't do anything because the tick would not try to acquire the lock since it is already locked.
+		TaskSynchronization.Increment();
+		LODData.Empty(ReadSize);
+		LODData.AddUninitialized(ReadSize);
+		IORequest = IOFileHandle->ReadRequest(
+			ReadOffset,
+			ReadSize,
+			bHighPrioIORequest ? AIOP_BelowNormal : AIOP_Low,
+			&AsyncFileCallback,
+			LODData.GetData());
+	}
+	else
+	{
+		MarkAsCancelled();
+	}
+}
+
+void FStaticMeshStreamIn_IO::ClearIORequest(const FContext& Context)
+{
+	if (IOFileHandle)
+	{
+		if (IORequest)
+		{
+			// If clearing requests not yet completed, cancel and wait.
+			if (!IORequest->PollCompletion())
+			{
+				IORequest->Cancel();
+				IORequest->WaitCompletion();
+			}
+			delete IORequest;
+			IORequest = nullptr;
+		}
+
+		delete IOFileHandle;
+		IOFileHandle = nullptr;
+	}
+}
+
+void FStaticMeshStreamIn_IO::SerializeLODData(const FContext& Context)
+{
+	check(!TaskSynchronization.GetValue());
+	UStaticMesh* Mesh = Context.Mesh;
+	FStaticMeshRenderData* RenderData = Context.RenderData;
+	if (!IsCancelled() && Mesh && RenderData)
+	{
+		check(PendingFirstMip < CurrentFirstLODIdx && CurrentFirstLODIdx == RenderData->CurrentFirstLODIdx);
+		FMemoryReader Ar(LODData, true);
+		for (int32 LODIdx = PendingFirstMip; LODIdx < CurrentFirstLODIdx; ++LODIdx)
+		{
+			FStaticMeshLODResources& LODResource = RenderData->LODResources[LODIdx];
+			constexpr uint8 DummyStripFlags = 0;
+			typename FStaticMeshLODResources::FStaticMeshBuffersSize DummyBuffersSize;
+			LODResource.SerializeBuffers(Ar, Mesh, DummyStripFlags, DummyBuffersSize);
+			check(DummyBuffersSize.CalcBuffersSize() == LODResource.BuffersSize);
+		}
+		LODData.Empty();
+	}
+}
+
+void FStaticMeshStreamIn_IO::CancelIORequest()
+{
+	if (IORequest)
+	{
+		// Calling cancel will trigger the SetAsyncFileCallback() which will also try a tick but will fail.
+		IORequest->Cancel();
+	}
+}
+
+template <bool bRenderThread>
+TStaticMeshStreamIn_IO<bRenderThread>::TStaticMeshStreamIn_IO(UStaticMesh* InMesh, int32 InRequestedMips, bool bHighPrio)
+	: FStaticMeshStreamIn_IO(InMesh, InRequestedMips, bHighPrio)
+{
+	PushTask(FContext(InMesh, TT_None), TT_Async, SRA_UPDATE_CALLBACK(DoInitiateIO), TT_None, nullptr);
+}
+
+template <bool bRenderThread>
+void TStaticMeshStreamIn_IO<bRenderThread>::DoInitiateIO(const FContext& Context)
+{
+	check(Context.CurrentThread == TT_Async);
+	const FString IOFilename = GetIOFilename(Context);
+	SetIORequest(Context, IOFilename);
+	PushTask(Context, TT_Async, SRA_UPDATE_CALLBACK(DoSerializeLODData), TT_Async, SRA_UPDATE_CALLBACK(DoCancelIO));
+}
+
+template <bool bRenderThread>
+void TStaticMeshStreamIn_IO<bRenderThread>::DoSerializeLODData(const FContext& Context)
+{
+	check(Context.CurrentThread == TT_Async);
+	SerializeLODData(Context);
+	ClearIORequest(Context);
+	const EThreadType TThread = bRenderThread ? TT_Render : TT_Async;
+	const EThreadType CThread = (EThreadType)Context.CurrentThread;
+	PushTask(Context, TThread, SRA_UPDATE_CALLBACK(DoCreateBuffers), CThread, SRA_UPDATE_CALLBACK(DoCancel));
+}
+
+template <bool bRenderThread>
+void TStaticMeshStreamIn_IO<bRenderThread>::DoCreateBuffers(const FContext& Context)
+{
+	if (bRenderThread)
+	{
+		CreateBuffers_RenderThread(Context);
+	}
+	else
+	{
+		CreateBuffers_Async(Context);
+	}
+	check(!TaskSynchronization.GetValue());
+	PushTask(Context, TT_Render, SRA_UPDATE_CALLBACK(DoFinishUpdate), (EThreadType)Context.CurrentThread, SRA_UPDATE_CALLBACK(DoCancel));
+}
+
+template <bool bRenderThread>
+void TStaticMeshStreamIn_IO<bRenderThread>::DoCancelIO(const FContext& Context)
+{
+	ClearIORequest(Context);
+	PushTask(Context, TT_None, nullptr, (EThreadType)Context.CurrentThread, SRA_UPDATE_CALLBACK(DoCancel));
+}
+
+template class TStaticMeshStreamIn_IO<true>;
+template class TStaticMeshStreamIn_IO<false>;
+
+#if WITH_EDITOR
+FStaticMeshStreamIn_DDC::FStaticMeshStreamIn_DDC(UStaticMesh* InMesh, int32 InRequestedMips)
+	: FStaticMeshStreamIn(InMesh, InRequestedMips)
+	, bDerivedDataInvalid(false)
+{}
+
+void FStaticMeshStreamIn_DDC::LoadNewLODsFromDDC(const FContext& Context)
+{
+	check(Context.CurrentThread == TT_Async);
+	// TODO: support streaming CPU data for editor builds
+}
+
+template <bool bRenderThread>
+TStaticMeshStreamIn_DDC<bRenderThread>::TStaticMeshStreamIn_DDC(UStaticMesh* InMesh, int32 InRequestedMips)
+	: FStaticMeshStreamIn_DDC(InMesh, InRequestedMips)
+{
+	PushTask(FContext(InMesh, TT_None), TT_Async, SRA_UPDATE_CALLBACK(DoLoadNewLODsFromDDC), TT_None, nullptr);
+}
+
+template <bool bRenderThread>
+void TStaticMeshStreamIn_DDC<bRenderThread>::DoLoadNewLODsFromDDC(const FContext& Context)
+{
+	LoadNewLODsFromDDC(Context);
+	check(!TaskSynchronization.GetValue());
+	const EThreadType TThread = bRenderThread ? TT_Render : TT_Async;
+	const EThreadType CThread = (EThreadType)Context.CurrentThread;
+	PushTask(Context, TThread, SRA_UPDATE_CALLBACK(DoCreateBuffers), CThread, SRA_UPDATE_CALLBACK(DoCancel));
+}
+
+template <bool bRenderThread>
+void TStaticMeshStreamIn_DDC<bRenderThread>::DoCreateBuffers(const FContext& Context)
+{
+	if (bRenderThread)
+	{
+		CreateBuffers_RenderThread(Context);
+	}
+	else
+	{
+		CreateBuffers_Async(Context);
+	}
+	check(!TaskSynchronization.GetValue());
+	PushTask(Context, TT_Render, SRA_UPDATE_CALLBACK(DoFinishUpdate), (EThreadType)Context.CurrentThread, SRA_UPDATE_CALLBACK(DoCancel));
+}
+
+template class TStaticMeshStreamIn_DDC<true>;
+template class TStaticMeshStreamIn_DDC<false>;
+#endif
diff --git a/Engine/Source/Runtime/Engine/Private/Streaming/StaticMeshUpdate.h b/Engine/Source/Runtime/Engine/Private/Streaming/StaticMeshUpdate.h
new file mode 100644
--- /dev/null
+++ b/Engine/Source/Runtime/Engine/Private/Streaming/StaticMeshUpdate.h
@@ -0,0 +1,250 @@
+// Copyright 1998-2018 Epic Games, Inc. All Rights Reserved.
+
+/*=============================================================================
+StaticMeshUpdate.h: Helpers to stream in and out static mesh LODs.
+=============================================================================*/
+
+#pragma once
+
+#include "CoreMinimal.h"
+#include "Streaming/RenderAssetUpdate.h"
+#include "Engine/StaticMesh.h"
+#include "Async/AsyncFileHandle.h"
+
+/**
+* A context used to update or proceed with the next update step.
+* The mesh and render data references could be stored in the update object
+* but are currently kept outside to avoid lifetime management within the object.
+*/
+struct FStaticMeshUpdateContext
+{
+	typedef int32 EThreadType;
+
+	FStaticMeshUpdateContext(UStaticMesh* InMesh, EThreadType InCurrentThread);
+
+	FStaticMeshUpdateContext(UStreamableRenderAsset* InMesh, EThreadType InCurrentThread);
+
+	UStreamableRenderAsset* GetRenderAsset() const
+	{
+		return Mesh;
+	}
+
+	EThreadType GetCurrentThread() const
+	{
+		return CurrentThread;
+	}
+
+	/** The mesh to update, this must be the same one as the one used when creating the FStaticMeshUpdate object. */
+	UStaticMesh* Mesh;
+	/** The current render data of this mesh. */
+	FStaticMeshRenderData* RenderData;
+	/** The thread on which the context was created. */
+	EThreadType CurrentThread;
+};
+
+/**
+* This class provides a framework for loading and unloading the LODs of static meshes.
+* Each thread essentially calls Tick() until the job is done.
+* The object can be safely deleted when IsCompleted() returns true.
+*/
+class FStaticMeshUpdate : public TRenderAssetUpdate<FStaticMeshUpdateContext>
+{
+public:
+	FStaticMeshUpdate(UStaticMesh* InMesh, int32 InRequestedMips);
+
+	virtual ~FStaticMeshUpdate() {}
+
+	virtual void Abort()
+	{
+		TRenderAssetUpdate<FStaticMeshUpdateContext>::Abort();
+	}
+
+#if WITH_EDITOR
+	/** Returns whether DDC of this texture needs to be regenerated.  */
+	virtual bool DDCIsInvalid() const { return false; }
+#endif
+
+protected:
+	/** Cached index of current first LOD that will be replaced by PendingFirstMip */
+	int32 CurrentFirstLODIdx;
+};
+
+class FStaticMeshStreamIn : public FStaticMeshUpdate
+{
+public:
+	FStaticMeshStreamIn(UStaticMesh* InMesh, int32 InRequestedMips);
+
+	virtual ~FStaticMeshStreamIn();
+
+protected:
+	/** Correspond to the buffers in FStaticMeshLODResources */
+	struct FIntermediateBuffers
+	{
+		FVertexBufferRHIRef TangentsVertexBuffer;
+		FVertexBufferRHIRef TexCoordVertexBuffer;
+		FVertexBufferRHIRef PositionVertexBuffer;
+		FVertexBufferRHIRef ColorVertexBuffer;
+		FIndexBufferRHIRef IndexBuffer;
+		FIndexBufferRHIRef ReversedIndexBuffer;
+		FIndexBufferRHIRef DepthOnlyIndexBuffer;
+		FIndexBufferRHIRef ReversedDepthOnlyIndexBuffer;
+		FIndexBufferRHIRef WireframeIndexBuffer;
+		FIndexBufferRHIRef AdjacencyIndexBuffer;
+
+		void CreateFromCPUData_RenderThread(UStaticMesh* Mesh, FStaticMeshLODResources& LODResource);
+		void CreateFromCPUData_Async(UStaticMesh* Mesh, FStaticMeshLODResources& LODResource);
+
+		void SafeRelease();
+
+		/** Transfer ownership of buffers to a LOD resource */
+		template <int32 MaxNumUpdates>
+		void TransferBuffers(FStaticMeshLODResources& LODResource, TRHIResourceUpdateBatcher<MaxNumUpdates>& Batcher);
+
+		void CheckIsNull() const;
+	};
+
+	/** Create buffers with new LOD data on render or pooled thread */
+	void CreateBuffers_RenderThread(const FContext& Context);
+	void CreateBuffers_Async(const FContext& Context);
+
+	/** Discard newly streamed-in CPU data */
+	void DiscardNewLODs(const FContext& Context);
+
+	/** Apply the new buffers (if not cancelled) and finish the update process. When cancelled, the intermediate buffers simply gets discarded. */
+	void DoFinishUpdate(const FContext& Context);
+
+	/** Discard streamed-in CPU data and intermediate RHI buffers */
+	void DoCancel(const FContext& Context);
+
+	/** The intermediate buffers created in the update process. */
+	FIntermediateBuffers IntermediateBuffersArray[MAX_MESH_LOD_COUNT];
+
+private:
+	template <bool bRenderThread>
+	void CreateBuffers_Internal(const FContext& Context);
+};
+
+class FStaticMeshStreamOut : public FStaticMeshUpdate
+{
+public:
+	FStaticMeshStreamOut(UStaticMesh* InMesh, int32 InRequestedMips);
+
+	virtual ~FStaticMeshStreamOut() {}
+
+private:
+	/** Release RHI buffers and update SRVs */
+	void DoReleaseBuffers(const FContext& Context);
+};
+
+class FStaticMeshStreamIn_IO : public FStaticMeshStreamIn
+{
+public:
+	FStaticMeshStreamIn_IO(UStaticMesh* InMesh, int32 InRequestedMips, bool bHighPrio);
+
+	virtual ~FStaticMeshStreamIn_IO() {}
+
+	virtual void Abort() override;
+
+protected:
+	class FCancelIORequestsTask : public FNonAbandonableTask
+	{
+	public:
+		FCancelIORequestsTask(FStaticMeshStreamIn_IO* InPendingUpdate)
+			: PendingUpdate(InPendingUpdate)
+		{}
+
+		void DoWork();
+
+		FORCEINLINE TStatId GetStatId() const
+		{
+			RETURN_QUICK_DECLARE_CYCLE_STAT(FCancelIORequestsTask_StaticMesh, STATGROUP_ThreadPoolAsyncTasks);
+		}
+
+	private:
+		FStaticMeshStreamIn_IO* PendingUpdate;
+	};
+
+	typedef FAutoDeleteAsyncTask<FCancelIORequestsTask> FAsyncCancelIORequestsTask;
+	friend class FCancelIORequestsTask;
+
+	/** Figure out the full name of the .bulk file */
+	FString GetIOFilename(const FContext& Context);
+
+	/** Set a callback called when IORequest is completed or cancelled */
+	void SetAsyncFileCallback(const FContext& Context);
+
+	/** Create a new async IO request to read in LOD data */
+	void SetIORequest(const FContext& Context, const FString& IOFilename);
+
+	/** Release IORequest and IOFileHandle. IORequest will be cancelled if still inflight */
+	void ClearIORequest(const FContext& Context);
+
+	/** Serialize data of new LODs to corresponding FStaticMeshLODResources */
+	void SerializeLODData(const FContext& Context);
+
+	/** Called by FAsyncCancelIORequestsTask to cancel inflight IO request if any */
+	void CancelIORequest();
+
+	TArray<uint8> LODData;
+	IAsyncReadFileHandle* IOFileHandle;
+	IAsyncReadRequest* IORequest;
+	FAsyncFileCallBack AsyncFileCallback;
+	bool bHighPrioIORequest;
+};
+
+template <bool bRenderThread>
+class TStaticMeshStreamIn_IO : public FStaticMeshStreamIn_IO
+{
+public:
+	TStaticMeshStreamIn_IO(UStaticMesh* InMesh, int32 InRequestedMips, bool bHighPrio);
+
+	virtual ~TStaticMeshStreamIn_IO() {}
+
+protected:
+	void DoInitiateIO(const FContext& Context);
+
+	void DoSerializeLODData(const FContext& Context);
+
+	void DoCreateBuffers(const FContext& Context);
+
+	void DoCancelIO(const FContext& Context);
+};
+
+typedef TStaticMeshStreamIn_IO<true> FStaticMeshStreamIn_IO_RenderThread;
+typedef TStaticMeshStreamIn_IO<false> FStaticMeshStreamIn_IO_Async;
+
+#if WITH_EDITOR
+class FStaticMeshStreamIn_DDC : public FStaticMeshStreamIn
+{
+public:
+	FStaticMeshStreamIn_DDC(UStaticMesh* InMesh, int32 InRequestedMips);
+
+	virtual ~FStaticMeshStreamIn_DDC() {}
+
+	virtual bool DDCIsInvalid() const override { return bDerivedDataInvalid; }
+
+protected:
+	void LoadNewLODsFromDDC(const FContext& Context);
+
+	bool bDerivedDataInvalid;
+};
+
+template <bool bRenderThread>
+class TStaticMeshStreamIn_DDC : public FStaticMeshStreamIn_DDC
+{
+public:
+	TStaticMeshStreamIn_DDC(UStaticMesh* InMesh, int32 InRequestedMips);
+
+	virtual ~TStaticMeshStreamIn_DDC() {}
+
+private:
+	/** Load new LOD buffers from DDC and queue a task to create RHI buffers on RT */
+	void DoLoadNewLODsFromDDC(const FContext& Context);
+
+	/** Create RHI buffers for newly streamed-in LODs and queue a task to rename references on RT */
+	void DoCreateBuffers(const FContext& Context);
+};
+
+typedef TStaticMeshStreamIn_DDC<true> FStaticMeshStreamIn_DDC_RenderThread;
+typedef TStaticMeshStreamIn_DDC<false> FStaticMeshStreamIn_DDC_Async;
+#endif
\ No newline at end of file
diff --git a/Engine/Source/Runtime/Engine/Private/Streaming/StaticTextureInstanceManager.cpp b/Engine/Source/Runtime/Engine/Private/Streaming/StaticTextureInstanceManager.cpp
--- a/Engine/Source/Runtime/Engine/Private/Streaming/StaticTextureInstanceManager.cpp
+++ b/Engine/Source/Runtime/Engine/Private/Streaming/StaticTextureInstanceManager.cpp
@@ -7,7 +7,7 @@
 #include "Streaming/StaticTextureInstanceManager.h"
 #include "Components/PrimitiveComponent.h"
 
-void FStaticTextureInstanceManager::FTasks::SyncResults()
+void FStaticRenderAssetInstanceManager::FTasks::SyncResults()
 {
 	RefreshVisibilityTask->TryWork(false);
 	NormalizeLightmapTexelFactorTask->TryWork(false);
@@ -17,12 +17,12 @@ void FStaticTextureInstanceManager::FTasks::SyncResults()
 
 }
 
-FStaticTextureInstanceManager::FStaticTextureInstanceManager(TextureInstanceTask::FDoWorkTask& AsyncTask)
+FStaticRenderAssetInstanceManager::FStaticRenderAssetInstanceManager(RenderAssetInstanceTask::FDoWorkTask& AsyncTask)
 	: DirtyIndex(0)
 {
 	FTasks& Tasks = StateSync.GetTasks();
 
-	Tasks.RefreshVisibilityTask = new FRefreshVisibilityTask(TextureInstanceTask::FRefreshVisibility::FOnWorkDone::CreateLambda([this](int32 InBeginIndex, int32 InEndIndex){ this->OnRefreshVisibilityDone(InBeginIndex, InEndIndex); }));
+	Tasks.RefreshVisibilityTask = new FRefreshVisibilityTask(RenderAssetInstanceTask::FRefreshVisibility::FOnWorkDone::CreateLambda([this](int32 InBeginIndex, int32 InEndIndex){ this->OnRefreshVisibilityDone(InBeginIndex, InEndIndex); }));
 	AsyncTask.Add(Tasks.RefreshVisibilityTask.GetReference());
 
 	Tasks.NormalizeLightmapTexelFactorTask = new FNormalizeLightmapTexelFactorTask();
@@ -30,11 +30,11 @@ FStaticTextureInstanceManager::FStaticTextureInstanceManager(TextureInstanceTask
 
 }
 
-void FStaticTextureInstanceManager::NormalizeLightmapTexelFactor()
+void FStaticRenderAssetInstanceManager::NormalizeLightmapTexelFactor()
 {
 	if (!AsyncView)
 	{
-		FTextureInstanceState* State = StateSync.SyncAndGetState();
+		FRenderAssetInstanceState* State = StateSync.SyncAndGetState();
 		if (State->NumBounds() > 0)
 		{
 			StateSync.GetTasks().NormalizeLightmapTexelFactorTask->Init(State);
@@ -42,7 +42,7 @@ void FStaticTextureInstanceManager::NormalizeLightmapTexelFactor()
 	}
 }
 
-void FStaticTextureInstanceManager::OnRefreshVisibilityDone(int32 InBeginIndex, int32 InEndIndex)
+void FStaticRenderAssetInstanceManager::OnRefreshVisibilityDone(int32 InBeginIndex, int32 InEndIndex)
 {
 	// Make sure there are no wholes between the DirtyIndex and the first updated index.
 	if (InBeginIndex <= DirtyIndex)
@@ -52,10 +52,10 @@ void FStaticTextureInstanceManager::OnRefreshVisibilityDone(int32 InBeginIndex,
 }
 
 /*-----------------------------------
------- ITextureInstanceManager ------
+------ IRenderAssetInstanceManager ------
 -----------------------------------*/
 
-bool FStaticTextureInstanceManager::CanManage(const UPrimitiveComponent* Component) const
+bool FStaticRenderAssetInstanceManager::CanManage(const UPrimitiveComponent* Component) const
 {
 	// This manager only manages static components from static actors.
 	// Note that once the view has been shared, no other modifications are allowed.
@@ -69,11 +69,11 @@ bool FStaticTextureInstanceManager::CanManage(const UPrimitiveComponent* Compone
 	return false;
 }
 
-void FStaticTextureInstanceManager::Refresh(float Percentage)
+void FStaticRenderAssetInstanceManager::Refresh(float Percentage)
 {
 	// Since this is only managing static components, only visibility needs to be refreshed.
 
-	FTextureInstanceState* State = StateSync.SyncAndGetState();
+	FRenderAssetInstanceState* State = StateSync.SyncAndGetState();
 	if (DirtyIndex < State->NumBounds())
 	{
 		const int32 EndIndex = FMath::Min(State->NumBounds(), DirtyIndex + FMath::CeilToInt((float)State->NumBounds() * Percentage));
@@ -81,19 +81,19 @@ void FStaticTextureInstanceManager::Refresh(float Percentage)
 	}
 }
 
-EAddComponentResult FStaticTextureInstanceManager::Add(const UPrimitiveComponent* Component, FStreamingTextureLevelContext& LevelContext, float MaxAllowedUIDensity)
+EAddComponentResult FStaticRenderAssetInstanceManager::Add(const UPrimitiveComponent* Component, FStreamingTextureLevelContext& LevelContext, float MaxAllowedUIDensity)
 {
 	if (!AsyncView)
 	{
-		FTextureInstanceState* State = StateSync.SyncAndGetState();
+		FRenderAssetInstanceState* State = StateSync.SyncAndGetState();
 		return State->AddComponent(Component, LevelContext, MaxAllowedUIDensity);
 	}
 	return EAddComponentResult::Fail;
 }
 
-void FStaticTextureInstanceManager::Remove(const UPrimitiveComponent* Component, FRemovedTextureArray* RemovedTextures)
+void FStaticRenderAssetInstanceManager::Remove(const UPrimitiveComponent* Component, FRemovedRenderAssetArray* RemovedRenderAssets)
 {
-	FTextureInstanceState* State = StateSync.SyncAndGetState();
+	FRenderAssetInstanceState* State = StateSync.SyncAndGetState();
 	// If the view is shared, we are limited to clearing the references (no realloc)
 	if (AsyncView)
 	{
@@ -101,13 +101,13 @@ void FStaticTextureInstanceManager::Remove(const UPrimitiveComponent* Component,
 	}
 	else // Otherwise it can be cleaned properly
 	{
-		State->RemoveComponent(Component, RemovedTextures);
+		State->RemoveComponent(Component, RemovedRenderAssets);
 	}
 }
 
-const FTextureInstanceView* FStaticTextureInstanceManager::GetAsyncView(bool bCreateIfNull)
+const FRenderAssetInstanceView* FStaticRenderAssetInstanceManager::GetAsyncView(bool bCreateIfNull)
 {
-	FTextureInstanceState* State = StateSync.SyncAndGetState();
+	FRenderAssetInstanceState* State = StateSync.SyncAndGetState();
 	if (!AsyncView && bCreateIfNull)
 	{
 		AsyncView = State;
@@ -116,16 +116,16 @@ const FTextureInstanceView* FStaticTextureInstanceManager::GetAsyncView(bool bCr
 	return AsyncView.GetReference();
 }
 
-uint32 FStaticTextureInstanceManager::GetAllocatedSize() const
+uint32 FStaticRenderAssetInstanceManager::GetAllocatedSize() const
 {
-	const FTextureInstanceState* State = StateSync.GetState();
-	return  State ? (sizeof(FTextureInstanceState) + State->GetAllocatedSize()) : 0;
+	const FRenderAssetInstanceState* State = StateSync.GetState();
+	return  State ? (sizeof(FRenderAssetInstanceState) + State->GetAllocatedSize()) : 0;
 }
 
 
-void FStaticTextureInstanceManager::OffsetBounds(const FVector& Offset)
+void FStaticRenderAssetInstanceManager::OffsetBounds(const FVector& Offset)
 {
-	FTextureInstanceState* State = StateSync.SyncAndGetState();
+	FRenderAssetInstanceState* State = StateSync.SyncAndGetState();
 	if (State)
 	{
 		State->OffsetBounds(Offset);
diff --git a/Engine/Source/Runtime/Engine/Private/Streaming/StaticTextureInstanceManager.h b/Engine/Source/Runtime/Engine/Private/Streaming/StaticTextureInstanceManager.h
--- a/Engine/Source/Runtime/Engine/Private/Streaming/StaticTextureInstanceManager.h
+++ b/Engine/Source/Runtime/Engine/Private/Streaming/StaticTextureInstanceManager.h
@@ -11,29 +11,29 @@ StaticTextureInstanceManager.h: Definitions of classes used for texture streamin
 #include "Streaming/TextureInstanceTask.h"
 
 /** 
- * A texture instance manager to manage fully static components. 
+ * A texture/mesh instance manager to manage fully static components. 
  * Once an async view has been requested, nothing can be changed as the async view and the internal state point to the same data. 
  * This allows for quicker refreshes and also prevent state duplication.
  */
-class FStaticTextureInstanceManager : public ITextureInstanceManager
+class FStaticRenderAssetInstanceManager : public IRenderAssetInstanceManager
 {
 public:
 
 	/** Contructor. */
-	FStaticTextureInstanceManager(TextureInstanceTask::FDoWorkTask& AsyncTask);
-	~FStaticTextureInstanceManager() { StateSync.Sync(); }
+	FStaticRenderAssetInstanceManager(RenderAssetInstanceTask::FDoWorkTask& AsyncTask);
+	~FStaticRenderAssetInstanceManager() { StateSync.Sync(); }
 
 	/** Normalize lightmap texel factors, this is ran on an async tasks. */
 	void NormalizeLightmapTexelFactor();
 
 	FORCEINLINE int32 CompileElements() { return StateSync.SyncAndGetState()->CompileElements(); }
 	FORCEINLINE int32 CheckRegistrationAndUnpackBounds(TArray<const UPrimitiveComponent*>& RemovedComponents) { return StateSync.SyncAndGetState()->CheckRegistrationAndUnpackBounds(RemovedComponents); }
-	FORCEINLINE FTextureInstanceState::FTextureIterator GetTextureIterator( ) {  return StateSync.SyncAndGetState()->GetTextureIterator(); }
+	FORCEINLINE FRenderAssetInstanceState::FRenderAssetIterator GetRenderAssetIterator( ) {  return StateSync.SyncAndGetState()->GetRenderAssetIterator(); }
 
-	FORCEINLINE bool HasTextureReferences() const { return StateSync.GetState()->NumBounds() > 0; }
+	FORCEINLINE bool HasRenderAssetReferences() const { return StateSync.GetState()->NumBounds() > 0; }
 
 	/*-----------------------------------
-	------ ITextureInstanceManager ------
+	------ IRenderAssetInstanceManager ------
 	-----------------------------------*/
 
 	/** Return whether this component can be managed by this manager. */
@@ -45,20 +45,20 @@ public:
 	/** Return whether this component is be managed by this manager. */
 	bool CanManage(const UPrimitiveComponent* Component) const final override;
 
-	/** Refresh component data (bounds, last render time, min and max view distance) - see TextureInstanceView. */
+	/** Refresh component data (bounds, last render time, min and max view distance) - see RenderAssetInstanceView. */
 	void Refresh(float Percentage) final override;
 
 	/** Add a component streaming data, the LevelContext gives support for precompiled data. */
 	EAddComponentResult Add(const UPrimitiveComponent* Component, FStreamingTextureLevelContext& LevelContext, float MaxAllowedUIDensity) final override;
 
 	/** Remove a component, the RemoveTextures is the list of textures not referred anymore. */
-	 void Remove(const UPrimitiveComponent* Component, FRemovedTextureArray* RemovedTextures) final override;
+	 void Remove(const UPrimitiveComponent* Component, FRemovedRenderAssetArray* RemovedRenderAssets) final override;
 
 	/** Notify the manager that an async view will be requested on the next frame. */
 	FORCEINLINE void PrepareAsyncView() final override {}
 
 	/** Return a view of the data that has to be 100% thread safe. The content is allowed to be updated, but not memory must be reallocated. */
-	const FTextureInstanceView* GetAsyncView(bool bCreateIfNull) final override;
+	const FRenderAssetInstanceView* GetAsyncView(bool bCreateIfNull) final override;
 
 	/** Return the size taken for sub-allocation. */
 	uint32 GetAllocatedSize() const final override;
@@ -72,8 +72,8 @@ protected:
 
 private:
 
-	typedef TextureInstanceTask::FRefreshVisibilityTask FRefreshVisibilityTask;
-	typedef TextureInstanceTask::FNormalizeLightmapTexelFactorTask FNormalizeLightmapTexelFactorTask;
+	typedef RenderAssetInstanceTask::FRefreshVisibilityTask FRefreshVisibilityTask;
+	typedef RenderAssetInstanceTask::FNormalizeLightmapTexelFactorTask FNormalizeLightmapTexelFactorTask;
 
 	struct FTasks
 	{
@@ -84,10 +84,10 @@ private:
 	};
 
 	/** The texture instances. Shared with the async task. */
-	FTextureInstanceStateTaskSync<FTasks> StateSync;
+	FRenderAssetInstanceStateTaskSync<FTasks> StateSync;
 
 	/** A duplicate view for the async streaming task. */
-	TRefCountPtr<const FTextureInstanceView> AsyncView;
+	TRefCountPtr<const FRenderAssetInstanceView> AsyncView;
 
 	/** Ranges from 0 to Bounds4Components.Num(). Used in the incremental update to update visibility. */
 	int32 DirtyIndex;
