Resolve robomerge conflicts

diff --git a/Engine/Source/Runtime/Engine/Private/Streaming/StreamingManagerTexture.cpp b/Engine/Source/Runtime/Engine/Private/Streaming/StreamingManagerTexture.cpp
--- a/Engine/Source/Runtime/Engine/Private/Streaming/StreamingManagerTexture.cpp
+++ b/Engine/Source/Runtime/Engine/Private/Streaming/StreamingManagerTexture.cpp
@@ -8,6 +8,7 @@
 #include "GameFramework/Actor.h"
 #include "Engine/World.h"
 #include "Engine/TextureStreamingTypes.h"
+#include "Engine/StaticMesh.h"
 #include "Materials/MaterialInterface.h"
 #include "Misc/CommandLine.h"
 #include "Misc/ConfigCacheIni.h"
@@ -20,12 +21,14 @@
 #include "Components/PrimitiveComponent.h"
 #include "Misc/CoreDelegates.h"
 #include "ProfilingDebugging/CsvProfiler.h"
+#include "Interfaces/ITargetPlatformManagerModule.h"
+#include "Interfaces/ITargetPlatform.h"
 
 CSV_DECLARE_CATEGORY_MODULE_EXTERN(CORE_API, Basic);
 
-bool TrackTexture( const FString& TextureName );
-bool UntrackTexture( const FString& TextureName );
-void ListTrackedTextures( FOutputDevice& Ar, int32 NumTextures );
+bool TrackRenderAsset( const FString& AssetName );
+bool UntrackRenderAsset( const FString& AssetName );
+void ListTrackedRenderAssets( FOutputDevice& Ar, int32 NumTextures );
 
 /**
  * Helper function to clamp the mesh to camera distance
@@ -37,25 +40,25 @@ FORCEINLINE float ClampMeshToCameraDistanceSquared(float MeshToCameraDistanceSqu
 }
 
 /*-----------------------------------------------------------------------------
-	FStreamingManagerTexture implementation.
+	FRenderAssetStreamingManager implementation.
 -----------------------------------------------------------------------------*/
 
 /** Constructor, initializing all members and  */
-FStreamingManagerTexture::FStreamingManagerTexture()
-:	CurrentUpdateStreamingTextureIndex(0)
+FRenderAssetStreamingManager::FRenderAssetStreamingManager()
+:	CurrentUpdateStreamingRenderAssetIndex(0)
 ,	bTriggerDumpTextureGroupStats( false )
 ,	bDetailedDumpTextureGroupStats( false )
 ,	AsyncWork( nullptr )
 ,	CurrentPendingMipCopyRequestIdx(0)
 ,	ProcessingStage( 0 )
-,	NumTextureProcessingStages(5)
+,	NumRenderAssetProcessingStages(5)
 ,	bUseDynamicStreaming( false )
 ,	BoostPlayerTextures( 3.0f )
 ,	MemoryMargin(0)
 ,	EffectiveStreamingPoolSize(0)
 ,	MemoryOverBudget(0)
 ,	MaxEverRequired(0)
-,	bPauseTextureStreaming(false)
+,	bPauseRenderAssetStreaming(false)
 ,	LastWorldUpdateTime(GIsEditor ? -FLT_MAX : 0) // In editor, visibility is not taken into consideration.
 {
 	// Read settings from ini file.
@@ -70,19 +73,19 @@ FStreamingManagerTexture::FStreamingManagerTexture()
 	GConfig->GetInt(TEXT("TextureStreaming"), TEXT("PoolSize"), PoolSizeIniSetting, GEngineIni);
 	GConfig->GetBool(TEXT("TextureStreaming"), TEXT("UseDynamicStreaming"), bUseDynamicStreaming, GEngineIni);
 	GConfig->GetFloat( TEXT("TextureStreaming"), TEXT("BoostPlayerTextures"), BoostPlayerTextures, GEngineIni );
-	GConfig->GetBool(TEXT("TextureStreaming"), TEXT("NeverStreamOutTextures"), GNeverStreamOutTextures, GEngineIni);
+	GConfig->GetBool(TEXT("TextureStreaming"), TEXT("NeverStreamOutRenderAssets"), GNeverStreamOutRenderAssets, GEngineIni);
 
-	// -NeverStreamOutTextures
-	if (FParse::Param(FCommandLine::Get(), TEXT("NeverStreamOutTextures")))
+	// -NeverStreamOutRenderAssets
+	if (FParse::Param(FCommandLine::Get(), TEXT("NeverStreamOutRenderAssets")))
 	{
-		GNeverStreamOutTextures = true;
+		GNeverStreamOutRenderAssets = true;
 	}
 	if (GIsEditor)
 	{
 		// this would not be good or useful in the editor
-		GNeverStreamOutTextures = false;
+		GNeverStreamOutRenderAssets = false;
 	}
-	if (GNeverStreamOutTextures)
+	if (GNeverStreamOutRenderAssets)
 	{
 		UE_LOG(LogContentStreaming, Log, TEXT("Textures will NEVER stream out!"));
 	}
@@ -101,29 +104,33 @@ FStreamingManagerTexture::FStreamingManagerTexture()
 	for ( int32 LODGroup=0; LODGroup < TEXTUREGROUP_MAX; ++LODGroup )
 	{
 		const FTextureLODGroup& TexGroup = UDeviceProfileManager::Get().GetActiveProfile()->GetTextureLODSettings()->GetTextureLODGroup(TextureGroup(LODGroup));
-		NumStreamedMips[LODGroup] = TexGroup.NumStreamedMips;
+		NumStreamedMips_Texture[LODGroup] = TexGroup.NumStreamedMips;
 	}
 
+	// TODO: NumStreamedMips_StaticMesh, NumStreamedMips_SkeletalMesh
+	NumStreamedMips_StaticMesh.Empty(1);
+	NumStreamedMips_StaticMesh.Add(INT32_MAX);
+
 	// setup the streaming resource flush function pointer
 	GFlushStreamingFunc = &FlushResourceStreaming;
 
 	ProcessingStage = 0;
-	AsyncWork = new FAsyncTask<FAsyncTextureStreamingTask>(this);
+	AsyncWork = new FAsyncTask<FRenderAssetStreamingMipCalcTask>(this);
 
-	TextureInstanceAsyncWork = new TextureInstanceTask::FDoWorkAsyncTask();
-	DynamicComponentManager.RegisterTasks(TextureInstanceAsyncWork->GetTask());
+	RenderAssetInstanceAsyncWork = new RenderAssetInstanceTask::FDoWorkAsyncTask();
+	DynamicComponentManager.RegisterTasks(RenderAssetInstanceAsyncWork->GetTask());
 
-	FCoreUObjectDelegates::GetPreGarbageCollectDelegate().AddRaw(this, &FStreamingManagerTexture::OnPreGarbageCollect);
+	FCoreUObjectDelegates::GetPreGarbageCollectDelegate().AddRaw(this, &FRenderAssetStreamingManager::OnPreGarbageCollect);
 
-	FCoreDelegates::PakFileMountedCallback.AddRaw(this, &FStreamingManagerTexture::OnPakFileMounted);
+	FCoreDelegates::PakFileMountedCallback.AddRaw(this, &FRenderAssetStreamingManager::OnPakFileMounted);
 }
 
-FStreamingManagerTexture::~FStreamingManagerTexture()
+FRenderAssetStreamingManager::~FRenderAssetStreamingManager()
 {
 	AsyncWork->EnsureCompletion();
 	delete AsyncWork;
 
-	TextureInstanceAsyncWork->EnsureCompletion();
+	RenderAssetInstanceAsyncWork->EnsureCompletion();
 	
 	FCoreUObjectDelegates::GetPreGarbageCollectDelegate().RemoveAll(this);
 
@@ -132,33 +139,33 @@ FStreamingManagerTexture::~FStreamingManagerTexture()
 	STAT(DisplayedStats.Apply();)
 }
 
-void FStreamingManagerTexture::OnPreGarbageCollect()
+void FRenderAssetStreamingManager::OnPreGarbageCollect()
 {
 	FScopeLock ScopeLock(&CriticalSection);
-	FRemovedTextureArray RemovedTextures;
+	FRemovedRenderAssetArray RemovedRenderAssets;
 
 	// Check all levels for pending kills.
-	for (int32 Index = 0; Index < LevelTextureManagers.Num(); ++Index)
+	for (int32 Index = 0; Index < LevelRenderAssetManagers.Num(); ++Index)
 	{
-		FLevelTextureManager& LevelManager = LevelTextureManagers[Index];
+		FLevelRenderAssetManager& LevelManager = LevelRenderAssetManagers[Index];
 		if (LevelManager.GetLevel()->IsPendingKill())
 		{
-			LevelManager.Remove(&RemovedTextures);
+			LevelManager.Remove(&RemovedRenderAssets);
 
 			// Remove the level entry. The async task view will still be valid as it uses a shared ptr.
-			LevelTextureManagers.RemoveAtSwap(Index);
+			LevelRenderAssetManagers.RemoveAtSwap(Index);
 			--Index;
 		}
 	}
 
-	DynamicComponentManager.OnPreGarbageCollect(RemovedTextures);
+	DynamicComponentManager.OnPreGarbageCollect(RemovedRenderAssets);
 
-	SetTexturesRemovedTimestamp(RemovedTextures);
+	SetRenderAssetsRemovedTimestamp(RemovedRenderAssets);
 }
 
 
 
-void FStreamingManagerTexture::OnPakFileMounted(const TCHAR* PakFilename)
+void FRenderAssetStreamingManager::OnPakFileMounted(const TCHAR* PakFilename)
 {
 	// clear the cached file exists checks which failed as they may now be loaded
 	bNewFilesLoaded = true;
@@ -167,30 +174,27 @@ void FStreamingManagerTexture::OnPakFileMounted(const TCHAR* PakFilename)
 /**
  * Cancels the timed Forced resources (i.e used the Kismet action "Stream In Textures").
  */
-void FStreamingManagerTexture::CancelForcedResources()
+void FRenderAssetStreamingManager::CancelForcedResources()
 {
 	FScopeLock ScopeLock(&CriticalSection);
 
-	// Update textures that are Forced on a timer.
-	for ( int32 TextureIndex=0; TextureIndex < StreamingTextures.Num(); ++TextureIndex )
+	// Update textures/meshes that are Forced on a timer.
+	for ( int32 Idx=0; Idx < StreamingRenderAssets.Num(); ++Idx )
 	{
-		FStreamingTexture& StreamingTexture = StreamingTextures[ TextureIndex ];
+		FStreamingRenderAsset& StreamingRenderAsset = StreamingRenderAssets[ Idx ];
 
-		// Make sure this streaming texture hasn't been marked for removal.
-		if ( StreamingTexture.Texture )
+		// Make sure this streaming texture/mesh hasn't been marked for removal.
+		if (StreamingRenderAsset.RenderAsset)
 		{
-			// Remove any prestream requests from textures
-			float TimeLeft = (float)(StreamingTexture.Texture->ForceMipLevelsToBeResidentTimestamp - FApp::GetCurrentTime());
+			// Remove any prestream requests from textures/meshes
+			float TimeLeft = (float)(StreamingRenderAsset.RenderAsset->ForceMipLevelsToBeResidentTimestamp - FApp::GetCurrentTime());
 			if ( TimeLeft >= 0.0f )
 			{
-				StreamingTexture.Texture->SetForceMipLevelsToBeResident( -1.0f );
-				StreamingTexture.InstanceRemovedTimestamp = -FLT_MAX;
-				if ( StreamingTexture.Texture->Resource )
-				{
-					StreamingTexture.Texture->InvalidateLastRenderTimeForStreaming();
-				}
+				StreamingRenderAsset.RenderAsset->SetForceMipLevelsToBeResident( -1.0f );
+				StreamingRenderAsset.InstanceRemovedTimestamp = -FLT_MAX;
+				StreamingRenderAsset.RenderAsset->InvalidateLastRenderTimeForStreaming();
 #if STREAMING_LOG_CANCELFORCED
-				UE_LOG(LogContentStreaming, Log, TEXT("Canceling forced texture: %s (had %.1f seconds left)"), *StreamingTexture.Texture->GetFullName(), TimeLeft );
+				UE_LOG(LogContentStreaming, Log, TEXT("Canceling forced texture: %s (had %.1f seconds left)"), *StreamingRenderAsset.Texture->GetFullName(), TimeLeft );
 #endif
 			}
 		}
@@ -203,67 +207,67 @@ void FStreamingManagerTexture::CancelForcedResources()
 /**
  * Notifies manager of "level" change so it can prioritize character textures for a few frames.
  */
-void FStreamingManagerTexture::NotifyLevelChange()
+void FRenderAssetStreamingManager::NotifyLevelChange()
 {
 }
 
 /** Don't stream world resources for the next NumFrames. */
-void FStreamingManagerTexture::SetDisregardWorldResourcesForFrames( int32 NumFrames )
+void FRenderAssetStreamingManager::SetDisregardWorldResourcesForFrames( int32 NumFrames )
 {
 	//@TODO: We could perhaps increase the priority factor for character textures...
 }
 
 /**
- *	Try to stream out texture mip-levels to free up more memory.
+ *	Try to stream out texture/mesh mip-levels to free up more memory.
  *	@param RequiredMemorySize	- Additional texture memory required
  *	@return						- Whether it succeeded or not
  **/
-bool FStreamingManagerTexture::StreamOutTextureData( int64 RequiredMemorySize )
+bool FRenderAssetStreamingManager::StreamOutRenderAssetData( int64 RequiredMemorySize )
 {
 	FScopeLock ScopeLock(&CriticalSection);
 
 	const int64 MaxTempMemoryAllowed = Settings.MaxTempMemoryAllowed * 1024 * 1024;
-	const bool CachedPauseTextureStreaming = bPauseTextureStreaming;
+	const bool CachedPauseTextureStreaming = bPauseRenderAssetStreaming;
 
 	// Pause texture streaming to prevent sending load requests.
-	bPauseTextureStreaming = true;
+	bPauseRenderAssetStreaming = true;
 	SyncStates(true);
 
-	// Sort texture, having those that should be dropped first.
-	TArray<int32> PrioritizedTextures;
-	PrioritizedTextures.Empty(StreamingTextures.Num());
-	for (int32 TextureIndex = 0; TextureIndex < StreamingTextures.Num(); ++TextureIndex)
+	// Sort texture/mesh, having those that should be dropped first.
+	TArray<int32> PrioritizedRenderAssets;
+	PrioritizedRenderAssets.Empty(StreamingRenderAssets.Num());
+	for (int32 Idx = 0; Idx < StreamingRenderAssets.Num(); ++Idx)
 	{
-		FStreamingTexture& StreamingTexture = StreamingTextures[TextureIndex];
+		FStreamingRenderAsset& StreamingRenderAsset = StreamingRenderAssets[Idx];
 		// Only texture for which we can drop mips.
-		if (StreamingTexture.IsMaxResolutionAffectedByGlobalBias())
+		if (StreamingRenderAsset.IsMaxResolutionAffectedByGlobalBias())
 		{
-			PrioritizedTextures.Add(TextureIndex);
+			PrioritizedRenderAssets.Add(Idx);
 		}
 	}
-	PrioritizedTextures.Sort(FCompareTextureByRetentionPriority(StreamingTextures));
+	PrioritizedRenderAssets.Sort(FCompareRenderAssetByRetentionPriority(StreamingRenderAssets));
 
 	int64 TempMemoryUsed = 0;
 	int64 MemoryDropped = 0;
 
-	// Process all texture, starting with the ones we least want to keep
-	for (int32 PriorityIndex = PrioritizedTextures.Num() - 1; PriorityIndex >= 0 && MemoryDropped < RequiredMemorySize; --PriorityIndex)
+	// Process all texture/mesh, starting with the ones we least want to keep
+	for (int32 PriorityIndex = PrioritizedRenderAssets.Num() - 1; PriorityIndex >= 0 && MemoryDropped < RequiredMemorySize; --PriorityIndex)
 	{
-		int32 TextureIndex = PrioritizedTextures[PriorityIndex];
-		if (!StreamingTextures.IsValidIndex(TextureIndex)) continue;
+		int32 RenderAssetIndex = PrioritizedRenderAssets[PriorityIndex];
+		if (!StreamingRenderAssets.IsValidIndex(RenderAssetIndex)) continue;
 
-		FStreamingTexture& StreamingTexture = StreamingTextures[TextureIndex];
-		if (!StreamingTexture.Texture) continue;
+		FStreamingRenderAsset& StreamingRenderAsset = StreamingRenderAssets[RenderAssetIndex];
+		if (!StreamingRenderAsset.RenderAsset) continue;
 
-		const int32 MinimalSize = StreamingTexture.GetSize(StreamingTexture.MinAllowedMips);
-		const int32 CurrentSize = StreamingTexture.GetSize(StreamingTexture.ResidentMips);
+		const int32 MinimalSize = StreamingRenderAsset.GetSize(StreamingRenderAsset.MinAllowedMips);
+		const int32 CurrentSize = StreamingRenderAsset.GetSize(StreamingRenderAsset.ResidentMips);
 
-		if (StreamingTexture.Texture->StreamOut(StreamingTexture.MinAllowedMips))
+		if (StreamingRenderAsset.RenderAsset->StreamOut(StreamingRenderAsset.MinAllowedMips))
 		{
 			MemoryDropped += CurrentSize - MinimalSize; 
 			TempMemoryUsed += MinimalSize;
 
-			StreamingTexture.UpdateStreamingStatus(false);
+			StreamingRenderAsset.UpdateStreamingStatus(false);
 
 			if (TempMemoryUsed >= MaxTempMemoryAllowed)
 			{
@@ -280,15 +284,15 @@ bool FStreamingManagerTexture::StreamOutTextureData( int64 RequiredMemorySize )
 		}
 	}
 
-	bPauseTextureStreaming = CachedPauseTextureStreaming;
+	bPauseRenderAssetStreaming = CachedPauseTextureStreaming;
 	UE_LOG(LogContentStreaming, Log, TEXT("Streaming out texture memory! Saved %.2f MB."), float(MemoryDropped)/1024.0f/1024.0f);
 	return true;
 }
 
-void FStreamingManagerTexture::IncrementalUpdate(float Percentage, bool bUpdateDynamicComponents)
+void FRenderAssetStreamingManager::IncrementalUpdate(float Percentage, bool bUpdateDynamicComponents)
 {
-	QUICK_SCOPE_CYCLE_COUNTER(STAT_FStreamingManagerTexture_IncrementalUpdate);
-	FRemovedTextureArray RemovedTextures;
+	QUICK_SCOPE_CYCLE_COUNTER(STAT_FRenderAssetStreamingManager_IncrementalUpdate);
+	FRemovedRenderAssetArray RemovedRenderAssets;
 
 	int64 NumStepsLeftForIncrementalBuild = CVarStreamingNumStaticComponentsProcessedPerFrame.GetValueOnGameThread();
 	if (NumStepsLeftForIncrementalBuild <= 0) // When 0, don't allow incremental updates.
@@ -296,60 +300,66 @@ void FStreamingManagerTexture::IncrementalUpdate(float Percentage, bool bUpdateD
 		NumStepsLeftForIncrementalBuild = MAX_int64;
 	}
 
-	for (FLevelTextureManager& LevelManager : LevelTextureManagers)
+	for (FLevelRenderAssetManager& LevelManager : LevelRenderAssetManagers)
 	{
-		LevelManager.IncrementalUpdate(DynamicComponentManager, RemovedTextures, NumStepsLeftForIncrementalBuild, Percentage, bUseDynamicStreaming); // Complete the incremental update.
+		LevelManager.IncrementalUpdate(DynamicComponentManager, RemovedRenderAssets, NumStepsLeftForIncrementalBuild, Percentage, bUseDynamicStreaming); // Complete the incremental update.
 	}
 
 	// Dynamic component are only udpated when it is useful for the dynamic async view.
 	if (bUpdateDynamicComponents && bUseDynamicStreaming)
 	{
-		DynamicComponentManager.IncrementalUpdate(RemovedTextures, Percentage);
+		DynamicComponentManager.IncrementalUpdate(RemovedRenderAssets, Percentage);
 	}
 
-	SetTexturesRemovedTimestamp(RemovedTextures);
+	SetRenderAssetsRemovedTimestamp(RemovedRenderAssets);
 }
 
-void FStreamingManagerTexture::ProcessRemovedTextures()
+void FRenderAssetStreamingManager::ProcessRemovedRenderAssets()
 {
-	for (int32 TextureIndex : RemovedTextureIndices)
+	for (int32 AssetIndex : RemovedRenderAssetIndices)
 	{
-		// Remove swap all elements, until this entry has a valid texture.
+		// Remove swap all elements, until this entry has a valid texture/mesh.
 		// This handles the case where the last element was also removed.
-		while (StreamingTextures.IsValidIndex(TextureIndex) && !StreamingTextures[TextureIndex].Texture)
+		while (StreamingRenderAssets.IsValidIndex(AssetIndex) && !StreamingRenderAssets[AssetIndex].RenderAsset)
 		{
-			StreamingTextures.RemoveAtSwap(TextureIndex);
+			StreamingRenderAssets.RemoveAtSwap(AssetIndex);
 		}
 
-		if (StreamingTextures.IsValidIndex(TextureIndex))
+		if (StreamingRenderAssets.IsValidIndex(AssetIndex))
 		{
 			// Update the texture with its new index.
-			StreamingTextures[TextureIndex].Texture->StreamingIndex = TextureIndex;
+			StreamingRenderAssets[AssetIndex].RenderAsset->StreamingIndex = AssetIndex;
 		}
 	}
-	RemovedTextureIndices.Empty();
+	RemovedRenderAssetIndices.Empty();
 }
 
-void FStreamingManagerTexture::ProcessAddedTextures()
+void FRenderAssetStreamingManager::ProcessAddedRenderAssets()
 {
-	// Add new textures.
-	StreamingTextures.Reserve(StreamingTextures.Num() + PendingStreamingTextures.Num());
-	for (int32 TextureIndex=0; TextureIndex < PendingStreamingTextures.Num(); ++TextureIndex)
+	// Add new textures or meshes.
+	StreamingRenderAssets.Reserve(StreamingRenderAssets.Num() + PendingStreamingRenderAssets.Num());
+	for (int32 Idx = 0; Idx < PendingStreamingRenderAssets.Num(); ++Idx)
 	{
-		UTexture2D* Texture = PendingStreamingTextures[TextureIndex];
-		if (!Texture) continue; // Could be null if it was removed after being added.
-
-		Texture->StreamingIndex = StreamingTextures.Num();
-		new (StreamingTextures) FStreamingTexture(Texture, NumStreamedMips, Settings);
+		UStreamableRenderAsset* Asset = PendingStreamingRenderAssets[Idx];
+		// Could be null if it was removed after being added.
+		if (Asset)
+		{
+			Asset->StreamingIndex = StreamingRenderAssets.Num();
+			const FStreamingRenderAsset::EAssetType Type = PendingStreamingRenderAssetTypes[Idx];
+			const int32* NumStreamedMips;
+			const int32 NumLODGroups = GetNumStreamedMipsArray(Type, NumStreamedMips);
+			new (StreamingRenderAssets) FStreamingRenderAsset(Asset, NumStreamedMips, NumLODGroups, Type, Settings);
+		}
 	}
-	PendingStreamingTextures.Empty();
+	PendingStreamingRenderAssets.Empty();
+	PendingStreamingRenderAssetTypes.Empty();
 }
 
-void FStreamingManagerTexture::ConditionalUpdateStaticData()
+void FRenderAssetStreamingManager::ConditionalUpdateStaticData()
 {
 	static float PreviousLightmapStreamingFactor = GLightmapStreamingFactor;
 	static float PreviousShadowmapStreamingFactor = GShadowmapStreamingFactor;
-	static FTextureStreamingSettings PreviousSettings = Settings;
+	static FRenderAssetStreamingSettings PreviousSettings = Settings;
 
 	if (PreviousLightmapStreamingFactor != GLightmapStreamingFactor || 
 		PreviousShadowmapStreamingFactor != GShadowmapStreamingFactor || 
@@ -357,16 +367,16 @@ void FStreamingManagerTexture::ConditionalUpdateStaticData()
 	{
 		STAT(GatheredStats.SetupAsyncTaskCycles += FPlatformTime::Cycles();)
 		// Update each texture static data.
-		for (FStreamingTexture& StreamingTexture : StreamingTextures)
+		for (FStreamingRenderAsset& StreamingRenderAsset : StreamingRenderAssets)
 		{
-			StreamingTexture.UpdateStaticData(Settings);
+			StreamingRenderAsset.UpdateStaticData(Settings);
 
 			// When the material quality changes, some textures could stop being used.
 			// Refreshing their removed timestamp ensures not texture ends up in the unkwown 
 			// ref heuristic (which would force load them).
 			if (PreviousSettings.MaterialQualityLevel != Settings.MaterialQualityLevel)
 			{
-				StreamingTexture.InstanceRemovedTimestamp = FApp::GetCurrentTime();
+				StreamingRenderAsset.InstanceRemovedTimestamp = FApp::GetCurrentTime();
 			}
 		}
 		STAT(GatheredStats.SetupAsyncTaskCycles -= (int32)FPlatformTime::Cycles();)
@@ -388,12 +398,12 @@ void FStreamingManagerTexture::ConditionalUpdateStaticData()
 			TArray<ULevel*, TInlineAllocator<32> > Levels;
 
 			// RemoveLevel data
-			for (FLevelTextureManager& LevelManager : LevelTextureManagers)
+			for (FLevelRenderAssetManager& LevelManager : LevelRenderAssetManagers)
 			{
 				Levels.Push(LevelManager.GetLevel());
 				LevelManager.Remove(nullptr);
 			}
-			LevelTextureManagers.Empty();
+			LevelRenderAssetManagers.Empty();
 
 			for (ULevel* Level : Levels)
 			{
@@ -416,12 +426,12 @@ void FStreamingManagerTexture::ConditionalUpdateStaticData()
 	}
 }
 
-void FStreamingManagerTexture::UpdatePendingStates(bool bUpdateDynamicComponents)
+void FRenderAssetStreamingManager::UpdatePendingStates(bool bUpdateDynamicComponents)
 {
 	CheckUserSettings();
 
-	ProcessRemovedTextures();
-	ProcessAddedTextures();
+	ProcessRemovedRenderAssets();
+	ProcessAddedRenderAssets();
 
 	Settings.Update();
 	ConditionalUpdateStaticData();
@@ -437,14 +447,16 @@ void FStreamingManagerTexture::UpdatePendingStates(bool bUpdateDynamicComponents
 }
 
 /**
- * Adds new textures and level data on the gamethread (while the worker thread isn't active).
+ * Adds new textures/meshes and level data on the gamethread (while the worker thread isn't active).
  */
-void FStreamingManagerTexture::PrepareAsyncTask(bool bProcessEverything)
+void FRenderAssetStreamingManager::PrepareAsyncTask(bool bProcessEverything)
 {
-	FAsyncTextureStreamingTask& AsyncTask = AsyncWork->GetTask();
+	FRenderAssetStreamingMipCalcTask& AsyncTask = AsyncWork->GetTask();
 	FTextureMemoryStats Stats;
 	RHIGetTextureMemoryStats(Stats);
 
+	// TODO: Track memory allocated by mesh LODs
+
 	// When processing all textures, we need unlimited budget so that textures get all at their required states.
 	// Same when forcing stream-in, for which we want all used textures to be fully loaded 
 	if (Stats.IsUsingLimitedPoolSize() && !bProcessEverything && !Settings.bFullyLoadUsedTextures)
@@ -457,14 +469,14 @@ void FStreamingManagerTexture::PrepareAsyncTask(bool bProcessEverything)
 		// Temp must be smaller since membudget only updates if it has a least temp memory available.
 		AsyncTask.Reset(0, Stats.AllocatedMemorySize, MAX_int64, MAX_int64 / 2, 0);
 	}
-	AsyncTask.StreamingData.Init(CurrentViewInfos, LastWorldUpdateTime, LevelTextureManagers, DynamicComponentManager);
+	AsyncTask.StreamingData.Init(CurrentViewInfos, LastWorldUpdateTime, LevelRenderAssetManagers, DynamicComponentManager);
 }
 
 /**
  * Temporarily boosts the streaming distance factor by the specified number.
  * This factor is automatically reset to 1.0 after it's been used for mip-calculations.
  */
-void FStreamingManagerTexture::BoostTextures( AActor* Actor, float BoostFactor )
+void FRenderAssetStreamingManager::BoostTextures( AActor* Actor, float BoostFactor )
 {
 	FScopeLock ScopeLock(&CriticalSection);
 
@@ -482,7 +494,7 @@ void FStreamingManagerTexture::BoostTextures( AActor* Actor, float BoostFactor )
 				Primitive->GetUsedTextures( Textures, EMaterialQualityLevel::Num );
 				for ( UTexture* Texture : Textures )
 				{
-					FStreamingTexture* StreamingTexture = GetStreamingTexture( Cast<UTexture2D>( Texture ) );
+					FStreamingRenderAsset* StreamingTexture = GetStreamingRenderAsset(Texture);
 					if ( StreamingTexture )
 					{
 						StreamingTexture->DynamicBoostFactor = FMath::Max( StreamingTexture->DynamicBoostFactor, BoostFactor );
@@ -494,7 +506,7 @@ void FStreamingManagerTexture::BoostTextures( AActor* Actor, float BoostFactor )
 }
 
 /** Adds a ULevel to the streaming manager. This is called from 2 paths : after PostPostLoad and after AddToWorld */
-void FStreamingManagerTexture::AddLevel( ULevel* Level )
+void FRenderAssetStreamingManager::AddLevel( ULevel* Level )
 {
 	FScopeLock ScopeLock(&CriticalSection);
 
@@ -509,7 +521,7 @@ void FStreamingManagerTexture::AddLevel( ULevel* Level )
 	else
 	{
 		// In game, because static components can not be changed, the level static data is computed and kept as long as the level is not destroyed.
-		for (const FLevelTextureManager& LevelManager : LevelTextureManagers)
+		for (const FLevelRenderAssetManager& LevelManager : LevelRenderAssetManagers)
 		{
 			if (LevelManager.GetLevel() == Level)
 			{
@@ -520,43 +532,43 @@ void FStreamingManagerTexture::AddLevel( ULevel* Level )
 	}
 
 	// If the level was not already there, add a new entry.
-	TextureInstanceAsyncWork->EnsureCompletion();
-	LevelTextureManagers.Add(new FLevelTextureManager(Level, TextureInstanceAsyncWork->GetTask()));
+	RenderAssetInstanceAsyncWork->EnsureCompletion();
+	LevelRenderAssetManagers.Add(new FLevelRenderAssetManager(Level, RenderAssetInstanceAsyncWork->GetTask()));
 }
 
 /** Removes a ULevel from the streaming manager. */
-void FStreamingManagerTexture::RemoveLevel( ULevel* Level )
+void FRenderAssetStreamingManager::RemoveLevel( ULevel* Level )
 {
 	FScopeLock ScopeLock(&CriticalSection);
 
 	check(Level);
 
 	// In editor we remove levels when visibility changes, while in game we want to kept the static data as long as possible.
-	// FLevelTextureManager::IncrementalUpdate will remove dynamic components and mark textures timestamps.
+	// FLevelRenderAssetManager::IncrementalUpdate will remove dynamic components and mark textures/meshes timestamps.
 	if (GIsEditor || Level->IsPendingKill() || Level->HasAnyFlags(RF_BeginDestroyed|RF_FinishDestroyed))
 	{
-		for (int32 Index = 0; Index < LevelTextureManagers.Num(); ++Index)
+		for (int32 Index = 0; Index < LevelRenderAssetManagers.Num(); ++Index)
 		{
-			FLevelTextureManager& LevelManager = LevelTextureManagers[Index];
+			FLevelRenderAssetManager& LevelManager = LevelRenderAssetManagers[Index];
 			if (LevelManager.GetLevel() == Level)
 			{
-				FRemovedTextureArray RemovedTextures;
-				LevelManager.Remove(&RemovedTextures);
-				SetTexturesRemovedTimestamp(RemovedTextures);
+				FRemovedRenderAssetArray RemovedRenderAssets;
+				LevelManager.Remove(&RemovedRenderAssets);
+				SetRenderAssetsRemovedTimestamp(RemovedRenderAssets);
 
 				// Remove the level entry. The async task view will still be valid as it uses a shared ptr.
-				LevelTextureManagers.RemoveAtSwap(Index);
+				LevelRenderAssetManagers.RemoveAtSwap(Index);
 				break;
 			}
 		}
 	}
 }
 
-void FStreamingManagerTexture::NotifyLevelOffset(ULevel* Level, const FVector& Offset)
+void FRenderAssetStreamingManager::NotifyLevelOffset(ULevel* Level, const FVector& Offset)
 {
 	FScopeLock ScopeLock(&CriticalSection);
 
-	for (FLevelTextureManager& LevelManager : LevelTextureManagers)
+	for (FLevelRenderAssetManager& LevelManager : LevelRenderAssetManagers)
 	{
 		if (LevelManager.GetLevel() == Level)
 		{
@@ -566,66 +578,83 @@ void FStreamingManagerTexture::NotifyLevelOffset(ULevel* Level, const FVector& O
 	}
 }
 
-/**
- * Adds a new texture to the streaming manager.
- */
-void FStreamingManagerTexture::AddStreamingTexture( UTexture2D* Texture )
+void FRenderAssetStreamingManager::AddStreamingRenderAsset_Internal(UStreamableRenderAsset* InAsset, FStreamingRenderAsset::EAssetType InType)
 {
 	FScopeLock ScopeLock(&CriticalSection);
 
 	STAT(GatheredStats.CallbacksCycles = -(int32)FPlatformTime::Cycles();)
 
-	// Adds the new texture to the Pending list, to avoid reallocation of the thread-safe StreamingTextures array.
-	check(Texture->StreamingIndex == INDEX_NONE);
-	Texture->StreamingIndex = PendingStreamingTextures.Add(Texture);
-	
+	// Adds the new texture/mesh to the Pending list, to avoid reallocation of the thread-safe StreamingRenderAssets array.
+	check(InAsset->StreamingIndex == INDEX_NONE);
+	InAsset->StreamingIndex = PendingStreamingRenderAssets.Add(InAsset);
+	PendingStreamingRenderAssetTypes.Add(InType);
+
 	// Mark as pending update while the streamer has not determined the required resolution (unless paused)
-	Texture->bHasStreamingUpdatePending = !bPauseTextureStreaming;
+	InAsset->bHasStreamingUpdatePending = !bPauseRenderAssetStreaming;
 
-	// Notify that this texture ptr is valid.
-	ReferencedTextures.Add(Texture);
+	// Notify that this texture/mesh ptr is valid.
+	ReferencedRenderAssets.Add(InAsset);
 
 	STAT(GatheredStats.CallbacksCycles += FPlatformTime::Cycles();)
 }
 
 /**
- * Removes a texture from the streaming manager.
+ * Adds a new texture/mesh to the streaming manager.
  */
-void FStreamingManagerTexture::RemoveStreamingTexture( UTexture2D* Texture )
+void FRenderAssetStreamingManager::AddStreamingRenderAsset( UTexture2D* Texture )
+{
+	AddStreamingRenderAsset_Internal(Texture, FStreamingRenderAsset::AT_Texture);
+}
+
+void FRenderAssetStreamingManager::AddStreamingRenderAsset(UStaticMesh* StaticMesh)
+{
+	AddStreamingRenderAsset_Internal(StaticMesh, FStreamingRenderAsset::AT_StaticMesh);
+}
+
+void FRenderAssetStreamingManager::AddStreamingRenderAsset(USkeletalMesh* SkeletalMesh)
+{
+	// TODO
+	LowLevelFatalError(TEXT("FRenderAssetStreamingManager::AddStreamingRenderAsset(USkeletalMesh* SkeletalMesh) is not implemented"));
+}
+
+/**
+ * Removes a texture/mesh from the streaming manager.
+ */
+void FRenderAssetStreamingManager::RemoveStreamingRenderAsset( UStreamableRenderAsset* RenderAsset )
 {
 	FScopeLock ScopeLock(&CriticalSection);
 
 	STAT(GatheredStats.CallbacksCycles = -(int32)FPlatformTime::Cycles();)
 
-	const int32	TextureIndex = Texture->StreamingIndex;
+	const int32	Idx = RenderAsset->StreamingIndex;
 
 	// Remove it from the Pending list if it is there.
-	if (PendingStreamingTextures.IsValidIndex(TextureIndex) && PendingStreamingTextures[TextureIndex] == Texture)
+	if (PendingStreamingRenderAssets.IsValidIndex(Idx) && PendingStreamingRenderAssets[Idx] == RenderAsset)
 	{
-		PendingStreamingTextures[TextureIndex] = nullptr;
+		PendingStreamingRenderAssets[Idx] = nullptr;
 	}
-	else if (StreamingTextures.IsValidIndex(TextureIndex) && StreamingTextures[TextureIndex].Texture == Texture)
+	else if (StreamingRenderAssets.IsValidIndex(Idx) && StreamingRenderAssets[Idx].RenderAsset == RenderAsset)
 	{
-		StreamingTextures[TextureIndex].Texture = nullptr;
-		RemovedTextureIndices.Add(TextureIndex);
+		StreamingRenderAssets[Idx].RenderAsset = nullptr;
+		RemovedRenderAssetIndices.Add(Idx);
 	}
 
-	Texture->StreamingIndex = INDEX_NONE;
-	Texture->bHasStreamingUpdatePending = false;
+	RenderAsset->StreamingIndex = INDEX_NONE;
+	RenderAsset->bHasStreamingUpdatePending = false;
 
-	// Remove reference to this texture.
-	ReferencedTextures.Remove(Texture);
+	// Remove reference to this texture/mesh.
+	ReferencedRenderAssets.Remove(RenderAsset);
 
 	STAT(GatheredStats.CallbacksCycles += FPlatformTime::Cycles();)
 }
 
 /** Called when a spawned primitive is deleted, or when an actor is destroyed in the editor. */
-void FStreamingManagerTexture::NotifyActorDestroyed( AActor* Actor )
+void FRenderAssetStreamingManager::NotifyActorDestroyed( AActor* Actor )
 {
 	FScopeLock ScopeLock(&CriticalSection);
 
 	STAT(GatheredStats.CallbacksCycles = -(int32)FPlatformTime::Cycles();)
-	FRemovedTextureArray RemovedTextures;
+	FRemovedRenderAssetArray RemovedRenderAssets;
 	check(Actor);
 
 	TInlineComponentArray<UPrimitiveComponent*> Components;
@@ -636,14 +665,14 @@ void FStreamingManagerTexture::NotifyActorDestroyed( AActor* Actor )
 	ULevel* Level = !GIsEditor ? Actor->GetLevel() : nullptr;
 
 	// Remove any reference in the level managers.
-	for (FLevelTextureManager& LevelManager : LevelTextureManagers)
+	for (FLevelRenderAssetManager& LevelManager : LevelRenderAssetManagers)
 	{
 		if (!Level || LevelManager.GetLevel() == Level)
 		{
 			LevelManager.RemoveActorReferences(Actor);
 			for (UPrimitiveComponent* Component : Components)
-				{
-				LevelManager.RemoveComponentReferences(Component, RemovedTextures);
+			{
+				LevelManager.RemoveComponentReferences(Component, RemovedRenderAssets);
 			}
 		}
 	}
@@ -651,17 +680,17 @@ void FStreamingManagerTexture::NotifyActorDestroyed( AActor* Actor )
 	for (UPrimitiveComponent* Component : Components)
 	{
 		// Remove any references in the dynamic component manager.
-		DynamicComponentManager.Remove(Component, &RemovedTextures);
+		DynamicComponentManager.Remove(Component, &RemovedRenderAssets);
 
 		// Reset this now as we have finished iterating over the levels
 		Component->bAttachedToStreamingManagerAsStatic = false;
 	}
 
-	SetTexturesRemovedTimestamp(RemovedTextures);
+	SetRenderAssetsRemovedTimestamp(RemovedRenderAssets);
 	STAT(GatheredStats.CallbacksCycles += FPlatformTime::Cycles();)
 }
 
-void FStreamingManagerTexture::RemoveStaticReferences(const UPrimitiveComponent* Primitive)
+void FRenderAssetStreamingManager::RemoveStaticReferences(const UPrimitiveComponent* Primitive)
 {
 	FScopeLock ScopeLock(&CriticalSection);
 
@@ -669,17 +698,17 @@ void FStreamingManagerTexture::RemoveStaticReferences(const UPrimitiveComponent*
 
 	if (Primitive->bAttachedToStreamingManagerAsStatic)
 	{
-		FRemovedTextureArray RemovedTextures;
+		FRemovedRenderAssetArray RemovedRenderAssets;
 		ULevel* Level = Primitive->GetComponentLevel();
-		for (FLevelTextureManager& LevelManager : LevelTextureManagers)
+		for (FLevelRenderAssetManager& LevelManager : LevelRenderAssetManagers)
 		{
 			if (!Level || LevelManager.GetLevel() == Level)
 			{
-				LevelManager.RemoveComponentReferences(Primitive, RemovedTextures);
+				LevelManager.RemoveComponentReferences(Primitive, RemovedRenderAssets);
 			}
 		}
 		Primitive->bAttachedToStreamingManagerAsStatic = false;
-		// Nothing to do with removed textures as we are about to reinsert
+		// Nothing to do with removed textures/meshes as we are about to reinsert
 	}
 }
 
@@ -687,7 +716,7 @@ void FStreamingManagerTexture::RemoveStaticReferences(const UPrimitiveComponent*
  * Called when a primitive is detached from an actor or another component.
  * Note: We should not be accessing the primitive or the UTexture2D after this call!
  */
-void FStreamingManagerTexture::NotifyPrimitiveDetached( const UPrimitiveComponent* Primitive )
+void FRenderAssetStreamingManager::NotifyPrimitiveDetached( const UPrimitiveComponent* Primitive )
 {
 	FScopeLock ScopeLock(&CriticalSection);
 
@@ -697,7 +726,7 @@ void FStreamingManagerTexture::NotifyPrimitiveDetached( const UPrimitiveComponen
 	}
 
 	STAT(GatheredStats.CallbacksCycles = -(int32)FPlatformTime::Cycles();)
-	FRemovedTextureArray RemovedTextures;
+	FRemovedRenderAssetArray RemovedRenderAssets;
 
 #if STREAMING_LOG_DYNAMIC
 		UE_LOG(LogContentStreaming, Log, TEXT("NotifyPrimitiveDetached(0x%08x \"%s\"), IsRegistered=%d"), SIZE_T(Primitive), *Primitive->GetReadableName(), Primitive->IsRegistered());
@@ -716,11 +745,11 @@ void FStreamingManagerTexture::NotifyPrimitiveDetached( const UPrimitiveComponen
 		// Unless in editor, we don't want to remove reference in static level data when toggling visibility.
 		else if (GIsEditor || Primitive->IsPendingKill() || Primitive->HasAnyFlags(RF_BeginDestroyed|RF_FinishDestroyed))
 		{
-			for (FLevelTextureManager& LevelManager : LevelTextureManagers)
+			for (FLevelRenderAssetManager& LevelManager : LevelRenderAssetManagers)
 			{
 				if (!Level || LevelManager.GetLevel() == Level)
 				{
-					LevelManager.RemoveComponentReferences(Primitive, RemovedTextures);
+					LevelManager.RemoveComponentReferences(Primitive, RemovedRenderAssets);
 				}
 			}
 			Primitive->bAttachedToStreamingManagerAsStatic = false;
@@ -728,38 +757,38 @@ void FStreamingManagerTexture::NotifyPrimitiveDetached( const UPrimitiveComponen
 	}
 	
 	// Dynamic component must be removed when visibility changes.
-	DynamicComponentManager.Remove(Primitive, &RemovedTextures);
+	DynamicComponentManager.Remove(Primitive, &RemovedRenderAssets);
 
-	SetTexturesRemovedTimestamp(RemovedTextures);
+	SetRenderAssetsRemovedTimestamp(RemovedRenderAssets);
 	STAT(GatheredStats.CallbacksCycles += FPlatformTime::Cycles();)
 }
 
 /**
-* Mark the textures with a timestamp. They're about to lose their location-based heuristic and we don't want them to
+* Mark the textures/meshes with a timestamp. They're about to lose their location-based heuristic and we don't want them to
 * start using LastRenderTime heuristic for a few seconds until they are garbage collected!
 *
-* @param RemovedTextures	List of removed textures.
+* @param RemovedRenderAssets	List of removed textures or meshes.
 */
-void FStreamingManagerTexture::SetTexturesRemovedTimestamp(const FRemovedTextureArray& RemovedTextures)
+void FRenderAssetStreamingManager::SetRenderAssetsRemovedTimestamp(const FRemovedRenderAssetArray& RemovedRenderAssets)
 {
 	const double CurrentTime = FApp::GetCurrentTime();
-	for ( int32 TextureIndex=0; TextureIndex < RemovedTextures.Num(); ++TextureIndex )
+	for ( int32 Idx=0; Idx < RemovedRenderAssets.Num(); ++Idx )
 	{
-		// When clearing references to textures, those textures could be already deleted.
-		// This happens because we don't clear texture references in RemoveStreamingTexture.
-		const UTexture2D* Texture = RemovedTextures[TextureIndex];
-		if (!ReferencedTextures.Contains(Texture)) continue;
+		// When clearing references to textures/meshes, those textures/meshes could be already deleted.
+		// This happens because we don't clear texture/mesh references in RemoveStreamingRenderAsset.
+		const UStreamableRenderAsset* Asset = RemovedRenderAssets[Idx];
+		if (!ReferencedRenderAssets.Contains(Asset)) continue;
 
-		FStreamingTexture* StreamingTexture = GetStreamingTexture(Texture);
-		if ( StreamingTexture )
+		FStreamingRenderAsset* StreamingRenderAsset = GetStreamingRenderAsset(Asset);
+		if (StreamingRenderAsset)
 		{
-			StreamingTexture->InstanceRemovedTimestamp = CurrentTime;
+			StreamingRenderAsset->InstanceRemovedTimestamp = CurrentTime;
 		}
 	}
 }
 
 
-void FStreamingManagerTexture::NotifyPrimitiveUpdated( const UPrimitiveComponent* Primitive )
+void FRenderAssetStreamingManager::NotifyPrimitiveUpdated( const UPrimitiveComponent* Primitive )
 {
 	STAT(GatheredStats.CallbacksCycles = -(int32)FPlatformTime::Cycles();)
 
@@ -805,11 +834,11 @@ void FStreamingManagerTexture::NotifyPrimitiveUpdated( const UPrimitiveComponent
 }
 
 /**
- * Called when a primitive has had its textured changed.
+ * Called when a primitive has had its textures/mesh changed.
  * Only affects primitives that were already attached.
  * Replaces previous info.
  */
-void FStreamingManagerTexture::NotifyPrimitiveUpdated_Concurrent( const UPrimitiveComponent* Primitive )
+void FRenderAssetStreamingManager::NotifyPrimitiveUpdated_Concurrent( const UPrimitiveComponent* Primitive )
 {
 	STAT(int32 CallbackCycle = -(int32)FPlatformTime::Cycles();)
 
@@ -825,7 +854,7 @@ void FStreamingManagerTexture::NotifyPrimitiveUpdated_Concurrent( const UPrimiti
 	STAT(FPlatformAtomics::InterlockedAdd(&GatheredStats.CallbacksCycles, CallbackCycle));
 }
 
-void FStreamingManagerTexture::SyncStates(bool bCompleteFullUpdateCycle)
+void FRenderAssetStreamingManager::SyncStates(bool bCompleteFullUpdateCycle)
 {
 	// Finish the current update cycle. 
 	while (ProcessingStage != 0 && bCompleteFullUpdateCycle)
@@ -835,28 +864,28 @@ void FStreamingManagerTexture::SyncStates(bool bCompleteFullUpdateCycle)
 
 	// Wait for async tasks
 	AsyncWork->EnsureCompletion();
-	TextureInstanceAsyncWork->EnsureCompletion();
+	RenderAssetInstanceAsyncWork->EnsureCompletion();
 
-	// Update any pending states, including added/removed textures.
-	// Doing so when ProcessingStage != 0 risk invalidating the indice in the async task used in StreamTextures().
+	// Update any pending states, including added/removed textures/meshes.
+	// Doing so when ProcessingStage != 0 risk invalidating the indices in the async task used in StreamRenderAssets().
 	// This would in practice postpone some of the load and cancel requests.
 	UpdatePendingStates(false);
 }
 
 /**
- * Returns the corresponding FStreamingTexture for a UTexture2D.
+ * Returns the corresponding FStreamingRenderAsset for a texture or mesh.
  */
-FStreamingTexture* FStreamingManagerTexture::GetStreamingTexture( const UTexture2D* Texture2D )
+FStreamingRenderAsset* FRenderAssetStreamingManager::GetStreamingRenderAsset( const UStreamableRenderAsset* RenderAsset )
 {
 	FScopeLock ScopeLock(&CriticalSection);
 
-	if (Texture2D && StreamingTextures.IsValidIndex(Texture2D->StreamingIndex))
+	if (RenderAsset && StreamingRenderAssets.IsValidIndex(RenderAsset->StreamingIndex))
 	{
-		FStreamingTexture* StreamingTexture = &StreamingTextures[Texture2D->StreamingIndex];
+		FStreamingRenderAsset* StreamingRenderAsset = &StreamingRenderAssets[RenderAsset->StreamingIndex];
 
-		// If the texture don't match, this means the texture is pending in PendingStreamingTextures, for which no FStreamingTexture* is yet allocated.
+		// If the texture/mesh don't match, this means the texture/mesh is pending in PendingStreamingRenderAssets, for which no FStreamingRenderAsset* is yet allocated.
 		// If this is not acceptable, the caller should first synchronize everything through SyncStates
-		return StreamingTexture->Texture == Texture2D ? StreamingTexture : nullptr;
+		return StreamingRenderAsset->RenderAsset == RenderAsset ? StreamingRenderAsset : nullptr;
 	}
 	else
 	{
@@ -865,31 +894,34 @@ FStreamingTexture* FStreamingManagerTexture::GetStreamingTexture( const UTexture
 }
 
 /**
- * Updates streaming for an individual texture, taking into account all view infos.
+ * Updates streaming for an individual texture/mesh, taking into account all view infos.
  *
- * @param Texture	Texture to update
+ * @param RenderAsset	Texture or mesh to update
  */
-void FStreamingManagerTexture::UpdateIndividualTexture( UTexture2D* Texture )
+void FRenderAssetStreamingManager::UpdateIndividualRenderAsset( UStreamableRenderAsset* RenderAsset )
 {
 	FScopeLock ScopeLock(&CriticalSection);
 
-	if (!IStreamingManager::Get().IsStreamingEnabled() || !Texture) return;
+	if (!IStreamingManager::Get().IsStreamingEnabled() || !RenderAsset) return;
 
 	// Because we want to priorize loading of this texture, 
 	// don't process everything as this would send load requests for all textures.
 	SyncStates(false);
 
-	FStreamingTexture* StreamingTexture = GetStreamingTexture(Texture);
-	if (!StreamingTexture) return;
+	FStreamingRenderAsset* StreamingRenderAsset = GetStreamingRenderAsset(RenderAsset);
+	if (!StreamingRenderAsset) return;
+
+	const int32* NumStreamedMips;
+	const int32 NumLODGroups = GetNumStreamedMipsArray(StreamingRenderAsset->RenderAssetType, NumStreamedMips);
 
-	StreamingTexture->UpdateDynamicData(NumStreamedMips, Settings, false);
+	StreamingRenderAsset->UpdateDynamicData(NumStreamedMips, NumLODGroups, Settings, false);
 
-	if (StreamingTexture->bForceFullyLoad) // Somewhat expected at this point.
+	if (StreamingRenderAsset->bForceFullyLoad) // Somewhat expected at this point.
 	{
-		StreamingTexture->WantedMips = StreamingTexture->BudgetedMips = StreamingTexture->MaxAllowedMips;
+		StreamingRenderAsset->WantedMips = StreamingRenderAsset->BudgetedMips = StreamingRenderAsset->MaxAllowedMips;
 	}
 
-	StreamingTexture->StreamWantedMips(*this);
+	StreamingRenderAsset->StreamWantedMips(*this);
 }
 
 /**
@@ -900,42 +932,45 @@ void FStreamingManagerTexture::UpdateIndividualTexture( UTexture2D* Texture )
  * @param StageIndex		Current stage index
  * @param NumUpdateStages	Number of texture update stages
  */
-void FStreamingManagerTexture::UpdateStreamingTextures( int32 StageIndex, int32 NumUpdateStages, bool bWaitForMipFading )
+void FRenderAssetStreamingManager::UpdateStreamingRenderAssets( int32 StageIndex, int32 NumUpdateStages, bool bWaitForMipFading )
 {
 	if ( StageIndex == 0 )
 	{
-		CurrentUpdateStreamingTextureIndex = 0;
-		InflightTextures.Reset();
+		CurrentUpdateStreamingRenderAssetIndex = 0;
+		InflightRenderAssets.Reset();
 	}
 
-	int32 StartIndex = CurrentUpdateStreamingTextureIndex;
-	int32 EndIndex = StreamingTextures.Num() * (StageIndex + 1) / NumUpdateStages;
+	int32 StartIndex = CurrentUpdateStreamingRenderAssetIndex;
+	int32 EndIndex = StreamingRenderAssets.Num() * (StageIndex + 1) / NumUpdateStages;
 	for ( int32 Index=StartIndex; Index < EndIndex; ++Index )
 	{
-		FStreamingTexture& StreamingTexture = StreamingTextures[Index];
-		FPlatformMisc::Prefetch( &StreamingTexture + 1 );
+		FStreamingRenderAsset& StreamingRenderAsset = StreamingRenderAssets[Index];
+		FPlatformMisc::Prefetch( &StreamingRenderAsset + 1 );
 
-		// Is this texture marked for removal? Will get cleanup once the async task is done.
-		if (!StreamingTexture.Texture) continue;
+		// Is this texture/mesh marked for removal? Will get cleanup once the async task is done.
+		if (!StreamingRenderAsset.RenderAsset) continue;
 
-		STAT(int32 PreviousResidentMips = StreamingTexture.ResidentMips;)
+		STAT(int32 PreviousResidentMips = StreamingRenderAsset.ResidentMips;)
 
-		StreamingTexture.UpdateDynamicData(NumStreamedMips, Settings, bWaitForMipFading);
+		const int32* NumStreamedMips;
+		const int32 NumLODGroups = GetNumStreamedMipsArray(StreamingRenderAsset.RenderAssetType, NumStreamedMips);
 
-		// Make a list of each texture that can potentially require additional UpdateStreamingStatus
-		if (StreamingTexture.bInFlight)
+		StreamingRenderAsset.UpdateDynamicData(NumStreamedMips, NumLODGroups, Settings, bWaitForMipFading);
+
+		// Make a list of each texture/mesh that can potentially require additional UpdateStreamingStatus
+		if (StreamingRenderAsset.bInFlight)
 		{
-			InflightTextures.Add(Index);
+			InflightRenderAssets.Add(Index);
 		}
 
 #if STATS
-		if (StreamingTexture.ResidentMips > PreviousResidentMips)
+		if (StreamingRenderAsset.ResidentMips > PreviousResidentMips)
 		{
-			GatheredStats.MipIOBandwidth += StreamingTexture.GetSize(StreamingTexture.ResidentMips) - StreamingTexture.GetSize(PreviousResidentMips);
+			GatheredStats.MipIOBandwidth += StreamingRenderAsset.GetSize(StreamingRenderAsset.ResidentMips) - StreamingRenderAsset.GetSize(PreviousResidentMips);
 		}
 #endif
 	}
-	CurrentUpdateStreamingTextureIndex = EndIndex;
+	CurrentUpdateStreamingRenderAssetIndex = EndIndex;
 }
 
 static TAutoConsoleVariable<int32> CVarTextureStreamingAmortizeCPUToGPUCopy(
@@ -959,23 +994,23 @@ static FORCEINLINE bool ShouldAmortizeMipCopies()
 }
 
 /**
- * Stream textures in/out, based on the priorities calculated by the async work.
+ * Stream textures/meshes in/out, based on the priorities calculated by the async work.
  * @param bProcessEverything	Whether we're processing all textures in one go
  */
-void FStreamingManagerTexture::StreamTextures( bool bProcessEverything )
+void FRenderAssetStreamingManager::StreamRenderAssets( bool bProcessEverything )
 {
-	const FAsyncTextureStreamingTask& AsyncTask = AsyncWork->GetTask();
+	const FRenderAssetStreamingMipCalcTask& AsyncTask = AsyncWork->GetTask();
 
-	// Note that texture indices referred by the async task could be outdated if UpdatePendingStates() was called between the
-	// end of the async task work, and this call to StreamTextures(). This happens when SyncStates(false) is called.
+	// Note that render asset indices referred by the async task could be outdated if UpdatePendingStates() was called between the
+	// end of the async task work, and this call to StreamRenderAssets(). This happens when SyncStates(false) is called.
 
-	if (!bPauseTextureStreaming || bProcessEverything)
+	if (!bPauseRenderAssetStreaming || bProcessEverything)
 	{
-		for (int32 TextureIndex : AsyncTask.GetCancelationRequests())
+		for (int32 AssetIndex : AsyncTask.GetCancelationRequests())
 		{
-			if (StreamingTextures.IsValidIndex(TextureIndex))
+			if (StreamingRenderAssets.IsValidIndex(AssetIndex))
 			{
-				StreamingTextures[TextureIndex].CancelPendingMipChangeRequest();
+				StreamingRenderAssets[AssetIndex].CancelPendingMipChangeRequest();
 			}
 		}
 
@@ -986,47 +1021,47 @@ void FStreamingManagerTexture::StreamTextures( bool bProcessEverything )
 			CurrentPendingMipCopyRequestIdx = 0;
 
 			// Make copies of the requests so that they can be processed later
-			for (int32 TextureIndex : AsyncTask.GetLoadRequests())
+			for (int32 AssetIndex : AsyncTask.GetLoadRequests())
 			{
-				if (StreamingTextures.IsValidIndex(TextureIndex)
-					&& StreamingTextures[TextureIndex].Texture)
+				if (StreamingRenderAssets.IsValidIndex(AssetIndex)
+					&& StreamingRenderAssets[AssetIndex].RenderAsset)
 				{
-					FStreamingTexture& StreamingTexture = StreamingTextures[TextureIndex];
-					StreamingTexture.CacheStreamingMetaData();
-					new (PendingMipCopyRequests) FPendingMipCopyRequest(StreamingTexture.Texture, TextureIndex);
+					FStreamingRenderAsset& StreamingRenderAsset = StreamingRenderAssets[AssetIndex];
+					StreamingRenderAsset.CacheStreamingMetaData();
+					new (PendingMipCopyRequests) FPendingMipCopyRequest(StreamingRenderAsset.RenderAsset, AssetIndex);
 				}
 			}
 		}
 		else
 		{
-			for (int32 TextureIndex : AsyncTask.GetLoadRequests())
+			for (int32 AssetIndex : AsyncTask.GetLoadRequests())
 			{
-				if (StreamingTextures.IsValidIndex(TextureIndex))
+				if (StreamingRenderAssets.IsValidIndex(AssetIndex))
 				{
-					StreamingTextures[TextureIndex].StreamWantedMips(*this);
+					StreamingRenderAssets[AssetIndex].StreamWantedMips(*this);
 				}
 			}
 		}
 	}
 	
-	for (int32 TextureIndex : AsyncTask.GetPendingUpdateDirties())
+	for (int32 AssetIndex : AsyncTask.GetPendingUpdateDirties())
 	{
-		if (StreamingTextures.IsValidIndex(TextureIndex))
+		if (StreamingRenderAssets.IsValidIndex(AssetIndex))
 		{
-			FStreamingTexture& StreamingTexture = StreamingTextures[TextureIndex];
-			const bool bNewState = StreamingTexture.HasUpdatePending(bPauseTextureStreaming, AsyncTask.HasAnyView());
+			FStreamingRenderAsset& StreamingRenderAsset = StreamingRenderAssets[AssetIndex];
+			const bool bNewState = StreamingRenderAsset.HasUpdatePending(bPauseRenderAssetStreaming, AsyncTask.HasAnyView());
 
-			// Always update the texture and the streaming texture together to make sure they are in sync.
-			StreamingTexture.bHasUpdatePending = bNewState;
-			if (StreamingTexture.Texture)
+			// Always update the texture/mesh and the streaming texture/mesh together to make sure they are in sync.
+			StreamingRenderAsset.bHasUpdatePending = bNewState;
+			if (StreamingRenderAsset.RenderAsset)
 			{
-				StreamingTexture.Texture->bHasStreamingUpdatePending = bNewState;
+				StreamingRenderAsset.RenderAsset->bHasStreamingUpdatePending = bNewState;
 			}
 		}
 	}
 }
 
-void FStreamingManagerTexture::ProcessPendingMipCopyRequests()
+void FRenderAssetStreamingManager::ProcessPendingMipCopyRequests()
 {
 	if (!ShouldAmortizeMipCopies())
 	{
@@ -1040,32 +1075,32 @@ void FStreamingManagerTexture::ProcessPendingMipCopyRequests()
 	{
 		const FPendingMipCopyRequest& Request = PendingMipCopyRequests[CurrentPendingMipCopyRequestIdx++];
 
-		if (Request.Texture)
+		if (Request.RenderAsset)
 		{
-			FStreamingTexture* StreamingTexture = nullptr;
+			FStreamingRenderAsset* StreamingRenderAsset = nullptr;
 
-			if (StreamingTextures.IsValidIndex(Request.CachedIdx)
-				&& StreamingTextures[Request.CachedIdx].Texture == Request.Texture)
+			if (StreamingRenderAssets.IsValidIndex(Request.CachedIdx)
+				&& StreamingRenderAssets[Request.CachedIdx].RenderAsset == Request.RenderAsset)
 			{
-				StreamingTexture = &StreamingTextures[Request.CachedIdx];
+				StreamingRenderAsset = &StreamingRenderAssets[Request.CachedIdx];
 			}
-			else if (ReferencedTextures.Contains(Request.Texture))
+			else if (ReferencedRenderAssets.Contains(Request.RenderAsset))
 			{
 				// Texture is still valid but its index has been changed
-				check(StreamingTextures.IsValidIndex(Request.Texture->StreamingIndex));
-				StreamingTexture = &StreamingTextures[Request.Texture->StreamingIndex];
+				check(StreamingRenderAssets.IsValidIndex(Request.RenderAsset->StreamingIndex));
+				StreamingRenderAsset = &StreamingRenderAssets[Request.RenderAsset->StreamingIndex];
 			}
 
-			if (StreamingTexture)
+			if (StreamingRenderAsset)
 			{
-				StreamingTexture->StreamWantedMipsUsingCachedData(*this);
+				StreamingRenderAsset->StreamWantedMipsUsingCachedData(*this);
 				--NumRemainingRequests;
 			}
 		}
 	}
 }
 
-void FStreamingManagerTexture::CheckUserSettings()
+void FRenderAssetStreamingManager::CheckUserSettings()
 {	
 	if (CVarStreamingUseFixedPoolSize.GetValueOnGameThread() == 0)
 	{
@@ -1094,15 +1129,15 @@ void FStreamingManagerTexture::CheckUserSettings()
 	}
 }
 
-void FStreamingManagerTexture::SetLastUpdateTime()
+void FRenderAssetStreamingManager::SetLastUpdateTime()
 {
 	// Update the last update time.
 	float WorldTime = 0;
 
-	for (int32 LevelIndex = 0; LevelIndex < LevelTextureManagers.Num(); ++LevelIndex)
+	for (int32 LevelIndex = 0; LevelIndex < LevelRenderAssetManagers.Num(); ++LevelIndex)
 	{
 		// Update last update time only if there is a reasonable threshold to define visibility.
-		WorldTime = LevelTextureManagers[LevelIndex].GetWorldTime();
+		WorldTime = LevelRenderAssetManagers[LevelIndex].GetWorldTime();
 		if (WorldTime > 0)
 		{
 			break;
@@ -1119,7 +1154,7 @@ void FStreamingManagerTexture::SetLastUpdateTime()
 	}
 }
 
-void FStreamingManagerTexture::UpdateStats()
+void FRenderAssetStreamingManager::UpdateStats()
 {
 	float DeltaStatTime = (float)(GatheredStats.Timestamp - DisplayedStats.Timestamp);
 	if (DeltaStatTime > SMALL_NUMBER)
@@ -1133,7 +1168,7 @@ void FStreamingManagerTexture::UpdateStats()
 	MaxEverRequired = FMath::Max<int64>(MaxEverRequired, DisplayedStats.RequiredPool);
 }
 
-void FStreamingManagerTexture::LogViewLocationChange()
+void FRenderAssetStreamingManager::LogViewLocationChange()
 {
 #if STREAMING_LOG_VIEWCHANGES
 	static bool bWasLocationOveridden = false;
@@ -1168,19 +1203,61 @@ static TAutoConsoleVariable<int32> CVarUseBackgroundThreadPool(
 	1,
 	TEXT("If true, use the background thread pool for mip calculations."));
 
-void FStreamingManagerTexture::UpdateResourceStreaming( float DeltaTime, bool bProcessEverything/*=false*/ )
+class FUpdateStreamingRenderAssetsTask
+{
+	FEvent* CompletionEvent;
+	FRenderAssetStreamingManager* Manager;
+	int32 StageIdx;
+	int32 NumUpdateStages;
+	bool bWaitForMipFading;
+public:
+	FUpdateStreamingRenderAssetsTask(
+		FEvent* InCompletionEvent,
+		FRenderAssetStreamingManager* InManager,
+		int32 InStageIdx,
+		int32 InNumUpdateStages,
+		bool bInWaitForMipFading)
+		: CompletionEvent(InCompletionEvent)
+		, Manager(InManager)
+		, StageIdx(InStageIdx)
+		, NumUpdateStages(InNumUpdateStages)
+		, bWaitForMipFading(bInWaitForMipFading)
+	{
+	}
+	static FORCEINLINE TStatId GetStatId()
+	{
+		RETURN_QUICK_DECLARE_CYCLE_STAT(FUpdateStreamingRenderAssetsTask, STATGROUP_TaskGraphTasks);
+	}
+	static FORCEINLINE ENamedThreads::Type GetDesiredThread()
+	{
+		return ENamedThreads::AnyHiPriThreadHiPriTask;
+	}
+	static FORCEINLINE ESubsequentsMode::Type GetSubsequentsMode()
+	{
+		return ESubsequentsMode::FireAndForget;
+	}
+	void DoTask(ENamedThreads::Type CurrentThread, const FGraphEventRef& MyCompletionGraphEvent)
+	{
+		Manager->UpdateStreamingRenderAssets(StageIdx, NumUpdateStages, bWaitForMipFading);
+		CompletionEvent->Trigger();
+	}
+};
+
+void FRenderAssetStreamingManager::UpdateResourceStreaming( float DeltaTime, bool bProcessEverything/*=false*/ )
 {
 	FScopeLock ScopeLock(&CriticalSection);
 
-	SCOPE_CYCLE_COUNTER(STAT_TextureStreaming_GameThreadUpdateTime);
-	CSV_SCOPED_TIMING_STAT_EXCLUSIVE(TextureStreaming);
+	SCOPE_CYCLE_COUNTER(STAT_RenderAssetStreaming_GameThreadUpdateTime);
+	CSV_SCOPED_TIMING_STAT_EXCLUSIVE(RenderAssetStreaming);
+
+	const bool bUseThreadingForPerf = FApp::ShouldUseThreadingForPerformance();
 
 	LogViewLocationChange();
 	STAT(DisplayedStats.Apply();)
 
-	TextureInstanceAsyncWork->EnsureCompletion();
+	RenderAssetInstanceAsyncWork->EnsureCompletion();
 
-	if (NumTextureProcessingStages <= 0 || bProcessEverything)
+	if (NumRenderAssetProcessingStages <= 0 || bProcessEverything)
 	{
 		if (!AsyncWork->IsDone())
 		{	// Is the AsyncWork is running for some reason? (E.g. we reset the system by simply setting ProcessingStage to 0.)
@@ -1188,17 +1265,17 @@ void FStreamingManagerTexture::UpdateResourceStreaming( float DeltaTime, bool bP
 		}
 
 		ProcessingStage = 0;
-		NumTextureProcessingStages = Settings.FramesForFullUpdate;
+		NumRenderAssetProcessingStages = Settings.FramesForFullUpdate;
 
 		// Update Thread Data
 		SetLastUpdateTime();
-		UpdateStreamingTextures(0, 1, false);
+		UpdateStreamingRenderAssets(0, 1, false);
 
 		UpdatePendingStates(true);
 		PrepareAsyncTask(bProcessEverything || Settings.bStressTest);
 		AsyncWork->StartSynchronousTask();
 
-		StreamTextures(bProcessEverything);
+		StreamRenderAssets(bProcessEverything);
 
 		STAT(GatheredStats.SetupAsyncTaskCycles = 0);
 		STAT(GatheredStats.UpdateStreamingDataCycles = 0);
@@ -1210,7 +1287,7 @@ void FStreamingManagerTexture::UpdateResourceStreaming( float DeltaTime, bool bP
 	{
 		STAT(GatheredStats.SetupAsyncTaskCycles = -(int32)FPlatformTime::Cycles();)
 
-		NumTextureProcessingStages = Settings.FramesForFullUpdate;
+		NumRenderAssetProcessingStages = Settings.FramesForFullUpdate;
 
 		if (!AsyncWork->IsDone())
 		{	// Is the AsyncWork is running for some reason? (E.g. we reset the system by simply setting ProcessingStage to 0.)
@@ -1225,7 +1302,7 @@ void FStreamingManagerTexture::UpdateResourceStreaming( float DeltaTime, bool bP
 
 		STAT(GatheredStats.SetupAsyncTaskCycles += FPlatformTime::Cycles();)
 	}
-	else if (ProcessingStage <= NumTextureProcessingStages)
+	else if (ProcessingStage <= NumRenderAssetProcessingStages)
 	{
 		STAT(int32 StartTime = (int32)FPlatformTime::Cycles();)
 
@@ -1234,10 +1311,32 @@ void FStreamingManagerTexture::UpdateResourceStreaming( float DeltaTime, bool bP
 			SetLastUpdateTime();
 		}
 
-		UpdateStreamingTextures(ProcessingStage - 1, NumTextureProcessingStages, DeltaTime > 0);
-		IncrementalUpdate(1.f / (float)FMath::Max(NumTextureProcessingStages - 1, 1), true); // -1 since we don't want to do anything at stage 0.
+		FEvent* SyncEvent;
+		// Optimization: overlapping UpdateStreamingRenderAssets() and IncrementalUpdate();
+		// Restrict this optimization to platforms tested to have a win;
+		// Platforms tested and results (ave exec time of UpdateResourceStreaming):
+		//   PS4 Pro - from ~0.55 ms/frame to ~0.15 ms/frame
+		const bool bOverlappedExecution = !!(PLATFORM_PS4) && !(WITH_EDITOR) && bUseThreadingForPerf;
+		if (bOverlappedExecution)
+		{
+			SyncEvent = FPlatformProcess::GetSynchEventFromPool(false);
+			check(SyncEvent);
+			TGraphTask<FUpdateStreamingRenderAssetsTask>::CreateTask(nullptr, ENamedThreads::GameThread)
+				.ConstructAndDispatchWhenReady(SyncEvent, this, ProcessingStage - 1, NumRenderAssetProcessingStages, DeltaTime > 0.f);
+		}
+		else
+		{
+			UpdateStreamingRenderAssets(ProcessingStage - 1, NumRenderAssetProcessingStages, DeltaTime > 0.f);
+		}
+
+		IncrementalUpdate(1.f / (float)FMath::Max(NumRenderAssetProcessingStages - 1, 1), true); // -1 since we don't want to do anything at stage 0.
 		++ProcessingStage;
 
+		if (bOverlappedExecution)
+		{
+			SyncEvent->Wait();
+			FPlatformProcess::ReturnSynchEventToPool(SyncEvent);
+		}
 		STAT(GatheredStats.UpdateStreamingDataCycles = FMath::Max<uint32>(ProcessingStage > 2 ? GatheredStats.UpdateStreamingDataCycles : 0, FPlatformTime::Cycles() - StartTime);)
 	}
 	else if (AsyncWork->IsDone())
@@ -1245,15 +1344,15 @@ void FStreamingManagerTexture::UpdateResourceStreaming( float DeltaTime, bool bP
 		STAT(GatheredStats.StreamTexturesCycles = -(int32)FPlatformTime::Cycles();)
 
 		// Since this step is lightweight, tick each texture inflight here, to accelerate the state changes.
-		for (int32 TextureIndex : InflightTextures)
+		for (int32 TextureIndex : InflightRenderAssets)
 		{
-			StreamingTextures[TextureIndex].UpdateStreamingStatus(DeltaTime > 0);
+			StreamingRenderAssets[TextureIndex].UpdateStreamingStatus(DeltaTime > 0);
 		}
 
-		StreamTextures(bProcessEverything);
+		StreamRenderAssets(bProcessEverything);
 		// Release the old view now as the destructors can be expensive. Now only the dynamic manager holds a ref.
 		AsyncWork->GetTask().ReleaseAsyncViews();
-		IncrementalUpdate(1.f / (float)FMath::Max(NumTextureProcessingStages - 1, 1), true); // Just in case continue any pending update.
+		IncrementalUpdate(1.f / (float)FMath::Max(NumRenderAssetProcessingStages - 1, 1), true); // Just in case continue any pending update.
 		DynamicComponentManager.PrepareAsyncView();
 
 		ProcessingStage = 0;
@@ -1267,13 +1366,13 @@ void FStreamingManagerTexture::UpdateResourceStreaming( float DeltaTime, bool bP
 		ProcessPendingMipCopyRequests();
 	}
 
-	if (FApp::ShouldUseThreadingForPerformance())
+	if (bUseThreadingForPerf)
 	{
-		TextureInstanceAsyncWork->StartBackgroundTask(GThreadPool);
+		RenderAssetInstanceAsyncWork->StartBackgroundTask(GThreadPool);
 	}
 	else
 	{
-		TextureInstanceAsyncWork->StartSynchronousTask();
+		RenderAssetInstanceAsyncWork->StartSynchronousTask();
 	}
 }
 
@@ -1284,7 +1383,7 @@ void FStreamingManagerTexture::UpdateResourceStreaming( float DeltaTime, bool bP
  * @param bLogResults	Whether to dump the results to the log.
  * @return				Number of streaming requests still in flight, if the time limit was reached before they were finished.
  */
-int32 FStreamingManagerTexture::BlockTillAllRequestsFinished( float TimeLimit /*= 0.0f*/, bool bLogResults /*= false*/ )
+int32 FRenderAssetStreamingManager::BlockTillAllRequestsFinished( float TimeLimit /*= 0.0f*/, bool bLogResults /*= false*/ )
 {
 	FScopeLock ScopeLock(&CriticalSection);
 
@@ -1294,10 +1393,10 @@ int32 FStreamingManagerTexture::BlockTillAllRequestsFinished( float TimeLimit /*
 	{
 		int32 NumOfInFlights = 0;
 
-		for (FStreamingTexture& StreamingTexture : StreamingTextures)
+		for (FStreamingRenderAsset& StreamingRenderAsset : StreamingRenderAssets)
 		{
-			StreamingTexture.UpdateStreamingStatus(false);
-			if (StreamingTexture.bInFlight)
+			StreamingRenderAsset.UpdateStreamingStatus(false);
+			if (StreamingRenderAsset.bInFlight)
 			{
 				++NumOfInFlights;
 			}
@@ -1320,29 +1419,29 @@ int32 FStreamingManagerTexture::BlockTillAllRequestsFinished( float TimeLimit /*
 	}
 }
 
-void FStreamingManagerTexture::GetObjectReferenceBounds(const UObject* RefObject, TArray<FBox>& AssetBoxes)
+void FRenderAssetStreamingManager::GetObjectReferenceBounds(const UObject* RefObject, TArray<FBox>& AssetBoxes)
 {
 	FScopeLock ScopeLock(&CriticalSection);
 
-	const UTexture2D* Texture2D = Cast<const UTexture2D>(RefObject);
-	if (Texture2D)
+	const UStreamableRenderAsset* RenderAsset = Cast<const UStreamableRenderAsset>(RefObject);
+	if (RenderAsset)
 	{
-		for (FLevelTextureManager& LevelManager : LevelTextureManagers)
+		for (FLevelRenderAssetManager& LevelManager : LevelRenderAssetManagers)
 		{
-			const FTextureInstanceView* View = LevelManager.GetRawAsyncView();
+			const FRenderAssetInstanceView* View = LevelManager.GetRawAsyncView();
 			if (View)
 			{
-				for (auto It = View->GetElementIterator(Texture2D); It; ++It)
+				for (auto It = View->GetElementIterator(RenderAsset); It; ++It)
 				{
 					AssetBoxes.Add(It.GetBounds().GetBox());
 				}
 			}
 		}
 
-		const FTextureInstanceView* View = DynamicComponentManager.GetAsyncView(false);
+		const FRenderAssetInstanceView* View = DynamicComponentManager.GetAsyncView(false);
 		if (View)
 		{
-			for (auto It = View->GetElementIterator(Texture2D); It; ++It)
+			for (auto It = View->GetElementIterator(RenderAsset); It; ++It)
 			{
 				AssetBoxes.Add(It.GetBounds().GetBox());
 			}
@@ -1350,20 +1449,20 @@ void FStreamingManagerTexture::GetObjectReferenceBounds(const UObject* RefObject
 	}
 }
 
-void FStreamingManagerTexture::PropagateLightingScenarioChange()
+void FRenderAssetStreamingManager::PropagateLightingScenarioChange()
 {
 	FScopeLock ScopeLock(&CriticalSection);
 
 	// Note that dynamic components don't need to be handled because their renderstates are updated, which triggers and update.
 	
 	TArray<ULevel*, TInlineAllocator<32> > Levels;
-	for (FLevelTextureManager& LevelManager : LevelTextureManagers)
+	for (FLevelRenderAssetManager& LevelManager : LevelRenderAssetManagers)
 	{
 		Levels.Push(LevelManager.GetLevel());
 		LevelManager.Remove(nullptr);
 	}
 
-	LevelTextureManagers.Empty();
+	LevelRenderAssetManagers.Empty();
 
 	for (ULevel* Level : Levels)
 	{
@@ -1372,7 +1471,7 @@ void FStreamingManagerTexture::PropagateLightingScenarioChange()
 }
 
 #if STATS_FAST
-bool FStreamingManagerTexture::HandleDumpTextureStreamingStatsCommand( const TCHAR* Cmd, FOutputDevice& Ar )
+bool FRenderAssetStreamingManager::HandleDumpTextureStreamingStatsCommand( const TCHAR* Cmd, FOutputDevice& Ar )
 {
 	FScopeLock ScopeLock(&CriticalSection);
 
@@ -1392,7 +1491,7 @@ bool FStreamingManagerTexture::HandleDumpTextureStreamingStatsCommand( const TCH
 #endif // STATS_FAST
 
 #if STATS
-bool FStreamingManagerTexture::HandleListStreamingTexturesCommand( const TCHAR* Cmd, FOutputDevice& Ar )
+bool FRenderAssetStreamingManager::HandleListStreamingRenderAssetsCommand( const TCHAR* Cmd, FOutputDevice& Ar )
 {
 	FScopeLock ScopeLock(&CriticalSection);
 
@@ -1400,51 +1499,78 @@ bool FStreamingManagerTexture::HandleListStreamingTexturesCommand( const TCHAR*
 
 	const bool bShouldOnlyListUnkownRef = FParse::Command(&Cmd, TEXT("UNKOWNREF"));
 
-	// Sort texture by names so that the state can be compared between runs.
-	TMap<FString, int32> SortedTextures;
-	for ( int32 TextureIndex=0; TextureIndex < StreamingTextures.Num(); ++TextureIndex )
+	// Sort texture/mesh by names so that the state can be compared between runs.
+	TMap<FString, int32> SortedRenderAssets;
+	for ( int32 Idx=0; Idx < StreamingRenderAssets.Num(); ++Idx)
 	{
-		const FStreamingTexture& StreamingTexture = StreamingTextures[TextureIndex];
-		if (!StreamingTexture.Texture) continue;
-		if (bShouldOnlyListUnkownRef && !StreamingTexture.bUseUnkownRefHeuristic) continue;
+		const FStreamingRenderAsset& StreamingRenderAsset = StreamingRenderAssets[Idx];
+		if (!StreamingRenderAsset.RenderAsset) continue;
+		if (bShouldOnlyListUnkownRef && !StreamingRenderAsset.bUseUnkownRefHeuristic) continue;
 
-		SortedTextures.Add(StreamingTexture.Texture->GetFullName(), TextureIndex);
+		SortedRenderAssets.Add(StreamingRenderAsset.RenderAsset->GetFullName(), Idx);
 	}
 
-	SortedTextures.KeySort(TLess<FString>());
+	SortedRenderAssets.KeySort(TLess<FString>());
 
-	for (TMap<FString, int32>::TConstIterator It(SortedTextures); It; ++It)
+	for (TMap<FString, int32>::TConstIterator It(SortedRenderAssets); It; ++It)
 	{
-		const FStreamingTexture& StreamingTexture = StreamingTextures[It.Value()];
-
-		const UTexture2D* Texture2D = StreamingTexture.Texture;
-		UE_LOG(LogContentStreaming, Log,  TEXT("Texture [%d] : %s"), It.Value(), *Texture2D->GetFullName() );
+		const FStreamingRenderAsset& StreamingRenderAsset = StreamingRenderAssets[It.Value()];
+		const UStreamableRenderAsset* RenderAsset = StreamingRenderAsset.RenderAsset;
+		const typename FStreamingRenderAsset::EAssetType AssetType = StreamingRenderAsset.RenderAssetType;
+		
+		UE_LOG(LogContentStreaming, Log,  TEXT("%s [%d] : %s"),
+			FStreamingRenderAsset::GetStreamingAssetTypeStr(AssetType),
+			It.Value(),
+			*RenderAsset->GetFullName() );
 
-		int32 CurrentMipIndex = FMath::Max(Texture2D->GetNumMips() - StreamingTexture.ResidentMips, 0);
-		int32 WantedMipIndex = FMath::Max(Texture2D->GetNumMips() - StreamingTexture.GetPerfectWantedMips(), 0);
-		int32 MaxAllowedMipIndex = FMath::Max(Texture2D->GetNumMips() - StreamingTexture.MaxAllowedMips, 0);
-		const TIndirectArray<struct FTexture2DMipMap>& Mips = Texture2D->PlatformData->Mips;
+		int32 CurrentMipIndex = FMath::Max(RenderAsset->GetNumMipsForStreaming() - StreamingRenderAsset.ResidentMips, 0);
+		int32 WantedMipIndex = FMath::Max(RenderAsset->GetNumMipsForStreaming() - StreamingRenderAsset.GetPerfectWantedMips(), 0);
+		int32 MaxAllowedMipIndex = FMath::Max(RenderAsset->GetNumMipsForStreaming() - StreamingRenderAsset.MaxAllowedMips, 0);
 
-		if (StreamingTexture.LastRenderTime != MAX_FLT)
+		if (AssetType == FStreamingRenderAsset::AT_Texture)
 		{
-			UE_LOG(LogContentStreaming, Log,  TEXT("    Current=%dx%d Wanted=%dx%d MaxAllowed=%dx%d LastRenderTime=%.3f BudgetBias=%d Group=%s"), 
-				Mips[CurrentMipIndex].SizeX, Mips[CurrentMipIndex].SizeY, 
-				Mips[WantedMipIndex].SizeX, Mips[WantedMipIndex].SizeY, 
-				Mips[MaxAllowedMipIndex].SizeX, Mips[MaxAllowedMipIndex].SizeY,
-				StreamingTexture.LastRenderTime,
-				StreamingTexture.BudgetMipBias,
-				UTexture::GetTextureGroupString(StreamingTexture.LODGroup));
+			const UTexture2D* Texture = CastChecked<UTexture2D>(RenderAsset);
+			const TIndirectArray<struct FTexture2DMipMap>& Mips = Texture->PlatformData->Mips;
+
+			if (StreamingRenderAsset.LastRenderTime != MAX_FLT)
+			{
+				UE_LOG(LogContentStreaming, Log, TEXT("    Current=%dx%d Wanted=%dx%d MaxAllowed=%dx%d LastRenderTime=%.3f BudgetBias=%d Group=%s"),
+					Mips[CurrentMipIndex].SizeX, Mips[CurrentMipIndex].SizeY,
+					Mips[WantedMipIndex].SizeX, Mips[WantedMipIndex].SizeY,
+					Mips[MaxAllowedMipIndex].SizeX, Mips[MaxAllowedMipIndex].SizeY,
+					StreamingRenderAsset.LastRenderTime,
+					StreamingRenderAsset.BudgetMipBias,
+					UTexture::GetTextureGroupString(static_cast<TextureGroup>(StreamingRenderAsset.LODGroup)));
+			}
+			else
+			{
+				UE_LOG(LogContentStreaming, Log, TEXT("    Current=%dx%d Wanted=%dx%d MaxAllowed=%dx%d BudgetBias=%d Group=%s"),
+					Mips[CurrentMipIndex].SizeX, Mips[CurrentMipIndex].SizeY,
+					Mips[WantedMipIndex].SizeX, Mips[WantedMipIndex].SizeY,
+					Mips[MaxAllowedMipIndex].SizeX, Mips[MaxAllowedMipIndex].SizeY,
+					StreamingRenderAsset.BudgetMipBias,
+					UTexture::GetTextureGroupString(static_cast<TextureGroup>(StreamingRenderAsset.LODGroup)));
+			}
 		}
 		else
 		{
-			UE_LOG(LogContentStreaming, Log,  TEXT("    Current=%dx%d Wanted=%dx%d MaxAllowed=%dx%d BudgetBias=%d Group=%s"), 
-				Mips[CurrentMipIndex].SizeX, Mips[CurrentMipIndex].SizeY, 
-				Mips[WantedMipIndex].SizeX, Mips[WantedMipIndex].SizeY, 
-				Mips[MaxAllowedMipIndex].SizeX, Mips[MaxAllowedMipIndex].SizeY,
-				StreamingTexture.BudgetMipBias,
-				UTexture::GetTextureGroupString(StreamingTexture.LODGroup));
+			const float LastRenderTime = StreamingRenderAsset.LastRenderTime;
+			const UStaticMesh* StaticMesh = Cast<UStaticMesh>(RenderAsset);
+			FString LODGroupName = TEXT("Unknown");
+#if WITH_EDITORONLY_DATA
+			if (StaticMesh)
+			{
+				LODGroupName = StaticMesh->LODGroup.ToString();
+			}
+#endif
+			UE_LOG(LogContentStreaming, Log, TEXT("    CurrentLOD=%d WantedLOD=%d MaxAllowedLOD=%d LastRenderTime=%s BudgetBias=%d Group=%s"),
+				CurrentMipIndex,
+				WantedMipIndex,
+				MaxAllowedMipIndex,
+				LastRenderTime == MAX_FLT ? TEXT("NotTracked") : *FString::Printf(TEXT("%.3f"), LastRenderTime),
+				StreamingRenderAsset.BudgetMipBias,
+				*LODGroupName);
 		}
-
 	}
 	return true;
 }
@@ -1452,7 +1578,7 @@ bool FStreamingManagerTexture::HandleListStreamingTexturesCommand( const TCHAR*
 
 #if !UE_BUILD_SHIPPING
 
-bool FStreamingManagerTexture::HandleResetMaxEverRequiredTexturesCommand(const TCHAR* Cmd, FOutputDevice& Ar)
+bool FRenderAssetStreamingManager::HandleResetMaxEverRequiredRenderAssetMemoryCommand(const TCHAR* Cmd, FOutputDevice& Ar)
 {
 	FScopeLock ScopeLock(&CriticalSection);
 
@@ -1461,7 +1587,7 @@ bool FStreamingManagerTexture::HandleResetMaxEverRequiredTexturesCommand(const T
 	return true;
 }
 
-bool FStreamingManagerTexture::HandleLightmapStreamingFactorCommand( const TCHAR* Cmd, FOutputDevice& Ar )
+bool FRenderAssetStreamingManager::HandleLightmapStreamingFactorCommand( const TCHAR* Cmd, FOutputDevice& Ar )
 {
 	FScopeLock ScopeLock(&CriticalSection);
 
@@ -1475,15 +1601,17 @@ bool FStreamingManagerTexture::HandleLightmapStreamingFactorCommand( const TCHAR
 	return true;
 }
 
-bool FStreamingManagerTexture::HandleCancelTextureStreamingCommand( const TCHAR* Cmd, FOutputDevice& Ar )
+bool FRenderAssetStreamingManager::HandleCancelRenderAssetStreamingCommand( const TCHAR* Cmd, FOutputDevice& Ar )
 {
 	FScopeLock ScopeLock(&CriticalSection);
 
 	UTexture2D::CancelPendingTextureStreaming();
+	UStaticMesh::CancelAllPendingStreamingActions();
+	// TODO: USkeletalMesh
 	return true;
 }
 
-bool FStreamingManagerTexture::HandleShadowmapStreamingFactorCommand( const TCHAR* Cmd, FOutputDevice& Ar )
+bool FRenderAssetStreamingManager::HandleShadowmapStreamingFactorCommand( const TCHAR* Cmd, FOutputDevice& Ar )
 {
 	FScopeLock ScopeLock(&CriticalSection);
 
@@ -1497,15 +1625,16 @@ bool FStreamingManagerTexture::HandleShadowmapStreamingFactorCommand( const TCHA
 	return true;
 }
 
-bool FStreamingManagerTexture::HandleNumStreamedMipsCommand( const TCHAR* Cmd, FOutputDevice& Ar )
+bool FRenderAssetStreamingManager::HandleNumStreamedMipsCommand( const TCHAR* Cmd, FOutputDevice& Ar )
 {
 	FScopeLock ScopeLock(&CriticalSection);
 
 	FString NumTextureString(FParse::Token(Cmd, 0));
 	FString NumMipsString(FParse::Token(Cmd, 0));
+	FString LODGroupType(FParse::Token(Cmd, false));
 	int32 LODGroup = ( NumTextureString.Len() > 0 ) ? FCString::Atoi(*NumTextureString) : MAX_int32;
 	int32 NumMips = ( NumMipsString.Len() > 0 ) ? FCString::Atoi(*NumMipsString) : MAX_int32;
-	if ( LODGroup >= 0 && LODGroup < TEXTUREGROUP_MAX )
+	if ((LODGroupType == TEXT("") || LODGroupType == TEXT("Texture")) && LODGroup >= 0 && LODGroup < TEXTUREGROUP_MAX)
 	{
 		FTextureLODGroup& TexGroup = UDeviceProfileManager::Get().GetActiveProfile()->GetTextureLODSettings()->GetTextureLODGroup(TextureGroup(LODGroup));
 		if ( NumMips >= -1 && NumMips <= MAX_TEXTURE_MIP_COUNT )
@@ -1514,32 +1643,42 @@ bool FStreamingManagerTexture::HandleNumStreamedMipsCommand( const TCHAR* Cmd, F
 		}
 		Ar.Logf( TEXT("%s.NumStreamedMips = %d"), UTexture::GetTextureGroupString(TextureGroup(LODGroup)), TexGroup.NumStreamedMips );
 	}
+	else if (LODGroupType == TEXT("StaticMesh"))
+	{
+		// TODO
+		Ar.Logf(TEXT("NumStreamedMips command is not implemented for static mesh yet"));
+	}
+	else if (LODGroupType == TEXT("SkeletalMesh"))
+	{
+		// TODO
+		Ar.Logf(TEXT("NumStreamedMips command is not implemented for skeletal mesh yet"));
+	}
 	else
 	{
-		Ar.Logf( TEXT("Usage: NumStreamedMips TextureGroupIndex <N>") );
+		Ar.Logf( TEXT("Usage: NumStreamedMips LODGroupIndex <N> [Texture|StaticMesh|SkeletalMesh]") );
 	}
 	return true;
 }
 
-bool FStreamingManagerTexture::HandleTrackTextureCommand( const TCHAR* Cmd, FOutputDevice& Ar )
+bool FRenderAssetStreamingManager::HandleTrackRenderAssetCommand( const TCHAR* Cmd, FOutputDevice& Ar )
 {
 	FScopeLock ScopeLock(&CriticalSection);
 
-	FString TextureName(FParse::Token(Cmd, 0));
-	if ( TrackTexture( TextureName ) )
+	FString AssetName(FParse::Token(Cmd, 0));
+	if ( TrackRenderAsset(AssetName) )
 	{
-		Ar.Logf( TEXT("Textures containing \"%s\" are now tracked."), *TextureName );
+		Ar.Logf(TEXT("Textures or meshes containing \"%s\" are now tracked."), *AssetName);
 	}
 	return true;
 }
 
-bool FStreamingManagerTexture::HandleListTrackedTexturesCommand( const TCHAR* Cmd, FOutputDevice& Ar )
+bool FRenderAssetStreamingManager::HandleListTrackedRenderAssetsCommand( const TCHAR* Cmd, FOutputDevice& Ar )
 {
 	FScopeLock ScopeLock(&CriticalSection);
 
-	FString NumTextureString(FParse::Token(Cmd, 0));
-	int32 NumTextures = ( NumTextureString.Len() > 0 ) ? FCString::Atoi(*NumTextureString) : -1;
-	ListTrackedTextures( Ar, NumTextures );
+	FString NumAssetString(FParse::Token(Cmd, 0));
+	int32 NumAssets = (NumAssetString.Len() > 0) ? FCString::Atoi(*NumAssetString) : -1;
+	ListTrackedRenderAssets(Ar, NumAssets);
 	return true;
 }
 
@@ -1548,18 +1687,18 @@ FORCEINLINE float SqrtKeepMax(float V)
 	return V == FLT_MAX ? FLT_MAX : FMath::Sqrt(V);
 }
 
-bool FStreamingManagerTexture::HandleDebugTrackedTexturesCommand( const TCHAR* Cmd, FOutputDevice& Ar )
+bool FRenderAssetStreamingManager::HandleDebugTrackedRenderAssetsCommand( const TCHAR* Cmd, FOutputDevice& Ar )
 {
 	FScopeLock ScopeLock(&CriticalSection);
 
-	// The ENABLE_TEXTURE_TRACKING macro is defined in ContentStreaming.cpp and not available here. This code does not compile any more.
+	// The ENABLE_RENDER_ASSET_TRACKING macro is defined in ContentStreaming.cpp and not available here. This code does not compile any more.
 #ifdef ENABLE_TEXTURE_TRACKING_BROKEN
 	int32 NumTrackedTextures = GTrackedTextureNames.Num();
 	if ( NumTrackedTextures )
 	{
 		for (int32 StreamingIndex = 0; StreamingIndex < StreamingTextures.Num(); ++StreamingIndex)
 		{
-			FStreamingTexture& StreamingTexture = StreamingTextures[StreamingIndex];
+			FStreamingRenderAsset& StreamingTexture = StreamingTextures[StreamingIndex];
 			if (StreamingTexture.Texture)
 			{
 				// See if it matches any of the texture names that we're tracking.
@@ -1621,7 +1760,7 @@ bool FStreamingManagerTexture::HandleDebugTrackedTexturesCommand( const TCHAR* C
 
 						for ( int32 LevelIndex=0; LevelIndex < ThreadSettings.LevelData.Num(); ++LevelIndex )
 						{
-							FStreamingManagerTexture::FThreadLevelData& LevelData = ThreadSettings.LevelData[ LevelIndex ].Value;
+							FRenderAssetStreamingManager::FThreadLevelData& LevelData = ThreadSettings.LevelData[ LevelIndex ].Value;
 							TArray<FStreamableTextureInstance4>* TextureInstances = LevelData.ThreadTextureInstances.Find( StreamingTexture.Texture );
 							if ( TextureInstances )
 							{
@@ -1649,24 +1788,24 @@ bool FStreamingManagerTexture::HandleDebugTrackedTexturesCommand( const TCHAR* C
 			}
 		}
 	}
-#endif // ENABLE_TEXTURE_TRACKING
+#endif // ENABLE_RENDER_ASSET_TRACKING
 
 	return true;
 }
 
-bool FStreamingManagerTexture::HandleUntrackTextureCommand( const TCHAR* Cmd, FOutputDevice& Ar )
+bool FRenderAssetStreamingManager::HandleUntrackRenderAssetCommand( const TCHAR* Cmd, FOutputDevice& Ar )
 {
 	FScopeLock ScopeLock(&CriticalSection);
 
-	FString TextureName(FParse::Token(Cmd, 0));
-	if ( UntrackTexture( TextureName ) )
+	FString AssetName(FParse::Token(Cmd, 0));
+	if (UntrackRenderAsset(AssetName))
 	{
-		Ar.Logf( TEXT("Textures containing \"%s\" are no longer tracked."), *TextureName );
+		Ar.Logf(TEXT("Textures or meshes containing \"%s\" are no longer tracked."), *AssetName);
 	}
 	return true;
 }
 
-bool FStreamingManagerTexture::HandleStreamOutCommand( const TCHAR* Cmd, FOutputDevice& Ar )
+bool FRenderAssetStreamingManager::HandleStreamOutCommand( const TCHAR* Cmd, FOutputDevice& Ar )
 {
 	FScopeLock ScopeLock(&CriticalSection);
 
@@ -1674,8 +1813,8 @@ bool FStreamingManagerTexture::HandleStreamOutCommand( const TCHAR* Cmd, FOutput
 	int64 FreeMB = (Parameter.Len() > 0) ? FCString::Atoi(*Parameter) : 0;
 	if ( FreeMB > 0 )
 	{
-		bool bSucceeded = StreamOutTextureData( FreeMB * 1024 * 1024 );
-		Ar.Logf( TEXT("Tried to stream out %llu MB of texture data: %s"), FreeMB, bSucceeded ? TEXT("Succeeded") : TEXT("Failed") );
+		bool bSucceeded = StreamOutRenderAssetData( FreeMB * 1024 * 1024 );
+		Ar.Logf( TEXT("Tried to stream out %llu MB of texture/mesh data: %s"), FreeMB, bSucceeded ? TEXT("Succeeded") : TEXT("Failed") );
 	}
 	else
 	{
@@ -1684,29 +1823,31 @@ bool FStreamingManagerTexture::HandleStreamOutCommand( const TCHAR* Cmd, FOutput
 	return true;
 }
 
-bool FStreamingManagerTexture::HandlePauseTextureStreamingCommand( const TCHAR* Cmd, FOutputDevice& Ar )
+bool FRenderAssetStreamingManager::HandlePauseRenderAssetStreamingCommand( const TCHAR* Cmd, FOutputDevice& Ar )
 {
 	FScopeLock ScopeLock(&CriticalSection);
 
-	bPauseTextureStreaming = !bPauseTextureStreaming;
-	Ar.Logf( TEXT("Texture streaming is now \"%s\"."), bPauseTextureStreaming ? TEXT("PAUSED") : TEXT("UNPAUSED") );
+	bPauseRenderAssetStreaming = !bPauseRenderAssetStreaming;
+	Ar.Logf( TEXT("Render asset streaming is now \"%s\"."), bPauseRenderAssetStreaming ? TEXT("PAUSED") : TEXT("UNPAUSED") );
 	return true;
 }
 
-bool FStreamingManagerTexture::HandleStreamingManagerMemoryCommand( const TCHAR* Cmd, FOutputDevice& Ar, UWorld* InWorld )
+bool FRenderAssetStreamingManager::HandleStreamingManagerMemoryCommand( const TCHAR* Cmd, FOutputDevice& Ar, UWorld* InWorld )
 {
 	FScopeLock ScopeLock(&CriticalSection);
 
 	SyncStates(true);
 
-	uint32 MemSize = sizeof(FStreamingManagerTexture);
-	MemSize += StreamingTextures.GetAllocatedSize();
+	uint32 MemSize = sizeof(FRenderAssetStreamingManager);
+	MemSize += StreamingRenderAssets.GetAllocatedSize();
 	MemSize += DynamicComponentManager.GetAllocatedSize();
-	MemSize += PendingStreamingTextures.GetAllocatedSize() + RemovedTextureIndices.GetAllocatedSize();
-	MemSize += LevelTextureManagers.GetAllocatedSize();
+	MemSize += PendingStreamingRenderAssets.GetAllocatedSize()
+		+ PendingStreamingRenderAssetTypes.GetAllocatedSize()
+		+ RemovedRenderAssetIndices.GetAllocatedSize();
+	MemSize += LevelRenderAssetManagers.GetAllocatedSize();
 	MemSize += AsyncWork->GetTask().StreamingData.GetAllocatedSize();
 
-	for (const FLevelTextureManager& LevelManager : LevelTextureManagers)
+	for (const FLevelRenderAssetManager& LevelManager : LevelRenderAssetManagers)
 	{
 		MemSize += LevelManager.GetAllocatedSize();
 	}
@@ -1716,88 +1857,110 @@ bool FStreamingManagerTexture::HandleStreamingManagerMemoryCommand( const TCHAR*
 	return true;
 }
 
-bool FStreamingManagerTexture::HandleTextureGroupsCommand( const TCHAR* Cmd, FOutputDevice& Ar )
+bool FRenderAssetStreamingManager::HandleLODGroupsCommand( const TCHAR* Cmd, FOutputDevice& Ar )
 {
 	bDetailedDumpTextureGroupStats = FParse::Param(Cmd, TEXT("Detailed"));
 	bTriggerDumpTextureGroupStats = true;
+	// TODO: mesh LOD groups
 	return true;
 }
 
-bool FStreamingManagerTexture::HandleInvestigateTextureCommand( const TCHAR* Cmd, FOutputDevice& Ar, UWorld* InWorld )
+bool FRenderAssetStreamingManager::HandleInvestigateRenderAssetCommand(const TCHAR* Cmd, FOutputDevice& Ar, UWorld* InWorld)
 {
 	FScopeLock ScopeLock(&CriticalSection);
 
 	SyncStates(true);
 
-	FString InvestigateTextureName(FParse::Token(Cmd, 0));
-	if ( InvestigateTextureName.Len() )
+	FString InvestigateAssetName(FParse::Token(Cmd, 0));
+	if (InvestigateAssetName.Len())
 	{
-		FAsyncTextureStreamingData& StreamingData = AsyncWork->GetTask().StreamingData;
-		StreamingData.Init(CurrentViewInfos, LastWorldUpdateTime, LevelTextureManagers, DynamicComponentManager);
+		FAsyncRenderAssetStreamingData& StreamingData = AsyncWork->GetTask().StreamingData;
+		StreamingData.Init(CurrentViewInfos, LastWorldUpdateTime, LevelRenderAssetManagers, DynamicComponentManager);
+		StreamingData.ComputeViewInfoExtras(Settings);
 		StreamingData.UpdateBoundSizes_Async(Settings);
 
-		for ( int32 TextureIndex=0; TextureIndex < StreamingTextures.Num(); ++TextureIndex )
+		for (int32 AssetIndex = 0; AssetIndex < StreamingRenderAssets.Num(); ++AssetIndex)
 		{
-			FStreamingTexture& StreamingTexture = StreamingTextures[TextureIndex];
-			FString TextureName = StreamingTexture.Texture->GetFullName();
-			if (TextureName.Contains(InvestigateTextureName))
+			FStreamingRenderAsset& StreamingRenderAsset = StreamingRenderAssets[AssetIndex];
+			FString AssetName = StreamingRenderAsset.RenderAsset->GetFullName();
+			if (AssetName.Contains(InvestigateAssetName))
 			{
-				UTexture2D* Texture2D = StreamingTexture.Texture;
-				int32 CurrentMipIndex = FMath::Max(Texture2D->GetNumMips() - StreamingTexture.ResidentMips, 0);
-				int32 WantedMipIndex = FMath::Max(Texture2D->GetNumMips() - StreamingTexture.GetPerfectWantedMips(), 0);
-				int32 MaxMipIndex = FMath::Max(Texture2D->GetNumMips() - StreamingTexture.MaxAllowedMips, 0);
-
-				UE_LOG(LogContentStreaming, Log,  TEXT("Texture: %s"), *TextureName );
-				UE_LOG(LogContentStreaming, Log,  TEXT("  Texture group:   %s"), UTexture::GetTextureGroupString(StreamingTexture.LODGroup) );
+				UStreamableRenderAsset* RenderAsset = StreamingRenderAsset.RenderAsset;
+				if (!RenderAsset) continue;
+				const typename FStreamingRenderAsset::EAssetType AssetType = StreamingRenderAsset.RenderAssetType;
+				UTexture2D* Texture2D = Cast<UTexture2D>(RenderAsset);
+				UStaticMesh* StaticMesh = Cast<UStaticMesh>(RenderAsset);
+				int32 CurrentMipIndex = FMath::Max(RenderAsset->GetNumMipsForStreaming() - StreamingRenderAsset.ResidentMips, 0);
+				int32 WantedMipIndex = FMath::Max(RenderAsset->GetNumMipsForStreaming() - StreamingRenderAsset.GetPerfectWantedMips(), 0);
+				int32 MaxMipIndex = FMath::Max(RenderAsset->GetNumMipsForStreaming() - StreamingRenderAsset.MaxAllowedMips, 0);
+
+				UE_LOG(LogContentStreaming, Log, TEXT("%s: %s"), FStreamingRenderAsset::GetStreamingAssetTypeStr(AssetType), *AssetName);
+				FString LODGroupName = Texture2D ? UTexture::GetTextureGroupString((TextureGroup)StreamingRenderAsset.LODGroup) : TEXT("Unknown");
+#if WITH_EDITORONLY_DATA
+				if (StaticMesh)
+				{
+					LODGroupName = StaticMesh->LODGroup.ToString();
+				}
+#endif
+				UE_LOG(LogContentStreaming, Log, TEXT("  LOD group:   %s"), *LODGroupName);
 
-				if ( Texture2D->bGlobalForceMipLevelsToBeResident )
+				if (RenderAsset->bGlobalForceMipLevelsToBeResident)
 				{
-					UE_LOG(LogContentStreaming, Log,  TEXT("  Force all mips:  bGlobalForceMipLevelsToBeResident") );
+					UE_LOG(LogContentStreaming, Log, TEXT("  Force all mips:  bGlobalForceMipLevelsToBeResident"));
 				}
-				else if ( Texture2D->bForceMiplevelsToBeResident )
+				else if (RenderAsset->bForceMiplevelsToBeResident)
 				{
-					UE_LOG(LogContentStreaming, Log,  TEXT("  Force all mips:  bForceMiplevelsToBeResident") );
+					UE_LOG(LogContentStreaming, Log, TEXT("  Force all mips:  bForceMiplevelsToBeResident"));
 				}
-				else if ( Texture2D->ShouldMipLevelsBeForcedResident() )
+				else if (RenderAsset->ShouldMipLevelsBeForcedResident())
 				{
-					float TimeLeft = (float)(Texture2D->ForceMipLevelsToBeResidentTimestamp - FApp::GetCurrentTime());
-					UE_LOG(LogContentStreaming, Log,  TEXT("  Force all mips:  %.1f seconds left"), FMath::Max(TimeLeft,0.0f) );
+					float TimeLeft = (float)(RenderAsset->ForceMipLevelsToBeResidentTimestamp - FApp::GetCurrentTime());
+					UE_LOG(LogContentStreaming, Log, TEXT("  Force all mips:  %.1f seconds left"), FMath::Max(TimeLeft, 0.0f));
 				}
-				else if ( StreamingTexture.bForceFullyLoadHeuristic )
+				else if (StreamingRenderAsset.bForceFullyLoadHeuristic)
 				{
-					UE_LOG(LogContentStreaming, Log,  TEXT("  Force all mips: bForceFullyLoad") );
+					UE_LOG(LogContentStreaming, Log, TEXT("  Force all mips: bForceFullyLoad"));
 				}
-				else if ( StreamingTexture.MipCount == 1 )
+				else if (StreamingRenderAsset.MipCount == 1)
+				{
+					UE_LOG(LogContentStreaming, Log, TEXT("  Force all mips:  No mip-maps"));
+				}
+				
+				if (Texture2D)
+				{
+					UE_LOG(LogContentStreaming, Log, TEXT("  Current size [Mips]: %dx%d [%d]"), Texture2D->PlatformData->Mips[CurrentMipIndex].SizeX, Texture2D->PlatformData->Mips[CurrentMipIndex].SizeY, StreamingRenderAsset.ResidentMips);
+					UE_LOG(LogContentStreaming, Log, TEXT("  Wanted size [Mips]:  %dx%d [%d]"), Texture2D->PlatformData->Mips[WantedMipIndex].SizeX, Texture2D->PlatformData->Mips[WantedMipIndex].SizeY, StreamingRenderAsset.GetPerfectWantedMips());
+				}
+				else
 				{
-					UE_LOG(LogContentStreaming, Log,  TEXT("  Force all mips:  No mip-maps") );
+					UE_LOG(LogContentStreaming, Log, TEXT("  Current LOD index: %d"), CurrentMipIndex);
+					UE_LOG(LogContentStreaming, Log, TEXT("  Wanted LOD index: %d"), WantedMipIndex);
 				}
-				UE_LOG(LogContentStreaming, Log,  TEXT("  Current size [Mips]: %dx%d [%d]"), Texture2D->PlatformData->Mips[CurrentMipIndex].SizeX, Texture2D->PlatformData->Mips[CurrentMipIndex].SizeY, StreamingTexture.ResidentMips );
-				UE_LOG(LogContentStreaming, Log,  TEXT("  Wanted size [Mips]:  %dx%d [%d]"), Texture2D->PlatformData->Mips[WantedMipIndex].SizeX, Texture2D->PlatformData->Mips[WantedMipIndex].SizeY, StreamingTexture.GetPerfectWantedMips() );
-				UE_LOG(LogContentStreaming, Log,  TEXT("  Allowed mips:        %d-%d"), StreamingTexture.MinAllowedMips, StreamingTexture.MaxAllowedMips );
-				UE_LOG(LogContentStreaming, Log,  TEXT("  LoadOrder Priority:  %d"), StreamingTexture.LoadOrderPriority );
-				UE_LOG(LogContentStreaming, Log,  TEXT("  Retention Priority:  %d"), StreamingTexture.RetentionPriority );
-				UE_LOG(LogContentStreaming, Log,  TEXT("  Boost factor:        %.1f"), StreamingTexture.BoostFactor );
-				UE_LOG(LogContentStreaming, Log,  TEXT("  Mip bias [Budget]:   %d [%d]"), Texture2D->PlatformData->Mips.Num() - StreamingTexture.MaxAllowedMips, StreamingTexture.BudgetMipBias + Settings.GlobalMipBias );
+				UE_LOG(LogContentStreaming, Log, TEXT("  Allowed mips:        %d-%d"), StreamingRenderAsset.MinAllowedMips, StreamingRenderAsset.MaxAllowedMips);
+				UE_LOG(LogContentStreaming, Log, TEXT("  LoadOrder Priority:  %d"), StreamingRenderAsset.LoadOrderPriority);
+				UE_LOG(LogContentStreaming, Log, TEXT("  Retention Priority:  %d"), StreamingRenderAsset.RetentionPriority);
+				UE_LOG(LogContentStreaming, Log, TEXT("  Boost factor:        %.1f"), StreamingRenderAsset.BoostFactor);
+				UE_LOG(LogContentStreaming, Log, TEXT("  Mip bias [Budget]:   %d [%d]"), StreamingRenderAsset.MipCount - StreamingRenderAsset.MaxAllowedMips, StreamingRenderAsset.BudgetMipBias + (Settings.bUsePerTextureBias ? 0 : Settings.GlobalMipBias));
 
 				if (InWorld && !GIsEditor)
 				{
-					UE_LOG(LogContentStreaming, Log,  TEXT("  Time: World=%.3f LastUpdate=%.3f "), InWorld->GetTimeSeconds(), LastWorldUpdateTime);
+					UE_LOG(LogContentStreaming, Log, TEXT("  Time: World=%.3f LastUpdate=%.3f "), InWorld->GetTimeSeconds(), LastWorldUpdateTime);
 				}
 
-				for( int32 ViewIndex=0; ViewIndex < StreamingData.GetViewInfos().Num(); ViewIndex++ )
+				for (int32 ViewIndex = 0; ViewIndex < StreamingData.GetViewInfos().Num(); ViewIndex++)
 				{
 					// Calculate distance of viewer to bounding sphere.
 					const FStreamingViewInfo& ViewInfo = StreamingData.GetViewInfos()[ViewIndex];
-					UE_LOG(LogContentStreaming, Log,  TEXT("  View%d: Position=(%s) ScreenSize=%f MaxEffectiveScreenSize=%f Boost=%f"), ViewIndex, *ViewInfo.ViewOrigin.ToString(), ViewInfo.ScreenSize, Settings.MaxEffectiveScreenSize, ViewInfo.BoostFactor);
+					UE_LOG(LogContentStreaming, Log, TEXT("  View%d: Position=(%s) ScreenSize=%f MaxEffectiveScreenSize=%f Boost=%f"), ViewIndex, *ViewInfo.ViewOrigin.ToString(), ViewInfo.ScreenSize, Settings.MaxEffectiveScreenSize, ViewInfo.BoostFactor);
 				}
 
-				StreamingData.UpdatePerfectWantedMips_Async(StreamingTexture, Settings, true);
+				StreamingData.UpdatePerfectWantedMips_Async(StreamingRenderAsset, Settings, true);
 			}
 		}
 	}
 	else
 	{
-		Ar.Logf( TEXT("Usage: InvestigateTexture <name>") );
+		Ar.Logf(TEXT("Usage: InvestigateTexture <name>"));
 	}
 	return true;
 }
@@ -1809,32 +1972,36 @@ bool FStreamingManagerTexture::HandleInvestigateTextureCommand( const TCHAR* Cmd
  * @param Ar		Output device for feedback
  * @return			true if the command was handled
  */
-bool FStreamingManagerTexture::Exec( UWorld* InWorld, const TCHAR* Cmd, FOutputDevice& Ar )
+bool FRenderAssetStreamingManager::Exec( UWorld* InWorld, const TCHAR* Cmd, FOutputDevice& Ar )
 {
 #if STATS_FAST
-	if (FParse::Command(&Cmd,TEXT("DumpTextureStreamingStats")))
+	if (FParse::Command(&Cmd,TEXT("DumpTextureStreamingStats"))
+		|| FParse::Command(&Cmd, TEXT("DumpRenderAssetStreamingStats")))
 	{
 		return HandleDumpTextureStreamingStatsCommand( Cmd, Ar );
 	}
 #endif
 #if STATS
-	if (FParse::Command(&Cmd,TEXT("ListStreamingTextures")))
+	if (FParse::Command(&Cmd,TEXT("ListStreamingTextures"))
+		|| FParse::Command(&Cmd, TEXT("ListStreamingRenderAssets")))
 	{
-		return HandleListStreamingTexturesCommand( Cmd, Ar );
+		return HandleListStreamingRenderAssetsCommand( Cmd, Ar );
 	}
 #endif
 #if !UE_BUILD_SHIPPING
-	if (FParse::Command(&Cmd, TEXT("ResetMaxEverRequiredTextures")))
+	if (FParse::Command(&Cmd, TEXT("ResetMaxEverRequiredTextures"))
+		|| FParse::Command(&Cmd, TEXT("ResetMaxEverRequiredRenderAssetMemory")))
 	{
-		return HandleResetMaxEverRequiredTexturesCommand(Cmd, Ar);
+		return HandleResetMaxEverRequiredRenderAssetMemoryCommand(Cmd, Ar);
 	}
 	if (FParse::Command(&Cmd,TEXT("LightmapStreamingFactor")))
 	{
 		return HandleLightmapStreamingFactorCommand( Cmd, Ar );
 	}
-	else if (FParse::Command(&Cmd,TEXT("CancelTextureStreaming")))
+	else if (FParse::Command(&Cmd,TEXT("CancelTextureStreaming"))
+		|| FParse::Command(&Cmd, TEXT("CancelRenderAssetStreaming")))
 	{
-		return HandleCancelTextureStreamingCommand( Cmd, Ar );
+		return HandleCancelRenderAssetStreamingCommand( Cmd, Ar );
 	}
 	else if (FParse::Command(&Cmd,TEXT("ShadowmapStreamingFactor")))
 	{
@@ -1844,41 +2011,48 @@ bool FStreamingManagerTexture::Exec( UWorld* InWorld, const TCHAR* Cmd, FOutputD
 	{
 		return HandleNumStreamedMipsCommand( Cmd, Ar );
 	}
-	else if (FParse::Command(&Cmd,TEXT("TrackTexture")))
+	else if (FParse::Command(&Cmd,TEXT("TrackTexture"))
+		|| FParse::Command(&Cmd, TEXT("TrackRenderAsset")))
 	{
-		return HandleTrackTextureCommand( Cmd, Ar );
+		return HandleTrackRenderAssetCommand( Cmd, Ar );
 	}
-	else if (FParse::Command(&Cmd,TEXT("ListTrackedTextures")))
+	else if (FParse::Command(&Cmd,TEXT("ListTrackedTextures"))
+		|| FParse::Command(&Cmd, TEXT("ListTrackedRenderAssets")))
 	{
-		return HandleListTrackedTexturesCommand( Cmd, Ar );
+		return HandleListTrackedRenderAssetsCommand( Cmd, Ar );
 	}
-	else if (FParse::Command(&Cmd,TEXT("DebugTrackedTextures")))
+	else if (FParse::Command(&Cmd,TEXT("DebugTrackedTextures"))
+		|| FParse::Command(&Cmd, TEXT("DebugTrackedRenderAssets")))
 	{
-		return HandleDebugTrackedTexturesCommand( Cmd, Ar );
+		return HandleDebugTrackedRenderAssetsCommand( Cmd, Ar );
 	}
-	else if (FParse::Command(&Cmd,TEXT("UntrackTexture")))
+	else if (FParse::Command(&Cmd,TEXT("UntrackTexture"))
+		|| FParse::Command(&Cmd, TEXT("UntrackRenderAsset")))
 	{
-		return HandleUntrackTextureCommand( Cmd, Ar );
+		return HandleUntrackRenderAssetCommand( Cmd, Ar );
 	}
 	else if (FParse::Command(&Cmd,TEXT("StreamOut")))
 	{
 		return HandleStreamOutCommand( Cmd, Ar );
 	}
-	else if (FParse::Command(&Cmd,TEXT("PauseTextureStreaming")))
+	else if (FParse::Command(&Cmd,TEXT("PauseTextureStreaming"))
+		|| FParse::Command(&Cmd, TEXT("PauseRenderAssetStreaming")))
 	{
-		return HandlePauseTextureStreamingCommand( Cmd, Ar );
+		return HandlePauseRenderAssetStreamingCommand( Cmd, Ar );
 	}
 	else if (FParse::Command(&Cmd,TEXT("StreamingManagerMemory")))
 	{
 		return HandleStreamingManagerMemoryCommand( Cmd, Ar, InWorld );
 	}
-	else if (FParse::Command(&Cmd,TEXT("TextureGroups")))
+	else if (FParse::Command(&Cmd,TEXT("TextureGroups"))
+		|| FParse::Command(&Cmd, TEXT("LODGroups")))
 	{
-		return HandleTextureGroupsCommand( Cmd, Ar );
+		return HandleLODGroupsCommand( Cmd, Ar );
 	}
-	else if (FParse::Command(&Cmd,TEXT("InvestigateTexture")))
+	else if (FParse::Command(&Cmd,TEXT("InvestigateTexture"))
+		|| FParse::Command(&Cmd, TEXT("InvestigateRenderAsset")))
 	{
-		return HandleInvestigateTextureCommand( Cmd, Ar, InWorld );
+		return HandleInvestigateRenderAssetCommand( Cmd, Ar, InWorld );
 	}
 	else if (FParse::Command(&Cmd,TEXT("ListMaterialsWithMissingTextureStreamingData")))
 	{
@@ -1901,7 +2075,7 @@ bool FStreamingManagerTexture::Exec( UWorld* InWorld, const TCHAR* Cmd, FOutputD
 	return false;
 }
 
-void FStreamingManagerTexture::DumpTextureGroupStats( bool bDetailedStats )
+void FRenderAssetStreamingManager::DumpTextureGroupStats( bool bDetailedStats )
 {
 	FScopeLock ScopeLock(&CriticalSection);
 
@@ -1950,7 +2124,7 @@ void FStreamingManagerTexture::DumpTextureGroupStats( bool bDetailedStats )
 		UTexture2D* Texture2D = Cast<UTexture2D>(Texture);
 		FTextureGroupStats& Stat = TextureGroupStats[Texture->LODGroup];
 		FTextureGroupStats& Waste = TextureGroupWaste[Texture->LODGroup];
-		FStreamingTexture* StreamingTexture = GetStreamingTexture(Texture2D);
+		FStreamingRenderAsset* StreamingTexture = GetStreamingRenderAsset(Texture2D);
 		uint32 TextureAlign = 0;
 		if ( StreamingTexture )
 		{
