Resolve robomerge conflicts

diff --git a/Engine/Source/Runtime/Engine/Private/Streaming/StreamingManagerTexture.h b/Engine/Source/Runtime/Engine/Private/Streaming/StreamingManagerTexture.h
--- a/Engine/Source/Runtime/Engine/Private/Streaming/StreamingManagerTexture.h
+++ b/Engine/Source/Runtime/Engine/Private/Streaming/StreamingManagerTexture.h
@@ -14,7 +14,7 @@ TextureStreamingManager.h: Definitions of classes used for texture streaming.
 #include "Streaming/TextureInstanceTask.h"
 
 class AActor;
-class FAsyncTextureStreamingTask;
+class FRenderAssetStreamingMipCalcTask;
 class UPrimitiveComponent;
 
 template<typename TTask> class FAsyncTask;
@@ -22,18 +22,18 @@ template<typename TTask> class FAsyncTask;
 #define STATS_FAST 0
 
 /*-----------------------------------------------------------------------------
-	Texture streaming.
+	Texture or mesh streaming.
 -----------------------------------------------------------------------------*/
 
 /**
- * Streaming manager dealing with textures.
+ * Streaming manager dealing with textures/meshes.
  */
-struct FStreamingManagerTexture : public ITextureStreamingManager
+struct FRenderAssetStreamingManager final : public IRenderAssetStreamingManager
 {
 	/** Constructor, initializing all members */
-	FStreamingManagerTexture();
+	FRenderAssetStreamingManager();
 
-	virtual ~FStreamingManagerTexture();
+	virtual ~FRenderAssetStreamingManager();
 
 	/** Called before GC to clear pending kill levels. */
 	void OnPreGarbageCollect();
@@ -47,11 +47,11 @@ struct FStreamingManagerTexture : public ITextureStreamingManager
 	virtual void UpdateResourceStreaming( float DeltaTime, bool bProcessEverything=false ) override;
 
 	/**
-	 * Updates streaming for an individual texture, taking into account all view infos.
+	 * Updates streaming for an individual texture/mesh, taking into account all view infos.
 	 *
-	 * @param Texture	Texture to update
+	 * @param RenderAsset	Texture or mesh to update
 	 */
-	virtual void UpdateIndividualTexture( UTexture2D* Texture ) override;
+	virtual void UpdateIndividualRenderAsset( UStreamableRenderAsset* RenderAsset ) override;
 
 	/**
 	 * Blocks till all pending requests are fulfilled.
@@ -76,11 +76,11 @@ struct FStreamingManagerTexture : public ITextureStreamingManager
 	virtual void SetDisregardWorldResourcesForFrames( int32 NumFrames ) override;
 
 	/**
-	 *	Try to stream out texture mip-levels to free up more memory.
+	 *	Try to stream out texture/mesh mip-levels to free up more memory.
 	 *	@param RequiredMemorySize	- Required minimum available texture memory
 	 *	@return						- Whether it succeeded or not
 	 **/
-	virtual bool StreamOutTextureData( int64 RequiredMemorySize ) override;
+	virtual bool StreamOutRenderAssetData( int64 RequiredMemorySize ) override;
 
 	virtual int64 GetMemoryOverBudget() const override { return MemoryOverBudget; }
 
@@ -108,29 +108,32 @@ struct FStreamingManagerTexture : public ITextureStreamingManager
 	bool HandleDumpTextureStreamingStatsCommand( const TCHAR* Cmd, FOutputDevice& Ar );
 #endif // STATS_FAST
 #if STATS
-	bool HandleListStreamingTexturesCommand(const TCHAR* Cmd, FOutputDevice& Ar);
+	bool HandleListStreamingRenderAssetsCommand(const TCHAR* Cmd, FOutputDevice& Ar);
 #endif // STATS
 #if !UE_BUILD_SHIPPING
-	bool HandleResetMaxEverRequiredTexturesCommand(const TCHAR* Cmd, FOutputDevice& Ar);
+	bool HandleResetMaxEverRequiredRenderAssetMemoryCommand(const TCHAR* Cmd, FOutputDevice& Ar);
 	bool HandleLightmapStreamingFactorCommand( const TCHAR* Cmd, FOutputDevice& Ar );
-	bool HandleCancelTextureStreamingCommand( const TCHAR* Cmd, FOutputDevice& Ar );
+	bool HandleCancelRenderAssetStreamingCommand( const TCHAR* Cmd, FOutputDevice& Ar );
 	bool HandleShadowmapStreamingFactorCommand( const TCHAR* Cmd, FOutputDevice& Ar );
 	bool HandleNumStreamedMipsCommand( const TCHAR* Cmd, FOutputDevice& Ar );
-	bool HandleTrackTextureCommand( const TCHAR* Cmd, FOutputDevice& Ar );
-	bool HandleListTrackedTexturesCommand( const TCHAR* Cmd, FOutputDevice& Ar );
-	bool HandleDebugTrackedTexturesCommand( const TCHAR* Cmd, FOutputDevice& Ar );
-	bool HandleUntrackTextureCommand( const TCHAR* Cmd, FOutputDevice& Ar );
+	bool HandleTrackRenderAssetCommand( const TCHAR* Cmd, FOutputDevice& Ar );
+	bool HandleListTrackedRenderAssetsCommand( const TCHAR* Cmd, FOutputDevice& Ar );
+	bool HandleDebugTrackedRenderAssetsCommand( const TCHAR* Cmd, FOutputDevice& Ar );
+	bool HandleUntrackRenderAssetCommand( const TCHAR* Cmd, FOutputDevice& Ar );
 	bool HandleStreamOutCommand( const TCHAR* Cmd, FOutputDevice& Ar );
-	bool HandlePauseTextureStreamingCommand( const TCHAR* Cmd, FOutputDevice& Ar );
+	bool HandlePauseRenderAssetStreamingCommand( const TCHAR* Cmd, FOutputDevice& Ar );
 	bool HandleStreamingManagerMemoryCommand( const TCHAR* Cmd, FOutputDevice& Ar, UWorld* InWorld );
-	bool HandleTextureGroupsCommand( const TCHAR* Cmd, FOutputDevice& Ar );
-	bool HandleInvestigateTextureCommand( const TCHAR* Cmd, FOutputDevice& Ar, UWorld* InWorld );
+	bool HandleLODGroupsCommand( const TCHAR* Cmd, FOutputDevice& Ar );
+	bool HandleInvestigateRenderAssetCommand( const TCHAR* Cmd, FOutputDevice& Ar, UWorld* InWorld );
 #endif // !UE_BUILD_SHIPPING
-	/** Adds a new texture to the streaming manager. */
-	virtual void AddStreamingTexture( UTexture2D* Texture ) override;
 
-	/** Removes a texture from the streaming manager. */
-	virtual void RemoveStreamingTexture( UTexture2D* Texture ) override;
+	/** Adds a new texture/mesh to the streaming manager. */
+	virtual void AddStreamingRenderAsset( UTexture2D* Texture ) override;
+	virtual void AddStreamingRenderAsset(UStaticMesh* StaticMesh) override;
+	virtual void AddStreamingRenderAsset(USkeletalMesh* SkeletalMesh) override;
+
+	/** Removes a texture/mesh from the streaming manager. */
+	virtual void RemoveStreamingRenderAsset( UStreamableRenderAsset* RenderAsset ) override;
 
 	/** Adds a ULevel to the streaming manager. */
 	virtual void AddLevel( class ULevel* Level ) override;
@@ -152,13 +155,13 @@ struct FStreamingManagerTexture : public ITextureStreamingManager
 	/**  Called when a primitive streaming data needs to be updated in the last stage of the frame. */
 	virtual void NotifyPrimitiveUpdated_Concurrent( const UPrimitiveComponent* Primitive ) override;
 
-	/** Returns the corresponding FStreamingTexture for a UTexture2D. */
-	FStreamingTexture* GetStreamingTexture( const UTexture2D* Texture2D );
+	/** Returns the corresponding FStreamingRenderAsset for a texture or mesh. */
+	FStreamingRenderAsset* GetStreamingRenderAsset( const UStreamableRenderAsset* RenderAsset );
 
-	/** Set current pause state for texture streaming */
-	virtual void PauseTextureStreaming(bool bInShouldPause) override
+	/** Set current pause state for texture/mesh streaming */
+	virtual void PauseRenderAssetStreaming(bool bInShouldPause) override
 	{
-		bPauseTextureStreaming = bInShouldPause;
+		bPauseRenderAssetStreaming = bInShouldPause;
 	}
 
 	/** Return all bounds related to the ref object */
@@ -169,25 +172,26 @@ struct FStreamingManagerTexture : public ITextureStreamingManager
 
 private:
 //BEGIN: Thread-safe functions and data
-		friend class FAsyncTextureStreamingTask;
+		friend class FRenderAssetStreamingMipCalcTask;
+		friend class FUpdateStreamingRenderAssetsTask;
 
 		/** Remove any references in level managers to this component */
 		void RemoveStaticReferences(const UPrimitiveComponent* Primitive);
 
 		/**
-		 * Not thread-safe: Updates a portion (as indicated by 'StageIndex') of all streaming textures,
+		 * Not thread-safe: Updates a portion (as indicated by 'StageIndex') of all streaming textures/meshes,
 		 * allowing their streaming state to progress.
 		 *
 		 * @param StageIndex		Current stage index
-		 * @param NumUpdateStages	Number of texture update stages
+		 * @param NumUpdateStages	Number of texture/mesh update stages
 		 */
-		void UpdateStreamingTextures( int32 StageIndex, int32 NumStages, bool bWaitForMipFading );
+		void UpdateStreamingRenderAssets( int32 StageIndex, int32 NumStages, bool bWaitForMipFading );
 
-		void ProcessRemovedTextures();
-		void ProcessAddedTextures();
+		void ProcessRemovedRenderAssets();
+		void ProcessAddedRenderAssets();
 		void ConditionalUpdateStaticData();
 
-		/** Adds new textures and level data on the gamethread (while the worker thread isn't active). */
+		/** Adds new textures/meshes and level data on the gamethread (while the worker thread isn't active). */
 		void PrepareAsyncTask( bool bProcessEverything );
 
 		/** Checks for updates in the user settings (CVars, etc). */
@@ -205,11 +209,11 @@ private:
 		void SetOptionalBulkData( UTexture* Texture, bool bHasOptionalBulkData );
 
 		/**
-		 * Stream textures in/out, based on the priorities calculated by the async work.
+		 * Stream textures/meshes in/out, based on the priorities calculated by the async work.
 		 *
 		 * @param bProcessEverything	Whether we're processing all textures in one go
 		 */
-		void StreamTextures( bool bProcessEverything );
+		void StreamRenderAssets( bool bProcessEverything );
 
 		/**
 		 * If new files have loaded this function will return true once
@@ -226,26 +230,46 @@ private:
 			return bResult;
 		}
 
-		/** All streaming UTexture2D objects. */
-		TArray<FStreamingTexture> StreamingTextures;
+		int32 GetNumStreamedMipsArray(FStreamingRenderAsset::EAssetType AssetType, const int32*& OutArray)
+		{
+			switch (AssetType)
+			{
+			case FStreamingRenderAsset::AT_Texture:
+				OutArray = NumStreamedMips_Texture;
+				return TEXTUREGROUP_MAX;
+			case FStreamingRenderAsset::AT_StaticMesh:
+				OutArray = NumStreamedMips_StaticMesh.GetData();
+				return NumStreamedMips_StaticMesh.Num();
+			case FStreamingRenderAsset::AT_SkeletalMesh:
+				OutArray = NumStreamedMips_SkeletalMesh.GetData();
+				return NumStreamedMips_SkeletalMesh.Num();
+			default:
+				check(false);
+				OutArray = nullptr;
+				return -1;
+			}
+		}
 
-		/** All the textures referenced in StreamingTextures. Used to handled deleted textures.  */
-		TSet<const UTexture2D*> ReferencedTextures;
+		/** All streaming texture or mesh objects. */
+		TArray<FStreamingRenderAsset> StreamingRenderAssets;
+
+		/** All the textures/meshes referenced in StreamingRenderAssets. Used to handled deleted textures/meshes.  */
+		TSet<const UStreamableRenderAsset*> ReferencedRenderAssets;
 
 		/** All the currently installed optional bulkdata files */
 		TSet<FString> OptionalBulkDataFiles;
 
-		/** Index of the StreamingTexture that will be updated next by UpdateStreamingTextures(). */
-		int32 CurrentUpdateStreamingTextureIndex;
+		/** Index of the StreamingTexture that will be updated next by UpdateStreamingRenderAssets(). */
+		int32 CurrentUpdateStreamingRenderAssetIndex;
 //END: Thread-safe functions and data
 
 	/**
-	 * Mark the textures with a timestamp. They're about to lose their location-based heuristic and we don't want them to
+	 * Mark the textures/meshes with a timestamp. They're about to lose their location-based heuristic and we don't want them to
 	 * start using LastRenderTime heuristic for a few seconds until they are garbage collected!
 	 *
-	 * @param RemovedTextures	List of removed textures.
+	 * @param RemovedRenderAssets	List of removed textures or meshes.
 	 */
-	void	SetTexturesRemovedTimestamp(const FRemovedTextureArray& RemovedTextures);
+	void	SetRenderAssetsRemovedTimestamp(const FRemovedRenderAssetArray& RemovedRenderAssets);
 
 	void	DumpTextureGroupStats( bool bDetailedStats );
 
@@ -269,11 +293,15 @@ private:
 	 */
 	void	SyncStates(bool bCompleteFullUpdateCycle);
 
-	// Called on game thread when no new elements are added to the StreamingTextures array and
-	// the elements in the array are not removed or reordered. This runs in parrallel with the
-	// async update task so the streaming meta data in each FStreamingTexture can change
+	/**
+	 * Called on game thread when no new elements are added to the StreamingRenderAssets array and
+	 * the elements in the array are not removed or reordered. This runs in parrallel with the
+	 * async update task so the streaming meta data in each FStreamingRenderAsset can change
+	 */
 	void ProcessPendingMipCopyRequests();
 
+	void AddStreamingRenderAsset_Internal(UStreamableRenderAsset* InAsset, FStreamingRenderAsset::EAssetType InType);
+
 	/** Next sync, dump texture group stats. */
 	bool	bTriggerDumpTextureGroupStats;
 
@@ -281,36 +309,39 @@ private:
 	bool	bDetailedDumpTextureGroupStats;
 
 	/** Cached from the system settings. */
-	int32 NumStreamedMips[TEXTUREGROUP_MAX];
+	int32 NumStreamedMips_Texture[TEXTUREGROUP_MAX];
+	TArray<int32> NumStreamedMips_StaticMesh;
+	TArray<int32> NumStreamedMips_SkeletalMesh;
 
-	FTextureStreamingSettings Settings;
+	FRenderAssetStreamingSettings Settings;
 
-	/** Async work for calculating priorities for all textures. */
-	FAsyncTask<FAsyncTextureStreamingTask>*	AsyncWork;
+	/** Async work for calculating priorities and target number of mips for all textures/meshes. */
+	FAsyncTask<FRenderAssetStreamingMipCalcTask>*	AsyncWork;
 
-	/** Async work for texture instance managers. */
-	TRefCountPtr<TextureInstanceTask::FDoWorkAsyncTask> TextureInstanceAsyncWork;
+	/** Async work for render asset instance managers. */
+	TRefCountPtr<RenderAssetInstanceTask::FDoWorkAsyncTask> RenderAssetInstanceAsyncWork;
 
-	/** Textures from dynamic primitives. Owns the data for all levels. */
-	FDynamicTextureInstanceManager DynamicComponentManager;
+	/** Texture/mesh instances from dynamic primitives. Owns the data for all levels. */
+	FDynamicRenderAssetInstanceManager DynamicComponentManager;
 
-	/** New textures, before they've been added to the thread-safe container. */
-	TArray<UTexture2D*>	PendingStreamingTextures;
+	/** New textures/meshes, before they've been added to the thread-safe container. */
+	TArray<UStreamableRenderAsset*>	PendingStreamingRenderAssets;
+	TArray<typename FStreamingRenderAsset::EAssetType> PendingStreamingRenderAssetTypes;
 
-	/** The list of indices with null texture in StreamingTextures. */
-	TArray<int32>	RemovedTextureIndices;
+	/** The list of indices with null render asset in StreamingRenderAssets. */
+	TArray<int32>	RemovedRenderAssetIndices;
 
-	// Represent a pending request to stream in/out mips to/from GPU for a 2D texture
+	// Represent a pending request to stream in/out mips to/from GPU for a texture or mesh
 	struct FPendingMipCopyRequest
 	{
-		const UTexture2D* Texture;
-		// Used to find the corresponding FStreamingTexture in the StreamingTextures array
+		const UStreamableRenderAsset* RenderAsset;
+		// Used to find the corresponding FStreamingRenderAsset in the StreamingRenderAssets array
 		int32 CachedIdx;
 
 		FPendingMipCopyRequest() = default;
 
-		FPendingMipCopyRequest(const UTexture2D* InTexture, int32 InCachedIdx)
-			: Texture(InTexture)
+		FPendingMipCopyRequest(const UStreamableRenderAsset* InAsset, int32 InCachedIdx)
+			: RenderAsset(InAsset)
 			, CachedIdx(InCachedIdx)
 		{}
 	};
@@ -320,23 +351,23 @@ private:
 	int32 CurrentPendingMipCopyRequestIdx;
 
 	/** Level data */
-	TIndirectArray<FLevelTextureManager> LevelTextureManagers;
+	TIndirectArray<FLevelRenderAssetManager> LevelRenderAssetManagers;
 
 	/** Stages [0,N-2] is non-threaded data collection, Stage N-1 is wait-for-AsyncWork-and-finalize. */
 	int32					ProcessingStage;
 
 	/** Total number of processing stages (N). */
-	int32					NumTextureProcessingStages;
+	int32					NumRenderAssetProcessingStages;
 
-	/** Whether to support texture instance streaming for dynamic (movable/spawned) objects. */
+	/** Whether to support texture/mesh instance streaming for dynamic (movable/spawned) objects. */
 	bool					bUseDynamicStreaming;
 
 	float					BoostPlayerTextures;
 
-	/** Amount of memory to leave free in the texture pool. */
+	/** Amount of memory to leave free in the render asset pool. */
 	int64					MemoryMargin;
 
-	/** The actual memory pool size available to stream textures, excludes non-streaming texture, temp memory (for streaming mips), memory margin (allocator overhead). */
+	/** The actual memory pool size available to stream textures/meshes, excludes non-streaming texture/mesh, temp memory (for streaming mips), memory margin (allocator overhead). */
 	int64					EffectiveStreamingPoolSize;
 
 	// Stats we need to keep across frames as we only iterate over a subset of textures.
@@ -344,18 +375,18 @@ private:
 	int64 MemoryOverBudget;
 	int64 MaxEverRequired;
 
-	/** Whether texture streaming is paused or not. When paused, it won't stream any textures in or out. */
-	bool bPauseTextureStreaming;
+	/** Whether render asset streaming is paused or not. When paused, it won't stream any textures/meshes in or out. */
+	bool bPauseRenderAssetStreaming;
 
 	bool bNewFilesLoaded;
 
 	/** Last time all data were fully updated. Instances are considered visible if they were rendered between that last time and the current time. */
 	float LastWorldUpdateTime;
 
-	FTextureStreamingStats DisplayedStats;
-	FTextureStreamingStats GatheredStats;
+	FRenderAssetStreamingStats DisplayedStats;
+	FRenderAssetStreamingStats GatheredStats;
 
-	TArray<int32> InflightTextures;
+	TArray<int32> InflightRenderAssets;
 
 	TMap<FString, bool> CachedFileExistsChecks;
 	void OnPakFileMounted(const TCHAR* PakFilename);
@@ -371,5 +402,5 @@ private:
 	// A critical section use around code that could be called in parallel with NotifyPrimitiveUpdated() or NotifyPrimitiveUpdated_Concurrent().
 	FCriticalSection CriticalSection;
 
-	friend bool TrackTextureEvent( FStreamingTexture* StreamingTexture, UTexture2D* Texture, bool bForceMipLevelsToBeResident, const FStreamingManagerTexture* Manager);
+	friend bool TrackRenderAssetEvent( FStreamingRenderAsset* StreamingRenderAsset, UStreamableRenderAsset* RenderAsset, bool bForceMipLevelsToBeResident, const FRenderAssetStreamingManager* Manager);
 };
diff --git a/Engine/Source/Runtime/Engine/Private/Streaming/StreamingTexture.cpp b/Engine/Source/Runtime/Engine/Private/Streaming/StreamingTexture.cpp
--- a/Engine/Source/Runtime/Engine/Private/Streaming/StreamingTexture.cpp
+++ b/Engine/Source/Runtime/Engine/Private/Streaming/StreamingTexture.cpp
@@ -8,17 +8,25 @@ StreamingTexture.cpp: Definitions of classes used for texture.
 #include "Misc/App.h"
 #include "Streaming/StreamingManagerTexture.h"
 #include "HAL/FileManager.h"
-
-FStreamingTexture::FStreamingTexture(UTexture2D* InTexture, const int32 NumStreamedMips[TEXTUREGROUP_MAX], const FTextureStreamingSettings& Settings)
-: Texture(InTexture)
+#include "Engine/StaticMesh.h"
+#include "Engine/SkeletalMesh.h"
+
+FStreamingRenderAsset::FStreamingRenderAsset(
+	UStreamableRenderAsset* InRenderAsset,
+	const int32* NumStreamedMips,
+	int32 NumLODGroups,
+	EAssetType InAssetType,
+	const FRenderAssetStreamingSettings& Settings)
+	: RenderAsset(InRenderAsset)
+	, RenderAssetType(InAssetType)
 {
 	UpdateStaticData(Settings);
-	UpdateDynamicData(NumStreamedMips, Settings, false);
+	UpdateDynamicData(NumStreamedMips, NumLODGroups, Settings, false);
 
 	InstanceRemovedTimestamp = FApp::GetCurrentTime();
 	DynamicBoostFactor = 1.f;
 
-	bHasUpdatePending = InTexture && InTexture->bHasStreamingUpdatePending;
+	bHasUpdatePending = InRenderAsset && InRenderAsset->bHasStreamingUpdatePending;
 
 	bForceFullyLoadHeuristic = false;
 	bUseUnkownRefHeuristic = false;
@@ -28,43 +36,78 @@ FStreamingTexture::FStreamingTexture(UTexture2D* InTexture, const int32 NumStrea
 	HiddenWantedMips = MinAllowedMips;
 	RetentionPriority = 0;
 	BudgetedMips = MinAllowedMips;
+	NumForcedMips = 0;
 	LoadOrderPriority = 0;
 	WantedMips = MinAllowedMips;
 }
 
-void FStreamingTexture::UpdateStaticData(const FTextureStreamingSettings& Settings)
+void FStreamingRenderAsset::UpdateStaticData(const FRenderAssetStreamingSettings& Settings)
 {
 
 	OptionalBulkDataFilename = TEXT("");
-	if (Texture)
+	if (RenderAsset)
 	{
-		LODGroup = (TextureGroup)Texture->LODGroup;
-		NumNonStreamingMips = Texture->GetNumNonStreamingMips();
-		MipCount = Texture->GetNumMips();
+		LODGroup = RenderAsset->GetLODGroupForStreaming();
+		NumNonStreamingMips = RenderAsset->GetNumNonStreamingMips();
+		MipCount = RenderAsset->GetNumMipsForStreaming();
 		BudgetMipBias = 0;
-		BoostFactor = GetExtraBoost(LODGroup, Settings);
 
-		bIsCharacterTexture = (LODGroup == TEXTUREGROUP_Character || LODGroup == TEXTUREGROUP_CharacterSpecular || LODGroup == TEXTUREGROUP_CharacterNormalMap);
-		bIsTerrainTexture = (LODGroup == TEXTUREGROUP_Terrain_Heightmap || LODGroup == TEXTUREGROUP_Terrain_Weightmap);
+		if (RenderAssetType == AT_Texture)
+		{
+			check(MipCount <= MAX_TEXTURE_MIP_COUNT);
+			const TextureGroup TextureLODGroup = static_cast<TextureGroup>(LODGroup);
+			BoostFactor = GetExtraBoost(TextureLODGroup, Settings);
+			bIsCharacterTexture = (TextureLODGroup == TEXTUREGROUP_Character || TextureLODGroup == TEXTUREGROUP_CharacterSpecular || TextureLODGroup == TEXTUREGROUP_CharacterNormalMap);
+			bIsTerrainTexture = (TextureLODGroup == TEXTUREGROUP_Terrain_Heightmap || TextureLODGroup == TEXTUREGROUP_Terrain_Weightmap);
+		}
+		else
+		{
+			check(MipCount <= MaxNumMeshLODs);
+			// Default boost value .71 is too small for meshes
+			BoostFactor = 1.f;
+			bIsCharacterTexture = false;
+			bIsTerrainTexture = false;
+			if (RenderAssetType == AT_StaticMesh)
+			{
+				const UStaticMesh* StaticMesh = CastChecked<UStaticMesh>(RenderAsset);
+				for (int32 Idx = 0; Idx < MaxNumMeshLODs; ++Idx)
+				{
+					const int32 LODIdx = FMath::Max(MipCount - Idx - 1, 0);
+					LODScreenSizes[Idx] = StaticMesh->RenderData->ScreenSize[LODIdx].GetValueForFeatureLevel(GMaxRHIFeatureLevel);
+				}
+			}
+			else // AT_SkeletalMesh
+			{
+				USkeletalMesh* SkeletalMesh = CastChecked<USkeletalMesh>(RenderAsset);
+				const TArray<FSkeletalMeshLODInfo>& LODInfos =  SkeletalMesh->GetLODInfoArray();
+				for (int32 Idx = 0; Idx < MaxNumMeshLODs; ++Idx)
+				{
+					const int32 LODIdx = FMath::Max(MipCount - Idx - 1, 0);
+					LODScreenSizes[Idx] = LODInfos[LODIdx].ScreenSize.GetValueForFeatureLevel(GMaxRHIFeatureLevel);
+				}
+			}
+		}
 
-		NumNonOptionalMips = MipCount - Texture->CalcNumOptionalMips();
+		NumNonOptionalMips = MipCount - RenderAsset->CalcNumOptionalMips();
 		OptionalMipsState = (NumNonOptionalMips == MipCount) ? EOptionalMipsState::NoOptionalMips : EOptionalMipsState::NotCached;
 
-		for (int32 MipIndex=0; MipIndex < MAX_TEXTURE_MIP_COUNT; ++MipIndex)
+		const int32 MaxNumMips = RenderAssetType == AT_Texture ? MAX_TEXTURE_MIP_COUNT : MaxNumMeshLODs;
+		for (int32 MipIndex = 0; MipIndex < MaxNumMips; ++MipIndex)
 		{
-			TextureSizes[MipIndex] = Texture->CalcTextureMemorySize(FMath::Min(MipIndex + 1, MipCount));
+			CumulativeLODSizes[MipIndex] = RenderAsset->CalcCumulativeLODSize(FMath::Min(MipIndex + 1, MipCount));
 		}
-		const FString* PtrBulkDataFilename = nullptr;
+
 		const int32 OptionalMipCount = MipCount - NumNonOptionalMips;
 		const int32 OptionalMipIndex = OptionalMipCount - 1; // just here so it's clear why this -1 is here
-		if (Texture->GetMipDataFilename(OptionalMipIndex, PtrBulkDataFilename))
+		if (!RenderAsset->GetMipDataFilename(OptionalMipIndex, OptionalBulkDataFilename))
 		{
-			OptionalBulkDataFilename = *PtrBulkDataFilename;
+			OptionalBulkDataFilename.Empty();
 		}
 	}
 	else
 	{
 		LODGroup = TEXTUREGROUP_World;
+		RenderAssetType = AT_Num;
 		NumNonStreamingMips = 0;
 		MipCount = 0;
 		BudgetMipBias = 0;
@@ -77,12 +120,12 @@ void FStreamingTexture::UpdateStaticData(const FTextureStreamingSettings& Settin
 
 		for (int32 MipIndex=0; MipIndex < MAX_TEXTURE_MIP_COUNT; ++MipIndex)
 		{
-			TextureSizes[MipIndex] = 0;
+			CumulativeLODSizes[MipIndex] = 0;
 		}
 	}
 }
 
-void FStreamingTexture::UpdateOptionalMipsState_Async()
+void FStreamingRenderAsset::UpdateOptionalMipsState_Async()
 {
 	// Here we do a lazy update where we check if the highres mip file exists only if it could be useful to do so.
 	// This requires texture to be at max resolution before the optional mips .
@@ -93,27 +136,27 @@ void FStreamingTexture::UpdateOptionalMipsState_Async()
 
 }
 
-void FStreamingTexture::UpdateDynamicData(const int32 NumStreamedMips[TEXTUREGROUP_MAX], const FTextureStreamingSettings& Settings, bool bWaitForMipFading)
+void FStreamingRenderAsset::UpdateDynamicData(const int32* NumStreamedMips, int32 NumLODGroups, const FRenderAssetStreamingSettings& Settings, bool bWaitForMipFading)
 {
 	// Note that those values are read from the async task and must not be assigned temporary values!!
-	if (Texture)
+	if (RenderAsset)
 	{
 		UpdateStreamingStatus(bWaitForMipFading);
 
-		// The last render time of this texture. Can be FLT_MAX when texture has no resource.
-		const float LastRenderTimeForTexture = Texture->GetLastRenderTimeForStreaming();
+		// The last render time of this texture/mesh. Can be FLT_MAX when texture has no resource.
+		const float LastRenderTimeForTexture = RenderAsset->GetLastRenderTimeForStreaming();
 		LastRenderTime = (FApp::GetCurrentTime() > LastRenderTimeForTexture) ? float( FApp::GetCurrentTime() - LastRenderTimeForTexture ) : 0.0f;
 
-		bForceFullyLoad = Texture->ShouldMipLevelsBeForcedResident();
+		bForceFullyLoad = RenderAsset->ShouldMipLevelsBeForcedResident();
 
-		bIgnoreStreamingMipBias = Texture->bIgnoreStreamingMipBias;
+		bIgnoreStreamingMipBias = RenderAsset->bIgnoreStreamingMipBias;
 
-		const int32 NumCinematicMipLevels = (bForceFullyLoad && Texture->bUseCinematicMipLevels) ? Texture->NumCinematicMipLevels : 0;
+		const int32 NumCinematicMipLevels = (bForceFullyLoad && RenderAsset->bUseCinematicMipLevels) ? RenderAsset->NumCinematicMipLevels : 0;
 
 		int32 LODBias = 0;
 		if (!Settings.bUseAllMips)
 		{
-			LODBias = FMath::Max<int32>(Texture->GetCachedLODBias() - NumCinematicMipLevels, 0);
+			LODBias = FMath::Max<int32>(RenderAsset->GetCachedLODBias() - NumCinematicMipLevels, 0);
 
 			// Reduce the max allowed resolution according to LODBias if the texture group allows it.
 			if (IsMaxResolutionAffectedByGlobalBias() && !Settings.bUsePerTextureBias)
@@ -124,16 +167,16 @@ void FStreamingTexture::UpdateDynamicData(const int32 NumStreamedMips[TEXTUREGRO
 			LODBias += BudgetMipBias;
 		}
 
-		// Update MaxAllowedMips in an atomic way to possible bad interaction with the async task.
+		// Update MaxAllowedMips in an atomic way to avoid possible bad interaction with the async task.
 		{
 			// The max mip count is affected by the texture bias and cinematic bias settings.
 			// don't set MaxAllowdMips more then once as it could be read by async texture task
 			int32 TempMaxAllowedMips = FMath::Clamp<int32>(FMath::Min<int32>(MipCount - LODBias, GMaxTextureMipCount), NumNonStreamingMips, MipCount);
 			if (NumNonOptionalMips < MipCount)
 			{
-				// If the optional mips are not available, or if we shouldn't load then now, clamp the possible mips requested. 
+				// If the optional mips are not available, or if we shouldn't load them now, clamp the possible mips requested. 
 				// (when the non-optional mips are not yet loaded, loading optional mips generates cross files requests).
-				// This is not bullet proof though since the texture could have a pending stream-out request.
+				// This is not bullet proof though since the texture/mesh could have a pending stream-out request.
 				if (OptionalMipsState != EOptionalMipsState::HasOptionalMips || ResidentMips < NumNonOptionalMips)
 				{
 					TempMaxAllowedMips = FMath::Min(TempMaxAllowedMips, NumNonOptionalMips);
@@ -142,7 +185,7 @@ void FStreamingTexture::UpdateDynamicData(const int32 NumStreamedMips[TEXTUREGRO
 			MaxAllowedMips = TempMaxAllowedMips;
 		}
 	
-
+		check(LODGroup < NumLODGroups);
 		if (NumStreamedMips[LODGroup] > 0)
 		{
 			MinAllowedMips = FMath::Clamp<int32>(MipCount - NumStreamedMips[LODGroup], NumNonStreamingMips, MaxAllowedMips);
@@ -168,16 +211,24 @@ void FStreamingTexture::UpdateDynamicData(const int32 NumStreamedMips[TEXTUREGRO
 	}
 }
 
-void FStreamingTexture::UpdateStreamingStatus(bool bWaitForMipFading)
+void FStreamingRenderAsset::UpdateStreamingStatus(bool bWaitForMipFading)
 {
-	if (Texture)
+	if (RenderAsset)
 	{
-		bReadyForStreaming = Texture->IsReadyForStreaming();
-		bInFlight = Texture->UpdateStreamingStatus(bWaitForMipFading);
+		bInFlight = RenderAsset->UpdateStreamingStatus(bWaitForMipFading);
+
+		// Optimization: Use GetCachedNumResidentLODs() and GetCachedReadyForStreaming()
+		// instead of GetNumResidentMips() and IsReadyForStreaming() to reduce cache misses
+		// Platforms tested and results (ave exec time of FRenderAssetStreamingManager::UpdateResourceStreaming):
+		//   PS4 Pro - from ~0.79 ms/frame to ~0.55 ms/frame
 
 		// This must be updated after UpdateStreamingStatus
-		ResidentMips = Texture->GetNumResidentMips();
-		RequestedMips = Texture->GetNumRequestedMips();
+		ResidentMips = RenderAsset->GetCachedNumResidentLODs();
+		if (!bReadyForStreaming)
+		{
+			bReadyForStreaming = RenderAsset->GetCachedReadyForStreaming();
+		}
+		RequestedMips = RenderAsset->GetNumRequestedMips();
 	}
 	else
 	{
@@ -186,10 +237,9 @@ void FStreamingTexture::UpdateStreamingStatus(bool bWaitForMipFading)
 	}
 }
 
-float FStreamingTexture::GetExtraBoost(TextureGroup	LODGroup, const FTextureStreamingSettings& Settings)
+float FStreamingRenderAsset::GetExtraBoost(TextureGroup	LODGroup, const FRenderAssetStreamingSettings& Settings)
 {
-	// When accurate distance computation, we need to relax the distance otherwise it gets too conservative. (ex 513 goes to 1024)
-	const float DistanceScale = Settings.bUseNewMetrics ? .71f : 1.f;
+	const float DistanceScale = GetDefaultExtraBoost(Settings.bUseNewMetrics);
 
 	if (LODGroup == TEXTUREGROUP_Terrain_Heightmap || LODGroup == TEXTUREGROUP_Terrain_Weightmap) 
 	{
@@ -211,31 +261,67 @@ float FStreamingTexture::GetExtraBoost(TextureGroup	LODGroup, const FTextureStre
 	}
 }
 
-int32 FStreamingTexture::GetWantedMipsFromSize(float Size) const
+int32 FStreamingRenderAsset::GetWantedMipsFromSize(float Size, float MaxScreenSizeOverAllViews) const
 {
-	float WantedMipsFloat = 1.0f + FMath::Log2(FMath::Max(1.f, Size));
-	int32 WantedMipsInt = FMath::CeilToInt(WantedMipsFloat);
-	return FMath::Clamp<int32>(WantedMipsInt, MinAllowedMips, MaxAllowedMips);
+	if (RenderAssetType == AT_Texture)
+	{
+		float WantedMipsFloat = 1.0f + FMath::Log2(FMath::Max(1.f, Size));
+		int32 WantedMipsInt = FMath::CeilToInt(WantedMipsFloat);
+		return FMath::Clamp<int32>(WantedMipsInt, MinAllowedMips, MaxAllowedMips);
+	}
+	else
+	{
+		check(MinAllowedMips >= 1);
+		check(MaxAllowedMips <= MipCount);
+		check(RenderAssetType == AT_StaticMesh || RenderAssetType == AT_SkeletalMesh);
+		if (Size != FLT_MAX)
+		{
+			const float NormalizedSize = Size / MaxScreenSizeOverAllViews;
+			for (int32 NumMips = MinAllowedMips; NumMips <= MaxAllowedMips; ++NumMips)
+			{
+				if (GetNormalizedScreenSize(NumMips) >= NormalizedSize)
+				{
+					return NumMips;
+				}
+			}
+		}
+		return MaxAllowedMips;
+	}
 }
 
 /** Set the wanted mips from the async task data */
-void FStreamingTexture::SetPerfectWantedMips_Async(float MaxSize, float MaxSize_VisibleOnly, bool InLooksLowRes, const FTextureStreamingSettings& Settings)
+void FStreamingRenderAsset::SetPerfectWantedMips_Async(
+	float MaxSize,
+	float MaxSize_VisibleOnly,
+	float MaxScreenSizeOverAllViews,
+	int32 MaxNumForcedLODs,
+	bool InLooksLowRes,
+	const FRenderAssetStreamingSettings& Settings)
 {
 	bForceFullyLoadHeuristic = (MaxSize == FLT_MAX || MaxSize_VisibleOnly == FLT_MAX);
-	VisibleWantedMips = GetWantedMipsFromSize(MaxSize_VisibleOnly);
 	bLooksLowRes = InLooksLowRes; // Things like lightmaps, HLOD and close instances.
 
+	if (MaxNumForcedLODs >= MaxAllowedMips)
+	{
+		VisibleWantedMips = HiddenWantedMips = NumForcedMips = MaxAllowedMips;
+		NumMissingMips = 0;
+		return;
+	}
+
+	NumForcedMips = FMath::Min(MaxNumForcedLODs, MaxAllowedMips);
+	VisibleWantedMips = FMath::Max(GetWantedMipsFromSize(MaxSize_VisibleOnly, MaxScreenSizeOverAllViews), NumForcedMips);
+
 	// Terrain, Forced Fully Load and Things that already look bad are not affected by hidden scale.
 	if (bIsTerrainTexture || bForceFullyLoadHeuristic || bLooksLowRes)
 	{
-		HiddenWantedMips = GetWantedMipsFromSize(MaxSize);
+		HiddenWantedMips = FMath::Max(GetWantedMipsFromSize(MaxSize, MaxScreenSizeOverAllViews), NumForcedMips);
 		NumMissingMips = 0; // No impact for terrains as there are not allowed to drop mips.
 	}
 	else
 	{
-		HiddenWantedMips = GetWantedMipsFromSize(MaxSize * Settings.HiddenPrimitiveScale);
+		HiddenWantedMips = FMath::Max(GetWantedMipsFromSize(MaxSize * Settings.HiddenPrimitiveScale, MaxScreenSizeOverAllViews), NumForcedMips);
 		// NumMissingMips contains the number of mips not loaded because of HiddenPrimitiveScale. When out of budget, those texture will be considered as already sacrificed.
-		NumMissingMips = FMath::Max<int32>(GetWantedMipsFromSize(MaxSize) - FMath::Max<int32>(VisibleWantedMips, HiddenWantedMips), 0);
+		NumMissingMips = FMath::Max<int32>(GetWantedMipsFromSize(MaxSize, MaxScreenSizeOverAllViews) - FMath::Max<int32>(VisibleWantedMips, HiddenWantedMips), 0);
 	}
 }
 
@@ -243,13 +329,13 @@ void FStreamingTexture::SetPerfectWantedMips_Async(float MaxSize, float MaxSize_
  * Once the wanted mips are computed, the async task will check if everything fits in the budget.
  *  This only consider the highest mip that will be requested eventually, so that slip requests are stables.
  */
-int64 FStreamingTexture::UpdateRetentionPriority_Async()
+int64 FStreamingRenderAsset::UpdateRetentionPriority_Async()
 {
 	// Reserve the budget for the max mip that will be loaded eventually (ignore the effect of split requests)
 	BudgetedMips = GetPerfectWantedMips();
 	RetentionPriority = 0;
 
-	if (Texture)
+	if (RenderAsset)
 	{
 		const bool bIsHuge    = GetSize(BudgetedMips) >= 8 * 1024 * 1024 && LODGroup != TEXTUREGROUP_Lightmap && LODGroup != TEXTUREGROUP_Shadowmap;
 		const bool bShouldKeep = bIsTerrainTexture || bForceFullyLoadHeuristic || (bLooksLowRes && !bIsHuge);
@@ -273,9 +359,9 @@ int64 FStreamingTexture::UpdateRetentionPriority_Async()
 	}
 }
 
-int64 FStreamingTexture::DropMaxResolution_Async(int32 NumDroppedMips)
+int64 FStreamingRenderAsset::DropMaxResolution_Async(int32 NumDroppedMips)
 {
-	if (Texture)
+	if (RenderAsset)
 	{
 		// Don't drop bellow min allowed mips. Also ensure that MinAllowedMips < MaxAllowedMips in order allow the BudgetMipBias to reset.
 		NumDroppedMips = FMath::Min<int32>(MaxAllowedMips - MinAllowedMips - 1, NumDroppedMips);
@@ -305,9 +391,9 @@ int64 FStreamingTexture::DropMaxResolution_Async(int32 NumDroppedMips)
 	return 0;
 }
 
-int64 FStreamingTexture::DropOneMip_Async()
+int64 FStreamingRenderAsset::DropOneMip_Async()
 {
-	if (Texture && BudgetedMips > MinAllowedMips)
+	if (RenderAsset && BudgetedMips > MinAllowedMips)
 	{
 		--BudgetedMips;
 		return GetSize(BudgetedMips + 1) - GetSize(BudgetedMips);
@@ -318,9 +404,9 @@ int64 FStreamingTexture::DropOneMip_Async()
 	}
 }
 
-int64 FStreamingTexture::KeepOneMip_Async()
+int64 FStreamingRenderAsset::KeepOneMip_Async()
 {
-	if (Texture && BudgetedMips < FMath::Min<int32>(ResidentMips, MaxAllowedMips))
+	if (RenderAsset && BudgetedMips < FMath::Min<int32>(ResidentMips, MaxAllowedMips))
 	{
 		++BudgetedMips;
 		return GetSize(BudgetedMips) - GetSize(BudgetedMips - 1);
@@ -331,7 +417,7 @@ int64 FStreamingTexture::KeepOneMip_Async()
 	}
 }
 
-bool FStreamingTexture::UpdateLoadOrderPriority_Async(int32 MinMipForSplitRequest)
+bool FStreamingRenderAsset::UpdateLoadOrderPriority_Async(int32 MinMipForSplitRequest)
 {
 	LoadOrderPriority = 0;
 
@@ -346,7 +432,7 @@ bool FStreamingTexture::UpdateLoadOrderPriority_Async(int32 MinMipForSplitReques
 	}
 
 	// If the entry is valid and we need to send a new request to load/drop the right mip.
-	if (bReadyForStreaming && Texture && WantedMips != RequestedMips)
+	if (bReadyForStreaming && RenderAsset && WantedMips != RequestedMips)
 	{
 		const bool bIsVisible			= ResidentMips < VisibleWantedMips; // Otherwise it means we are loading mips that are only useful for non visible primitives.
 		const bool bMustLoadFirst		= bForceFullyLoadHeuristic || bIsTerrainTexture || bIsCharacterTexture;
@@ -365,40 +451,40 @@ bool FStreamingTexture::UpdateLoadOrderPriority_Async(int32 MinMipForSplitReques
 	}
 }
 
-void FStreamingTexture::CancelPendingMipChangeRequest()
+void FStreamingRenderAsset::CancelPendingMipChangeRequest()
 {
-	if (Texture)
+	if (RenderAsset)
 	{
-		Texture->CancelPendingMipChangeRequest();
+		RenderAsset->CancelPendingMipChangeRequest();
 		UpdateStreamingStatus(false);
 	}
 }
 
-void FStreamingTexture::StreamWantedMips(FStreamingManagerTexture& Manager)
+void FStreamingRenderAsset::StreamWantedMips(FRenderAssetStreamingManager& Manager)
 {
 	StreamWantedMips_Internal(Manager, false);
 }
 
-void FStreamingTexture::CacheStreamingMetaData()
+void FStreamingRenderAsset::CacheStreamingMetaData()
 {
 	bCachedForceFullyLoadHeuristic = bForceFullyLoadHeuristic;
 	CachedWantedMips = WantedMips;
 	CachedVisibleWantedMips = VisibleWantedMips;
 }
 
-void FStreamingTexture::StreamWantedMipsUsingCachedData(FStreamingManagerTexture& Manager)
+void FStreamingRenderAsset::StreamWantedMipsUsingCachedData(FRenderAssetStreamingManager& Manager)
 {
 	StreamWantedMips_Internal(Manager, true);
 }
 
-void FStreamingTexture::StreamWantedMips_Internal(FStreamingManagerTexture& Manager, bool bUseCachedData)
+void FStreamingRenderAsset::StreamWantedMips_Internal(FRenderAssetStreamingManager& Manager, bool bUseCachedData)
 {
-	if (Texture && !Texture->HasPendingUpdate())
+	if (RenderAsset && !RenderAsset->HasPendingUpdate())
 	{
 		const uint32 bLocalForceFullyLoadHeuristic = bUseCachedData ? bCachedForceFullyLoadHeuristic : bForceFullyLoadHeuristic;
 		const int32 LocalVisibleWantedMips = bUseCachedData ? CachedVisibleWantedMips : VisibleWantedMips;
 		// Update ResidentMips now as it is guarantied to not change here (since no pending requests).
-		ResidentMips = Texture->GetNumResidentMips();
+		ResidentMips = RenderAsset->GetNumResidentMips();
 
 		// Prevent streaming-in optional mips and non optional mips as they are from different files.
 		int32 LocalWantedMips = bUseCachedData ? CachedWantedMips : WantedMips;
@@ -411,15 +497,15 @@ void FStreamingTexture::StreamWantedMips_Internal(FStreamingManagerTexture& Mana
 		{
 			if (LocalWantedMips < ResidentMips)
 			{
-				Texture->StreamOut(LocalWantedMips);
+				RenderAsset->StreamOut(LocalWantedMips);
 			}
 			else // WantedMips > ResidentMips
 			{
 				const bool bShouldPrioritizeAsyncIORequest = (bLocalForceFullyLoadHeuristic || bIsTerrainTexture || bIsCharacterTexture) && LocalWantedMips <= LocalVisibleWantedMips;
-				Texture->StreamIn(LocalWantedMips, bShouldPrioritizeAsyncIORequest);
+				RenderAsset->StreamIn(LocalWantedMips, bShouldPrioritizeAsyncIORequest);
 			}
 			UpdateStreamingStatus(false);
-			TrackTextureEvent(this, Texture, bLocalForceFullyLoadHeuristic != 0, &Manager);
+			TrackRenderAssetEvent(this, RenderAsset, bLocalForceFullyLoadHeuristic != 0, &Manager);
 		}
 	}
 }
diff --git a/Engine/Source/Runtime/Engine/Private/Streaming/StreamingTexture.h b/Engine/Source/Runtime/Engine/Private/Streaming/StreamingTexture.h
--- a/Engine/Source/Runtime/Engine/Private/Streaming/StreamingTexture.h
+++ b/Engine/Source/Runtime/Engine/Private/Streaming/StreamingTexture.h
@@ -9,54 +9,98 @@ StreamingTexture.h: Definitions of classes used for texture streaming.
 #include "CoreMinimal.h"
 #include "Engine/Texture2D.h"
 
-struct FStreamingManagerTexture;
-struct FTextureStreamingSettings;
+struct FRenderAssetStreamingManager;
+struct FRenderAssetStreamingSettings;
 
 /*-----------------------------------------------------------------------------
-	FStreamingTexture, the streaming system's version of UTexture2D.
+	FStreamingRenderAsset, the streaming system's version of UTexture2D.
 -----------------------------------------------------------------------------*/
 
-/** Self-contained structure to manage a streaming texture, possibly on a separate thread. */
-struct FStreamingTexture
+/** Self-contained structure to manage a streaming texture/mesh, possibly on a separate thread. */
+struct FStreamingRenderAsset
 {
 private:
-	enum EOptionalMipsState
+	enum EOptionalMipsState : uint8
 	{
 		NotCached,
 		NoOptionalMips,
 		HasOptionalMips,
 	};
+
+	static constexpr int32 MaxNumMeshLODs = MAX_MESH_LOD_COUNT;
+	static_assert(2 * MaxNumMeshLODs >= MAX_TEXTURE_MIP_COUNT, "Failed mip count assumption");
+
 public:
+	enum EAssetType : uint8
+	{
+		AT_Texture,
+		AT_StaticMesh,
+		AT_SkeletalMesh,
+		AT_Num
+	};
 
-	FStreamingTexture(UTexture2D* InTexture, const int32 NumStreamedMips[TEXTUREGROUP_MAX], const FTextureStreamingSettings& Settings);
+	FStreamingRenderAsset(
+		UStreamableRenderAsset* InRenderAsset,
+		const int32* NumStreamedMips,
+		int32 NumLODGroups,
+		EAssetType InAssetType,
+		const FRenderAssetStreamingSettings& Settings);
 
 	/** Update data that should not change unless changing settings. */
-	void UpdateStaticData(const FTextureStreamingSettings& Settings);
+	void UpdateStaticData(const FRenderAssetStreamingSettings& Settings);
 
 	/** Update data that the engine could change through gameplay. */
-	void UpdateDynamicData(const int32 NumStreamedMips[TEXTUREGROUP_MAX], const FTextureStreamingSettings& Settings, bool bWaitForMipFading);
+	void UpdateDynamicData(const int32* NumStreamedMips, int32 NumLODGroups, const FRenderAssetStreamingSettings& Settings, bool bWaitForMipFading);
 
 	/** Lightweight version of UpdateDynamicData. */
 	void UpdateStreamingStatus(bool bWaitForMipFading);
 
 	/**
-	 * Returns the amount of memory used by the texture given a specified number of mip-maps, in bytes.
+	 * Returns the amount of memory used by the texture/mesh given a specified number of mip-maps, in bytes.
 	 *
 	 * @param MipCount	Number of mip-maps to account for
 	 * @return			Total amount of memory used for the specified mip-maps, in bytes
 	 */
 	int32 GetSize( int32 InMipCount ) const
 	{
-		check(InMipCount >= 0 && InMipCount <= MAX_TEXTURE_MIP_COUNT);
-		return TextureSizes[ InMipCount - 1 ];
+		check(InMipCount > 0);
+		check(InMipCount <= RenderAssetType == AT_Texture ? MAX_TEXTURE_MIP_COUNT : MaxNumMeshLODs);
+		return CumulativeLODSizes[InMipCount - 1];
+	}
+
+	static const TCHAR* GetStreamingAssetTypeStr(EAssetType InAssetType)
+	{
+		switch (InAssetType)
+		{
+		case AT_Texture:
+			return TEXT("Texture");
+		case AT_StaticMesh:
+			return TEXT("StaticMesh");
+		case AT_SkeletalMesh:
+			return TEXT("SkeletalMesh");
+		default:
+			return TEXT("Unkown");
+		}
+	}
+
+	static float GetDefaultExtraBoost(bool bUseNewMetrics)
+	{
+		// When accurate distance computation, we need to relax the distance otherwise it gets too conservative. (ex 513 goes to 1024)
+		return bUseNewMetrics ? .71f : 1.f;
 	}
 
-	static float GetExtraBoost(TextureGroup	LODGroup, const FTextureStreamingSettings& Settings);
+	static float GetExtraBoost(TextureGroup	LODGroup, const FRenderAssetStreamingSettings& Settings);
 
-	FORCEINLINE int32 GetWantedMipsFromSize(float Size) const;
+	int32 GetWantedMipsFromSize(float Size, float MaxScreenSizeOverAllViews) const;
 
 	/** Set the wanted mips from the async task data */
-	void SetPerfectWantedMips_Async(float MaxSize, float MaxSize_VisibleOnly, bool InLooksLowRes, const FTextureStreamingSettings& Settings);
+	void SetPerfectWantedMips_Async(
+		float MaxSize,
+		float MaxSize_VisibleOnly,
+		float MaxScreenSizeOverAllViews,
+		int32 MaxNumForcedLODs,
+		bool InLooksLowRes,
+		const FRenderAssetStreamingSettings& Settings);
 
 	/** Init BudgetedMip and update RetentionPriority. Returns the size that would be taken if all budgeted mips where loaded. */
 	int64 UpdateRetentionPriority_Async();
@@ -70,7 +114,25 @@ public:
 	/** Increase BudgetedMip by 1, up to resident mips, and return the size taken. */
 	int64 KeepOneMip_Async();
 
-	float GetMaxAllowedSize() { return (float)(0x1 << (MaxAllowedMips - 1)); }
+	float GetMaxAllowedSize(float MaxScreenSizeOverAllViews) const
+	{
+		return RenderAssetType == AT_Texture ? (float)(0x1 << (MaxAllowedMips - 1)) : MaxScreenSizeOverAllViews;
+	}
+
+	float GetNormalizedScreenSize(int32 NumMips) const
+	{
+		check(RenderAssetType == AT_StaticMesh || RenderAssetType == AT_SkeletalMesh);
+		check(NumMips > 0 && NumMips <= MipCount);
+		return LODScreenSizes[NumMips - 1];
+	}
+
+	float GetLODScreenSize(int32 NumMips, float MaxScreenSizeOverAllViews) const
+	{
+		check(NumMips > 0 && NumMips <= MipCount);
+		return RenderAssetType == AT_Texture ?
+			static_cast<float>(1 << (NumMips - 1)) :
+			GetNormalizedScreenSize(NumMips) * MaxScreenSizeOverAllViews;
+	}
 
 	/** Init load order. Return wether this texture has any load/unload request */
 	bool UpdateLoadOrderPriority_Async(int32 MinMipForSplitRequest);
@@ -78,23 +140,26 @@ public:
 	void UpdateOptionalMipsState_Async();
 	
 	void CancelPendingMipChangeRequest();
-	void StreamWantedMips(FStreamingManagerTexture& Manager);
+	void StreamWantedMips(FRenderAssetStreamingManager& Manager);
 
 	// Cache meta data (e.g. WantedMips) for StreamWantedMipsUsingCachedData to use later on
 	void CacheStreamingMetaData();
 
-	// Stream using the meta data produced by last run of FAsyncTextureStreamingTask.
+	// Stream using the meta data produced by last run of FRenderAssetStreamingMipCalcTask.
 	// This allows streaming to happen in parallel with the async update task
-	void StreamWantedMipsUsingCachedData(FStreamingManagerTexture& Manager);
+	void StreamWantedMipsUsingCachedData(FRenderAssetStreamingManager& Manager);
 
 	FORCEINLINE int32 GetPerfectWantedMips() const { return FMath::Max<int32>(VisibleWantedMips,  HiddenWantedMips); }
 
-	// Whether this texture can be affected by Global Bias and Budget Bias per texture.
-	// It controls if the texture resolution can be scarificed to fit into budget.
+	// Whether this texture/mesh can be affected by Global Bias and Budget Bias per texture/mesh.
+	// It controls if the texture/mesh resolution can be scarificed to fit into budget.
 	FORCEINLINE bool IsMaxResolutionAffectedByGlobalBias() const 
 	{
 		// In editor, forced stream in should never have reduced mips as they can be edited.
-		return LODGroup != TEXTUREGROUP_HierarchicalLOD && !bIsTerrainTexture && !bIgnoreStreamingMipBias && !(GIsEditor && bForceFullyLoadHeuristic); 
+		return (RenderAssetType != AT_Texture || LODGroup != TEXTUREGROUP_HierarchicalLOD)
+			&& !bIsTerrainTexture
+			&& !bIgnoreStreamingMipBias
+			&& !(GIsEditor && bForceFullyLoadHeuristic); 
 	}
 
 	FORCEINLINE bool HasUpdatePending(bool bIsStreamingPaused, bool bHasViewPoint) const 
@@ -116,123 +181,113 @@ public:
 	}
 
 	/***************************************************************
-	 * Members initialized when this is constructed => NEVER CHANGES
+	 * Member data categories:
+	 * (1) Members initialized when this is constructed => NEVER CHANGES
+	 * (2) Cached dynamic members that need constant update => UPDATES IN UpdateDynamicData()
+	 * (3) Helper data set by the streamer to handle special cases => CHANGES ANYTIME (gamethread)
+	 * (4) Data generated by the async task. CHANGES ANYTIME (taskthread)
+	 * (5) Data cached to allow streaming to run in parallel with meta data update
 	 ***************************************************************/
 
-	/** Texture to manage. Note that this becomes null when the texture is removed. */
-	UTexture2D*		Texture;
-
-	/** Cached texture group. */
-	TextureGroup	LODGroup;
-
-	/** Cached number of mipmaps that are not allowed to stream. */
+	/** (1) Texture/mesh to manage. Note that this becomes null when the texture/mesh is removed. */
+	UStreamableRenderAsset*		RenderAsset;
+	/** (2) */
+	FString			OptionalBulkDataFilename;
+	
+	/** (1) Cached texture/mesh LOD group. */
+	int32	LODGroup;
+	/** (1) Cached number of mipmaps that are not allowed to stream. */
 	int32			NumNonStreamingMips;
-
-	/** Cached number of mip-maps in the UTexture2D mip array (including the base mip) */
+	/** (1) Cached number of mip-maps in the asset's mip array (including the base mip) */
 	int32			MipCount;
-
-	/** Sum of all boost factors that applies to this texture. */
+	/** (1) Sum of all boost factors that applies to this texture/mesh. */
 	float			BoostFactor;
+	/** (1) Cached memory sizes for each possible mipcount. */
+	union
+	{
+		int32 CumulativeLODSizes[2 * MaxNumMeshLODs];
+		struct
+		{
+			int32 CumulativeLODSizes_Mesh[MaxNumMeshLODs];
+			// Normalized size of projected bounding sphere - [0, 1]
+			float LODScreenSizes[MaxNumMeshLODs];
+		};
+	};
 
-	/** Cached memory sizes for each possible mipcount. */
-	int32			TextureSizes[MAX_TEXTURE_MIP_COUNT + 1];
-
-	/** Whether the texture is ready to be streamed in/out (cached from IsReadyForStreaming()). */
-	uint32			bIsCharacterTexture : 1;
-
-	/** Whether the texture should be forcibly fully loaded. */
-	uint32			bIsTerrainTexture : 1;
-
-
-	/***********************************************************************************
-	 * Cached dynamic members that need constant update => UPDATES IN UpdateDynamicData()
-	 ***********************************************************************************/
-
-	/** Whether the texture is ready to be streamed in/out (cached from IsReadyForStreaming()). */
-	uint32			bReadyForStreaming : 1;
-
-	/** Whether the texture should be forcibly fully loaded. */
-	uint32			bForceFullyLoad : 1;
-
-	/** Whether the texture resolution should be affected by the memory budget. */
-	uint32			bIgnoreStreamingMipBias : 1;
-
-	/** Cached number of mip-maps in memory (including the base mip) */
+	/** (2) Cached number of mip-maps in memory (including the base mip) */
 	int32			ResidentMips;
-	/** Min number of mip-maps requested by the streaming system. */
+	/** (2) Min number of mip-maps requested by the streaming system. */
 	int32			RequestedMips;
-
-	/** Min mip to be requested by the streaming  */
+	/** (2) Min mip to be requested by the streaming  */
 	int32			MinAllowedMips;
-	/** Max mip to be requested by the streaming  */
+	/** (2) Max mip to be requested by the streaming  */
 	int32			MaxAllowedMips;
-	/** Mips which are in an optional bulk data file (may not be present on device) */
+	/** (2) Mips which are in an optional bulk data file (may not be present on device) */
 	int32			NumNonOptionalMips;
-	/** Cached state on disk of the optional mips for this streaming texture */
-	EOptionalMipsState	OptionalMipsState;
-	FString			OptionalBulkDataFilename;
-
-	/** How much game time has elapsed since the texture was bound for rendering. Based on FApp::GetCurrentTime(). */
+	/** (2) How much game time has elapsed since the texture was bound for rendering. Based on FApp::GetCurrentTime(). */
 	float			LastRenderTime;
 
-	/*****************************************************************************************
-	 * Helper data set by the streamer to handle special cases => CHANGES ANYTIME (gamethread)
-	 *****************************************************************************************/
-
-	/** Whether the texture is currently being streamed in/out. */
-	uint32			bInFlight : 1;
-
-	/** Wheter the streamer has streaming plans for this texture. */
-	uint32			bHasUpdatePending : 1;
-
-	/** If non-zero, the most recent time an instance location was removed for this texture. */
+	/** (3) If non-zero, the most recent time an instance location was removed for this texture. */
 	double			InstanceRemovedTimestamp;
-
-	/** Extra gameplay boost factor. Reset after every update. */
+	/** (3) Extra gameplay boost factor. Reset after every update. */
 	float			DynamicBoostFactor;
 
-	/****************************************************************
-	 * Data generated by the async task. CHANGES ANYTIME (taskthread)
-	 ****************************************************************/
-	
-	/** Same as force fully load, but takes into account component settings. */
-	uint32			bForceFullyLoadHeuristic : 1; 
-
-	/** Whether this has not component referencing it. */
-	uint32			bUseUnkownRefHeuristic : 1;
-
-	/** Same as force fully load, but takes into account component settings. */
-	uint32			bLooksLowRes : 1; 
-
-	/** How many mips are missing to satisfy ideal quality because of max size limitation. Used to prevent sacrificing mips that are already visually sacrificed. */
+	/** (4) How many mips are missing to satisfy ideal quality because of max size limitation. Used to prevent sacrificing mips that are already visually sacrificed. */
 	int32			NumMissingMips;
-
-	/** Max wanted mips for visible instances.*/
+	/** (4) Max wanted mips for visible instances.*/
 	int32			VisibleWantedMips;
-
-	/** Wanted mips for non visible instances.*/
+	/** (4) Wanted mips for non visible instances.*/
 	int32			HiddenWantedMips;
-
-	/** Retention priority used to sacrifice mips when out of budget. */
+	/** (4) Retention priority used to sacrifice mips when out of budget. */
 	int32			RetentionPriority;
-	/** The max allowed mips (based on Visible and Hidden wanted mips) in order to fit in budget. */
+	/** (4) The max allowed mips (based on Visible and Hidden wanted mips) in order to fit in budget. */
 	int32			BudgetedMips;
-
-	/** The load request priority. */
+	/** (4) The load request priority. */
 	int32			LoadOrderPriority;	
-	/** The mip that will be requested. Note, that some texture are loaded using split requests, so that the first request can be smaller than budgeted mip. */
+	/** (4) The mip that will be requested. Note, that some texture are loaded using split requests, so that the first request can be smaller than budgeted mip. */
 	int32			WantedMips;
-
-	/** A persistent bias applied to this texture. Increase whenever the streamer needs to make sacrifices to fit in budget */
+	/** (4) A persistent bias applied to this texture. Increase whenever the streamer needs to make sacrifices to fit in budget */
 	int32			BudgetMipBias;
+	/** (4) Number of LODs forced resident. Only used by meshes currently */
+	int32			NumForcedMips;
 
-	/****************************************************************
-	* Data cached to allow streaming to run in parallel with meta data update
-	****************************************************************/
-	uint32 bCachedForceFullyLoadHeuristic : 1;
+	/** (5) */
 	int32 CachedWantedMips;
+	/** (5) */
 	int32 CachedVisibleWantedMips;
 
+	/** (1) */
+	EAssetType RenderAssetType;
+	/** (2) Cached state on disk of the optional mips for this streaming texture */
+	EOptionalMipsState	OptionalMipsState;
+
+	/** (1) Whether the texture is ready to be streamed in/out (cached from IsReadyForStreaming()). */
+	uint32			bIsCharacterTexture : 1;
+	/** (1) Whether the texture should be forcibly fully loaded. */
+	uint32			bIsTerrainTexture : 1;
+
+	/** (2) Whether the texture is ready to be streamed in/out (cached from IsReadyForStreaming()). */
+	uint32			bReadyForStreaming : 1;
+	/** (2) Whether the texture should be forcibly fully loaded. */
+	uint32			bForceFullyLoad : 1;
+	/** (2) Whether the texture resolution should be affected by the memory budget. */
+	uint32			bIgnoreStreamingMipBias : 1;
+
+	/** (3) Whether the texture is currently being streamed in/out. */
+	uint32			bInFlight : 1;
+	/** (3) Wheter the streamer has streaming plans for this texture. */
+	uint32			bHasUpdatePending : 1;
+
+	/** (4) Same as force fully load, but takes into account component settings. */
+	uint32			bForceFullyLoadHeuristic : 1;
+	/** (4) Whether this has not component referencing it. */
+	uint32			bUseUnkownRefHeuristic : 1;
+	/** (4) Same as force fully load, but takes into account component settings. */
+	uint32			bLooksLowRes : 1;
+
+	/** (5) */
+	uint32			bCachedForceFullyLoadHeuristic : 1;
+
 private:
-	FORCEINLINE_DEBUGGABLE void StreamWantedMips_Internal(FStreamingManagerTexture& Manager, bool bUseCachedData);
+	FORCEINLINE_DEBUGGABLE void StreamWantedMips_Internal(FRenderAssetStreamingManager& Manager, bool bUseCachedData);
 };
diff --git a/Engine/Source/Runtime/Engine/Private/Streaming/Texture2DStreamIn_DDC_AsyncCreate.cpp b/Engine/Source/Runtime/Engine/Private/Streaming/Texture2DStreamIn_DDC_AsyncCreate.cpp
--- a/Engine/Source/Runtime/Engine/Private/Streaming/Texture2DStreamIn_DDC_AsyncCreate.cpp
+++ b/Engine/Source/Runtime/Engine/Private/Streaming/Texture2DStreamIn_DDC_AsyncCreate.cpp
@@ -12,7 +12,7 @@ Texture2DStreamIn_DDC_AsyncCreate.h: Load texture 2D mips from the DDC using asy
 FTexture2DStreamIn_DDC_AsyncCreate::FTexture2DStreamIn_DDC_AsyncCreate(UTexture2D* InTexture, int32 InRequestedMips)
 	: FTexture2DStreamIn_DDC(InTexture, InRequestedMips) 
 {
-	PushTask(FContext(InTexture, TT_None), TT_Async, TEXTURE2D_UPDATE_CALLBACK(AllocateAndLoadMips), TT_None, nullptr);
+	PushTask(FContext(InTexture, TT_None), TT_Async, SRA_UPDATE_CALLBACK(AllocateAndLoadMips), TT_None, nullptr);
 }
 
 // ****************************
@@ -27,7 +27,7 @@ void FTexture2DStreamIn_DDC_AsyncCreate::AllocateAndLoadMips(const FContext& Con
 	DoAllocateNewMips(Context);
 	DoLoadNewMipsFromDDC(Context);
 
-	PushTask(Context, TT_Async, TEXTURE2D_UPDATE_CALLBACK(AsyncCreate), TT_Render, TEXTURE2D_UPDATE_CALLBACK(Cancel));
+	PushTask(Context, TT_Async, SRA_UPDATE_CALLBACK(AsyncCreate), TT_Render, SRA_UPDATE_CALLBACK(Cancel));
 }
 
 void FTexture2DStreamIn_DDC_AsyncCreate::AsyncCreate(const FContext& Context)
@@ -38,7 +38,7 @@ void FTexture2DStreamIn_DDC_AsyncCreate::AsyncCreate(const FContext& Context)
 	DoAsyncCreateWithNewMips(Context);
 	DoFreeNewMips(Context);
 
-	PushTask(Context, TT_Render, TEXTURE2D_UPDATE_CALLBACK(Finalize), TT_Render, TEXTURE2D_UPDATE_CALLBACK(Cancel));
+	PushTask(Context, TT_Render, SRA_UPDATE_CALLBACK(Finalize), TT_Render, SRA_UPDATE_CALLBACK(Cancel));
 }
 
 void FTexture2DStreamIn_DDC_AsyncCreate::Finalize(const FContext& Context)
diff --git a/Engine/Source/Runtime/Engine/Private/Streaming/Texture2DStreamIn_DDC_AsyncReallocate.cpp b/Engine/Source/Runtime/Engine/Private/Streaming/Texture2DStreamIn_DDC_AsyncReallocate.cpp
--- a/Engine/Source/Runtime/Engine/Private/Streaming/Texture2DStreamIn_DDC_AsyncReallocate.cpp
+++ b/Engine/Source/Runtime/Engine/Private/Streaming/Texture2DStreamIn_DDC_AsyncReallocate.cpp
@@ -12,7 +12,7 @@ Texture2DStreamIn_DDC_AsyncReallocate.cpp: Load texture 2D mips from the DDC usi
 FTexture2DStreamIn_DDC_AsyncReallocate::FTexture2DStreamIn_DDC_AsyncReallocate(UTexture2D* InTexture, int32 InRequestedMips)
 	: FTexture2DStreamIn_DDC(InTexture, InRequestedMips)
 {
-	PushTask(FContext(InTexture, TT_None), TT_Render, TEXTURE2D_UPDATE_CALLBACK(AsyncReallocate), TT_None, nullptr);
+	PushTask(FContext(InTexture, TT_None), TT_Render, SRA_UPDATE_CALLBACK(AsyncReallocate), TT_None, nullptr);
 }
 
 // ****************************
@@ -26,7 +26,7 @@ void FTexture2DStreamIn_DDC_AsyncReallocate::AsyncReallocate(const FContext& Con
 
 	DoAsyncReallocate(Context);
 
-	PushTask(Context, TT_Render, TEXTURE2D_UPDATE_CALLBACK(LockMips), TT_Render, TEXTURE2D_UPDATE_CALLBACK(Cancel));
+	PushTask(Context, TT_Render, SRA_UPDATE_CALLBACK(LockMips), TT_Render, SRA_UPDATE_CALLBACK(Cancel));
 }
 
 void FTexture2DStreamIn_DDC_AsyncReallocate::LockMips(const FContext& Context)
@@ -37,7 +37,7 @@ void FTexture2DStreamIn_DDC_AsyncReallocate::LockMips(const FContext& Context)
 	RHIFinalizeAsyncReallocateTexture2D(IntermediateTextureRHI, true);
 	DoLockNewMips(Context);
 
-	PushTask(Context, TT_Async, TEXTURE2D_UPDATE_CALLBACK(LoadMips), TT_Render, TEXTURE2D_UPDATE_CALLBACK(Cancel));
+	PushTask(Context, TT_Async, SRA_UPDATE_CALLBACK(LoadMips), TT_Render, SRA_UPDATE_CALLBACK(Cancel));
 }
 
 void FTexture2DStreamIn_DDC_AsyncReallocate::LoadMips(const FContext& Context)
@@ -47,7 +47,7 @@ void FTexture2DStreamIn_DDC_AsyncReallocate::LoadMips(const FContext& Context)
 
 	DoLoadNewMipsFromDDC(Context);
 
-	PushTask(Context, TT_Render, TEXTURE2D_UPDATE_CALLBACK(Finalize), TT_Render, TEXTURE2D_UPDATE_CALLBACK(Cancel));
+	PushTask(Context, TT_Render, SRA_UPDATE_CALLBACK(Finalize), TT_Render, SRA_UPDATE_CALLBACK(Cancel));
 }
 
 void FTexture2DStreamIn_DDC_AsyncReallocate::Finalize(const FContext& Context)
diff --git a/Engine/Source/Runtime/Engine/Private/Streaming/Texture2DStreamIn_IO.cpp b/Engine/Source/Runtime/Engine/Private/Streaming/Texture2DStreamIn_IO.cpp
--- a/Engine/Source/Runtime/Engine/Private/Streaming/Texture2DStreamIn_IO.cpp
+++ b/Engine/Source/Runtime/Engine/Private/Streaming/Texture2DStreamIn_IO.cpp
@@ -177,9 +177,9 @@ void FTexture2DStreamIn_IO::SetAsyncFileCallback(const FContext& Context)
 
 #if !UE_BUILD_SHIPPING
 		// On some platforms the IO is too fast to test cancelation requests timing issues.
-		if (FTextureStreamingSettings::ExtraIOLatency > 0 && TaskSynchronization.GetValue() == 0)
+		if (FRenderAssetStreamingSettings::ExtraIOLatency > 0 && TaskSynchronization.GetValue() == 0)
 		{
-			FPlatformProcess::Sleep(FTextureStreamingSettings::ExtraIOLatency * .001f); // Slow down the streaming.
+			FPlatformProcess::Sleep(FRenderAssetStreamingSettings::ExtraIOLatency * .001f); // Slow down the streaming.
 		}
 #endif
 
diff --git a/Engine/Source/Runtime/Engine/Private/Streaming/Texture2DStreamIn_IO.h b/Engine/Source/Runtime/Engine/Private/Streaming/Texture2DStreamIn_IO.h
--- a/Engine/Source/Runtime/Engine/Private/Streaming/Texture2DStreamIn_IO.h
+++ b/Engine/Source/Runtime/Engine/Private/Streaming/Texture2DStreamIn_IO.h
@@ -47,7 +47,7 @@ private:
 		void DoWork();
 		FORCEINLINE TStatId GetStatId() const
 		{
-			RETURN_QUICK_DECLARE_CYCLE_STAT(FCancelIORequestsTask, STATGROUP_ThreadPoolAsyncTasks);
+			RETURN_QUICK_DECLARE_CYCLE_STAT(FCancelIORequestsTask_Texture, STATGROUP_ThreadPoolAsyncTasks);
 		}
 	protected:
 		FTexture2DStreamIn_IO* PendingUpdate;
diff --git a/Engine/Source/Runtime/Engine/Private/Streaming/Texture2DStreamIn_IO_AsyncCreate.cpp b/Engine/Source/Runtime/Engine/Private/Streaming/Texture2DStreamIn_IO_AsyncCreate.cpp
--- a/Engine/Source/Runtime/Engine/Private/Streaming/Texture2DStreamIn_IO_AsyncCreate.cpp
+++ b/Engine/Source/Runtime/Engine/Private/Streaming/Texture2DStreamIn_IO_AsyncCreate.cpp
@@ -10,7 +10,7 @@ Texture2DStreamIn_IO_AsyncCreate.cpp: Async create path for streaming in texture
 FTexture2DStreamIn_IO_AsyncCreate::FTexture2DStreamIn_IO_AsyncCreate(UTexture2D* InTexture, int32 InRequestedMips, bool InPrioritizedIORequest)
 	: FTexture2DStreamIn_IO(InTexture, InRequestedMips, InPrioritizedIORequest) 
 {
-	PushTask(FContext(InTexture, TT_None), TT_Async, TEXTURE2D_UPDATE_CALLBACK(AllocateAndLoadMips), TT_None, nullptr);
+	PushTask(FContext(InTexture, TT_None), TT_Async, SRA_UPDATE_CALLBACK(AllocateAndLoadMips), TT_None, nullptr);
 }
 
 // ****************************
@@ -26,7 +26,7 @@ void FTexture2DStreamIn_IO_AsyncCreate::AllocateAndLoadMips(const FContext& Cont
 	DoAllocateNewMips(Context);
 	SetIORequests(Context);
 
-	PushTask(Context, TT_Async, TEXTURE2D_UPDATE_CALLBACK(AsyncCreate), TT_Async, TEXTURE2D_UPDATE_CALLBACK(CancelIO));
+	PushTask(Context, TT_Async, SRA_UPDATE_CALLBACK(AsyncCreate), TT_Async, SRA_UPDATE_CALLBACK(CancelIO));
 }
 
 
@@ -39,7 +39,7 @@ void FTexture2DStreamIn_IO_AsyncCreate::AsyncCreate(const FContext& Context)
 	DoFreeNewMips(Context);
 	ClearIORequests(Context);
 
-	PushTask(Context, TT_Render, TEXTURE2D_UPDATE_CALLBACK(Finalize), TT_Render, TEXTURE2D_UPDATE_CALLBACK(Cancel));
+	PushTask(Context, TT_Render, SRA_UPDATE_CALLBACK(Finalize), TT_Render, SRA_UPDATE_CALLBACK(Cancel));
 }
 
 void FTexture2DStreamIn_IO_AsyncCreate::Finalize(const FContext& Context)
@@ -62,7 +62,7 @@ void FTexture2DStreamIn_IO_AsyncCreate::CancelIO(const FContext& Context)
 
 	ClearIORequests(Context);
 
-	PushTask(Context, TT_None, nullptr, TT_Render, TEXTURE2D_UPDATE_CALLBACK(Cancel));
+	PushTask(Context, TT_None, nullptr, TT_Render, SRA_UPDATE_CALLBACK(Cancel));
 }
 
 
diff --git a/Engine/Source/Runtime/Engine/Private/Streaming/Texture2DStreamIn_IO_AsyncReallocate.cpp b/Engine/Source/Runtime/Engine/Private/Streaming/Texture2DStreamIn_IO_AsyncReallocate.cpp
--- a/Engine/Source/Runtime/Engine/Private/Streaming/Texture2DStreamIn_IO_AsyncReallocate.cpp
+++ b/Engine/Source/Runtime/Engine/Private/Streaming/Texture2DStreamIn_IO_AsyncReallocate.cpp
@@ -10,7 +10,7 @@ Texture2DStreamIn_IO_AsyncReallocate.cpp: Default path for streaming in texture
 FTexture2DStreamIn_IO_AsyncReallocate::FTexture2DStreamIn_IO_AsyncReallocate(UTexture2D* InTexture, int32 InRequestedMips, bool InPrioritizedIORequest) 
 	: FTexture2DStreamIn_IO(InTexture, InRequestedMips, InPrioritizedIORequest)
 {
-	PushTask(FContext(InTexture, TT_None), TT_Render, TEXTURE2D_UPDATE_CALLBACK(AsyncReallocate), TT_None, nullptr);
+	PushTask(FContext(InTexture, TT_None), TT_Render, SRA_UPDATE_CALLBACK(AsyncReallocate), TT_None, nullptr);
 }
 
 // ****************************
@@ -25,7 +25,7 @@ void FTexture2DStreamIn_IO_AsyncReallocate::AsyncReallocate(const FContext& Cont
 	SetIOFilename(Context);
 	DoAsyncReallocate(Context);
 
-	PushTask(Context, TT_Render, TEXTURE2D_UPDATE_CALLBACK(LockMips), TT_Render, TEXTURE2D_UPDATE_CALLBACK(Cancel));
+	PushTask(Context, TT_Render, SRA_UPDATE_CALLBACK(LockMips), TT_Render, SRA_UPDATE_CALLBACK(Cancel));
 }
 
 void FTexture2DStreamIn_IO_AsyncReallocate::LockMips(const FContext& Context)
@@ -36,7 +36,7 @@ void FTexture2DStreamIn_IO_AsyncReallocate::LockMips(const FContext& Context)
 	RHIFinalizeAsyncReallocateTexture2D(IntermediateTextureRHI, true);
 	DoLockNewMips(Context);
 
-	PushTask(Context, TT_Async, TEXTURE2D_UPDATE_CALLBACK(LoadMips), TT_Render, TEXTURE2D_UPDATE_CALLBACK(Cancel));
+	PushTask(Context, TT_Async, SRA_UPDATE_CALLBACK(LoadMips), TT_Render, SRA_UPDATE_CALLBACK(Cancel));
 }
 
 void FTexture2DStreamIn_IO_AsyncReallocate::LoadMips(const FContext& Context)
@@ -46,7 +46,7 @@ void FTexture2DStreamIn_IO_AsyncReallocate::LoadMips(const FContext& Context)
 
 	SetIORequests(Context);
 
-	PushTask(Context, TT_Render, TEXTURE2D_UPDATE_CALLBACK(Finalize), TT_Async, TEXTURE2D_UPDATE_CALLBACK(CancelIO));
+	PushTask(Context, TT_Render, SRA_UPDATE_CALLBACK(Finalize), TT_Async, SRA_UPDATE_CALLBACK(CancelIO));
 }
 
 void FTexture2DStreamIn_IO_AsyncReallocate::Finalize(const FContext& Context)
@@ -70,7 +70,7 @@ void FTexture2DStreamIn_IO_AsyncReallocate::CancelIO(const FContext& Context)
 
 	ClearIORequests(Context);
 
-	PushTask(Context, TT_None, nullptr, TT_Render, TEXTURE2D_UPDATE_CALLBACK(Cancel));
+	PushTask(Context, TT_None, nullptr, TT_Render, SRA_UPDATE_CALLBACK(Cancel));
 }
 
 void FTexture2DStreamIn_IO_AsyncReallocate::Cancel(const FContext& Context)
diff --git a/Engine/Source/Runtime/Engine/Private/Streaming/Texture2DStreamIn_IO_Virtual.cpp b/Engine/Source/Runtime/Engine/Private/Streaming/Texture2DStreamIn_IO_Virtual.cpp
--- a/Engine/Source/Runtime/Engine/Private/Streaming/Texture2DStreamIn_IO_Virtual.cpp
+++ b/Engine/Source/Runtime/Engine/Private/Streaming/Texture2DStreamIn_IO_Virtual.cpp
@@ -10,7 +10,7 @@ Texture2DStreamIn_IO_AsyncCreate.cpp: Async create path for streaming in texture
 FTexture2DStreamIn_IO_Virtual::FTexture2DStreamIn_IO_Virtual(UTexture2D* InTexture, int32 InRequestedMips, bool InPrioritizedIORequest) 
 	: FTexture2DStreamIn_IO(InTexture, InRequestedMips, InPrioritizedIORequest)
 {
-	PushTask(FContext(InTexture, TT_None), TT_Render, TEXTURE2D_UPDATE_CALLBACK(LockMips), TT_None, nullptr);
+	PushTask(FContext(InTexture, TT_None), TT_Render, SRA_UPDATE_CALLBACK(LockMips), TT_None, nullptr);
 }
 
 void FTexture2DStreamIn_IO_Virtual::LockMips(const FContext& Context)
@@ -22,7 +22,7 @@ void FTexture2DStreamIn_IO_Virtual::LockMips(const FContext& Context)
 	DoConvertToVirtualWithNewMips(Context);
 	DoLockNewMips(Context);
 
-	PushTask(Context, TT_Async, TEXTURE2D_UPDATE_CALLBACK(LoadMips), TT_Render, TEXTURE2D_UPDATE_CALLBACK(Cancel));
+	PushTask(Context, TT_Async, SRA_UPDATE_CALLBACK(LoadMips), TT_Render, SRA_UPDATE_CALLBACK(Cancel));
 }
 
 void FTexture2DStreamIn_IO_Virtual::LoadMips(const FContext& Context)
@@ -32,7 +32,7 @@ void FTexture2DStreamIn_IO_Virtual::LoadMips(const FContext& Context)
 
 	SetIORequests(Context);
 
-	PushTask(Context, TT_Async, TEXTURE2D_UPDATE_CALLBACK(PostLoadMips), TT_Async, TEXTURE2D_UPDATE_CALLBACK(CancelIO));
+	PushTask(Context, TT_Async, SRA_UPDATE_CALLBACK(PostLoadMips), TT_Async, SRA_UPDATE_CALLBACK(CancelIO));
 }
 
 void FTexture2DStreamIn_IO_Virtual::PostLoadMips(const FContext& Context)
@@ -42,7 +42,7 @@ void FTexture2DStreamIn_IO_Virtual::PostLoadMips(const FContext& Context)
 
 	ClearIORequests(Context);
 
-	PushTask(Context, TT_Render, TEXTURE2D_UPDATE_CALLBACK(Finalize), TT_Render, TEXTURE2D_UPDATE_CALLBACK(Cancel));
+	PushTask(Context, TT_Render, SRA_UPDATE_CALLBACK(Finalize), TT_Render, SRA_UPDATE_CALLBACK(Cancel));
 }
 
 
@@ -70,7 +70,7 @@ void FTexture2DStreamIn_IO_Virtual::CancelIO(const FContext& Context)
 
 	ClearIORequests(Context);
 
-	PushTask(Context, TT_None, nullptr, TT_Render, TEXTURE2D_UPDATE_CALLBACK(Cancel));
+	PushTask(Context, TT_None, nullptr, TT_Render, SRA_UPDATE_CALLBACK(Cancel));
 }
 
 void FTexture2DStreamIn_IO_Virtual::Cancel(const FContext& Context)
diff --git a/Engine/Source/Runtime/Engine/Private/Streaming/Texture2DStreamOut_AsyncReallocate.cpp b/Engine/Source/Runtime/Engine/Private/Streaming/Texture2DStreamOut_AsyncReallocate.cpp
--- a/Engine/Source/Runtime/Engine/Private/Streaming/Texture2DStreamOut_AsyncReallocate.cpp
+++ b/Engine/Source/Runtime/Engine/Private/Streaming/Texture2DStreamOut_AsyncReallocate.cpp
@@ -16,7 +16,7 @@ FTexture2DStreamOut_AsyncReallocate::FTexture2DStreamOut_AsyncReallocate(UTextur
 {
 	ensure(InRequestedMips < InTexture->GetNumResidentMips());
 
-	PushTask(FContext(InTexture, TT_None), TT_Render, TEXTURE2D_UPDATE_CALLBACK(AsyncReallocate), TT_None, nullptr);
+	PushTask(FContext(InTexture, TT_None), TT_Render, SRA_UPDATE_CALLBACK(AsyncReallocate), TT_None, nullptr);
 }
 
 void FTexture2DStreamOut_AsyncReallocate::AsyncReallocate(const FContext& Context)
@@ -26,7 +26,7 @@ void FTexture2DStreamOut_AsyncReallocate::AsyncReallocate(const FContext& Contex
 
 	DoAsyncReallocate(Context);
 	
-	PushTask(Context, TT_Render, TEXTURE2D_UPDATE_CALLBACK(Finalize), TT_Render, TEXTURE2D_UPDATE_CALLBACK(Cancel));
+	PushTask(Context, TT_Render, SRA_UPDATE_CALLBACK(Finalize), TT_Render, SRA_UPDATE_CALLBACK(Cancel));
 }
 
 void FTexture2DStreamOut_AsyncReallocate::Finalize(const FContext& Context)
diff --git a/Engine/Source/Runtime/Engine/Private/Streaming/Texture2DStreamOut_Virtual.cpp b/Engine/Source/Runtime/Engine/Private/Streaming/Texture2DStreamOut_Virtual.cpp
--- a/Engine/Source/Runtime/Engine/Private/Streaming/Texture2DStreamOut_Virtual.cpp
+++ b/Engine/Source/Runtime/Engine/Private/Streaming/Texture2DStreamOut_Virtual.cpp
@@ -12,7 +12,7 @@ FTexture2DStreamOut_Virtual::FTexture2DStreamOut_Virtual(UTexture2D* InTexture,
 {
 	ensure(InRequestedMips < InTexture->GetNumResidentMips());
 	
-	PushTask(FContext(InTexture, TT_None), TT_Render, TEXTURE2D_UPDATE_CALLBACK(Finalize), TT_None, nullptr);
+	PushTask(FContext(InTexture, TT_None), TT_Render, SRA_UPDATE_CALLBACK(Finalize), TT_None, nullptr);
 }
 
 // ****************************
diff --git a/Engine/Source/Runtime/Engine/Private/Streaming/Texture2DUpdate.cpp b/Engine/Source/Runtime/Engine/Private/Streaming/Texture2DUpdate.cpp
--- a/Engine/Source/Runtime/Engine/Private/Streaming/Texture2DUpdate.cpp
+++ b/Engine/Source/Runtime/Engine/Private/Streaming/Texture2DUpdate.cpp
@@ -8,221 +8,37 @@ Texture2DUpdate.cpp: Helpers to stream in and out mips.
 #include "RenderUtils.h"
 #include "Containers/ResourceArray.h"
 
-volatile int32 FTexture2DUpdate::GSuspendRenderThreadTasks = 0;
-
-FTexture2DUpdate::FContext::FContext(UTexture2D* InTexture, EThreadType InCurrentThread) 
+FTexture2DUpdateContext::FTexture2DUpdateContext(UTexture2D* InTexture, EThreadType InCurrentThread)
 	: Texture(InTexture)
 	, CurrentThread(InCurrentThread)
 {
 	check(InTexture);
-	checkSlow(InCurrentThread != TT_Render || IsInRenderingThread());
+	checkSlow(InCurrentThread != FTexture2DUpdate::TT_Render || IsInRenderingThread());
 	Resource = static_cast<FTexture2DResource*>(Texture->Resource);
 }
 
+FTexture2DUpdateContext::FTexture2DUpdateContext(UStreamableRenderAsset* InTexture, EThreadType InCurrentThread)
+#if UE_BUILD_SHIPPING
+	: FTexture2DUpdateContext(static_cast<UTexture2D*>(InTexture), InCurrentThread)
+#else
+	: FTexture2DUpdateContext(Cast<UTexture2D>(InTexture), InCurrentThread)
+#endif
+{}
+
 FTexture2DUpdate::FTexture2DUpdate(UTexture2D* InTexture, int32 InRequestedMips) 
-	: PendingFirstMip(INDEX_NONE)
-	, RequestedMips(INDEX_NONE)
-	, ScheduledTaskCount(0)
-	, bIsCancelled(false)
-	, TaskState(TS_Locked) // The object is created in the locked state to follow the Tick path
-	, PendingTaskState(TS_None)
-	, TaskThread(TT_None)
-	, CancelationThread(TT_None)	
+	: TRenderAssetUpdate<FTexture2DUpdateContext>(InTexture, InRequestedMips)
 {
-	check(InTexture);
-
-	const int32 NonStreamingMipCount = InTexture->GetNumNonStreamingMips();
-	const int32 MaxMipCount = InTexture->GetNumMips();
-	InRequestedMips = FMath::Clamp<int32>(InRequestedMips, NonStreamingMipCount, MaxMipCount);
-
-	if (InRequestedMips > 0 && InRequestedMips != InTexture->GetNumResidentMips() && InTexture->bIsStreamable && InTexture->Resource)
-	{
-		RequestedMips = InRequestedMips;
-		PendingFirstMip	= InTexture->GetPlatformMips().Num() - RequestedMips;
-	}
-	else // This shouldn't happen but if it does, then the update is canceled
+	if (!InTexture->Resource)
 	{
+		RequestedMips = INDEX_NONE;
+		PendingFirstMip = INDEX_NONE;
 		bIsCancelled = true;
 	}
 }
 
 FTexture2DUpdate::~FTexture2DUpdate()
 {
-	// Work must be done here because derived destructors have been called now and so derived members are invalid.
-	ensure(ScheduledTaskCount <= 0);
 	ensure(!IntermediateTextureRHI);
-
-}
-
-void FTexture2DUpdate::Tick(UTexture2D* InTexture, EThreadType InCurrentThread)
-{
-	if (TaskState == TS_None || (TaskSynchronization.GetValue() > 0 && InCurrentThread == TT_None))
-	{
-		// Early exit if the task is not ready to execute and we are ticking from non executing thread.
-		// Executing thread must not early exit in order to make sure that tasks are correctly scheduled.
-		// Otherwise, this assumes that the game thread regularly ticks.
-		return;
-	}
-
-	// Acquire the lock if there is work to do and if it is allowed to wait for the lock
-	if (DoConditionalLock(InCurrentThread))
-	{
-		ensure(PendingTaskState == TS_Scheduled || PendingTaskState == TS_Pending);
-
-		// If the task is ready to execute. 
-		// If this is the renderthread and we shouldn't run renderthread task, mark as pending.
-		// This will require a tick from the game thread to reschedule the task.
-		if (TaskSynchronization.GetValue() <= 0 && !(GSuspendRenderThreadTasks > 0 && InCurrentThread == TT_Render))
-		{
-			FContext Context(InTexture, InCurrentThread);
-
-			// The task params can not change at this point but the bIsCancelled state could change.
-			// To ensure coherency, the cancel state is cached as it affects which thread is relevant.
-			const bool bCachedIsCancelled = bIsCancelled;
-			const EThreadType RelevantThread = !bCachedIsCancelled ? TaskThread : CancelationThread;
-
-			if (RelevantThread == TT_None)
-			{
-				ClearTask();
-			}
-			else if (InCurrentThread == RelevantThread)
-			{
-				FCallback CachedCallback = !bCachedIsCancelled ? TaskCallback : CancelationCallback;
-				ClearTask();
-				CachedCallback(Context); // Valid if the thread is valid.
-			}
-			else if (PendingTaskState != TS_Scheduled || InCurrentThread != TT_None)
-			{
-				// If the task was never scheduled (because synchro was not ready) schedule now.
-				// We also reschedule if this is an executing thread and it end up not being the good thread.
-				// This can happen when a task gets cancelled between scheduling and execution.
-				// We enforce that executing threads either execute or reschedule to prevent a possible stalls
-				// since game thread will not reschedule after the first time.
-				// It's completely safe to schedule several time as task execution can only ever happen once.
-				// we also keep track of how many callbacks are scheduled through ScheduledTaskCount to prevent
-				// deleting the object while another thread is scheduled to access it.
-				ScheduleTick(Context, RelevantThread);
-			}
-			else // Otherwise unlock the task for the executing thread to process it.
-			{
-				PendingTaskState = TS_Scheduled;
-			}
-		}
-		else // If synchro is not ready, mark the task as pending to be scheduled.
-		{
-			PendingTaskState = TS_Pending;
-		}
-
-		DoUnlock();
-	}
-}
-
-void FTexture2DUpdate::PushTask(const FContext& Context, EThreadType InTaskThread, const FCallback& InTaskCallback, EThreadType InCancelationThread, const FCallback& InCancelationCallback)
-{
-	// PushTask can only be called by the one thread/callback that is doing the processing. 
-	// This means we don't need to check whether other threads could be trying to push tasks.
-	check(TaskState == TS_Locked);
-	checkSlow((bool)InTaskCallback == (InTaskThread != TT_None));
-	checkSlow((bool)InCancelationCallback == (InCancelationThread != TT_None));
-
-	// To ensure coherency, the cancel state is cached as it affects which thread is relevant.
-	const bool bCachedIsCancelled = bIsCancelled;
-	const EThreadType RelevantThread = !bCachedIsCancelled ? InTaskThread : InCancelationThread;
-
-	// TaskSynchronization is expected to be set before call this.
-	// If the rendering thread is suspended, delay the scheduling until not suspended anymore.
-	const bool bCanExecuteNow = TaskSynchronization.GetValue() <= 0 && !(GSuspendRenderThreadTasks > 0 && RelevantThread == TT_Render);
-
-	if (RelevantThread == TT_None)
-	{
-		// Nothing to do.
-	}
-	else if (bCanExecuteNow && Context.CurrentThread == RelevantThread)
-	{
-		FCallback CachedCallback = !bCachedIsCancelled ? InTaskCallback : InCancelationCallback;
-		// Not null otherwise relevant thread would be TT_None.
-		CachedCallback(Context);
-	}
-	else
-	{
-		TaskThread = InTaskThread;
-		TaskCallback = InTaskCallback;
-		CancelationThread = InCancelationThread;
-		CancelationCallback = InCancelationCallback;
-
-		if (bCanExecuteNow)
-		{
-			ScheduleTick(Context, RelevantThread);
-		}
-		else
-		{
-			PendingTaskState = TS_Pending;
-		}
-	}
-}	
-
-void FTexture2DUpdate::ScheduleTick(const FContext& Context, EThreadType InThread)
-{
-	// Task can only be scheduled once the synchronization is completed.
-	check(TaskSynchronization.GetValue() <= 0);
-
-	// The pending update needs to be cached because the scheduling can happen in the constructor, before the assignment.
-	
-	// When not having many threads, async task should never schedule tasks since they would wake threads with higher priority while still holding the lock.
-	if (Context.CurrentThread == TT_Async && !FApp::ShouldUseThreadingForPerformance())
-	{
-		PendingTaskState = TS_Pending;
-	}
-	else if (InThread == TT_Render)
-	{
-		FPlatformAtomics::InterlockedIncrement(&ScheduledTaskCount);
-		PendingTaskState = TS_Scheduled;
-
-		ENQUEUE_UNIQUE_RENDER_COMMAND_TWOPARAMETER(
-		Texture2DUpdateCommand,
-		UTexture2D*, Texture, Context.Texture,
-		FTexture2DUpdate*, CachedPendingUpdate, this,
-		{
-			check(Texture && CachedPendingUpdate);
-
-			// Recompute the context has things might have changed!
-			CachedPendingUpdate->Tick(Texture, TT_Render);
-
-			FPlatformMisc::MemoryBarrier();
-			FPlatformAtomics::InterlockedDecrement(&CachedPendingUpdate->ScheduledTaskCount);
-		});
-	}
-	else // InThread == TT_Async
-	{
-		check(InThread == TT_Async);
-
-		FPlatformAtomics::InterlockedIncrement(&ScheduledTaskCount);
-		PendingTaskState = TS_Scheduled;
-
-		(new FAsyncMipUpdateTask(Context.Texture, this))->StartBackgroundTask();
-	}
-}
-
-void FTexture2DUpdate::ClearTask()
-{
-	// Reset everything so that the callback setup a new callback
-	PendingTaskState = TS_None;
-	TaskThread = TT_None;
-	TaskCallback = nullptr;
-	CancelationThread = TT_None;
-	CancelationCallback = nullptr;
-	TaskSynchronization.Set(0);
-}
-
-void FTexture2DUpdate::FMipUpdateTask::DoWork()
-{
-	check(Texture && CachedPendingUpdate);
-
-	// Recompute the context has things might have changed!
-	CachedPendingUpdate->Tick(Texture, FTexture2DUpdate::TT_Async);
-
-	FPlatformMisc::MemoryBarrier();
-	FPlatformAtomics::InterlockedDecrement(&CachedPendingUpdate->ScheduledTaskCount);
 }
 
 
@@ -322,83 +138,3 @@ void FTexture2DUpdate::DoFinishUpdate(const FContext& Context)
 
 	}
 }
-
-bool FTexture2DUpdate::DoConditionalLock(EThreadType InCurrentThread)
-{
-	DECLARE_SCOPE_CYCLE_COUNTER(TEXT("FTexture2DUpdate::DoConditionalLock"), STAT_FTexture2DUpdate_DoConditionalLock, STATGROUP_StreamingDetails);
-
-	// Acquire the lock if there is work to do and if it is allowed to wait for the lock
-	int32 CachedTaskState = TS_None;
-	do
-	{
-		// Sleep in between retries.
-		if (CachedTaskState != TS_None)
-		{
-			FPlatformProcess::Sleep(0);
-		}
-		// Cache the task state.
-		CachedTaskState = TaskState;
-
-		// Return immediately if there is no work to do or if it is locked and we are not on an executing thread.
-		// When the renderthread is the gamethread, don't lock if this is the renderthread to prevent stalling on low priority async tasks.
-		if (CachedTaskState == TS_None || (CachedTaskState == TS_Locked && (InCurrentThread == TT_None || (InCurrentThread == TT_Render && !GIsThreadedRendering))))
-		{
-			return false;
-		}
-	}
-	while (CachedTaskState == TS_Locked || FPlatformAtomics::InterlockedCompareExchange(&TaskState, TS_Locked, CachedTaskState) != CachedTaskState);
-
-	ensure(PendingTaskState == TS_None);
-	PendingTaskState = (ETaskState)CachedTaskState;
-
-	return true;
-}
-
-
-void FTexture2DUpdate::DoLock()
-{
-	DECLARE_SCOPE_CYCLE_COUNTER(TEXT("FTexture2DUpdate::DoLock"), STAT_FTexture2DUpdate_DoLock, STATGROUP_StreamingDetails);
-
-	// Acquire the lock
-	int32 CachedTaskState = TS_None;
-	do
-	{
-		// Sleep in between retries.
-		if (CachedTaskState != TS_None)
-		{
-			FPlatformProcess::Sleep(0);
-		}
-		// Cache the task state.
-		CachedTaskState = TaskState;
-
-	}
-	while (CachedTaskState == TS_Locked || FPlatformAtomics::InterlockedCompareExchange(&TaskState, TS_Locked, CachedTaskState) != CachedTaskState);
-	
-	// If we just acquired the lock, nothing should be in the process.
-	ensure(PendingTaskState == TS_None);
-	PendingTaskState = (ETaskState)CachedTaskState;
-}
-
-void FTexture2DUpdate::DoUnlock()
-{
-	ensure(TaskState == TS_Locked && PendingTaskState != TS_Locked);
-
-	ETaskState CachedPendingTaskState = PendingTaskState;
-
-	// Reset the pending task state first to prevent racing condition that could fail ensure(PendingTaskState == TS_None) in DoLock()
-	PendingTaskState = TS_None;
-	TaskState = CachedPendingTaskState;
-}
-
-void SuspendTextureStreamingRenderTasksInternal()
-{
-	// This doesn't prevent from having a task pushed immediately after as some threads could already be deep in FTexture2DUpdate::PushTask.
-	// This is why GSuspendRenderThreadTasks also gets checked in FTexture2DUpdate::Tick. The goal being to avoid accessing the RHI more
-	// than avoiding pushing new render commands. The reason behind is that some code paths access the RHI outside the render thread.
-	FPlatformAtomics::InterlockedIncrement(&FTexture2DUpdate::GSuspendRenderThreadTasks);
-}
-
-void ResumeTextureStreamingRenderTasksInternal()
-{
-	FPlatformAtomics::InterlockedDecrement(&FTexture2DUpdate::GSuspendRenderThreadTasks);
-}
\ No newline at end of file
diff --git a/Engine/Source/Runtime/Engine/Private/Streaming/Texture2DUpdate.h b/Engine/Source/Runtime/Engine/Private/Streaming/Texture2DUpdate.h
--- a/Engine/Source/Runtime/Engine/Private/Streaming/Texture2DUpdate.h
+++ b/Engine/Source/Runtime/Engine/Private/Streaming/Texture2DUpdate.h
@@ -7,127 +7,66 @@ Texture2DUpdate.h: Helpers to stream in and out mips.
 #pragma once
 
 #include "CoreMinimal.h"
-#include "HAL/ThreadSafeCounter.h"
-#include "Async/AsyncWork.h"
+#include "Streaming/RenderAssetUpdate.h"
 #include "Engine/Texture2D.h"
 #include "Async/AsyncFileHandle.h"
 
-#define TEXTURE2D_UPDATE_CALLBACK(FunctionName) [this](const FContext& C){ FunctionName(C); }
 extern TAutoConsoleVariable<int32> CVarFlushRHIThreadOnSTreamingTextureLocks;
 
+/**
+* A context used to update or proceed with the next update step.
+* The texture and resource references could be stored in the update object
+* but are currently kept outside to avoid lifetime management within the object.
+*/
+struct FTexture2DUpdateContext
+{
+	typedef int32 EThreadType;
+
+	FTexture2DUpdateContext(UTexture2D* InTexture, EThreadType InCurrentThread);
+
+	FTexture2DUpdateContext(UStreamableRenderAsset* InTexture, EThreadType InCurrentThread);
+
+	UStreamableRenderAsset* GetRenderAsset() const
+	{
+		return Texture;
+	}
+
+	EThreadType GetCurrentThread() const
+	{
+		return CurrentThread;
+	}
+
+	/** The texture to update, this must be the same one as the one used when creating the FTexture2DUpdate object. */
+	UTexture2D* Texture;
+	/** The current 2D resource of this texture. */
+	FTexture2DResource* Resource;
+	/** The thread on which the context was created. */
+	EThreadType CurrentThread;
+};
+
 /**
  * This class provides a framework for loading and unloading the mips of 2D textures.
  * Each thread essentially calls Tick() until the job is done.
  * The object can be safely deleted when IsCompleted() returns true.
  */
-class FTexture2DUpdate
+class FTexture2DUpdate : public TRenderAssetUpdate<FTexture2DUpdateContext>
 {
 public:
-	/** A thread type used for doing a part of the update process.  */
-	enum EThreadType
-	{
-		TT_None,	// No thread.
-		TT_Render,	// The render thread.
-		TT_Async	// An async work thread.
-	};
-
-	/**  The state of scheduled work for the update process. */
-	enum ETaskState
-	{
-		TS_None,		// Nothing to do.
-		TS_Pending,		// The next task (or update step) is configured, but a callback has not been scheduled yet.
-		TS_Scheduled,	// The next task (or update step) is configured and a callback has been scheduled on either the renderthread or async thread.
-		TS_Locked,		// The object is locked, and no one is allowed to process or look at the next task.
-	};
-
 	FTexture2DUpdate(UTexture2D* InTexture, int32 InRequestedMips);
-	virtual ~FTexture2DUpdate();
-
-	/**
-	 * Do or schedule any pending work for a given texture.
-	 *
-	 * @param InTexture - the texture being updated, this must be the same texture as the texture used to create this object.
-	 * @param InCurrentThread - the thread from which the tick is being called. Using TT_None ensures that no work will be immediately performed.
-	 */
-	void Tick(UTexture2D* InTexture, EThreadType InCurrentThread);
 
-	/** Returns whether the task has finished executing and there is no other thread possibly accessing it. */
-	bool IsCompleted() const { return ScheduledTaskCount <= 0 && TaskState == TS_None; }
+	virtual ~FTexture2DUpdate();
 
-	/** Cancel the current update. Will also attempt to cancel pending IO requests, see FTexture2DStreamIn_IO::Abort(). */
-	virtual void Abort() { MarkAsCancelled(); }
+	virtual void Abort()
+	{
+		TRenderAssetUpdate<FTexture2DUpdateContext>::Abort();
+	}
 
 #if WITH_EDITORONLY_DATA
 	/** Returns whether DDC of this texture needs to be regenerated.  */
 	virtual bool DDCIsInvalid() const { return false; }
 #endif
 
-	/** Returns whether the task was aborted through Abort() or cancelled.  */
-	bool IsCancelled() const  { return bIsCancelled; }
-
-	/** 
-	 * Perform a lock on the object, preventing any other thread from processing a pending task in Tick().  
-	 * The lock will stall if there was another thread already processing the task data.
-	 */
-	void DoLock();
-	
-	/** Release any lock on the object, allowing other thread to modify it. */
-	void DoUnlock();
-
-	/** Get the number of requested mips for this update, ignoring cancellation attempts. */
-	int32 GetNumRequestedMips() const { return RequestedMips; }
-
 protected:
-
-	/** Set the task state as cancelled. This is internally called in Abort() and when any critical conditions are not met when performing the update. */
-	void MarkAsCancelled() { bIsCancelled = true; }
-
-	/**
-	 * A context used to update or proceed with the next update step.
-	 * The texture and resource references could be stored in the update object
-	 * but are currently kept outside to avoid lifetime management within the object.
-	 */
-	struct FContext
-	{
-		FContext(UTexture2D* InTexture, EThreadType InCurrentThread);
-
-		/** The texture to update, this must be the same one as the one used when creating the FTexture2DUpdate object. */
-		UTexture2D* Texture;
-		/** The current 2D resource of this texture. */
-		FTexture2DResource* Resource;
-		/** The thread on which the context was created. */
-		EThreadType CurrentThread;
-	};
-
-	/** A callback used to perform a task in the update process. Each task must be executed on a specific thread. */
-	typedef TFunction<void(const FContext& Context)> FCallback;
-
-	/**
-	 * Defines the next step to be executed. The next step will be executed by calling the callback on the specified thread.
-	 * The callback (for both success and cancelation) will only be executed if TaskSynchronization reaches 0.
-	 * If all requirements are immediately satisfied when calling the PushTask the relevant callback will be called immediately.
-	 *
-	 * @param Context - The context defining which texture is being updated and on which thread this is being called.
-	 * @param InTaskThread - The thread on which to call the next step of the update, being TaskCallback.
-	 * @param InTaskCallback - The callback that will perform the next step of the update.
-	 * @param InCancelationThread - The thread on which to call the cancellation of the update (only if the update gets cancelled).
-	 * @param InCancelationCallback - The callback handling the cancellation of the update (only if the update gets cancelled).
-	 */
-	void PushTask(const FContext& Context, EThreadType InTaskThread, const FCallback& InTaskCallback, EThreadType InCancelationThread, const FCallback& InCancelationCallback);
-
-	/** The intermediate texture created in the update process. In the virtual path, this can exceptionally end update being the same as the original texture. */
-	FTexture2DRHIRef IntermediateTextureRHI;
-	/** The mip index within UTexture2D::GetPlatformMips() that will end as being the first mip of the intermediate (future) texture. */
-	int32 PendingFirstMip;
-	/** The total number of mips of the intermediate (future) texture. */
-	int32 RequestedMips;
-
-	/** Synchronization used for trigger the task next step execution. */
-	FThreadSafeCounter	TaskSynchronization;
-
-	/** The number of scheduled ticks (and exceptionally other work from FCancelIORequestsTask) from the renderthread and async thread. Used to prevent deleting the object while it could be accessed. */
-	volatile int32 ScheduledTaskCount;
-
 	// ****************************
 	// ********* Helpers **********
 	// ****************************
@@ -141,68 +80,6 @@ protected:
 	/** Apply the new texture (if not cancelled) and finish the update process. When cancelled, the intermediate texture simply gets discarded. */
 	void DoFinishUpdate(const FContext& Context);
 
-private:
-
-	/**
-	 * Locks the object only if there is work to do. If the calling thread is has no capability to actually perform any work,
-	 * the lock attempt will also fails if the object is already locked. This is used essentially to prevent the gamethread
-	 * from being blocked when ticking the update while the update is already being ticked on another thread.
-	 *
-	 * @param InCurrentThread - the thread trying to lock the object. Passing TS_None may end up preventing to acquired the object.
-	 * @return Whether the lock succeeded.
-	 */
-	bool DoConditionalLock(EThreadType InCurrentThread);
-
-	/**
-	 * Pushes a callback on either the gamethread or an async work thread.
-	 *
-	 * @param Context - the context from the caller thread.
-	 * @param InThread - the thread on which to schedule a tick.
-	 */
-	void ScheduleTick(const FContext& Context, EThreadType InThread);
-
-	/** Clears any pending work. */
-	void ClearTask();
-
-	/** An async task used to call tick on the pending update. */
-	class FMipUpdateTask : public FNonAbandonableTask
-	{
-	public:
-		FMipUpdateTask(UTexture2D* InTexture, FTexture2DUpdate* InPendingUpdate) : Texture(InTexture), CachedPendingUpdate(InPendingUpdate) {}
-		void DoWork();
-		FORCEINLINE TStatId GetStatId() const
-		{
-			RETURN_QUICK_DECLARE_CYCLE_STAT(FMipUpdateTask, STATGROUP_ThreadPoolAsyncTasks);
-		}
-	protected:
-		UTexture2D* Texture;
-		FTexture2DUpdate* CachedPendingUpdate;
-	};
-
-	/** The async task to update this object, only one can be active at anytime. It just calls Tick(). */
-	typedef FAutoDeleteAsyncTask<FMipUpdateTask> FAsyncMipUpdateTask;
-
-	/** Whether the task has been cancelled because the update could not proceed or because the user called Abort(). */
-	bool bIsCancelled;
-
-	/** The state of the work yet to be performed to complete the update or cancelation. */
-	volatile int32 TaskState;
-	/** The pending state of future work. Used because the object is in locked state (TS_Locked) when being updated, until the call to Unlock() */
-	ETaskState PendingTaskState;
-
-	/** The thread on which to call the next step of the update, being TaskCallback. */
-	EThreadType TaskThread;
-	/** The callback that will perform the next step of the update. */
-	FCallback TaskCallback;
-	/** The thread on which to call the cancellation of the update (only if the update gets cancelled). */
-	EThreadType CancelationThread; // The thread on which the callbacks should be called.
-	/** The callback handling the cancellation of the update (only if the update gets cancelled). */
-	FCallback CancelationCallback;
-
-public:
-	/** A counter defining whether render thread tasks should be postponed The callback handling the cancellation of the update (only if the update gets cancelled). */
-	static volatile int32 GSuspendRenderThreadTasks;
+	/** The intermediate texture created in the update process. In the virtual path, this can exceptionally end update being the same as the original texture. */
+	FTexture2DRHIRef IntermediateTextureRHI;
 };
-
-void SuspendTextureStreamingRenderTasksInternal();
-void ResumeTextureStreamingRenderTasksInternal();
diff --git a/Engine/Source/Runtime/Engine/Private/Streaming/TextureInstanceManager.h b/Engine/Source/Runtime/Engine/Private/Streaming/TextureInstanceManager.h
--- a/Engine/Source/Runtime/Engine/Private/Streaming/TextureInstanceManager.h
+++ b/Engine/Source/Runtime/Engine/Private/Streaming/TextureInstanceManager.h
@@ -10,11 +10,11 @@ TextureInstanceManager.h: Definitions of classes used for texture streaming.
 #include "Streaming/TextureInstanceState.h"
 
 /** An interface to manage the TextureInstanceState from a group of components with similar properties, allowing to add/remove components incrementally and update their content. */
-class ITextureInstanceManager
+class IRenderAssetInstanceManager
 {
 public:
 
-	virtual ~ITextureInstanceManager() {}
+	virtual ~IRenderAssetInstanceManager() {}
 
 	/** Return whether the component is referenced. */
 	virtual bool IsReferenced(const UPrimitiveComponent* Component) const = 0;
@@ -30,13 +30,13 @@ public:
 	virtual EAddComponentResult Add(const UPrimitiveComponent* Component, FStreamingTextureLevelContext& LevelContext, float MaxAllowedUIDensity) = 0;
 
 	/** Remove a component, the RemoveTextures is the list of textures not referred anymore. */
-	virtual void Remove(const UPrimitiveComponent* Component, FRemovedTextureArray* RemovedTextures) = 0;
+	virtual void Remove(const UPrimitiveComponent* Component, FRemovedRenderAssetArray* RemovedTextures) = 0;
 
 	/** Notify the manager that an async view will be requested on the next frame. */
 	virtual void PrepareAsyncView() = 0;
 
 	/** Return a view of the data that has to be 100% thread safe. The content is allowed to be updated, but not memory must be reallocated. */
-	virtual const FTextureInstanceView* GetAsyncView(bool bCreateIfNull) = 0;
+	virtual const FRenderAssetInstanceView* GetAsyncView(bool bCreateIfNull) = 0;
 
 	/** Return the size taken for sub-allocation. */
 	virtual uint32 GetAllocatedSize() const = 0;
