Resolve robomerge conflicts

diff --git a/Engine/Source/Runtime/Engine/Private/Streaming/TextureInstanceState.cpp b/Engine/Source/Runtime/Engine/Private/Streaming/TextureInstanceState.cpp
--- a/Engine/Source/Runtime/Engine/Private/Streaming/TextureInstanceState.cpp
+++ b/Engine/Source/Runtime/Engine/Private/Streaming/TextureInstanceState.cpp
@@ -14,12 +14,12 @@
 #include "UObject/UObjectHash.h"
 #include "Templates/RefCounting.h"
 
-int32 FTextureInstanceState::AddBounds(const UPrimitiveComponent* Component)
+int32 FRenderAssetInstanceState::AddBounds(const UPrimitiveComponent* Component)
 {
 	return AddBounds(Component->Bounds, PackedRelativeBox_Identity, Component, Component->LastRenderTimeOnScreen, Component->Bounds.Origin, 0, 0, FLT_MAX);
 }
 
-int32 FTextureInstanceState::AddBounds(const FBoxSphereBounds& Bounds, uint32 PackedRelativeBox, const UPrimitiveComponent* InComponent, float LastRenderTime, const FVector4& RangeOrigin, float MinDistance, float MinRange, float MaxRange)
+int32 FRenderAssetInstanceState::AddBounds(const FBoxSphereBounds& Bounds, uint32 PackedRelativeBox, const UPrimitiveComponent* InComponent, float LastRenderTime, const FVector4& RangeOrigin, float MinDistance, float MinRange, float MaxRange)
 {
 	check(InComponent);
 
@@ -52,7 +52,7 @@ int32 FTextureInstanceState::AddBounds(const FBoxSphereBounds& Bounds, uint32 Pa
 	return BoundsIndex;
 }
 
-void FTextureInstanceState::RemoveBounds(int32 BoundsIndex)
+void FRenderAssetInstanceState::RemoveBounds(int32 BoundsIndex)
 {
 	checkSlow(!FreeBoundIndices.Contains(BoundsIndex));
 
@@ -91,9 +91,9 @@ void FTextureInstanceState::RemoveBounds(int32 BoundsIndex)
 	}
 }
 
-void FTextureInstanceState::AddElement(const UPrimitiveComponent* InComponent, const UTexture2D* InTexture, int InBoundsIndex, float InTexelFactor, bool InForceLoad, int32*& ComponentLink)
+void FRenderAssetInstanceState::AddElement(const UPrimitiveComponent* InComponent, const UStreamableRenderAsset* InAsset, int InBoundsIndex, float InTexelFactor, bool InForceLoad, int32*& ComponentLink)
 {
-	check(InComponent && InTexture);
+	check(InComponent && InAsset);
 
 	// Keep Max texel factor up to date.
 	MaxTexelFactor = FMath::Max(InTexelFactor, MaxTexelFactor);
@@ -112,24 +112,24 @@ void FTextureInstanceState::AddElement(const UPrimitiveComponent* InComponent, c
 	FElement& Element = Elements[ElementIndex];
 
 	Element.Component = InComponent;
-	Element.Texture = InTexture;
+	Element.RenderAsset = InAsset;
 	Element.BoundsIndex = InBoundsIndex;
 	Element.TexelFactor = InTexelFactor;
 	Element.bForceLoad = InForceLoad;
 
-	FTextureDesc* TextureDesc = TextureMap.Find(InTexture);
-	if (TextureDesc)
+	FRenderAssetDesc* AssetDesc = RenderAssetMap.Find(InAsset);
+	if (AssetDesc)
 	{
-		FElement& TextureLinkElement = Elements[TextureDesc->HeadLink];
+		FElement& AssetLinkElement = Elements[AssetDesc->HeadLink];
 
 		// The new inserted element as the head element.
-		Element.NextTextureLink = TextureDesc->HeadLink;
-		TextureLinkElement.PrevTextureLink = ElementIndex;
-		TextureDesc->HeadLink = ElementIndex;
+		Element.NextRenderAssetLink = AssetDesc->HeadLink;
+		AssetLinkElement.PrevRenderAssetLink = ElementIndex;
+		AssetDesc->HeadLink = ElementIndex;
 	}
 	else
 	{
-		TextureMap.Add(InTexture, FTextureDesc(ElementIndex, InTexture->LODGroup));
+		RenderAssetMap.Add(InAsset, FRenderAssetDesc(ElementIndex, InAsset->GetLODGroupForStreaming()));
 	}
 
 	// Simple sanity check to ensure that the component link passed in param is the right one
@@ -149,11 +149,11 @@ void FTextureInstanceState::AddElement(const UPrimitiveComponent* InComponent, c
 	// This will happen when not all components could be inserted in the incremental build.
 	if (HasCompiledElements()) 
 	{
-		CompiledTextureMap.FindOrAdd(Element.Texture).Add(FCompiledElement(Element));
+		CompiledRenderAssetMap.FindOrAdd(Element.RenderAsset).Add(FCompiledElement(Element));
 	}
 }
 
-void FTextureInstanceState::RemoveElement(int32 ElementIndex, int32& NextComponentLink, int32& BoundsIndex, const UTexture2D*& Texture)
+void FRenderAssetInstanceState::RemoveElement(int32 ElementIndex, int32& NextComponentLink, int32& BoundsIndex, const UStreamableRenderAsset*& Asset)
 {
 	FElement& Element = Elements[ElementIndex];
 	NextComponentLink = Element.NextComponentLink; 
@@ -162,34 +162,34 @@ void FTextureInstanceState::RemoveElement(int32 ElementIndex, int32& NextCompone
 	// Removed compiled elements. This happens when a static component is not registered after the level became visible.
 	if (HasCompiledElements())
 	{
-		CompiledTextureMap.FindChecked(Element.Texture).RemoveSingleSwap(FCompiledElement(Element), false);
+		CompiledRenderAssetMap.FindChecked(Element.RenderAsset).RemoveSingleSwap(FCompiledElement(Element), false);
 	}
 
-	// Unlink textures
-	if (Element.Texture)
+	// Unlink texutres or meshes
+	if (Element.RenderAsset)
 	{
-		if (Element.PrevTextureLink == INDEX_NONE) // If NONE, that means this is the head of the texture list.
+		if (Element.PrevRenderAssetLink == INDEX_NONE) // If NONE, that means this is the head of the texture/mesh list.
 		{
-			if (Element.NextTextureLink != INDEX_NONE) // Check if there are other entries for this texture.
+			if (Element.NextRenderAssetLink != INDEX_NONE) // Check if there are other entries for this texture/mesh.
 			{
 				 // Replace the head
-				TextureMap.Find(Element.Texture)->HeadLink = Element.NextTextureLink;
-				Elements[Element.NextTextureLink].PrevTextureLink = INDEX_NONE;
+				RenderAssetMap.Find(Element.RenderAsset)->HeadLink = Element.NextRenderAssetLink;
+				Elements[Element.NextRenderAssetLink].PrevRenderAssetLink = INDEX_NONE;
 			}
-			else // Otherwise, remove the texture entry
+			else // Otherwise, remove the texture/mesh entry
 			{
-				TextureMap.Remove(Element.Texture);
-				CompiledTextureMap.Remove(Element.Texture);
-				Texture = Element.Texture;
+				RenderAssetMap.Remove(Element.RenderAsset);
+				CompiledRenderAssetMap.Remove(Element.RenderAsset);
+				Asset = Element.RenderAsset;
 			}
 		}
 		else // Otherwise, just relink entries.
 		{
-			Elements[Element.PrevTextureLink].NextTextureLink = Element.NextTextureLink;
+			Elements[Element.PrevRenderAssetLink].NextRenderAssetLink = Element.NextRenderAssetLink;
 
-			if (Element.NextTextureLink != INDEX_NONE)
+			if (Element.NextRenderAssetLink != INDEX_NONE)
 			{
-				Elements[Element.NextTextureLink].PrevTextureLink = Element.PrevTextureLink;
+				Elements[Element.NextRenderAssetLink].PrevRenderAssetLink = Element.PrevRenderAssetLink;
 			}
 		}
 	}
@@ -215,22 +215,22 @@ FORCEINLINE bool operator<(const FBoxSphereBounds& Lhs, const FBoxSphereBounds&
 	return FMemory::Memcmp(&Lhs, &Rhs, sizeof(FBoxSphereBounds)) < 0;
 }
 
-void FTextureInstanceState::AddTextureElements(const UPrimitiveComponent* Component, const TArrayView<FStreamingTexturePrimitiveInfo>& TextureInstanceInfos, int32 BoundsIndex, int32*& ComponentLink)
+void FRenderAssetInstanceState::AddRenderAssetElements(const UPrimitiveComponent* Component, const TArrayView<FStreamingRenderAssetPrimitiveInfo>& RenderAssetInstanceInfos, int32 BoundsIndex, int32*& ComponentLink)
 {
-	// Loop for each texture - texel factor group (a group being of same texel factor sign)
-	for (int32 InfoIndex = 0; InfoIndex < TextureInstanceInfos.Num();)
+	// Loop for each render asset - texel factor group (a group being of same texel factor sign)
+	for (int32 InfoIndex = 0; InfoIndex < RenderAssetInstanceInfos.Num();)
 	{
-		const FStreamingTexturePrimitiveInfo& Info = TextureInstanceInfos[InfoIndex];
+		const FStreamingRenderAssetPrimitiveInfo& Info = RenderAssetInstanceInfos[InfoIndex];
 		float MergedTexelFactor = Info.TexelFactor;
 		int32 NumOfMergedElements = 1;
 
 		// Merge all texel factor >= 0 for the same texture, or all those < 0
 		if (Info.TexelFactor >= 0)
 		{
-			for (int32 NextInfoIndex = InfoIndex + 1; NextInfoIndex < TextureInstanceInfos.Num(); ++NextInfoIndex)
+			for (int32 NextInfoIndex = InfoIndex + 1; NextInfoIndex < RenderAssetInstanceInfos.Num(); ++NextInfoIndex)
 			{
-				const FStreamingTexturePrimitiveInfo& NextInfo = TextureInstanceInfos[NextInfoIndex];
-				if (NextInfo.Texture == Info.Texture && NextInfo.TexelFactor >= 0)
+				const FStreamingRenderAssetPrimitiveInfo& NextInfo = RenderAssetInstanceInfos[NextInfoIndex];
+				if (NextInfo.RenderAsset == Info.RenderAsset && NextInfo.TexelFactor >= 0)
 				{
 					MergedTexelFactor = FMath::Max(MergedTexelFactor, NextInfo.TexelFactor);
 					++NumOfMergedElements;
@@ -243,10 +243,10 @@ void FTextureInstanceState::AddTextureElements(const UPrimitiveComponent* Compon
 		}
 		else // Info.TexelFactor < 0
 		{
-			for (int32 NextInfoIndex = InfoIndex + 1; NextInfoIndex < TextureInstanceInfos.Num(); ++NextInfoIndex)
+			for (int32 NextInfoIndex = InfoIndex + 1; NextInfoIndex < RenderAssetInstanceInfos.Num(); ++NextInfoIndex)
 			{
-				const FStreamingTexturePrimitiveInfo& NextInfo = TextureInstanceInfos[NextInfoIndex];
-				if (NextInfo.Texture == Info.Texture && NextInfo.TexelFactor < 0)
+				const FStreamingRenderAssetPrimitiveInfo& NextInfo = RenderAssetInstanceInfos[NextInfoIndex];
+				if (NextInfo.RenderAsset == Info.RenderAsset && NextInfo.TexelFactor < 0)
 				{
 					MergedTexelFactor = FMath::Min(MergedTexelFactor, NextInfo.TexelFactor);
 					++NumOfMergedElements;
@@ -257,21 +257,20 @@ void FTextureInstanceState::AddTextureElements(const UPrimitiveComponent* Compon
 				}
 			}
 		}
-		AddElement(Component, Info.Texture, BoundsIndex, MergedTexelFactor, Component->bForceMipStreaming, ComponentLink);
+		AddElement(Component, Info.RenderAsset, BoundsIndex, MergedTexelFactor, Component->bForceMipStreaming, ComponentLink);
 
 		InfoIndex += NumOfMergedElements;
 	}
 }
 
-EAddComponentResult FTextureInstanceState::AddComponent(const UPrimitiveComponent* Component, FStreamingTextureLevelContext& LevelContext, float MaxAllowedUIDensity)
+EAddComponentResult FRenderAssetInstanceState::AddComponent(const UPrimitiveComponent* Component, FStreamingTextureLevelContext& LevelContext, float MaxAllowedUIDensity)
 {
 	check(Component);
 
-	TArray<FStreamingTexturePrimitiveInfo> TextureInstanceInfos;
-	Component->GetStreamingTextureInfoWithNULLRemoval(LevelContext, TextureInstanceInfos);
+	TArray<FStreamingRenderAssetPrimitiveInfo> RenderAssetInstanceInfos;
+	Component->GetStreamingRenderAssetInfoWithNULLRemoval(LevelContext, RenderAssetInstanceInfos);
 	// Texture entries are guarantied to be relevant here, except for bounds if the component is not registered.
-
-	if (!TextureInstanceInfos.Num())
+	if (!RenderAssetInstanceInfos.Num())
 	{
 		return EAddComponentResult::Fail;
 	}
@@ -279,7 +278,7 @@ EAddComponentResult FTextureInstanceState::AddComponent(const UPrimitiveComponen
 	// First check if all entries are below the max allowed UI density, otherwise abort immediately.
 	if (MaxAllowedUIDensity > 0)
 	{
-		for (const FStreamingTexturePrimitiveInfo& Info : TextureInstanceInfos)
+		for (const FStreamingRenderAssetPrimitiveInfo& Info : RenderAssetInstanceInfos)
 		{
 			if (Info.TexelFactor > MaxAllowedUIDensity)
 			{
@@ -293,7 +292,7 @@ EAddComponentResult FTextureInstanceState::AddComponent(const UPrimitiveComponen
 		// When the components are not registered, the bound will be generated from PackedRelativeBox in CheckRegistrationAndUnpackBounds.
 		// Otherwise, the entry is not usable as we don't know the bound to use. The component will need to be reinstered later, once registered.
 		// it will not be possible to recreate the bounds correctly.
-		for (const FStreamingTexturePrimitiveInfo& Info : TextureInstanceInfos)
+		for (const FStreamingRenderAssetPrimitiveInfo& Info : RenderAssetInstanceInfos)
 		{
 			if (!Info.PackedRelativeBox)
 			{
@@ -304,15 +303,15 @@ EAddComponentResult FTextureInstanceState::AddComponent(const UPrimitiveComponen
 		// Sort by PackedRelativeBox, to identical bounds identical entries together
 		// Sort by Texture to merge duplicate texture entries.
 		// Then sort by TexelFactor, to merge negative entries together.
-		TextureInstanceInfos.Sort([](const FStreamingTexturePrimitiveInfo& Lhs, const FStreamingTexturePrimitiveInfo& Rhs) 
+		RenderAssetInstanceInfos.Sort([](const FStreamingRenderAssetPrimitiveInfo& Lhs, const FStreamingRenderAssetPrimitiveInfo& Rhs)
 		{
 			if (Lhs.PackedRelativeBox == Rhs.PackedRelativeBox)
 			{
-				if (Lhs.Texture == Rhs.Texture)
+				if (Lhs.RenderAsset == Rhs.RenderAsset)
 				{
 					return Lhs.TexelFactor < Rhs.TexelFactor;
 				}
-				return Lhs.Texture < Rhs.Texture;
+				return Lhs.RenderAsset < Rhs.RenderAsset;
 			}
 			return Lhs.PackedRelativeBox < Rhs.PackedRelativeBox;
 		});
@@ -320,18 +319,18 @@ EAddComponentResult FTextureInstanceState::AddComponent(const UPrimitiveComponen
 		int32* ComponentLink = ComponentMap.Find(Component);
 
 		// Loop for each bound.
-		for (int32 InfoIndex = 0; InfoIndex < TextureInstanceInfos.Num();)
+		for (int32 InfoIndex = 0; InfoIndex < RenderAssetInstanceInfos.Num();)
 		{
-			const FStreamingTexturePrimitiveInfo& Info = TextureInstanceInfos[InfoIndex];
+			const FStreamingRenderAssetPrimitiveInfo& Info = RenderAssetInstanceInfos[InfoIndex];
 
 			int32 NumOfBoundReferences = 1;
-			for (int32 NextInfoIndex = InfoIndex + 1; NextInfoIndex < TextureInstanceInfos.Num() && TextureInstanceInfos[NextInfoIndex].PackedRelativeBox == Info.PackedRelativeBox; ++NextInfoIndex)
+			for (int32 NextInfoIndex = InfoIndex + 1; NextInfoIndex < RenderAssetInstanceInfos.Num() && RenderAssetInstanceInfos[NextInfoIndex].PackedRelativeBox == Info.PackedRelativeBox; ++NextInfoIndex)
 			{
 				++NumOfBoundReferences;
 			}
 
 			const int32 BoundsIndex = AddBounds(FBoxSphereBounds(ForceInit), Info.PackedRelativeBox, Component, Component->LastRenderTimeOnScreen, FVector(ForceInit), 0, 0, FLT_MAX);
-			AddTextureElements(Component, TArrayView<FStreamingTexturePrimitiveInfo>(TextureInstanceInfos.GetData() + InfoIndex, NumOfBoundReferences), BoundsIndex, ComponentLink);
+			AddRenderAssetElements(Component, TArrayView<FStreamingRenderAssetPrimitiveInfo>(RenderAssetInstanceInfos.GetData() + InfoIndex, NumOfBoundReferences), BoundsIndex, ComponentLink);
 			BoundsToUnpack.Push(BoundsIndex);
 
 			InfoIndex += NumOfBoundReferences;
@@ -342,15 +341,15 @@ EAddComponentResult FTextureInstanceState::AddComponent(const UPrimitiveComponen
 		// Sort by Bounds, to merge identical bounds entries together
 		// Sort by Texture to merge duplicate texture entries.
 		// Then sort by TexelFactor, to merge negative entries together.
-		TextureInstanceInfos.Sort([](const FStreamingTexturePrimitiveInfo& Lhs, const FStreamingTexturePrimitiveInfo& Rhs) 
+		RenderAssetInstanceInfos.Sort([](const FStreamingRenderAssetPrimitiveInfo& Lhs, const FStreamingRenderAssetPrimitiveInfo& Rhs)
 		{
 			if (Lhs.Bounds == Rhs.Bounds)
 			{
-				if (Lhs.Texture == Rhs.Texture)
+				if (Lhs.RenderAsset == Rhs.RenderAsset)
 				{
 					return Lhs.TexelFactor < Rhs.TexelFactor;
 				}
-				return Lhs.Texture < Rhs.Texture;
+				return Lhs.RenderAsset < Rhs.RenderAsset;
 			}
 			return Lhs.Bounds < Rhs.Bounds;
 		});
@@ -359,19 +358,19 @@ EAddComponentResult FTextureInstanceState::AddComponent(const UPrimitiveComponen
 		float MinDistance = 0, MinRange = 0, MaxRange = FLT_MAX;
 
 		// Loop for each bound.
-		for (int32 InfoIndex = 0; InfoIndex < TextureInstanceInfos.Num();)
+		for (int32 InfoIndex = 0; InfoIndex < RenderAssetInstanceInfos.Num();)
 		{
-			const FStreamingTexturePrimitiveInfo& Info = TextureInstanceInfos[InfoIndex];
+			const FStreamingRenderAssetPrimitiveInfo& Info = RenderAssetInstanceInfos[InfoIndex];
 
 			int32 NumOfBoundReferences = 1;
-			for (int32 NextInfoIndex = InfoIndex + 1; NextInfoIndex < TextureInstanceInfos.Num() && TextureInstanceInfos[NextInfoIndex].Bounds == Info.Bounds; ++NextInfoIndex)
+			for (int32 NextInfoIndex = InfoIndex + 1; NextInfoIndex < RenderAssetInstanceInfos.Num() && RenderAssetInstanceInfos[NextInfoIndex].Bounds == Info.Bounds; ++NextInfoIndex)
 			{
 				++NumOfBoundReferences;
 			}
 
 			GetDistanceAndRange(Component, Info.Bounds, MinDistance, MinRange, MaxRange);
 			const int32 BoundsIndex = AddBounds(Info.Bounds, PackedRelativeBox_Identity, Component, Component->LastRenderTimeOnScreen, Component->Bounds.Origin, MinDistance, MinRange, MaxRange);
-			AddTextureElements(Component, TArrayView<FStreamingTexturePrimitiveInfo>(TextureInstanceInfos.GetData() + InfoIndex, NumOfBoundReferences), BoundsIndex, ComponentLink);
+			AddRenderAssetElements(Component, TArrayView<FStreamingRenderAssetPrimitiveInfo>(RenderAssetInstanceInfos.GetData() + InfoIndex, NumOfBoundReferences), BoundsIndex, ComponentLink);
 
 			InfoIndex += NumOfBoundReferences;
 		}
@@ -379,36 +378,36 @@ EAddComponentResult FTextureInstanceState::AddComponent(const UPrimitiveComponen
 	return EAddComponentResult::Success;
 }
 
-EAddComponentResult FTextureInstanceState::AddComponentIgnoreBounds(const UPrimitiveComponent* Component, FStreamingTextureLevelContext& LevelContext)
+EAddComponentResult FRenderAssetInstanceState::AddComponentIgnoreBounds(const UPrimitiveComponent* Component, FStreamingTextureLevelContext& LevelContext)
 {
 	check(Component->IsRegistered()); // Must be registered otherwise bounds are invalid.
 
-	TArray<FStreamingTexturePrimitiveInfo> TextureInstanceInfos;
-	Component->GetStreamingTextureInfoWithNULLRemoval(LevelContext, TextureInstanceInfos);
+	TArray<FStreamingRenderAssetPrimitiveInfo> RenderAssetInstanceInfos;
+	Component->GetStreamingRenderAssetInfoWithNULLRemoval(LevelContext, RenderAssetInstanceInfos);
 
-	if (!TextureInstanceInfos.Num())
+	if (!RenderAssetInstanceInfos.Num())
 	{
 		return EAddComponentResult::Fail;
 	}
 
 	// Sort by Texture to merge duplicate texture entries.
 	// Then sort by TexelFactor, to merge negative entries together.
-	TextureInstanceInfos.Sort([](const FStreamingTexturePrimitiveInfo& Lhs, const FStreamingTexturePrimitiveInfo& Rhs) 
+	RenderAssetInstanceInfos.Sort([](const FStreamingRenderAssetPrimitiveInfo& Lhs, const FStreamingRenderAssetPrimitiveInfo& Rhs)
 	{
-		if (Lhs.Texture == Rhs.Texture)
+		if (Lhs.RenderAsset == Rhs.RenderAsset)
 		{
 			return Lhs.TexelFactor < Rhs.TexelFactor;
 		}
-		return Lhs.Texture < Rhs.Texture;
+		return Lhs.RenderAsset < Rhs.RenderAsset;
 	});
 
 	int32* ComponentLink = ComponentMap.Find(Component);
-	AddTextureElements(Component, TextureInstanceInfos, AddBounds(Component), ComponentLink);
+	AddRenderAssetElements(Component, RenderAssetInstanceInfos, AddBounds(Component), ComponentLink);
 	return EAddComponentResult::Success;
 }
 
 
-void FTextureInstanceState::RemoveComponent(const UPrimitiveComponent* Component, FRemovedTextureArray* RemovedTextures)
+void FRenderAssetInstanceState::RemoveComponent(const UPrimitiveComponent* Component, FRemovedRenderAssetArray* RemovedTextures)
 {
 	TArray<int32, TInlineAllocator<12> > RemovedBoundsIndices;
 	int32 ElementIndex = INDEX_NONE;
@@ -417,7 +416,7 @@ void FTextureInstanceState::RemoveComponent(const UPrimitiveComponent* Component
 	while (ElementIndex != INDEX_NONE)
 	{
 		int32 BoundsIndex = INDEX_NONE;
-		const UTexture2D* Texture = nullptr;
+		const UStreamableRenderAsset* Texture = nullptr;
 
 		RemoveElement(ElementIndex, ElementIndex, BoundsIndex, Texture);
 
@@ -438,7 +437,7 @@ void FTextureInstanceState::RemoveComponent(const UPrimitiveComponent* Component
 	}
 }
 
-bool FTextureInstanceState::RemoveComponentReferences(const UPrimitiveComponent* Component) 
+bool FRenderAssetInstanceState::RemoveComponentReferences(const UPrimitiveComponent* Component) 
 { 
 	// Because the async streaming task could be running, we can't change the async view state. 
 	// We limit ourself to clearing the component ptr to avoid invalid access when updating visibility.
@@ -468,7 +467,7 @@ bool FTextureInstanceState::RemoveComponentReferences(const UPrimitiveComponent*
 	}
 }
 
-void FTextureInstanceState::GetReferencedComponents(TArray<const UPrimitiveComponent*>& Components) const
+void FRenderAssetInstanceState::GetReferencedComponents(TArray<const UPrimitiveComponent*>& Components) const
 {
 	for (TMap<const UPrimitiveComponent*, int32>::TConstIterator It(ComponentMap); It; ++It)
 	{
@@ -476,7 +475,7 @@ void FTextureInstanceState::GetReferencedComponents(TArray<const UPrimitiveCompo
 	}
 }
 
-void FTextureInstanceState::UpdateBounds(const UPrimitiveComponent* Component)
+void FRenderAssetInstanceState::UpdateBounds(const UPrimitiveComponent* Component)
 {
 	int32* ComponentLink = ComponentMap.Find(Component);
 	if (ComponentLink)
@@ -494,7 +493,7 @@ void FTextureInstanceState::UpdateBounds(const UPrimitiveComponent* Component)
 	}
 }
 
-bool FTextureInstanceState::UpdateBounds(int32 BoundIndex)
+bool FRenderAssetInstanceState::UpdateBounds(int32 BoundIndex)
 {
 	const UPrimitiveComponent* Component = ensure(Bounds4Components.IsValidIndex(BoundIndex)) ? Bounds4Components[BoundIndex] : nullptr;
 	if (Component)
@@ -508,7 +507,7 @@ bool FTextureInstanceState::UpdateBounds(int32 BoundIndex)
 	}
 }
 
-bool FTextureInstanceState::ConditionalUpdateBounds(int32 BoundIndex)
+bool FRenderAssetInstanceState::ConditionalUpdateBounds(int32 BoundIndex)
 {
 	const UPrimitiveComponent* Component = ensure(Bounds4Components.IsValidIndex(BoundIndex)) ? Bounds4Components[BoundIndex] : nullptr;
 	if (Component)
@@ -540,7 +539,7 @@ bool FTextureInstanceState::ConditionalUpdateBounds(int32 BoundIndex)
 }
 
 
-void FTextureInstanceState::UpdateLastRenderTime(int32 BoundIndex)
+void FRenderAssetInstanceState::UpdateLastRenderTime(int32 BoundIndex)
 {
 	const UPrimitiveComponent* Component = ensure(Bounds4Components.IsValidIndex(BoundIndex)) ? Bounds4Components[BoundIndex] : nullptr;
 	if (Component)
@@ -549,10 +548,10 @@ void FTextureInstanceState::UpdateLastRenderTime(int32 BoundIndex)
 	}
 }
 
-uint32 FTextureInstanceState::GetAllocatedSize() const
+uint32 FRenderAssetInstanceState::GetAllocatedSize() const
 {
 	int32 CompiledElementsSize = 0;
-	for (TMap<const UTexture2D*, TArray<FCompiledElement> >::TConstIterator It(CompiledTextureMap); It; ++It)
+	for (TMap<const UStreamableRenderAsset*, TArray<FCompiledElement> >::TConstIterator It(CompiledRenderAssetMap); It; ++It)
 	{
 		CompiledElementsSize += It.Value().GetAllocatedSize();
 	}
@@ -562,38 +561,38 @@ uint32 FTextureInstanceState::GetAllocatedSize() const
 		Elements.GetAllocatedSize() +
 		FreeBoundIndices.GetAllocatedSize() +
 		FreeElementIndices.GetAllocatedSize() +
-		TextureMap.GetAllocatedSize() +
-		CompiledTextureMap.GetAllocatedSize() + CompiledElementsSize +
+		RenderAssetMap.GetAllocatedSize() +
+		CompiledRenderAssetMap.GetAllocatedSize() + CompiledElementsSize +
 		ComponentMap.GetAllocatedSize();
 }
 
-int32 FTextureInstanceState::CompileElements()
+int32 FRenderAssetInstanceState::CompileElements()
 {
-	CompiledTextureMap.Empty();
+	CompiledRenderAssetMap.Empty();
 	MaxTexelFactor = 0;
 
 	// First create an entry for all elements, so that there are no reallocs when inserting each compiled elements.
-	for (TMap<const UTexture2D*, FTextureDesc>::TConstIterator TextureIt(TextureMap); TextureIt; ++TextureIt)
+	for (TMap<const UStreamableRenderAsset*, FRenderAssetDesc>::TConstIterator AssetIt(RenderAssetMap); AssetIt; ++AssetIt)
 	{
-		const UTexture2D* Texture = TextureIt.Key();
-		CompiledTextureMap.Add(Texture);
+		const UStreamableRenderAsset* Asset = AssetIt.Key();
+		CompiledRenderAssetMap.Add(Asset);
 	}
 
 	// Then fill in each array.
-	for (TMap<const UTexture2D*, TArray<FCompiledElement> >::TIterator It(CompiledTextureMap); It; ++It)
+	for (TMap<const UStreamableRenderAsset*, TArray<FCompiledElement> >::TIterator It(CompiledRenderAssetMap); It; ++It)
 	{
-		const UTexture2D* Texture = It.Key();
+		const UStreamableRenderAsset* Asset = It.Key();
 		TArray<FCompiledElement>& CompiledElemements = It.Value();
 
 		int32 CompiledElementCount = 0;
-		for (auto ElementIt = GetElementIterator(Texture); ElementIt; ++ElementIt)
+		for (auto ElementIt = GetElementIterator(Asset); ElementIt; ++ElementIt)
 		{
 			++CompiledElementCount;
 		}
 		CompiledElemements.AddUninitialized(CompiledElementCount);
 
 		CompiledElementCount = 0;
-		for (auto ElementIt = GetElementIterator(Texture); ElementIt; ++ElementIt)
+		for (auto ElementIt = GetElementIterator(Asset); ElementIt; ++ElementIt)
 		{
 			const float TexelFactor = ElementIt.GetTexelFactor();
 			// No need to care about force load as MaxTexelFactor is used to ignore far away levels.
@@ -606,10 +605,10 @@ int32 FTextureInstanceState::CompileElements()
 			++CompiledElementCount;
 		}
 	}
-	return CompiledTextureMap.Num();
+	return CompiledRenderAssetMap.Num();
 }
 
-int32 FTextureInstanceState::CheckRegistrationAndUnpackBounds(TArray<const UPrimitiveComponent*>& RemovedComponents)
+int32 FRenderAssetInstanceState::CheckRegistrationAndUnpackBounds(TArray<const UPrimitiveComponent*>& RemovedComponents)
 {
 	for (int32 BoundIndex : BoundsToUnpack)
 	{
@@ -638,7 +637,7 @@ int32 FTextureInstanceState::CheckRegistrationAndUnpackBounds(TArray<const UPrim
 	return NumSteps;
 }
 
-bool FTextureInstanceState::MoveBound(int32 SrcBoundIndex, int32 DstBoundIndex)
+bool FRenderAssetInstanceState::MoveBound(int32 SrcBoundIndex, int32 DstBoundIndex)
 {
 	check(!HasCompiledElements() && !BoundsToUnpack.Num()); // Defrag is for the dynamic elements which does not support dynamic compiled elements.
 
@@ -689,7 +688,7 @@ bool FTextureInstanceState::MoveBound(int32 SrcBoundIndex, int32 DstBoundIndex)
 	}
 }
 
-void FTextureInstanceState::TrimBounds()
+void FRenderAssetInstanceState::TrimBounds()
 {
 	const int32 DefragThreshold = 8; // Must be a multiple of 4
 	check(NumBounds4() * 4 == NumBounds());
@@ -746,7 +745,7 @@ void FTextureInstanceState::TrimBounds()
 	}
 }
 
-void FTextureInstanceState::OffsetBounds(const FVector& Offset)
+void FRenderAssetInstanceState::OffsetBounds(const FVector& Offset)
 {
 	for (int32 BoundIndex = 0; BoundIndex < Bounds4Components.Num(); ++BoundIndex)
 	{
diff --git a/Engine/Source/Runtime/Engine/Private/Streaming/TextureInstanceState.h b/Engine/Source/Runtime/Engine/Private/Streaming/TextureInstanceState.h
--- a/Engine/Source/Runtime/Engine/Private/Streaming/TextureInstanceState.h
+++ b/Engine/Source/Runtime/Engine/Private/Streaming/TextureInstanceState.h
@@ -14,8 +14,7 @@ TextureInstanceState.h: Definitions of classes used for texture streaming.
 class FStreamingTextureLevelContext;
 class ULevel;
 class UPrimitiveComponent;
-class UTexture2D;
-struct FStreamingTexturePrimitiveInfo;
+struct FStreamingRenderAssetPrimitiveInfo;
 
 enum class EAddComponentResult : uint8
 {
@@ -25,7 +24,7 @@ enum class EAddComponentResult : uint8
 };
 
 // Can be used either for static primitives or dynamic primitives
-class FTextureInstanceState : public FTextureInstanceView
+class FRenderAssetInstanceState : public FRenderAssetInstanceView
 {
 public:
 
@@ -37,7 +36,7 @@ public:
 	EAddComponentResult AddComponentIgnoreBounds(const UPrimitiveComponent* Component, FStreamingTextureLevelContext& LevelContext);
 
 	FORCEINLINE bool HasComponentReferences(const UPrimitiveComponent* Component) const { return ComponentMap.Contains(Component); }
-	void RemoveComponent(const UPrimitiveComponent* Component, FRemovedTextureArray* RemovedTextures);
+	void RemoveComponent(const UPrimitiveComponent* Component, FRemovedRenderAssetArray* RemovedTextures);
 	bool RemoveComponentReferences(const UPrimitiveComponent* Component);
 
 	void GetReferencedComponents(TArray<const UPrimitiveComponent*>& Components) const;
@@ -63,15 +62,15 @@ public:
 
 private:
 
-	void AddElement(const UPrimitiveComponent* Component, const UTexture2D* Texture, int BoundsIndex, float TexelFactor, bool bForceLoad, int32*& ComponentLink);
+	void AddElement(const UPrimitiveComponent* Component, const UStreamableRenderAsset* Asset, int BoundsIndex, float TexelFactor, bool bForceLoad, int32*& ComponentLink);
 	// Returns the next elements using the same component.
-	void RemoveElement(int32 ElementIndex, int32& NextComponentLink, int32& BoundsIndex, const UTexture2D*& Texture);
+	void RemoveElement(int32 ElementIndex, int32& NextComponentLink, int32& BoundsIndex, const UStreamableRenderAsset*& Asset);
 
 	int32 AddBounds(const FBoxSphereBounds& Bounds, uint32 PackedRelativeBox, const UPrimitiveComponent* Component, float LastRenderTime, const FVector4& RangeOrigin, float MinDistance, float MinRange, float MaxRange);
 	FORCEINLINE int32 AddBounds(const UPrimitiveComponent* Component);
 	void RemoveBounds(int32 Index);
 
-	void AddTextureElements(const UPrimitiveComponent* Component, const TArrayView<FStreamingTexturePrimitiveInfo>& TextureInstanceInfos, int32 BoundsIndex, int32*& ComponentLink);
+	void AddRenderAssetElements(const UPrimitiveComponent* Component, const TArrayView<FStreamingRenderAssetPrimitiveInfo>& RenderAssetInstanceInfos, int32 BoundsIndex, int32*& ComponentLink);
 
 private:
 
@@ -93,36 +92,36 @@ private:
 
 	TMap<const UPrimitiveComponent*, int32> ComponentMap;
 
-	friend class FTextureLinkIterator;
-	friend class FTextureIterator;
+	friend class FRenderAssetLinkIterator;
+	friend class FRenderAssetIterator;
 };
 
 template <typename TTasks>
-class FTextureInstanceStateTaskSync
+class FRenderAssetInstanceStateTaskSync
 {
 public:
 
-	FTextureInstanceStateTaskSync() : State(new FTextureInstanceState()) {}
+	FRenderAssetInstanceStateTaskSync() : State(new FRenderAssetInstanceState()) {}
 
 	FORCEINLINE void Sync()
 	{
 		Tasks.SyncResults();
 	}
 
-	FORCEINLINE FTextureInstanceState* SyncAndGetState()
+	FORCEINLINE FRenderAssetInstanceState* SyncAndGetState()
 	{
 		Tasks.SyncResults();
 		return State.GetReference();
 	}
 
 	// Get State but must be constant as async tasks could be reading data.
-	FORCEINLINE const FTextureInstanceState* GetState() const
+	FORCEINLINE const FRenderAssetInstanceState* GetState() const
 	{
 		return State.GetReference();
 	}
 
 	// Used when updating the state, but with no possible reallocation.
-	FORCEINLINE FTextureInstanceState* GetStateUnsafe()
+	FORCEINLINE FRenderAssetInstanceState* GetStateUnsafe()
 	{
 		return State.GetReference();
 	}
@@ -132,6 +131,6 @@ public:
 
 private:
 
-	TRefCountPtr<FTextureInstanceState> State;
+	TRefCountPtr<FRenderAssetInstanceState> State;
 	TTasks Tasks;
 };
diff --git a/Engine/Source/Runtime/Engine/Private/Streaming/TextureInstanceTask.cpp b/Engine/Source/Runtime/Engine/Private/Streaming/TextureInstanceTask.cpp
--- a/Engine/Source/Runtime/Engine/Private/Streaming/TextureInstanceTask.cpp
+++ b/Engine/Source/Runtime/Engine/Private/Streaming/TextureInstanceTask.cpp
@@ -9,7 +9,7 @@ TextureInstanceTask.cpp: Definitions of classes used for texture streaming.
 #include "Engine/Texture2D.h"
 #include "HAL/LowLevelMemTracker.h"
 
-namespace TextureInstanceTask {
+namespace RenderAssetInstanceTask {
 
 FRefreshVisibility::FRefreshVisibility(const FOnWorkDone& InOnWorkDoneDelegate)
 	: OnWorkDoneDelegate(InOnWorkDoneDelegate)
@@ -18,7 +18,7 @@ FRefreshVisibility::FRefreshVisibility(const FOnWorkDone& InOnWorkDoneDelegate)
 {
 }
 
-void FRefreshVisibility::Init(FTextureInstanceState* InState, int32 InBeginIndex, int32 InEndIndex)
+void FRefreshVisibility::Init(FRenderAssetInstanceState* InState, int32 InBeginIndex, int32 InEndIndex)
 {
 	State = InState;
 	BeginIndex = InBeginIndex;
@@ -48,7 +48,7 @@ FRefreshFull::FRefreshFull(const FOnWorkDone& InOnWorkDoneDelegate)
 {
 }
 
-void FRefreshFull::Init(FTextureInstanceState* InState, int32 InBeginIndex, int32 InEndIndex)
+void FRefreshFull::Init(FRenderAssetInstanceState* InState, int32 InBeginIndex, int32 InEndIndex)
 {
 	check(!SkippedIndices.Num());
 
@@ -115,9 +115,9 @@ void FRefreshFull::Sync()
 void FNormalizeLightmapTexelFactor::operator()(bool bAsync)
 {
 	TArray<float> TexelFactors;
-	for (auto TextureIt = State->GetTextureIterator(); TextureIt; ++TextureIt)
+	for (auto TextureIt = State->GetRenderAssetIterator(); TextureIt; ++TextureIt)
 	{
-		const UTexture2D* Texture = *TextureIt;
+		const UStreamableRenderAsset* Texture = *TextureIt;
 		const int32 TextureLODGroup = TextureIt.GetLODGroup();
 		if (TextureLODGroup == TEXTUREGROUP_Lightmap || TextureLODGroup == TEXTUREGROUP_Shadowmap)
 		{
@@ -143,7 +143,7 @@ void FCreateViewWithUninitializedBounds::operator()(bool bAsync)
 { 
 	// The task creates another view but does not fully initialize the bounds.
 	// At sync point, the bound array are swapped as the dirty array gets refreshed.
-	View = FTextureInstanceView::CreateViewWithUninitializedBounds(State.GetReference());
+	View = FRenderAssetInstanceView::CreateViewWithUninitializedBounds(State.GetReference());
 
 	// Release the previous view. Ran here because the destructor is super expansive.
 	ViewToRelease.SafeRelease();
@@ -187,4 +187,4 @@ void FDoWorkTask::DoWork()
 	ProcessTasks(NormalizeLightmapTexelFactorTasks);
 }
 
-} // namespace TextureInstanceTask
+} // namespace RenderAssetInstanceTask
diff --git a/Engine/Source/Runtime/Engine/Private/Streaming/TextureInstanceTask.h b/Engine/Source/Runtime/Engine/Private/Streaming/TextureInstanceTask.h
--- a/Engine/Source/Runtime/Engine/Private/Streaming/TextureInstanceTask.h
+++ b/Engine/Source/Runtime/Engine/Private/Streaming/TextureInstanceTask.h
@@ -10,7 +10,7 @@ TextureInstanceTask.h: Definitions of classes used for texture streaming.
 #include "Streaming/TextureInstanceState.h"
 #include "Async/AsyncWork.h"
 
-namespace TextureInstanceTask {
+namespace RenderAssetInstanceTask {
 
 enum ETaskState
 {
@@ -79,7 +79,7 @@ public:
 	DECLARE_DELEGATE_TwoParams(FOnWorkDone, int32, int32);
 
 	FRefreshVisibility(const FOnWorkDone& InOnWorkDoneDelegate);
-	void Init(FTextureInstanceState* InState, int32 InBeginIndex, int32 InEndIndex);
+	void Init(FRenderAssetInstanceState* InState, int32 InBeginIndex, int32 InEndIndex);
 	void operator()(bool bAsync);
 	void Sync();
 
@@ -88,7 +88,7 @@ private:
 	// Callback to process the results.
 	FOnWorkDone OnWorkDoneDelegate;
 	// The state to update (no re/allocation allowed)
-	TRefCountPtr<FTextureInstanceState> State;
+	TRefCountPtr<FRenderAssetInstanceState> State;
 	// The index of the first bound to update.
 	int32 BeginIndex;
 	// The index of the last bound to update.
@@ -103,7 +103,7 @@ public:
 	DECLARE_DELEGATE_FiveParams(FOnWorkDone, int32, int32, const TArray<int32>&, int32, int32);
 
 	FRefreshFull(const FOnWorkDone& InOnWorkDoneDelegate);
-	void Init(FTextureInstanceState* InState, int32 InBeginIndex, int32 InEndIndex);
+	void Init(FRenderAssetInstanceState* InState, int32 InBeginIndex, int32 InEndIndex);
 	void operator()(bool bAsync);
 	void Sync();
 
@@ -118,7 +118,7 @@ private:
 	// Any bounds that couldn't be updated for some reason (incoherent bounds).
 	TArray<int32> SkippedIndices;
 	// The state to update (no re/allocation allowed)
-	TRefCountPtr<FTextureInstanceState> State;
+	TRefCountPtr<FRenderAssetInstanceState> State;
 	// The index of the first bound to update.
 	int32 BeginIndex;
 	// The index of the last bound to update.
@@ -130,24 +130,24 @@ class FNormalizeLightmapTexelFactor
 {
 public:
 
-	void Init(FTextureInstanceState* InState) { State = InState; }
+	void Init(FRenderAssetInstanceState* InState) { State = InState; }
 	void operator()(bool bAsync);
 	void Sync() { State.SafeRelease(); }
 
 private:
 
 	// The state to update (no re/allocation allowed)
-	TRefCountPtr<FTextureInstanceState> State;
+	TRefCountPtr<FRenderAssetInstanceState> State;
 };
 
 // Create an independent view of a state
 class FCreateViewWithUninitializedBounds : public FNonAbandonableTask
 {
 public:
-	DECLARE_DELEGATE_OneParam(FOnWorkDone, FTextureInstanceView*);
+	DECLARE_DELEGATE_OneParam(FOnWorkDone, FRenderAssetInstanceView*);
 
 	FCreateViewWithUninitializedBounds(const FOnWorkDone& InOnWorkDoneDelegate) : OnWorkDoneDelegate(InOnWorkDoneDelegate) {}
-	void Init(const FTextureInstanceState* InState, const FTextureInstanceView* InViewToRelease)
+	void Init(const FRenderAssetInstanceState* InState, const FRenderAssetInstanceView* InViewToRelease)
 	{
 		State = InState;
 		ViewToRelease = InViewToRelease;
@@ -161,11 +161,11 @@ private:
 	// Callback to process the results.
 	FOnWorkDone OnWorkDoneDelegate;
 	// The view created from the state, as a result of the execution.
-	TRefCountPtr<FTextureInstanceView> View;
+	TRefCountPtr<FRenderAssetInstanceView> View;
 	// The state for which to create the view.
-	TRefCountPtr<const FTextureInstanceState> State;
+	TRefCountPtr<const FRenderAssetInstanceState> State;
 	// The previous view of the state. Used to release the state and run the destructor async.
-	TRefCountPtr<const FTextureInstanceView> ViewToRelease;
+	TRefCountPtr<const FRenderAssetInstanceView> ViewToRelease;
 };
 
 typedef TDoWorkTask<FRefreshFull> FRefreshFullTask;
@@ -201,4 +201,4 @@ class FDoWorkAsyncTask : public FAsyncTask<FDoWorkTask>, public FRefCountedObjec
 {
 };
 
-} // namespace TextureInstanceTask
+} // namespace RenderAssetInstanceTask
diff --git a/Engine/Source/Runtime/Engine/Private/Streaming/TextureInstanceView.cpp b/Engine/Source/Runtime/Engine/Private/Streaming/TextureInstanceView.cpp
--- a/Engine/Source/Runtime/Engine/Private/Streaming/TextureInstanceView.cpp
+++ b/Engine/Source/Runtime/Engine/Private/Streaming/TextureInstanceView.cpp
@@ -11,7 +11,7 @@
 #include "Components/PrimitiveComponent.h"
 #include "ContentStreaming.h"
 
-void FTextureInstanceView::FBounds4::Set(int32 Index, const FBoxSphereBounds& Bounds, uint32 InPackedRelativeBox, float InLastRenderTime, const FVector& RangeOrigin, float InMinDistance, float InMinRange, float InMaxRange)
+void FRenderAssetInstanceView::FBounds4::Set(int32 Index, const FBoxSphereBounds& Bounds, uint32 InPackedRelativeBox, float InLastRenderTime, const FVector& RangeOrigin, float InMinDistance, float InMinRange, float InMaxRange)
 {
 	check(Index >= 0 && Index < 4);
 
@@ -32,7 +32,7 @@ void FTextureInstanceView::FBounds4::Set(int32 Index, const FBoxSphereBounds& Bo
 	LastRenderTime.Component(Index) = InLastRenderTime;
 }
 
-void FTextureInstanceView::FBounds4::UnpackBounds(int32 Index, const UPrimitiveComponent* Component)
+void FRenderAssetInstanceView::FBounds4::UnpackBounds(int32 Index, const UPrimitiveComponent* Component)
 {
 	check(Component);
 	check(Index >= 0 && Index < 4);
@@ -44,7 +44,7 @@ void FTextureInstanceView::FBounds4::UnpackBounds(int32 Index, const UPrimitiveC
 
 		// Update the visibility range once we have the bounds.
 		float MinDistance = 0, MinRange = 0, MaxRange = FLT_MAX;
-		FTextureInstanceView::GetDistanceAndRange(Component, SubBounds, MinDistance, MinRange, MaxRange);
+		FRenderAssetInstanceView::GetDistanceAndRange(Component, SubBounds, MinDistance, MinRange, MaxRange);
 
 		OriginX.Component(Index) = SubBounds.Origin.X;
 		OriginY.Component(Index) = SubBounds.Origin.Y;
@@ -64,7 +64,7 @@ void FTextureInstanceView::FBounds4::UnpackBounds(int32 Index, const UPrimitiveC
 }
 
 /** Dynamic Path, this needs to reset all members since the dynamic data is rebuilt from scratch every update (the previous data is given to the async task) */
-void FTextureInstanceView::FBounds4::FullUpdate(int32 Index, const FBoxSphereBounds& Bounds, float InLastRenderTime)
+void FRenderAssetInstanceView::FBounds4::FullUpdate(int32 Index, const FBoxSphereBounds& Bounds, float InLastRenderTime)
 {
 	check(Index >= 0 && Index < 4);
 
@@ -85,18 +85,18 @@ void FTextureInstanceView::FBounds4::FullUpdate(int32 Index, const FBoxSphereBou
 	LastRenderTime.Component(Index) = InLastRenderTime;
 }
 
-FTextureInstanceView::FTextureLinkConstIterator::FTextureLinkConstIterator(const FTextureInstanceView& InState, const UTexture2D* InTexture) 
+FRenderAssetInstanceView::FRenderAssetLinkConstIterator::FRenderAssetLinkConstIterator(const FRenderAssetInstanceView& InState, const UStreamableRenderAsset* InAsset)
 :	State(InState)
 ,	CurrElementIndex(INDEX_NONE)
 {
-	const FTextureDesc* TextureDesc = State.TextureMap.Find(InTexture);
-	if (TextureDesc)
+	const FRenderAssetDesc* AssetDesc = State.RenderAssetMap.Find(InAsset);
+	if (AssetDesc)
 	{
-		CurrElementIndex = TextureDesc->HeadLink;
+		CurrElementIndex = AssetDesc->HeadLink;
 	}
 }
 
-FBoxSphereBounds FTextureInstanceView::FTextureLinkConstIterator::GetBounds() const
+FBoxSphereBounds FRenderAssetInstanceView::FRenderAssetLinkConstIterator::GetBounds() const
 { 
 	FBoxSphereBounds Bounds(ForceInitToZero);
 
@@ -119,7 +119,7 @@ FBoxSphereBounds FTextureInstanceView::FTextureLinkConstIterator::GetBounds() co
 	return Bounds;
 }
 
-void FTextureInstanceView::FTextureLinkConstIterator::OutputToLog(float MaxNormalizedSize, float MaxNormalizedSize_VisibleOnly, const TCHAR* Prefix) const
+void FRenderAssetInstanceView::FRenderAssetLinkConstIterator::OutputToLog(float MaxNormalizedSize, float MaxNormalizedSize_VisibleOnly, const TCHAR* Prefix) const
 {
 #if !UE_BUILD_SHIPPING
 	const UPrimitiveComponent* Component = GetComponent();
@@ -197,33 +197,33 @@ void FTextureInstanceView::FTextureLinkConstIterator::OutputToLog(float MaxNorma
 #endif // !UE_BUILD_SHIPPING
 }
 
-TRefCountPtr<const FTextureInstanceView> FTextureInstanceView::CreateView(const FTextureInstanceView* RefView)
+TRefCountPtr<const FRenderAssetInstanceView> FRenderAssetInstanceView::CreateView(const FRenderAssetInstanceView* RefView)
 {
-	TRefCountPtr<FTextureInstanceView> NewView(new FTextureInstanceView());
+	TRefCountPtr<FRenderAssetInstanceView> NewView(new FRenderAssetInstanceView());
 	
 	NewView->Bounds4 = RefView->Bounds4;
 	NewView->Elements = RefView->Elements;
-	NewView->TextureMap = RefView->TextureMap;
+	NewView->RenderAssetMap = RefView->RenderAssetMap;
 	NewView->MaxTexelFactor = RefView->MaxTexelFactor;
 	// NewView->CompiledTextureMap = RefView->CompiledTextureMap;
 
-	return TRefCountPtr<const FTextureInstanceView>(NewView.GetReference());
+	return TRefCountPtr<const FRenderAssetInstanceView>(NewView.GetReference());
 }
 
-TRefCountPtr<FTextureInstanceView> FTextureInstanceView::CreateViewWithUninitializedBounds(const FTextureInstanceView* RefView)
+TRefCountPtr<FRenderAssetInstanceView> FRenderAssetInstanceView::CreateViewWithUninitializedBounds(const FRenderAssetInstanceView* RefView)
 {
-	TRefCountPtr<FTextureInstanceView> NewView(new FTextureInstanceView());
+	TRefCountPtr<FRenderAssetInstanceView> NewView(new FRenderAssetInstanceView());
 	
 	NewView->Bounds4.AddUninitialized(RefView->Bounds4.Num());
 	NewView->Elements = RefView->Elements;
-	NewView->TextureMap = RefView->TextureMap;
+	NewView->RenderAssetMap = RefView->RenderAssetMap;
 	NewView->MaxTexelFactor = RefView->MaxTexelFactor;
 	// NewView->CompiledTextureMap = RefView->RefView;
 
 	return NewView;
 }
 
-void FTextureInstanceView::GetDistanceAndRange(const UPrimitiveComponent* Component, const FBoxSphereBounds& TextureInstanceBounds, float& MinDistance, float& MinRange, float& MaxRange)
+void FRenderAssetInstanceView::GetDistanceAndRange(const UPrimitiveComponent* Component, const FBoxSphereBounds& RenderAssetInstanceBounds, float& MinDistance, float& MinRange, float& MaxRange)
 {
 	check(Component && Component->IsRegistered());
 
@@ -235,71 +235,40 @@ void FTextureInstanceView::GetDistanceAndRange(const UPrimitiveComponent* Compon
 
 	const UPrimitiveComponent* LODParent = Component->GetLODParentPrimitive();
 
-	MinDistance = FMath::Max<float>(0, Component->MinDrawDistance - (TextureInstanceBounds.Origin - Component->Bounds.Origin).Size() - TextureInstanceBounds.SphereRadius);
+	MinDistance = FMath::Max<float>(0, Component->MinDrawDistance - (RenderAssetInstanceBounds.Origin - Component->Bounds.Origin).Size() - RenderAssetInstanceBounds.SphereRadius);
 	MinRange = FMath::Max<float>(0, Component->MinDrawDistance);
 	// Max distance when HLOD becomes visible.
 	MaxRange = LODParent ? (LODParent->MinDrawDistance + (Component->Bounds.Origin - LODParent->Bounds.Origin).Size()) : FLT_MAX;
-	if (LODParent)
-	{
-		MaxRange = LODParent->MinDrawDistance + (Component->Bounds.Origin - LODParent->Bounds.Origin).Size();
-	}
 }
 
-void FTextureInstanceView::SwapData(FTextureInstanceView* Lfs, FTextureInstanceView* Rhs)
+void FRenderAssetInstanceView::SwapData(FRenderAssetInstanceView* Lfs, FRenderAssetInstanceView* Rhs)
 {
 	// Things must be compatible somehow or derived classes will be in incoherent state.
 	check(Lfs->Bounds4.Num() == Rhs->Bounds4.Num());
 	check(Lfs->Elements.Num() == Rhs->Elements.Num());
-	check(Lfs->TextureMap.Num() == Rhs->TextureMap.Num());
-	check(Lfs->CompiledTextureMap.Num() == 0 && Rhs->CompiledTextureMap.Num() == 0);
+	check(Lfs->RenderAssetMap.Num() == Rhs->RenderAssetMap.Num());
+	check(Lfs->CompiledRenderAssetMap.Num() == 0 && Rhs->CompiledRenderAssetMap.Num() == 0);
 
 	FMemory::Memswap(&Lfs->Bounds4 , &Rhs->Bounds4, sizeof(Lfs->Bounds4));
 	FMemory::Memswap(&Lfs->Elements , &Rhs->Elements, sizeof(Lfs->Elements));
-	FMemory::Memswap(&Lfs->TextureMap , &Rhs->TextureMap, sizeof(Lfs->TextureMap));
+	FMemory::Memswap(&Lfs->RenderAssetMap, &Rhs->RenderAssetMap, sizeof(Lfs->RenderAssetMap));
 	FMemory::Memswap(&Lfs->MaxTexelFactor , &Rhs->MaxTexelFactor, sizeof(Lfs->MaxTexelFactor));
 }
 
-struct FStreamingViewInfoExtra
+void FRenderAssetInstanceAsyncView::UpdateBoundSizes_Async(
+	const TArray<FStreamingViewInfo>& ViewInfos,
+	const FStreamingViewInfoExtraArray& ViewInfoExtras,
+	float LastUpdateTime,
+	const FRenderAssetStreamingSettings& Settings)
 {
-	// The screen size factor including the view boost.
-	float ScreenSizeFloat;
-
-	// The extra view boost for visible primitive (if ViewInfo.BoostFactor > "r.Streaming.MaxHiddenPrimitiveViewBoost").
-	float ExtraBoostForVisiblePrimitiveFloat;
-};
+	check(ViewInfos.Num() == ViewInfoExtras.Num());
 
-void FTextureInstanceAsyncView::UpdateBoundSizes_Async(const TArray<FStreamingViewInfo>& ViewInfos, float LastUpdateTime, const FTextureStreamingSettings& Settings)
-{
 	if (!View.IsValid())  return;
 
 	const int32 NumViews = ViewInfos.Num();
 	const int32 NumBounds4 = View->NumBounds4();
 
 	const VectorRegister LastUpdateTime4 = VectorSet(LastUpdateTime, LastUpdateTime, LastUpdateTime, LastUpdateTime);
-	const float OneOverMaxHiddenPrimitiveViewBoost = 1.f / Settings.MaxHiddenPrimitiveViewBoost;
-	
-	TArray<FStreamingViewInfoExtra, TInlineAllocator<4> > ViewInfoExtras;
-	ViewInfoExtras.AddZeroed(ViewInfos.Num());
-	
-	for (int32 ViewIndex = 0; ViewIndex < NumViews; ++ViewIndex)
-	{
-		const FStreamingViewInfo& ViewInfo = ViewInfos[ViewIndex];
-		FStreamingViewInfoExtra& ViewInfoExtra = ViewInfoExtras[ViewIndex];
-
-		const float EffectiveScreenSize = (Settings.MaxEffectiveScreenSize > 0.0f) ? FMath::Min(Settings.MaxEffectiveScreenSize, ViewInfo.ScreenSize) : ViewInfo.ScreenSize;
-		ViewInfoExtra.ScreenSizeFloat = EffectiveScreenSize * .5f; // Multiply by half since the ratio factors map to half the screen only
-		ViewInfoExtra.ExtraBoostForVisiblePrimitiveFloat = 1.f;
-
-		if (ViewInfo.BoostFactor > Settings.MaxHiddenPrimitiveViewBoost)
-		{
-			ViewInfoExtra.ScreenSizeFloat *= Settings.MaxHiddenPrimitiveViewBoost;
-			ViewInfoExtra.ExtraBoostForVisiblePrimitiveFloat = ViewInfo.BoostFactor * OneOverMaxHiddenPrimitiveViewBoost;
-		}
-		else
-		{
-			ViewInfoExtra.ScreenSizeFloat *= ViewInfo.BoostFactor;
-		}
-	}
 
 	BoundsViewInfo.Empty(NumBounds4 * 4);
 	BoundsViewInfo.AddUninitialized(NumBounds4 * 4);
@@ -309,7 +278,7 @@ void FTextureInstanceAsyncView::UpdateBoundSizes_Async(const TArray<FStreamingVi
 
 	for (int32 Bounds4Index = 0; Bounds4Index < NumBounds4; ++Bounds4Index)
 	{
-		const FTextureInstanceView::FBounds4& CurrentBounds4 = View->GetBounds4(Bounds4Index);
+		const FRenderAssetInstanceView::FBounds4& CurrentBounds4 = View->GetBounds4(Bounds4Index);
 
 		// Calculate distance of viewer to bounding sphere.
 		const VectorRegister OriginX = VectorLoadAligned( &CurrentBounds4.OriginX );
@@ -423,18 +392,25 @@ void FTextureInstanceAsyncView::UpdateBoundSizes_Async(const TArray<FStreamingVi
 		}
 	}
 
-	if (Settings.MinLevelTextureScreenSize > 0)
+	if (Settings.MinLevelRenderAssetScreenSize > 0)
 	{
 		float ViewMaxNormalizedSizeResult = VectorGetComponent(ViewMaxNormalizedSize, 0);
 		for (int32 SubIndex = 1; SubIndex < 4; ++SubIndex)
 		{
 			ViewMaxNormalizedSizeResult = FMath::Max(ViewMaxNormalizedSizeResult, VectorGetComponent(ViewMaxNormalizedSize, SubIndex));
 		}
-		MaxLevelTextureScreenSize = View->GetMaxTexelFactor() * ViewMaxNormalizedSizeResult;
+		MaxLevelRenderAssetScreenSize = View->GetMaxTexelFactor() * ViewMaxNormalizedSizeResult;
 	}
 }
 
-void FTextureInstanceAsyncView::ProcessElement(const FBoundsViewInfo& BoundsVieWInfo, float TexelFactor, bool bForceLoad, float& MaxSize, float& MaxSize_VisibleOnly) const
+void FRenderAssetInstanceAsyncView::ProcessElement(
+	typename FStreamingRenderAsset::EAssetType AssetType,
+	const FBoundsViewInfo& BoundsVieWInfo,
+	float TexelFactor,
+	bool bForceLoad,
+	float& MaxSize,
+	float& MaxSize_VisibleOnly,
+	int32& MaxNumForcedLODs) const
 {
 	if (TexelFactor == FLT_MAX) // If this is a forced load component.
 	{
@@ -454,8 +430,17 @@ void FTextureInstanceAsyncView::ProcessElement(const FBoundsViewInfo& BoundsVieW
 	}
 	else // Negative texel factors map to fixed resolution. Currently used for lanscape.
 	{
-		MaxSize = FMath::Max(MaxSize, -TexelFactor);
-		MaxSize_VisibleOnly = FMath::Max(MaxSize_VisibleOnly, -TexelFactor);
+		if (AssetType == FStreamingRenderAsset::AT_Texture)
+		{
+			MaxSize = FMath::Max(MaxSize, -TexelFactor);
+			MaxSize_VisibleOnly = FMath::Max(MaxSize_VisibleOnly, -TexelFactor);
+		}
+		else
+		{
+			check(AssetType == FStreamingRenderAsset::AT_StaticMesh || AssetType == FStreamingRenderAsset::AT_SkeletalMesh);
+			check(-TexelFactor <= MAX_MESH_LOD_COUNT);
+			MaxNumForcedLODs = FMath::Max(MaxNumForcedLODs, static_cast<int32>(-TexelFactor));
+		}
 
 		// Force load will load the immediatly visible part, and later the full texture.
 		if (bForceLoad && (BoundsVieWInfo.MaxNormalizedSize > 0 || BoundsVieWInfo.MaxNormalizedSize_VisibleOnly > 0))
@@ -466,9 +451,18 @@ void FTextureInstanceAsyncView::ProcessElement(const FBoundsViewInfo& BoundsVieW
 	}
 }
 
-void FTextureInstanceAsyncView::GetTexelSize(const UTexture2D* InTexture, float& MaxSize, float& MaxSize_VisibleOnly, const TCHAR* LogPrefix) const
+void FRenderAssetInstanceAsyncView::GetRenderAssetScreenSize(
+	typename FStreamingRenderAsset::EAssetType AssetType,
+	const UStreamableRenderAsset* InAsset,
+	float& MaxSize,
+	float& MaxSize_VisibleOnly,
+	int32& MaxNumForcedLODs,
+	const TCHAR* LogPrefix) const
 {
 	// No need to iterate more if texture is already at maximum resolution.
+	// Meshes don't really fit into the concept of max resolution but current
+	// max_res for texture is 8k which is large enough to let mesh screen
+	// sizes be constrained by this value
 
 	int32 CurrCount = 0;
 
@@ -477,19 +471,26 @@ void FTextureInstanceAsyncView::GetTexelSize(const UTexture2D* InTexture, float&
 		// Use the fast path if available, about twice as fast when there are a lot of elements.
 		if (View->HasCompiledElements() && !LogPrefix)
 		{
-			const TArray<FTextureInstanceView::FCompiledElement>* CompiledElements = View->GetCompiledElements(InTexture);
+			const TArray<FRenderAssetInstanceView::FCompiledElement>* CompiledElements = View->GetCompiledElements(InAsset);
 			if (CompiledElements)
 			{
 				const int32 NumCompiledElements = CompiledElements->Num();
-				const FTextureInstanceView::FCompiledElement* CompiledElementData = CompiledElements->GetData();
+				const FRenderAssetInstanceView::FCompiledElement* CompiledElementData = CompiledElements->GetData();
 
 				int32 CompiledElementIndex = 0;
 				while (CompiledElementIndex < NumCompiledElements && MaxSize_VisibleOnly < MAX_TEXTURE_SIZE)
 				{
-					const FTextureInstanceView::FCompiledElement& CompiledElement = CompiledElementData[CompiledElementIndex];
+					const FRenderAssetInstanceView::FCompiledElement& CompiledElement = CompiledElementData[CompiledElementIndex];
 					if (ensure(BoundsViewInfo.IsValidIndex(CompiledElement.BoundsIndex)))
 					{
-						ProcessElement(BoundsViewInfo[CompiledElement.BoundsIndex], CompiledElement.TexelFactor, CompiledElement.bForceLoad, MaxSize, MaxSize_VisibleOnly);
+						ProcessElement(
+							AssetType,
+							BoundsViewInfo[CompiledElement.BoundsIndex],
+							CompiledElement.TexelFactor,
+							CompiledElement.bForceLoad,
+							MaxSize,
+							MaxSize_VisibleOnly,
+							MaxNumForcedLODs);
 					}
 					++CompiledElementIndex;
 				}
@@ -497,20 +498,20 @@ void FTextureInstanceAsyncView::GetTexelSize(const UTexture2D* InTexture, float&
 				if (MaxSize_VisibleOnly >= MAX_TEXTURE_SIZE && CompiledElementIndex > 1)
 				{
 					// This does not realloc anything but moves the closest element at head, making the next update find it immediately and early exit.
-					FTextureInstanceView::FCompiledElement* SwapElementData = const_cast<FTextureInstanceView::FCompiledElement*>(CompiledElementData);
-					Swap<FTextureInstanceView::FCompiledElement>(SwapElementData[0], SwapElementData[CompiledElementIndex - 1]);
+					FRenderAssetInstanceView::FCompiledElement* SwapElementData = const_cast<FRenderAssetInstanceView::FCompiledElement*>(CompiledElementData);
+					Swap<FRenderAssetInstanceView::FCompiledElement>(SwapElementData[0], SwapElementData[CompiledElementIndex - 1]);
 				}
 			}
 		}
 		else
 		{
-			for (auto It = View->GetElementIterator(InTexture); It && (MaxSize_VisibleOnly < MAX_TEXTURE_SIZE || LogPrefix); ++It)
+			for (auto It = View->GetElementIterator(InAsset); It && (MaxSize_VisibleOnly < MAX_TEXTURE_SIZE || LogPrefix); ++It)
 			{
 				// Only handle elements that are in bounds.
 				if (ensure(BoundsViewInfo.IsValidIndex(It.GetBoundsIndex())))
 				{
 					const FBoundsViewInfo& BoundsVieWInfo = BoundsViewInfo[It.GetBoundsIndex()];
-					ProcessElement(BoundsVieWInfo, It.GetTexelFactor(), It.GetForceLoad(), MaxSize, MaxSize_VisibleOnly);
+					ProcessElement(AssetType, BoundsVieWInfo, It.GetTexelFactor(), It.GetForceLoad(), MaxSize, MaxSize_VisibleOnly, MaxNumForcedLODs);
 					if (LogPrefix)
 					{
 						It.OutputToLog(BoundsVieWInfo.MaxNormalizedSize, BoundsVieWInfo.MaxNormalizedSize_VisibleOnly, LogPrefix);
@@ -521,7 +522,7 @@ void FTextureInstanceAsyncView::GetTexelSize(const UTexture2D* InTexture, float&
 	}
 }
 
-bool FTextureInstanceAsyncView::HasTextureReferences(const UTexture2D* InTexture) const
+bool FRenderAssetInstanceAsyncView::HasRenderAssetReferences(const UStreamableRenderAsset* InAsset) const
 {
-	return View.IsValid() && (bool)View->GetElementIterator(InTexture);
+	return View.IsValid() && (bool)View->GetElementIterator(InAsset);
 }
diff --git a/Engine/Source/Runtime/Engine/Private/Streaming/TextureInstanceView.h b/Engine/Source/Runtime/Engine/Private/Streaming/TextureInstanceView.h
--- a/Engine/Source/Runtime/Engine/Private/Streaming/TextureInstanceView.h
+++ b/Engine/Source/Runtime/Engine/Private/Streaming/TextureInstanceView.h
@@ -8,20 +8,21 @@ TextureInstanceView.h: Definitions of classes used for texture streaming.
 
 #include "CoreMinimal.h"
 #include "Templates/RefCounting.h"
+#include "Streaming/StreamingTexture.h"
 
 class UPrimitiveComponent;
-class UTexture2D;
+class UStreamableRenderAsset;
 struct FStreamingViewInfo;
-struct FTextureStreamingSettings;
+struct FRenderAssetStreamingSettings;
 
 #define MAX_TEXTURE_SIZE (float(1 << (MAX_TEXTURE_MIP_COUNT - 1)))
 
 // Main Thread Job Requirement : find all instance of a component and update it's bound.
-// Threaded Job Requirement : get the list of instance texture easily from the list of 
+// Threaded Job Requirement : get the list of instance texture or mesh easily from the list of 
 
-// A constant view on the relationship between textures, components and bounds. 
-// Has everything needed for the worker task to compute the required size per texture.
-class FTextureInstanceView : public FRefCountedObject
+// A constant view on the relationship between textures/meshes, components and bounds. 
+// Has everything needed for the worker task to compute the required size per texture/mesh.
+class FRenderAssetInstanceView : public FRefCountedObject
 {
 public:
 
@@ -30,11 +31,11 @@ public:
 	{
 		FORCEINLINE FBounds4();
 
-		/** X coordinates for the bounds origin of 4 texture instances */
+		/** X coordinates for the bounds origin of 4 texture/mesh instances */
 		FVector4 OriginX;
-		/** Y coordinates for the bounds origin of 4 texture instances */
+		/** Y coordinates for the bounds origin of 4 texture/mesh instances */
 		FVector4 OriginY;
-		/** Z coordinates for the bounds origin of 4 texture instances */
+		/** Z coordinates for the bounds origin of 4 texture/mesh instances */
 		FVector4 OriginZ;
 
 		/** X coordinates for used to compute the distance condition between min and max */
@@ -44,14 +45,14 @@ public:
 		/** Z coordinates for used to compute the distance condition between min and max */
 		FVector4 RangeOriginZ;
 
-		/** X size of the bounds box extent of 4 texture instances */
+		/** X size of the bounds box extent of 4 texture/mesh instances */
 		FVector4 ExtentX;
-		/** Y size of the bounds box extent of 4 texture instances */
+		/** Y size of the bounds box extent of 4 texture/mesh instances */
 		FVector4 ExtentY;
-		/** Z size of the bounds box extent of 4 texture instances */
+		/** Z size of the bounds box extent of 4 texture/mesh instances */
 		FVector4 ExtentZ;
 
-		/** Sphere radii for the bounding sphere of 4 texture instances */
+		/** Sphere radii for the bounding sphere of 4 texture/mesh instances */
 		FVector4 Radius;
 
 		/** The relative box the bound was computed with. Aligned to be interpreted as FVector4  */
@@ -84,14 +85,14 @@ public:
 		FORCEINLINE FElement();
 
 		const UPrimitiveComponent* Component; // Which component this relates too
-		const UTexture2D* Texture;	// Texture, never dereferenced.
+		const UStreamableRenderAsset* RenderAsset;	// Texture or mesh, never dereferenced.
 
 		int32 BoundsIndex;		// The Index associated to this component (static component can have several bounds).
 		float TexelFactor;		// The texture scale to be applied to this instance.
-		bool bForceLoad;		// The texture needs to be force loaded.
+		bool bForceLoad;		// The texture or mesh needs to be force loaded.
 
-		int32 PrevTextureLink;	// The previous element which uses the same texture as this Element. The first element referred by TextureMap will have INDEX_NONE.
-		int32 NextTextureLink;	// The next element which uses the same texture as this Element. Last element will have INDEX_NONE
+		int32 PrevRenderAssetLink;	// The previous element which uses the same asset as this Element. The first element referred by RenderAssetMap will have INDEX_NONE.
+		int32 NextRenderAssetLink;	// The next element which uses the same asset as this Element. Last element will have INDEX_NONE
 
 		// Components are always updated as a whole, so individual elements can not be removed. Removing the need for PrevComponentLink.
 		int32 NextComponentLink;	// The next element that uses the same component as this Element. The first one is referred by ComponentMap and the last one will have INDEX_NONE.
@@ -113,25 +114,25 @@ public:
 		FORCEINLINE bool operator==(const FCompiledElement& Rhs) const { return BoundsIndex == Rhs.BoundsIndex && TexelFactor == Rhs.TexelFactor && bForceLoad == Rhs.bForceLoad; }
 	};
 
-	struct FTextureDesc
+	struct FRenderAssetDesc
 	{
-		FTextureDesc(int32 InHeadLink, int32 InLODGroup) : HeadLink(InHeadLink), LODGroup(InLODGroup) {}
+		FRenderAssetDesc(int32 InHeadLink, int32 InLODGroup) : HeadLink(InHeadLink), LODGroup(InLODGroup) {}
 
-		// The index of head element using the texture.
+		// The index of head element using the renderable asset.
 		int32 HeadLink;
-		// The LODGroup of the texture, used to performe some tasks async.
+		// The LODGroup of the texture or mesh, used to performe some tasks async.
 		const int32 LODGroup;
 	};
 
 
-	// Iterator processing all elements refering to a texture.
-	class FTextureLinkConstIterator
+	// Iterator processing all elements refering to a texture/mesh.
+	class FRenderAssetLinkConstIterator
 	{
 	public:
-		FTextureLinkConstIterator(const FTextureInstanceView& InState, const UTexture2D* InTexture);
+		FRenderAssetLinkConstIterator(const FRenderAssetInstanceView& InState, const UStreamableRenderAsset* InAsset);
 
 		FORCEINLINE operator bool() const { return CurrElementIndex != INDEX_NONE; }
-		FORCEINLINE void operator++() { CurrElementIndex = State.Elements[CurrElementIndex].NextTextureLink; }
+		FORCEINLINE void operator++() { CurrElementIndex = State.Elements[CurrElementIndex].NextRenderAssetLink; }
 
 		void OutputToLog(float MaxNormalizedSize, float MaxNormalizedSize_VisibleOnly, const TCHAR* Prefix) const;
 
@@ -145,59 +146,59 @@ public:
 
 		FORCEINLINE const UPrimitiveComponent* GetComponent() const { return State.Elements[CurrElementIndex].Component; }
 
-		const FTextureInstanceView& State;
+		const FRenderAssetInstanceView& State;
 		int32 CurrElementIndex;
 	};
 
-	class FTextureLinkIterator : public FTextureLinkConstIterator
+	class FRenderAssetLinkIterator : public FRenderAssetLinkConstIterator
 	{
 	public:
-		FTextureLinkIterator(FTextureInstanceView& InState, const UTexture2D* InTexture) : FTextureLinkConstIterator(InState, InTexture) {}
+		FRenderAssetLinkIterator(FRenderAssetInstanceView& InState, const UStreamableRenderAsset* InAsset) : FRenderAssetLinkConstIterator(InState, InAsset) {}
 
 		FORCEINLINE void ClampTexelFactor(float CMin, float CMax)
 		{ 
-			float& TexelFactor = const_cast<FTextureInstanceView&>(State).Elements[CurrElementIndex].TexelFactor;
+			float& TexelFactor = const_cast<FRenderAssetInstanceView&>(State).Elements[CurrElementIndex].TexelFactor;
 			TexelFactor = FMath::Clamp<float>(TexelFactor, CMin, CMax);
 		}
 	};
 
-	class FTextureIterator
+	class FRenderAssetIterator
 	{
 	public:
-		FTextureIterator(const FTextureInstanceView& InState) : MapIt(InState.TextureMap) {}
+		FRenderAssetIterator(const FRenderAssetInstanceView& InState) : MapIt(InState.RenderAssetMap) {}
 
 		operator bool() const { return (bool)MapIt; }
 		void operator++() { ++MapIt; }
 
-		const UTexture2D* operator*() const { return MapIt.Key(); }
+		const UStreamableRenderAsset* operator*() const { return MapIt.Key(); }
 		int32 GetLODGroup() const { return MapIt.Value().LODGroup; }
 
 	private:
 
-		TMap<const UTexture2D*, FTextureDesc>::TConstIterator MapIt;
+		TMap<const UStreamableRenderAsset*, FRenderAssetDesc>::TConstIterator MapIt;
 	};
 
-	FTextureInstanceView() : MaxTexelFactor(FLT_MAX) {}
+	FRenderAssetInstanceView() : MaxTexelFactor(FLT_MAX) {}
 
 	FORCEINLINE int32 NumBounds4() const { return Bounds4.Num(); }
 	FORCEINLINE const FBounds4& GetBounds4(int32 Bounds4Index ) const {  return Bounds4[Bounds4Index]; }
 
-	FORCEINLINE FTextureLinkIterator GetElementIterator(const UTexture2D* InTexture ) {  return FTextureLinkIterator(*this, InTexture); }
-	FORCEINLINE FTextureLinkConstIterator GetElementIterator(const UTexture2D* InTexture ) const {  return FTextureLinkConstIterator(*this, InTexture); }
-	FORCEINLINE FTextureIterator GetTextureIterator( ) const {  return FTextureIterator(*this); }
+	FORCEINLINE FRenderAssetLinkIterator GetElementIterator(const UStreamableRenderAsset* InTexture ) {  return FRenderAssetLinkIterator(*this, InTexture); }
+	FORCEINLINE FRenderAssetLinkConstIterator GetElementIterator(const UStreamableRenderAsset* InTexture ) const {  return FRenderAssetLinkConstIterator(*this, InTexture); }
+	FORCEINLINE FRenderAssetIterator GetRenderAssetIterator( ) const {  return FRenderAssetIterator(*this); }
 
 	// Whether or not this state has compiled elements.
-	bool HasCompiledElements() const { return CompiledTextureMap.Num() != 0; }
-	// If this has compiled elements, return the array relate to a given texture.
-	const TArray<FCompiledElement>* GetCompiledElements(const UTexture2D* Texture) const { return CompiledTextureMap.Find(Texture); }
+	bool HasCompiledElements() const { return CompiledRenderAssetMap.Num() != 0; }
+	// If this has compiled elements, return the array relate to a given texture or mesh.
+	const TArray<FCompiledElement>* GetCompiledElements(const UStreamableRenderAsset* Asset) const { return CompiledRenderAssetMap.Find(Asset); }
 
-	static TRefCountPtr<const FTextureInstanceView> CreateView(const FTextureInstanceView* RefView);
-	static TRefCountPtr<FTextureInstanceView> CreateViewWithUninitializedBounds(const FTextureInstanceView* RefView);
-	static void SwapData(FTextureInstanceView* Lfs, FTextureInstanceView* Rhs);
+	static TRefCountPtr<const FRenderAssetInstanceView> CreateView(const FRenderAssetInstanceView* RefView);
+	static TRefCountPtr<FRenderAssetInstanceView> CreateViewWithUninitializedBounds(const FRenderAssetInstanceView* RefView);
+	static void SwapData(FRenderAssetInstanceView* Lfs, FRenderAssetInstanceView* Rhs);
 
 	float GetMaxTexelFactor() const { return MaxTexelFactor; }
 
-	static void GetDistanceAndRange(const UPrimitiveComponent* Component, const FBoxSphereBounds& TextureInstanceBounds, float& MinDistance, float& MinRange, float& MaxRange);
+	static void GetDistanceAndRange(const UPrimitiveComponent* Component, const FBoxSphereBounds& RenderAssetInstanceBounds, float& MinDistance, float& MinRange, float& MaxRange);
 
 protected:
 
@@ -205,40 +206,63 @@ protected:
 
 	TArray<FElement> Elements;
 
-	TMap<const UTexture2D*, FTextureDesc> TextureMap;
+	TMap<const UStreamableRenderAsset*, FRenderAssetDesc> RenderAssetMap;
 
 	// CompiledTextureMap is used to iterate more quickly on each elements by avoiding the linked list indirections.
-	TMap<const UTexture2D*, TArray<FCompiledElement> > CompiledTextureMap;
+	TMap<const UStreamableRenderAsset*, TArray<FCompiledElement> > CompiledRenderAssetMap;
 
 	/** Max texel factor across all elements. Used for early culling */
 	float MaxTexelFactor;
 };
 
+struct FStreamingViewInfoExtra
+{
+	// The screen size factor including the view boost.
+	float ScreenSizeFloat;
+
+	// The extra view boost for visible primitive (if ViewInfo.BoostFactor > "r.Streaming.MaxHiddenPrimitiveViewBoost").
+	float ExtraBoostForVisiblePrimitiveFloat;
+};
+
+typedef TArray<FStreamingViewInfoExtra, TInlineAllocator<4>> FStreamingViewInfoExtraArray;
+
 // Data used to compute visibility
-class FTextureInstanceAsyncView
+class FRenderAssetInstanceAsyncView
 {
 public:
 
-	FTextureInstanceAsyncView() : MaxLevelTextureScreenSize(MAX_FLT) {}
+	FRenderAssetInstanceAsyncView() : MaxLevelRenderAssetScreenSize(MAX_FLT) {}
 
-	FTextureInstanceAsyncView(const FTextureInstanceView* InView) : View(InView), MaxLevelTextureScreenSize(MAX_FLT) {}
+	FRenderAssetInstanceAsyncView(const FRenderAssetInstanceView* InView) : View(InView), MaxLevelRenderAssetScreenSize(MAX_FLT) {}
 
-	void UpdateBoundSizes_Async(const TArray<FStreamingViewInfo>& ViewInfos, float LastUpdateTime, const FTextureStreamingSettings& Settings);
+	void UpdateBoundSizes_Async(
+		const TArray<FStreamingViewInfo>& ViewInfos,
+		const FStreamingViewInfoExtraArray& ViewInfoExtras,
+		float LastUpdateTime,
+		const FRenderAssetStreamingSettings& Settings);
 
-	// MaxSize : Biggest texture size for all instances.
-	// MaxSize_VisibleOnly : Biggest texture size for visble instances only.
-	void GetTexelSize(const UTexture2D* InTexture, float& MaxSize, float& MaxSize_VisibleOnly, const TCHAR* LogPrefix) const;
+	// Screen size is the number of screen pixels overlap with this asset
+	// MaxSize : Biggest screen size for all instances.
+	// MaxSize_VisibleOnly : Biggest screen size for visble instances only.
+	// MaxForcedNumLODs: max number of LODs forced resident. Used for meshes only
+	void GetRenderAssetScreenSize(
+		typename FStreamingRenderAsset::EAssetType AssetType,
+		const UStreamableRenderAsset* InAsset,
+		float& MaxSize,
+		float& MaxSize_VisibleOnly,
+		int32& MaxNumForcedLODs,
+		const TCHAR* LogPrefix) const;
 
-	bool HasTextureReferences(const UTexture2D* InTexture) const;
+	bool HasRenderAssetReferences(const UStreamableRenderAsset* InAsset) const;
 
 	// Release the data now as this is expensive.
 	void OnTaskDone() { BoundsViewInfo.Empty(); }
 
-	float GetMaxLevelTextureScreenSize() const { return MaxLevelTextureScreenSize; }
+	float GetMaxLevelRenderAssetScreenSize() const { return MaxLevelRenderAssetScreenSize; }
 
 private:
 
-	TRefCountPtr<const FTextureInstanceView> View;
+	TRefCountPtr<const FRenderAssetInstanceView> View;
 
 	struct FBoundsViewInfo
 	{
@@ -256,7 +280,14 @@ private:
 	TArray<FBoundsViewInfo> BoundsViewInfo;
 
 	/** The max possible size (conservative) across all elements of this view. */
-	float MaxLevelTextureScreenSize;
-
-	void ProcessElement(const FBoundsViewInfo& BoundsVieWInfo, float TexelFactor, bool bForceLoad, float& MaxSize, float& MaxSize_VisibleOnly) const;
+	float MaxLevelRenderAssetScreenSize;
+
+	void ProcessElement(
+		typename FStreamingRenderAsset::EAssetType AssetType,
+		const FBoundsViewInfo& BoundsVieWInfo,
+		float TexelFactor,
+		bool bForceLoad,
+		float& MaxSize,
+		float& MaxSize_VisibleOnly,
+		int32& MaxNumForcedLODs) const;
 };
diff --git a/Engine/Source/Runtime/Engine/Private/Streaming/TextureInstanceView.inl b/Engine/Source/Runtime/Engine/Private/Streaming/TextureInstanceView.inl
--- a/Engine/Source/Runtime/Engine/Private/Streaming/TextureInstanceView.inl
+++ b/Engine/Source/Runtime/Engine/Private/Streaming/TextureInstanceView.inl
@@ -6,7 +6,7 @@ TextureInstanceView.h: Definitions of classes used for texture streaming.
 
 #pragma once
 
-FTextureInstanceView::FBounds4::FBounds4()
+FRenderAssetInstanceView::FBounds4::FBounds4()
 :	OriginX( 0, 0, 0, 0 )
 ,	OriginY( 0, 0, 0, 0 )
 ,	OriginZ( 0, 0, 0, 0 )
@@ -25,7 +25,7 @@ FTextureInstanceView::FBounds4::FBounds4()
 {
 }
 
-void FTextureInstanceView::FBounds4::Clear(int32 Index)
+void FRenderAssetInstanceView::FBounds4::Clear(int32 Index)
 {
 	check(Index >= 0 && Index < 4);
 
@@ -46,7 +46,7 @@ void FTextureInstanceView::FBounds4::Clear(int32 Index)
 	LastRenderTime.Component(Index) = -FLT_MAX;
 }
 
-void FTextureInstanceView::FBounds4::OffsetBounds(int32 Index, const FVector& Offset)
+void FRenderAssetInstanceView::FBounds4::OffsetBounds(int32 Index, const FVector& Offset)
 {
 	check(Index >= 0 && Index < 4);
 
@@ -55,20 +55,20 @@ void FTextureInstanceView::FBounds4::OffsetBounds(int32 Index, const FVector& Of
 	OriginZ.Component(Index) += Offset.Z;
 }
 
-void FTextureInstanceView::FBounds4::UpdateLastRenderTime(int32 Index, float InLastRenderTime)
+void FRenderAssetInstanceView::FBounds4::UpdateLastRenderTime(int32 Index, float InLastRenderTime)
 {
 	check(Index >= 0 && Index < 4);
 
 	LastRenderTime.Component(Index) = InLastRenderTime;
 }
 
-FTextureInstanceView::FElement::FElement()
+FRenderAssetInstanceView::FElement::FElement()
 :	Component(nullptr)
-,	Texture(nullptr)
+,	RenderAsset(nullptr)
 ,	BoundsIndex(INDEX_NONE)
 ,	TexelFactor(0)
-,	PrevTextureLink(INDEX_NONE)
-,	NextTextureLink(INDEX_NONE)
+,	PrevRenderAssetLink(INDEX_NONE)
+,	NextRenderAssetLink(INDEX_NONE)
 ,	NextComponentLink(INDEX_NONE)
 {
 }
diff --git a/Engine/Source/Runtime/Engine/Private/Streaming/TextureStreamingBuild.cpp b/Engine/Source/Runtime/Engine/Private/Streaming/TextureStreamingBuild.cpp
--- a/Engine/Source/Runtime/Engine/Private/Streaming/TextureStreamingBuild.cpp
+++ b/Engine/Source/Runtime/Engine/Private/Streaming/TextureStreamingBuild.cpp
@@ -160,13 +160,13 @@ ENGINE_API bool BuildTextureStreamingComponentData(UWorld* InWorld, EMaterialQua
 #undef LOCTEXT_NAMESPACE
 
 /**
- * Checks whether a UTexture2D is supposed to be streaming.
- * @param Texture	Texture to check
- * @return			true if the UTexture2D is supposed to be streaming
- */
-bool IsStreamingTexture( const UTexture2D* Texture2D )
+* Checks whether a UStreamableRenderAsset is a texture/mesh with streamable mips
+* @param Asset		Asset to check
+* @return			true if the UStreamableRenderAsset is supposed to be streaming
+*/
+bool IsStreamingRenderAsset( const UStreamableRenderAsset* Asset )
 {
-	return Texture2D && Texture2D->bIsStreamable && !Texture2D->NeverStream && Texture2D->GetNumMips() > Texture2D->GetNumNonStreamingMips();
+	return Asset && Asset->bIsStreamable && !Asset->NeverStream && Asset->GetNumMipsForStreaming() > Asset->GetNumNonStreamingMips();
 }
 
 uint32 PackRelativeBox(const FVector& RefOrigin, const FVector& RefExtent, const FVector& Origin, const FVector& Extent)
@@ -246,13 +246,14 @@ void UnpackRelativeBox(const FBoxSphereBounds& InRefBounds, uint32 InPackedRelBo
 	}
 }
 
-void FStreamingTextureBuildInfo::PackFrom(ULevel* Level, const FBoxSphereBounds& RefBounds, const FStreamingTexturePrimitiveInfo& Info)
+void FStreamingTextureBuildInfo::PackFrom(ULevel* Level, const FBoxSphereBounds& RefBounds, const FStreamingRenderAssetPrimitiveInfo& Info)
 {
 	check(Level);
 
 	PackedRelativeBox = PackRelativeBox(RefBounds.Origin, RefBounds.BoxExtent, Info.Bounds.Origin, Info.Bounds.BoxExtent);
 
-	UTexture2D* Texture2D = Info.Texture;
+	UTexture2D* Texture2D = Cast<UTexture2D>(Info.RenderAsset);
+	check(Texture2D);
 	if (Texture2D->LevelIndex == INDEX_NONE)
 	{
 		// If this is the first time this texture gets processed in the packing process, encode it.
@@ -406,7 +407,7 @@ int32* FStreamingTextureLevelContext::GetBuildDataIndexRef(UTexture2D* Texture2D
 	return nullptr;
 }
 
-void FStreamingTextureLevelContext::ProcessMaterial(const FBoxSphereBounds& ComponentBounds, const FPrimitiveMaterialInfo& MaterialData, float ComponentScaling, TArray<FStreamingTexturePrimitiveInfo>& OutStreamingTextures)
+void FStreamingTextureLevelContext::ProcessMaterial(const FBoxSphereBounds& ComponentBounds, const FPrimitiveMaterialInfo& MaterialData, float ComponentScaling, TArray<FStreamingRenderAssetPrimitiveInfo>& OutStreamingTextures)
 {
 	ensure(MaterialData.IsValid());
 
@@ -416,7 +417,7 @@ void FStreamingTextureLevelContext::ProcessMaterial(const FBoxSphereBounds& Comp
 	for (UTexture* Texture : Textures)
 	{
 		UTexture2D* Texture2D = Cast<UTexture2D>(Texture);
-		if (!IsStreamingTexture(Texture2D))
+		if (!IsStreamingRenderAsset(Texture2D))
 		{
 			continue;
 		}
@@ -426,10 +427,10 @@ void FStreamingTextureLevelContext::ProcessMaterial(const FBoxSphereBounds& Comp
 		{
 			if (*BuildDataIndex != INDEX_NONE)
 			{
-				FStreamingTexturePrimitiveInfo& Info = *new(OutStreamingTextures) FStreamingTexturePrimitiveInfo();
+				FStreamingRenderAssetPrimitiveInfo& Info = *new(OutStreamingTextures) FStreamingRenderAssetPrimitiveInfo();
 				const FStreamingTextureBuildInfo& BuildInfo = (*ComponentBuildData)[*BuildDataIndex];
 
-				Info.Texture = Texture2D;
+				Info.RenderAsset = Texture2D;
 				Info.TexelFactor = BuildInfo.TexelFactor * ComponentScaling;
 				Info.PackedRelativeBox = bUseRelativeBoxes ? BuildInfo.PackedRelativeBox : PackedRelativeBox_Identity;
 				UnpackRelativeBox(ComponentBounds, Info.PackedRelativeBox, Info.Bounds);
@@ -451,9 +452,9 @@ void FStreamingTextureLevelContext::ProcessMaterial(const FBoxSphereBounds& Comp
 
 			if (TextureDensity)
 			{
-				FStreamingTexturePrimitiveInfo& Info = *new(OutStreamingTextures) FStreamingTexturePrimitiveInfo();
+				FStreamingRenderAssetPrimitiveInfo& Info = *new(OutStreamingTextures) FStreamingRenderAssetPrimitiveInfo();
 
-				Info.Texture = Texture2D;
+				Info.RenderAsset = Texture2D;
 				Info.TexelFactor = TextureDensity * ComponentScaling;
 				Info.PackedRelativeBox = bUseRelativeBoxes ? MaterialData.PackedRelativeBox : PackedRelativeBox_Identity;
 				UnpackRelativeBox(ComponentBounds, Info.PackedRelativeBox, Info.Bounds);
diff --git a/Engine/Source/Runtime/Engine/Private/Streaming/TextureStreamingHelpers.cpp b/Engine/Source/Runtime/Engine/Private/Streaming/TextureStreamingHelpers.cpp
--- a/Engine/Source/Runtime/Engine/Private/Streaming/TextureStreamingHelpers.cpp
+++ b/Engine/Source/Runtime/Engine/Private/Streaming/TextureStreamingHelpers.cpp
@@ -44,13 +44,13 @@ DECLARE_MEMORY_STAT_POOL(TEXT("     NeverStream"),		STAT_StreamingOverview06_Nev
 DECLARE_MEMORY_STAT_POOL(TEXT("     UI Group"),			STAT_StreamingOverview07_UIGroup, STATGROUP_StreamingOverview, FPlatformMemory::MCR_StreamingPool);
 DECLARE_MEMORY_STAT_POOL(TEXT("Average Required PoolSize"),	STAT_StreamingOverview08_AverageRequiredPool, STATGROUP_StreamingOverview, FPlatformMemory::MCR_StreamingPool);
 
-DEFINE_STAT(STAT_TextureStreaming_GameThreadUpdateTime);
+DEFINE_STAT(STAT_RenderAssetStreaming_GameThreadUpdateTime);
 
 CSV_DEFINE_CATEGORY(TextureStreaming, true);
 
 DEFINE_LOG_CATEGORY(LogContentStreaming);
 
-int32 FTextureStreamingSettings::ExtraIOLatency = 0;
+int32 FRenderAssetStreamingSettings::ExtraIOLatency = 0;
 
 ENGINE_API TAutoConsoleVariable<int32> CVarStreamingUseNewMetrics(
 	TEXT("r.Streaming.UseNewMetrics"),
@@ -221,8 +221,8 @@ TAutoConsoleVariable<int32> CVarStreamingMinMipForSplitRequest(
 	TEXT("If non-zero, the minimum hidden mip for which load requests will first load the visible mip"),
 	ECVF_Default);
 
-TAutoConsoleVariable<float> CVarStreamingMinLevelTextureScreenSize(
-	TEXT("r.Streaming.MinLevelTextureScreenSize"),
+TAutoConsoleVariable<float> CVarStreamingMinLevelRenderAssetScreenSize(
+	TEXT("r.Streaming.MinLevelRenderAssetScreenSize"),
 	100,
 	TEXT("If non-zero, levels only get handled if any of their referenced texture could be required of this size. Using conservative metrics on the level data."),
 	ECVF_Default);
@@ -259,7 +259,7 @@ static TAutoConsoleVariable<int32> CVarStreamingStressTestFramesForFullUpdate(
 	TEXT("Num frames to update texture states when doing the stress tests."),
 	ECVF_Cheat);
 
-void FTextureStreamingSettings::Update()
+void FRenderAssetStreamingSettings::Update()
 {
 	MaxEffectiveScreenSize = CVarStreamingScreenSizeEffectiveMax.GetValueOnAnyThread();
 	MaxTempMemoryAllowed = CVarStreamingMaxTempMemoryAllowed.GetValueOnAnyThread();
@@ -274,7 +274,7 @@ void FTextureStreamingSettings::Update()
 	MinMipForSplitRequest = CVarStreamingMinMipForSplitRequest.GetValueOnAnyThread();
 	PerTextureBiasViewBoostThreshold = CVarStreamingPerTextureBiasViewBoostThreshold.GetValueOnAnyThread();
 	MaxHiddenPrimitiveViewBoost = FMath::Max<float>(1.f, CVarStreamingMaxHiddenPrimitiveViewBoost.GetValueOnAnyThread());
-	MinLevelTextureScreenSize = CVarStreamingMinLevelTextureScreenSize.GetValueOnAnyThread();
+	MinLevelRenderAssetScreenSize = CVarStreamingMinLevelRenderAssetScreenSize.GetValueOnAnyThread();
 	MaxTextureUVDensity = CVarStreamingMaxTextureUVDensity.GetValueOnAnyThread();
 	bUseMaterialData = bUseNewMetrics && CVarStreamingUseMaterialData.GetValueOnAnyThread() != 0;
 	HiddenPrimitiveScale = bUseNewMetrics ? CVarStreamingHiddenPrimitiveScale.GetValueOnAnyThread() : 1.f;
@@ -329,7 +329,7 @@ float GShadowmapStreamingFactor = 0.09f;
 /** For testing, finding useless textures or special demo purposes. If true, textures will never be streamed out (but they can be GC'd). 
 * Caution: this only applies to unlimited texture pools (i.e. not consoles)
 */
-bool GNeverStreamOutTextures = false;
+bool GNeverStreamOutRenderAssets = false;
 
 #if STATS
 extern int64 GUITextureMemory;
@@ -340,11 +340,11 @@ int64 GRequiredPoolSizeCount= 0;
 int64 GAverageRequiredPool = 0;
 #endif
 
-void FTextureStreamingStats::Apply()
+void FRenderAssetStreamingStats::Apply()
 {
 	/** Streaming stats */
 
-	SET_MEMORY_STAT(MCR_TexturePool, TexturePool); 
+	SET_MEMORY_STAT(MCR_TexturePool, RenderAssetPool); 
 	SET_MEMORY_STAT(MCR_StreamingPool, StreamingPool);
 	SET_MEMORY_STAT(MCR_UsedStreamingPool, UsedStreamingPool);
 
diff --git a/Engine/Source/Runtime/Engine/Private/Streaming/TextureStreamingHelpers.h b/Engine/Source/Runtime/Engine/Private/Streaming/TextureStreamingHelpers.h
--- a/Engine/Source/Runtime/Engine/Private/Streaming/TextureStreamingHelpers.h
+++ b/Engine/Source/Runtime/Engine/Private/Streaming/TextureStreamingHelpers.h
@@ -11,13 +11,13 @@ TextureStreamingHelpers.h: Definitions of classes used for texture streaming.
 #include "HAL/IConsoleManager.h"
 #include "Misc/MemStack.h"
 
-class UTexture2D;
+class UStreamableRenderAsset;
 
 /**
  * Streaming stats
  */
 
-DECLARE_CYCLE_STAT_EXTERN(TEXT("Texture Streaming Game Thread Update Time"), STAT_TextureStreaming_GameThreadUpdateTime,STATGROUP_Streaming, );
+DECLARE_CYCLE_STAT_EXTERN(TEXT("Renderable Asset Streaming Game Thread Update Time"), STAT_RenderAssetStreaming_GameThreadUpdateTime,STATGROUP_Streaming, );
 
 
 // Streaming Details
@@ -34,7 +34,7 @@ DECLARE_LOG_CATEGORY_EXTERN(LogContentStreaming, Log, All);
 
 extern float GLightmapStreamingFactor;
 extern float GShadowmapStreamingFactor;
-extern bool GNeverStreamOutTextures;
+extern bool GNeverStreamOutRenderAssets;
 
 //@DEBUG:
 // Set to 1 to log all dynamic component notifications
@@ -60,13 +60,13 @@ extern TAutoConsoleVariable<int32> CVarStreamingNumStaticComponentsProcessedPerF
 extern TAutoConsoleVariable<int32> CVarStreamingDefragDynamicBounds;
 extern TAutoConsoleVariable<float> CVarStreamingMaxTextureUVDensity;
 
-struct FTextureStreamingSettings
+struct FRenderAssetStreamingSettings
 {
-	FTextureStreamingSettings() { Update(); }
+	FRenderAssetStreamingSettings() { Update(); }
 	void Update();
 
-	FORCEINLINE bool operator ==(const FTextureStreamingSettings& Rhs) const { return FMemory::Memcmp(this, &Rhs, sizeof(FTextureStreamingSettings)) == 0; }
-	FORCEINLINE bool operator !=(const FTextureStreamingSettings& Rhs) const { return FMemory::Memcmp(this, &Rhs, sizeof(FTextureStreamingSettings)) != 0; }
+	FORCEINLINE bool operator ==(const FRenderAssetStreamingSettings& Rhs) const { return FMemory::Memcmp(this, &Rhs, sizeof(FRenderAssetStreamingSettings)) == 0; }
+	FORCEINLINE bool operator !=(const FRenderAssetStreamingSettings& Rhs) const { return FMemory::Memcmp(this, &Rhs, sizeof(FRenderAssetStreamingSettings)) != 0; }
 
 
 	float MaxEffectiveScreenSize;
@@ -85,7 +85,7 @@ struct FTextureStreamingSettings
 	bool bUsePerTextureBias;
 	bool bUseMaterialData;
 	int32 MinMipForSplitRequest;
-	float MinLevelTextureScreenSize;
+	float MinLevelRenderAssetScreenSize;
 	float MaxTextureUVDensity;
 	int32 MaterialQualityLevel;
 	int32 FramesForFullUpdate;
@@ -98,7 +98,7 @@ protected:
 };
 
 typedef TArray<int32, TMemStackAllocator<> > FStreamingRequests;
-typedef TArray<const UTexture2D*, TInlineAllocator<12> > FRemovedTextureArray;
+typedef TArray<const UStreamableRenderAsset*, TInlineAllocator<12> > FRemovedRenderAssetArray;
 
 #define NUM_BANDWIDTHSAMPLES 512
 #define NUM_LATENCYSAMPLES 512
@@ -108,52 +108,50 @@ typedef TArray<const UTexture2D*, TInlineAllocator<12> > FRemovedTextureArray;
 #define MAX_MIPDELTA			5.0f
 #define MAX_LASTRENDERTIME		90.0f
 
-class UTexture2D;
 class UPrimitiveComponent;
-class FTextureBoundsVisibility;
-class FDynamicTextureInstanceManager;
+class FDynamicRenderAssetInstanceManager;
 template<typename T>
 class FAsyncTask;
-class FAsyncTextureStreamingTask;
+class FRenderAssetStreamingMipCalcTask;
 
 
-struct FStreamingTexture;
+struct FStreamingRenderAsset;
 struct FStreamingContext;
 struct FStreamingHandlerTextureBase;
 struct FTexturePriority;
 
-struct FTextureStreamingStats
+struct FRenderAssetStreamingStats
 {
-	FTextureStreamingStats() { Reset(); }
+	FRenderAssetStreamingStats() { Reset(); }
 
-	void Reset() { FMemory::Memzero(this, sizeof(FTextureStreamingStats)); }
+	void Reset() { FMemory::Memzero(this, sizeof(FRenderAssetStreamingStats)); }
 
 	void Apply();
 
-	int64 TexturePool;
+	int64 RenderAssetPool;		// Streaming pool size in bytes given by RHI, usally some percentage of available VRAM
 	// int64 StreamingPool;
-	int64 UsedStreamingPool;
+	int64 UsedStreamingPool;	// Estimated memory in bytes the streamer actually use
 
-	int64 SafetyPool;
-	int64 TemporaryPool;
-	int64 StreamingPool;
-	int64 NonStreamingMips;
+	int64 SafetyPool;			// Memory margin in bytes we want to left unused
+	int64 TemporaryPool;		// Pool size in bytes for staging resources
+	int64 StreamingPool;		// Estimated memory in bytes available for streaming after subtracting non-streamable memory and safety pool
+	int64 NonStreamingMips;		// Estimated memory in bytes the streamer deemed non-streamable
 
-	int64 RequiredPool;
-	int64 VisibleMips;
-	int64 HiddenMips;
-	int64 ForcedMips;
-	int64 UnkownRefMips;
-	int64 CachedMips;
+	int64 RequiredPool;			// Estimated memory in bytes the streamer would use if there was no limit
+	int64 VisibleMips;			// Estimated memory in bytes used by visible assets
+	int64 HiddenMips;			// Estimated memory in bytes used by hidden (e.g. culled by view frustum or occlusion) assets
+	int64 ForcedMips;			// Estimated memory in bytes used by forced-fully-load assets
+	int64 UnkownRefMips;		// Estimated memory in bytes used by unknown-ref-heuristic assets
+	int64 CachedMips;			// Estimated memory in bytes used by the streamer for caching
 
-	int64 WantedMips;
-	int64 NewRequests; // How much texture memory is required by new requests.
-	int64 PendingRequests; // How much texture memory is waiting to be loaded for previous requests.
-	int64 MipIOBandwidth;
+	int64 WantedMips;			// Estimated memory in bytes would be resident if there was no caching
+	int64 NewRequests;			// Estimated memory in bytes required by new requests (TODO)
+	int64 PendingRequests;		// Estimated memory in bytes waiting to be loaded for previous requests
+	int64 MipIOBandwidth;		// Estimated IO bandwidth in bytes/sec
 
-	int64 OverBudget;
+	int64 OverBudget;			// RequiredPool - StreamingPool
 
-	double Timestamp;
+	double Timestamp;			// Last time point these stats are updated in seconds
 
 	
 	volatile int32 CallbacksCycles;
diff --git a/Engine/Source/Runtime/Engine/Private/Texture.cpp b/Engine/Source/Runtime/Engine/Private/Texture.cpp
--- a/Engine/Source/Runtime/Engine/Private/Texture.cpp
+++ b/Engine/Source/Runtime/Engine/Private/Texture.cpp
@@ -118,11 +118,6 @@ bool UTexture::IsPostLoadThreadSafe() const
 	return false;
 }
 
-int32 UTexture::GetCachedLODBias() const
-{
-	return CachedCombinedLODBias;
-}
-
 #if WITH_EDITOR
 bool UTexture::CanEditChange(const UProperty* InProperty) const
 {
diff --git a/Engine/Source/Runtime/Engine/Private/Texture2D.cpp b/Engine/Source/Runtime/Engine/Private/Texture2D.cpp
--- a/Engine/Source/Runtime/Engine/Private/Texture2D.cpp
+++ b/Engine/Source/Runtime/Engine/Private/Texture2D.cpp
@@ -128,8 +128,8 @@ int64 GUITextureMemory = 0;
 int64 GNeverStreamTextureMemory = 0;
 #endif
 
-/** Turn on ENABLE_TEXTURE_TRACKING in ContentStreaming.cpp and setup GTrackedTextures to track specific textures through the streaming system. */
-extern bool TrackTextureEvent( FStreamingTexture* StreamingTexture, UTexture2D* Texture, bool bForceMipLevelsToBeResident, const FStreamingManagerTexture* Manager );
+/** Turn on ENABLE_RENDER_ASSET_TRACKING in ContentStreaming.cpp and setup GTrackedTextures to track specific textures/meshes through the streaming system. */
+extern bool TrackTextureEvent( FStreamingRenderAsset* StreamingTexture, UStreamableRenderAsset* Texture, bool bForceMipLevelsToBeResident, const FRenderAssetStreamingManager* Manager );
 
 
 /** Scoped debug info that provides the texture name to memory allocation and crash callstacks. */
@@ -254,7 +254,7 @@ void UTexture2D::Serialize(FArchive& Ar)
 #endif // #if WITH_EDITOR
 }
 
-float UTexture2D::GetLastRenderTimeForStreaming()
+float UTexture2D::GetLastRenderTimeForStreaming() const
 {
 	float LastRenderTime = -FLT_MAX;
 	if (Resource)
@@ -313,7 +313,7 @@ int32 UTexture2D::GetNumRequestedMips() const
 	}
 	else
 	{
-		return GetNumResidentMips();
+		return GetCachedNumResidentLODs();
 	}
 }
 
@@ -413,9 +413,9 @@ float UTexture2D::GetAverageBrightness(bool bIgnoreTrueBlack, bool bUseGrayscale
 
 void UTexture2D::LinkStreaming()
 {
-	if (!IsTemplate() && IStreamingManager::Get().IsTextureStreamingEnabled() && IsStreamingTexture(this))
+	if (!IsTemplate() && IStreamingManager::Get().IsTextureStreamingEnabled() && IsStreamingRenderAsset(this))
 	{
-		IStreamingManager::Get().GetTextureStreamingManager().AddStreamingTexture(this);
+		IStreamingManager::Get().GetTextureStreamingManager().AddStreamingRenderAsset(this);
 	}
 	else
 	{
@@ -427,7 +427,7 @@ void UTexture2D::UnlinkStreaming()
 {
 	if (!IsTemplate() && IStreamingManager::Get().IsTextureStreamingEnabled())
 	{
-		IStreamingManager::Get().GetTextureStreamingManager().RemoveStreamingTexture(this);
+		IStreamingManager::Get().GetTextureStreamingManager().RemoveStreamingRenderAsset(this);
 	}
 }
 
@@ -562,7 +562,7 @@ void UTexture2D::WaitForStreaming()
 		// Update the wanted mip and stream in..		
 		if (IStreamingManager::Get().IsTextureStreamingEnabled())
 		{
-			IStreamingManager::Get().GetTextureStreamingManager().UpdateIndividualTexture( this );
+			IStreamingManager::Get().GetTextureStreamingManager().UpdateIndividualRenderAsset( this );
 
 			while (	UpdateStreamingStatus() ) 
 			{
@@ -956,11 +956,7 @@ void UTexture2D::GetResourceSizeEx(FResourceSizeEx& CumulativeResourceSize)
 
 bool UTexture2D::ShouldMipLevelsBeForcedResident() const
 {
-	if ( bGlobalForceMipLevelsToBeResident || bForceMiplevelsToBeResident || LODGroup == TEXTUREGROUP_Skybox )
-	{
-		return true;
-	}
-	if ( ForceMipLevelsToBeResidentTimestamp >= FApp::GetCurrentTime() )
+	if (LODGroup == TEXTUREGROUP_Skybox || this->Super::ShouldMipLevelsBeForcedResident())
 	{
 		return true;
 	}
@@ -1026,14 +1022,6 @@ UTexture2D* UTexture2D::CreateTransient(int32 InSizeX, int32 InSizeY, EPixelForm
 	return NewTexture;
 }
 
-void UTexture2D::SetForceMipLevelsToBeResident( float Seconds, int32 CinematicTextureGroups )
-{
-	uint32 TextureGroupBitfield = (uint32) CinematicTextureGroups;
-	uint32 MyTextureGroup = FMath::BitFlag[LODGroup];
-	bUseCinematicMipLevels = (TextureGroupBitfield & MyTextureGroup) ? true : false;
-	ForceMipLevelsToBeResidentTimestamp = FApp::GetCurrentTime() + Seconds;
-}
-
 int32 UTexture2D::Blueprint_GetSizeX() const
 {
 #if WITH_EDITORONLY_DATA
@@ -1187,6 +1175,7 @@ FTexture2DResource::FTexture2DResource( UTexture2D* InOwner, int32 InitialMipCou
 
 	// Keep track of first miplevel to use.
 	CurrentFirstMip = InOwner->GetNumMips() - InitialMipCount;
+	InOwner->SetCachedNumResidentLODs(static_cast<uint8>(InitialMipCount));
 
 	check(CurrentFirstMip>=0);
 	// texture must be as big as base miptail level
@@ -1315,6 +1304,7 @@ void FTexture2DResource::InitRHI()
 
 				// We're done with initialization.
 				bReadyForStreaming = true;
+				Owner->SetCachedReadyForStreaming(true);
 
 				return;
 			}
@@ -1360,6 +1350,7 @@ void FTexture2DResource::InitRHI()
 
 		// We're done with initialization.
 		bReadyForStreaming = true;
+		Owner->SetCachedReadyForStreaming(true);
 	}
 	else
 	{
@@ -1616,10 +1607,11 @@ void FTexture2DResource::UpdateTexture(FTexture2DRHIRef& InTextureRHI, int32 InN
 
 	if (Owner)
 	{
+		const int32 NumMips = Owner->GetNumMips();
+		
 		// Update mip-level fading.
 		if (CurrentFirstMip != InNewFirstMip)
 		{
-			const int32 NumMips = Owner->GetNumMips();
 			const int32 ResidentMips = NumMips - CurrentFirstMip;
 			const int32 RequestedMips = NumMips - InNewFirstMip;
 			MipBiasFade.SetNewMipCount(FMath::Max<int32>(RequestedMips, ResidentMips), RequestedMips, LastRenderTime, MipFadeSetting);
@@ -1639,6 +1631,7 @@ void FTexture2DResource::UpdateTexture(FTexture2DRHIRef& InTextureRHI, int32 InN
 		TextureRHI		= InTextureRHI;
 		Texture2DRHI	= InTextureRHI;
 		CurrentFirstMip = InNewFirstMip;
+		Owner->SetCachedNumResidentLODs(static_cast<uint8>(NumMips - InNewFirstMip));
 		RHIUpdateTextureReference(Owner->TextureReference.TextureReferenceRHI, TextureRHI);
 	}
 }
diff --git a/Engine/Source/Runtime/Engine/Private/UnrealEngine.cpp b/Engine/Source/Runtime/Engine/Private/UnrealEngine.cpp
--- a/Engine/Source/Runtime/Engine/Private/UnrealEngine.cpp
+++ b/Engine/Source/Runtime/Engine/Private/UnrealEngine.cpp
@@ -254,8 +254,8 @@ void FEngineModule::StartupModule()
 		CVARShowMaterialDrawEvents->SetOnChangedCallback(FConsoleVariableDelegate::CreateStatic(&OnChangeEngineCVarRequiringRecreateRenderState));
 	}
 
-	SuspendTextureStreamingRenderTasks = &SuspendTextureStreamingRenderTasksInternal;
-	ResumeTextureStreamingRenderTasks = &ResumeTextureStreamingRenderTasksInternal;
+	SuspendTextureStreamingRenderTasks = &SuspendRenderAssetStreamingRenderTasksInternal;
+	ResumeTextureStreamingRenderTasks = &ResumeRenderAssetStreamingRenderTasksInternal;
 
 	FParticleSystemWorldManager::OnStartup();
 
@@ -4804,13 +4804,13 @@ bool UEngine::HandleListTexturesCommand( const TCHAR* Cmd, FOutputDevice& Ar )
 		// that GetStreamingTextureInfo doesn't check whether a texture is actually streamable or not
 		// and is also implemented for skeletal meshes and such.
 		FStreamingTextureLevelContext LevelContext(EMaterialQualityLevel::Num, PrimitiveComponent);
-		TArray<FStreamingTexturePrimitiveInfo> StreamingTextures;
-		PrimitiveComponent->GetStreamingTextureInfo( LevelContext, StreamingTextures );
+		TArray<FStreamingRenderAssetPrimitiveInfo> StreamingTextures;
+		PrimitiveComponent->GetStreamingRenderAssetInfo( LevelContext, StreamingTextures );
 
 		// Increase usage count for all referenced textures
 		for( int32 TextureIndex=0; TextureIndex<StreamingTextures.Num(); TextureIndex++ )
 		{
-			UTexture2D* Texture = StreamingTextures[TextureIndex].Texture;
+			UTexture2D* Texture = Cast<UTexture2D>(StreamingTextures[TextureIndex].RenderAsset);
 			if( Texture )
 			{
 				// Initializes UsageCount to 0 if texture is not found.
diff --git a/Engine/Source/Runtime/Engine/Public/ContentStreaming.h b/Engine/Source/Runtime/Engine/Public/ContentStreaming.h
--- a/Engine/Source/Runtime/Engine/Public/ContentStreaming.h
+++ b/Engine/Source/Runtime/Engine/Public/ContentStreaming.h
@@ -14,7 +14,7 @@ class FSoundSource;
 class UPrimitiveComponent;
 class USoundWave;
 class UTexture2D;
-struct FStreamingManagerTexture;
+struct FRenderAssetStreamingManager;
 struct FWaveInstance;
 
 /*-----------------------------------------------------------------------------
@@ -25,10 +25,13 @@ struct FWaveInstance;
 class UPrimitiveComponent;
 class AActor;
 class UTexture2D;
+class UStaticMesh;
+class USkeletalMesh;
+class UStreamableRenderAsset;
 class FSoundSource;
 class USoundWave;
 struct FWaveInstance;
-struct FStreamingManagerTexture;
+struct FRenderAssetStreamingManager;
 
 /** Helper function to flush resource streaming. */
 void FlushResourceStreaming();
@@ -337,16 +340,16 @@ protected:
 };
 
 /**
- * Interface to add functions specifically related to texture streaming
+ * Interface to add functions specifically related to texture/mesh streaming
  */
-struct ITextureStreamingManager : public IStreamingManager
+struct IRenderAssetStreamingManager : public IStreamingManager
 {
 	/**
-	* Updates streaming for an individual texture, taking into account all view infos.
+	* Updates streaming for an individual texture/mesh, taking into account all view infos.
 	*
-	* @param Texture		Texture to update
+	* @param RenderAsset		Texture or mesh to update
 	*/
-	virtual void UpdateIndividualTexture(UTexture2D* Texture) = 0;
+	virtual void UpdateIndividualRenderAsset(UStreamableRenderAsset* RenderAsset) = 0;
 
 	/**
 	* Temporarily boosts the streaming distance factor by the specified number.
@@ -355,31 +358,33 @@ struct ITextureStreamingManager : public IStreamingManager
 	virtual void BoostTextures(AActor* Actor, float BoostFactor) = 0;
 
 	/**
-	*	Try to stream out texture mip-levels to free up more memory.
+	*	Try to stream out texture/mesh mip-levels to free up more memory.
 	*	@param RequiredMemorySize	- Required minimum available texture memory
 	*	@return						- Whether it succeeded or not
 	**/
-	virtual bool StreamOutTextureData(int64 RequiredMemorySize) = 0;
+	virtual bool StreamOutRenderAssetData(int64 RequiredMemorySize) = 0;
 
-	/** Adds a new texture to the streaming manager. */
-	virtual void AddStreamingTexture(UTexture2D* Texture) = 0;
+	/** Adds a new texture/mesh to the streaming manager. */
+	virtual void AddStreamingRenderAsset(UTexture2D* Texture) = 0;
+	virtual void AddStreamingRenderAsset(UStaticMesh* StaticMesh) = 0;
+	virtual void AddStreamingRenderAsset(USkeletalMesh* SkeletalMesh) = 0;
 
-	/** Removes a texture from the streaming manager. */
-	virtual void RemoveStreamingTexture(UTexture2D* Texture) = 0;
+	/** Removes a texture/mesh from the streaming manager. */
+	virtual void RemoveStreamingRenderAsset(UStreamableRenderAsset* RenderAsset) = 0;
 
 	virtual int64 GetMemoryOverBudget() const = 0;
 
 	/** Pool size for streaming. */
 	virtual int64 GetPoolSize() const = 0;
 
-	/** Max required textures ever seen in bytes. */
+	/** Max required textures/meshes ever seen in bytes. */
 	virtual int64 GetMaxEverRequired() const = 0;
 
-	/** Resets the max ever required textures.  For possibly when changing resolutions or screen pct. */
+	/** Resets the max ever required textures/meshes.  For possibly when changing resolutions or screen pct. */
 	virtual void ResetMaxEverRequired() = 0;
 
-	/** Set current pause state for texture streaming */
-	virtual void PauseTextureStreaming(bool bInShouldPause) = 0;
+	/** Set current pause state for texture/mesh streaming */
+	virtual void PauseRenderAssetStreaming(bool bInShouldPause) = 0;
 
 	/** Return all bounds related to the ref object */
 	ENGINE_API virtual void GetObjectReferenceBounds(const UObject* RefObject, TArray<FBox>& AssetBoxes) = 0;
@@ -499,7 +504,7 @@ struct FStreamingManagerCollection : public IStreamingManager
 	/**
 	 * Gets a reference to the Texture Streaming Manager interface
 	 */
-	ENGINE_API ITextureStreamingManager& GetTextureStreamingManager() const;
+	ENGINE_API IRenderAssetStreamingManager& GetTextureStreamingManager() const;
 
 	/**
 	 * Gets a reference to the Audio Streaming Manager interface
@@ -592,7 +597,7 @@ protected:
 	float LoadMapTimeLimit;
 
 	/** The currently added texture streaming manager. Can be NULL*/
-	FStreamingManagerTexture* TextureStreamingManager;
+	FRenderAssetStreamingManager* TextureStreamingManager;
 
 	/** The audio streaming manager, should always exist */
 	IAudioStreamingManager* AudioStreamingManager;
diff --git a/Engine/Source/Runtime/Engine/Public/PrimitiveSceneProxy.h b/Engine/Source/Runtime/Engine/Public/PrimitiveSceneProxy.h
--- a/Engine/Source/Runtime/Engine/Public/PrimitiveSceneProxy.h
+++ b/Engine/Source/Runtime/Engine/Public/PrimitiveSceneProxy.h
@@ -726,6 +726,8 @@ public:
 	 */
 	ENGINE_API virtual struct FLODMask GetCustomWholeSceneShadowLOD(const FSceneView& InView, float InViewLODScale, int32 InForcedLODLevel, const struct FLODMask& InVisibilePrimitiveLODMask, float InShadowMapTextureResolution, float InShadowMapCascadeSize, int8 InShadowCascadeId, bool InHasSelfShadow) const;
 
+	virtual uint8 GetCurrentFirstLODIdx_RenderThread() const { return 0; }
+
 protected:
 
 	/** Allow subclasses to override the primitive name. Used primarily by BSP. */
