Resolve robomerge conflicts

diff --git a/Engine/Source/Runtime/Engine/Public/RawIndexBuffer.h b/Engine/Source/Runtime/Engine/Public/RawIndexBuffer.h
--- a/Engine/Source/Runtime/Engine/Public/RawIndexBuffer.h
+++ b/Engine/Source/Runtime/Engine/Public/RawIndexBuffer.h
@@ -226,7 +226,7 @@ public:
 	 */
 	FORCEINLINE int32 GetNumIndices() const
 	{
-		return b32Bit ? (IndexStorage.Num()/4) : (IndexStorage.Num()/2);
+		return CachedNumIndices >= 0 ? CachedNumIndices : (b32Bit ? (IndexStorage.Num()/4) : (IndexStorage.Num()/2));
 	}
 
 	/**
@@ -237,6 +237,19 @@ public:
 		return IndexStorage.GetAllocatedSize();
 	}
 
+	/** == GetNumIndices() * (b32Bit ? 4 : 2) */
+	int32 GetIndexDataSize() const { return IndexStorage.Num(); }
+
+	/** Create an RHI index buffer with CPU data. CPU data may be discarded after creation (see TResourceArray::Discard) */
+	FIndexBufferRHIRef CreateRHIBuffer_RenderThread();
+	FIndexBufferRHIRef CreateRHIBuffer_Async();
+
+	/** Set whether this buffer is managed by the streamer. Must set before InitRHI is called */
+	void SetIsStreamed(bool bValue) { bStreamed = bValue; }
+
+	/** Take over ownership of IntermediateBuffer */
+	void InitRHIForStreaming(FIndexBufferRHIParamRef IntermediateBuffer);
+
 	/**
 	 * Serialization.
 	 * @param	Ar				Archive to serialize with
@@ -244,6 +257,9 @@ public:
 	 */
 	void Serialize(FArchive& Ar, bool bNeedsCPUAccess);
 
+	/** Serialize only meta data (e.g. number of indices) but not the actual index data */
+	void SerializeMetaData(FArchive& Ar);
+
     /**
      * Discard
      * discards the serialized data when it is not needed
@@ -256,10 +272,19 @@ public:
 	inline bool Is32Bit() const { return b32Bit; }
 
 private:
+	template <bool bRenderThread>
+	FIndexBufferRHIRef CreateRHIBuffer_Internal();
+
 	/** Storage for indices. */
 	TResourceArray<uint8,INDEXBUFFER_ALIGNMENT> IndexStorage;
+
+	/** If >= 0, represents the number of indices in this IB. Needed in cooked build since InitRHI may discard CPU data */
+	int32 CachedNumIndices;
+
 	/** 32bit or 16bit? */
 	bool b32Bit;
+
+	bool bStreamed;
 };
 
 /**
diff --git a/Engine/Source/Runtime/Engine/Public/Rendering/ColorVertexBuffer.h b/Engine/Source/Runtime/Engine/Public/Rendering/ColorVertexBuffer.h
--- a/Engine/Source/Runtime/Engine/Public/Rendering/ColorVertexBuffer.h
+++ b/Engine/Source/Runtime/Engine/Public/Rendering/ColorVertexBuffer.h
@@ -57,6 +57,8 @@ public:
 	*/
 	void Serialize(FArchive& Ar, bool bNeedsCPUAccess);
 
+	void SerializeMetaData(FArchive& Ar);
+
 	/**
 	* Export the data to a string, used for editor Copy&Paste.
 	* The method must not be called if there is no data.
@@ -132,6 +134,36 @@ public:
 		InitFromColorArray(&InColor, Count, 0);
 	}
 
+	/** Create an RHI vertex buffer with CPU data. CPU data may be discarded after creation (see TResourceArray::Discard) */
+	FVertexBufferRHIRef CreateRHIBuffer_RenderThread();
+	FVertexBufferRHIRef CreateRHIBuffer_Async();
+
+	/** Set whether this buffer is managed by the streamer. Must be set before InitRHI is called */
+	void SetIsStreamed(bool bValue) { bStreamed = bValue; }
+
+	/** Similar to Init/ReleaseRHI but only update existing SRV so references to the SRV stays valid */
+	template <int32 MaxNumUpdates>
+	void InitRHIForStreaming(FVertexBufferRHIParamRef IntermediateBuffer, TRHIResourceUpdateBatcher<MaxNumUpdates>& Batcher)
+	{
+		check(!VertexBufferRHI);
+		if (IntermediateBuffer)
+		{
+			check(ColorComponentsSRV);
+			VertexBufferRHI = IntermediateBuffer;
+			Batcher.QueueUpdateRequest(ColorComponentsSRV, VertexBufferRHI, 4, PF_R8G8B8A8);
+		}
+	}
+
+	template <int32 MaxNumUpdates>
+	void ReleaseRHIForStreaming(TRHIResourceUpdateBatcher<MaxNumUpdates>& Batcher)
+	{
+		if (ColorComponentsSRV)
+		{
+			Batcher.QueueUpdateRequest(ColorComponentsSRV, nullptr, 0, 0);
+		}
+		FVertexBuffer::ReleaseRHI();
+	}
+
 	// FRenderResource interface.
 	ENGINE_API virtual void InitRHI() override;
 	ENGINE_API virtual void ReleaseRHI() override;
@@ -168,9 +200,14 @@ private:
 
 	bool NeedsCPUAccess = true;
 
+	bool bStreamed;
+
 	/** Allocates the vertex data storage type. */
 	void AllocateData(bool bNeedsCPUAccess = true);
 
+	template <bool bRenderThread>
+	FVertexBufferRHIRef CreateRHIBuffer_Internal();
+
 	/** Purposely hidden */
 	ENGINE_API FColorVertexBuffer(const FColorVertexBuffer &rhs);
 };
\ No newline at end of file
diff --git a/Engine/Source/Runtime/Engine/Public/Rendering/PositionVertexBuffer.h b/Engine/Source/Runtime/Engine/Public/Rendering/PositionVertexBuffer.h
--- a/Engine/Source/Runtime/Engine/Public/Rendering/PositionVertexBuffer.h
+++ b/Engine/Source/Runtime/Engine/Public/Rendering/PositionVertexBuffer.h
@@ -64,6 +64,8 @@ public:
 	*/
 	void Serialize(FArchive& Ar, bool bInNeedsCPUAccess);
 
+	void SerializeMetaData(FArchive& Ar);
+
 	/**
 	* Specialized assignment operator, only used when importing LOD's.
 	*/
@@ -90,6 +92,38 @@ public:
 		return NumVertices;
 	}
 
+	/** Create an RHI vertex buffer with CPU data. CPU data may be discarded after creation (see TResourceArray::Discard) */
+	FVertexBufferRHIRef CreateRHIBuffer_RenderThread();
+	FVertexBufferRHIRef CreateRHIBuffer_Async();
+
+	/** Set whether this buffer is managed by the streamer. Must be set before InitRHI is called */
+	void SetIsStreamed(bool bValue) { bStreamed = bValue; }
+
+	/** Similar to Init/ReleaseRHI but only update existing SRV so references to the SRV stays valid */
+	template <int32 MaxNumUpdates>
+	void InitRHIForStreaming(FVertexBufferRHIParamRef IntermediateBuffer, TRHIResourceUpdateBatcher<MaxNumUpdates>& Batcher)
+	{
+		check(!VertexBufferRHI);
+		if (IntermediateBuffer)
+		{
+			VertexBufferRHI = IntermediateBuffer;
+			if (PositionComponentSRV)
+			{
+				Batcher.QueueUpdateRequest(PositionComponentSRV, VertexBufferRHI, 4, PF_R32_FLOAT);
+			}
+		}
+	}
+
+	template <int32 MaxNumUpdates>
+	void ReleaseRHIForStreaming(TRHIResourceUpdateBatcher<MaxNumUpdates>& Batcher)
+	{
+		if (PositionComponentSRV)
+		{
+			Batcher.QueueUpdateRequest(PositionComponentSRV, nullptr, 0, 0);
+		}
+		FVertexBuffer::ReleaseRHI();
+	}
+
 	// FRenderResource interface.
 	ENGINE_API virtual void InitRHI() override;
 	ENGINE_API virtual void ReleaseRHI() override;
@@ -120,6 +154,11 @@ private:
 
 	bool bNeedsCPUAccess = true;
 
+	bool bStreamed;
+
 	/** Allocates the vertex data storage type. */
 	void AllocateData(bool bInNeedsCPUAccess = true);
+
+	template <bool bRenderThread>
+	FVertexBufferRHIRef CreateRHIBuffer_Internal();
 };
diff --git a/Engine/Source/Runtime/Engine/Public/Rendering/StaticMeshVertexBuffer.h b/Engine/Source/Runtime/Engine/Public/Rendering/StaticMeshVertexBuffer.h
--- a/Engine/Source/Runtime/Engine/Public/Rendering/StaticMeshVertexBuffer.h
+++ b/Engine/Source/Runtime/Engine/Public/Rendering/StaticMeshVertexBuffer.h
@@ -176,6 +176,8 @@ public:
 	*/
 	void Serialize(FArchive& Ar, bool bNeedsCPUAccess);
 
+	void SerializeMetaData(FArchive& Ar);
+
 	/**
 	* Specialized assignment operator, only used when importing LOD's.
 	*/
@@ -384,6 +386,64 @@ public:
 		return (TangentsStride + (TexcoordStride * GetNumTexCoords())) * NumVertices;
 	}
 
+	/** Create an RHI vertex buffer with CPU data. CPU data may be discarded after creation (see TResourceArray::Discard) */
+	FVertexBufferRHIRef CreateTangentsRHIBuffer_RenderThread();
+	FVertexBufferRHIRef CreateTangentsRHIBuffer_Async();
+	FVertexBufferRHIRef CreateTexCoordRHIBuffer_RenderThread();
+	FVertexBufferRHIRef CreateTexCoordRHIBuffer_Async();
+
+	/** Set whether this buffer is managed by the streamer. Must be set before InitRHI is called */
+	void SetIsStreamed(bool bValue) { bStreamed = bValue; }
+
+	/** Similar to Init/ReleaseRHI but only update existing SRV so references to the SRV stays valid */
+	template <int32 MaxNumUpdates>
+	void InitRHIForStreaming(
+		FVertexBufferRHIParamRef IntermediateTangentsBuffer,
+		FVertexBufferRHIParamRef IntermediateTexCoordBuffer,
+		TRHIResourceUpdateBatcher<MaxNumUpdates>& Batcher)
+	{
+		check(!TangentsVertexBuffer.VertexBufferRHI && !TexCoordVertexBuffer.VertexBufferRHI);
+		if (IntermediateTangentsBuffer)
+		{
+			TangentsVertexBuffer.VertexBufferRHI = IntermediateTangentsBuffer;
+			if (TangentsSRV)
+			{
+				Batcher.QueueUpdateRequest(
+					TangentsSRV,
+					TangentsVertexBuffer.VertexBufferRHI,
+					GetUseHighPrecisionTangentBasis() ? 8u : 4u,
+					GetUseHighPrecisionTangentBasis() ? (uint8)PF_R16G16B16A16_SNORM : (uint8)PF_R8G8B8A8_SNORM);;
+			}
+		}
+		if (IntermediateTexCoordBuffer)
+		{
+			TexCoordVertexBuffer.VertexBufferRHI = IntermediateTexCoordBuffer;
+			if (TextureCoordinatesSRV)
+			{
+				Batcher.QueueUpdateRequest(
+					TextureCoordinatesSRV,
+					TexCoordVertexBuffer.VertexBufferRHI,
+					GetUseFullPrecisionUVs() ? 8u : 4u,
+					GetUseFullPrecisionUVs() ? (uint8)PF_G32R32F : (uint8)PF_G16R16F);
+			}
+		}
+	}
+
+	template<int32 MaxNumUpdates>
+	void ReleaseRHIForStreaming(TRHIResourceUpdateBatcher<MaxNumUpdates>& Batcher)
+	{
+		if (TangentsSRV)
+		{
+			Batcher.QueueUpdateRequest(TangentsSRV, nullptr, 0, 0);
+		}
+		if (TextureCoordinatesSRV)
+		{
+			Batcher.QueueUpdateRequest(TextureCoordinatesSRV, nullptr, 0, 0);
+		}
+		TangentsVertexBuffer.ReleaseRHI();
+		TexCoordVertexBuffer.ReleaseRHI();
+	}
+
 	// FRenderResource interface.
 	ENGINE_API virtual void InitRHI() override;
 	ENGINE_API virtual void ReleaseRHI() override;
@@ -465,6 +525,8 @@ private:
 
 	bool NeedsCPUAccess = true;
 
+	bool bStreamed;
+
 	/** Allocates the vertex data storage type. */
 	void AllocateData(bool bNeedsCPUAccess = true);
 
@@ -472,4 +534,12 @@ private:
 	* @param InData - optional half float source data to convert into full float texture coordinate buffer. if null, convert existing half float texture coordinates to a new float buffer.
 	*/
 	void ConvertHalfTexcoordsToFloat(const uint8* InData);
+
+	template <bool bRenderThread>
+	FVertexBufferRHIRef CreateTangentsRHIBuffer_Internal();
+
+	template<bool bRenderThread>
+	FVertexBufferRHIRef CreateTexCoordRHIBuffer_Internal();
+
+	void InitTangentAndTexCoordStrides();
 };
\ No newline at end of file
diff --git a/Engine/Source/Runtime/Engine/Public/SceneManagement.h b/Engine/Source/Runtime/Engine/Public/SceneManagement.h
--- a/Engine/Source/Runtime/Engine/Public/SceneManagement.h
+++ b/Engine/Source/Runtime/Engine/Public/SceneManagement.h
@@ -2649,6 +2649,11 @@ struct FLODMask
 	{
 		DitheredLODIndices[SampleIndex] = (int8)LODIndex;
 	}
+	void ClampToFirstLOD(int8 FirstLODIdx)
+	{
+		DitheredLODIndices[0] = FMath::Max(DitheredLODIndices[0], FirstLODIdx);
+		DitheredLODIndices[1] = FMath::Max(DitheredLODIndices[1], FirstLODIdx);
+	}
 	bool ContainsLOD(int32 LODIndex) const
 	{
 		return DitheredLODIndices[0] == LODIndex || DitheredLODIndices[1] == LODIndex;
@@ -2671,7 +2676,7 @@ struct FLODMask
 		return DitheredLODIndices[0] != DitheredLODIndices[1];
 	}
 };
-FLODMask ENGINE_API ComputeLODForMeshes(const TArray<class FStaticMeshBatchRelevance>& StaticMeshRelevances, const FSceneView& View, const FVector4& Origin, float SphereRadius, int32 ForcedLODLevel, float& OutScreenRadiusSquared, float ScreenSizeScale = 1.0f, bool bDitheredLODTransition = true);
+FLODMask ENGINE_API ComputeLODForMeshes(const TArray<class FStaticMeshBatchRelevance>& StaticMeshRelevances, const FSceneView& View, const FVector4& Origin, float SphereRadius, int32 ForcedLODLevel, float& OutScreenRadiusSquared, int8 CurFirstLODIdx, float ScreenSizeScale = 1.0f, bool bDitheredLODTransition = true);
 FLODMask ENGINE_API ComputeFastLODForMeshes(const TArray<float>& ScreenSizes, const FSceneView& View, const FVector4& Origin, float SphereRadius, int32 ForcedLODLevel, float& OutScreenRadiusSquared, float ScreenSizeScale = 1.0f, bool bDitheredLODTransition = true);
 
 class FSharedSamplerState : public FRenderResource
diff --git a/Engine/Source/Runtime/Engine/Public/StaticMeshResources.h b/Engine/Source/Runtime/Engine/Public/StaticMeshResources.h
--- a/Engine/Source/Runtime/Engine/Public/StaticMeshResources.h
+++ b/Engine/Source/Runtime/Engine/Public/StaticMeshResources.h
@@ -46,6 +46,7 @@ class UBodySetup;
 #define STATICMESH_ENABLE_DEBUG_RENDERING (!(UE_BUILD_SHIPPING || UE_BUILD_TEST) || WITH_EDITOR)
 
 struct FStaticMaterial;
+struct FStaticMeshBuffersSize;
 
 /**
  * The LOD settings to use for a group of static meshes.
@@ -56,8 +57,10 @@ public:
 	/** Default values. */
 	FStaticMeshLODGroup()
 		: DefaultNumLODs(1)
+		, DefaultMaxNumStreamedLODs(0)
 		, DefaultLightMapResolution(64)
 		, BasePercentTrianglesMult(1.0f)
+		, bSupportLODStreaming(false)
 		, DisplayName( NSLOCTEXT( "UnrealEd", "None", "None" ) )
 	{
 		FMemory::Memzero(SettingsBias);
@@ -70,12 +73,24 @@ public:
 		return DefaultNumLODs;
 	}
 
+	/** Returns the default maximum of streamed LODs */
+	int32 GetDefaultMaxNumStreamedLODs() const
+	{
+		return DefaultMaxNumStreamedLODs;
+	}
+
 	/** Returns the default lightmap resolution. */
 	int32 GetDefaultLightMapResolution() const
 	{
 		return DefaultLightMapResolution;
 	}
 
+	/** Returns whether this LOD group supports LOD streaming. */
+	bool IsLODStreamingSupported() const
+	{
+		return bSupportLODStreaming;
+	}
+
 	/** Returns default reduction settings for the specified LOD. */
 	FMeshReductionSettings GetDefaultSettings(int32 LODIndex) const
 	{
@@ -91,10 +106,14 @@ private:
 	friend class FStaticMeshLODSettings;
 	/** The default number of LODs to build. */
 	int32 DefaultNumLODs;
+	/** Maximum number of streamed LODs */
+	int32 DefaultMaxNumStreamedLODs;
 	/** Default lightmap resolution. */
 	int32 DefaultLightMapResolution;
 	/** An additional reduction of base meshes in this group. */
 	float BasePercentTrianglesMult;
+	/** Whether static meshes in this LOD group can be streamed. */
+	bool bSupportLODStreaming;
 	/** Display name. */
 	FText DisplayName;
 	/** Default reduction settings for meshes in this group. */
@@ -128,6 +147,12 @@ public:
 		return *Group;
 	}
 
+	int32 GetLODGroupIdx(FName GroupName) const
+	{
+		const int32* IdxPtr = GroupName2Index.Find(GroupName);
+		return IdxPtr ? *IdxPtr : INDEX_NONE;
+	}
+
 	/** Retrieve the names of all defined LOD groups. */
 	void GetLODGroupNames(TArray<FName>& OutNames) const;
 
@@ -139,6 +164,8 @@ private:
 	void ReadEntry(FStaticMeshLODGroup& Group, FString Entry);
 	/** Per-group settings. */
 	TMap<FName,FStaticMeshLODGroup> Groups;
+	/** For fast index lookup. Must not change after initialization */
+	TMap<FName, int32> GroupName2Index;
 };
 
 
@@ -283,6 +310,16 @@ struct FStaticMeshLODResources
 
 	/** True if the reversed index buffers contained data at init. Needed as it will not be available to the CPU afterwards. */
 	uint32 bHasReversedDepthOnlyIndices: 1;
+
+	uint32 bHasColorVertexData : 1;
+
+	uint32 bHasWireframeIndices : 1;
+
+	/** True if vertex and index data are serialized inline */
+	uint32 bBuffersInlined : 1;
+
+	/** True if this LOD is optional. That is, vertex and index data may not be available */
+	uint32 bIsOptionalLOD : 1;
 	
 	/**	Allows uniform random selection of mesh sections based on their area. */
 	FStaticMeshAreaWeightedSectionSampler AreaWeightedSampler;
@@ -291,9 +328,23 @@ struct FStaticMeshLODResources
 
 	uint32 DepthOnlyNumTriangles;
 
+	/** Sum of all vertex and index buffer sizes. Calculated in SerializeBuffers */
+	uint32 BuffersSize;
+
+	/** Offset in the .bulk file if this LOD is streamed */
+	uint32 OffsetInFile;
+	/** Size of serialized buffer data in bytes */
+	uint32 BulkDataSize;
+
 #if STATS
 	uint32 StaticMeshIndexMemory;
 #endif
+
+#if WITH_EDITOR
+	FByteBulkData BulkData;
+
+	FString DerivedDataKey;
+#endif
 	
 	/** Default constructor. */
 	ENGINE_API FStaticMeshLODResources();
@@ -316,6 +367,87 @@ struct FStaticMeshLODResources
 	ENGINE_API int32 GetNumVertices() const;
 
 	ENGINE_API int32 GetNumTexCoords() const;
+
+private:
+	enum EClassDataStripFlag : uint8
+	{
+		CDSF_AdjacencyData = 1,
+		CDSF_MinLodData = 2,
+		CDSF_ReversedIndexBuffer = 4,
+	};
+
+	/**
+	 * Due to discard on load, size of an static mesh LOD is not known at cook time and
+	 * this struct is used to keep track of all the information needed to compute LOD size
+	 * at load time
+	 */
+	struct FStaticMeshBuffersSize
+	{
+		uint32 SerializedBuffersSize;
+		uint32 DepthOnlyIBSize;
+		uint32 ReversedIBsSize;
+
+		void Clear()
+		{
+			SerializedBuffersSize = 0;
+			DepthOnlyIBSize = 0;
+			ReversedIBsSize = 0;
+		}
+
+		uint32 CalcBuffersSize() const;
+
+		friend FArchive& operator<<(FArchive& Ar, FStaticMeshBuffersSize& Info)
+		{
+			Ar << Info.SerializedBuffersSize;
+			Ar << Info.DepthOnlyIBSize;
+			Ar << Info.ReversedIBsSize;
+			return Ar;
+		}
+	};
+
+	static uint8 GenerateClassStripFlags(FArchive& Ar, UStaticMesh* OwnerStaticMesh, int32 Index);
+
+	static bool IsLODCookedOut(const ITargetPlatform* TargetPlatform, UStaticMesh* StaticMesh, bool bIsBelowMinLOD);
+
+	static bool IsLODInlined(const ITargetPlatform* TargetPlatform, UStaticMesh* StaticMesh, int32 LODIdx, bool bIsBelowMinLOD);
+
+	/** Compute the size of VertexBuffers and add the result to OutSize */
+	static void AccumVertexBuffersSize(const FStaticMeshVertexBuffers& VertexBuffers, uint32& OutSize);
+
+	/** Compute the size of IndexBuffer and add the result to OutSize */
+	static void AccumIndexBufferSize(const FRawStaticIndexBuffer& IndexBuffer, uint32& OutSize);
+
+	/**
+	 * Serialize vertex and index buffer data for this LOD
+	 * OutBuffersSize - Size of all serialized data in bytes
+	 */
+	void SerializeBuffers(FArchive& Ar, UStaticMesh* OwnerStaticMesh, uint8 InStripFlags, FStaticMeshBuffersSize& OutBuffersSize);
+
+	/**
+	 * Serialize availability information such as bHasDepthOnlyIndices and size of buffers so it
+	 * can be retrieved without loading in actual vertex or index data
+	 */
+	void SerializeAvailabilityInfo(FArchive& Ar);
+
+	template <bool bIncrement>
+	void UpdateIndexMemoryStats();
+
+	template <bool bIncrement>
+	void UpdateVertexMemoryStats() const;
+
+	void ConditionalForce16BitIndexBuffer(EShaderPlatform MaxShaderPlatform, UStaticMesh* Parent);
+
+	void IncrementMemoryStats();
+
+	void DecrementMemoryStats();
+
+	/** Discard loaded vertex and index data. Used when a streaming request is cancelled */
+	void DiscardCPUData();
+
+	friend class FStaticMeshRenderData;
+	friend class FStaticMeshStreamIn;
+	friend class FStaticMeshStreamIn_IO;
+	friend class FStaticMeshStreamOut;
 };
 
 struct ENGINE_API FStaticMeshVertexFactories
@@ -385,6 +517,13 @@ public:
 	/** True if LODs share static lighting data. */
 	bool bLODsShareStaticLighting;
 
+	/** True if rhi resources are initialized */
+	bool bReadyForStreaming;
+
+	uint8 NumInlinedLODs;
+
+	uint8 CurrentFirstLODIdx;
+
 #if WITH_EDITORONLY_DATA
 	/** The derived data key associated with this render data. */
 	FString DerivedDataKey;
@@ -572,6 +711,11 @@ public:
 	/** Sets up a wireframe FMeshBatch for a specific LOD. */
 	virtual bool GetWireframeMeshElement(int32 LODIndex, int32 BatchIndex, const FMaterialRenderProxy* WireframeRenderProxy, uint8 InDepthPriorityGroup, bool bAllowPreCulledIndices, FMeshBatch& OutMeshBatch) const;
 
+	virtual uint8 GetCurrentFirstLODIdx_RenderThread() const override
+	{
+		return GetCurrentFirstLODIdx_Internal();
+	}
+
 protected:
 	/**
 	 * Sets IndexBuffer, FirstIndex and NumPrimitives of OutMeshElement.
@@ -579,6 +723,12 @@ protected:
 	virtual void SetIndexSource(int32 LODIndex, int32 ElementIndex, FMeshBatch& OutMeshElement, bool bWireframe, bool bRequiresAdjacencyInformation, bool bUseInversedIndices, bool bAllowPreCulledIndices) const;
 	bool IsCollisionView(const FEngineShowFlags& EngineShowFlags, bool& bDrawSimpleCollision, bool& bDrawComplexCollision) const;
 
+	/** Only call on render thread timeline */
+	uint8 GetCurrentFirstLODIdx_Internal() const
+	{
+		return RenderData->CurrentFirstLODIdx;
+	}
+
 public:
 	// FPrimitiveSceneProxy interface.
 #if WITH_EDITOR
diff --git a/Engine/Source/Runtime/Engine/Public/StaticMeshVertexData.h b/Engine/Source/Runtime/Engine/Public/StaticMeshVertexData.h
--- a/Engine/Source/Runtime/Engine/Public/StaticMeshVertexData.h
+++ b/Engine/Source/Runtime/Engine/Public/StaticMeshVertexData.h
@@ -82,6 +82,12 @@ public:
 	{
 		return &Data;
 	}
+
+	const FResourceArrayInterface* GetResourceArray() const
+	{
+		return &Data;
+	}
+
 	/**
 	* Serializer for this class
 	*
diff --git a/Engine/Source/Runtime/Engine/Public/TextureResource.h b/Engine/Source/Runtime/Engine/Public/TextureResource.h
--- a/Engine/Source/Runtime/Engine/Public/TextureResource.h
+++ b/Engine/Source/Runtime/Engine/Public/TextureResource.h
@@ -144,7 +144,7 @@ private:
 	friend class FTexture2DUpdate;
 
 	/** The UTexture2D which this resource represents.														*/
-	const UTexture2D*	Owner;
+	UTexture2D*	Owner;
 	/** Resource memory allocated by the owner for serialize bulk mip data into								*/
 	FTexture2DResourceMem* ResourceMem;
 
