Try to fix static analysis warning

diff --git a/Engine/Source/Runtime/Engine/Private/ContentStreaming.cpp b/Engine/Source/Runtime/Engine/Private/ContentStreaming.cpp
--- a/Engine/Source/Runtime/Engine/Private/ContentStreaming.cpp
+++ b/Engine/Source/Runtime/Engine/Private/ContentStreaming.cpp
@@ -263,10 +263,12 @@ bool TrackRenderAssetEvent(FStreamingRenderAsset* StreamingRenderAsset, UStreama
 
 					int32 WantedMips		= RenderAsset->GetNumRequestedMips();
 					float BoostFactor		= 1.0f;
+					FStreamingRenderAsset::EAssetType AssetType = FStreamingRenderAsset::AT_Num;
 					if ( StreamingRenderAsset )
 					{
 						WantedMips			= StreamingRenderAsset->WantedMips;
 						BoostFactor			= StreamingRenderAsset->BoostFactor;
+						AssetType			= StreamingRenderAsset->RenderAssetType;
 					}
 
 					if ( LastEvent->NumResidentMips != RenderAsset->GetNumResidentMips() ||
@@ -284,7 +286,7 @@ bool TrackRenderAssetEvent(FStreamingRenderAsset* StreamingRenderAsset, UStreama
 						NewEvent.Timestamp				= LastEvent->Timestamp			= float(FPlatformTime::Seconds() - GStartTime);
 						NewEvent.BoostFactor			= LastEvent->BoostFactor		= BoostFactor;
 						UE_LOG(LogContentStreaming, Log, TEXT("%s: \"%s\", ResidentMips: %d/%d, RequestedMips: %d, WantedMips: %d, Boost: %.1f (%s)"),
-							FStreamingRenderAsset::GetStreamingAssetTypeStr(StreamingRenderAsset->RenderAssetType),
+							FStreamingRenderAsset::GetStreamingAssetTypeStr(AssetType),
 							AssetName, LastEvent->NumResidentMips, RenderAsset->GetNumMipsForStreaming(), bIsDestroying ? 0 : LastEvent->NumRequestedMips, LastEvent->WantedMips, 
 							BoostFactor, bIsDestroying ? TEXT("DESTROYED") : TEXT("updated") );
 					}
diff --git a/Engine/Source/Runtime/RHI/Public/RHICommandList.h b/Engine/Source/Runtime/RHI/Public/RHICommandList.h
--- a/Engine/Source/Runtime/RHI/Public/RHICommandList.h
+++ b/Engine/Source/Runtime/RHI/Public/RHICommandList.h
@@ -4890,6 +4890,8 @@ struct TRHIResourceUpdateBatcher
 
 	void QueueUpdateRequest(FShaderResourceViewRHIParamRef SRV, FVertexBufferRHIParamRef VertexBuffer, uint32 Stride, uint8 Format)
 	{
+#pragma warning( push )
+#pragma warning( disable : 6386 ) // Disable incorrect compile warning: C6386 buffer overrun while writing to 'this->UpdateInfos'
 		check(NumBatched <= MaxNumUpdates);
 		if (NumBatched == MaxNumUpdates)
 		{
@@ -4899,6 +4901,7 @@ struct TRHIResourceUpdateBatcher
 		UpdateInfos[Idx].Type = FRHIResourceUpdateInfo::UT_VertexBufferSRV;
 		UpdateInfos[Idx].VertexBufferSRV = { SRV, VertexBuffer, Stride, Format };
 	}
+#pragma warning( pop )
 };
 
 #undef RHICOMMAND_CALLSTACK
