Rename bSpawnPawn to accurate bSpawnPlayerController for UGameplayStatics::CreatePlayer.
Also clarify parameter name for UGameInstance::CreateLocalPlayer

diff --git a/Engine/Config/BaseEngine.ini b/Engine/Config/BaseEngine.ini
--- a/Engine/Config/BaseEngine.ini
+++ b/Engine/Config/BaseEngine.ini
@@ -635,6 +635,8 @@ bUseStreamingPause=false
 +FunctionRedirects=(OldName="UserWidget.PlayAnimationAtTime",NewName="UserWidget.PlayAnimation")
 +FunctionRedirects=(OldName="AddChildWrapBox", NewName="AddChildToWrapBox")
 
++PropertyRedirects=(OldName="GameplayStatics.CreatePlayer.bSpawnPawn",NewName="bSpawnPlayerController")
+
 ; Chaos
 +ClassRedirects=(OldName="/Script/GeometryCollectionCore.GeometryCollection",NewName="/Script/GeometryCollectionEngine.GeometryCollection")
 +ClassRedirects=(OldName="/Script/GeometryCollectionCore.GeometryCollectionCache",NewName="/Script/GeometryCollectionEngine.GeometryCollectionCache")
diff --git a/Engine/Source/Runtime/Engine/Classes/Engine/GameInstance.h b/Engine/Source/Runtime/Engine/Classes/Engine/GameInstance.h
--- a/Engine/Source/Runtime/Engine/Classes/Engine/GameInstance.h
+++ b/Engine/Source/Runtime/Engine/Classes/Engine/GameInstance.h
@@ -249,10 +249,10 @@ public:
 	 * Adds a new player.
 	 * @param ControllerId - The controller ID the player should accept input from.
 	 * @param OutError - If no player is returned, OutError will contain a string describing the reason.
-	 * @param SpawnActor - True if an actor should be spawned for the new player.
+	 * @param bSpawnPlayerController - True if a player controller should be spawned immediately for the new player.
 	 * @return The player which was created.
 	 */
-	ULocalPlayer*			CreateLocalPlayer(int32 ControllerId, FString& OutError, bool bSpawnActor);
+	ULocalPlayer*			CreateLocalPlayer(int32 ControllerId, FString& OutError, bool bSpawnPlayerController);
 
 	/**
 	 * Adds a LocalPlayer to the local and global list of Players.
diff --git a/Engine/Source/Runtime/Engine/Classes/Kismet/GameplayStatics.h b/Engine/Source/Runtime/Engine/Classes/Kismet/GameplayStatics.h
--- a/Engine/Source/Runtime/Engine/Classes/Kismet/GameplayStatics.h
+++ b/Engine/Source/Runtime/Engine/Classes/Kismet/GameplayStatics.h
@@ -123,11 +123,12 @@ class ENGINE_API UGameplayStatics : public UBlueprintFunctionLibrary
 	static class APlayerCameraManager* GetPlayerCameraManager(const UObject* WorldContextObject, int32 PlayerIndex);
 
 	/** Create a new player for this game.  
-	 *  @param ControllerId		The ID of the controller that the should control the newly created player.  A value of -1 specifies to use the next available ID
-	 *  @param bSpawnPawn		Whether a pawn should be spawned immediately. If false a pawn will not be created until transition to the next map.
+	 *  @param ControllerId     The ID of the controller that the should control the newly created player.  A value of -1 specifies to use the next available ID
+	 *  @param bSpawnPawn       Whether a player controller should be spawned immediately for this player. If false a player controller will not be created automatically until transition to the next map.
+	 *  @return                 The created player controller if one is created. 
 	 */
 	UFUNCTION(BlueprintCallable, Category="Game", meta=(WorldContext="WorldContextObject", AdvancedDisplay="2", UnsafeDuringActorConstruction="true"))
-	static class APlayerController* CreatePlayer(const UObject* WorldContextObject, int32 ControllerId = -1, bool bSpawnPawn = true);
+	static class APlayerController* CreatePlayer(const UObject* WorldContextObject, int32 ControllerId = -1, bool bSpawnPlayerController = true);
 
 	/** Removes a player from this game.  
 	 *  @param Player			The player controller of the player to be removed
diff --git a/Engine/Source/Runtime/Engine/Private/GameInstance.cpp b/Engine/Source/Runtime/Engine/Private/GameInstance.cpp
--- a/Engine/Source/Runtime/Engine/Private/GameInstance.cpp
+++ b/Engine/Source/Runtime/Engine/Private/GameInstance.cpp
@@ -607,7 +607,7 @@ ULocalPlayer* UGameInstance::CreateInitialPlayer(FString& OutError)
 	return CreateLocalPlayer( 0, OutError, false );
 }
 
-ULocalPlayer* UGameInstance::CreateLocalPlayer(int32 ControllerId, FString& OutError, bool bSpawnActor)
+ULocalPlayer* UGameInstance::CreateLocalPlayer(int32 ControllerId, FString& OutError, bool bSpawnPlayerController)
 {
 	check(GetEngine()->LocalPlayerClass != NULL);
 
@@ -641,7 +641,7 @@ ULocalPlayer* UGameInstance::CreateLocalPlayer(int32 ControllerId, FString& OutE
 
 		NewPlayer = NewObject<ULocalPlayer>(GetEngine(), GetEngine()->LocalPlayerClass);
 		InsertIndex = AddLocalPlayer(NewPlayer, ControllerId);
-		if (bSpawnActor && InsertIndex != INDEX_NONE && GetWorld() != NULL)
+		if (bSpawnPlayerController && InsertIndex != INDEX_NONE && GetWorld() != NULL)
 		{
 			if (GetWorld()->GetNetMode() != NM_Client)
 			{
diff --git a/Engine/Source/Runtime/Engine/Private/GameplayStatics.cpp b/Engine/Source/Runtime/Engine/Private/GameplayStatics.cpp
--- a/Engine/Source/Runtime/Engine/Private/GameplayStatics.cpp
+++ b/Engine/Source/Runtime/Engine/Private/GameplayStatics.cpp
@@ -235,12 +235,12 @@ APlayerCameraManager* UGameplayStatics::GetPlayerCameraManager(const UObject* Wo
 	return PC ? PC->PlayerCameraManager : nullptr;
 }
 
-APlayerController* UGameplayStatics::CreatePlayer(const UObject* WorldContextObject, int32 ControllerId, bool bSpawnPawn)
+APlayerController* UGameplayStatics::CreatePlayer(const UObject* WorldContextObject, int32 ControllerId, bool bSpawnPlayerController)
 {
 	UWorld* World = GEngine->GetWorldFromContextObject(WorldContextObject, EGetWorldErrorMode::LogAndReturnNull);
 	FString Error;
 
-	ULocalPlayer* LocalPlayer = World ? World->GetGameInstance()->CreateLocalPlayer(ControllerId, Error, bSpawnPawn) : NULL;
+	ULocalPlayer* LocalPlayer = World ? World->GetGameInstance()->CreateLocalPlayer(ControllerId, Error, bSpawnPlayerController) : nullptr;
 
 	if (Error.Len() > 0)
 	{
