Remove deprecated PostPhysicsComponentTick from UPrimitiveComponent (saves 88 bytes)

diff --git a/Engine/Source/Runtime/Engine/Classes/Components/PrimitiveComponent.h b/Engine/Source/Runtime/Engine/Classes/Components/PrimitiveComponent.h
--- a/Engine/Source/Runtime/Engine/Classes/Components/PrimitiveComponent.h
+++ b/Engine/Source/Runtime/Engine/Classes/Components/PrimitiveComponent.h
@@ -1561,22 +1561,6 @@ public:
 	 */
 	virtual void GetUsedTextures(TArray<UTexture*>& OutTextures, EMaterialQualityLevel::Type QualityLevel);
 
-	/** Tick function for physics ticking **/
-	UPROPERTY()
-	struct FPrimitiveComponentPostPhysicsTickFunction PostPhysicsComponentTick;
-
-	/** Controls if we get a post physics tick or not. If set during ticking, will take effect next frame **/
-	UE_DEPRECATED(4.11, "Please register your own tick function or use the primary tick function")
-	void SetPostPhysicsComponentTickEnabled(bool bEnable);
-
-	/** Returns whether we have the post physics tick enabled **/
-	UE_DEPRECATED(4.11, "Please register your own tick function or use the primary tick function")
-	bool IsPostPhysicsComponentTickEnabled() const;
-
-	/** Tick function called after physics (sync scene) has finished simulation */
-	UE_DEPRECATED(4.11, "Please register your own tick function or use the primary tick function")
-	virtual void PostPhysicsTick(FPrimitiveComponentPostPhysicsTickFunction &ThisTickFunction) {}
-
 	/** Return the BodySetup to use for this PrimitiveComponent (single body case) */
 	virtual class UBodySetup* GetBodySetup() { return NULL; }
 
@@ -1770,7 +1754,6 @@ public:
 	virtual void OnCreatePhysicsState() override;
 	virtual void OnDestroyPhysicsState() override;
 	virtual void OnActorEnableCollisionChanged() override;
-	virtual void RegisterComponentTickFunctions(bool bRegister) override;
 	virtual void InvalidateLightingCacheDetailed(bool bInvalidateBuildEnqueuedLighting, bool bTranslationOnly) override;
 	virtual bool IsEditorOnly() const override;
 	virtual bool ShouldCreatePhysicsState() const override;
diff --git a/Engine/Source/Runtime/Engine/Classes/Engine/EngineBaseTypes.h b/Engine/Source/Runtime/Engine/Classes/Engine/EngineBaseTypes.h
--- a/Engine/Source/Runtime/Engine/Classes/Engine/EngineBaseTypes.h
+++ b/Engine/Source/Runtime/Engine/Classes/Engine/EngineBaseTypes.h
@@ -501,39 +501,6 @@ struct TStructOpsTypeTraits<FActorComponentTickFunction> : public TStructOpsType
 		WithCopy = false
 	};
 };
-/** 
-* Tick function that calls UPrimitiveComponent::PostPhysicsTick
-**/
-
-//DEPRECATED: This struct has been deprecated. Please use your own tick functions if you need something other than the primary tick function
-USTRUCT()
-struct FPrimitiveComponentPostPhysicsTickFunction : public FTickFunction
-{
-	GENERATED_USTRUCT_BODY()
-
-	/** PrimitiveComponent component that is the target of this tick **/
-	class UPrimitiveComponent*	Target;
-
-	/** 
-		* Abstract function actually execute the tick. 
-		* @param DeltaTime - frame time to advance, in seconds
-		* @param TickType - kind of tick for this frame
-		* @param CurrentThread - thread we are executing on, useful to pass along as new tasks are created
-		* @param MyCompletionGraphEvent - completion event for this task. Useful for holding the completetion of this task until certain child tasks are complete.
-	**/
-	virtual void ExecuteTick(float DeltaTime, ELevelTick TickType, ENamedThreads::Type CurrentThread, const FGraphEventRef& MyCompletionGraphEvent) override;
-	/** Abstract function to describe this tick. Used to print messages about illegal cycles in the dependency graph **/
-	virtual FString DiagnosticMessage() override;
-};
-
-template<>
-struct TStructOpsTypeTraits<FPrimitiveComponentPostPhysicsTickFunction> : public TStructOpsTypeTraitsBase2<FPrimitiveComponentPostPhysicsTickFunction>
-{
-	enum
-	{
-		WithCopy = false
-	};
-};
 
 /** Types of network failures broadcast from the engine */
 UENUM(BlueprintType)
diff --git a/Engine/Source/Runtime/Engine/Private/Components/PrimitiveComponent.cpp b/Engine/Source/Runtime/Engine/Private/Components/PrimitiveComponent.cpp
--- a/Engine/Source/Runtime/Engine/Private/Components/PrimitiveComponent.cpp
+++ b/Engine/Source/Runtime/Engine/Private/Components/PrimitiveComponent.cpp
@@ -225,10 +225,6 @@ FThreadSafeCounter UPrimitiveComponent::NextComponentId;
 UPrimitiveComponent::UPrimitiveComponent(const FObjectInitializer& ObjectInitializer /*= FObjectInitializer::Get()*/)
 	: Super(ObjectInitializer)
 {
-	PostPhysicsComponentTick.bCanEverTick = false;
-	PostPhysicsComponentTick.bStartWithTickEnabled = true;
-	PostPhysicsComponentTick.TickGroup = TG_PostPhysics;
-
 	LastRenderTime = -1000.0f;
 	LastRenderTimeOnScreen = -1000.0f;
 	BoundsScale = 1.0f;
@@ -429,74 +425,6 @@ void UPrimitiveComponent::GetUsedTextures(TArray<UTexture*>& OutTextures, EMater
 	}
 }
 
-/** 
-* Abstract function actually execute the tick. 
-* @param DeltaTime - frame time to advance, in seconds
-* @param TickType - kind of tick for this frame
-* @param CurrentThread - thread we are executing on, useful to pass along as new tasks are created
-* @param MyCompletionGraphEvent - completion event for this task. Useful for holding the completetion of this task until certain child tasks are complete.
-**/
-void FPrimitiveComponentPostPhysicsTickFunction::ExecuteTick(float DeltaTime, enum ELevelTick TickType, ENamedThreads::Type CurrentThread, const FGraphEventRef& MyCompletionGraphEvent)
-{
-	PRAGMA_DISABLE_DEPRECATION_WARNINGS
-	FActorComponentTickFunction::ExecuteTickHelper(Target, /*bTickInEditor=*/ false, DeltaTime, TickType, [this](float DilatedTime){ Target->PostPhysicsTick(*this); });
-	PRAGMA_ENABLE_DEPRECATION_WARNINGS
-}
-
-/** Abstract function to describe this tick. Used to print messages about illegal cycles in the dependency graph **/
-FString FPrimitiveComponentPostPhysicsTickFunction::DiagnosticMessage()
-{
-	return Target->GetFullName() + TEXT("[UPrimitiveComponent::PostPhysicsTick]");
-}
-
-void UPrimitiveComponent::RegisterComponentTickFunctions(bool bRegister)
-{
-	Super::RegisterComponentTickFunctions(bRegister);
-
-	if (bRegister)
-	{
-		if (SetupActorComponentTickFunction(&PostPhysicsComponentTick))
-		{
-			PostPhysicsComponentTick.Target = this;
-
-			// If primary tick is registered, add a prerequisate to it
-			if(PrimaryComponentTick.bCanEverTick)
-			{
-				PostPhysicsComponentTick.AddPrerequisite(this,PrimaryComponentTick); 
-			}
-
-			// Set a prereq for the post physics tick to happen after physics is finished
-			UWorld* World = GetWorld();
-			if (World != nullptr)
-			{
-				PostPhysicsComponentTick.AddPrerequisite(World, World->EndPhysicsTickFunction);
-			}
-		}
-	}
-	else
-	{
-		if(PostPhysicsComponentTick.IsTickFunctionRegistered())
-		{
-			PostPhysicsComponentTick.UnRegisterTickFunction();
-		}
-	}
-}
-
-void UPrimitiveComponent::SetPostPhysicsComponentTickEnabled(bool bEnable)
-{
-	// @todo ticking, James, turn this off when not needed
-	if (PostPhysicsComponentTick.bCanEverTick && !IsTemplate())
-	{
-		PostPhysicsComponentTick.SetTickFunctionEnable(bEnable);
-	}
-}
-
-bool UPrimitiveComponent::IsPostPhysicsComponentTickEnabled() const
-{
-	return PostPhysicsComponentTick.IsTickFunctionEnabled();
-}
-
-
 //////////////////////////////////////////////////////////////////////////
 // Render
 
