Fix a bug where some static mesh LOD group settings are not overridden properly

diff --git a/Engine/Source/Runtime/Engine/Private/StaticMesh.cpp b/Engine/Source/Runtime/Engine/Private/StaticMesh.cpp
--- a/Engine/Source/Runtime/Engine/Private/StaticMesh.cpp
+++ b/Engine/Source/Runtime/Engine/Private/StaticMesh.cpp
@@ -1568,10 +1568,10 @@ void FStaticMeshLODSettings::ReadEntry(FStaticMeshLODGroup& Group, FString Entry
 	{
 		Group.DefaultMaxNumStreamedLODs = FMath::Max(Group.DefaultMaxNumStreamedLODs, 0);
 	}
-
+	
+	int32 LocalSupportLODStreaming = 0;
+	if (FParse::Value(*Entry, TEXT("bSupportLODStreaming="), LocalSupportLODStreaming))
 	{
-		int32 LocalSupportLODStreaming = 0;
-		FParse::Value(*Entry, TEXT("bSupportLODStreaming="), LocalSupportLODStreaming);
 		Group.bSupportLODStreaming = !!LocalSupportLODStreaming;
 	}
 
@@ -1585,15 +1585,15 @@ void FStaticMeshLODSettings::ReadEntry(FStaticMeshLODGroup& Group, FString Entry
 	if (FParse::Value(*Entry, TEXT("BasePercentTriangles="), BasePercentTriangles))
 	{
 		BasePercentTriangles = FMath::Clamp<float>(BasePercentTriangles, 0.0f, 100.0f);
+		Group.DefaultSettings[0].PercentTriangles = BasePercentTriangles * 0.01f;
 	}
-	Group.DefaultSettings[0].PercentTriangles = BasePercentTriangles * 0.01f;
 
 	float LODPercentTriangles = 100.0f;
 	if (FParse::Value(*Entry, TEXT("LODPercentTriangles="), LODPercentTriangles))
 	{
 		LODPercentTriangles = FMath::Clamp<float>(LODPercentTriangles, 0.0f, 100.0f);
+		Group.DefaultSettings[1].PercentTriangles = LODPercentTriangles * 0.01f;
 	}
-	Group.DefaultSettings[1].PercentTriangles = LODPercentTriangles * 0.01f;
 
 	if (FParse::Value(*Entry, TEXT("MaxDeviation="), Settings.MaxDeviation))
 	{
@@ -1634,15 +1634,15 @@ void FStaticMeshLODSettings::ReadEntry(FStaticMeshLODGroup& Group, FString Entry
 	if (FParse::Value(*Entry, TEXT("BasePercentTrianglesMult="), BasePercentTrianglesMult))
 	{
 		BasePercentTrianglesMult = FMath::Clamp<float>(BasePercentTrianglesMult, 0.0f, 100.0f);
+		Group.BasePercentTrianglesMult = BasePercentTrianglesMult * 0.01f;
 	}
-	Group.BasePercentTrianglesMult = BasePercentTrianglesMult * 0.01f;
 
 	float LODPercentTrianglesMult = 100.0f;
 	if (FParse::Value(*Entry, TEXT("LODPercentTrianglesMult="), LODPercentTrianglesMult))
 	{
 		LODPercentTrianglesMult = FMath::Clamp<float>(LODPercentTrianglesMult, 0.0f, 100.0f);
+		Bias.PercentTriangles = LODPercentTrianglesMult * 0.01f;
 	}
-	Bias.PercentTriangles = LODPercentTrianglesMult * 0.01f;
 
 	if (FParse::Value(*Entry, TEXT("MaxDeviationBias="), Bias.MaxDeviation))
 	{
