Deprecating cdkey strings

diff --git a/Engine/Source/Runtime/Engine/Classes/Engine/NetConnection.h b/Engine/Source/Runtime/Engine/Classes/Engine/NetConnection.h
--- a/Engine/Source/Runtime/Engine/Classes/Engine/NetConnection.h
+++ b/Engine/Source/Runtime/Engine/Classes/Engine/NetConnection.h
@@ -342,7 +342,9 @@ public:
 	uint8					ExpectedClientLoginMsgType;	// Used to determine what the next expected control channel msg type should be from a connecting client
 
 	// CD key authentication
+	UE_DEPRECATED(4.23, "CDKeyHash is deprecated.")
 	FString			CDKeyHash;				// Hash of client's CD key
+	UE_DEPRECATED(4.23, "CDKeyResponse is deprecated.")
 	FString			CDKeyResponse;			// Client's response to CD key challenge
 
 	// Internal.
diff --git a/Engine/Source/Runtime/Engine/Private/NetConnection.cpp b/Engine/Source/Runtime/Engine/Private/NetConnection.cpp
--- a/Engine/Source/Runtime/Engine/Private/NetConnection.cpp
+++ b/Engine/Source/Runtime/Engine/Private/NetConnection.cpp
@@ -509,8 +509,11 @@ void UNetConnection::Serialize( FArchive& Ar )
 		GRANULAR_NETWORK_MEMORY_TRACKING_TRACK("Challenge", Challenge.CountBytes(Ar));
 		GRANULAR_NETWORK_MEMORY_TRACKING_TRACK("ClientResponse", ClientResponse.CountBytes(Ar));
 		GRANULAR_NETWORK_MEMORY_TRACKING_TRACK("RequestURL", RequestURL.CountBytes(Ar));
+
+PRAGMA_DISABLE_DEPRECATION_WARNINGS
 		GRANULAR_NETWORK_MEMORY_TRACKING_TRACK("CDKeyHash", CDKeyHash.CountBytes(Ar));
 		GRANULAR_NETWORK_MEMORY_TRACKING_TRACK("CDKeyResponse", CDKeyResponse.CountBytes(Ar));
+PRAGMA_ENABLE_DEPRECATION_WARNINGS
 
 		GRANULAR_NETWORK_MEMORY_TRACKING_TRACK("SendBuffer", SendBuffer.CountMemory(Ar));
 
