Adding several features to baked analysis feature.

- New getter functions on audio component which retrieve data directly with no frequency interpolation
- Ability to log baked fft/envelope data to csv files
- Ability to override sound wave to use for analysis
- FFT attack/release envelope following

diff --git a/Engine/Source/Runtime/Engine/Classes/Components/AudioComponent.h b/Engine/Source/Runtime/Engine/Classes/Components/AudioComponent.h
--- a/Engine/Source/Runtime/Engine/Classes/Components/AudioComponent.h
+++ b/Engine/Source/Runtime/Engine/Classes/Components/AudioComponent.h
@@ -441,14 +441,38 @@ public:
 	UFUNCTION(BlueprintCallable, Category = "Audio|Components|Audio")
 	bool HasCookedAmplitudeEnvelopeData() const;
 
-	/** Returns the current cooked spectral data of the playing audio component. Returns true if there is data and the audio component is playing. */
+	/** 
+	* Retrieves the current-time cooked spectral data of the sounds playing on the audio component. 
+	* Spectral data is averaged and interpolated for all playing sounds on this audio component. 
+	* Returns true if there is data and the audio component is playing. 
+	*/
 	UFUNCTION(BlueprintCallable, Category = "Audio|Components|Audio")
 	bool GetCookedFFTData(const TArray<float>& FrequenciesToGet, TArray<FSoundWaveSpectralData>& OutSoundWaveSpectralData);
 
-	/** Returns the current cooked envelope data of the playing audio component. Returns true if there is data and the audio component is playing. */
+	/** 
+	* Retrieves the current-time cooked spectral data of the sounds playing audio component. 
+	* Spectral data is not averaged or interpolated. Instead an array of data with all playing sound waves with cooked data is returned.
+	* Returns true if there is data and the audio component is playing. 
+	*/
+	UFUNCTION(BlueprintCallable, Category = "Audio|Components|Audio")
+	bool GetCookedFFTDataForAllPlayingSounds(TArray<FSoundWaveSpectralDataPerSound>& OutSoundWaveSpectralData);
+
+	/** 
+	* Retrieves the current-time cooked envelope data of the playing audio component. 
+	* Cooked data is interpolated and averaged across all playing sound waves.
+	* Returns true if there is data and the audio component is playing. 
+	*/
 	UFUNCTION(BlueprintCallable, Category = "Audio|Components|Audio")
 	bool GetCookedEnvelopeData(float& OutEnvelopeData);
 
+	/**
+	* Retrieves the current-time envelope data of the sounds playing audio component.
+	* Envelope data is not averaged or interpolated. Instead an array of data with all playing sound waves with cooked data is returned.
+	* Returns true if there is data and the audio component is playing.
+	*/
+	UFUNCTION(BlueprintCallable, Category = "Audio|Components|Audio")
+	bool GetCookedEnvelopeDataForAllPlayingSounds(TArray<FSoundWaveEnvelopeDataPerSound>& OutEnvelopeData);
+
 	static void PlaybackCompleted(uint64 AudioComponentID, bool bFailedToStart);
 
 private:
diff --git a/Engine/Source/Runtime/Engine/Classes/Sound/SoundWave.h b/Engine/Source/Runtime/Engine/Classes/Sound/SoundWave.h
--- a/Engine/Source/Runtime/Engine/Classes/Sound/SoundWave.h
+++ b/Engine/Source/Runtime/Engine/Classes/Sound/SoundWave.h
@@ -129,22 +129,21 @@ struct FStreamedAudioPlatformData
 
 };
 
-// Struct used to retrieve the spectral magnitude of an audio signal at a given frequency to BP
 USTRUCT(BlueprintType)
 struct FSoundWaveSpectralData
 {
 	GENERATED_USTRUCT_BODY()
 
 	// The frequency hz of the spectrum value
-	UPROPERTY(EditAnywhere, BlueprintReadWrite, Category = "SpectralData|Foo")
+	UPROPERTY(EditAnywhere, BlueprintReadWrite, Category = "SpectralData")
 	float FrequencyHz;
 
 	// The magnitude of the spectrum at this frequency
-	UPROPERTY(EditAnywhere, BlueprintReadWrite, Category = "SpectralData|Bar")
+	UPROPERTY(EditAnywhere, BlueprintReadWrite, Category = "SpectralData")
 	float Magnitude;
 
 	// The normalized magnitude of the spectrum at this frequency
-	UPROPERTY(EditAnywhere, BlueprintReadWrite, Category = "SpectralData|Bar")
+	UPROPERTY(EditAnywhere, BlueprintReadWrite, Category = "SpectralData")
 	float NormalizedMagnitude;
 
 	FSoundWaveSpectralData()
@@ -154,6 +153,42 @@ struct FSoundWaveSpectralData
 	{}
 };
 
+USTRUCT(BlueprintType)
+struct FSoundWaveSpectralDataPerSound
+{
+	GENERATED_USTRUCT_BODY()
+
+	// The array of current spectral data for this sound wave
+	UPROPERTY(EditAnywhere, BlueprintReadWrite, Category = "SpectralData")
+	TArray<FSoundWaveSpectralData> SpectralData;
+
+	// The current playback time of this sound wave
+	UPROPERTY(EditAnywhere, BlueprintReadWrite, Category = "SpectralData")
+	float PlaybackTime;
+
+	// The sound wave this spectral data is associated with
+	UPROPERTY(EditAnywhere, BlueprintReadWrite, Category = "SpectralData")
+	USoundWave* SoundWave;
+};
+
+USTRUCT(BlueprintType)
+struct FSoundWaveEnvelopeDataPerSound
+{
+	GENERATED_USTRUCT_BODY()
+
+	// The current envelope of the playing sound
+	UPROPERTY(EditAnywhere, BlueprintReadWrite, Category = "EnvelopeData")
+	float Envelope;
+
+	// The current playback time of this sound wave
+	UPROPERTY(EditAnywhere, BlueprintReadWrite, Category = "EnvelopeData")
+	float PlaybackTime;
+
+	// The sound wave this envelope data is associated with
+	UPROPERTY(EditAnywhere, BlueprintReadWrite, Category = "EnvelopeData")
+	USoundWave* SoundWave;
+};
+
 // Sort predicate for sorting spectral data by frequency (lowest first)
 struct FCompareSpectralDataByFrequencyHz
 {
@@ -301,6 +336,17 @@ class ENGINE_API USoundWave : public USoundBase
 	uint8 bDecompressedFromOgg : 1;
 
 #if WITH_EDITORONLY_DATA
+	/** Specify a sound to use for the baked analysis. Will default to this USoundWave if not sete. */
+	UPROPERTY(EditAnywhere, Category = "Analysis")
+	USoundWave* OverrideSoundToUseForAnalysis;
+
+	/** 
+		Whether or not we should treat the sound wave used for analysis (this or the override) as looping while performing analysis. 
+		A looping sound may include the end of the file for inclusion in analysis for envelope and FFT analysis.
+	*/
+	UPROPERTY(EditAnywhere, Category = "Analysis")
+	uint8 TreatFileAsLoopingForAnalysis:1;
+
 	/** Whether or not to enable cook-time baked FFT analysis. */
 	UPROPERTY(EditAnywhere, Category = "Analysis|FFT")
 	uint8 bEnableBakedFFTAnalysis : 1;
@@ -314,18 +360,26 @@ class ENGINE_API USoundWave : public USoundBase
 	ESoundWaveFFTSize FFTSize;
 
 	/** How many audio frames analyze at a time. */
-	UPROPERTY(EditAnywhere, Category = "Analysis|FFT", meta = (EditCondition = "bEnableAmplitudeEnvelopeAnalysis", ClampMin = "512", UIMin = "512"))
+	UPROPERTY(EditAnywhere, Category = "Analysis|FFT", meta = (EditCondition = "bEnableBakedFFTAnalysis", ClampMin = "512", UIMin = "512"))
 	int32 FFTAnalysisFrameSize;
 
+	/** Attack time in milliseconds of the spectral envelope follower. */
+	UPROPERTY(EditAnywhere, Category = "Analysis|FFT", meta = (EditCondition = "bEnableBakedFFTAnalysis", ClampMin = "0", UIMin = "0"))
+	int32 FFTAnalysisAttackTime;
+
+	/** Release time in milliseconds of the spectral envelope follower. */
+	UPROPERTY(EditAnywhere, Category = "Analysis|FFT", meta = (EditCondition = "bEnableBakedFFTAnalysis", ClampMin = "0", UIMin = "0"))
+	int32 FFTAnalysisReleaseTime;
+
 	/** How many audio frames to average a new envelope value. Larger values use less memory for audio envelope data but will result in lower envelope accuracy. */
 	UPROPERTY(EditAnywhere, Category = "Analysis|Envelope", meta = (EditCondition = "bEnableAmplitudeEnvelopeAnalysis", ClampMin = "512", UIMin = "512"))
 	int32 EnvelopeFollowerFrameSize;
 
-	/** How quickly the envelope analyzer responds to increasing amplitudes. */
+	/** The attack time in milliseconds. Describes how quickly the envelope analyzer responds to increasing amplitudes. */
 	UPROPERTY(EditAnywhere, Category = "Analysis|Envelope", meta = (EditCondition = "bEnableAmplitudeEnvelopeAnalysis", ClampMin = "0", UIMin = "0"))
 	int32 EnvelopeFollowerAttackTime;
 
-	/** How quickly the envelope analyzer responds to decreasing amplitudes. */
+	/** The release time in milliseconds. Describes how quickly the envelope analyzer responds to decreasing amplitudes. */
 	UPROPERTY(EditAnywhere, Category = "Analysis|Envelope", meta = (EditCondition = "bEnableAmplitudeEnvelopeAnalysis", ClampMin = "0", UIMin = "0"))
 	int32 EnvelopeFollowerReleaseTime;
 #endif
@@ -648,6 +702,11 @@ private:
 #endif
 
 public:
+
+#if WITH_EDITOR
+	void LogBakedData();
+#endif
+
 	virtual void BeginGetCompressedData(FName Format, const FPlatformAudioCookOverrides* CompressionOverrides);
 
 	/** 
diff --git a/Engine/Source/Runtime/Engine/Private/Components/AudioComponent.cpp b/Engine/Source/Runtime/Engine/Private/Components/AudioComponent.cpp
--- a/Engine/Source/Runtime/Engine/Private/Components/AudioComponent.cpp
+++ b/Engine/Source/Runtime/Engine/Private/Components/AudioComponent.cpp
@@ -1265,6 +1265,34 @@ bool UAudioComponent::GetCookedFFTData(const TArray<float>& FrequenciesToGet, TA
 	return bHadData;
 }
 
+bool UAudioComponent::GetCookedFFTDataForAllPlayingSounds(TArray<FSoundWaveSpectralDataPerSound>& OutSoundWaveSpectralData)
+{
+	bool bHadData = false;
+	if (IsPlaying() && SoundWavePlaybackTimes.Num() > 0)
+	{
+		OutSoundWaveSpectralData.Reset();
+
+		for (auto& Entry : SoundWavePlaybackTimes)
+		{
+			if (Entry.Value.PlaybackTime > 0.0f && Entry.Value.SoundWave->CookedSpectralTimeData.Num() > 0)
+			{
+				FSoundWaveSpectralDataPerSound NewOutput;
+				NewOutput.SoundWave = Entry.Value.SoundWave;
+				NewOutput.PlaybackTime = Entry.Value.PlaybackTime;
+
+				// Find the point in the spectral data that corresponds to the time
+				Entry.Value.SoundWave->GetInterpolatedCookedFFTDataForTime(Entry.Value.PlaybackTime, Entry.Value.LastFFTCookedIndex, NewOutput.SpectralData, Sound->IsLooping());
+				if (NewOutput.SpectralData.Num())
+				{
+					OutSoundWaveSpectralData.Add(NewOutput);
+				}
+				bHadData = true;
+			}
+		}
+	}
+	return bHadData;
+}
+
 bool UAudioComponent::GetCookedEnvelopeData(float& OutEnvelopeData)
 {
 	bool bHadData = false;
@@ -1299,3 +1327,30 @@ bool UAudioComponent::GetCookedEnvelopeData(float& OutEnvelopeData)
 
 	return bHadData;
 }
+
+bool UAudioComponent::GetCookedEnvelopeDataForAllPlayingSounds(TArray<FSoundWaveEnvelopeDataPerSound>& OutEnvelopeData)
+{
+	bool bHadData = false;
+	if (IsPlaying() && SoundWavePlaybackTimes.Num() > 0)
+	{
+		for (auto& Entry : SoundWavePlaybackTimes)
+		{
+			if (Entry.Value.SoundWave->CookedEnvelopeTimeData.Num() > 0 && Entry.Value.PlaybackTime > 0.0f)
+			{
+				// Find the point in the spectral data that corresponds to the time
+				float SoundWaveAmplitude = 0.0f;
+				if (Entry.Value.SoundWave->GetInterpolatedCookedEnvelopeDataForTime(Entry.Value.PlaybackTime, Entry.Value.LastEnvelopeCookedIndex, SoundWaveAmplitude, Sound->IsLooping()))
+				{
+					FSoundWaveEnvelopeDataPerSound NewOutput;
+					NewOutput.SoundWave = Entry.Value.SoundWave;
+					NewOutput.PlaybackTime = Entry.Value.PlaybackTime;
+					NewOutput.Envelope = SoundWaveAmplitude;
+					OutEnvelopeData.Add(NewOutput);
+					bHadData = true;
+				}
+
+			}
+		}
+	}
+	return bHadData;
+}
diff --git a/Engine/Source/Runtime/Engine/Private/SoundWave.cpp b/Engine/Source/Runtime/Engine/Private/SoundWave.cpp
--- a/Engine/Source/Runtime/Engine/Private/SoundWave.cpp
+++ b/Engine/Source/Runtime/Engine/Private/SoundWave.cpp
@@ -4,6 +4,7 @@
 #include "Serialization/MemoryWriter.h"
 #include "UObject/FrameworkObjectVersion.h"
 #include "UObject/Package.h"
+#include "UObject/UObjectIterator.h"
 #include "EngineDefines.h"
 #include "Components/AudioComponent.h"
 #include "ContentStreaming.h"
@@ -19,10 +20,14 @@
 #include "EditorFramework/AssetImportData.h"
 #include "ProfilingDebugging/CookStats.h"
 #include "HAL/LowLevelMemTracker.h"
+#include "HAL/IConsoleManager.h"
+#include "HAL/FileManager.h"
 #include "AudioCompressionSettingsUtils.h"
 #include "DSP/SpectrumAnalyzer.h"
 #include "DSP/EnvelopeFollower.h"
 #include "DSP/BufferVectorOperations.h"
+#include "Misc/OutputDeviceArchiveWrapper.h"
+#include "Sound/SampleBuffer.h"
 
 static int32 BypassVirtualizeWhenSilentCVar = 0;
 FAutoConsoleVariableRef CVarBypassVirtualizeWhenSilent(
@@ -32,6 +37,39 @@ FAutoConsoleVariableRef CVarBypassVirtualizeWhenSilent(
 	TEXT("0: Honor the Play When Silent flag, 1: stop all silent non-procedural sources."),
 	ECVF_Default);
 
+#if !UE_BUILD_SHIPPING
+static void DumpBakedAnalysisData(const TArray<FString>& Args)
+{
+	if (IsInGameThread())
+	{
+		if (Args.Num() == 1)
+		{
+			const FString& SoundWaveToDump = Args[0];
+			UE_LOG(LogTemp, Log, TEXT("Foo"));
+			for (TObjectIterator<USoundWave> It; It; ++It)
+			{
+				if (It->IsTemplate(RF_ClassDefaultObject))
+				{
+					continue;
+				}
+
+				if (SoundWaveToDump.Equals(It->GetName()))
+				{
+					UE_LOG(LogTemp, Log, TEXT("Foo"));
+					It->LogBakedData();
+				}
+			}
+		}
+	}
+}
+
+static FAutoConsoleCommand DumpBakedAnalysisDataCmd(
+	TEXT("au.DumpBakedAnalysisData"),
+	TEXT("debug command to dump the baked analysis data of a sound wave to a csv file."),
+	FConsoleCommandWithArgsDelegate::CreateStatic(&DumpBakedAnalysisData)
+);
+#endif
+
 #if ENABLE_COOK_STATS
 namespace SoundWaveCookStats
 {
@@ -129,6 +167,8 @@ USoundWave::USoundWave(const FObjectInitializer& ObjectInitializer)
 	FrequenciesToAnalyze.Add(1000.0f);
 	FrequenciesToAnalyze.Add(5000.0f);
 	FFTAnalysisFrameSize = 1024;
+	FFTAnalysisAttackTime = 10;
+	FFTAnalysisReleaseTime = 3000;
 	EnvelopeFollowerFrameSize = 1024;
 	EnvelopeFollowerAttackTime = 10;
 	EnvelopeFollowerReleaseTime = 100;
@@ -791,28 +831,115 @@ float USoundWave::GetSampleRateForTargetPlatform(const ITargetPlatform* TargetPl
 	}
 }
 
+void USoundWave::LogBakedData()
+{
+	const FString AnalysisPathName = *(FPaths::ProjectLogDir() + TEXT("BakedAudioAnalysisData/"));
+	IFileManager::Get().MakeDirectory(*AnalysisPathName);
+
+	FString SoundWaveName = FString::Printf(TEXT("%s.%s"), *FDateTime::Now().ToString(TEXT("%d-%H.%M.%S")), *GetName());
+
+	if (CookedEnvelopeTimeData.Num())
+	{
+		FString EnvelopeFileName = FString::Printf(TEXT("%s.envelope.csv"), *SoundWaveName);	
+		FString FilenameFull = AnalysisPathName + EnvelopeFileName;
+
+		FArchive* FileAr = IFileManager::Get().CreateDebugFileWriter(*FilenameFull);
+		FOutputDeviceArchiveWrapper* FileArWrapper = new FOutputDeviceArchiveWrapper(FileAr);
+		FOutputDevice* ReportAr = FileArWrapper;
+
+		ReportAr->Log(TEXT("TimeStamp (Sec),Amplitude"));
+
+		for (const FSoundWaveEnvelopeTimeData& EnvTimeData : CookedEnvelopeTimeData)
+		{
+			ReportAr->Logf(TEXT("%.4f,%.4f"), EnvTimeData.TimeSec, EnvTimeData.Amplitude);
+		}
+
+		// Shutdown and free archive resources
+		FileArWrapper->TearDown();
+		delete FileArWrapper;
+		delete FileAr;
+	}
+
+	if (CookedSpectralTimeData.Num())
+	{
+		FString AnalysisFileName = FString::Printf(TEXT("%s.spectral.csv"), *SoundWaveName);
+		FString FilenameFull = AnalysisPathName + AnalysisFileName;
+
+		FArchive* FileAr = IFileManager::Get().CreateDebugFileWriter(*FilenameFull);
+		FOutputDeviceArchiveWrapper* FileArWrapper = new FOutputDeviceArchiveWrapper(FileAr);
+		FOutputDevice* ReportAr = FileArWrapper;
+
+		// Build the header string
+		FString ScratchString;
+		ScratchString.Append(TEXT("Time Stamp (Sec),"));
+
+		for (int32 i = 0; i < FrequenciesToAnalyze.Num(); ++i)
+		{
+			ScratchString.Append(FString::Printf(TEXT("%.2f Hz"), FrequenciesToAnalyze[i]));
+			if (i != FrequenciesToAnalyze.Num() - 1)
+			{
+				ScratchString.Append(TEXT(","));
+			}
+		}
+
+		ReportAr->Log(ScratchString);
+
+		for (const FSoundWaveSpectralTimeData& SpectralTimeData : CookedSpectralTimeData)
+		{
+			ScratchString.Reset();
+			ScratchString.Append(FString::Printf(TEXT("%.4f,"), SpectralTimeData.TimeSec));
+
+			for (int32 i = 0; i < SpectralTimeData.Data.Num(); ++i)
+			{
+				ScratchString.Append(FString::Printf(TEXT("%.4f"), SpectralTimeData.Data[i].Magnitude));
+				if (i != SpectralTimeData.Data.Num() - 1)
+				{
+					ScratchString.Append(TEXT(","));
+				}
+			}
+
+			ReportAr->Log(*ScratchString);
+		}
+
+		// Shutdown and free archive resources
+		FileArWrapper->TearDown();
+		delete FileArWrapper;
+		delete FileAr;
+	}
+
+}
+
 static bool AnyFFTAnalysisPropertiesChanged(const FName& PropertyName)
 {
-	// List of properties which cause re-cooking to get triggered
+	// List of properties which cause analysis to get triggered
+	static FName OverrideSoundName						= GET_MEMBER_NAME_CHECKED(USoundWave, OverrideSoundToUseForAnalysis);
 	static FName EnableFFTAnalysisFName					= GET_MEMBER_NAME_CHECKED(USoundWave, bEnableBakedFFTAnalysis);
 	static FName FFTSizeFName							= GET_MEMBER_NAME_CHECKED(USoundWave, FFTSize);
 	static FName FFTAnalysisFrameSizeFName				= GET_MEMBER_NAME_CHECKED(USoundWave, FFTAnalysisFrameSize);
 	static FName FrequenciesToAnalyzeFName				= GET_MEMBER_NAME_CHECKED(USoundWave, FrequenciesToAnalyze);
+	static FName FFTAnalysisAttackTimeFName				= GET_MEMBER_NAME_CHECKED(USoundWave, FFTAnalysisAttackTime);
+	static FName FFTAnalysisReleaseTimeFName			= GET_MEMBER_NAME_CHECKED(USoundWave, FFTAnalysisReleaseTime);
 
-	return	PropertyName == EnableFFTAnalysisFName ||
+	return	PropertyName == OverrideSoundName ||
+			PropertyName == EnableFFTAnalysisFName ||
 			PropertyName == FFTSizeFName ||
 			PropertyName == FFTAnalysisFrameSizeFName ||
-			PropertyName == FrequenciesToAnalyzeFName;
+			PropertyName == FrequenciesToAnalyzeFName || 
+			PropertyName == FFTAnalysisAttackTimeFName ||
+			PropertyName == FFTAnalysisReleaseTimeFName;
 }
 
 static bool AnyEnvelopeAnalysisPropertiesChanged(const FName& PropertyName)
 {
-	static FName EnableAmplitudeEnvelopeAnalysisFName = GET_MEMBER_NAME_CHECKED(USoundWave, bEnableAmplitudeEnvelopeAnalysis);
-	static FName EnvelopeFollowerFrameSizeFName = GET_MEMBER_NAME_CHECKED(USoundWave, EnvelopeFollowerFrameSize);
-	static FName EnvelopeFollowerAttackTimeFName = GET_MEMBER_NAME_CHECKED(USoundWave, EnvelopeFollowerAttackTime);
-	static FName EnvelopeFollowerReleaseTimeFName = GET_MEMBER_NAME_CHECKED(USoundWave, EnvelopeFollowerReleaseTime);
-
-	return	PropertyName == EnableAmplitudeEnvelopeAnalysisFName ||
+	// List of properties which cause re-analysis to get triggered
+	static FName OverrideSoundName						= GET_MEMBER_NAME_CHECKED(USoundWave, OverrideSoundToUseForAnalysis);
+	static FName EnableAmplitudeEnvelopeAnalysisFName	= GET_MEMBER_NAME_CHECKED(USoundWave, bEnableAmplitudeEnvelopeAnalysis);
+	static FName EnvelopeFollowerFrameSizeFName			= GET_MEMBER_NAME_CHECKED(USoundWave, EnvelopeFollowerFrameSize);
+	static FName EnvelopeFollowerAttackTimeFName		= GET_MEMBER_NAME_CHECKED(USoundWave, EnvelopeFollowerAttackTime);
+	static FName EnvelopeFollowerReleaseTimeFName		= GET_MEMBER_NAME_CHECKED(USoundWave, EnvelopeFollowerReleaseTime);
+
+	return	PropertyName == OverrideSoundName ||
+			PropertyName == EnableAmplitudeEnvelopeAnalysisFName ||
 			PropertyName == EnvelopeFollowerFrameSizeFName ||
 			PropertyName == EnvelopeFollowerAttackTimeFName ||
 			PropertyName == EnvelopeFollowerReleaseTimeFName;
@@ -845,7 +972,13 @@ void USoundWave::BakeFFTAnalysis()
 		uint32 RawDataSampleRate = 0;
 		uint16 RawDataNumChannels = 0;
 
-		if (!GetImportedSoundWaveData(RawImportedWaveData, RawDataSampleRate, RawDataNumChannels))
+		USoundWave* SoundWaveToUseForAnalysis = this;
+		if (OverrideSoundToUseForAnalysis)
+		{
+			SoundWaveToUseForAnalysis = OverrideSoundToUseForAnalysis;
+		}
+		
+		if (!SoundWaveToUseForAnalysis->GetImportedSoundWaveData(RawImportedWaveData, RawDataSampleRate, RawDataNumChannels))
 		{
 			return;
 		}
@@ -885,8 +1018,17 @@ void USoundWave::BakeFFTAnalysis()
 
 		}
 
+		// Prepare the spectral envelope followers
+		TArray<Audio::FEnvelopeFollower> SpectralEnvelopeFollowers;
+		SpectralEnvelopeFollowers.AddDefaulted(FrequenciesToAnalyze.Num());
+
+		for (Audio::FEnvelopeFollower& EnvFollower : SpectralEnvelopeFollowers)
+		{
+			EnvFollower.Init((float)RawDataSampleRate / FFTAnalysisFrameSize, FFTAnalysisAttackTime, FFTAnalysisReleaseTime);
+		}
+
 		// Build a new spectrum analyzer
-		Audio::FSpectrumAnalyzer SpectrumAnalyzer(SpectrumAnalyzerSettings, (float)SampleRate);
+		Audio::FSpectrumAnalyzer SpectrumAnalyzer(SpectrumAnalyzerSettings, (float)RawDataSampleRate);
 
 		// The audio data block to use to submit audio data to the spectrum analyzer
 		Audio::AlignedFloatBuffer AnalysisData;
@@ -921,11 +1063,15 @@ void USoundWave::BakeFFTAnalysis()
 				SpectrumAnalyzer.LockOutputBuffer();
 
 				// Get the magntiudes for the specified frequencies
-				for (float Frequency : FrequenciesToAnalyze)
+				for (int32 Index = 0; Index < FrequenciesToAnalyze.Num(); ++Index)
 				{
+					float Frequency = FrequenciesToAnalyze[Index];
 					FSoundWaveSpectralDataEntry DataEntry;
 					DataEntry.Magnitude = SpectrumAnalyzer.GetMagnitudeForFrequency(Frequency);
 
+					// Feed the magnitude through the spectral envelope follower for this band
+					DataEntry.Magnitude = SpectralEnvelopeFollowers[Index].ProcessAudio(DataEntry.Magnitude);
+
 					// Track the max magnitude so we can later set normalized magnitudes
 					if (DataEntry.Magnitude > MaximumMagnitude)
 					{
@@ -940,12 +1086,49 @@ void USoundWave::BakeFFTAnalysis()
 				// The time stamp is derived from the frame index and sample rate
 				NewData.TimeSec = FMath::Max((float)(FrameIndex - FFTAnalysisFrameSize + 1) / RawDataSampleRate, 0.0f);
 
-				CookedSpectralTimeData.Add(NewData);
+				/*
+				// TODO: add FFTAnalysisTimeOffset
+				// Don't let the time shift be more than the negative or postive duration
+				float Duration = (float)NumFrames / RawDataSampleRate;
+				float TimeShift = FMath::Clamp((float)FFTAnalysisTimeOffset / 1000, -Duration, Duration);
+				
+				NewData.TimeSec = NewData.TimeSec + (float)FFTAnalysisTimeOffset / 1000;
+
+				// Truncate data if time shift is far enough to left that it's before the start of the sound
+				if (TreatFileAsLoopingForAnalysis)
+				{
+					// Wrap the time value from endpoints if we're told this sound wave is looping
+					if (NewData.TimeSec < 0.0f)
+					{
+						NewData.TimeSec = Duration + NewData.TimeSec;
+					}
+					else if (NewData.TimeSec >= Duration)
+					{
+						NewData.TimeSec = NewData.TimeSec - Duration;
+					}
+					CookedSpectralTimeData.Add(NewData);
+				}
+				else if (NewData.TimeSec > 0.0f)
+				{
+					CookedSpectralTimeData.Add(NewData);
+				}
+				*/
 
 				AnalysisData.Reset();
 			}
 		}
 
+		// Sort predicate for sorting spectral data by time (lowest first)
+		struct FSortSpectralDataByTime
+		{
+			FORCEINLINE bool operator()(const FSoundWaveSpectralTimeData& A, const FSoundWaveSpectralTimeData& B) const
+			{
+				return A.TimeSec < B.TimeSec;
+			}
+		};
+
+		CookedSpectralTimeData.Sort(FSortSpectralDataByTime());
+
 		// It's possible for the maximum magnitude to be 0.0 if the audio file was silent. 
 		if (MaximumMagnitude > 0.0f)
 		{
@@ -981,7 +1164,13 @@ void USoundWave::BakeEnvelopeAnalysis()
 		uint32 RawDataSampleRate = 0;
 		uint16 RawDataNumChannels = 0;
 
-		if (!GetImportedSoundWaveData(RawImportedWaveData, RawDataSampleRate, RawDataNumChannels))
+		USoundWave* SoundWaveToUseForAnalysis = this;
+		if (OverrideSoundToUseForAnalysis)
+		{
+			SoundWaveToUseForAnalysis = OverrideSoundToUseForAnalysis;
+		}
+
+		if (!SoundWaveToUseForAnalysis->GetImportedSoundWaveData(RawImportedWaveData, RawDataSampleRate, RawDataNumChannels))
 		{
 			return;
 		}
@@ -1019,9 +1208,7 @@ void USoundWave::BakeEnvelopeAnalysis()
 				CookedEnvelopeTimeData.Add(NewData);
 			}
 		}
-
 	}
-
 }
 
 void USoundWave::PostEditChangeProperty(FPropertyChangedEvent& PropertyChangedEvent)
@@ -1045,11 +1232,13 @@ void USoundWave::PostEditChangeProperty(FPropertyChangedEvent& PropertyChangedEv
 				UpdatePlatformData();
 				MarkPackageDirty();
 			}
-			else if (AnyFFTAnalysisPropertiesChanged(Name))
+			
+			if (AnyFFTAnalysisPropertiesChanged(Name))
 			{
 				BakeFFTAnalysis();
 			}
-			else if (AnyEnvelopeAnalysisPropertiesChanged(Name))
+			
+			if (AnyEnvelopeAnalysisPropertiesChanged(Name))
 			{
 				BakeEnvelopeAnalysis();
 			}
