Fix build

diff --git a/Engine/Source/Runtime/Engine/Private/SoundWave.cpp b/Engine/Source/Runtime/Engine/Private/SoundWave.cpp
--- a/Engine/Source/Runtime/Engine/Private/SoundWave.cpp
+++ b/Engine/Source/Runtime/Engine/Private/SoundWave.cpp
@@ -56,7 +56,9 @@ static void DumpBakedAnalysisData(const TArray<FString>& Args)
 				if (SoundWaveToDump.Equals(It->GetName()))
 				{
 					UE_LOG(LogTemp, Log, TEXT("Foo"));
+#if WITH_EDITOR
 					It->LogBakedData();
+#endif // WITH_EDITOR
 				}
 			}
 		}
