Removing redundant DemoNetDriver FName creation

diff --git a/Engine/Source/Runtime/Engine/Private/GameInstance.cpp b/Engine/Source/Runtime/Engine/Private/GameInstance.cpp
--- a/Engine/Source/Runtime/Engine/Private/GameInstance.cpp
+++ b/Engine/Source/Runtime/Engine/Private/GameInstance.cpp
@@ -958,8 +958,6 @@ void UGameInstance::StartRecordingReplay(const FString& Name, const FString& Fri
 		CurrentWorld->DestroyDemoNetDriver();
 		bDestroyedDemoNetDriver = true; 
 
-		const FName NAME_DemoNetDriver(TEXT("DemoNetDriver"));
-
 		if (!GEngine->CreateNamedNetDriver(CurrentWorld, NAME_DemoNetDriver, NAME_DemoNetDriver))
 		{
 			UE_LOG(LogDemo, Warning, TEXT("RecordReplay: failed to create demo net driver!"));
@@ -1048,8 +1046,6 @@ bool UGameInstance::PlayReplay(const FString& Name, UWorld* WorldOverride, const
 		DemoURL.AddOption(*Option);
 	}
 
-	const FName NAME_DemoNetDriver( TEXT( "DemoNetDriver" ) );
-
 	if ( !GEngine->CreateNamedNetDriver( CurrentWorld, NAME_DemoNetDriver, NAME_DemoNetDriver ) )
 	{
 		UE_LOG(LogDemo, Warning, TEXT( "PlayReplay: failed to create demo net driver!" ) );
