Fixed cpu skinning crash. CPU skinning was using a local vertex factory, but we were always providing a primitive uniform buffer for skinned meshes.
Repro:
* Open skeletal mesh in the editor.
* Pick a joint from skeleton tree window.
* Click mesh -> show selected bone weight option.

diff --git a/Engine/Source/Runtime/Engine/Private/SkeletalMesh.cpp b/Engine/Source/Runtime/Engine/Private/SkeletalMesh.cpp
--- a/Engine/Source/Runtime/Engine/Private/SkeletalMesh.cpp
+++ b/Engine/Source/Runtime/Engine/Private/SkeletalMesh.cpp
@@ -3509,7 +3509,7 @@ FSkeletalMeshSceneProxy::FSkeletalMeshSceneProxy(const USkinnedMeshComponent* Co
 	}
 
 	// Skip primitive uniform buffer if we will be using local vertex factory which gets it's data from GPUScene.
-	bVFRequiresPrimitiveUniformBuffer = !(bRenderStatic && UseGPUScene(GMaxRHIShaderPlatform, FeatureLevel));
+	bVFRequiresPrimitiveUniformBuffer = !((bIsCPUSkinned || bRenderStatic) && UseGPUScene(GMaxRHIShaderPlatform, FeatureLevel));
 }
 
 
diff --git a/Engine/Source/Runtime/Renderer/Private/GPUScene.cpp b/Engine/Source/Runtime/Renderer/Private/GPUScene.cpp
--- a/Engine/Source/Runtime/Renderer/Private/GPUScene.cpp
+++ b/Engine/Source/Runtime/Renderer/Private/GPUScene.cpp
@@ -361,7 +361,7 @@ void AddPrimitiveToUpdateGPU(FScene& Scene, int32 PrimitiveId)
 			Scene.GPUScene.PrimitivesMarkedToUpdate.Add(0, NewSize - Scene.GPUScene.PrimitivesMarkedToUpdate.Num());
 		}
 
-		// Make sure we are no updating same primitive multiple times.
+		// Make sure we aren't updating same primitive multiple times.
 		if (!Scene.GPUScene.PrimitivesMarkedToUpdate[PrimitiveId])
 		{
 			Scene.GPUScene.PrimitivesToUpdate.Add(PrimitiveId);
