Add flag (a.StripOddFrameAnimsWhenFrameStripping) for performing frame stripping on animations that have an odd number of frames (unlike animations with even frames these need resampling so it is a bigger quality tradeoff)

diff --git a/Engine/Source/Runtime/Engine/Classes/Animation/AnimSequence.h b/Engine/Source/Runtime/Engine/Classes/Animation/AnimSequence.h
--- a/Engine/Source/Runtime/Engine/Classes/Animation/AnimSequence.h
+++ b/Engine/Source/Runtime/Engine/Classes/Animation/AnimSequence.h
@@ -305,6 +305,9 @@ struct ENGINE_API FRequestAnimCompressionParams
 	// Should we attempt to do framestripping (removing every other frame from raw animation tracks)
 	bool bPerformFrameStripping;
 
+	// If false we only perform frame stripping on even numbered frames (as a quality measure)
+	bool bPerformFrameStrippingOnOddNumberedFrames;
+
 	// Compression context
 	TSharedPtr<FAnimCompressContext> CompressContext;
 
diff --git a/Engine/Source/Runtime/Engine/Private/Animation/AnimCompress_RemoveLinearKeys.cpp b/Engine/Source/Runtime/Engine/Private/Animation/AnimCompress_RemoveLinearKeys.cpp
--- a/Engine/Source/Runtime/Engine/Private/Animation/AnimCompress_RemoveLinearKeys.cpp
+++ b/Engine/Source/Runtime/Engine/Private/Animation/AnimCompress_RemoveLinearKeys.cpp
@@ -24,33 +24,6 @@ static FQuat EnforceShortestArc(const FQuat& A, const FQuat& B)
 	return B*Bias;
 }
 
-/**
- * Helper template function to interpolate between two data types.
- * Used in the FilterLinearKeysTemplate function below
- */
-template <typename T>
-T Interpolate(const T& A, const T& B, float Alpha)
-{
-	// only the custom instantiations below are valid
-	check(0);
-	return 0;
-}
-
-/** custom instantiation of Interpolate for FVectors */
-template <> FVector Interpolate<FVector>(const FVector& A, const FVector& B, float Alpha)
-{
-	return FMath::Lerp(A,B,Alpha);
-}
-
-/** custom instantiation of Interpolate for FQuats */
-template <> FQuat Interpolate<FQuat>(const FQuat& A, const FQuat& B, float Alpha)
-{
-	FQuat result = FQuat::FastLerp(A,B,Alpha);
-	result.Normalize();
-
-	return result;
-}
-
 /**
  * Helper template function to calculate the delta between two data types.
  * Used in the FilterLinearKeysTemplate function below
@@ -259,7 +232,7 @@ void FilterLinearKeysTemplate(
 
 					// compute the proposed, interpolated value for the key
 					const float Alpha = (TestTime - LowTime) * InvRange;
-					const KeyType LerpValue = Interpolate(LowValue, HighValue, Alpha);
+					const KeyType LerpValue = AnimationCompressionUtils::Interpolate(LowValue, HighValue, Alpha);
 
 					// compute the error between our interpolated value and the desired value
 					float LerpError = CalcDelta(TestValue, LerpValue);
@@ -463,7 +436,7 @@ void FilterLinearKeysTemplate(
 
 					// compute the proposed, interpolated value for the key
 					const float Alpha = (TestTime - LowTime) * InvRange;
-					const KeyType LerpValue = Interpolate(LowValue, HighValue, Alpha);
+					const KeyType LerpValue = AnimationCompressionUtils::Interpolate(LowValue, HighValue, Alpha);
 
 					// compute the error between our interpolated value and the desired value
 					float LerpError = CalcDelta(TestValue, LerpValue);
diff --git a/Engine/Source/Runtime/Engine/Private/Animation/AnimCompressionDerivedData.cpp b/Engine/Source/Runtime/Engine/Private/Animation/AnimCompressionDerivedData.cpp
--- a/Engine/Source/Runtime/Engine/Private/Animation/AnimCompressionDerivedData.cpp
+++ b/Engine/Source/Runtime/Engine/Private/Animation/AnimCompressionDerivedData.cpp
@@ -15,7 +15,7 @@
 DECLARE_CYCLE_STAT(TEXT("Anim Compression (Derived Data)"), STAT_AnimCompressionDerivedData, STATGROUP_Anim);
 
 template<typename ArrayValue>
-void StripFrames(TArray<ArrayValue>& Keys, const int32 NumFrames)
+void StripFramesEven(TArray<ArrayValue>& Keys, const int32 NumFrames)
 {
 	if (Keys.Num() > 1)
 	{
@@ -33,7 +33,46 @@ void StripFrames(TArray<ArrayValue>& Keys, const int32 NumFrames)
 	}
 }
 
-FDerivedDataAnimationCompression::FDerivedDataAnimationCompression(UAnimSequence* InAnimSequence, TSharedPtr<FAnimCompressContext> InCompressContext, bool bInDoCompressionInPlace, bool bInTryFrameStripping)
+template<typename ArrayValue>
+void StripFramesOdd(TArray<ArrayValue>& Keys, const int32 NumFrames)
+{
+	if (Keys.Num() > 1)
+	{
+		const int32 NewNumFrames = NumFrames / 2;
+
+		TArray<ArrayValue> NewKeys;
+		NewKeys.Reserve(NewNumFrames);
+
+		check(Keys.Num() == NumFrames);
+
+		NewKeys.Add(Keys[0]); //Always keep first 
+
+		//Always keep first and last
+		const int32 NumFramesToCalculate = NewNumFrames - 2;
+
+		// Frame increment is ratio of old frame spaces vs new frame spaces 
+		const double FrameIncrement = (double)(NumFrames - 1) / (double)(NewNumFrames - 1);
+			
+		for(int32 Frame=0; Frame < NumFramesToCalculate; ++Frame)
+		{
+			const double NextFramePosition = FrameIncrement * (Frame + 1);
+			const int32 Frame1 = (int32)NextFramePosition;
+			const float Alpha = (NextFramePosition - (double)Frame1);
+
+			NewKeys.Add(AnimationCompressionUtils::Interpolate(Keys[Frame1], Keys[Frame1 + 1], Alpha));
+
+		}
+
+		NewKeys.Add(Keys.Last()); // Always Keep Last
+
+		const int32 HalfSize = (NumFrames - 1) / 2;
+		const int32 StartRemoval = HalfSize + 1;
+
+		Keys = MoveTemp(NewKeys);
+	}
+}
+
+FDerivedDataAnimationCompression::FDerivedDataAnimationCompression(UAnimSequence* InAnimSequence, TSharedPtr<FAnimCompressContext> InCompressContext, bool bInDoCompressionInPlace, bool bInTryFrameStripping, bool bTryStrippingOnOddFramedAnims)
 	: OriginalAnimSequence(InAnimSequence)
 	, DuplicateSequence(nullptr)
 	, CompressContext(InCompressContext)
@@ -43,7 +82,11 @@ FDerivedDataAnimationCompression::FDerivedDataAnimationCompression(UAnimSequence
 	InAnimSequence->AddToRoot(); //Keep this around until we are finished
 
 	// Can only do stripping on animations that have an even number of frames once the end frame is removed)
-	const bool bStripCandidate = (OriginalAnimSequence->GetRawNumberOfFrames() > 10) && (((OriginalAnimSequence->GetRawNumberOfFrames() - 1) % 2) == 0);
+	bIsEvenFramed = ((OriginalAnimSequence->GetRawNumberOfFrames() - 1) % 2) == 0;
+	const bool bIsValidForStripping = bIsEvenFramed || bTryStrippingOnOddFramedAnims;
+
+	const bool bStripCandidate = (OriginalAnimSequence->GetRawNumberOfFrames() > 10) && bIsValidForStripping;
+	
 	bPerformStripping = bStripCandidate && bInTryFrameStripping;
 }
 
@@ -52,7 +95,9 @@ FDerivedDataAnimationCompression::~FDerivedDataAnimationCompression()
 	OriginalAnimSequence->RemoveFromRoot();
 	if (DuplicateSequence)
 	{
+		DuplicateSequence->ClearFlags(RF_Standalone | RF_Public);
 		DuplicateSequence->RemoveFromRoot();
+		DuplicateSequence->MarkPendingKill();
 	}
 }
 
@@ -144,17 +189,34 @@ bool FDerivedDataAnimationCompression::Build( TArray<uint8>& OutData )
 			const int32 NumTracks = AnimToOperateOn->GetRawAnimationData().Num();
 
 			//Strip every other frame from tracks
-			for (int32 TrackIndex = 0; TrackIndex < NumTracks; ++TrackIndex)
+			if(bIsEvenFramed)
 			{
-				FRawAnimSequenceTrack& Track = AnimToOperateOn->GetRawAnimationTrack(TrackIndex);
+				for (int32 TrackIndex = 0; TrackIndex < NumTracks; ++TrackIndex)
+				{
+					FRawAnimSequenceTrack& Track = AnimToOperateOn->GetRawAnimationTrack(TrackIndex);
+
+					StripFramesEven(Track.PosKeys, NumFrames);
+					StripFramesEven(Track.RotKeys, NumFrames);
+					StripFramesEven(Track.ScaleKeys, NumFrames);
+				}
 
-				StripFrames(Track.PosKeys, NumFrames);
-				StripFrames(Track.RotKeys, NumFrames);
-				StripFrames(Track.ScaleKeys, NumFrames);
+				const int32 ActualFrames = AnimToOperateOn->GetRawNumberOfFrames() - 1; // strip bookmark end frame
+				AnimToOperateOn->SetRawNumberOfFrame((ActualFrames / 2) + 1);
 			}
+			else
+			{
+				for (int32 TrackIndex = 0; TrackIndex < NumTracks; ++TrackIndex)
+				{
+					FRawAnimSequenceTrack& Track = AnimToOperateOn->GetRawAnimationTrack(TrackIndex);
+
+					StripFramesOdd(Track.PosKeys, NumFrames);
+					StripFramesOdd(Track.RotKeys, NumFrames);
+					StripFramesOdd(Track.ScaleKeys, NumFrames);
+				}
 
-			const int32 StripFrames = AnimToOperateOn->GetRawNumberOfFrames() - 1;
-			AnimToOperateOn->SetRawNumberOfFrame((StripFrames / 2) + 1);
+				const int32 ActualFrames = AnimToOperateOn->GetRawNumberOfFrames(); // strip bookmark end frame
+				AnimToOperateOn->SetRawNumberOfFrame((ActualFrames / 2));
+			}
 		}
 
 		AnimToOperateOn->UpdateCompressedTrackMapFromRaw();
diff --git a/Engine/Source/Runtime/Engine/Private/Animation/AnimCompressionDerivedData.h b/Engine/Source/Runtime/Engine/Private/Animation/AnimCompressionDerivedData.h
--- a/Engine/Source/Runtime/Engine/Private/Animation/AnimCompressionDerivedData.h
+++ b/Engine/Source/Runtime/Engine/Private/Animation/AnimCompressionDerivedData.h
@@ -33,8 +33,11 @@ private:
 	// Whether we should frame strip (remove every other frame from even frames animations)
 	bool bPerformStripping;
 
+	// Track if it is an even framed animation (when stripping odd framed animations will need to be resampled)
+	bool bIsEvenFramed;
+
 public:
-	FDerivedDataAnimationCompression(UAnimSequence* InAnimSequence, TSharedPtr<FAnimCompressContext> InCompressContext, bool bInDoCompressionInPlace, bool bInTryFrameStripping);
+	FDerivedDataAnimationCompression(UAnimSequence* InAnimSequence, TSharedPtr<FAnimCompressContext> InCompressContext, bool bInDoCompressionInPlace, bool bInTryFrameStripping, bool bTryStrippingOnOddFramedAnims);
 	virtual ~FDerivedDataAnimationCompression();
 
 	virtual const TCHAR* GetPluginName() const override
@@ -47,7 +50,7 @@ public:
 		// This is a version string that mimics the old versioning scheme. If you
 		// want to bump this version, generate a new guid using VS->Tools->Create GUID and
 		// return it here. Ex.
-		return TEXT("D58C7C3E48274C04BFC4D7823AF46F5E");
+		return TEXT("A74732B14AD94950856C6A09F708445C");
 	}
 
 	virtual FString GetPluginSpecificCacheKeySuffix() const override;
diff --git a/Engine/Source/Runtime/Engine/Private/Animation/AnimSequence.cpp b/Engine/Source/Runtime/Engine/Private/Animation/AnimSequence.cpp
--- a/Engine/Source/Runtime/Engine/Private/Animation/AnimSequence.cpp
+++ b/Engine/Source/Runtime/Engine/Private/Animation/AnimSequence.cpp
@@ -48,27 +48,44 @@ DECLARE_CYCLE_STAT(TEXT("Build Anim Track Pairs"), STAT_BuildAnimTrackPairs, STA
 DECLARE_CYCLE_STAT(TEXT("Extract Pose From Anim Data"), STAT_ExtractPoseFromAnimData, STATGROUP_Anim);
 
 int32 GPerformFrameStripping = 0;
+int32 GPerformFrameStrippingOddFramedAnimations = 0;
 
 static const TCHAR* StripFrameCVarName = TEXT("a.StripFramesOnCompression");
+static const TCHAR* OddFrameStripStrippingCVarName = TEXT("a.StripOddFramesWhenFrameStripping");
 
 static FAutoConsoleVariableRef CVarAnimStripFramesOnCompression(
 	StripFrameCVarName,
 	GPerformFrameStripping,
 	TEXT("1 = Strip every other frame on animations that have an even number of frames. 0 = off"));
 
+static FAutoConsoleVariableRef CVarOddFrameStripping(
+	OddFrameStripStrippingCVarName,
+	GPerformFrameStrippingOddFramedAnimations,
+	TEXT("1 = When frame stripping apply to animations with an odd number of frames too. 0 = only even framed animations"));
+
 void OnCVarsChanged()
 {
 	static bool bCompressionFrameStrip = (GPerformFrameStripping == 1);
+	static bool bOddFramedStrip = (GPerformFrameStrippingOddFramedAnimations == 1);
+
 	static TArray<UAnimSequence*> SequenceCache;
 	static FString OutputMessage;
 
 	const bool bCurrentFrameStrip = (GPerformFrameStripping == 1);
-	if (bCompressionFrameStrip != bCurrentFrameStrip)
+	const bool bCurrentOddFramedStrip = (GPerformFrameStrippingOddFramedAnimations == 1);
+
+	const bool bFrameStripChanged = bCompressionFrameStrip != bCurrentFrameStrip;
+	const bool bOddFrameStripChanged = bOddFramedStrip != bCurrentOddFramedStrip;
+
+	if (bFrameStripChanged || bOddFrameStripChanged)
 	{
 		bCompressionFrameStrip = bCurrentFrameStrip;
+		bOddFramedStrip = bCurrentOddFramedStrip;
 
 		SequenceCache.Reset();
 
+		CollectGarbage(GARBAGE_COLLECTION_KEEPFLAGS);
+
 		for (TObjectIterator<UAnimSequence> It; It; ++It)
 		{
 			SequenceCache.Add(*It);
@@ -81,10 +98,7 @@ void OnCVarsChanged()
 		{
 			Seq->RequestSyncAnimRecompression();
 
-			if (!bCurrentFrameStrip || Seq->GetCompressedNumberOfFrames() != Seq->GetRawNumberOfFrames())
-			{
-				Sizes.Emplace(Seq->GetApproxCompressedSize(), Seq);
-			}
+			Sizes.Emplace(Seq->GetApproxCompressedSize(), Seq);
 		}
 
 		Sizes.Sort([](const TPair<int32, UAnimSequence*>& A, const TPair<int32, UAnimSequence*>& B)
@@ -94,14 +108,25 @@ void OnCVarsChanged()
 
 		OutputMessage.Reset();
 
+		const TCHAR* StripMessage = bCompressionFrameStrip ? TEXT("Stripping: On") : TEXT("Stripping: Off");
+		const TCHAR* OddMessage = bOddFramedStrip ? TEXT("Odd Frames: On") : TEXT("Odd Frames: Off");
+
+		OutputMessage += FString::Printf(TEXT("%s - %s\n\n"), StripMessage, OddMessage);
+
 		int32 TotalSize = 0;
 		int32 NumAnimations = 0;
 		for (const TPair<int32, UAnimSequence*>& Pair : Sizes)
 		{
-			OutputMessage += FString::Printf(TEXT("%s - %.1fK\n"), *Pair.Value->GetPathName(), (float)Pair.Key / 1000.f);
-			TotalSize += Pair.Key;
-			NumAnimations++;
+			const bool bIsOddFramed = (Pair.Value->GetNumberOfFrames() % 2) == 0;
+			if (bIsOddFramed)
+			{
+				OutputMessage += FString::Printf(TEXT("%s - %.1fK\n"), *Pair.Value->GetPathName(), (float)Pair.Key / 1000.f);
+				TotalSize += Pair.Key;
+				NumAnimations++;
+			}
 		}
+
+		OutputMessage += FString::Printf(TEXT("\n\nTotalAnims: %i TotalSize = %.1fK"), NumAnimations, ((float)TotalSize / 1000.f));
 		FPlatformApplicationMisc::ClipboardCopy(*OutputMessage);
 	}
 }
@@ -127,6 +152,7 @@ FRequestAnimCompressionParams::FRequestAnimCompressionParams(bool bInAsyncCompre
 void FRequestAnimCompressionParams::InitFrameStrippingFromCVar()
 {
 	bPerformFrameStripping = (GPerformFrameStripping == 1);
+	bPerformFrameStrippingOnOddNumberedFrames = (GPerformFrameStrippingOddFramedAnimations == 1);
 }
 
 void FRequestAnimCompressionParams::InitFrameStrippingFromPlatform(const class ITargetPlatform* TargetPlatform)
@@ -2514,9 +2540,10 @@ void UAnimSequence::RequestAnimCompression(FRequestAnimCompressionParams Params)
 	else
 	{
 		const bool bPerformFrameStripping = Params.bPerformFrameStripping && bAllowFrameStripping;
+		const bool bPerformStrippingOnOddFramedAnims = Params.bPerformFrameStrippingOnOddNumberedFrames;
 
 		TArray<uint8> OutData;
-		FDerivedDataAnimationCompression* AnimCompressor = new FDerivedDataAnimationCompression(this, Params.CompressContext, bDoCompressionInPlace, bPerformFrameStripping);
+		FDerivedDataAnimationCompression* AnimCompressor = new FDerivedDataAnimationCompression(this, Params.CompressContext, bDoCompressionInPlace, bPerformFrameStripping, bPerformStrippingOnOddFramedAnims);
 		// For debugging DDC/Compression issues		
 		const bool bSkipDDC = false;
 		if (bSkipDDC || (CompressCommandletVersion == INDEX_NONE))
diff --git a/Engine/Source/Runtime/Engine/Private/Animation/AnimationUtils.cpp b/Engine/Source/Runtime/Engine/Private/Animation/AnimationUtils.cpp
--- a/Engine/Source/Runtime/Engine/Private/Animation/AnimationUtils.cpp
+++ b/Engine/Source/Runtime/Engine/Private/Animation/AnimationUtils.cpp
@@ -790,6 +790,8 @@ static void ClearAnimCompressionJobs(FGraphEventArray& AnimCompressionTask_Compl
 	for (FAnimCompressionJobContext* Context : AnimCompressionJobContexes)
 	{
 		Context->AnimSeq->RecycleAnimSequence();
+		Context->AnimSeq->ClearFlags(RF_Standalone | RF_Public);
+		Context->AnimSeq->MarkPendingKill();
 		delete Context;
 	}
 	AnimCompressionJobContexes.Reset();
diff --git a/Engine/Source/Runtime/Engine/Public/AnimationCompression.h b/Engine/Source/Runtime/Engine/Public/AnimationCompression.h
--- a/Engine/Source/Runtime/Engine/Public/AnimationCompression.h
+++ b/Engine/Source/Runtime/Engine/Public/AnimationCompression.h
@@ -54,6 +54,33 @@ namespace AnimationCompressionUtils
 		return Result;
 #endif
 	}
+
+	/**
+	* Helper template function to interpolate between two data types.
+	* Used in the FilterLinearKeysTemplate function below
+	*/
+	template <typename T>
+	FORCEINLINE T Interpolate(const T& A, const T& B, float Alpha)
+	{
+		// only the custom instantiations below are valid
+		check(0);
+		return 0;
+	}
+
+	/** custom instantiation of Interpolate for FVectors */
+	template <> FORCEINLINE FVector Interpolate<FVector>(const FVector& A, const FVector& B, float Alpha)
+	{
+		return FMath::Lerp(A, B, Alpha);
+	}
+
+	/** custom instantiation of Interpolate for FQuats */
+	template <> FORCEINLINE FQuat Interpolate<FQuat>(const FQuat& A, const FQuat& B, float Alpha)
+	{
+		FQuat result = FQuat::FastLerp(A, B, Alpha);
+		result.Normalize();
+
+		return result;
+	}
 }
 
 class FQuatFixed48NoW
