Fixed issue where server could receive data from the client on a closed actor channel. The server would then try to open the channel, fail and report that to the client which would start explicitly nacking all packets containing refereces to that channel.

diff --git a/Engine/Source/Runtime/Engine/Private/NetConnection.cpp b/Engine/Source/Runtime/Engine/Private/NetConnection.cpp
--- a/Engine/Source/Runtime/Engine/Private/NetConnection.cpp
+++ b/Engine/Source/Runtime/Engine/Private/NetConnection.cpp
@@ -1978,6 +1978,14 @@ PRAGMA_ENABLE_DEPRECATION_WARNINGS
 					return;
 				}
 
+				// Ignore incoming data on channel types that the client are not allowed to create. This can occur if we have in-flight data when server is closing a channel
+				if ( Driver->IsServer() && (Driver->ChannelDefinitionMap[Bunch.ChName].bClientOpen == false) )
+				{
+					UE_LOG(LogNetTraffic, Warning, TEXT("      Ignoring Bunch Create received from client since only server is allowed to create this type of channel: Bunch  %i: ChName %s, ChSequence: %i, bReliable: %i, bPartial: %i, bPartialInitial: %i, bPartialFinal: %i"), Bunch.ChIndex, *Bunch.ChName.ToString(), Bunch.ChSequence, (int)Bunch.bReliable, (int)Bunch.bPartial, (int)Bunch.bPartialInitial, (int)Bunch.bPartialFinal );
+					RejectedChans.AddUnique(Bunch.ChIndex);
+					continue;
+				}
+
 				// Reliable (either open or later), so create new channel.
 				UE_LOG(LogNetTraffic, Log, TEXT("      Bunch Create %i: ChName %s, ChSequence: %i, bReliable: %i, bPartial: %i, bPartialInitial: %i, bPartialFinal: %i"), Bunch.ChIndex, *Bunch.ChName.ToString(), Bunch.ChSequence, (int)Bunch.bReliable, (int)Bunch.bPartial, (int)Bunch.bPartialInitial, (int)Bunch.bPartialFinal );
 				Channel = CreateChannelByName( Bunch.ChName, EChannelCreateFlags::None, Bunch.ChIndex );
