#UE4 Fix for a rendering crash that happens when compiling a blueprint that has a static mesh component while the level editor has a level loaded with an instance of that blueprint that has mesh paint on it.

diff --git a/Engine/Source/Runtime/Engine/Private/Components/StaticMeshComponent.cpp b/Engine/Source/Runtime/Engine/Private/Components/StaticMeshComponent.cpp
--- a/Engine/Source/Runtime/Engine/Private/Components/StaticMeshComponent.cpp
+++ b/Engine/Source/Runtime/Engine/Private/Components/StaticMeshComponent.cpp
@@ -2730,13 +2730,20 @@ FArchive& operator<<(FArchive& Ar,FStaticMeshComponentLODInfo& I)
 			check(I.OverrideVertexColors != nullptr);
 			I.OverrideVertexColors->Serialize(Ar, bNeedsCPUAccess);
 			
-			//When IsSaving, we cannot have this situation since it is check before
-			if (Ar.IsLoading() && (I.OverrideVertexColors->GetNumVertices() <= 0 || I.OverrideVertexColors->GetStride() <= 0))
+			if (Ar.IsLoading())
 			{
-				UE_LOG(LogStaticMesh, Log, TEXT("Loading a staticmesh component that is flag with override vertex color buffer, but the buffer is empty after loading(serializing) it. Resave the map to fix the component override vertex color data."));
-				//Avoid saving an empty array
-				delete I.OverrideVertexColors;
-				I.OverrideVertexColors = nullptr;
+				//When IsSaving, we cannot have this situation since it is check before
+				if (I.OverrideVertexColors->GetNumVertices() <= 0 || I.OverrideVertexColors->GetStride() <= 0)
+				{
+					UE_LOG(LogStaticMesh, Log, TEXT("Loading a staticmesh component that is flag with override vertex color buffer, but the buffer is empty after loading(serializing) it. Resave the map to fix the component override vertex color data."));
+					//Avoid saving an empty array
+					delete I.OverrideVertexColors;
+					I.OverrideVertexColors = nullptr;
+				}
+				else
+				{
+					BeginInitResource(I.OverrideVertexColors);
+				}
 			}
 		}
 	}
