Fixing up bugs with the FFT envelope follower.

diff --git a/Engine/Source/Runtime/AudioMixer/Private/DSP/EnvelopeFollower.cpp b/Engine/Source/Runtime/AudioMixer/Private/DSP/EnvelopeFollower.cpp
--- a/Engine/Source/Runtime/AudioMixer/Private/DSP/EnvelopeFollower.cpp
+++ b/Engine/Source/Runtime/AudioMixer/Private/DSP/EnvelopeFollower.cpp
@@ -94,6 +94,25 @@ namespace Audio
 		return CurrentEnvelopeValue = NewEnvelopeValue;
 	}
 
+	float FEnvelopeFollower::ProcessAudioNonClamped(const float InAudioSample)
+	{
+		// Take the absolute value of the input sample
+		float Sample = FMath::Abs(InAudioSample);
+
+		// If we're not Peak detecting, then square the input
+		if (EnvMode != EPeakMode::Peak)
+		{
+			Sample = Sample * Sample;
+		}
+
+		float TimeSamples = (Sample > CurrentEnvelopeValue) ? AttackTimeSamples : ReleaseTimeSamples;
+		float NewEnvelopeValue = TimeSamples * (CurrentEnvelopeValue - Sample) + Sample;;
+		NewEnvelopeValue = Audio::UnderflowClamp(NewEnvelopeValue);
+
+		// Update and return the envelope value
+		return CurrentEnvelopeValue = NewEnvelopeValue;
+	}
+
 	int16 FEnvelopeFollower::ProcessAudio(const int16 InAudioSample)
 	{
 		// Convert to float
diff --git a/Engine/Source/Runtime/AudioMixer/Public/DSP/EnvelopeFollower.h b/Engine/Source/Runtime/AudioMixer/Public/DSP/EnvelopeFollower.h
--- a/Engine/Source/Runtime/AudioMixer/Public/DSP/EnvelopeFollower.h
+++ b/Engine/Source/Runtime/AudioMixer/Public/DSP/EnvelopeFollower.h
@@ -50,6 +50,9 @@ namespace Audio
 		// Processes the input audio stream and returns the envelope value.
 		float ProcessAudio(const float InAudioSample);
 
+		// Processes the input audio stream and returns the envelope value.
+		float ProcessAudioNonClamped(const float InAudioSample);
+
 		// Process the input audio stream in int16 format
 		int16 ProcessAudio(const int16 InAudioSample);
 
diff --git a/Engine/Source/Runtime/Engine/Private/SoundWave.cpp b/Engine/Source/Runtime/Engine/Private/SoundWave.cpp
--- a/Engine/Source/Runtime/Engine/Private/SoundWave.cpp
+++ b/Engine/Source/Runtime/Engine/Private/SoundWave.cpp
@@ -1057,7 +1057,7 @@ void USoundWave::BakeFFTAnalysis()
 				SpectrumAnalyzer.PushAudio(AnalysisData.GetData(), AnalysisData.Num());
 
 				// Block while the analyzer does the analysis
-				while (SpectrumAnalyzer.PerformAnalysisIfPossible()) {}
+				SpectrumAnalyzer.PerformAnalysisIfPossible(true);
 
 				FSoundWaveSpectralTimeData NewData;
 
@@ -1072,7 +1072,7 @@ void USoundWave::BakeFFTAnalysis()
 					DataEntry.Magnitude = SpectrumAnalyzer.GetMagnitudeForFrequency(Frequency);
 
 					// Feed the magnitude through the spectral envelope follower for this band
-					DataEntry.Magnitude = SpectralEnvelopeFollowers[Index].ProcessAudio(DataEntry.Magnitude);
+					DataEntry.Magnitude = SpectralEnvelopeFollowers[Index].ProcessAudioNonClamped(DataEntry.Magnitude);
 
 					// Track the max magnitude so we can later set normalized magnitudes
 					if (DataEntry.Magnitude > MaximumMagnitude)
@@ -1116,6 +1116,8 @@ void USoundWave::BakeFFTAnalysis()
 				}
 				*/
 
+				CookedSpectralTimeData.Add(NewData);
+
 				AnalysisData.Reset();
 			}
 		}
