#ue4 - Fix jittery character proxies when using Linear smoothing. ACharacter::PreReplication() was stomping timestamps when clientside replay recording was enabled.

diff --git a/Engine/Source/Runtime/Engine/Private/Character.cpp b/Engine/Source/Runtime/Engine/Private/Character.cpp
--- a/Engine/Source/Runtime/Engine/Private/Character.cpp
+++ b/Engine/Source/Runtime/Engine/Private/Character.cpp
@@ -1456,14 +1456,19 @@ void ACharacter::PreReplication( IRepChangedPropertyTracker & ChangedPropertyTra
 		}
 	}
 
-	// Save bandwidth by not replicating this value unless it is necessary, since it changes every update.
-	if ((CharacterMovement->NetworkSmoothingMode == ENetworkSmoothingMode::Linear) || CharacterMovement->bNetworkAlwaysReplicateTransformUpdateTimestamp)
+	// Avoid stomping timestamp values when client-side replay recording calls this code!
+	// GetServerLastTransformUpdateTimeStamp() returns 0 on clients, but ReplayLastTransformUpdateTimeStamp will take care of it.
+	if (Role == ROLE_Authority)
 	{
-		ReplicatedServerLastTransformUpdateTimeStamp = CharacterMovement->GetServerLastTransformUpdateTimeStamp();
-	}
-	else
-	{
-		ReplicatedServerLastTransformUpdateTimeStamp = 0.f;
+		// Save bandwidth by not replicating this value unless it is necessary, since it changes every update.
+		if ((CharacterMovement->NetworkSmoothingMode == ENetworkSmoothingMode::Linear) || CharacterMovement->bNetworkAlwaysReplicateTransformUpdateTimestamp)
+		{
+			ReplicatedServerLastTransformUpdateTimeStamp = CharacterMovement->GetServerLastTransformUpdateTimeStamp();
+		}
+		else
+		{
+			ReplicatedServerLastTransformUpdateTimeStamp = 0.f;
+		}
 	}
 }
 
