Fix editor compile error

diff --git a/Engine/Source/Runtime/Engine/Public/ContentStreaming.h b/Engine/Source/Runtime/Engine/Public/ContentStreaming.h
--- a/Engine/Source/Runtime/Engine/Public/ContentStreaming.h
+++ b/Engine/Source/Runtime/Engine/Public/ContentStreaming.h
@@ -388,6 +388,14 @@ struct IRenderAssetStreamingManager : public IStreamingManager
 
 	/** Return all bounds related to the ref object */
 	ENGINE_API virtual void GetObjectReferenceBounds(const UObject* RefObject, TArray<FBox>& AssetBoxes) = 0;
+
+	//BEGIN: APIs for backward compatibility
+	void UpdateIndividualTexture(UTexture2D* Texture) { UpdateIndividualRenderAsset(Texture); }
+	bool StreamOutTextureData(int64 RequiredMemorySize) { return StreamOutRenderAssetData(RequiredMemorySize); }
+	void AddStreamingTexture(UTexture2D* Texture) { AddStreamingRenderAsset(Texture); }
+	void RemoveStreamingTexture(UTexture2D* Texture) { RemoveStreamingRenderAsset(Texture); }
+	void PauseTextureStreaming(bool bInShouldPause) { PauseRenderAssetStreaming(bInShouldPause); }
+	//END: APIs for backward compatibility
 };
 
 /**
