#UE4 Fix crash that occurs if LOD1 is painted on a mesh but not LOD0 when a blueprint is compiled while an instance of that blueprint is placed in the level.

diff --git a/Engine/Source/Runtime/Engine/Private/Components/StaticMeshComponent.cpp b/Engine/Source/Runtime/Engine/Private/Components/StaticMeshComponent.cpp
--- a/Engine/Source/Runtime/Engine/Private/Components/StaticMeshComponent.cpp
+++ b/Engine/Source/Runtime/Engine/Private/Components/StaticMeshComponent.cpp
@@ -131,7 +131,7 @@ bool FStaticMeshComponentInstanceData::ApplyVertexColorData(UStaticMeshComponent
 				// may not be null at the start (could have been initialized 
 				// from a  component template with vert coloring), but should
 				// be null at this point, after RemoveInstanceVertexColorsFromLOD()
-				if (LODInfo.OverrideVertexColors == NULL)
+				if (LODInfo.OverrideVertexColors == NULL && VertexColorLODData.VertexBufferColors.Num() > 0)
 				{
 					LODInfo.PaintedVertices = VertexColorLODData.PaintedVertices;
 
