#ue4 - Back out changelist 4925436, not necessary because the PreReplication() function doesn't run for proxies.

diff --git a/Engine/Source/Runtime/Engine/Private/Character.cpp b/Engine/Source/Runtime/Engine/Private/Character.cpp
--- a/Engine/Source/Runtime/Engine/Private/Character.cpp
+++ b/Engine/Source/Runtime/Engine/Private/Character.cpp
@@ -1456,19 +1456,14 @@ void ACharacter::PreReplication( IRepChangedPropertyTracker & ChangedPropertyTra
 		}
 	}
 
-	// Avoid stomping timestamp values when client-side replay recording calls this code!
-	// GetServerLastTransformUpdateTimeStamp() returns 0 on clients, but ReplayLastTransformUpdateTimeStamp will take care of it.
-	if (Role == ROLE_Authority)
+	// Save bandwidth by not replicating this value unless it is necessary, since it changes every update.
+	if ((CharacterMovement->NetworkSmoothingMode == ENetworkSmoothingMode::Linear) || CharacterMovement->bNetworkAlwaysReplicateTransformUpdateTimestamp)
 	{
-		// Save bandwidth by not replicating this value unless it is necessary, since it changes every update.
-		if ((CharacterMovement->NetworkSmoothingMode == ENetworkSmoothingMode::Linear) || CharacterMovement->bNetworkAlwaysReplicateTransformUpdateTimestamp)
-		{
-			ReplicatedServerLastTransformUpdateTimeStamp = CharacterMovement->GetServerLastTransformUpdateTimeStamp();
-		}
-		else
-		{
-			ReplicatedServerLastTransformUpdateTimeStamp = 0.f;
-		}
+		ReplicatedServerLastTransformUpdateTimeStamp = CharacterMovement->GetServerLastTransformUpdateTimeStamp();
+	}
+	else
+	{
+		ReplicatedServerLastTransformUpdateTimeStamp = 0.f;
 	}
 }
 
