Audio - Adding a CVar to disable spatialization of audio.

Aaron.McLeran

diff --git a/Engine/Source/Runtime/Engine/Private/AudioDevice.cpp b/Engine/Source/Runtime/Engine/Private/AudioDevice.cpp
--- a/Engine/Source/Runtime/Engine/Private/AudioDevice.cpp
+++ b/Engine/Source/Runtime/Engine/Private/AudioDevice.cpp
@@ -111,6 +111,13 @@ FAutoConsoleVariableRef CVarNumPrecacheFrames(
 	TEXT("0: Use default value for precache frames, >0: Number of frames to precache."),
 	ECVF_Default);
 
+static int32 AllowAudioSpatializationCVar = 1;
+FAutoConsoleVariableRef CVarAllowAudioSpatializationCVar(
+	TEXT("au.AllowAudioSpatialization"),
+	AllowAudioSpatializationCVar,
+	TEXT("Controls if we allow spatialization of audio, normally this is enabled.  If disabled all audio won't be spatialized, but will have attenuation.\n")
+	TEXT("0: Disable, >0: Enable"),
+	ECVF_Default);
 
 /*-----------------------------------------------------------------------------
 FDynamicParameter implementation.
@@ -4156,6 +4163,12 @@ void FAudioDevice::AddNewActiveSound(const FActiveSound& NewActiveSound)
 		OneShotCount++;
 	}
 
+	// If we've disabled audio spatialization, then we need to force this active sound to no longer spatialize.
+	if (AllowAudioSpatializationCVar == 0)
+	{
+		ActiveSound->bAllowSpatialization = false;
+	}
+
 	// Set the active sound to be playing audio so it gets parsed at least once.
 	ActiveSound->bIsPlayingAudio = true;
 
