fix compile error on some platform and compiler

diff --git a/Engine/Source/Runtime/Engine/Private/ContentStreaming.cpp b/Engine/Source/Runtime/Engine/Private/ContentStreaming.cpp
--- a/Engine/Source/Runtime/Engine/Private/ContentStreaming.cpp
+++ b/Engine/Source/Runtime/Engine/Private/ContentStreaming.cpp
@@ -703,6 +703,34 @@ void IStreamingManager::Tick( float DeltaTime, bool bProcessEverything/*=false*/
 	bPendingRemoveViews = true;
 }
 
+/*-----------------------------------------------------------------------------
+	IRenderAssetStreamingManager implementation.
+-----------------------------------------------------------------------------*/
+
+void IRenderAssetStreamingManager::UpdateIndividualTexture(UTexture2D* Texture)
+{
+	UpdateIndividualRenderAsset(Texture);
+}
+
+bool IRenderAssetStreamingManager::StreamOutTextureData(int64 RequiredMemorySize)
+{
+	return StreamOutRenderAssetData(RequiredMemorySize);
+}
+
+void IRenderAssetStreamingManager::AddStreamingTexture(UTexture2D* Texture)
+{
+	AddStreamingRenderAsset(Texture);
+}
+
+void IRenderAssetStreamingManager::RemoveStreamingTexture(UTexture2D* Texture)
+{
+	RemoveStreamingRenderAsset(Texture);
+}
+
+void IRenderAssetStreamingManager::PauseTextureStreaming(bool bInShouldPause)
+{
+	PauseRenderAssetStreaming(bInShouldPause);
+}
 
 /*-----------------------------------------------------------------------------
 	FStreamingManagerCollection implementation.
diff --git a/Engine/Source/Runtime/Engine/Public/ContentStreaming.h b/Engine/Source/Runtime/Engine/Public/ContentStreaming.h
--- a/Engine/Source/Runtime/Engine/Public/ContentStreaming.h
+++ b/Engine/Source/Runtime/Engine/Public/ContentStreaming.h
@@ -390,11 +390,11 @@ struct IRenderAssetStreamingManager : public IStreamingManager
 	ENGINE_API virtual void GetObjectReferenceBounds(const UObject* RefObject, TArray<FBox>& AssetBoxes) = 0;
 
 	//BEGIN: APIs for backward compatibility
-	void UpdateIndividualTexture(UTexture2D* Texture) { UpdateIndividualRenderAsset(Texture); }
-	bool StreamOutTextureData(int64 RequiredMemorySize) { return StreamOutRenderAssetData(RequiredMemorySize); }
-	void AddStreamingTexture(UTexture2D* Texture) { AddStreamingRenderAsset(Texture); }
-	void RemoveStreamingTexture(UTexture2D* Texture) { RemoveStreamingRenderAsset(Texture); }
-	void PauseTextureStreaming(bool bInShouldPause) { PauseRenderAssetStreaming(bInShouldPause); }
+	void UpdateIndividualTexture(UTexture2D* Texture);
+	bool StreamOutTextureData(int64 RequiredMemorySize);
+	void AddStreamingTexture(UTexture2D* Texture);
+	void RemoveStreamingTexture(UTexture2D* Texture);
+	void PauseTextureStreaming(bool bInShouldPause);
 	//END: APIs for backward compatibility
 };
 
