Export some functions to fix link error

diff --git a/Engine/Source/Runtime/Engine/Public/ContentStreaming.h b/Engine/Source/Runtime/Engine/Public/ContentStreaming.h
--- a/Engine/Source/Runtime/Engine/Public/ContentStreaming.h
+++ b/Engine/Source/Runtime/Engine/Public/ContentStreaming.h
@@ -390,11 +390,11 @@ struct IRenderAssetStreamingManager : public IStreamingManager
 	ENGINE_API virtual void GetObjectReferenceBounds(const UObject* RefObject, TArray<FBox>& AssetBoxes) = 0;
 
 	//BEGIN: APIs for backward compatibility
-	void UpdateIndividualTexture(UTexture2D* Texture);
-	bool StreamOutTextureData(int64 RequiredMemorySize);
-	void AddStreamingTexture(UTexture2D* Texture);
-	void RemoveStreamingTexture(UTexture2D* Texture);
-	void PauseTextureStreaming(bool bInShouldPause);
+	ENGINE_API void UpdateIndividualTexture(UTexture2D* Texture);
+	ENGINE_API bool StreamOutTextureData(int64 RequiredMemorySize);
+	ENGINE_API void AddStreamingTexture(UTexture2D* Texture);
+	ENGINE_API void RemoveStreamingTexture(UTexture2D* Texture);
+	ENGINE_API void PauseTextureStreaming(bool bInShouldPause);
 	//END: APIs for backward compatibility
 };
 
