Fixed bug where CreateCanvasRenderTarget2D would assert when run from cmdlet (e.g. during packaging).

diff --git a/Engine/Source/Runtime/Engine/Private/CanvasRenderTarget2D.cpp b/Engine/Source/Runtime/Engine/Private/CanvasRenderTarget2D.cpp
--- a/Engine/Source/Runtime/Engine/Private/CanvasRenderTarget2D.cpp
+++ b/Engine/Source/Runtime/Engine/Private/CanvasRenderTarget2D.cpp
@@ -67,25 +67,29 @@ void UCanvasRenderTarget2D::RepaintCanvas()
 	const UWorld* WorldPtr = World.Get();
 	const ERHIFeatureLevel::Type FeatureLevel = WorldPtr != nullptr ? World->FeatureLevel.GetValue() : GMaxRHIFeatureLevel;
 
+	// NOTE: This texture may be null when this is invoked through blueprint from a cmdlet or server.
 	FTextureRenderTarget2DResource* TextureRenderTarget = (FTextureRenderTarget2DResource*) GameThread_GetRenderTargetResource();
 
 	FCanvas RenderCanvas(TextureRenderTarget, nullptr, FApp::GetCurrentTime() - GStartTime, FApp::GetDeltaTime(), FApp::GetCurrentTime() - GStartTime, FeatureLevel);
 	Canvas->Init(GetSurfaceWidth(), GetSurfaceHeight(), nullptr, &RenderCanvas);
 
-	// Enqueue the rendering command to set up the rendering canvas.
-	bool bClearRenderTarget = bShouldClearRenderTargetOnReceiveUpdate;
-	ENQUEUE_RENDER_COMMAND(CanvasRenderTargetMakeCurrentCommand)(
-		[TextureRenderTarget, bClearRenderTarget](FRHICommandListImmediate& RHICmdList)
+	if (TextureRenderTarget)
 	{
-		RHICmdList.TransitionResource(EResourceTransitionAccess::EWritable, TextureRenderTarget->GetRenderTargetTexture());
-
-		if (bClearRenderTarget)
+		// Enqueue the rendering command to set up the rendering canvas.
+		bool bClearRenderTarget = bShouldClearRenderTargetOnReceiveUpdate;
+		ENQUEUE_RENDER_COMMAND(CanvasRenderTargetMakeCurrentCommand)(
+			[TextureRenderTarget, bClearRenderTarget](FRHICommandListImmediate& RHICmdList)
 		{
-			FRHIRenderPassInfo RPInfo(TextureRenderTarget->GetRenderTargetTexture(), ERenderTargetActions::Clear_Store);
-			RHICmdList.BeginRenderPass(RPInfo, TEXT("ClearUCanvas"));
-			RHICmdList.EndRenderPass();
-		}
-	});
+			RHICmdList.TransitionResource(EResourceTransitionAccess::EWritable, TextureRenderTarget->GetRenderTargetTexture());
+
+			if (bClearRenderTarget)
+			{
+				FRHIRenderPassInfo RPInfo(TextureRenderTarget->GetRenderTargetTexture(), ERenderTargetActions::Clear_Store);
+				RHICmdList.BeginRenderPass(RPInfo, TEXT("ClearUCanvas"));
+				RHICmdList.EndRenderPass();
+			}
+		});
+	}
 
 	if (!IsPendingKill() && OnCanvasRenderTargetUpdate.IsBound())
 	{
@@ -96,7 +100,11 @@ void UCanvasRenderTarget2D::RepaintCanvas()
 
 	// Clean up and flush the rendering canvas.
 	Canvas->Canvas = nullptr;
-	RenderCanvas.Flush_GameThread();
+
+	if (TextureRenderTarget)
+	{
+		RenderCanvas.Flush_GameThread();
+	}
 
 	UpdateResourceImmediate(false);
 }
diff --git a/Engine/Source/Runtime/Engine/Private/UserInterface/Canvas.cpp b/Engine/Source/Runtime/Engine/Private/UserInterface/Canvas.cpp
--- a/Engine/Source/Runtime/Engine/Private/UserInterface/Canvas.cpp
+++ b/Engine/Source/Runtime/Engine/Private/UserInterface/Canvas.cpp
@@ -265,14 +265,14 @@ FCanvas::FCanvas(FRenderTarget* InRenderTarget,FHitProxyConsumer* InHitProxyCons
 
 void FCanvas::Construct()
 {
-	check(RenderTarget);
-
 	bStereoRendering = false;
 	bScaledToRenderTarget = false;
 	bAllowsToSwitchVerticalAxis = true;
 
+	const FIntPoint RenderTargetSizeXY = RenderTarget ? RenderTarget->GetSizeXY() : FIntPoint(1, 1);
+
 	new(TransformStack) FTransformEntry( 
-		FMatrix( FScaleMatrix(GetDPIScale()) * CalcBaseTransform2D(RenderTarget->GetSizeXY().X,RenderTarget->GetSizeXY().Y) ) 
+		FMatrix( FScaleMatrix(GetDPIScale()) * CalcBaseTransform2D(RenderTargetSizeXY.X, RenderTargetSizeXY.Y) )
 		);
 
 	// init alpha to 1
@@ -751,7 +751,7 @@ void FCanvas::Flush_GameThread(bool bForce)
 	}
 
 	// current render target set for the canvas
-	check(RenderTarget);	 	
+	check(RenderTarget);
 
 	// no need to set the render target if we aren't going to draw anything to it!
 	if (SortedElements.Num() == 0)
