preparation work for XMA2 streaming:
 - streamed decompression honours 'NumCompressedBytesConsumed'.
 - audio encoder indicates whether the sound is being cooked for streaming or not.
 - some XMA format checks.

diff --git a/Engine/Source/Developer/TargetPlatform/Public/Interfaces/IAudioFormat.h b/Engine/Source/Developer/TargetPlatform/Public/Interfaces/IAudioFormat.h
--- a/Engine/Source/Developer/TargetPlatform/Public/Interfaces/IAudioFormat.h
+++ b/Engine/Source/Developer/TargetPlatform/Public/Interfaces/IAudioFormat.h
@@ -25,6 +25,9 @@ struct FSoundQualityInfo
 	/** Holds the length of the sound in seconds. */
 	float Duration;
 
+	/** Holds whether the sound will be streamed. */
+	bool bStreaming;
+
 	/** Holds a string for debugging purposes. */
 	FString DebugName;
 };
diff --git a/Engine/Source/Runtime/Engine/Private/AudioDecompress.cpp b/Engine/Source/Runtime/Engine/Private/AudioDecompress.cpp
--- a/Engine/Source/Runtime/Engine/Private/AudioDecompress.cpp
+++ b/Engine/Source/Runtime/Engine/Private/AudioDecompress.cpp
@@ -290,6 +290,12 @@ int32 IStreamedCompressedInfo::DecompressToPCMBuffer(uint16 FrameSize)
 	LastPCMOffset = 0;
 	
 	const FDecodeResult DecodeResult = Decode(SrcPtr, FrameSize, LastDecodedPCM.GetData(), LastDecodedPCM.Num());
+	if (DecodeResult.NumCompressedBytesConsumed != INDEX_NONE )
+	{
+		SrcBufferOffset -= FrameSize;
+		SrcBufferOffset += DecodeResult.NumCompressedBytesConsumed;
+	}
+
 	return DecodeResult.NumAudioFramesProduced;
 }
 
diff --git a/Engine/Source/Runtime/Engine/Private/AudioDerivedData.cpp b/Engine/Source/Runtime/Engine/Private/AudioDerivedData.cpp
--- a/Engine/Source/Runtime/Engine/Private/AudioDerivedData.cpp
+++ b/Engine/Source/Runtime/Engine/Private/AudioDerivedData.cpp
@@ -958,6 +958,7 @@ static void CookSimpleWave(USoundWave* SoundWave, FName FormatName, const IAudio
 		QualityInfo.NumChannels = *WaveInfo.pChannels;
 		QualityInfo.SampleRate = WaveSampleRate;
 		QualityInfo.SampleDataSize = Input.Num();
+		QualityInfo.bStreaming = SoundWave->IsStreaming();
 		QualityInfo.DebugName = SoundWave->GetFullName();
 
 		// Cook the data.
@@ -1140,6 +1141,7 @@ static void CookSurroundWave( USoundWave* SoundWave, FName FormatName, const IAu
 			QualityInfo.NumChannels = ChannelCount;
 			QualityInfo.SampleRate = WaveSampleRate;
 			QualityInfo.SampleDataSize = SampleDataSize;
+			QualityInfo.bStreaming = SoundWave->IsStreaming();
 			QualityInfo.DebugName = SoundWave->GetFullName();
 			//@todo tighten up the checking for empty results here
 			if(Format.CookSurround(FormatName, SourceBuffers, QualityInfo, Output)) 
diff --git a/Engine/Source/Runtime/Windows/AudioMixerXAudio2/Private/AudioMixerPlatformXAudio2.cpp b/Engine/Source/Runtime/Windows/AudioMixerXAudio2/Private/AudioMixerPlatformXAudio2.cpp
--- a/Engine/Source/Runtime/Windows/AudioMixerXAudio2/Private/AudioMixerPlatformXAudio2.cpp
+++ b/Engine/Source/Runtime/Windows/AudioMixerXAudio2/Private/AudioMixerPlatformXAudio2.cpp
@@ -767,6 +767,14 @@ namespace Audio
 	{
 		static FName NAME_OGG(TEXT("OGG"));
 		static FName NAME_OPUS(TEXT("OPUS"));
+		static FName NAME_XMA(TEXT("XMA"));
+
+#if WITH_XMA2 && USE_XMA2_FOR_STREAMING
+		if (InSoundWave->IsStreaming() && InSoundWave->NumChannels <= 2)
+		{
+			return NAME_XMA;
+		}
+#endif
 
 		if (InSoundWave->IsStreaming())
 		{
@@ -780,7 +788,6 @@ namespace Audio
 #if WITH_XMA2
 		if (InSoundWave->NumChannels <= 2)
 		{
-			static FName NAME_XMA(TEXT("XMA"));
 			return NAME_XMA;
 		}
 #endif //#if WITH_XMA2
diff --git a/Engine/Source/Runtime/Windows/XAudio2/Private/XAudio2Device.cpp b/Engine/Source/Runtime/Windows/XAudio2/Private/XAudio2Device.cpp
--- a/Engine/Source/Runtime/Windows/XAudio2/Private/XAudio2Device.cpp
+++ b/Engine/Source/Runtime/Windows/XAudio2/Private/XAudio2Device.cpp
@@ -351,6 +351,19 @@ class ICompressedAudioInfo* FXAudio2Device::CreateCompressedAudioInfo(USoundWave
 {
 	check(SoundWave);
 
+#if WITH_XMA2 && USE_XMA2_FOR_STREAMING
+	if (SoundWave->IsStreaming() && SoundWave->NumChannels <= 2 )
+	{
+		ICompressedAudioInfo* CompressedInfo = new FXMAAudioInfo();
+		if (!CompressedInfo)
+		{
+			UE_LOG(LogAudio, Error, TEXT("Failed to create new FXMAAudioInfo for streaming SoundWave %s: out of memory."), *SoundWave->GetName());
+			return nullptr;
+		}
+		return CompressedInfo;
+	}
+#endif
+
 	if (SoundWave->IsStreaming())
 	{
 #if USE_VORBIS_FOR_STREAMING
diff --git a/Engine/Source/Runtime/Windows/XAudio2/Public/XAudio2Device.h b/Engine/Source/Runtime/Windows/XAudio2/Public/XAudio2Device.h
--- a/Engine/Source/Runtime/Windows/XAudio2/Public/XAudio2Device.h
+++ b/Engine/Source/Runtime/Windows/XAudio2/Public/XAudio2Device.h
@@ -120,6 +120,15 @@ private:
 	{
 		static FName NAME_OPUS(TEXT("OPUS"));
 		static FName NAME_OGG(TEXT("OGG"));
+		static FName NAME_XMA(TEXT("XMA"));
+
+#if WITH_XMA2 && USE_XMA2_FOR_STREAMING
+		if (SoundWave->IsStreaming() && SoundWave->NumChannels <= 2)
+		{
+			return NAME_XMA;
+		}
+#endif
+
 
 #if WITH_OGGVORBIS
 
@@ -132,7 +141,6 @@ private:
 #endif
 		}
 
-		static FName NAME_XMA(TEXT("XMA"));
 
 #if WITH_XMA2
 		if (SoundWave->NumChannels > 2)
@@ -149,7 +157,6 @@ private:
 #endif
 
 #else //WITH_OGGVORBIS
-		static FName NAME_XMA(TEXT("XMA"));
 		return NAME_XMA;
 #endif //WITH_OGGVORBIS
 	}
