Reduce FBodyInstance by 56 bytes in shipping

diff --git a/Engine/Plugins/Runtime/PhysXVehicles/Source/PhysXVehicles/Private/WheeledVehicleMovementComponent.cpp b/Engine/Plugins/Runtime/PhysXVehicles/Source/PhysXVehicles/Private/WheeledVehicleMovementComponent.cpp
--- a/Engine/Plugins/Runtime/PhysXVehicles/Source/PhysXVehicles/Private/WheeledVehicleMovementComponent.cpp
+++ b/Engine/Plugins/Runtime/PhysXVehicles/Source/PhysXVehicles/Private/WheeledVehicleMovementComponent.cpp
@@ -417,7 +417,7 @@ void UWheeledVehicleMovementComponent::SetupVehicleMass()
 	}
 
 	//Ensure that if mass properties ever change we set them back to our override
-	UpdatedPrimitive->GetBodyInstance()->OnRecalculatedMassProperties.AddUObject(this, &UWheeledVehicleMovementComponent::UpdateMassProperties);
+	UpdatedPrimitive->GetBodyInstance()->OnRecalculatedMassProperties().AddUObject(this, &UWheeledVehicleMovementComponent::UpdateMassProperties);
 
 	UpdateMassProperties(UpdatedPrimitive->GetBodyInstance());
 }
diff --git a/Engine/Source/Runtime/Engine/Classes/PhysicsEngine/BodyInstance.h b/Engine/Source/Runtime/Engine/Classes/PhysicsEngine/BodyInstance.h
--- a/Engine/Source/Runtime/Engine/Classes/PhysicsEngine/BodyInstance.h
+++ b/Engine/Source/Runtime/Engine/Classes/PhysicsEngine/BodyInstance.h
@@ -120,6 +120,8 @@ enum class BodyInstanceSceneState : uint8
 	Removed
 };
 
+#define USE_BODYINSTANCE_DEBUG_NAMES (WITH_EDITORONLY_DATA || UE_BUILD_DEBUG || LOOKING_FOR_PERF_ISSUES) && !(UE_BUILD_SHIPPING || UE_BUILD_TEST) && !NO_LOGGING
+
 /** Container for a physics representation of an object */
 USTRUCT(BlueprintType)
 struct ENGINE_API FBodyInstance
@@ -303,6 +305,17 @@ private:
 	UPROPERTY(EditAnywhere, Category=Custom)
 	FName CollisionProfileName;
 
+public:
+
+	/** This physics body's solver iteration count for position. Increasing this will be more CPU intensive, but better stabilized.  */
+	UPROPERTY(EditAnywhere, AdvancedDisplay, BlueprintReadOnly, Category=Physics)
+	uint8 PositionSolverIterationCount;
+
+	/** This physics body's solver iteration count for velocity. Increasing this will be more CPU intensive, but better stabilized. */
+	UPROPERTY(EditAnywhere, AdvancedDisplay, BlueprintReadOnly, Category = Physics)
+	uint8 VelocitySolverIterationCount;
+
+private:
 	/** Custom Channels for Responses*/
 	UPROPERTY(EditAnywhere, Category = Custom)
 	struct FCollisionResponse CollisionResponses;
@@ -412,14 +425,6 @@ public:
 	UPROPERTY()
 	float PhysicsBlendWeight;
 
-	/** This physics body's solver iteration count for position. Increasing this will be more CPU intensive, but better stabilized.  */
-	UPROPERTY(EditAnywhere, AdvancedDisplay, BlueprintReadOnly, Category=Physics)
-	int32 PositionSolverIterationCount;
-
-	/** This physics body's solver iteration count for velocity. Increasing this will be more CPU intensive, but better stabilized. */
-	UPROPERTY(EditAnywhere, AdvancedDisplay, BlueprintReadOnly, Category = Physics)
-	int32 VelocitySolverIterationCount;
-
 public:
 
 	FPhysicsActorHandle& GetPhysicsActorHandle();
@@ -429,7 +434,9 @@ public:
 	// Internal physics representation of our body instance
 	FPhysicsActorHandle ActorHandle;
 
+#if USE_BODYINSTANCE_DEBUG_NAMES
 	TSharedPtr<TArray<ANSICHAR>> CharDebugName;
+#endif
 
 	/** PrimitiveComponent containing this body.   */
 	TWeakObjectPtr<class UPrimitiveComponent> OwnerComponent;
@@ -668,11 +675,62 @@ public:
 	/** Enable/disable Continuous Collidion Detection feature */
 	void SetUseCCD(bool bInUseCCD);
 
-	/** Custom projection for physics (callback to update component transform based on physics data) */
-	FCalculateCustomProjection OnCalculateCustomProjection;
+private:
+
+	/** Struct of body instance delegates that are rarely bound so that we can only allocate memory if one is actually being used. */
+	struct FBodyInstanceDelegates
+	{
+		/** Custom projection for physics (callback to update component transform based on physics data) */
+		FCalculateCustomProjection OnCalculateCustomProjection;
+
+		/** Called whenever mass properties have been re-calculated. */
+		FRecalculatedMassProperties OnRecalculatedMassProperties;
+	};
+
+	/** Specialization of TUniquePtr for storing body instance delegates as there is a need to copy the contents of the delegate structure. */
+	struct FBodyInstanceDelegatesPtr : public TUniquePtr<FBodyInstanceDelegates>
+	{
+		FBodyInstanceDelegatesPtr() = default;
+		FBodyInstanceDelegatesPtr(const FBodyInstanceDelegatesPtr& Other)
+		{
+			if (Other.IsValid())
+			{
+				Reset(new FBodyInstanceDelegates);
+				*Get() = *Other;
+			}
+		}
+
+		FBodyInstanceDelegatesPtr& operator=(const FBodyInstanceDelegatesPtr& Other)
+		{
+			if (Other.IsValid())
+			{
+				if (!IsValid())
+				{
+					Reset(new FBodyInstanceDelegates);
+				}
+				*Get() = *Other;
+			}
+			else
+			{
+				Reset();
+			}
+
+			return *this;
+		}
+	};
+
+	/** Pointer to lazily created container for the body instance delegates. */
+	FBodyInstanceDelegatesPtr BodyInstanceDelegates;
+
+public:
+	/** Executes the OnCalculateCustomProjection delegate if bound. */
+	void ExecuteOnCalculateCustomProjection(FTransform& WorldTM) const;
+
+	/** Returns reference to the OnCalculateCustomProjection delegate. Will allocate delegate struct if not already created. */
+	FCalculateCustomProjection& OnCalculateCustomProjection();
 
-	/** Called whenever mass properties have been re-calculated. */
-	FRecalculatedMassProperties OnRecalculatedMassProperties;
+	/** Returns reference to the OnRecalculatedMassProperties delegate. Will allocate delegate struct if not already created. */
+	FRecalculatedMassProperties& OnRecalculatedMassProperties();
 
 	/** See if this body is valid. */
 	bool IsValidBodyInstance() const;
diff --git a/Engine/Source/Runtime/Engine/Private/PhysicsEngine/BodyInstance.cpp b/Engine/Source/Runtime/Engine/Private/PhysicsEngine/BodyInstance.cpp
--- a/Engine/Source/Runtime/Engine/Private/PhysicsEngine/BodyInstance.cpp
+++ b/Engine/Source/Runtime/Engine/Private/PhysicsEngine/BodyInstance.cpp
@@ -30,10 +30,10 @@
 #include "Misc/UObjectToken.h"
 
 #if WITH_PHYSX
-#include "PhysXPublic.h"
-#include "Physics/PhysicsFiltering.h"
-#include "PhysicsEngine/PhysXSupport.h"
-#include "Collision/CollisionConversions.h"
+	#include "PhysXPublic.h"
+	#include "Physics/PhysicsFiltering.h"
+	#include "PhysicsEngine/PhysXSupport.h"
+	#include "Collision/CollisionConversions.h"
 #include "PxShape.h"
 #endif // WITH_PHYSX
 
@@ -108,7 +108,7 @@ void FCollisionResponse::SetResponse(ECollisionChannel Channel, ECollisionRespon
 {
 #if 1// @hack until PostLoad is disabled for CDO of BP WITH_EDITOR
 	ECollisionResponse DefaultResponse = FCollisionResponseContainer::GetDefaultResponseContainer().GetResponse(Channel);
-	if(DefaultResponse == NewResponse)
+	if (DefaultResponse == NewResponse)
 	{
 		RemoveReponseFromArray(Channel);
 	}
@@ -153,11 +153,11 @@ void FCollisionResponse::SetResponsesArray(const TArray<FResponseChannel>& InCha
 #if DO_GUARD_SLOW
 	// verify if the name is overlapping, if so, ensure, do not remove in debug because it will cause inconsistent bug between debug/release
 	int32 const ResponseNum = InChannelResponses.Num();
-	for(int32 I = 0; I < ResponseNum; ++I)
+	for (int32 I=0; I<ResponseNum; ++I)
 	{
-		for(int32 J = I + 1; J < ResponseNum; ++J)
+		for (int32 J=I+1; J<ResponseNum; ++J)
 		{
-			if(InChannelResponses[I].Channel == InChannelResponses[J].Channel)
+			if (InChannelResponses[I].Channel == InChannelResponses[J].Channel)
 			{
 				UE_LOG(LogCollision, Warning, TEXT("Collision Channel : Redundant name exists"));
 			}
@@ -177,9 +177,9 @@ bool FCollisionResponse::RemoveReponseFromArray(ECollisionChannel Channel)
 	// without editor, this does not have to match 
 	// We'd need to save name just in case that name is gone or not
 	FName ChannelName = UCollisionProfile::Get()->ReturnChannelNameFromContainerIndex(Channel);
-	for(auto Iter = ResponseArray.CreateIterator(); Iter; ++Iter)
+	for (auto Iter=ResponseArray.CreateIterator(); Iter; ++Iter)
 	{
-		if(ChannelName == (*Iter).Channel)
+		if (ChannelName == (*Iter).Channel)
 		{
 			ResponseArray.RemoveAt(Iter.GetIndex());
 			return true;
@@ -194,9 +194,9 @@ bool FCollisionResponse::AddReponseToArray(ECollisionChannel Channel, ECollision
 	// so only allow it in editor
 	// without editor, this does not have to match 
 	FName ChannelName = UCollisionProfile::Get()->ReturnChannelNameFromContainerIndex(Channel);
-	for(auto Iter = ResponseArray.CreateIterator(); Iter; ++Iter)
+	for (auto Iter=ResponseArray.CreateIterator(); Iter; ++Iter)
 	{
-		if(ChannelName == (*Iter).Channel)
+		if (ChannelName == (*Iter).Channel)
 		{
 			(*Iter).Response = Response;
 			return true;
@@ -215,13 +215,13 @@ void FCollisionResponse::UpdateArrayFromResponseContainer()
 	const FCollisionResponseContainer& DefaultResponse = FCollisionResponseContainer::GetDefaultResponseContainer();
 	const UCollisionProfile* CollisionProfile = UCollisionProfile::Get();
 
-	for(int32 i = 0; i < ARRAY_COUNT(ResponseToChannels.EnumArray); i++)
+	for (int32 i = 0; i < ARRAY_COUNT(ResponseToChannels.EnumArray); i++)
 	{
 		// if not same as default
-		if(ResponseToChannels.EnumArray[i] != DefaultResponse.EnumArray[i])
+		if (ResponseToChannels.EnumArray[i] != DefaultResponse.EnumArray[i])
 		{
 			FName ChannelName = CollisionProfile->ReturnChannelNameFromContainerIndex(i);
-			if(ChannelName != NAME_None)
+			if (ChannelName != NAME_None)
 			{
 				FResponseChannel NewResponse;
 				NewResponse.Channel = ChannelName;
@@ -238,12 +238,12 @@ void FCollisionResponse::UpdateResponseContainerFromArray()
 {
 	ResponseToChannels = FCollisionResponseContainer::GetDefaultResponseContainer();
 
-	for(auto Iter = ResponseArray.CreateIterator(); Iter; ++Iter)
+	for (auto Iter = ResponseArray.CreateIterator(); Iter; ++Iter)
 	{
 		FResponseChannel& Response = *Iter;
 
 		int32 EnumIndex = UCollisionProfile::Get()->ReturnContainerIndexFromChannelName(Response.Channel);
-		if(EnumIndex != INDEX_NONE)
+		if ( EnumIndex != INDEX_NONE )
 		{
 			ResponseToChannels.SetResponse((ECollisionChannel)EnumIndex, Response.Response);
 		}
@@ -271,7 +271,7 @@ bool FCollisionResponse::operator==(const FCollisionResponse& Other) const
 					break;
 				}
 			}
-
+			
 		}
 	}
 
@@ -341,15 +341,15 @@ const FPhysicsActorHandle& FBodyInstance::GetActorReferenceWithWelding() const
 	return WeldParent ? WeldParent->ActorHandle : ActorHandle;
 }
 
-FArchive& operator<<(FArchive& Ar, FBodyInstance& BodyInst)
+FArchive& operator<<(FArchive& Ar,FBodyInstance& BodyInst)
 {
-	if(!Ar.IsLoading() && !Ar.IsSaving())
+	if (!Ar.IsLoading() && !Ar.IsSaving())
 	{
 		Ar << BodyInst.OwnerComponent;
 		Ar << BodyInst.PhysMaterialOverride;
 	}
 
-	if(Ar.IsLoading() && Ar.UE4Ver() < VER_UE4_MAX_ANGULAR_VELOCITY_DEFAULT)
+	if (Ar.IsLoading() && Ar.UE4Ver() < VER_UE4_MAX_ANGULAR_VELOCITY_DEFAULT)
 	{
 		if(BodyInst.MaxAngularVelocity != 400.f)
 		{
@@ -402,10 +402,10 @@ bool FBodyInstance::IsShapeBoundToBody(const FPhysicsShapeHandle& Shape) const
 const TMap<FPhysicsShapeHandle, FBodyInstance::FWeldInfo>* FBodyInstance::GetCurrentWeldInfo() const
 {
 	return ShapeToBodiesMap.Get();
-}
+	}
 
 int32 FBodyInstance::GetAllShapes_AssumesLocked(TArray<FPhysicsShapeHandle>& OutShapes) const
-{
+	{
 	return FPhysicsInterface::GetAllShapes_AssumedLocked(ActorHandle, OutShapes);
 }
 
@@ -415,7 +415,7 @@ void FBodyInstance::UpdateTriMeshVertices(const TArray<FVector> & NewPositions)
 	check(false);
 #else
 #if WITH_PHYSX
-	if(BodySetup.IsValid())
+	if (BodySetup.IsValid())
 	{
 		FPhysicsCommand::ExecuteWrite(ActorHandle, [&](const FPhysicsActorHandle& Actor)
 		{
@@ -473,7 +473,7 @@ void FBodyInstance::SetResponseToAllChannels(ECollisionResponse NewResponse)
 	CollisionResponses.SetAllChannels(NewResponse);
 	UpdatePhysicsFilterData();
 }
-
+	
 void FBodyInstance::ReplaceResponseToChannels(ECollisionResponse OldResponse, ECollisionResponse NewResponse)
 {
 	InvalidateCollisionProfileName();
@@ -518,7 +518,7 @@ void FBodyInstance::SetCollisionProfileName(FName InCollisionProfileName)
 	//Note that GetCollisionProfileName will use the external profile if one is set.
 	//GetCollisionProfileName will be consistent with the values set by LoadProfileData.
 	//This is why we can't use CollisionProfileName directly during the equality check
-	if(bPendingCollisionProfileSetup || GetCollisionProfileName() != InCollisionProfileName)
+	if (bPendingCollisionProfileSetup || GetCollisionProfileName() != InCollisionProfileName)
 	{
 		//LoadProfileData uses GetCollisionProfileName internally so we must now set the external collision data to null.
 		ExternalCollisionProfileBodySetup = nullptr;
@@ -528,18 +528,18 @@ void FBodyInstance::SetCollisionProfileName(FName InCollisionProfileName)
 
 		bPendingCollisionProfileSetup = false;
 	}
-
+	
 	ExternalCollisionProfileBodySetup = nullptr;	//Even if incoming is the same as GetCollisionProfileName we turn it into "manual mode"
 }
 
 FName FBodyInstance::GetCollisionProfileName() const
 {
 	FName ReturnProfileName = CollisionProfileName;
-	if(UBodySetup* BodySetupPtr = ExternalCollisionProfileBodySetup.Get(true))
+	if (UBodySetup* BodySetupPtr = ExternalCollisionProfileBodySetup.Get(true))
 	{
 		ReturnProfileName = BodySetupPtr->DefaultInstance.CollisionProfileName;
 	}
-
+	
 	return ReturnProfileName;
 }
 
@@ -557,14 +557,14 @@ void FBodyInstance::SetMassScale(float InMassScale)
 
 void FBodyInstance::SetCollisionEnabled(ECollisionEnabled::Type NewType, bool bUpdatePhysicsFilterData)
 {
-	if(CollisionEnabled != NewType)
+	if (CollisionEnabled != NewType)
 	{
 		ECollisionEnabled::Type OldType = CollisionEnabled;
 		InvalidateCollisionProfileName();
 		CollisionEnabled = NewType;
-
+		
 		// Only update physics filter data if required
-		if(bUpdatePhysicsFilterData)
+		if (bUpdatePhysicsFilterData)
 		{
 			UpdatePhysicsFilterData();
 		}
@@ -574,7 +574,7 @@ void FBodyInstance::SetCollisionEnabled(ECollisionEnabled::Type NewType, bool bU
 
 		// Whenever we change physics state, call Recreate
 		// This should also handle destroying the state (in case it's newly disabled).
-		if(bWasPhysicsEnabled != bIsPhysicsEnabled)
+		if (bWasPhysicsEnabled != bIsPhysicsEnabled)
 		{
 			if(UPrimitiveComponent* PrimComponent = OwnerComponent.Get())
 			{
@@ -589,13 +589,13 @@ void FBodyInstance::SetCollisionEnabled(ECollisionEnabled::Type NewType, bool bU
 EDOFMode::Type FBodyInstance::ResolveDOFMode(EDOFMode::Type DOFMode)
 {
 	EDOFMode::Type ResultDOF = DOFMode;
-	if(DOFMode == EDOFMode::Default)
+	if (DOFMode == EDOFMode::Default)
 	{
 		ESettingsDOF::Type SettingDefaultPlane = UPhysicsSettings::Get()->DefaultDegreesOfFreedom;
-		if(SettingDefaultPlane == ESettingsDOF::XYPlane) ResultDOF = EDOFMode::XYPlane;
-		if(SettingDefaultPlane == ESettingsDOF::XZPlane) ResultDOF = EDOFMode::XZPlane;
-		if(SettingDefaultPlane == ESettingsDOF::YZPlane) ResultDOF = EDOFMode::YZPlane;
-		if(SettingDefaultPlane == ESettingsDOF::Full3D) ResultDOF = EDOFMode::SixDOF;
+		if (SettingDefaultPlane == ESettingsDOF::XYPlane) ResultDOF = EDOFMode::XYPlane;
+		if (SettingDefaultPlane == ESettingsDOF::XZPlane) ResultDOF = EDOFMode::XZPlane;
+		if (SettingDefaultPlane == ESettingsDOF::YZPlane) ResultDOF = EDOFMode::YZPlane;
+		if (SettingDefaultPlane == ESettingsDOF::Full3D) ResultDOF  = EDOFMode::SixDOF;
 	}
 
 	return ResultDOF;
@@ -605,15 +605,15 @@ FVector FBodyInstance::GetLockedAxis() const
 {
 	EDOFMode::Type MyDOFMode = ResolveDOFMode(DOFMode);
 
-	switch(MyDOFMode)
+	switch (MyDOFMode)
 	{
-		case EDOFMode::None: return FVector::ZeroVector;
-		case EDOFMode::YZPlane: return FVector(1, 0, 0);
-		case EDOFMode::XZPlane: return FVector(0, 1, 0);
-		case EDOFMode::XYPlane: return FVector(0, 0, 1);
-		case EDOFMode::CustomPlane: return CustomDOFPlaneNormal;
-		case EDOFMode::SixDOF: return FVector::ZeroVector;
-		default:	check(0);	//unsupported locked axis type
+	case EDOFMode::None: return FVector::ZeroVector;
+	case EDOFMode::YZPlane: return FVector(1, 0, 0);
+	case EDOFMode::XZPlane: return FVector(0, 1, 0);
+	case EDOFMode::XYPlane: return FVector(0, 0, 1);
+	case EDOFMode::CustomPlane: return CustomDOFPlaneNormal;
+	case EDOFMode::SixDOF: return FVector::ZeroVector;
+	default:	check(0);	//unsupported locked axis type
 	}
 
 	return FVector::ZeroVector;
@@ -642,7 +642,7 @@ void FBodyInstance::SetDOFLock(EDOFMode::Type NewAxisMode)
 
 void FBodyInstance::CreateDOFLock()
 {
-	if(DOFConstraint)
+	if (DOFConstraint)
 	{
 		DOFConstraint->TermConstraint();
 		FConstraintInstance::Free(DOFConstraint);
@@ -652,13 +652,13 @@ void FBodyInstance::CreateDOFLock()
 	const FVector LockedAxis = GetLockedAxis();
 	const EDOFMode::Type DOF = ResolveDOFMode(DOFMode);
 
-	if(IsDynamic() == false || (LockedAxis.IsNearlyZero() && DOF != EDOFMode::SixDOF))
+	if (IsDynamic() == false || (LockedAxis.IsNearlyZero() && DOF != EDOFMode::SixDOF))
 	{
 		return;
 	}
 
 	//if we're using SixDOF make sure we have at least one constraint
-	if(DOF == EDOFMode::SixDOF && !bLockXTranslation && !bLockYTranslation && !bLockZTranslation && !bLockXRotation && !bLockYRotation && !bLockZRotation)
+	if (DOF == EDOFMode::SixDOF && !bLockXTranslation && !bLockYTranslation && !bLockZTranslation && !bLockXRotation && !bLockYRotation && !bLockZRotation)
 	{
 		return;
 	}
@@ -666,8 +666,8 @@ void FBodyInstance::CreateDOFLock()
 	DOFConstraint = FConstraintInstance::Alloc();
 	{
 		DOFConstraint->ProfileInstance.ConeLimit.bSoftConstraint = false;
-		DOFConstraint->ProfileInstance.TwistLimit.bSoftConstraint = false;
-		DOFConstraint->ProfileInstance.LinearLimit.bSoftConstraint = false;
+		DOFConstraint->ProfileInstance.TwistLimit.bSoftConstraint  = false;
+		DOFConstraint->ProfileInstance.LinearLimit.bSoftConstraint  = false;
 
 		const FTransform TM = GetUnrealWorldTransform(false);
 		FVector Normal = FVector(1, 0, 0);
@@ -679,7 +679,7 @@ void FBodyInstance::CreateDOFLock()
 			DOFConstraint->SetAngularSwing1Motion((bLockRotation || DOFMode != EDOFMode::CustomPlane) ? EAngularConstraintMotion::ACM_Locked : EAngularConstraintMotion::ACM_Free);
 			DOFConstraint->SetAngularSwing2Motion((bLockRotation || DOFMode != EDOFMode::CustomPlane) ? EAngularConstraintMotion::ACM_Locked : EAngularConstraintMotion::ACM_Free);
 			DOFConstraint->SetAngularTwistMotion(EAngularConstraintMotion::ACM_Free);
-
+			
 			DOFConstraint->SetLinearXMotion((bLockTranslation || DOFMode != EDOFMode::CustomPlane) ? ELinearConstraintMotion::LCM_Locked : ELinearConstraintMotion::LCM_Free);
 			DOFConstraint->SetLinearYMotion(ELinearConstraintMotion::LCM_Free);
 			DOFConstraint->SetLinearZMotion(ELinearConstraintMotion::LCM_Free);
@@ -716,24 +716,24 @@ ECollisionEnabled::Type FBodyInstance::GetCollisionEnabled_CheckOwner() const
 	// Check actor override
 	const UPrimitiveComponent* OwnerComponentInst = OwnerComponent.Get();
 	AActor* Owner = OwnerComponentInst ? OwnerComponentInst->GetOwner() : nullptr;
-	if(Owner && !Owner->GetActorEnableCollision())
-	{
-		return ECollisionEnabled::NoCollision;
-	}
-	else if(const USkeletalMeshComponent* SkelMeshComp = Cast<USkeletalMeshComponent>(OwnerComponentInst))
-	{
-		// Check component override (skel mesh case)
-		return (ECollisionEnabled::Type)FMath::Max(SkelMeshComp->BodyInstance.CollisionEnabled, CollisionEnabled);
-	}
+		if (Owner && !Owner->GetActorEnableCollision())
+		{
+			return ECollisionEnabled::NoCollision;
+		}
+		else if(const USkeletalMeshComponent* SkelMeshComp = Cast<USkeletalMeshComponent>(OwnerComponentInst))
+{
+			// Check component override (skel mesh case)
+			return (ECollisionEnabled::Type)FMath::Max(SkelMeshComp->BodyInstance.CollisionEnabled, CollisionEnabled);
+}
 	else
 	{
 		return CollisionEnabled;
 	}
-}
+			}
 
 void FBodyInstance::SetMaskFilter(FMaskFilter InMaskFilter)
 {
-	if(MaskFilter == InMaskFilter)
+	if (MaskFilter == InMaskFilter)
 	{
 		return;
 	}
@@ -748,13 +748,13 @@ void FBodyInstance::SetMaskFilter(FMaskFilter InMaskFilter)
 			const FBodyInstance* BI = GetOriginalBodyInstance(Shape);
 
 			if(BI == this)
-			{
-				FPhysicsCommand::ExecuteShapeWrite(this, Shape, [&](const FPhysicsShapeHandle& InnerShape)
 				{
+				FPhysicsCommand::ExecuteShapeWrite(this, Shape, [&](const FPhysicsShapeHandle& InnerShape)
+					{
 					FPhysicsInterface::SetMaskFilter(InnerShape, InMaskFilter);
-				});
+					});
+				}
 			}
-		}
 	});
 
 	MaskFilter = InMaskFilter;
@@ -772,13 +772,13 @@ void FBodyInstance::UpdatePhysicsFilterData()
 	}
 
 	// Do nothing if no physics actor
-	if(!IsValidBodyInstance())
+	if (!IsValidBodyInstance())
 	{
 		return;
 	}
 
 	// this can happen in landscape height field collision component
-	if(!BodySetup.IsValid())
+	if (!BodySetup.IsValid())
 	{
 		return;
 	}
@@ -804,7 +804,7 @@ void FBodyInstance::UpdatePhysicsFilterData()
 			// we need to generate new filter data using the owning original instance (and its BodySetup) 
 			FBodyCollisionData PerShapeCollisionData;
 			if(BI != this)
-			{
+	{
 				BI->BuildBodyFilterData(PerShapeCollisionData.CollisionFilterData);
 
 				const bool bInstanceComplexAsSimple = BI->BodySetup.IsValid() ? (BI->BodySetup->GetCollisionTraceFlag() == CTF_UseComplexAsSimple) : false;
@@ -813,7 +813,7 @@ void FBodyInstance::UpdatePhysicsFilterData()
 			else
 			{
 				PerShapeCollisionData = BodyCollisionData;
-			}
+		}
 
 			FPhysicsCommand::ExecuteShapeWrite(this, Shape, [&](const FPhysicsShapeHandle& InnerShape)
 			{
@@ -833,18 +833,18 @@ void FBodyInstance::UpdatePhysicsFilterData()
 
 				// If we changed 'simulation collision' on a shape, we need to recalc mass properties
 				if(bWasSimulationShape != bNewSimShape)
-				{
+		{
 					bUpdateMassProperties = true;
-				}
+		}
 
 				// Apply new collision settings to this shape
 				FPhysicsInterface::SetSimulationFilter(InnerShape, FilterData.SimFilter);
 				FPhysicsInterface::SetQueryFilter(InnerShape, bIsTrimesh ? FilterData.QueryComplexFilter : FilterData.QuerySimpleFilter);
 			});
-		}
+	}
 
 		if(bUpdateMassProperties)
-		{
+	{
 			UpdateMassProperties();
 		}
 	});
@@ -854,14 +854,14 @@ void FBodyInstance::UpdatePhysicsFilterData()
 
 TAutoConsoleVariable<int32> CDisableQueryOnlyActors(TEXT("p.DisableQueryOnlyActors"), 0, TEXT("If QueryOnly is used, actors are marked as simulation disabled. This is NOT compatible with origin shifting at the moment."));
 
+#if USE_BODYINSTANCE_DEBUG_NAMES
 TSharedPtr<TArray<ANSICHAR>> GetDebugDebugName(const UPrimitiveComponent* PrimitiveComp, const UBodySetup* BodySetup, FString& DebugName)
 {
 	// Setup names
 	// Make the debug name for this geometry...
-	DebugName = (TEXT(""));
+	DebugName.Reset();
 	TSharedPtr<TArray<ANSICHAR>> PhysXName;
 
-#if (WITH_EDITORONLY_DATA || UE_BUILD_DEBUG || LOOKING_FOR_PERF_ISSUES) && !(UE_BUILD_SHIPPING || UE_BUILD_TEST) && !NO_LOGGING
 	if (PrimitiveComp)
 	{
 		DebugName += FString::Printf(TEXT("Component: '%s' "), *PrimitiveComp->GetPathName());
@@ -874,10 +874,10 @@ TSharedPtr<TArray<ANSICHAR>> GetDebugDebugName(const UPrimitiveComponent* Primit
 
 	// Convert to char* for PhysX
 	PhysXName = MakeShareable(new TArray<ANSICHAR>(StringToArray<ANSICHAR>(*DebugName, DebugName.Len() + 1)));
-#endif
 
 	return PhysXName;
 }
+#endif
 
 void GetSimulatingAndBlendWeight(const USkeletalMeshComponent* SkelMeshComp, const UBodySetup* BodySetup, float& InstanceBlendWeight, bool& bInstanceSimulatePhysics)
 {
@@ -922,21 +922,22 @@ template <bool bCompileStatic>
 struct FInitBodiesHelper
 {
 	FInitBodiesHelper(TArray<FBodyInstance*>& InBodies, TArray<FTransform>& InTransforms, class UBodySetup* InBodySetup, class UPrimitiveComponent* InPrimitiveComp, FPhysScene* InRBScene, const FBodyInstance::FInitBodySpawnParams& InSpawnParams, FPhysicsAggregateHandle InAggregate)
-		: Bodies(InBodies)
-		, Transforms(InTransforms)
-		, BodySetup(InBodySetup)
-		, PrimitiveComp(InPrimitiveComp)
+	: Bodies(InBodies)
+	, Transforms(InTransforms)
+	, BodySetup(InBodySetup)
+	, PrimitiveComp(InPrimitiveComp)
 		, PhysScene(InRBScene)
 		, Aggregate(InAggregate)
-		, DebugName(TEXT(""))
-		, InstanceBlendWeight(-1.f)
 		, bInstanceSimulatePhysics(false)
-		, SkelMeshComp(nullptr)
-		, SpawnParams(InSpawnParams)
+	, InstanceBlendWeight(-1.f)
+	, SkelMeshComp(nullptr)
+	, SpawnParams(InSpawnParams)
 	{
 		//Compute all the needed constants
 
+#if USE_BODYINSTANCE_DEBUG_NAMES
 		PhysXName = GetDebugDebugName(PrimitiveComp, BodySetup, DebugName);
+#endif
 
 		bStatic = bCompileStatic || SpawnParams.bStaticPhysics;
 		SkelMeshComp = bCompileStatic ? nullptr : Cast<USkeletalMeshComponent>(PrimitiveComp);
@@ -944,26 +945,28 @@ struct FInitBodiesHelper
 		{
 			GetSimulatingAndBlendWeight(SkelMeshComp, BodySetup, InstanceBlendWeight, bInstanceSimulatePhysics);
 		}
-
+		
 	}
 
 	FORCEINLINE bool IsStatic() const { return bCompileStatic || bStatic; }
 
 	//The arguments passed into InitBodies
-	TArray<FBodyInstance*>& Bodies;
+	TArray<FBodyInstance*>& Bodies;   
 	TArray<FTransform>& Transforms;
 	class UBodySetup* BodySetup;
 	class UPrimitiveComponent* PrimitiveComp;
 	FPhysScene* PhysScene;
 	FPhysicsAggregateHandle Aggregate;
 
+#if USE_BODYINSTANCE_DEBUG_NAMES
 	FString DebugName;
 	TSharedPtr<TArray<ANSICHAR>> PhysXName; // Get rid of ANSICHAR in physics
+#endif
 
 	//The constants shared between PhysX and Box2D
 	bool bStatic;
-	float InstanceBlendWeight;
 	bool bInstanceSimulatePhysics;
+	float InstanceBlendWeight;
 
 	const USkeletalMeshComponent* SkelMeshComp;
 
@@ -978,12 +981,14 @@ struct FInitBodiesHelper
 
 		FActorCreationParams ActorParams;
 		ActorParams.InitialTM = Transform;
+#if USE_BODYINSTANCE_DEBUG_NAMES
 		ActorParams.DebugName = Instance->CharDebugName.IsValid() ? Instance->CharDebugName->GetData() : nullptr;
+#endif
 		ActorParams.bEnableGravity = Instance->bEnableGravity;
 		ActorParams.bQueryOnly = bDisableSim;
 		ActorParams.Scene = PhysScene;
 
-		if(IsStatic())
+		if (IsStatic())
 		{
 			ActorParams.bStatic = true;
 
@@ -1033,32 +1038,37 @@ struct FInitBodiesHelper
 		// Ensure we have the AggGeom inside the body setup so we can calculate the number of shapes
 		BodySetup->CreatePhysicsMeshes();
 
-		for(int32 BodyIdx = NumBodies - 1; BodyIdx >= 0; BodyIdx--)   // iterate in reverse since list might shrink
+		for (int32 BodyIdx = NumBodies - 1; BodyIdx >= 0; BodyIdx--)   // iterate in reverse since list might shrink
 		{
 			FBodyInstance* Instance = Bodies[BodyIdx];
 			const FTransform& Transform = Transforms[BodyIdx];
 
+#if !USE_BODYINSTANCE_DEBUG_NAMES
+			FString DebugName;
+#endif
 			FBodyInstance::ValidateTransform(Transform, DebugName, BodySetup);
 
 			Instance->OwnerComponent = PrimitiveComp;
 			Instance->BodySetup = BodySetup;
 			Instance->Scale3D = Transform.GetScale3D();
+#if USE_BODYINSTANCE_DEBUG_NAMES
 			Instance->CharDebugName = PhysXName;
+#endif
 			Instance->bEnableGravity = Instance->bEnableGravity && (SkelMeshComp ? SkelMeshComp->BodyInstance.bEnableGravity : true);	//In the case of skeletal mesh component we AND bodies with the parent body
 
 			// Handle autowelding here to avoid extra work
-			if(!IsStatic() && Instance->bAutoWeld)
+			if (!IsStatic() && Instance->bAutoWeld)
 			{
 				ECollisionEnabled::Type CollisionType = Instance->GetCollisionEnabled();
-				if(CollisionType != ECollisionEnabled::QueryOnly)
+				if (CollisionType != ECollisionEnabled::QueryOnly)
 				{
-					if(UPrimitiveComponent * ParentPrimComponent = PrimitiveComp ? Cast<UPrimitiveComponent>(PrimitiveComp->GetAttachParent()) : NULL)
+					if (UPrimitiveComponent * ParentPrimComponent = PrimitiveComp ? Cast<UPrimitiveComponent>(PrimitiveComp->GetAttachParent()) : NULL)
 					{
 						UWorld* World = PrimitiveComp->GetWorld();
-						if(World && World->IsGameWorld())
+						if (World && World->IsGameWorld())
 						{
 							//if we have a parent we will now do the weld and exit any further initialization
-							if(PrimitiveComp->WeldToImplementation(ParentPrimComponent, PrimitiveComp->GetAttachSocketName(), false))	//welded new simulated body so initialization is done
+							if (PrimitiveComp->WeldToImplementation(ParentPrimComponent, PrimitiveComp->GetAttachSocketName(), false))	//welded new simulated body so initialization is done
 							{
 								return false;
 							}
@@ -1072,17 +1082,17 @@ struct FInitBodiesHelper
 			if(Instance->GetPhysicsActorHandle().IsValid())
 			{
 				Instance->OwnerComponent = nullptr;
-				Instance->BodySetup = nullptr;
+				Instance->BodySetup      = nullptr;
 				Bodies.RemoveAt(BodyIdx);  // so we wont add it to the physx scene again later.
 				Transforms.RemoveAt(BodyIdx);
 				continue;
 			}
 
 			// Set sim parameters for bodies from skeletal mesh components
-			if(!IsStatic() && SpawnParams.bPhysicsTypeDeterminesSimulation)
+			if (!IsStatic() && SpawnParams.bPhysicsTypeDeterminesSimulation)
 			{
 				Instance->bSimulatePhysics = bInstanceSimulatePhysics;
-				if(InstanceBlendWeight != -1.0f)
+				if (InstanceBlendWeight != -1.0f)
 				{
 					Instance->PhysicsBlendWeight = InstanceBlendWeight;
 				}
@@ -1093,7 +1103,7 @@ struct FInitBodiesHelper
 
 			CreateActor_AssumesLocked(Instance, Transform);
 			const bool bInitFail = CreateShapes_AssumesLocked(Instance);
-			if(bInitFail)
+			if (bInitFail)
 			{
 #if WITH_EDITOR
 				//In the editor we may have ended up here because of world trace ignoring our EnableCollision. Since we can't get at the data in that function we check for it here
@@ -1113,63 +1123,63 @@ struct FInitBodiesHelper
 			}
 
 			FPhysicsInterface::SetActorUserData_AssumesLocked(Instance->ActorHandle, &Instance->PhysxUserData);
-		}
+			}
 
 		return true;
-	}
+			}
 
 
 	void InitBodies()
-	{
+				{
 		LLM_SCOPE(ELLMTag::PhysX);
 
 		check(IsInGameThread());
 
 		if(CreateShapesAndActors())
-		{
+					{
 			FPhysicsCommand::ExecuteWrite(PhysScene, [&]()
-			{
+	{
 				// If an aggregate present, add to that
 				if(Aggregate.IsValid())
-				{
+		{
 					for(FBodyInstance* BI : Bodies)
-					{
+			{
 						const FPhysicsActorHandle& ActorHandle = BI->GetPhysicsActorHandle();
 						if(ActorHandle.IsValid())
-						{
+				{
 							FPhysicsInterface::AddActorToAggregate_AssumesLocked(Aggregate, ActorHandle);
-						}
-					}
 				}
+			}
+		}
 				else if(PhysScene)
-				{
+		{
 					TArray<FPhysicsActorHandle> ActorHandles;
 					ActorHandles.Reserve(Bodies.Num());
 
 					for(FBodyInstance* BI : Bodies)
-					{
+			{
 						const FPhysicsActorHandle& ActorHandle = BI->GetPhysicsActorHandle();
 						if(ActorHandle.IsValid())
-						{
+			{
 							ActorHandles.Add(ActorHandle);
-						}
-					}
+			}
+		}
 
 					PhysScene->AddActorsToScene_AssumesLocked(ActorHandles);
 				}
-
+				
 				// Set up dynamic instance data
-				if(!IsStatic())
+				if (!IsStatic())
 				{
 					SCOPE_CYCLE_COUNTER(STAT_InitBodyPostAdd);
-					for(int32 BodyIdx = 0, NumBodies = Bodies.Num(); BodyIdx < NumBodies; ++BodyIdx)
+					for (int32 BodyIdx = 0, NumBodies = Bodies.Num(); BodyIdx < NumBodies; ++BodyIdx)
 					{
 						FBodyInstance* Instance = Bodies[BodyIdx];
 						Instance->InitDynamicProperties_AssumesLocked();
 					}
 				}
 			});
-		}
+			}
 	}
 };
 
@@ -1205,7 +1215,7 @@ void FBodyInstance::InitBody(class UBodySetup* Setup, const FTransform& Transfor
 
 	SCOPE_CYCLE_COUNTER(STAT_InitBody);
 	check(Setup);
-
+	
 	static TArray<FBodyInstance*> Bodies;
 	static TArray<FTransform> Transforms;
 
@@ -1236,11 +1246,11 @@ void FBodyInstance::InitBody(class UBodySetup* Setup, const FTransform& Transfor
 FVector GetInitialLinearVelocity(const AActor* OwningActor, bool& bComponentAwake)
 {
 	FVector InitialLinVel(EForceInit::ForceInitToZero);
-	if(OwningActor)
+	if (OwningActor)
 	{
 		InitialLinVel = OwningActor->GetVelocity();
 
-		if(InitialLinVel.SizeSquared() > FMath::Square(KINDA_SMALL_NUMBER))
+		if (InitialLinVel.SizeSquared() > FMath::Square(KINDA_SMALL_NUMBER))
 		{
 			bComponentAwake = true;
 		}
@@ -1281,28 +1291,28 @@ void FBodyInstance::TermBody(bool bNeverDeferRelease)
 	OwnerComponent = NULL;
 	ExternalCollisionProfileBodySetup = nullptr;
 
-	if(DOFConstraint)
+	if (DOFConstraint)
 	{
 		DOFConstraint->TermConstraint();
 		FConstraintInstance::Free(DOFConstraint);
-		DOFConstraint = NULL;
+			DOFConstraint = NULL;
 	}
-
+	
 }
 
 bool FBodyInstance::Weld(FBodyInstance* TheirBody, const FTransform& TheirTM)
 {
 	check(IsInGameThread());
 	check(TheirBody);
-	if(TheirBody->BodySetup.IsValid() == false)	//attach actor can be called before body has been initialized. In this case just return false
+	if (TheirBody->BodySetup.IsValid() == false)	//attach actor can be called before body has been initialized. In this case just return false
 	{
 		return false;
 	}
 
-	if(TheirBody->WeldParent == this) // The body is already welded to this component. Do nothing.
-	{
-		return false;
-	}
+    if (TheirBody->WeldParent == this) // The body is already welded to this component. Do nothing.
+    {
+        return false;
+    }
 
 	TArray<FPhysicsShapeHandle> PNewShapes;
 
@@ -1336,7 +1346,7 @@ bool FBodyInstance::Weld(FBodyInstance* TheirBody, const FTransform& TheirTM)
 				ShapeToBodiesMap = TSharedPtr<TMap<FPhysicsShapeHandle, FWeldInfo>>(new TMap<FPhysicsShapeHandle, FWeldInfo>());
 			}
 
-			for(int32 ShapeIdx = 0; ShapeIdx < PNewShapes.Num(); ++ShapeIdx)
+			for (int32 ShapeIdx = 0; ShapeIdx < PNewShapes.Num(); ++ShapeIdx)
 			{
 				ShapeToBodiesMap->Add(PNewShapes[ShapeIdx], FWeldInfo(TheirBody, RelativeTM));
 			}
@@ -1366,7 +1376,7 @@ bool FBodyInstance::Weld(FBodyInstance* TheirBody, const FTransform& TheirTM)
 		// remove their body from scenes (don't call TermBody because we don't want to clear things like BodySetup)
 		FPhysicsInterface::ReleaseActor(TheirBody->ActorHandle, TheirBody->GetPhysicsScene());
 	});
-
+	
 	UpdateInterpolateWhenSubStepping();
 
 	TheirBody->UpdateDebugRendering();
@@ -1382,7 +1392,7 @@ void FBodyInstance::UnWeld(FBodyInstance* TheirBI)
 	bool bShapesChanged = false;
 
 	FPhysicsCommand::ExecuteWrite(ActorHandle, [&](const FPhysicsActorHandle& Actor)
-	{
+				{
 		TArray<FPhysicsShapeHandle> Shapes;
 		const int32 NumSyncShapes = GetAllShapes_AssumesLocked(Shapes);
 		const int32 NumTotalShapes = Shapes.Num();
@@ -1390,7 +1400,7 @@ void FBodyInstance::UnWeld(FBodyInstance* TheirBI)
 		for(FPhysicsShapeHandle& Shape : Shapes)
 		{
 			const FBodyInstance* BI = GetOriginalBodyInstance(Shape);
-			if(TheirBI == BI)
+			if (TheirBI == BI)
 			{
 				ShapeToBodiesMap->Remove(Shape);
 				FPhysicsInterface::DetachShape(Actor, Shape);
@@ -1398,10 +1408,10 @@ void FBodyInstance::UnWeld(FBodyInstance* TheirBI)
 			}
 		}
 
-		if(bShapesChanged)
-		{
-			PostShapeChange();
-		}
+	if (bShapesChanged)
+	{
+		PostShapeChange();
+	}
 
 		TheirBI->WeldParent = nullptr;
 	});
@@ -1426,11 +1436,11 @@ float AdjustForSmallThreshold(float NewVal, float OldVal)
 {
 	float Threshold = 0.1f;
 	float Delta = NewVal - OldVal;
-	if(Delta < 0 && FMath::Abs(NewVal) < Threshold)	//getting smaller and passed threshold so flip sign
+	if (Delta < 0 && FMath::Abs(NewVal) < Threshold)	//getting smaller and passed threshold so flip sign
 	{
 		return -Threshold;
 	}
-	else if(Delta > 0 && FMath::Abs(NewVal) < Threshold)	//getting bigger and passed small threshold so flip sign
+	else if (Delta > 0 && FMath::Abs(NewVal) < Threshold)	//getting bigger and passed small threshold so flip sign
 	{
 		return Threshold;
 	}
@@ -1460,35 +1470,35 @@ void ComputeScalingVectors(EScaleMode::Type ScaleMode, const FVector& InScale3D,
 	NewScale3D.Z = FMath::IsNearlyZero(InScale3D.Z) ? KINDA_SMALL_NUMBER : InScale3D.Z;
 
 	const FVector NewScale3DAbs = NewScale3D.GetAbs();
-	switch(ScaleMode)
+	switch (ScaleMode)
 	{
-		case EScaleMode::Free:
-		{
-			OutScale3D = NewScale3D;
-			break;
-		}
-		case EScaleMode::LockedXY:
-		{
-			float XYScaleAbs = FMath::Max(NewScale3DAbs.X, NewScale3DAbs.Y);
-			float XYScale = FMath::Max(NewScale3D.X, NewScale3D.Y) < 0.f ? -XYScaleAbs : XYScaleAbs;	//if both xy are negative we should make the xy scale negative
+	case EScaleMode::Free:
+	{
+		OutScale3D = NewScale3D;
+		break;
+	}
+	case EScaleMode::LockedXY:
+	{
+		float XYScaleAbs = FMath::Max(NewScale3DAbs.X, NewScale3DAbs.Y);
+		float XYScale = FMath::Max(NewScale3D.X, NewScale3D.Y) < 0.f ? -XYScaleAbs : XYScaleAbs;	//if both xy are negative we should make the xy scale negative
 
-			OutScale3D = NewScale3D;
-			OutScale3D.X = OutScale3D.Y = XYScale;
+		OutScale3D = NewScale3D;
+		OutScale3D.X = OutScale3D.Y = XYScale;
 
-			break;
-		}
-		case EScaleMode::LockedXYZ:
-		{
-			float UniformScaleAbs = NewScale3DAbs.GetMin();	//uniform scale uses the smallest magnitude
-			float UniformScale = FMath::Max3(NewScale3D.X, NewScale3D.Y, NewScale3D.Z) < 0.f ? -UniformScaleAbs : UniformScaleAbs;	//if all three values are negative we should make uniform scale negative
+		break;
+	}
+	case EScaleMode::LockedXYZ:
+	{
+		float UniformScaleAbs = NewScale3DAbs.GetMin();	//uniform scale uses the smallest magnitude
+		float UniformScale = FMath::Max3(NewScale3D.X, NewScale3D.Y, NewScale3D.Z) < 0.f ? -UniformScaleAbs : UniformScaleAbs;	//if all three values are negative we should make uniform scale negative
 
-			OutScale3D = FVector(UniformScale);
-			break;
-		}
-		default:
-		{
-			check(false);	//invalid scale mode
-		}
+		OutScale3D = FVector(UniformScale);
+		break;
+	}
+	default:
+	{
+		check(false);	//invalid scale mode
+	}
 	}
 
 	OutScale3DAbs = OutScale3D.GetAbs();
@@ -1525,7 +1535,7 @@ void FBodyInstance::SetMassOverride(float MassInKG, bool bNewOverrideMass)
 
 bool FBodyInstance::GetRigidBodyState(FRigidBodyState& OutState)
 {
-	if(IsInstanceSimulatingPhysics())
+	if (IsInstanceSimulatingPhysics())
 	{
 		FTransform BodyTM = GetUnrealWorldTransform();
 		OutState.Position = BodyTM.GetTranslation();
@@ -1543,14 +1553,14 @@ bool FBodyInstance::UpdateBodyScale(const FVector& InScale3D, bool bForceUpdate)
 {
 	SCOPE_CYCLE_COUNTER(STAT_BodyInstanceUpdateBodyScale);
 
-	if(!IsValidBodyInstance())
+	if (!IsValidBodyInstance())
 	{
 		//UE_LOG(LogPhysics, Log, TEXT("Body hasn't been initialized. Call InitBody to initialize."));
 		return false;
 	}
 
 	// if scale is already correct, and not forcing an update, do nothing
-	if(Scale3D.Equals(InScale3D) && !bForceUpdate)
+	if (Scale3D.Equals(InScale3D) && !bForceUpdate)
 	{
 		return false;
 	}
@@ -1558,7 +1568,7 @@ bool FBodyInstance::UpdateBodyScale(const FVector& InScale3D, bool bForceUpdate)
 	bool bSuccess = false;
 
 #if !(UE_BUILD_SHIPPING || UE_BUILD_TEST)
-	ensureMsgf(!Scale3D.ContainsNaN() && !InScale3D.ContainsNaN(), TEXT("Scale3D = (%f,%f,%f) InScale3D = (%f,%f,%f)"), Scale3D.X, Scale3D.Y, Scale3D.Z, InScale3D.X, InScale3D.Y, InScale3D.Z);
+	ensureMsgf ( !Scale3D.ContainsNaN() && !InScale3D.ContainsNaN(), TEXT("Scale3D = (%f,%f,%f) InScale3D = (%f,%f,%f)"), Scale3D.X, Scale3D.Y, Scale3D.Z, InScale3D.X, InScale3D.Y, InScale3D.Z );
 #endif
 
 	FVector UpdatedScale3D;
@@ -1600,7 +1610,7 @@ bool FBodyInstance::UpdateBodyScale(const FVector& InScale3D, bool bForceUpdate)
 #if WITH_CHAOS
 			//check(false);
 #else
-			switch(GeomType)
+			switch (GeomType)
 			{
 				case ECollisionShapeType::Sphere:
 				{
@@ -1608,13 +1618,13 @@ bool FBodyInstance::UpdateBodyScale(const FVector& InScale3D, bool bForceUpdate)
 					ensure(ScaleMode == EScaleMode::LockedXYZ);
 
 					GeoCollection.GetSphereGeometry(PSphereGeom);
-
+					 
 					PSphereGeom.radius = FMath::Max(SphereElem->Radius * AdjustedScale3DAbs.X, FCollisionShape::MinSphereRadius());
 
 					FVector NewTranslation = RelativeTM.TransformPosition(SphereElem->Center) * AdjustedScale3D.X;
 					LocalTransform.SetTranslation(NewTranslation);
 
-					if(PSphereGeom.isValid())
+					if (PSphereGeom.isValid())
 					{
 						UpdatedGeometry = &PSphereGeom;
 						bSuccess = true;
@@ -1637,7 +1647,7 @@ bool FBodyInstance::UpdateBodyScale(const FVector& InScale3D, bool bForceUpdate)
 					LocalTransform = BoxElem->GetTransform() * RelativeTM;
 					LocalTransform.ScaleTranslation(AdjustedScale3D);
 
-					if(PBoxGeom.isValid())
+					if (PBoxGeom.isValid())
 					{
 						UpdatedGeometry = &PBoxGeom;
 						bSuccess = true;
@@ -1675,7 +1685,7 @@ bool FBodyInstance::UpdateBodyScale(const FVector& InScale3D, bool bForceUpdate)
 					LocalTransform.SetRotation(SphylElem->Rotation.Quaternion() * U2PSphylBasis_UE); // #PHYS2 we probably want to put this behind the interface?
 					LocalTransform.ScaleTranslation(AdjustedScale3D);
 
-					if(PCapsuleGeom.isValid())
+					if (PCapsuleGeom.isValid())
 					{
 						UpdatedGeometry = &PCapsuleGeom;
 						bSuccess = true;
@@ -1700,7 +1710,7 @@ bool FBodyInstance::UpdateBodyScale(const FVector& InScale3D, bool bForceUpdate)
 					LocalTransform.TransformRotation(RelativeTM.GetRotation());
 					LocalTransform.ScaleTranslation(AdjustedScale3D);
 
-					if(PConvexGeom.isValid())
+					if (PConvexGeom.isValid())
 					{
 						UpdatedGeometry = &PConvexGeom;
 						bSuccess = true;
@@ -1720,19 +1730,19 @@ bool FBodyInstance::UpdateBodyScale(const FVector& InScale3D, bool bForceUpdate)
 
 					// find which trimesh elems it is
 					// it would be nice to know if the order of PShapes array index is in the order of createShape
-					if(BodySetup.IsValid())
+					if (BodySetup.IsValid())
 					{
-						for(PxTriangleMesh* TriMesh : BodySetup->TriMeshes)
+						for (PxTriangleMesh* TriMesh : BodySetup->TriMeshes)
 						{
 							// found it
-							if(TriMesh == PTriMeshGeom.triangleMesh)
+							if (TriMesh == PTriMeshGeom.triangleMesh)
 							{
 								PTriMeshGeom.scale.scale = U2PVector(AdjustedScale3D);
 
 								LocalTransform = RelativeTM;
 								LocalTransform.ScaleTranslation(AdjustedScale3D);
 
-								if(PTriMeshGeom.isValid())
+								if (PTriMeshGeom.isValid())
 								{
 									UpdatedGeometry = &PTriMeshGeom;
 									bSuccess = true;
@@ -1754,12 +1764,12 @@ bool FBodyInstance::UpdateBodyScale(const FVector& InScale3D, bool bForceUpdate)
 				}
 				default:
 				{
-					UE_LOG(LogPhysics, Error, TEXT("Unknown geom type."));
+						   UE_LOG(LogPhysics, Error, TEXT("Unknown geom type."));
 				}
 			}// end switch
 #endif
 
-			if(UpdatedGeometry)
+			if (UpdatedGeometry)
 			{
 				FPhysicsCommand::ExecuteShapeWrite(this, Shape, [&](FPhysicsShapeHandle& InShape)
 				{
@@ -1769,7 +1779,7 @@ bool FBodyInstance::UpdateBodyScale(const FVector& InScale3D, bool bForceUpdate)
 
 				UpdatedScale3D = AdjustedScale3D;
 			}
-			else if(bInvalid)
+			else if (bInvalid)
 			{
 #if !(UE_BUILD_SHIPPING || UE_BUILD_TEST)
 				FMessageLog("PIE").Warning()
@@ -1781,16 +1791,16 @@ bool FBodyInstance::UpdateBodyScale(const FVector& InScale3D, bool bForceUpdate)
 			}
 		}
 	});
-
+	
 #endif
 
 	// if success, overwrite old Scale3D, otherwise, just don't do it. It will have invalid scale next time
-	if(bSuccess)
+	if (bSuccess)
 	{
 		Scale3D = UpdatedScale3D;
 
 		// update mass if required
-		if(bUpdateMassWhenScaleChanges)
+		if (bUpdateMassWhenScaleChanges)
 		{
 			UpdateMassProperties();
 		}
@@ -1813,16 +1823,16 @@ void FBodyInstance::UpdateInstanceSimulatePhysics()
 		FPhysicsInterface::SetCcdEnabled_AssumesLocked(Actor, bUseCCD);
 
 		if(bSimulatePhysics && bStartAwake)
-		{
+				{
 			FPhysicsInterface::WakeUp_AssumesLocked(Actor);
 		}
 	});
 
 	//In the original physx only implementation this was wrapped in a PRigidDynamic != NULL check.
 	//We use bInitialized to check rigid actor has been created in either engine because if we haven't even initialized yet, we don't want to undo our settings
-	if(bInitialized)
+	if (bInitialized)
 	{
-		if(bUseSimulate)
+		if (bUseSimulate)
 		{
 			PhysicsBlendWeight = 1.f;
 		}
@@ -1853,11 +1863,11 @@ void FBodyInstance::ApplyWeldOnChildren()
 		TArray<FName> ChildrenLabels;
 		OwnerComponentInst->GetWeldedBodies(ChildrenBodies, ChildrenLabels, /*bIncludingAutoWeld=*/true);
 
-		for(int32 ChildIdx = 0; ChildIdx < ChildrenBodies.Num(); ++ChildIdx)
+		for (int32 ChildIdx = 0; ChildIdx < ChildrenBodies.Num(); ++ChildIdx)
 		{
 			FBodyInstance* ChildBI = ChildrenBodies[ChildIdx];
 			checkSlow(ChildBI);
-			if(ChildBI != this)
+			if (ChildBI != this)
 			{
 				const ECollisionEnabled::Type ChildCollision = ChildBI->GetCollisionEnabled();
 				if(CollisionEnabledHasPhysics(ChildCollision))
@@ -1870,7 +1880,7 @@ void FBodyInstance::ApplyWeldOnChildren()
 			}
 		}
 	}
-
+	
 }
 
 bool FBodyInstance::ShouldInstanceSimulatingPhysics() const
@@ -1881,33 +1891,33 @@ bool FBodyInstance::ShouldInstanceSimulatingPhysics() const
 
 void FBodyInstance::SetInstanceSimulatePhysics(bool bSimulate, bool bMaintainPhysicsBlending)
 {
-	if(bSimulate)
+	if (bSimulate)
 	{
 		UPrimitiveComponent* OwnerComponentInst = OwnerComponent.Get();
 
 		// If we are enabling simulation, and we are the root body of our component (or we are welded), we detach the component 
-		if(OwnerComponentInst && OwnerComponentInst->IsRegistered() && (OwnerComponentInst->GetBodyInstance() == this || OwnerComponentInst->IsWelded()))
+		if (OwnerComponentInst && OwnerComponentInst->IsRegistered() && (OwnerComponentInst->GetBodyInstance() == this || OwnerComponentInst->IsWelded()))
 		{
-			if(OwnerComponentInst->GetAttachParent())
+			if (OwnerComponentInst->GetAttachParent())
 			{
 				OwnerComponentInst->DetachFromComponent(FDetachmentTransformRules::KeepWorldTransform);
 			}
 
-			if(bSimulatePhysics == false)	//if we're switching from kinematic to simulated
+			if (bSimulatePhysics == false)	//if we're switching from kinematic to simulated
 			{
 				ApplyWeldOnChildren();
 			}
 		}
 
 #if !(UE_BUILD_SHIPPING || UE_BUILD_TEST)
-		if(OwnerComponentInst)
+		if (OwnerComponentInst)
 		{
-			if(!IsValidBodyInstance())
+			if (!IsValidBodyInstance())
 			{
 				FMessageLog("PIE").Warning(FText::Format(LOCTEXT("SimPhysNoBody", "Trying to simulate physics on ''{0}'' but no physics body."),
 					FText::FromString(GetPathNameSafe(OwnerComponentInst))));
 			}
-			else if(!IsDynamic())
+			else if (!IsDynamic())
 			{
 				FMessageLog("PIE").Warning(FText::Format(LOCTEXT("SimPhysStatic", "Trying to simulate physics on ''{0}'' but it is static."),
 					FText::FromString(GetPathNameSafe(OwnerComponentInst))));
@@ -1922,9 +1932,9 @@ void FBodyInstance::SetInstanceSimulatePhysics(bool bSimulate, bool bMaintainPhy
 	}
 
 	bSimulatePhysics = bSimulate;
-	if(!bMaintainPhysicsBlending)
+	if ( !bMaintainPhysicsBlending )
 	{
-		if(bSimulatePhysics)
+		if (bSimulatePhysics)
 		{
 			PhysicsBlendWeight = 1.f;
 		}
@@ -1952,7 +1962,7 @@ FTransform GetUnrealWorldTransformImp_AssumesLocked(const FBodyInstance* BodyIns
 
 		if(bWithProjection)
 		{
-			BodyInstance->OnCalculateCustomProjection.ExecuteIfBound(BodyInstance, WorldTM);
+			BodyInstance->ExecuteOnCalculateCustomProjection(WorldTM);
 		}
 	}
 
@@ -1983,7 +1993,7 @@ void FBodyInstance::SetBodyTransform(const FTransform& NewTransform, ETeleportTy
 
 #if !(UE_BUILD_SHIPPING || UE_BUILD_TEST)
 	extern bool GShouldLogOutAFrameOfSetBodyTransform;
-	if(GShouldLogOutAFrameOfSetBodyTransform == true)
+	if (GShouldLogOutAFrameOfSetBodyTransform == true)
 	{
 		UE_LOG(LogPhysics, Log, TEXT("SetBodyTransform: %s"), *GetBodyDebugName());
 	}
@@ -1991,7 +2001,7 @@ void FBodyInstance::SetBodyTransform(const FTransform& NewTransform, ETeleportTy
 
 	// Catch NaNs and elegantly bail out.
 
-	if(!ensureMsgf(!NewTransform.ContainsNaN(), TEXT("SetBodyTransform contains NaN (%s)\n%s"), (OwnerComponent.Get() ? *OwnerComponent->GetPathName() : TEXT("NONE")), *NewTransform.ToString()))
+	if( !ensureMsgf(!NewTransform.ContainsNaN(), TEXT("SetBodyTransform contains NaN (%s)\n%s"), (OwnerComponent.Get() ? *OwnerComponent->GetPathName() : TEXT("NONE")), *NewTransform.ToString()) )
 	{
 		return;
 	}
@@ -2023,27 +2033,27 @@ void FBodyInstance::SetBodyTransform(const FTransform& NewTransform, ETeleportTy
 				const bool bIsSimKinematic = bKinematic && bSimulated;
 
 				if(bIsSimKinematic && Teleport == ETeleportType::None)
-				{
+					{
 					Scene->SetKinematicTarget_AssumesLocked(this, NewTransform, true);
-				}
-				else
-				{
-					if(bIsSimKinematic)
+					}
+					else
 					{
+					if(bIsSimKinematic)
+						{
 						FPhysicsInterface::SetKinematicTarget_AssumesLocked(Actor, NewTransform);
-					}
+						}
 
 					FPhysicsInterface::SetGlobalPose_AssumesLocked(Actor, NewTransform);
-				}
+					}
 			});
-		}
+				}
 		else if(Scene)
 		{
 			FPhysicsCommand::ExecuteWrite(ActorHandle, [&](const FPhysicsActorHandle& Actor)
 			{
 				FPhysicsInterface::SetGlobalPose_AssumesLocked(Actor, NewTransform);
 			});
-		}
+			}
 	}
 	else if(WeldParent)
 	{
@@ -2073,9 +2083,9 @@ FVector FBodyInstance::GetUnrealWorldVelocity_AssumesLocked() const
 {
 	FVector LinVel(EForceInit::ForceInitToZero);
 	if(ActorHandle.IsValid())
-	{
+{
 		LinVel = FPhysicsInterface::GetLinearVelocity_AssumesLocked(ActorHandle);
-	}
+}
 
 	return LinVel;
 }
@@ -2121,7 +2131,7 @@ FTransform FBodyInstance::GetMassSpaceToWorldSpace() const
 	FTransform MassSpaceToWorldSpace = FTransform::Identity;
 
 	FPhysicsCommand::ExecuteRead(ActorHandle, [&](const FPhysicsActorHandle& Actor)
-	{
+		{
 		MassSpaceToWorldSpace = FPhysicsInterface::GetComTransform_AssumesLocked(Actor);
 	});
 
@@ -2188,7 +2198,7 @@ FBox FBodyInstance::GetBodyBounds() const
 
 void FBodyInstance::DrawCOMPosition(FPrimitiveDrawInterface* PDI, float COMRenderSize, const FColor& COMRenderColor)
 {
-	if(IsValidBodyInstance())
+	if (IsValidBodyInstance())
 	{
 		DrawWireStar(PDI, GetCOMPosition(), COMRenderSize, COMRenderColor, SDPG_World);
 	}
@@ -2227,10 +2237,10 @@ void FBodyInstance::CopyRuntimeBodyInstancePropertiesFrom(const FBodyInstance* F
 FPhysScene* FBodyInstance::GetPhysicsScene() const
 {
 	if(ActorHandle.IsValid())
-	{
+{
 		return FPhysicsInterface::GetCurrentScene(ActorHandle);
 	}
-
+	
 	return nullptr;
 }
 
@@ -2246,7 +2256,7 @@ const FPhysicsActorHandle& FBodyInstance::GetPhysicsActorHandle() const
 
 const FWalkableSlopeOverride& FBodyInstance::GetWalkableSlopeOverride() const
 {
-	if(bOverrideWalkableSlopeOnInstance || !BodySetup.IsValid())
+	if (bOverrideWalkableSlopeOnInstance || !BodySetup.IsValid())
 	{
 		return WalkableSlopeOverride;
 	}
@@ -2263,12 +2273,12 @@ void FBodyInstance::SetWalkableSlopeOverride(const FWalkableSlopeOverride& NewOv
 }
 
 
-/**
-*	Changes the current PhysMaterialOverride for this body.
-*	Note that if physics is already running on this component, this will _not_ alter its mass/inertia etc, it will only change its
+/** 
+*	Changes the current PhysMaterialOverride for this body. 
+*	Note that if physics is already running on this component, this will _not_ alter its mass/inertia etc, it will only change its 
 *	surface properties like friction and the damping.
 */
-void FBodyInstance::SetPhysMaterialOverride(UPhysicalMaterial* NewPhysMaterial)
+void FBodyInstance::SetPhysMaterialOverride( UPhysicalMaterial* NewPhysMaterial )
 {
 	// Save ref to PhysicalMaterial
 	PhysMaterialOverride = NewPhysMaterial;
@@ -2296,11 +2306,11 @@ UPhysicalMaterial* FBodyInstance::GetSimplePhysicalMaterial(const FBodyInstance*
 
 	// Find the PhysicalMaterial we need to apply to the physics bodies.
 	// (LOW priority) Engine Mat, Material PhysMat, BodySetup Mat, Component Override, Body Override (HIGH priority)
-
+	
 	UPhysicalMaterial* ReturnPhysMaterial = NULL;
 
 	// BodyInstance override
-	if(BodyInstance->PhysMaterialOverride != NULL)
+	if (BodyInstance->PhysMaterialOverride != NULL)
 	{
 		ReturnPhysMaterial = BodyInstance->PhysMaterialOverride;
 		check(!ReturnPhysMaterial || ReturnPhysMaterial->IsValidLowLevel());
@@ -2309,7 +2319,7 @@ UPhysicalMaterial* FBodyInstance::GetSimplePhysicalMaterial(const FBodyInstance*
 	{
 		// Component override
 		UPrimitiveComponent* OwnerPrimComponent = OwnerComp.Get();
-		if(OwnerPrimComponent && OwnerPrimComponent->BodyInstance.PhysMaterialOverride != NULL)
+		if (OwnerPrimComponent && OwnerPrimComponent->BodyInstance.PhysMaterialOverride != NULL)
 		{
 			ReturnPhysMaterial = OwnerComp->BodyInstance.PhysMaterialOverride;
 			check(!ReturnPhysMaterial || ReturnPhysMaterial->IsValidLowLevel());
@@ -2318,7 +2328,7 @@ UPhysicalMaterial* FBodyInstance::GetSimplePhysicalMaterial(const FBodyInstance*
 		{
 			// BodySetup
 			UBodySetup* BodySetupRawPtr = BodySetupPtr.Get();
-			if(BodySetupRawPtr && BodySetupRawPtr->PhysMaterial != NULL)
+			if (BodySetupRawPtr && BodySetupRawPtr->PhysMaterial != NULL)
 			{
 				ReturnPhysMaterial = BodySetupPtr->PhysMaterial;
 				check(!ReturnPhysMaterial || ReturnPhysMaterial->IsValidLowLevel());
@@ -2328,16 +2338,16 @@ UPhysicalMaterial* FBodyInstance::GetSimplePhysicalMaterial(const FBodyInstance*
 				// See if the Material has a PhysicalMaterial
 				UMeshComponent* MeshComp = Cast<UMeshComponent>(OwnerPrimComponent);
 				UPhysicalMaterial* PhysMatFromMaterial = NULL;
-				if(MeshComp != NULL)
+				if (MeshComp != NULL)
 				{
 					UMaterialInterface* Material = MeshComp->GetMaterial(0);
-					if(Material != NULL)
+					if (Material != NULL)
 					{
 						PhysMatFromMaterial = Material->GetPhysicalMaterial();
 					}
 				}
 
-				if(PhysMatFromMaterial != NULL)
+				if (PhysMatFromMaterial != NULL)
 				{
 					ReturnPhysMaterial = PhysMatFromMaterial;
 					check(!ReturnPhysMaterial || ReturnPhysMaterial->IsValidLowLevel());
@@ -2351,7 +2361,7 @@ UPhysicalMaterial* FBodyInstance::GetSimplePhysicalMaterial(const FBodyInstance*
 			}
 		}
 	}
-
+	
 	return ReturnPhysMaterial;
 }
 
@@ -2373,16 +2383,16 @@ void FBodyInstance::GetComplexPhysicalMaterials(const FBodyInstance*, TWeakObjec
 	check(GEngine->DefaultPhysMaterial != NULL);
 	// See if the Material has a PhysicalMaterial
 	UPrimitiveComponent* PrimComp = OwnerComp.Get();
-	if(PrimComp)
+	if (PrimComp)
 	{
 		const int32 NumMaterials = PrimComp->GetNumMaterials();
 		OutPhysicalMaterials.SetNum(NumMaterials);
 
-		for(int32 MatIdx = 0; MatIdx < NumMaterials; MatIdx++)
+		for (int32 MatIdx = 0; MatIdx < NumMaterials; MatIdx++)
 		{
 			UPhysicalMaterial* PhysMat = GEngine->DefaultPhysMaterial;
 			UMaterialInterface* Material = PrimComp->GetMaterial(MatIdx);
-			if(Material)
+			if (Material)
 			{
 				PhysMat = Material->GetPhysicalMaterial();
 			}
@@ -2435,7 +2445,7 @@ PxMassProperties ComputeMassProperties(const FBodyInstance* OwningBodyInstance,
 	// physical material - nothing can weigh less than hydrogen (0.09 kg/m^3)
 	float DensityKGPerCubicUU = 1.0f;
 	float RaiseMassToPower = 0.75f;
-	if(UPhysicalMaterial* PhysMat = OwningBodyInstance->GetSimplePhysicalMaterial())
+	if (UPhysicalMaterial* PhysMat = OwningBodyInstance->GetSimplePhysicalMaterial())
 	{
 		DensityKGPerCubicUU = FMath::Max(KgPerM3ToKgPerCm3(0.09f), gPerCm3ToKgPerCm3(PhysMat->Density));
 		RaiseMassToPower = PhysMat->RaiseMassToPower;
@@ -2447,7 +2457,7 @@ PxMassProperties ComputeMassProperties(const FBodyInstance* OwningBodyInstance,
 	float OldMass = MassProps.mass;
 	float NewMass = 0.f;
 
-	if(OwningBodyInstance->bOverrideMass == false)
+	if (OwningBodyInstance->bOverrideMass == false)
 	{
 		float UsePow = FMath::Clamp<float>(RaiseMassToPower, KINDA_SMALL_NUMBER, 1.f);
 		NewMass = FMath::Pow(OldMass, UsePow);
@@ -2463,7 +2473,7 @@ PxMassProperties ComputeMassProperties(const FBodyInstance* OwningBodyInstance,
 	check(NewMass > 0.f);
 
 	float MassRatio = NewMass / OldMass;
-
+	
 	PxMassProperties FinalMassProps = MassProps * MassRatio;
 
 	FinalMassProps.centerOfMass += U2PVector(MassModifierTransform.TransformVector(OwningBodyInstance->COMNudge));
@@ -2489,90 +2499,90 @@ void FBodyInstance::UpdateMassProperties()
 				const int32 NumShapes = FPhysicsInterface::GetNumShapes(Actor);
 
 				TArray<FPhysicsShapeHandle> Shapes;
-				Shapes.AddUninitialized(NumShapes);
+			Shapes.AddUninitialized(NumShapes);
 				FPhysicsInterface::GetAllShapes_AssumedLocked(Actor, Shapes);
 
-				// Ignore trimeshes & shapes which don't contribute to the mass
-				for(int32 ShapeIdx = Shapes.Num() - 1; ShapeIdx >= 0; --ShapeIdx)
-				{
+			// Ignore trimeshes & shapes which don't contribute to the mass
+			for(int32 ShapeIdx = Shapes.Num() - 1; ShapeIdx >= 0; --ShapeIdx)
+			{
 					const FPhysicsShapeHandle& Shape = Shapes[ShapeIdx];
 					const FKShapeElem* ShapeElem = FPhysxUserData::Get<FKShapeElem>(FPhysicsInterface::GetUserData(Shape));
 					bool bIsTriangleMesh = FPhysicsInterface::GetShapeType(Shape) == ECollisionShapeType::Trimesh;
-					bool bHasNoMass = ShapeElem && !ShapeElem->GetContributeToMass();
-					if (bIsTriangleMesh || bHasNoMass)
-					{
-						Shapes.RemoveAtSwap(ShapeIdx);
-					}
+				bool bHasNoMass = ShapeElem && !ShapeElem->GetContributeToMass();
+				if (bIsTriangleMesh || bHasNoMass)
+				{
+					Shapes.RemoveAtSwap(ShapeIdx);
 				}
+			}
 
-				PxMassProperties TotalMassProperties;
-				if(ShapeToBodiesMap.IsValid() && ShapeToBodiesMap->Num() > 0)
+			PxMassProperties TotalMassProperties;
+			if(ShapeToBodiesMap.IsValid() && ShapeToBodiesMap->Num() > 0)
+			{
+				struct FWeldedBatch
 				{
-					struct FWeldedBatch
-					{
-							TArray<FPhysicsShapeHandle> Shapes;
-						FTransform RelTM;
-					};
+						TArray<FPhysicsShapeHandle> Shapes;
+					FTransform RelTM;
+				};
 
-					//If we have welded children we must compute the mass properties of each individual body first and then combine them all together
-					TMap<FBodyInstance*, FWeldedBatch> BodyToShapes;
+				//If we have welded children we must compute the mass properties of each individual body first and then combine them all together
+				TMap<FBodyInstance*, FWeldedBatch> BodyToShapes;
 
 					for(const FPhysicsShapeHandle& Shape : Shapes) //sort all welded children by their original bodies
+				{
+					if (FWeldInfo* WeldInfo = ShapeToBodiesMap->Find(Shape))
 					{
-						if (FWeldInfo* WeldInfo = ShapeToBodiesMap->Find(Shape))
+						FWeldedBatch* WeldedBatch = BodyToShapes.Find(WeldInfo->ChildBI);
+						if(!WeldedBatch)
 						{
-							FWeldedBatch* WeldedBatch = BodyToShapes.Find(WeldInfo->ChildBI);
-							if(!WeldedBatch)
-							{
-								WeldedBatch = &BodyToShapes.Add(WeldInfo->ChildBI);
-								WeldedBatch->RelTM = WeldInfo->RelativeTM;
-							}
-
-							WeldedBatch->Shapes.Add(Shape);
+							WeldedBatch = &BodyToShapes.Add(WeldInfo->ChildBI);
+							WeldedBatch->RelTM = WeldInfo->RelativeTM;
 						}
-						else
-						{
-							//no weld info so shape really belongs to this body
-							FWeldedBatch* WeldedBatch = BodyToShapes.Find(this);
-							if (!WeldedBatch)
-							{
-								WeldedBatch = &BodyToShapes.Add(this);
-								WeldedBatch->RelTM = FTransform::Identity;
-							}
 
-							WeldedBatch->Shapes.Add(Shape);
-						}
+						WeldedBatch->Shapes.Add(Shape);
 					}
-
-					TArray<PxMassProperties> SubMassProperties;
-					TArray<PxTransform> MassTMs;
-					for(auto BodyShapesItr : BodyToShapes)
+					else
 					{
-						const FBodyInstance* OwningBI = BodyShapesItr.Key;
-						const FWeldedBatch& WeldedBatch = BodyShapesItr.Value;
-						FTransform MassModifierTransform = WeldedBatch.RelTM;
-						MassModifierTransform.SetScale3D(MassModifierTransform.GetScale3D() * Scale3D);	//Ensure that any scaling that is done on the component is passed into the mass frame modifiers
-
-						PxMassProperties BodyMassProperties = ComputeMassProperties(OwningBI, WeldedBatch.Shapes, MassModifierTransform);
-						SubMassProperties.Add(BodyMassProperties);
-						MassTMs.Add(PxTransform(PxIdentity));
-					}
+						//no weld info so shape really belongs to this body
+						FWeldedBatch* WeldedBatch = BodyToShapes.Find(this);
+						if (!WeldedBatch)
+						{
+							WeldedBatch = &BodyToShapes.Add(this);
+							WeldedBatch->RelTM = FTransform::Identity;
+						}
 
-					TotalMassProperties = PxMassProperties::sum(SubMassProperties.GetData(), MassTMs.GetData(), SubMassProperties.Num());
+						WeldedBatch->Shapes.Add(Shape);
+					}
 				}
-				else
+
+				TArray<PxMassProperties> SubMassProperties;
+				TArray<PxTransform> MassTMs;
+				for(auto BodyShapesItr : BodyToShapes)
 				{
+					const FBodyInstance* OwningBI = BodyShapesItr.Key;
+					const FWeldedBatch& WeldedBatch = BodyShapesItr.Value;
+					FTransform MassModifierTransform = WeldedBatch.RelTM;
+					MassModifierTransform.SetScale3D(MassModifierTransform.GetScale3D() * Scale3D);	//Ensure that any scaling that is done on the component is passed into the mass frame modifiers
+
+					PxMassProperties BodyMassProperties = ComputeMassProperties(OwningBI, WeldedBatch.Shapes, MassModifierTransform);
+					SubMassProperties.Add(BodyMassProperties);
+					MassTMs.Add(PxTransform(PxIdentity));
+				}
+
+				TotalMassProperties = PxMassProperties::sum(SubMassProperties.GetData(), MassTMs.GetData(), SubMassProperties.Num());
+			}
+			else
+			{
 					// If we have no shapes that affect mass we cannot compute the mass properties in a meaningful way.
 					if (Shapes.Num())
 					{
-						//No children welded so just get this body's mass properties
-						FTransform MassModifierTransform(FQuat::Identity, FVector(0.f, 0.f, 0.f), Scale3D);	//Ensure that any scaling that is done on the component is passed into the mass frame modifiers
-						TotalMassProperties = ComputeMassProperties(this, Shapes, MassModifierTransform);
-					}
+				//No children welded so just get this body's mass properties
+				FTransform MassModifierTransform(FQuat::Identity, FVector(0.f, 0.f, 0.f), Scale3D);	//Ensure that any scaling that is done on the component is passed into the mass frame modifiers
+				TotalMassProperties = ComputeMassProperties(this, Shapes, MassModifierTransform);
+			}
 				}
 			
 				// #PHYS2 Refactor out PxMassProperties (Our own impl?)
-				PxQuat MassOrientation;
+			PxQuat MassOrientation;
 				const FVector MassSpaceInertiaTensor = P2UVector(PxMassProperties::getMassSpaceInertia(TotalMassProperties.inertiaTensor, MassOrientation));
 
 				FPhysicsInterface::SetMass_AssumesLocked(Actor, TotalMassProperties.mass);
@@ -2580,20 +2590,23 @@ void FBodyInstance::UpdateMassProperties()
 
 				FTransform Com(P2UQuat(MassOrientation), P2UVector(TotalMassProperties.centerOfMass));
 				FPhysicsInterface::SetComLocalPose_AssumesLocked(Actor, Com);
-			}
-		});
+		}
+	});
 	}
 #endif
 
 	//Let anyone who cares about mass properties know they've been updated
-	OnRecalculatedMassProperties.Broadcast(this);
+	if (BodyInstanceDelegates.IsValid())
+	{
+		BodyInstanceDelegates->OnRecalculatedMassProperties.Broadcast(this);
+	}
 }
 
 void FBodyInstance::UpdateDebugRendering()
 {
 #if !(UE_BUILD_SHIPPING || UE_BUILD_TEST)
 	//After we update the mass properties, we should update any debug rendering
-	if(UPrimitiveComponent* OwnerPrim = OwnerComponent.Get())
+	if (UPrimitiveComponent* OwnerPrim = OwnerComponent.Get())
 	{
 		OwnerPrim->SendRenderDebugPhysics();
 	}
@@ -2605,7 +2618,7 @@ void FBodyInstance::UpdateDampingProperties()
 	FPhysicsCommand::ExecuteWrite(ActorHandle, [&](const FPhysicsActorHandle& Actor)
 	{
 		if(FPhysicsInterface::IsDynamic(Actor))
-		{
+	{
 			FPhysicsInterface::SetLinearDamping_AssumesLocked(Actor, LinearDamping);
 			FPhysicsInterface::SetAngularDamping_AssumesLocked(Actor, AngularDamping);
 		}
@@ -2619,7 +2632,7 @@ bool FBodyInstance::IsInstanceAwake() const
 	FPhysicsCommand::ExecuteRead(ActorHandle, [&](const FPhysicsActorHandle& Actor)
 	{
 		if(FPhysicsInterface::IsDynamic(Actor))
-		{
+	{
 			bIsAwake = !FPhysicsInterface::IsSleeping(Actor);
 		}
 	});
@@ -2651,11 +2664,11 @@ void FBodyInstance::PutInstanceToSleep()
 
 float FBodyInstance::GetSleepThresholdMultiplier() const
 {
-	if(SleepFamily == ESleepFamily::Sensitive)
+	if (SleepFamily == ESleepFamily::Sensitive)
 	{
 		return 1 / 20.0f;
 	}
-	else if(SleepFamily == ESleepFamily::Custom)
+	else if (SleepFamily == ESleepFamily::Custom)
 	{
 		return CustomSleepThresholdMultiplier;
 	}
@@ -2668,13 +2681,13 @@ void FBodyInstance::SetLinearVelocity(const FVector& NewVel, bool bAddToCurrent,
 	FPhysicsCommand::ExecuteWrite(ActorHandle, [&](const FPhysicsActorHandle& Actor)
 	{
 		if(FPhysicsInterface::IsRigidBody(Actor))
-		{
+	{
 			FVector FinalVelocity = NewVel;
 
-			if(bAddToCurrent)
-			{
+		if (bAddToCurrent)
+		{
 				FinalVelocity += FPhysicsInterface::GetLinearVelocity_AssumesLocked(Actor);
-			}
+		}
 
 			FPhysicsInterface::SetLinearVelocity_AssumesLocked(Actor, FinalVelocity);
 		}
@@ -2687,13 +2700,13 @@ void FBodyInstance::SetAngularVelocityInRadians(const FVector& NewAngVel, bool b
 	FPhysicsCommand::ExecuteWrite(ActorHandle, [&](const FPhysicsActorHandle& Actor)
 	{
 		if(FPhysicsInterface::IsRigidBody(Actor))
-		{
+	{
 			FVector FinalVelocity = NewAngVel;
 
-			if(bAddToCurrent)
-			{
+		if (bAddToCurrent)
+		{
 				FinalVelocity += FPhysicsInterface::GetAngularVelocity_AssumesLocked(Actor);
-			}
+		}
 
 			FPhysicsInterface::SetAngularVelocity_AssumesLocked(Actor, FinalVelocity);
 		}
@@ -2711,7 +2724,7 @@ void FBodyInstance::SetMaxAngularVelocityInRadians(float NewMaxAngVel, bool bAdd
 
 	FPhysicsCommand::ExecuteWrite(ActorHandle, [&](const FPhysicsActorHandle& Actor)
 	{
-		if(bAddToCurrent)
+		if (bAddToCurrent)
 		{
 			float OldValue = FPhysicsInterface::GetMaxAngularVelocity_AssumesLocked(Actor);
 			NewMaxAngVel += OldValue;
@@ -2832,11 +2845,11 @@ void FBodyInstance::AddAngularImpulseInRadians(const FVector& AngularImpulse, bo
 		if(FPhysicsInterface::IsRigidBody(Actor) && FPhysicsInterface::IsInScene(Actor) && !IsRigidBodyKinematic_AssumesLocked(Actor))
 		{
 			if(bVelChange)
-			{
+	{
 				FPhysicsInterface::AddTorqueMassIndependent_AssumesLocked(Actor, AngularImpulse);
 			}
 			else
-			{
+		{
 				FPhysicsInterface::AddTorque_AssumesLocked(Actor, AngularImpulse);
 			}
 		}
@@ -2850,11 +2863,11 @@ void FBodyInstance::AddImpulse(const FVector& Impulse, bool bVelChange)
 		if(FPhysicsInterface::IsRigidBody(Actor) && FPhysicsInterface::IsInScene(Actor) && !IsRigidBodyKinematic_AssumesLocked(Actor))
 		{
 			if(bVelChange)
-			{
+	{
 				FPhysicsInterface::AddForceMassIndependent_AssumesLocked(Actor, Impulse);
 			}
 			else
-			{
+		{
 				FPhysicsInterface::AddForce_AssumesLocked(Actor, Impulse);
 			}
 		}
@@ -2880,7 +2893,7 @@ void FBodyInstance::SetInstanceNotifyRBCollision(bool bNewNotifyCollision)
 
 void FBodyInstance::SetEnableGravity(bool bInGravityEnabled)
 {
-	if(bEnableGravity != bInGravityEnabled)
+	if (bEnableGravity != bInGravityEnabled)
 	{
 		bEnableGravity = bInGravityEnabled;
 
@@ -2891,10 +2904,10 @@ void FBodyInstance::SetEnableGravity(bool bInGravityEnabled)
 				{
 					FPhysicsInterface::SetGravityEnabled_AssumesLocked(Actor, bEnableGravity);
 				}
-			});
+		});
 		}
 
-		if(bEnableGravity)
+		if (bEnableGravity)
 		{
 			WakeInstance();
 		}
@@ -2903,7 +2916,7 @@ void FBodyInstance::SetEnableGravity(bool bInGravityEnabled)
 
 void FBodyInstance::SetContactModification(bool bNewContactModification)
 {
-	if(bNewContactModification != bContactModification)
+	if (bNewContactModification != bContactModification)
 	{
 		bContactModification = bNewContactModification;
 		UpdatePhysicsFilterData();
@@ -2912,7 +2925,7 @@ void FBodyInstance::SetContactModification(bool bNewContactModification)
 
 void FBodyInstance::SetUseCCD(bool bInUseCCD)
 {
-	if(bUseCCD != bInUseCCD)
+	if (bUseCCD != bInUseCCD)
 	{
 		bUseCCD = bInUseCCD;
 		// Need to set body flag
@@ -2959,17 +2972,17 @@ FString FBodyInstance::GetBodyDebugName() const
 	FString DebugName;
 
 	UPrimitiveComponent* OwnerComponentInst = OwnerComponent.Get();
-	if(OwnerComponentInst != NULL)
+	if (OwnerComponentInst != NULL)
 	{
 		DebugName = OwnerComponentInst->GetPathName();
-		if(const UObject* StatObject = OwnerComponentInst->AdditionalStatObject())
+		if (const UObject* StatObject = OwnerComponentInst->AdditionalStatObject())
 		{
 			DebugName += TEXT(" ");
 			StatObject->AppendName(DebugName);
 		}
 	}
 
-	if((BodySetup != NULL) && (BodySetup->BoneName != NAME_None))
+	if ((BodySetup != NULL) && (BodySetup->BoneName != NAME_None))
 	{
 		DebugName += FString(TEXT(" Bone: ")) + BodySetup->BoneName.ToString();
 	}
@@ -2992,7 +3005,7 @@ bool FBodyInstance::LineTrace(struct FHitResult& OutHit, const FVector& Start, c
 bool FBodyInstance::Sweep(struct FHitResult& OutHit, const FVector& Start, const FVector& End, const FQuat& ShapeWorldRotation, const FCollisionShape& CollisionShape, bool bTraceComplex) const
 {
 	return FPhysicsInterface::Sweep_Geom(OutHit, this, Start, End, ShapeWorldRotation, CollisionShape, bTraceComplex);
-}
+		}
 
 bool FBodyInstance::GetSquaredDistanceToBody(const FVector& Point, float& OutDistanceSquared, FVector& OutPointOnBody) const
 {
@@ -3033,11 +3046,11 @@ bool FBodyInstance::OverlapTestForBodiesImpl(const FVector& Pos, const FQuat& Ro
 
 			// Calc shape global pose
 			FTransform PShapeGlobalPose = FPhysicsInterface::GetLocalTransform(Shape) * PTestGlobalPose;
-			for(const FBodyInstance* BodyInstance : Bodies)
+			for (const FBodyInstance* BodyInstance : Bodies)
 			{
 				bHaveOverlap = FPhysicsInterface::Overlap_Geom(BodyInstance, FPhysicsInterface::GetGeometryCollection(Shape), PShapeGlobalPose);
 
-				if(bHaveOverlap)
+				if (bHaveOverlap)
 				{
 					return;
 				}
@@ -3072,10 +3085,10 @@ bool FBodyInstance::OverlapTest(const FVector& Position, const FQuat& Rotation,
 
 FTransform RootSpaceToWeldedSpace(const FBodyInstance* BI, const FTransform& RootTM)
 {
-	if(BI->WeldParent)
+	if (BI->WeldParent)
 	{
 		UPrimitiveComponent* BIOwnerComponentInst = BI->OwnerComponent.Get();
-		if(BIOwnerComponentInst)
+		if (BIOwnerComponentInst)
 		{
 			FTransform RootToWelded = BIOwnerComponentInst->GetRelativeTransform().Inverse();
 			RootToWelded.ScaleTranslation(BI->Scale3D);
@@ -3092,7 +3105,7 @@ bool FBodyInstance::OverlapMulti(TArray<struct FOverlapResult>& InOutOverlaps, c
 	SCOPE_CYCLE_COUNTER(STAT_Collision_SceneQueryTotal);
 	SCOPE_CYCLE_COUNTER(STAT_Collision_FBodyInstance_OverlapMulti);
 
-	if(!IsValidBodyInstance() && (!WeldParent || !WeldParent->IsValidBodyInstance()))
+	if ( !IsValidBodyInstance()  && (!WeldParent || !WeldParent->IsValidBodyInstance()))
 	{
 		UE_LOG(LogCollision, Log, TEXT("FBodyInstance::OverlapMulti : (%s) No physics data"), *GetBodyDebugName());
 		return false;
@@ -3104,7 +3117,7 @@ bool FBodyInstance::OverlapMulti(TArray<struct FOverlapResult>& InOutOverlaps, c
 	const FTransform ComponentSpaceToTestSpace(Quat, Pos);
 
 	FTransform BodyInstanceSpaceToTestSpace(NoInit);
-	if(pWorldToComponent)
+	if (pWorldToComponent)
 	{
 		const FTransform RootTM = WeldParent ? WeldParent->GetUnrealWorldTransform() : GetUnrealWorldTransform();
 		const FTransform LocalOffset = (*pWorldToComponent) * RootTM;
@@ -3117,7 +3130,7 @@ bool FBodyInstance::OverlapMulti(TArray<struct FOverlapResult>& InOutOverlaps, c
 
 	//We want to test using global position. However, the global position of the body will be in terms of the root body which we are welded to. So we must undo the relative transform so that our shapes are centered
 	//Global = Parent * Relative => Global * RelativeInverse = Parent
-	if(WeldParent)
+	if (WeldParent)
 	{
 		BodyInstanceSpaceToTestSpace = RootSpaceToWeldedSpace(this, BodyInstanceSpaceToTestSpace);
 	}
@@ -3128,13 +3141,13 @@ bool FBodyInstance::OverlapMulti(TArray<struct FOverlapResult>& InOutOverlaps, c
 	{
 		if(Actor.IsValid())
 		{
-			// Get all the shapes from the actor
+		// Get all the shapes from the actor
 			FInlineShapeArray PShapes;
 			const int32 NumShapes = FillInlineShapeArray_AssumesLocked(PShapes, Actor);
 
 			// Iterate over each shape
 			TArray<struct FOverlapResult> TempOverlaps;
-			for(int32 ShapeIdx = 0; ShapeIdx < NumShapes; ShapeIdx++)
+			for (int32 ShapeIdx = 0; ShapeIdx < NumShapes; ShapeIdx++)
 			{
 #if WITH_CHAOS || WITH_IMMEDIATE_PHYSX
 				check(false);
@@ -3157,7 +3170,7 @@ bool FBodyInstance::OverlapMulti(TArray<struct FOverlapResult>& InOutOverlaps, c
 				// Calc shape global pose
 				const FTransform LocalTransform = FPhysicsInterface::GetLocalTransform(ShapeRef);
 				const FTransform GlobalTransform = LocalTransform * BodyInstanceSpaceToTestSpace;
-
+			
 				TempOverlaps.Reset();
 				if(FPhysicsInterface::GeomOverlapMulti(World, GeomCollection, GlobalTransform.GetTranslation(), GlobalTransform.GetRotation(), TempOverlaps, TestChannel, Params, ResponseParams, ObjectQueryParams))
 				{
@@ -3166,8 +3179,8 @@ bool FBodyInstance::OverlapMulti(TArray<struct FOverlapResult>& InOutOverlaps, c
 				InOutOverlaps.Append(TempOverlaps);
 #endif
 			}
-		}
-	});
+			}
+		});
 
 	return bHaveBlockingHit;
 }
@@ -3180,38 +3193,38 @@ bool FBodyInstance::IsValidCollisionProfileName(FName InCollisionProfileName)
 void FBodyInstance::LoadProfileData(bool bVerifyProfile)
 {
 	const FName UseCollisionProfileName = GetCollisionProfileName();
-	if(bVerifyProfile)
+	if ( bVerifyProfile )
 	{
 		// if collision profile name exists, 
 		// check with current settings
 		// if same, then keep the profile name
 		// if not same, that means it has been modified from default
 		// leave it as it is, and clear profile name
-		if(IsValidCollisionProfileName(UseCollisionProfileName))
+		if ( IsValidCollisionProfileName(UseCollisionProfileName) )
 		{
 			FCollisionResponseTemplate Template;
-			if(UCollisionProfile::Get()->GetProfileTemplate(UseCollisionProfileName, Template))
+			if ( UCollisionProfile::Get()->GetProfileTemplate(UseCollisionProfileName, Template) ) 
 			{
 				// this function is only used for old code that did require verification of using profile or not
 				// so that means it will have valid ResponsetoChannels value, so this is okay to access. 
-				if(Template.IsEqual(CollisionEnabled, ObjectType, CollisionResponses.GetResponseContainer()) == false)
+				if (Template.IsEqual(CollisionEnabled, ObjectType, CollisionResponses.GetResponseContainer()) == false)
 				{
-					InvalidateCollisionProfileName();
+					InvalidateCollisionProfileName(); 
 				}
 			}
 			else
 			{
 				UE_LOG(LogPhysics, Warning, TEXT("COLLISION PROFILE [%s] is not found"), *UseCollisionProfileName.ToString());
 				// if not nothing to do
-				InvalidateCollisionProfileName();
+				InvalidateCollisionProfileName(); 
 			}
 		}
 	}
 	else
 	{
-		if(IsValidCollisionProfileName(UseCollisionProfileName))
+		if ( IsValidCollisionProfileName(UseCollisionProfileName) )
 		{
-			if(UCollisionProfile::Get()->ReadConfig(UseCollisionProfileName, *this) == false)
+			if ( UCollisionProfile::Get()->ReadConfig(UseCollisionProfileName, *this) == false)
 			{
 				// clear the name
 				InvalidateCollisionProfileName();
@@ -3219,10 +3232,10 @@ void FBodyInstance::LoadProfileData(bool bVerifyProfile)
 		}
 
 		// no profile, so it just needs to update container from array data
-		if(DoesUseCollisionProfile() == false)
+		if ( DoesUseCollisionProfile() == false )
 		{
 			// if external profile copy the data over
-			if(ExternalCollisionProfileBodySetup.IsValid(true))
+			if (ExternalCollisionProfileBodySetup.IsValid(true))
 			{
 				UBodySetup* BodySetupInstance = ExternalCollisionProfileBodySetup.Get(true);
 				const FBodyInstance& ExternalBodyInstance = BodySetupInstance->DefaultInstance;
@@ -3246,35 +3259,35 @@ void FBodyInstance::GetBodyInstanceResourceSizeEx(FResourceSizeEx& CumulativeRes
 
 void FBodyInstance::FixupData(class UObject* Loader)
 {
-	check(Loader);
+	check (Loader);
 
 	int32 const UE4Version = Loader->GetLinkerUE4Version();
 
 #if WITH_EDITOR
-	if(UE4Version < VER_UE4_ADD_CUSTOMPROFILENAME_CHANGE)
+	if (UE4Version < VER_UE4_ADD_CUSTOMPROFILENAME_CHANGE)
 	{
-		if(CollisionProfileName == NAME_None)
+		if (CollisionProfileName == NAME_None)
 		{
 			CollisionProfileName = UCollisionProfile::CustomCollisionProfileName;
 		}
 	}
 
-	if(UE4Version < VER_UE4_SAVE_COLLISIONRESPONSE_PER_CHANNEL)
+	if (UE4Version < VER_UE4_SAVE_COLLISIONRESPONSE_PER_CHANNEL)
 	{
 		CollisionResponses.SetCollisionResponseContainer(ResponseToChannels_DEPRECATED);
 	}
 #endif // WITH_EDITORONLY_DATA
 
 	// Load profile. If older version, please verify profile name first
-	bool bNeedToVerifyProfile = (UE4Version < VER_UE4_COLLISION_PROFILE_SETTING) ||
+	bool bNeedToVerifyProfile = (UE4Version < VER_UE4_COLLISION_PROFILE_SETTING) || 
 		// or shape component needs to convert since we added profile
 		(UE4Version < VER_UE4_SAVE_COLLISIONRESPONSE_PER_CHANNEL && Loader->IsA(UShapeComponent::StaticClass()));
 	LoadProfileData(bNeedToVerifyProfile);
 
 	// if profile isn't set, then fix up channel responses
-	if(CollisionProfileName == UCollisionProfile::CustomCollisionProfileName)
+	if( CollisionProfileName == UCollisionProfile::CustomCollisionProfileName ) 
 	{
-		if(UE4Version >= VER_UE4_SAVE_COLLISIONRESPONSE_PER_CHANNEL)
+		if (UE4Version >= VER_UE4_SAVE_COLLISIONRESPONSE_PER_CHANNEL)
 		{
 			CollisionResponses.UpdateResponseContainerFromArray();
 		}
@@ -3282,11 +3295,11 @@ void FBodyInstance::FixupData(class UObject* Loader)
 }
 
 void FBodyInstance::ApplyMaterialToShape_AssumesLocked(const FPhysicsShapeHandle& InShape, UPhysicalMaterial* SimplePhysMat, const TArrayView<UPhysicalMaterial*>& ComplexPhysMats)
-{
+	{
 	// If a triangle mesh, need to get array of materials...
 	ECollisionShapeType GeomType = FPhysicsInterface::GetShapeType(InShape);
 	if(GeomType == ECollisionShapeType::Trimesh)
-	{
+		{
 		if(ComplexPhysMats.Num())
 		{
 			FPhysicsInterface::SetMaterials(InShape, ComplexPhysMats);
@@ -3329,12 +3342,12 @@ void FBodyInstance::ApplyMaterialToInstanceShapes_AssumesLocked(UPhysicalMateria
 		if(TheirBI->IsShapeBoundToBody(Shape))
 		{
 			FPhysicsCommand::ExecuteShapeWrite(BIWithActor, Shape, [&](const FPhysicsShapeHandle& InnerShape)
-			{
+		{
 				ApplyMaterialToShape_AssumesLocked(InnerShape, SimplePhysMat, ComplexPhysMats);
-			});
-		}
+		});		
 	}
 }
+}
 
 bool FBodyInstance::ValidateTransform(const FTransform &Transform, const FString& DebugName, const UBodySetup* Setup)
 {
@@ -3371,7 +3384,7 @@ bool FBodyInstance::ValidateTransform(const FTransform &Transform, const FString
 
 void FBodyInstance::InitDynamicProperties_AssumesLocked()
 {
-	if(!BodySetup.IsValid())
+	if (!BodySetup.IsValid())
 	{
 		// This may be invalid following an undo if the BodySetup was a transient object (e.g. in Mesh Paint mode)
 		// Just exit gracefully if so.
@@ -3380,16 +3393,16 @@ void FBodyInstance::InitDynamicProperties_AssumesLocked()
 
 	//QueryOnly bodies cannot become simulated at runtime. To do this they must change their CollisionEnabled which recreates the physics state
 	//So early out to save a lot of useless work
-	if(GetCollisionEnabled() == ECollisionEnabled::QueryOnly)
+	if (GetCollisionEnabled() == ECollisionEnabled::QueryOnly)
 	{
 		return;
 	}
-
+	
 	if(FPhysicsInterface::IsDynamic(ActorHandle))
 	{
 		//A non simulated body may become simulated at runtime, so we need to compute its mass.
 		//However, this is not supported for complexAsSimple since a trimesh cannot itself be simulated, it can only be used for collision of other simple shapes.
-		if(BodySetup->GetCollisionTraceFlag() != ECollisionTraceFlag::CTF_UseComplexAsSimple)
+		if (BodySetup->GetCollisionTraceFlag() != ECollisionTraceFlag::CTF_UseComplexAsSimple)
 		{
 			UpdateMassProperties();
 			UpdateDampingProperties();
@@ -3399,7 +3412,7 @@ void FBodyInstance::InitDynamicProperties_AssumesLocked()
 		else
 		{
 #if !(UE_BUILD_SHIPPING || UE_BUILD_TEST)
-			if(bSimulatePhysics)
+			if (bSimulatePhysics)
 			{
 				if(UPrimitiveComponent* OwnerComponentInst = OwnerComponent.Get())
 				{
@@ -3416,7 +3429,7 @@ void FBodyInstance::InitDynamicProperties_AssumesLocked()
 		bool bComponentAwake = false;
 		FVector InitialLinVel = GetInitialLinearVelocity(OwningActor, bComponentAwake);
 
-		if(ShouldInstanceSimulatingPhysics())
+		if (ShouldInstanceSimulatingPhysics())
 		{
 			FPhysicsInterface::SetLinearVelocity_AssumesLocked(ActorHandle, InitialLinVel);
 		}
@@ -3431,8 +3444,8 @@ void FBodyInstance::InitDynamicProperties_AssumesLocked()
 		StabilizationThreshold *= StabilizationThresholdMultiplier;
 		FPhysicsInterface::SetStabilizationEnergyThreshold_AssumesLocked(ActorHandle, StabilizationThreshold);
 
-		uint32 PositionIterCount = FMath::Clamp(PositionSolverIterationCount, 1, 255);
-		uint32 VelocityIterCount = FMath::Clamp(VelocitySolverIterationCount, 1, 255);
+		uint32 PositionIterCount = FMath::Clamp<uint8>(PositionSolverIterationCount, 1, 255);
+		uint32 VelocityIterCount = FMath::Clamp<uint8>(VelocitySolverIterationCount, 1, 255);
 		FPhysicsInterface::SetSolverPositionIterationCount_AssumesLocked(ActorHandle, PositionIterCount);
 		FPhysicsInterface::SetSolverVelocityIterationCount_AssumesLocked(ActorHandle, VelocityIterCount);
 
@@ -3449,7 +3462,7 @@ void FBodyInstance::InitDynamicProperties_AssumesLocked()
 }
 
 void FBodyInstance::BuildBodyFilterData(FBodyCollisionFilterData& OutFilterData) const
-{
+	{
 	// this can happen in landscape height field collision component
 	if(!BodySetup.IsValid())
 	{
@@ -3480,7 +3493,7 @@ void FBodyInstance::BuildBodyFilterData(FBodyCollisionFilterData& OutFilterData)
 		{
 			UseResponse.SetAllChannels(ECR_Ignore);
 			UseCollisionEnabled = ECollisionEnabled::PhysicsOnly;
-		}
+			}
 		else if(BodySetup->CollisionReponse == EBodyCollisionResponse::BodyCollision_Enabled)
 		{
 			UseResponse.SetAllChannels(ECR_Block);
@@ -3494,7 +3507,7 @@ void FBodyInstance::BuildBodyFilterData(FBodyCollisionFilterData& OutFilterData)
 		const bool bDisableSkelComponentOverride = CVarEnableDynamicPerBodyFilterHacks.GetValueOnGameThread() && bHACK_DisableSkelComponentFilterOverriding;
 		if(!bDisableSkelComponentOverride)
 		{
-			UseResponse = FCollisionResponseContainer::CreateMinContainer(UseResponse, SkelMeshComp->BodyInstance.CollisionResponses.GetResponseContainer());
+		UseResponse = FCollisionResponseContainer::CreateMinContainer(UseResponse, SkelMeshComp->BodyInstance.CollisionResponses.GetResponseContainer());
 		}
 
 		bUseNotifyRBCollision = bUseNotifyRBCollision && SkelMeshComp->BodyInstance.bNotifyRigidBodyCollision;
@@ -3507,7 +3520,7 @@ void FBodyInstance::BuildBodyFilterData(FBodyCollisionFilterData& OutFilterData)
 		UWorld* World = Owner->GetWorld();
 		UPrimitiveComponent* PrimComp = OwnerComponentInst;
 		if(World && World->bEnableTraceCollision &&
-			(PrimComp->IsA(UStaticMeshComponent::StaticClass()) || PrimComp->IsA(USkeletalMeshComponent::StaticClass()) || PrimComp->IsA(UBrushComponent::StaticClass())))
+		   (PrimComp->IsA(UStaticMeshComponent::StaticClass()) || PrimComp->IsA(USkeletalMeshComponent::StaticClass()) || PrimComp->IsA(UBrushComponent::StaticClass())))
 		{
 			//UE_LOG(LogPhysics, Warning, TEXT("Enabling collision %s : %s"), *GetNameSafe(Owner), *GetNameSafe(OwnerComponent.Get()));
 			// clear all other channel just in case other people using those channels to do something
@@ -3528,28 +3541,28 @@ void FBodyInstance::BuildBodyFilterData(FBodyCollisionFilterData& OutFilterData)
 		FCollisionFilterData SimFilterData;
 		FCollisionFilterData SimpleQueryData;
 
-		uint32 ActorID = Owner ? Owner->GetUniqueID() : 0;
-		uint32 CompID = (OwnerComponentInst != nullptr) ? OwnerComponentInst->GetUniqueID() : 0;
+			uint32 ActorID = Owner ? Owner->GetUniqueID() : 0;
+			uint32 CompID = (OwnerComponentInst != nullptr) ? OwnerComponentInst->GetUniqueID() : 0;
 		CreateShapeFilterData(UseChannel, MaskFilter, ActorID, UseResponse, CompID, InstanceBodyIndex, SimpleQueryData, SimFilterData, bRootCCD && !bPhysicsStatic, bUseNotifyRBCollision, bPhysicsStatic, bUseContactModification);
 
 		FCollisionFilterData ComplexQueryData = SimpleQueryData;
-
-		// Set output sim data
+			
+			// Set output sim data
 		OutFilterData.SimFilter = SimFilterData;
 
-		// Build filterdata variations for complex and simple
+			// Build filterdata variations for complex and simple
 		SimpleQueryData.Word3 |= EPDF_SimpleCollision;
-		if(bUseSimpleAsComplex)
-		{
+			if(bUseSimpleAsComplex)
+			{
 			SimpleQueryData.Word3 |= EPDF_ComplexCollision;
-		}
+			}
 
 		ComplexQueryData.Word3 |= EPDF_ComplexCollision;
-		if(bUseComplexAsSimple)
-		{
+			if(bUseComplexAsSimple)
+			{
 			ComplexQueryData.Word3 |= EPDF_SimpleCollision;
-		}
-
+			}
+			
 		OutFilterData.QuerySimpleFilter = SimpleQueryData;
 		OutFilterData.QueryComplexFilter = ComplexQueryData;
 	}
@@ -3582,7 +3595,7 @@ void FBodyInstance::InitStaticBodies(const TArray<FBodyInstance*>& Bodies, const
 
 
 void FBodyInstance::BuildBodyCollisionFlags(FBodyCollisionFlags& OutFlags, ECollisionEnabled::Type UseCollisionEnabled, bool bUseComplexAsSimple)
-{
+	{
 	if(UseCollisionEnabled != ECollisionEnabled::NoCollision)
 	{
 		// Query collision
@@ -3595,15 +3608,15 @@ void FBodyInstance::BuildBodyCollisionFlags(FBodyCollisionFlags& OutFlags, EColl
 		if(bSimCollision)
 		{
 			OutFlags.bEnableSimCollisionSimple = true;
-
+			
 			// on dynamic objects and objects which don't use complex as simple, tri mesh not used for sim
 			if(bUseComplexAsSimple)
 			{
 				OutFlags.bEnableSimCollisionComplex = true;
 			}
 		}
-	}
-}
+			}
+		}
 
 void FBodyInstance::UpdateInterpolateWhenSubStepping()
 {
@@ -3619,7 +3632,7 @@ void FBodyInstance::UpdateInterpolateWhenSubStepping()
 		{
 			UseCollisionEnabled = GetCollisionEnabled();
 		}
-
+	
 		bInterpolateWhenSubStepping = UseCollisionEnabled == ECollisionEnabled::PhysicsOnly || UseCollisionEnabled == ECollisionEnabled::QueryAndPhysics;
 
 		// If we have a weld parent we should take into account that too as that may be simulating while we are not
@@ -3632,6 +3645,35 @@ void FBodyInstance::UpdateInterpolateWhenSubStepping()
 	}
 }
 
+void FBodyInstance::ExecuteOnCalculateCustomProjection(FTransform& WorldTM) const
+{
+	if (BodyInstanceDelegates.IsValid())
+	{
+		BodyInstanceDelegates->OnCalculateCustomProjection.ExecuteIfBound(this, WorldTM);
+	}
+}
+
+FCalculateCustomProjection& FBodyInstance::OnCalculateCustomProjection()
+{
+	if (!BodyInstanceDelegates.IsValid())
+	{
+		BodyInstanceDelegates.Reset(new FBodyInstanceDelegates);
+	}
+
+	return BodyInstanceDelegates->OnCalculateCustomProjection;
+}
+
+FRecalculatedMassProperties& FBodyInstance::OnRecalculatedMassProperties()
+{
+	if (!BodyInstanceDelegates.IsValid())
+	{
+		BodyInstanceDelegates.Reset(new FBodyInstanceDelegates);
+	}
+
+	return BodyInstanceDelegates->OnRecalculatedMassProperties;
+}
+
+
 ////////////////////////////////////////////////////////////////////////////
 // FBodyInstanceEditorHelpers
 
@@ -3639,7 +3681,7 @@ void FBodyInstance::UpdateInterpolateWhenSubStepping()
 
 void FBodyInstanceEditorHelpers::EnsureConsistentMobilitySimulationSettingsOnPostEditChange(UPrimitiveComponent* Component, FPropertyChangedEvent& PropertyChangedEvent)
 {
-	if(UProperty* PropertyThatChanged = PropertyChangedEvent.Property)
+	if (UProperty* PropertyThatChanged = PropertyChangedEvent.Property)
 	{
 		const FName PropertyName = PropertyThatChanged->GetFName();
 
@@ -3647,15 +3689,15 @@ void FBodyInstanceEditorHelpers::EnsureConsistentMobilitySimulationSettingsOnPos
 		const bool bMobilityChanged = PropertyName == GET_MEMBER_NAME_CHECKED(USceneComponent, Mobility);
 		const bool bSimulatePhysicsChanged = PropertyName == GET_MEMBER_NAME_CHECKED(FBodyInstance, bSimulatePhysics);
 
-		if(bMobilityChanged || bSimulatePhysicsChanged)
+		if (bMobilityChanged || bSimulatePhysicsChanged)
 		{
 			// If we enabled physics simulation, but we are not marked movable, do that for them
-			if(bSimulatePhysicsChanged && Component->BodyInstance.bSimulatePhysics && (Component->Mobility != EComponentMobility::Movable))
+			if (bSimulatePhysicsChanged && Component->BodyInstance.bSimulatePhysics && (Component->Mobility != EComponentMobility::Movable))
 			{
 				Component->SetMobility(EComponentMobility::Movable);
 			}
 			// If we made the component no longer movable, but simulation was enabled, disable that for them
-			else if(bMobilityChanged && (Component->Mobility != EComponentMobility::Movable) && Component->BodyInstance.bSimulatePhysics)
+			else if (bMobilityChanged && (Component->Mobility != EComponentMobility::Movable) && Component->BodyInstance.bSimulatePhysics)
 			{
 				Component->BodyInstance.bSimulatePhysics = false;
 			}
@@ -3663,13 +3705,13 @@ void FBodyInstanceEditorHelpers::EnsureConsistentMobilitySimulationSettingsOnPos
 			// If the collision profile is one of the 'default' ones for a StaticMeshActor, make sure it is the correct one
 			// If user has changed it to something else, don't touch it
 			const FName CurrentProfileName = Component->BodyInstance.GetCollisionProfileName();
-			if((CurrentProfileName == UCollisionProfile::BlockAll_ProfileName) ||
+			if ((CurrentProfileName == UCollisionProfile::BlockAll_ProfileName) ||
 				(CurrentProfileName == UCollisionProfile::BlockAllDynamic_ProfileName) ||
 				(CurrentProfileName == UCollisionProfile::PhysicsActor_ProfileName))
 			{
-				if(Component->Mobility == EComponentMobility::Movable)
+				if (Component->Mobility == EComponentMobility::Movable)
 				{
-					if(Component->BodyInstance.bSimulatePhysics)
+					if (Component->BodyInstance.bSimulatePhysics)
 					{
 						Component->SetCollisionProfileName(UCollisionProfile::PhysicsActor_ProfileName);
 					}
