#UE4 Reduce compile time of FFastArraySerializer::FastArrayDeltaSerialize by moving the non-specialized parts to a cpp file.

diff --git a/Engine/Source/Runtime/Engine/Classes/Engine/NetSerialization.h b/Engine/Source/Runtime/Engine/Classes/Engine/NetSerialization.h
--- a/Engine/Source/Runtime/Engine/Classes/Engine/NetSerialization.h
+++ b/Engine/Source/Runtime/Engine/Classes/Engine/NetSerialization.h
@@ -416,7 +416,11 @@ struct FFastArraySerializer
 	}
 
 	template< typename Type, typename SerializerType >
-	static bool FastArrayDeltaSerialize( TArray<Type> &Items, FNetDeltaSerializeInfo& Parms, SerializerType& ArraySerializer );
+	static bool FastArrayDeltaSerialize( TArray<Type>& Items, FNetDeltaSerializeInfo& Parms, SerializerType& ArraySerializer )
+	{
+		TFastArrayDeltaSerializeAccessors<Type, SerializerType> Accessors(Items, ArraySerializer);
+		return FastArrayDeltaSerialize_Internal(Accessors, Parms, ArraySerializer, Type::StaticStruct());
+	}
 
 	/**
 	 * Called before removing elements and after the elements themselves are notified.  The indices are valid for this function call only!
@@ -455,6 +459,40 @@ struct FFastArraySerializer
 	}
 
 private:
+	struct FFastArrayDeltaSerializeAccessors
+	{
+		virtual ~FFastArrayDeltaSerializeAccessors() {}
+		virtual int32 GetNumItems() const = 0;
+		virtual FFastArraySerializerItem& GetItem(int32 Index) const = 0;
+		virtual FFastArraySerializerItem& AddItem() const = 0;
+		virtual void RemoveItem(int32 Index) const = 0;
+		virtual void PostReplicatedAdd(int32 Index) const = 0;
+		virtual void PostReplicatedChange(int32 Index) const = 0;
+		virtual bool ShouldWriteFastArrayItem(int32 Index, bool bIsWritingOnClient) const = 0;
+	};
+
+	template< typename Type, typename SerializerType >
+	struct TFastArrayDeltaSerializeAccessors : public FFastArrayDeltaSerializeAccessors
+	{
+		TFastArrayDeltaSerializeAccessors(TArray<Type>& InItems, SerializerType& InArraySerializer) : Items(InItems), ArraySerializer(InArraySerializer) {}
+		virtual int32 GetNumItems() const override { return Items.Num(); }
+		virtual FFastArraySerializerItem& GetItem(int32 Index) const override { return Items[Index]; }
+		virtual FFastArraySerializerItem& AddItem() const override { return Items.AddDefaulted_GetRef(); }
+		virtual void RemoveItem(int32 Index) const override { Items.RemoveAtSwap(Index, 1, false); }
+		virtual void PostReplicatedAdd(int32 Index) const override { Items[Index].PostReplicatedAdd(ArraySerializer); }
+		virtual void PostReplicatedChange(int32 Index) const override { Items[Index].PostReplicatedChange(ArraySerializer); }
+		virtual bool ShouldWriteFastArrayItem(int32 Index, bool bIsWritingOnClient) const override { return ArraySerializer.template ShouldWriteFastArrayItem<Type, SerializerType>(Items[Index], bIsWritingOnClient); }
+
+		TArray<Type>& Items;
+		SerializerType& ArraySerializer;
+	};
+
+	ENGINE_API FORCENOINLINE static bool FastArrayDeltaSerialize_Internal(
+		FFastArrayDeltaSerializeAccessors& Accessors,
+		FNetDeltaSerializeInfo& Parms,
+		FFastArraySerializer& ArraySerializer,
+		UScriptStruct* InnerStruct);
+
 	// Cached item counts, used for fast sanity checking when writing.
 	int32				CachedNumItems;
 	int32				CachedNumItemsToConsiderForWriting;
@@ -469,658 +507,6 @@ struct FFastArraySerializer_FastArrayDeltaSerialize_FIdxIDPair
 	int32	ID;
 };
 
-/** The function that implements Fast TArray Replication  */
-template< typename Type, typename SerializerType >
-bool FFastArraySerializer::FastArrayDeltaSerialize( TArray<Type> &Items, FNetDeltaSerializeInfo& Parms, SerializerType& ArraySerializer )
-{
-	SCOPE_CYCLE_COUNTER(STAT_NetSerializeFastArray);
-	class UScriptStruct* InnerStruct = Type::StaticStruct();
-
-	UE_LOG(LogNetFastTArray, Log, TEXT("FastArrayDeltaSerialize for %s. %s. %s"), *InnerStruct->GetName(), *InnerStruct->GetOwnerStruct()->GetName(), Parms.Reader ? TEXT("Reading") : TEXT("Writing"));
-
-	if ( Parms.bUpdateUnmappedObjects || Parms.Writer == NULL )
-	{
-		//---------------
-		// Build ItemMap if necessary. This maps ReplicationID to our local index into the Items array.
-		//---------------
-		if (ArraySerializer.ItemMap.Num() != Items.Num())
-		{
-			SCOPE_CYCLE_COUNTER(STAT_NetSerializeFastArray_BuildMap);
-			UE_LOG(LogNetFastTArray, Log, TEXT("FastArrayDeltaSerialize: Recreating Items map. Struct: %s, Items.Num: %d Map.Num: %d"), *InnerStruct->GetOwnerStruct()->GetName(), Items.Num(), ArraySerializer.ItemMap.Num() );
-
-			ArraySerializer.ItemMap.Empty();
-			for (int32 i=0; i < Items.Num(); ++i)
-			{
-				if ( Items[i].ReplicationID == INDEX_NONE  )
-				{
-					if (Parms.Writer)
-					{
-						UE_LOG( LogNetFastTArray, Warning, TEXT( "FastArrayDeltaSerialize: Item with uninitialized ReplicationID. Struct: %s, ItemIndex: %i" ), *InnerStruct->GetOwnerStruct()->GetName(), i );
-					}
-					else
-					{
-						// This is benign for clients, they may add things to their local array without assigning a ReplicationID
-						continue;
-					}
-				}
-				ArraySerializer.ItemMap.Add( Items[i].ReplicationID, i );
-			}
-		}
-	}
-
-	if ( Parms.GatherGuidReferences )
-	{
-		// Loop over all tracked guids, and return what we have
-		for ( const auto& GuidReferencesPair : ArraySerializer.GuidReferencesMap )
-		{
-			const FFastArraySerializerGuidReferences& GuidReferences = GuidReferencesPair.Value;
-
-			Parms.GatherGuidReferences->Append( GuidReferences.UnmappedGUIDs );
-			Parms.GatherGuidReferences->Append( GuidReferences.MappedDynamicGUIDs );
-
-			if ( Parms.TrackedGuidMemoryBytes )
-			{
-				*Parms.TrackedGuidMemoryBytes += GuidReferences.Buffer.Num();
-			}
-		}
-
-		return true;
-	}
-
-	if ( Parms.MoveGuidToUnmapped )
-	{
-		bool bFound = false;
-
-		const FNetworkGUID GUID = *Parms.MoveGuidToUnmapped;
-
-		// Try to find the guid in the list, and make sure it's on the unmapped lists now
-		for ( auto& GuidReferencesPair : ArraySerializer.GuidReferencesMap )
-		{
-			FFastArraySerializerGuidReferences& GuidReferences = GuidReferencesPair.Value;
-
-			if ( GuidReferences.MappedDynamicGUIDs.Contains( GUID ) )
-			{
-				GuidReferences.MappedDynamicGUIDs.Remove( GUID );
-				GuidReferences.UnmappedGUIDs.Add( GUID );
-				bFound = true;
-			}
-		}
-
-		return bFound;
-	}
-
-	if ( Parms.bUpdateUnmappedObjects )
-	{
-		// Loop over each item that has unmapped objects
-		for ( auto It = ArraySerializer.GuidReferencesMap.CreateIterator(); It; ++It )
-		{
-			// Get the element id
-			const int32 ElementID = It.Key();
-			
-			// Get a reference to the unmapped item itself
-			FFastArraySerializerGuidReferences& GuidReferences = It.Value();
-
-			if ( ( GuidReferences.UnmappedGUIDs.Num() == 0 && GuidReferences.MappedDynamicGUIDs.Num() == 0 ) || ArraySerializer.ItemMap.Find( ElementID ) == NULL )
-			{
-				// If for some reason the item is gone (or all guids were removed), we don't need to track guids for this item anymore
-				It.RemoveCurrent();
-				continue;		// We're done with this unmapped item
-			}
-
-			// Loop over all the guids, and check to see if any of them are loaded yet
-			bool bMappedSomeGUIDs = false;
-
-			for ( auto UnmappedIt = GuidReferences.UnmappedGUIDs.CreateIterator(); UnmappedIt; ++UnmappedIt )
-			{
-				const FNetworkGUID& GUID = *UnmappedIt;
-
-				if ( Parms.Map->IsGUIDBroken( GUID, false ) )
-				{
-					// Stop trying to load broken guids
-					UE_LOG( LogNetFastTArray, Warning, TEXT( "FastArrayDeltaSerialize: Broken GUID. NetGuid: %s" ), *GUID.ToString() );
-					UnmappedIt.RemoveCurrent();
-					continue;
-				}
-
-				UObject* Object = Parms.Map->GetObjectFromNetGUID( GUID, false );
-
-				if ( Object != NULL )
-				{
-					// This guid loaded!
-					if ( GUID.IsDynamic() )
-					{
-						GuidReferences.MappedDynamicGUIDs.Add( GUID );		// Move back to mapped list
-					}
-					UnmappedIt.RemoveCurrent();
-					bMappedSomeGUIDs = true;
-				}
-			}
-
-			// Check to see if we loaded any guids. If we did, we can serialize the element again which will load it this time
-			if ( bMappedSomeGUIDs )
-			{
-				Parms.bOutSomeObjectsWereMapped = true;
-
-				if ( !Parms.bCalledPreNetReceive )
-				{
-					// Call PreNetReceive if we are going to change a value (some game code will need to think this is an actual replicated value)
-					Parms.Object->PreNetReceive();
-					Parms.bCalledPreNetReceive = true;
-				}
-
-				Type* ThisElement = &Items[ArraySerializer.ItemMap.FindChecked( ElementID )];
-
-				// Initialize the reader with the stored buffer that we need to read from
-				FNetBitReader Reader( Parms.Map, GuidReferences.Buffer.GetData(), GuidReferences.NumBufferBits );
-
-				// Read the property (which should serialize any newly mapped objects as well)
-				bool bHasUnmapped = false;
-				Parms.NetSerializeCB->NetSerializeStruct( InnerStruct, Reader, Parms.Map, ThisElement, bHasUnmapped );
-
-				// Let the element know it changed
-				ThisElement->PostReplicatedChange( ArraySerializer );
-			}
-
-			// If we have no more guids, we can remove this item for good
-			if ( GuidReferences.UnmappedGUIDs.Num() == 0 && GuidReferences.MappedDynamicGUIDs.Num() == 0 )
-			{
-				It.RemoveCurrent();
-			}
-		}
-
-		// If we still have unmapped items, then communicate this to the outside
-		if ( ArraySerializer.GuidReferencesMap.Num() > 0 )
-		{
-			Parms.bOutHasMoreUnmapped = true;
-		}
-
-		return true;
-	}
-
-	if (Parms.Writer)
-	{
-		//-----------------------------
-		// Saving
-		//-----------------------------	
-		check(Parms.Struct);
-		FBitWriter& Writer = *Parms.Writer;
-
-		// Get the old map if its there
-		TMap<int32, int32> * OldMap = Parms.OldState ? &((FNetFastTArrayBaseState*)Parms.OldState)->IDToCLMap : NULL;
-		int32 BaseReplicationKey = Parms.OldState ? ((FNetFastTArrayBaseState*)Parms.OldState)->ArrayReplicationKey : -1;
-
-		// Helper for counting num elements we should consider
-		auto CalcNumItemsForConsideration = [&ArraySerializer, &Items, &Parms]()
-		{
-			int32 Count = 0;
-
-			// Count the number of items in the current array that may be written. On clients, items that were predicted will be skipped.
-			for (const Type& Item : Items)
-			{
-				if (ArraySerializer.template ShouldWriteFastArrayItem<Type, SerializerType>(Item, Parms.bIsWritingOnClient))
-				{
-					Count++;
-				}
-			}
-			return Count;
-		};
-
-		// See if the array changed at all. If the ArrayReplicationKey matches we can skip checking individual items
-		if (Parms.OldState && (ArraySerializer.ArrayReplicationKey == BaseReplicationKey))
-		{
-			// Double check the old map is valid and that we will consider writing the same number of elements that are in the old map.
-			if (ensureMsgf(OldMap, TEXT("Invalid OldMap")))
-			{
-				// If the keys didn't change, only update the item count caches if necessary.
-				if (ArraySerializer.CachedNumItems == INDEX_NONE ||
-					ArraySerializer.CachedNumItems != Items.Num() ||
-					ArraySerializer.CachedNumItemsToConsiderForWriting == INDEX_NONE)
-				{
-					ArraySerializer.CachedNumItems = Items.Num();
-					ArraySerializer.CachedNumItemsToConsiderForWriting = CalcNumItemsForConsideration();
-				}
-
-				if (UNLIKELY(OldMap->Num() != ArraySerializer.CachedNumItemsToConsiderForWriting))
-				{
-					UE_LOG(LogNetFastTArray, Warning, TEXT("OldMap size (%d) does not match item count (%d)"), OldMap->Num(), ArraySerializer.CachedNumItemsToConsiderForWriting);
-				}
-			}
-
-			if (Parms.OldState)
-			{
-				// Nothing changed and we had a valid old state, so just use/share the existing state. No need to create a new one.
-				*Parms.NewState = Parms.OldState->AsShared();
-			}
-			else
-			{
-				// Nothing changed but we don't have an existing state of our own yet so we need to make one here.
-				FNetFastTArrayBaseState * NewState = new FNetFastTArrayBaseState();
-				*Parms.NewState = MakeShareable( NewState );
-				NewState->ArrayReplicationKey = ArraySerializer.ArrayReplicationKey;
-			}
-			return false;
-		}
-
-
-		// Create a new map from the current state of the array		
-		FNetFastTArrayBaseState * NewState = new FNetFastTArrayBaseState();
-
-		check(Parms.NewState);
-		*Parms.NewState = MakeShareable( NewState );
-		TMap<int32, int32> & NewMap = NewState->IDToCLMap;
-		NewState->ArrayReplicationKey = ArraySerializer.ArrayReplicationKey;
-
-
-		int32 NumConsideredItems = CalcNumItemsForConsideration();
-
-		TArray<FFastArraySerializer_FastArrayDeltaSerialize_FIdxIDPair, TInlineAllocator<8> >	ChangedElements;
-		TArray<int32, TInlineAllocator<8> >		DeletedElements;
-
-		int32 DeleteCount = (OldMap ? OldMap->Num() : 0) - NumConsideredItems; // Note: this is incremented when we add new items below.
-		UE_LOG(LogNetFastTArray, Log, TEXT("NetSerializeItemDeltaFast: %s. DeleteCount: %d"), *Parms.DebugName, DeleteCount);
-
-		// Log out entire state of current/base state
-		if (UE_LOG_ACTIVE(LogNetFastTArray,Log))
-		{
-			FString CurrentState = FString::Printf(TEXT("Current: %d "), ArraySerializer.ArrayReplicationKey );
-			for (int32 i=0; i < Items.Num(); ++i)
-			{
-				CurrentState += FString::Printf(TEXT("[%d/%d], "), Items[i].ReplicationID, Items[i].ReplicationKey);
-			}
-			UE_LOG(LogNetFastTArray, Log, TEXT("%s"), *CurrentState);
-
-		
-			FString ClientStateStr = FString::Printf(TEXT("Client: %d "),  Parms.OldState ? ((FNetFastTArrayBaseState*)Parms.OldState)->ArrayReplicationKey: 0);
-			if (OldMap)
-			{
-				for (auto It = OldMap->CreateIterator(); It; ++It)
-				{
-					ClientStateStr += FString::Printf(TEXT("[%d/%d], "), It.Key(), It.Value() );
-				}
-			}
-			UE_LOG(LogNetFastTArray, Log, TEXT("%s"), *ClientStateStr);
-		}
-
-
-		//--------------------------------------------
-		// Find out what is new or what has changed
-		//--------------------------------------------
-		for (int32 i=0; i < Items.Num(); ++i)
-		{
-			UE_LOG(LogNetFastTArray, Log, TEXT("    Array[%d] - ID %d. CL %d."), i, Items[i].ReplicationID, Items[i].ReplicationKey);
-			if (!ArraySerializer.template ShouldWriteFastArrayItem<Type, SerializerType>(Items[i], Parms.bIsWritingOnClient))
-			{
-				// On clients, this will skip items that were added predictively.
-				continue;
-			}
-			if (Items[i].ReplicationID == INDEX_NONE)
-			{
-				ArraySerializer.MarkItemDirty(Items[i]);
-			}
-			NewMap.Add( Items[i].ReplicationID, Items[i].ReplicationKey );
-
-			int32* OldValuePtr = OldMap ? OldMap->Find( Items[i].ReplicationID ) : NULL;
-			if (OldValuePtr)
-			{
-				if (*OldValuePtr == Items[i].ReplicationKey)
-				{
-					UE_LOG(LogNetFastTArray, Log, TEXT("       Stayed The Same - Skipping"));
-
-					// Stayed the same, it might have moved but we dont care
-					continue;
-				}
-				else
-				{
-					UE_LOG(LogNetFastTArray, Log, TEXT("       Changed! Was: %d. Element ID: %d. %s"), *OldValuePtr, Items[i].ReplicationID, *Items[i].GetDebugString());
-
-					// Changed
-					ChangedElements.Add(FFastArraySerializer_FastArrayDeltaSerialize_FIdxIDPair(i, Items[i].ReplicationID));
-				}
-			}
-			else
-			{
-				UE_LOG(LogNetFastTArray, Log, TEXT("       New! Element ID: %d. %s"), Items[i].ReplicationID, *Items[i].GetDebugString());
-				
-				// The item really should have a valid ReplicationID but in the case of loading from a save game,
-				// items may not have been marked dirty individually. Its ok to just assign them one here.
-				// New
-				ChangedElements.Add(FFastArraySerializer_FastArrayDeltaSerialize_FIdxIDPair(i, Items[i].ReplicationID));
-				DeleteCount++; // We added something new, so our initial DeleteCount value must be incremented.
-			}
-		}
-
-		// Find out what was deleted
-		if( DeleteCount > 0 && OldMap)
-		{
-			for (auto It = OldMap->CreateIterator(); It; ++It)
-			{
-				if (!NewMap.Contains(It.Key()))
-				{
-					UE_LOG(LogNetFastTArray, Log, TEXT("   Deleting ID: %d"), It.Key());
-
-					DeletedElements.Add( It.Key() );
-					if (--DeleteCount <= 0)
-						break;
-				}
-			}
-		}
-
-		// Note: we used to early return false here if nothing had changed, but we still need to send
-		// a bunch with the array key / base key, so that clients can look for implicit deletes.
-
-		// The array replication key may have changed while adding new elemnts (in the call to MarkItemDirty above)
-		NewState->ArrayReplicationKey = ArraySerializer.ArrayReplicationKey;
-
-		//----------------------
-		// Write it out.
-		//----------------------
-
-		int32 ArrayReplicationKey = ArraySerializer.ArrayReplicationKey;
-		Writer << ArrayReplicationKey;
-		Writer << BaseReplicationKey;
-
-		uint32 ElemNum = DeletedElements.Num();
-		Writer << ElemNum;
-
-		ElemNum = ChangedElements.Num();
-		Writer << ElemNum;
-
-		UE_LOG(LogNetFastTArray, Log, TEXT("   Writing Bunch. NumChange: %d. NumDel: %d [%d/%d]"), ChangedElements.Num(), DeletedElements.Num(), ArrayReplicationKey, BaseReplicationKey );
-
-		// Serialize deleted items, just by their ID
-		for (auto It = DeletedElements.CreateIterator(); It; ++It)
-		{
-			int32 ID = *It;
-			Writer << ID;
-			UE_LOG(LogNetFastTArray, Log, TEXT("   Deleted ElementID: %d"), ID);
-		}
-
-		// Serialized new elements with their payload
-		for (auto It = ChangedElements.CreateIterator(); It; ++It)
-		{
-			void* ThisElement = &Items[It->Idx];
-
-			// Dont pack this, want property to be byte aligned
-			uint32 ID = It->ID;
-			Writer << ID;
-
-			UE_LOG(LogNetFastTArray, Log, TEXT("   Changed ElementID: %d"), ID);
-
-			bool bHasUnmapped = false;
-			Parms.NetSerializeCB->NetSerializeStruct(InnerStruct, Writer, Parms.Map, ThisElement, bHasUnmapped);
-		}
-	}
-	else
-	{
-		//-----------------------------
-		// Loading
-		//-----------------------------	
-		check(Parms.Reader);
-		FBitReader& Reader = *Parms.Reader;
-
-		static const int32 MAX_NUM_CHANGED = 2048;
-		static const int32 MAX_NUM_DELETED = 2048;
-
-		//---------------
-		// Read header
-		//---------------
-		
-		int32 ArrayReplicationKey;
-		Reader << ArrayReplicationKey;
-		
-		int32 BaseReplicationKey;
-		Reader << BaseReplicationKey;
-
-		uint32 NumDeletes;
-		Reader << NumDeletes;
-
-		UE_LOG( LogNetFastTArray, Log, TEXT( "Received [%d/%d]." ), ArrayReplicationKey, BaseReplicationKey );
-	
-		if ( NumDeletes > MAX_NUM_DELETED )
-		{
-			UE_LOG( LogNetFastTArray, Warning, TEXT( "NumDeletes > MAX_NUM_DELETED: %d." ), NumDeletes );
-			Reader.SetError();
-			return false;;
-		}
-
-		uint32 NumChanged;
-		Reader << NumChanged;
-
-		if (NumChanged > MAX_NUM_CHANGED)
-		{
-			UE_LOG(LogNetFastTArray, Warning, TEXT("NumChanged > MAX_NUM_CHANGED: %d."), NumChanged);
-			Reader.SetError();
-			return false;;
-		}
-
-		UE_LOG( LogNetFastTArray, Log, TEXT( "Read NumChanged: %d NumDeletes: %d." ), NumChanged, NumDeletes );
-
-		TArray<int32, TInlineAllocator<8> > DeleteIndices;
-		TArray<int32, TInlineAllocator<8> > AddedIndices;
-		TArray<int32, TInlineAllocator<8> > ChangedIndices;
-
-		//---------------
-		// Read deleted elements
-		//---------------
-		if (NumDeletes > 0)
-		{
-			for (uint32 i = 0; i < NumDeletes; ++i)
-			{
-				int32 ElementID;
-				Reader << ElementID;
-
-				int32* ElementIndexPtr = ArraySerializer.ItemMap.Find(ElementID);
-				if (ElementIndexPtr)
-				{
-					int32 DeleteIndex = *ElementIndexPtr;
-					DeleteIndices.Add(DeleteIndex);
-					UE_LOG(LogNetFastTArray, Log, TEXT("   Adding ElementID: %d for deletion"), ElementID);
-				}
-				else
-				{
-					UE_LOG(LogNetFastTArray, Log, TEXT("   Couldn't find ElementID: %d for deletion!"), ElementID);
-				}
-			}
-		}
-
-		//---------------
-		// Read Changed/New elements
-		//---------------
-		for(uint32 i=0; i < NumChanged; ++i)
-		{
-			int32 ElementID;
-			Reader << ElementID;
-
-			int32* ElementIndexPtr = ArraySerializer.ItemMap.Find(ElementID);
-			int32 ElementIndex = 0;
-			Type* ThisElement = nullptr;
-		
-			if (!ElementIndexPtr)
-			{
-				UE_LOG(LogNetFastTArray, Log, TEXT("   New. ID: %d. New Element!"), ElementID);
-
-				ThisElement = new (Items)Type();
-				ThisElement->ReplicationID = ElementID;
-
-				ElementIndex = Items.Num()-1;
-				ArraySerializer.ItemMap.Add(ElementID, ElementIndex);
-				
-				AddedIndices.Add(ElementIndex);
-			}
-			else
-			{
-				UE_LOG(LogNetFastTArray, Log, TEXT("   Changed. ID: %d -> Idx: %d"), ElementID, *ElementIndexPtr);
-				ElementIndex = *ElementIndexPtr;
-				ThisElement = &Items[ElementIndex];
-				ChangedIndices.Add(ElementIndex);
-			}
-
-			// Update this element's most recent array replication key
-			ThisElement->MostRecentArrayReplicationKey = ArrayReplicationKey;
-
-			// Update this element's replication key so that a client can re-serialize the array for client replay recording
-			ThisElement->ReplicationKey++;
-
-			// Let package map know we want to track and know about any guids that are unmapped during the serialize call
-			Parms.Map->ResetTrackedGuids( true );
-
-			// Remember where we started reading from, so that if we have unmapped properties, we can re-deserialize from this data later
-			FBitReaderMark Mark( Reader );
-
-			bool bHasUnmapped = false;
-			Parms.NetSerializeCB->NetSerializeStruct( InnerStruct, Reader, Parms.Map, ThisElement, bHasUnmapped );
-
-			if ( !Reader.IsError() )
-			{
-				// Track unmapped guids
-				const TSet< FNetworkGUID >& TrackedUnmappedGuids = Parms.Map->GetTrackedUnmappedGuids();
-				const TSet< FNetworkGUID >& TrackedMappedDynamicGuids = Parms.Map->GetTrackedDynamicMappedGuids();
-
-				if ( TrackedUnmappedGuids.Num() || TrackedMappedDynamicGuids.Num() )
-				{	
-					FFastArraySerializerGuidReferences& GuidReferences = ArraySerializer.GuidReferencesMap.FindOrAdd( ElementID );
-
-					// If guid lists are different, make note of that, and copy respective list
-					if ( !NetworkGuidSetsAreSame( GuidReferences.UnmappedGUIDs, TrackedUnmappedGuids ) )
-					{
-						// Copy the unmapped guid list to this unmapped item
-						GuidReferences.UnmappedGUIDs = TrackedUnmappedGuids;
-						Parms.bGuidListsChanged = true;
-					}
-
-					if ( !NetworkGuidSetsAreSame( GuidReferences.MappedDynamicGUIDs, TrackedMappedDynamicGuids ) )
-					{ 
-						// Copy the mapped guid list
-						GuidReferences.MappedDynamicGUIDs = TrackedMappedDynamicGuids;
-						Parms.bGuidListsChanged = true;
-					}
-
-					GuidReferences.Buffer.Empty();
-
-					// Remember the number of bits in the buffer
-					GuidReferences.NumBufferBits = Reader.GetPosBits() - Mark.GetPos();
-					
-					// Copy the buffer itself
-					Mark.Copy( Reader, GuidReferences.Buffer );
-
-					// Hijack this property to communicate that we need to be tracked since we have some unmapped guids
-					if ( TrackedUnmappedGuids.Num() )
-					{
-						Parms.bOutHasMoreUnmapped = true;
-					}
-				}
-				else
-				{
-					// If we don't have any unmapped objects, make sure we're no longer tracking this item in the unmapped lists
-					ArraySerializer.GuidReferencesMap.Remove( ElementID );
-				}
-			}
-
-			// Stop tracking unmapped objects
-			Parms.Map->ResetTrackedGuids( false );
-
-			if ( Reader.IsError() )
-			{
-				UE_LOG( LogNetFastTArray, Warning, TEXT( "Parms.NetSerializeCB->NetSerializeStruct: Reader.IsError() == true" ) );
-				return false;
-			}
-		}
-
-		// ---------------------------------------------------------
-		// Look for implicit deletes that would happen due to Naks
-		// ---------------------------------------------------------
-		
-		for (int32 idx=0; idx < Items.Num(); ++idx)
-		{
-			Type& Item = Items[idx];
-			if (Item.MostRecentArrayReplicationKey < ArrayReplicationKey && Item.MostRecentArrayReplicationKey > BaseReplicationKey)
-			{
-				// Make sure this wasn't an explicit delete in this bunch (otherwise we end up deleting an extra element!)
-				if (DeleteIndices.Contains(idx) == false)
-				{
-					// This will happen in normal conditions in network replays.
-					UE_LOG( LogNetFastTArray, Log, TEXT( "Adding implicit delete for ElementID: %d. MostRecentArrayReplicationKey: %d. Current Payload: [%d/%d]"), Item.ReplicationID, Item.MostRecentArrayReplicationKey, ArrayReplicationKey, BaseReplicationKey);
-					
-					DeleteIndices.Add(idx);
-				}
-			}
-		}
-
-		// Increment keys so that a client can re-serialize the array if needed, such as for client replay recording.
-		// Must check the size of DeleteIndices instead of NumDeletes to handle implicit deletes.
-		if ( DeleteIndices.Num() > 0 || NumChanged > 0 )
-		{
-			ArraySerializer.IncrementArrayReplicationKey();
-		}
-
-		// ---------------------------------------------------------
-		// Invoke all callbacks: removed -> added -> changed
-		// ---------------------------------------------------------
-
-		int32 PreRemoveSize = Items.Num();
-		int32 FinalSize = PreRemoveSize - DeleteIndices.Num();
-		for (int32 idx : DeleteIndices)
-		{
-			if (Items.IsValidIndex(idx))
-			{
-				// Remove the deleted element's tracked GUID references
-				if (ArraySerializer.GuidReferencesMap.Remove(Items[idx].ReplicationID) > 0)
-				{
-					Parms.bGuidListsChanged = true;
-				}
-
-				// Call the delete callbacks now, actually remove them at the end
-				Items[idx].PreReplicatedRemove(ArraySerializer);
-			}
-		}
-		ArraySerializer.PreReplicatedRemove(DeleteIndices, FinalSize);
-
-		if (PreRemoveSize != Items.Num())
-		{
-			UE_LOG( LogNetFastTArray, Error, TEXT( "Item size changed after PreReplicatedRemove! PremoveSize: %d  Item.Num: %d"), PreRemoveSize, Items.Num() );
-		}
-
-		for (int32 idx : AddedIndices)
-		{
-			Items[idx].PostReplicatedAdd(ArraySerializer);
-		}
-		ArraySerializer.PostReplicatedAdd(AddedIndices, FinalSize);
-
-		for (int32 idx : ChangedIndices)
-		{
-			Items[idx].PostReplicatedChange(ArraySerializer);
-		}	
-		ArraySerializer.PostReplicatedChange(ChangedIndices, FinalSize);
-
-		if (PreRemoveSize != Items.Num())
-		{
-			UE_LOG( LogNetFastTArray, Error, TEXT( "Item size changed after PostReplicatedAdd/PostReplicatedChange! PremoveSize: %d  Item.Num: %d"), PreRemoveSize, Items.Num() );
-		}
-
-		if (DeleteIndices.Num() > 0)
-		{
-			DeleteIndices.Sort();
-			for (int32 i = DeleteIndices.Num() - 1; i >= 0; --i)
-			{
-				int32 DeleteIndex = DeleteIndices[i];
-				if (Items.IsValidIndex(DeleteIndex))
-				{
-					Items.RemoveAtSwap(DeleteIndex, 1, false);
-
-					UE_LOG(LogNetFastTArray, Log, TEXT("   Deleting: %d"), DeleteIndex);
-				}
-			}
-
-			// Clear the map now that the indices are all shifted around. This kind of sucks, we could use slightly better data structures here I think.
-			// This will force the ItemMap to be rebuilt for the current Items array
-			ArraySerializer.ItemMap.Empty();
-		}
-	}
-
-	return true;
-}
-
-
 
 
 /**
diff --git a/Engine/Source/Runtime/Engine/Private/Net/NetSerialization.cpp b/Engine/Source/Runtime/Engine/Private/Net/NetSerialization.cpp
new file mode 100644
--- /dev/null
+++ b/Engine/Source/Runtime/Engine/Private/Net/NetSerialization.cpp
@@ -0,0 +1,653 @@
+// Copyright 1998-2019 Epic Games, Inc. All Rights Reserved.
+
+#include "Engine/NetSerialization.h"
+
+/** The function that implements Fast TArray Replication  */
+bool FFastArraySerializer::FastArrayDeltaSerialize_Internal(FFastArrayDeltaSerializeAccessors& Accessors, FNetDeltaSerializeInfo& Parms, FFastArraySerializer& ArraySerializer, UScriptStruct* InnerStruct)
+{
+	SCOPE_CYCLE_COUNTER(STAT_NetSerializeFastArray);
+
+	UE_LOG(LogNetFastTArray, Log, TEXT("FastArrayDeltaSerialize for %s. %s. %s"), *InnerStruct->GetName(), *InnerStruct->GetOwnerStruct()->GetName(), Parms.Reader ? TEXT("Reading") : TEXT("Writing"));
+
+	if (Parms.bUpdateUnmappedObjects || Parms.Writer == NULL)
+	{
+		//---------------
+		// Build ItemMap if necessary. This maps ReplicationID to our local index into the Items array.
+		//---------------
+		if (ArraySerializer.ItemMap.Num() != Accessors.GetNumItems())
+		{
+			SCOPE_CYCLE_COUNTER(STAT_NetSerializeFastArray_BuildMap);
+			UE_LOG(LogNetFastTArray, Log, TEXT("FastArrayDeltaSerialize: Recreating Items map. Struct: %s, Items.Num: %d Map.Num: %d"), *InnerStruct->GetOwnerStruct()->GetName(), Accessors.GetNumItems(), ArraySerializer.ItemMap.Num());
+
+			ArraySerializer.ItemMap.Empty();
+			for (int32 i = 0; i < Accessors.GetNumItems(); ++i)
+			{
+				if (Accessors.GetItem(i).ReplicationID == INDEX_NONE)
+				{
+					if (Parms.Writer)
+					{
+						UE_LOG(LogNetFastTArray, Warning, TEXT("FastArrayDeltaSerialize: Item with uninitialized ReplicationID. Struct: %s, ItemIndex: %i"), *InnerStruct->GetOwnerStruct()->GetName(), i);
+					}
+					else
+					{
+						// This is benign for clients, they may add things to their local array without assigning a ReplicationID
+						continue;
+					}
+				}
+				ArraySerializer.ItemMap.Add(Accessors.GetItem(i).ReplicationID, i);
+			}
+		}
+	}
+
+	if (Parms.GatherGuidReferences)
+	{
+		// Loop over all tracked guids, and return what we have
+		for (const auto& GuidReferencesPair : ArraySerializer.GuidReferencesMap)
+		{
+			const FFastArraySerializerGuidReferences& GuidReferences = GuidReferencesPair.Value;
+
+			Parms.GatherGuidReferences->Append(GuidReferences.UnmappedGUIDs);
+			Parms.GatherGuidReferences->Append(GuidReferences.MappedDynamicGUIDs);
+
+			if (Parms.TrackedGuidMemoryBytes)
+			{
+				*Parms.TrackedGuidMemoryBytes += GuidReferences.Buffer.Num();
+			}
+		}
+
+		return true;
+	}
+
+	if (Parms.MoveGuidToUnmapped)
+	{
+		bool bFound = false;
+
+		const FNetworkGUID GUID = *Parms.MoveGuidToUnmapped;
+
+		// Try to find the guid in the list, and make sure it's on the unmapped lists now
+		for (auto& GuidReferencesPair : ArraySerializer.GuidReferencesMap)
+		{
+			FFastArraySerializerGuidReferences& GuidReferences = GuidReferencesPair.Value;
+
+			if (GuidReferences.MappedDynamicGUIDs.Contains(GUID))
+			{
+				GuidReferences.MappedDynamicGUIDs.Remove(GUID);
+				GuidReferences.UnmappedGUIDs.Add(GUID);
+				bFound = true;
+			}
+		}
+
+		return bFound;
+	}
+
+	if (Parms.bUpdateUnmappedObjects)
+	{
+		// Loop over each item that has unmapped objects
+		for (auto It = ArraySerializer.GuidReferencesMap.CreateIterator(); It; ++It)
+		{
+			// Get the element id
+			const int32 ElementID = It.Key();
+
+			// Get a reference to the unmapped item itself
+			FFastArraySerializerGuidReferences& GuidReferences = It.Value();
+
+			if ((GuidReferences.UnmappedGUIDs.Num() == 0 && GuidReferences.MappedDynamicGUIDs.Num() == 0) || ArraySerializer.ItemMap.Find(ElementID) == NULL)
+			{
+				// If for some reason the item is gone (or all guids were removed), we don't need to track guids for this item anymore
+				It.RemoveCurrent();
+				continue;		// We're done with this unmapped item
+			}
+
+			// Loop over all the guids, and check to see if any of them are loaded yet
+			bool bMappedSomeGUIDs = false;
+
+			for (auto UnmappedIt = GuidReferences.UnmappedGUIDs.CreateIterator(); UnmappedIt; ++UnmappedIt)
+			{
+				const FNetworkGUID& GUID = *UnmappedIt;
+
+				if (Parms.Map->IsGUIDBroken(GUID, false))
+				{
+					// Stop trying to load broken guids
+					UE_LOG(LogNetFastTArray, Warning, TEXT("FastArrayDeltaSerialize: Broken GUID. NetGuid: %s"), *GUID.ToString());
+					UnmappedIt.RemoveCurrent();
+					continue;
+				}
+
+				UObject* Object = Parms.Map->GetObjectFromNetGUID(GUID, false);
+
+				if (Object != NULL)
+				{
+					// This guid loaded!
+					if (GUID.IsDynamic())
+					{
+						GuidReferences.MappedDynamicGUIDs.Add(GUID);		// Move back to mapped list
+					}
+					UnmappedIt.RemoveCurrent();
+					bMappedSomeGUIDs = true;
+				}
+			}
+
+			// Check to see if we loaded any guids. If we did, we can serialize the element again which will load it this time
+			if (bMappedSomeGUIDs)
+			{
+				Parms.bOutSomeObjectsWereMapped = true;
+
+				if (!Parms.bCalledPreNetReceive)
+				{
+					// Call PreNetReceive if we are going to change a value (some game code will need to think this is an actual replicated value)
+					Parms.Object->PreNetReceive();
+					Parms.bCalledPreNetReceive = true;
+				}
+
+				int32 ThisElementIdx = ArraySerializer.ItemMap.FindChecked(ElementID);
+				FFastArraySerializerItem* ThisElement = &Accessors.GetItem(ThisElementIdx);
+
+				// Initialize the reader with the stored buffer that we need to read from
+				FNetBitReader Reader(Parms.Map, GuidReferences.Buffer.GetData(), GuidReferences.NumBufferBits);
+
+				// Read the property (which should serialize any newly mapped objects as well)
+				bool bHasUnmapped = false;
+				Parms.NetSerializeCB->NetSerializeStruct(InnerStruct, Reader, Parms.Map, ThisElement, bHasUnmapped);
+
+				// Let the element know it changed
+				Accessors.PostReplicatedChange(ThisElementIdx);
+			}
+
+			// If we have no more guids, we can remove this item for good
+			if (GuidReferences.UnmappedGUIDs.Num() == 0 && GuidReferences.MappedDynamicGUIDs.Num() == 0)
+			{
+				It.RemoveCurrent();
+			}
+		}
+
+		// If we still have unmapped items, then communicate this to the outside
+		if (ArraySerializer.GuidReferencesMap.Num() > 0)
+		{
+			Parms.bOutHasMoreUnmapped = true;
+		}
+
+		return true;
+	}
+
+	if (Parms.Writer)
+	{
+		//-----------------------------
+		// Saving
+		//-----------------------------	
+		check(Parms.Struct);
+		FBitWriter& Writer = *Parms.Writer;
+
+		// Get the old map if its there
+		TMap<int32, int32> * OldMap = Parms.OldState ? &((FNetFastTArrayBaseState*)Parms.OldState)->IDToCLMap : NULL;
+		int32 BaseReplicationKey = Parms.OldState ? ((FNetFastTArrayBaseState*)Parms.OldState)->ArrayReplicationKey : -1;
+
+		// See if the array changed at all. If the ArrayReplicationKey matches we can skip checking individual items
+		if (Parms.OldState && (ArraySerializer.ArrayReplicationKey == BaseReplicationKey))
+		{
+			// Double check the old map is valid and that we will consider writing the same number of elements that are in the old map.
+			if (ensureMsgf(OldMap, TEXT("Invalid OldMap")))
+			{
+				// If the keys didn't change, only update the item count caches if necessary.
+				if (ArraySerializer.CachedNumItems == INDEX_NONE ||
+					ArraySerializer.CachedNumItems != Accessors.GetNumItems() ||
+					ArraySerializer.CachedNumItemsToConsiderForWriting == INDEX_NONE)
+				{
+					ArraySerializer.CachedNumItems = Accessors.GetNumItems();
+					ArraySerializer.CachedNumItemsToConsiderForWriting = 0;
+
+					// Count the number of items in the current array that may be written. On clients, items that were predicted will be skipped.
+					for (int32 ItemIdx = 0; ItemIdx < Accessors.GetNumItems(); ++ItemIdx)
+					{
+						if (Accessors.ShouldWriteFastArrayItem(ItemIdx, Parms.bIsWritingOnClient))
+						{
+							ArraySerializer.CachedNumItemsToConsiderForWriting++;
+						}
+					}
+				}
+
+				if (UNLIKELY(OldMap->Num() != ArraySerializer.CachedNumItemsToConsiderForWriting))
+				{
+					UE_LOG(LogNetFastTArray, Warning, TEXT("OldMap size (%d) does not match item count (%d)"), OldMap->Num(), ArraySerializer.CachedNumItemsToConsiderForWriting);
+				}
+			}
+
+			if (Parms.OldState)
+			{
+				// Nothing changed and we had a valid old state, so just use/share the existing state. No need to create a new one.
+				*Parms.NewState = Parms.OldState->AsShared();
+			}
+			else
+			{
+				// Nothing changed but we don't have an existing state of our own yet so we need to make one here.
+				FNetFastTArrayBaseState * NewState = new FNetFastTArrayBaseState();
+				*Parms.NewState = MakeShareable(NewState);
+				NewState->ArrayReplicationKey = ArraySerializer.ArrayReplicationKey;
+			}
+			return false;
+		}
+
+
+		// Create a new map from the current state of the array		
+		FNetFastTArrayBaseState * NewState = new FNetFastTArrayBaseState();
+
+		check(Parms.NewState);
+		*Parms.NewState = MakeShareable(NewState);
+		TMap<int32, int32> & NewMap = NewState->IDToCLMap;
+		NewState->ArrayReplicationKey = ArraySerializer.ArrayReplicationKey;
+
+
+		int32 NumConsideredItems = 0;
+		for (int32 ItemIdx = 0; ItemIdx < Accessors.GetNumItems(); ++ItemIdx)
+		{
+			if (Accessors.ShouldWriteFastArrayItem(ItemIdx, Parms.bIsWritingOnClient))
+			{
+				NumConsideredItems++;
+			}
+		}
+
+		TArray<FFastArraySerializer_FastArrayDeltaSerialize_FIdxIDPair, TInlineAllocator<8> >	ChangedElements;
+		TArray<int32, TInlineAllocator<8> >		DeletedElements;
+
+		int32 DeleteCount = (OldMap ? OldMap->Num() : 0) - NumConsideredItems; // Note: this is incremented when we add new items below.
+		UE_LOG(LogNetFastTArray, Log, TEXT("NetSerializeItemDeltaFast: %s. DeleteCount: %d"), *Parms.DebugName, DeleteCount);
+
+		// Log out entire state of current/base state
+		if (UE_LOG_ACTIVE(LogNetFastTArray, Log))
+		{
+			FString CurrentState = FString::Printf(TEXT("Current: %d "), ArraySerializer.ArrayReplicationKey);
+			for (int32 i = 0; i < Accessors.GetNumItems(); ++i)
+			{
+				CurrentState += FString::Printf(TEXT("[%d/%d], "), Accessors.GetItem(i).ReplicationID, Accessors.GetItem(i).ReplicationKey);
+			}
+			UE_LOG(LogNetFastTArray, Log, TEXT("%s"), *CurrentState);
+
+
+			FString ClientStateStr = FString::Printf(TEXT("Client: %d "), Parms.OldState ? ((FNetFastTArrayBaseState*)Parms.OldState)->ArrayReplicationKey : 0);
+			if (OldMap)
+			{
+				for (auto It = OldMap->CreateIterator(); It; ++It)
+				{
+					ClientStateStr += FString::Printf(TEXT("[%d/%d], "), It.Key(), It.Value());
+				}
+			}
+			UE_LOG(LogNetFastTArray, Log, TEXT("%s"), *ClientStateStr);
+		}
+
+
+		//--------------------------------------------
+		// Find out what is new or what has changed
+		//--------------------------------------------
+		for (int32 i = 0; i < Accessors.GetNumItems(); ++i)
+		{
+			UE_LOG(LogNetFastTArray, Log, TEXT("    Array[%d] - ID %d. CL %d."), i, Accessors.GetItem(i).ReplicationID, Accessors.GetItem(i).ReplicationKey);
+			if (!Accessors.ShouldWriteFastArrayItem(i, Parms.bIsWritingOnClient))
+			{
+				// On clients, this will skip items that were added predictively.
+				continue;
+			}
+			if (Accessors.GetItem(i).ReplicationID == INDEX_NONE)
+			{
+				ArraySerializer.MarkItemDirty(Accessors.GetItem(i));
+			}
+			NewMap.Add(Accessors.GetItem(i).ReplicationID, Accessors.GetItem(i).ReplicationKey);
+
+			int32* OldValuePtr = OldMap ? OldMap->Find(Accessors.GetItem(i).ReplicationID) : NULL;
+			if (OldValuePtr)
+			{
+				if (*OldValuePtr == Accessors.GetItem(i).ReplicationKey)
+				{
+					UE_LOG(LogNetFastTArray, Log, TEXT("       Stayed The Same - Skipping"));
+
+					// Stayed the same, it might have moved but we dont care
+					continue;
+				}
+				else
+				{
+					UE_LOG(LogNetFastTArray, Log, TEXT("       Changed! Was: %d. Element ID: %d. %s"), *OldValuePtr, Accessors.GetItem(i).ReplicationID, *Accessors.GetItem(i).GetDebugString());
+
+					// Changed
+					ChangedElements.Add(FFastArraySerializer_FastArrayDeltaSerialize_FIdxIDPair(i, Accessors.GetItem(i).ReplicationID));
+				}
+			}
+			else
+			{
+				UE_LOG(LogNetFastTArray, Log, TEXT("       New! Element ID: %d. %s"), Accessors.GetItem(i).ReplicationID, *Accessors.GetItem(i).GetDebugString());
+
+				// The item really should have a valid ReplicationID but in the case of loading from a save game,
+				// items may not have been marked dirty individually. Its ok to just assign them one here.
+				// New
+				ChangedElements.Add(FFastArraySerializer_FastArrayDeltaSerialize_FIdxIDPair(i, Accessors.GetItem(i).ReplicationID));
+				DeleteCount++; // We added something new, so our initial DeleteCount value must be incremented.
+			}
+		}
+
+		// Find out what was deleted
+		if (DeleteCount > 0 && OldMap)
+		{
+			for (auto It = OldMap->CreateIterator(); It; ++It)
+			{
+				if (!NewMap.Contains(It.Key()))
+				{
+					UE_LOG(LogNetFastTArray, Log, TEXT("   Deleting ID: %d"), It.Key());
+
+					DeletedElements.Add(It.Key());
+					if (--DeleteCount <= 0)
+						break;
+				}
+			}
+		}
+
+		// Note: we used to early return false here if nothing had changed, but we still need to send
+		// a bunch with the array key / base key, so that clients can look for implicit deletes.
+
+		// The array replication key may have changed while adding new elemnts (in the call to MarkItemDirty above)
+		NewState->ArrayReplicationKey = ArraySerializer.ArrayReplicationKey;
+
+		//----------------------
+		// Write it out.
+		//----------------------
+
+		int32 ArrayReplicationKey = ArraySerializer.ArrayReplicationKey;
+		Writer << ArrayReplicationKey;
+		Writer << BaseReplicationKey;
+
+		uint32 ElemNum = DeletedElements.Num();
+		Writer << ElemNum;
+
+		ElemNum = ChangedElements.Num();
+		Writer << ElemNum;
+
+		UE_LOG(LogNetFastTArray, Log, TEXT("   Writing Bunch. NumChange: %d. NumDel: %d [%d/%d]"), ChangedElements.Num(), DeletedElements.Num(), ArrayReplicationKey, BaseReplicationKey);
+
+		// Serialize deleted items, just by their ID
+		for (auto It = DeletedElements.CreateIterator(); It; ++It)
+		{
+			int32 ID = *It;
+			Writer << ID;
+			UE_LOG(LogNetFastTArray, Log, TEXT("   Deleted ElementID: %d"), ID);
+		}
+
+		// Serialized new elements with their payload
+		for (auto It = ChangedElements.CreateIterator(); It; ++It)
+		{
+			void* ThisElement = &Accessors.GetItem(It->Idx);
+
+			// Dont pack this, want property to be byte aligned
+			uint32 ID = It->ID;
+			Writer << ID;
+
+			UE_LOG(LogNetFastTArray, Log, TEXT("   Changed ElementID: %d"), ID);
+
+			bool bHasUnmapped = false;
+			Parms.NetSerializeCB->NetSerializeStruct(InnerStruct, Writer, Parms.Map, ThisElement, bHasUnmapped);
+		}
+	}
+	else
+	{
+		//-----------------------------
+		// Loading
+		//-----------------------------	
+		check(Parms.Reader);
+		FBitReader& Reader = *Parms.Reader;
+
+		static const int32 MAX_NUM_CHANGED = 2048;
+		static const int32 MAX_NUM_DELETED = 2048;
+
+		//---------------
+		// Read header
+		//---------------
+
+		int32 ArrayReplicationKey;
+		Reader << ArrayReplicationKey;
+
+		int32 BaseReplicationKey;
+		Reader << BaseReplicationKey;
+
+		uint32 NumDeletes;
+		Reader << NumDeletes;
+
+		UE_LOG(LogNetFastTArray, Log, TEXT("Received [%d/%d]."), ArrayReplicationKey, BaseReplicationKey);
+
+		if (NumDeletes > MAX_NUM_DELETED)
+		{
+			UE_LOG(LogNetFastTArray, Warning, TEXT("NumDeletes > MAX_NUM_DELETED: %d."), NumDeletes);
+			Reader.SetError();
+			return false;;
+		}
+
+		uint32 NumChanged;
+		Reader << NumChanged;
+
+		if (NumChanged > MAX_NUM_CHANGED)
+		{
+			UE_LOG(LogNetFastTArray, Warning, TEXT("NumChanged > MAX_NUM_CHANGED: %d."), NumChanged);
+			Reader.SetError();
+			return false;;
+		}
+
+		UE_LOG(LogNetFastTArray, Log, TEXT("Read NumChanged: %d NumDeletes: %d."), NumChanged, NumDeletes);
+
+		TArray<int32, TInlineAllocator<8> > DeleteIndices;
+		TArray<int32, TInlineAllocator<8> > AddedIndices;
+		TArray<int32, TInlineAllocator<8> > ChangedIndices;
+
+		//---------------
+		// Read deleted elements
+		//---------------
+		if (NumDeletes > 0)
+		{
+			for (uint32 i = 0; i < NumDeletes; ++i)
+			{
+				int32 ElementID;
+				Reader << ElementID;
+
+				int32* ElementIndexPtr = ArraySerializer.ItemMap.Find(ElementID);
+				if (ElementIndexPtr)
+				{
+					int32 DeleteIndex = *ElementIndexPtr;
+					DeleteIndices.Add(DeleteIndex);
+					UE_LOG(LogNetFastTArray, Log, TEXT("   Adding ElementID: %d for deletion"), ElementID);
+				}
+				else
+				{
+					UE_LOG(LogNetFastTArray, Log, TEXT("   Couldn't find ElementID: %d for deletion!"), ElementID);
+				}
+			}
+		}
+
+		//---------------
+		// Read Changed/New elements
+		//---------------
+		for (uint32 i = 0; i < NumChanged; ++i)
+		{
+			int32 ElementID;
+			Reader << ElementID;
+
+			int32* ElementIndexPtr = ArraySerializer.ItemMap.Find(ElementID);
+			int32 ElementIndex = 0;
+			FFastArraySerializerItem* ThisElement = nullptr;
+
+			if (!ElementIndexPtr)
+			{
+				UE_LOG(LogNetFastTArray, Log, TEXT("   New. ID: %d. New Element!"), ElementID);
+
+				ThisElement = &Accessors.AddItem();
+				ThisElement->ReplicationID = ElementID;
+
+				ElementIndex = Accessors.GetNumItems() - 1;
+				ArraySerializer.ItemMap.Add(ElementID, ElementIndex);
+
+				AddedIndices.Add(ElementIndex);
+			}
+			else
+			{
+				UE_LOG(LogNetFastTArray, Log, TEXT("   Changed. ID: %d -> Idx: %d"), ElementID, *ElementIndexPtr);
+				ElementIndex = *ElementIndexPtr;
+				ThisElement = &Accessors.GetItem(ElementIndex);
+				ChangedIndices.Add(ElementIndex);
+			}
+
+			// Update this element's most recent array replication key
+			ThisElement->MostRecentArrayReplicationKey = ArrayReplicationKey;
+
+			// Update this element's replication key so that a client can re-serialize the array for client replay recording
+			ThisElement->ReplicationKey++;
+
+			// Let package map know we want to track and know about any guids that are unmapped during the serialize call
+			Parms.Map->ResetTrackedGuids(true);
+
+			// Remember where we started reading from, so that if we have unmapped properties, we can re-deserialize from this data later
+			FBitReaderMark Mark(Reader);
+
+			bool bHasUnmapped = false;
+			Parms.NetSerializeCB->NetSerializeStruct(InnerStruct, Reader, Parms.Map, ThisElement, bHasUnmapped);
+
+			if (!Reader.IsError())
+			{
+				// Track unmapped guids
+				const TSet< FNetworkGUID >& TrackedUnmappedGuids = Parms.Map->GetTrackedUnmappedGuids();
+				const TSet< FNetworkGUID >& TrackedMappedDynamicGuids = Parms.Map->GetTrackedDynamicMappedGuids();
+
+				if (TrackedUnmappedGuids.Num() || TrackedMappedDynamicGuids.Num())
+				{
+					FFastArraySerializerGuidReferences& GuidReferences = ArraySerializer.GuidReferencesMap.FindOrAdd(ElementID);
+
+					// If guid lists are different, make note of that, and copy respective list
+					if (!NetworkGuidSetsAreSame(GuidReferences.UnmappedGUIDs, TrackedUnmappedGuids))
+					{
+						// Copy the unmapped guid list to this unmapped item
+						GuidReferences.UnmappedGUIDs = TrackedUnmappedGuids;
+						Parms.bGuidListsChanged = true;
+					}
+
+					if (!NetworkGuidSetsAreSame(GuidReferences.MappedDynamicGUIDs, TrackedMappedDynamicGuids))
+					{
+						// Copy the mapped guid list
+						GuidReferences.MappedDynamicGUIDs = TrackedMappedDynamicGuids;
+						Parms.bGuidListsChanged = true;
+					}
+
+					GuidReferences.Buffer.Empty();
+
+					// Remember the number of bits in the buffer
+					GuidReferences.NumBufferBits = Reader.GetPosBits() - Mark.GetPos();
+
+					// Copy the buffer itself
+					Mark.Copy(Reader, GuidReferences.Buffer);
+
+					// Hijack this property to communicate that we need to be tracked since we have some unmapped guids
+					if (TrackedUnmappedGuids.Num())
+					{
+						Parms.bOutHasMoreUnmapped = true;
+					}
+				}
+				else
+				{
+					// If we don't have any unmapped objects, make sure we're no longer tracking this item in the unmapped lists
+					ArraySerializer.GuidReferencesMap.Remove(ElementID);
+				}
+			}
+
+			// Stop tracking unmapped objects
+			Parms.Map->ResetTrackedGuids(false);
+
+			if (Reader.IsError())
+			{
+				UE_LOG(LogNetFastTArray, Warning, TEXT("Parms.NetSerializeCB->NetSerializeStruct: Reader.IsError() == true"));
+				return false;
+			}
+		}
+
+		// ---------------------------------------------------------
+		// Look for implicit deletes that would happen due to Naks
+		// ---------------------------------------------------------
+
+		for (int32 idx = 0; idx < Accessors.GetNumItems(); ++idx)
+		{
+			FFastArraySerializerItem& Item = Accessors.GetItem(idx);
+			if (Item.MostRecentArrayReplicationKey < ArrayReplicationKey && Item.MostRecentArrayReplicationKey > BaseReplicationKey)
+			{
+				// Make sure this wasn't an explicit delete in this bunch (otherwise we end up deleting an extra element!)
+				if (DeleteIndices.Contains(idx) == false)
+				{
+					// This will happen in normal conditions in network replays.
+					UE_LOG(LogNetFastTArray, Log, TEXT("Adding implicit delete for ElementID: %d. MostRecentArrayReplicationKey: %d. Current Payload: [%d/%d]"), Item.ReplicationID, Item.MostRecentArrayReplicationKey, ArrayReplicationKey, BaseReplicationKey);
+
+					DeleteIndices.Add(idx);
+				}
+			}
+		}
+
+		// Increment keys so that a client can re-serialize the array if needed, such as for client replay recording.
+		// Must check the size of DeleteIndices instead of NumDeletes to handle implicit deletes.
+		if (DeleteIndices.Num() > 0 || NumChanged > 0)
+		{
+			ArraySerializer.IncrementArrayReplicationKey();
+		}
+
+		// ---------------------------------------------------------
+		// Invoke all callbacks: removed -> added -> changed
+		// ---------------------------------------------------------
+
+		int32 PreRemoveSize = Accessors.GetNumItems();
+		int32 FinalSize = PreRemoveSize - DeleteIndices.Num();
+		for (int32 idx : DeleteIndices)
+		{
+			if (idx >= 0 && idx < Accessors.GetNumItems())
+			{
+				// Remove the deleted element's tracked GUID references
+				if (ArraySerializer.GuidReferencesMap.Remove(Accessors.GetItem(idx).ReplicationID) > 0)
+				{
+					Parms.bGuidListsChanged = true;
+				}
+
+				// Call the delete callbacks now, actually remove them at the end
+				Accessors.GetItem(idx).PreReplicatedRemove(ArraySerializer);
+			}
+		}
+		ArraySerializer.PreReplicatedRemove(DeleteIndices, FinalSize);
+
+		if (PreRemoveSize != Accessors.GetNumItems())
+		{
+			UE_LOG(LogNetFastTArray, Error, TEXT("Item size changed after PreReplicatedRemove! PremoveSize: %d  Item.Num: %d"), PreRemoveSize, Accessors.GetNumItems());
+		}
+
+		for (int32 idx : AddedIndices)
+		{
+			Accessors.PostReplicatedAdd(idx);
+		}
+		ArraySerializer.PostReplicatedAdd(AddedIndices, FinalSize);
+
+		for (int32 idx : ChangedIndices)
+		{
+			Accessors.PostReplicatedChange(idx);
+		}
+		ArraySerializer.PostReplicatedChange(ChangedIndices, FinalSize);
+
+		if (PreRemoveSize != Accessors.GetNumItems())
+		{
+			UE_LOG(LogNetFastTArray, Error, TEXT("Item size changed after PostReplicatedAdd/PostReplicatedChange! PremoveSize: %d  Item.Num: %d"), PreRemoveSize, Accessors.GetNumItems());
+		}
+
+		if (DeleteIndices.Num() > 0)
+		{
+			DeleteIndices.Sort();
+			for (int32 i = DeleteIndices.Num() - 1; i >= 0; --i)
+			{
+				int32 DeleteIndex = DeleteIndices[i];
+				if (DeleteIndex >= 0 && DeleteIndex < Accessors.GetNumItems())
+				{
+					Accessors.RemoveItem(DeleteIndex);
+
+					UE_LOG(LogNetFastTArray, Log, TEXT("   Deleting: %d"), DeleteIndex);
+				}
+			}
+
+			// Clear the map now that the indices are all shifted around. This kind of sucks, we could use slightly better data structures here I think.
+			// This will force the ItemMap to be rebuilt for the current Items array
+			ArraySerializer.ItemMap.Empty();
+		}
+	}
+
+	return true;
+}
\ No newline at end of file
