Optimized RenderAssetStreamingMipCalcTask - Task is now ~2.3 times faster
    - Enabled by default - CVar to enable / disable it is CVarStreamingMipCalculationEnablePerLevelList
    - Added to FStreamRenderAsset structure bit array that indexes levels that reference the asset
    - When iterating through render asset and computing max visible mip/lod, algo iterates through associated levels of the render assets instead of all loaded levels

diff --git a/Engine/Source/Runtime/Engine/Private/Streaming/AsyncTextureStreaming.cpp b/Engine/Source/Runtime/Engine/Private/Streaming/AsyncTextureStreaming.cpp
--- a/Engine/Source/Runtime/Engine/Private/Streaming/AsyncTextureStreaming.cpp
+++ b/Engine/Source/Runtime/Engine/Private/Streaming/AsyncTextureStreaming.cpp
@@ -8,11 +8,12 @@ AsyncTextureStreaming.cpp: Definitions of classes used for texture streaming asy
 #include "Misc/App.h"
 #include "Streaming/StreamingManagerTexture.h"
 #include "Engine/World.h"
+#include "ProfilingDebugging/CsvProfiler.h"
 
 void FAsyncRenderAssetStreamingData::Init(
 	TArray<FStreamingViewInfo> InViewInfos,
 	float InLastUpdateTime,
-	TIndirectArray<FLevelRenderAssetManager>& LevelStaticInstanceManagers,
+	TArray<FLevelRenderAssetManager*>& LevelStaticInstanceManagers,
 	FDynamicRenderAssetInstanceManager& DynamicComponentManager)
 {
 	ViewInfos = InViewInfos;
@@ -23,12 +24,26 @@ void FAsyncRenderAssetStreamingData::Init(
 	StaticInstancesViews.Reset();
 	StaticInstancesViewIndices.Reset();
 	CulledStaticInstancesViewIndices.Reset();
+	StaticInstancesViewLevelIndices.Reset();
 
-	for (FLevelRenderAssetManager& LevelManager : LevelStaticInstanceManagers)
+	for (int32 LevelIndex = 0; LevelIndex<LevelStaticInstanceManagers.Num(); ++LevelIndex)
 	{
+		if (LevelStaticInstanceManagers[LevelIndex] == nullptr)
+		{
+			StaticInstancesViewLevelIndices.Add(INDEX_NONE);
+			continue;
+		}
+
+		const FLevelRenderAssetManager& LevelManager = *LevelStaticInstanceManagers[LevelIndex];
+
 		if (LevelManager.IsInitialized() && LevelManager.GetLevel()->bIsVisible && LevelManager.HasRenderAssetReferences())
 		{
-			StaticInstancesViews.Push(LevelManager.GetAsyncView());
+			StaticInstancesViewLevelIndices.Add(StaticInstancesViews.Num());
+			StaticInstancesViews.Add(LevelStaticInstanceManagers[LevelIndex]->GetAsyncView());
+		}
+		else
+		{
+			StaticInstancesViewLevelIndices.Add(INDEX_NONE);
 		}
 	}
 }
@@ -132,20 +147,75 @@ void FAsyncRenderAssetStreamingData::UpdatePerfectWantedMips_Async(FStreamingRen
 		const typename FStreamingRenderAsset::EAssetType AssetType = StreamingRenderAsset.RenderAssetType;
 		DynamicInstancesView.GetRenderAssetScreenSize(AssetType, RenderAsset, MaxSize, MaxSize_VisibleOnly, MaxNumForcedLODs, bOutputToLog ? TEXT("Dynamic") : nullptr);
 
-		for (int32 StaticViewIndex : StaticInstancesViewIndices)
+		bool bCulled = false;
+		if (Settings.bMipCalculationEnablePerLevelList)
 		{
-			const FRenderAssetInstanceAsyncView& StaticInstancesView = StaticInstancesViews[StaticViewIndex];
+			int32 LevelToIterateCount = FMath::Min(StreamingRenderAsset.LevelIndexUsage.Num(), StaticInstancesViewLevelIndices.Num());
 			
-			// No need to iterate more if texture is already at maximum resolution.
-			if ((MaxNumForcedLODs >= StreamingRenderAsset.MaxAllowedMips
-				|| MaxSize_VisibleOnly >= MAX_TEXTURE_SIZE
-				|| (MaxSize_VisibleOnly > StaticInstancesView.GetMaxLevelRenderAssetScreenSize() && Settings.MinLevelRenderAssetScreenSize > 0))
-				&& !bOutputToLog)
+			for (TBitArray<>::FIterator LevelUsageIterator(StreamingRenderAsset.LevelIndexUsage); LevelUsageIterator.GetIndex() < LevelToIterateCount; ++LevelUsageIterator)
 			{
-				break;
+				if (!LevelUsageIterator.GetValue())
+				{
+					continue;
+				}
+
+				int32 ViewIndex = StaticInstancesViewLevelIndices[LevelUsageIterator.GetIndex()];
+
+				if (ViewIndex==INDEX_NONE)
+				{
+					continue;
+				}
+
+				const FRenderAssetInstanceAsyncView& StaticInstancesView = StaticInstancesViews[ViewIndex];
+
+				if (!StaticInstancesView.HasRenderAssetReferences(RenderAsset))
+				{
+					// Level entry has been replaced by another level
+					// Remove its reference.
+					StreamingRenderAsset.LevelIndexUsage[LevelUsageIterator.GetIndex()] = false;
+					continue;
+				}
+
+				if (StaticInstancesView.GetMaxLevelRenderAssetScreenSize() < Settings.MinLevelRenderAssetScreenSize)
+				{
+					bCulled = true;
+					continue;
+				}
+
+				// No need to iterate more if render asset is already at maximum resolution.
+				if (MaxSize_VisibleOnly >= MAX_TEXTURE_SIZE || MaxNumForcedLODs >= StreamingRenderAsset.MaxAllowedMips)
+				{
+					break;
+				}
+
+				float TmpMaxSize = MaxSize;
+				float TmpMaxVisibleOnly = MaxSize_VisibleOnly;
+				int32 TmpMaxNumForcedLODs = MaxNumForcedLODs;
+				
+				StaticInstancesView.GetRenderAssetScreenSize(AssetType, RenderAsset, TmpMaxSize, TmpMaxVisibleOnly, TmpMaxNumForcedLODs, bOutputToLog ? TEXT("Static") : nullptr);
+
+				MaxSize = FMath::Max(TmpMaxSize, MaxSize);
+				MaxSize_VisibleOnly = FMath::Max(TmpMaxVisibleOnly, MaxSize_VisibleOnly);
+				MaxNumForcedLODs = FMath::Max(TmpMaxNumForcedLODs, MaxNumForcedLODs);
 			}
+		}
+		else
+		{
+			for (int32 StaticViewIndex : StaticInstancesViewIndices)
+			{
+				const FRenderAssetInstanceAsyncView& StaticInstancesView = StaticInstancesViews[StaticViewIndex];
 
-			StaticInstancesView.GetRenderAssetScreenSize(AssetType, RenderAsset, MaxSize, MaxSize_VisibleOnly, MaxNumForcedLODs, bOutputToLog ? TEXT("Static") : nullptr);
+				// No need to iterate more if texture is already at maximum resolution.
+				if ((MaxNumForcedLODs >= StreamingRenderAsset.MaxAllowedMips
+					|| MaxSize_VisibleOnly >= MAX_TEXTURE_SIZE
+					|| (MaxSize_VisibleOnly > StaticInstancesView.GetMaxLevelRenderAssetScreenSize() && Settings.MinLevelRenderAssetScreenSize > 0))
+					&& !bOutputToLog)
+				{
+					break;
+				}
+
+				StaticInstancesView.GetRenderAssetScreenSize(AssetType, RenderAsset, MaxSize, MaxSize_VisibleOnly, MaxNumForcedLODs, bOutputToLog ? TEXT("Static") : nullptr);
+			}
 		}
 
 		// Don't apply to FLT_MAX since it is used as forced streaming. BoostFactor as only meaning for texture/mesh instances since the other heuristics are based on max resolution.
@@ -168,14 +238,21 @@ void FAsyncRenderAssetStreamingData::UpdatePerfectWantedMips_Async(FStreamingRen
 		StreamingRenderAsset.bUseUnkownRefHeuristic = MaxSize == 0 && MaxSize_VisibleOnly == 0 && !MaxNumForcedLODs && StreamingRenderAsset.LastRenderTime < TimeSinceRemoved - 5.f;
 		if (StreamingRenderAsset.bUseUnkownRefHeuristic)
 		{
-			// Check that it's not simply culled
-			for (int32 StaticViewIndex : CulledStaticInstancesViewIndices)
+			if (Settings.bMipCalculationEnablePerLevelList)
 			{
-				const FRenderAssetInstanceAsyncView& StaticInstancesView = StaticInstancesViews[StaticViewIndex];
-				if (StaticInstancesView.HasRenderAssetReferences(RenderAsset))
+				StreamingRenderAsset.bUseUnkownRefHeuristic = !bCulled;
+			}
+			else
+			{
+				// Check that it's not simply culled
+				for (int32 StaticViewIndex : CulledStaticInstancesViewIndices)
 				{
-					StreamingRenderAsset.bUseUnkownRefHeuristic = false;
-					break;
+					const FRenderAssetInstanceAsyncView& StaticInstancesView = StaticInstancesViews[StaticViewIndex];
+					if (StaticInstancesView.HasRenderAssetReferences(RenderAsset))
+					{
+						StreamingRenderAsset.bUseUnkownRefHeuristic = false;
+						break;
+					}
 				}
 			}
 
@@ -636,6 +713,7 @@ void FRenderAssetStreamingMipCalcTask::DoWork()
 {
 	SCOPED_NAMED_EVENT(FRenderAssetStreamingMipCalcTask_DoWork, FColor::Turquoise);
 	DECLARE_SCOPE_CYCLE_COUNTER(TEXT("FRenderAssetStreamingMipCalcTask::DoWork"), STAT_FRenderAssetStreamingMipCalcTask_DoWork, STATGROUP_StreamingDetails);
+	CSV_SCOPED_TIMING_STAT_GLOBAL(FRenderAssetStreamingMipCalcTask);
 
 	// While the async task is runnning, the StreamingRenderAssets are guarantied not to be reallocated.
 	// 2 things can happen : a texture can be removed, in which case the texture will be set to null
diff --git a/Engine/Source/Runtime/Engine/Private/Streaming/AsyncTextureStreaming.h b/Engine/Source/Runtime/Engine/Private/Streaming/AsyncTextureStreaming.h
--- a/Engine/Source/Runtime/Engine/Private/Streaming/AsyncTextureStreaming.h
+++ b/Engine/Source/Runtime/Engine/Private/Streaming/AsyncTextureStreaming.h
@@ -26,7 +26,7 @@ public:
 	void Init(
 		TArray<FStreamingViewInfo> InViewInfos,
 		float InWorldTime,
-		TIndirectArray<FLevelRenderAssetManager>& LevelStaticInstanceManagers,
+		TArray<FLevelRenderAssetManager*>& LevelStaticInstanceManagers,
 		FDynamicRenderAssetInstanceManager& DynamicComponentManager);
 
 	void ComputeViewInfoExtras(const FRenderAssetStreamingSettings& Settings);
@@ -61,8 +61,8 @@ public:
 		}
 	}
 
-private:
 
+private:
 	/** Cached from FStreamingManagerBase. */
 	TArray<FStreamingViewInfo> ViewInfos;
 	
@@ -84,6 +84,9 @@ private:
 
 	/** List of all static instances that were rejected. */
 	TArray<int32> CulledStaticInstancesViewIndices;
+
+	/** list of all static instances based on level index */
+	TArray<int32> StaticInstancesViewLevelIndices;
 };
 
 struct FCompareRenderAssetByRetentionPriority // Bigger retention priority first.
diff --git a/Engine/Source/Runtime/Engine/Private/Streaming/LevelTextureManager.cpp b/Engine/Source/Runtime/Engine/Private/Streaming/LevelTextureManager.cpp
--- a/Engine/Source/Runtime/Engine/Private/Streaming/LevelTextureManager.cpp
+++ b/Engine/Source/Runtime/Engine/Private/Streaming/LevelTextureManager.cpp
@@ -14,6 +14,7 @@ FLevelRenderAssetManager::FLevelRenderAssetManager(ULevel* InLevel, RenderAssetI
 	, bIsInitialized(false)
 	, StaticInstances(AsyncTask)
 	, BuildStep(EStaticBuildStep::BuildTextureLookUpMap) 
+	, bHasBeenReferencedToStreamedTextures(false)
 {
 	if (Level)
 	{
diff --git a/Engine/Source/Runtime/Engine/Private/Streaming/LevelTextureManager.h b/Engine/Source/Runtime/Engine/Private/Streaming/LevelTextureManager.h
--- a/Engine/Source/Runtime/Engine/Private/Streaming/LevelTextureManager.h
+++ b/Engine/Source/Runtime/Engine/Private/Streaming/LevelTextureManager.h
@@ -53,6 +53,8 @@ public:
 	uint32 GetAllocatedSize() const;
 
 	bool IsInitialized() const { return bIsInitialized; }
+	bool HasBeenReferencedToStreamedTextures() const { return bHasBeenReferencedToStreamedTextures; }
+	void SetReferencedToStreamedTextures() { bHasBeenReferencedToStreamedTextures = true; }
 
 	void NotifyLevelOffset(const FVector& Offset);
 
@@ -61,7 +63,8 @@ private:
 	ULevel* Level;
 
 	bool bIsInitialized;
-
+	bool bHasBeenReferencedToStreamedTextures;
+	
 	FStaticRenderAssetInstanceManager StaticInstances;
 
 	/** Incremental build implementation. */
diff --git a/Engine/Source/Runtime/Engine/Private/Streaming/StreamingManagerTexture.cpp b/Engine/Source/Runtime/Engine/Private/Streaming/StreamingManagerTexture.cpp
--- a/Engine/Source/Runtime/Engine/Private/Streaming/StreamingManagerTexture.cpp
+++ b/Engine/Source/Runtime/Engine/Private/Streaming/StreamingManagerTexture.cpp
@@ -147,14 +147,19 @@ void FRenderAssetStreamingManager::OnPreGarbageCollect()
 	// Check all levels for pending kills.
 	for (int32 Index = 0; Index < LevelRenderAssetManagers.Num(); ++Index)
 	{
-		FLevelRenderAssetManager& LevelManager = LevelRenderAssetManagers[Index];
+		if (LevelRenderAssetManagers[Index] == nullptr)
+		{
+			continue;
+		}
+
+		FLevelRenderAssetManager& LevelManager = *LevelRenderAssetManagers[Index];
 		if (LevelManager.GetLevel()->IsPendingKill())
 		{
 			LevelManager.Remove(&RemovedRenderAssets);
 
 			// Remove the level entry. The async task view will still be valid as it uses a shared ptr.
-			LevelRenderAssetManagers.RemoveAtSwap(Index);
-			--Index;
+			delete LevelRenderAssetManagers[Index];
+			LevelRenderAssetManagers[Index] = nullptr;
 		}
 	}
 
@@ -300,9 +305,12 @@ void FRenderAssetStreamingManager::IncrementalUpdate(float Percentage, bool bUpd
 		NumStepsLeftForIncrementalBuild = MAX_int64;
 	}
 
-	for (FLevelRenderAssetManager& LevelManager : LevelRenderAssetManagers)
+	for (FLevelRenderAssetManager* LevelManager : LevelRenderAssetManagers)
 	{
-		LevelManager.IncrementalUpdate(DynamicComponentManager, RemovedRenderAssets, NumStepsLeftForIncrementalBuild, Percentage, bUseDynamicStreaming); // Complete the incremental update.
+		if (LevelManager != nullptr)
+		{
+			LevelManager->IncrementalUpdate(DynamicComponentManager, RemovedRenderAssets, NumStepsLeftForIncrementalBuild, Percentage, bUseDynamicStreaming); // Complete the incremental update.
+		}
 	}
 
 	// Dynamic component are only udpated when it is useful for the dynamic async view.
@@ -398,10 +406,13 @@ void FRenderAssetStreamingManager::ConditionalUpdateStaticData()
 			TArray<ULevel*, TInlineAllocator<32> > Levels;
 
 			// RemoveLevel data
-			for (FLevelRenderAssetManager& LevelManager : LevelRenderAssetManagers)
+			for (FLevelRenderAssetManager* LevelManager : LevelRenderAssetManagers)
 			{
-				Levels.Push(LevelManager.GetLevel());
-				LevelManager.Remove(nullptr);
+				if (LevelManager!=nullptr)
+				{
+					Levels.Push(LevelManager->GetLevel());
+					LevelManager->Remove(nullptr);
+				}
 			}
 			LevelRenderAssetManagers.Empty();
 
@@ -426,6 +437,60 @@ void FRenderAssetStreamingManager::ConditionalUpdateStaticData()
 	}
 }
 
+void FRenderAssetStreamingManager::ProcessLevelsToReferenceToStreamedTextures()
+{
+	// Iterate through levels and reference Levels to StreamedTexture if needed
+	for (int32 LevelIndex = 0; LevelIndex < LevelRenderAssetManagers.Num(); ++LevelIndex)
+	{
+		if (LevelRenderAssetManagers[LevelIndex] == nullptr)
+		{
+			continue;
+		}
+
+		FLevelRenderAssetManager& LevelRenderAssetManager = *LevelRenderAssetManagers[LevelIndex];
+		if (LevelRenderAssetManager.HasBeenReferencedToStreamedTextures())
+		{
+			continue;
+		}
+
+		const FRenderAssetInstanceView* View = LevelRenderAssetManager.GetRawAsyncView();
+		if (View == nullptr)
+		{
+			continue;
+		}
+
+		LevelRenderAssetManager.SetReferencedToStreamedTextures();
+
+		FRenderAssetInstanceView::FRenderAssetIterator RenderAssetIterator = LevelRenderAssetManager.GetRawAsyncView()->GetRenderAssetIterator();
+
+		for (; RenderAssetIterator; ++RenderAssetIterator)
+		{
+			const UStreamableRenderAsset* RenderAsset = *RenderAssetIterator;
+			if (RenderAsset == nullptr || !ReferencedRenderAssets.Contains(RenderAsset) || !StreamingRenderAssets.IsValidIndex(RenderAsset->StreamingIndex))
+			{
+				continue;
+			}
+
+			FStreamingRenderAsset& StreamingRenderAsset = StreamingRenderAssets[RenderAsset->StreamingIndex];
+
+			check(StreamingRenderAsset.RenderAsset == RenderAsset);
+
+			TBitArray<>& LevelIndexUsage = StreamingRenderAsset.LevelIndexUsage;
+
+			if (LevelIndex >= LevelIndexUsage.Num())
+			{
+				uint32 NumBits = LevelIndex + 1 - LevelIndexUsage.Num();
+				for (uint32 Index = 0; Index < NumBits; ++Index)
+				{
+					LevelIndexUsage.Add(false);
+				}
+			}
+
+			LevelIndexUsage[LevelIndex] = true;
+		}
+	}
+}
+
 void FRenderAssetStreamingManager::UpdatePendingStates(bool bUpdateDynamicComponents)
 {
 	CheckUserSettings();
@@ -444,6 +509,8 @@ void FRenderAssetStreamingManager::UpdatePendingStates(bool bUpdateDynamicCompon
 	{
 		DynamicComponentManager.PrepareAsyncView();
 	}
+
+	ProcessLevelsToReferenceToStreamedTextures();
 }
 
 /**
@@ -521,9 +588,9 @@ void FRenderAssetStreamingManager::AddLevel( ULevel* Level )
 	else
 	{
 		// In game, because static components can not be changed, the level static data is computed and kept as long as the level is not destroyed.
-		for (const FLevelRenderAssetManager& LevelManager : LevelRenderAssetManagers)
+		for (const FLevelRenderAssetManager* LevelManager : LevelRenderAssetManagers)
 		{
-			if (LevelManager.GetLevel() == Level)
+			if (LevelManager!=nullptr && LevelManager->GetLevel() == Level)
 			{
 				// Nothing to do, since the incremental update automatically manages what needs to be done.
 				return;
@@ -531,9 +598,19 @@ void FRenderAssetStreamingManager::AddLevel( ULevel* Level )
 		}
 	}
 
-	// If the level was not already there, add a new entry.
+	// If the level was not already there, create a new one, find an available slot or add a new one.
 	RenderAssetInstanceAsyncWork->EnsureCompletion();
-	LevelRenderAssetManagers.Add(new FLevelRenderAssetManager(Level, RenderAssetInstanceAsyncWork->GetTask()));
+	FLevelRenderAssetManager* LevelRenderAssetManager = new FLevelRenderAssetManager(Level, RenderAssetInstanceAsyncWork->GetTask());
+
+	uint32 LevelIndex = LevelRenderAssetManagers.FindLastByPredicate([](FLevelRenderAssetManager* Ptr) { return (Ptr == nullptr); });
+	if (LevelIndex != INDEX_NONE)
+	{
+		LevelRenderAssetManagers[LevelIndex] = LevelRenderAssetManager;
+	}
+	else
+	{
+		LevelRenderAssetManagers.Add(LevelRenderAssetManager);
+	}
 }
 
 /** Removes a ULevel from the streaming manager. */
@@ -549,15 +626,16 @@ void FRenderAssetStreamingManager::RemoveLevel( ULevel* Level )
 	{
 		for (int32 Index = 0; Index < LevelRenderAssetManagers.Num(); ++Index)
 		{
-			FLevelRenderAssetManager& LevelManager = LevelRenderAssetManagers[Index];
-			if (LevelManager.GetLevel() == Level)
+			FLevelRenderAssetManager* LevelManager = LevelRenderAssetManagers[Index];
+			if (LevelManager!=nullptr && LevelManager->GetLevel() == Level)
 			{
 				FRemovedRenderAssetArray RemovedRenderAssets;
-				LevelManager.Remove(&RemovedRenderAssets);
+				LevelManager->Remove(&RemovedRenderAssets);
 				SetRenderAssetsRemovedTimestamp(RemovedRenderAssets);
 
 				// Remove the level entry. The async task view will still be valid as it uses a shared ptr.
-				LevelRenderAssetManagers.RemoveAtSwap(Index);
+				LevelRenderAssetManagers[Index] = nullptr;
+				delete LevelManager;
 				break;
 			}
 		}
@@ -568,11 +646,11 @@ void FRenderAssetStreamingManager::NotifyLevelOffset(ULevel* Level, const FVecto
 {
 	FScopeLock ScopeLock(&CriticalSection);
 
-	for (FLevelRenderAssetManager& LevelManager : LevelRenderAssetManagers)
+	for (FLevelRenderAssetManager* LevelManager : LevelRenderAssetManagers)
 	{
-		if (LevelManager.GetLevel() == Level)
+		if (LevelManager!=nullptr && LevelManager->GetLevel() == Level)
 		{
-			LevelManager.NotifyLevelOffset(Offset);
+			LevelManager->NotifyLevelOffset(Offset);
 			break;
 		}
 	}
@@ -665,14 +743,14 @@ void FRenderAssetStreamingManager::NotifyActorDestroyed( AActor* Actor )
 	ULevel* Level = !GIsEditor ? Actor->GetLevel() : nullptr;
 
 	// Remove any reference in the level managers.
-	for (FLevelRenderAssetManager& LevelManager : LevelRenderAssetManagers)
+	for (FLevelRenderAssetManager* LevelManager : LevelRenderAssetManagers)
 	{
-		if (!Level || LevelManager.GetLevel() == Level)
+		if (!Level || (LevelManager!=nullptr && LevelManager->GetLevel() == Level))
 		{
-			LevelManager.RemoveActorReferences(Actor);
+			LevelManager->RemoveActorReferences(Actor);
 			for (UPrimitiveComponent* Component : Components)
 			{
-				LevelManager.RemoveComponentReferences(Component, RemovedRenderAssets);
+				LevelManager->RemoveComponentReferences(Component, RemovedRenderAssets);
 			}
 		}
 	}
@@ -700,11 +778,11 @@ void FRenderAssetStreamingManager::RemoveStaticReferences(const UPrimitiveCompon
 	{
 		FRemovedRenderAssetArray RemovedRenderAssets;
 		ULevel* Level = Primitive->GetComponentLevel();
-		for (FLevelRenderAssetManager& LevelManager : LevelRenderAssetManagers)
+		for (FLevelRenderAssetManager* LevelManager : LevelRenderAssetManagers)
 		{
-			if (!Level || LevelManager.GetLevel() == Level)
+			if (!Level || (LevelManager!=nullptr && LevelManager->GetLevel() == Level))
 			{
-				LevelManager.RemoveComponentReferences(Primitive, RemovedRenderAssets);
+				LevelManager->RemoveComponentReferences(Primitive, RemovedRenderAssets);
 			}
 		}
 		Primitive->bAttachedToStreamingManagerAsStatic = false;
@@ -745,11 +823,11 @@ void FRenderAssetStreamingManager::NotifyPrimitiveDetached( const UPrimitiveComp
 		// Unless in editor, we don't want to remove reference in static level data when toggling visibility.
 		else if (GIsEditor || Primitive->IsPendingKill() || Primitive->HasAnyFlags(RF_BeginDestroyed|RF_FinishDestroyed))
 		{
-			for (FLevelRenderAssetManager& LevelManager : LevelRenderAssetManagers)
+			for (FLevelRenderAssetManager* LevelManager : LevelRenderAssetManagers)
 			{
-				if (!Level || LevelManager.GetLevel() == Level)
+				if (!Level || (LevelManager!=nullptr && LevelManager->GetLevel() == Level))
 				{
-					LevelManager.RemoveComponentReferences(Primitive, RemovedRenderAssets);
+					LevelManager->RemoveComponentReferences(Primitive, RemovedRenderAssets);
 				}
 			}
 			Primitive->bAttachedToStreamingManagerAsStatic = false;
@@ -1136,8 +1214,13 @@ void FRenderAssetStreamingManager::SetLastUpdateTime()
 
 	for (int32 LevelIndex = 0; LevelIndex < LevelRenderAssetManagers.Num(); ++LevelIndex)
 	{
+		if (LevelRenderAssetManagers[LevelIndex] == nullptr)
+		{
+			continue;
+		}
+
 		// Update last update time only if there is a reasonable threshold to define visibility.
-		WorldTime = LevelRenderAssetManagers[LevelIndex].GetWorldTime();
+		WorldTime = LevelRenderAssetManagers[LevelIndex]->GetWorldTime();
 		if (WorldTime > 0)
 		{
 			break;
@@ -1426,9 +1509,14 @@ void FRenderAssetStreamingManager::GetObjectReferenceBounds(const UObject* RefOb
 	const UStreamableRenderAsset* RenderAsset = Cast<const UStreamableRenderAsset>(RefObject);
 	if (RenderAsset)
 	{
-		for (FLevelRenderAssetManager& LevelManager : LevelRenderAssetManagers)
+		for (FLevelRenderAssetManager *LevelManager : LevelRenderAssetManagers)
 		{
-			const FRenderAssetInstanceView* View = LevelManager.GetRawAsyncView();
+			if (LevelManager == nullptr)
+			{
+				continue;
+			}
+
+			const FRenderAssetInstanceView* View = LevelManager->GetRawAsyncView();
 			if (View)
 			{
 				for (auto It = View->GetElementIterator(RenderAsset); It; ++It)
@@ -1456,10 +1544,13 @@ void FRenderAssetStreamingManager::PropagateLightingScenarioChange()
 	// Note that dynamic components don't need to be handled because their renderstates are updated, which triggers and update.
 	
 	TArray<ULevel*, TInlineAllocator<32> > Levels;
-	for (FLevelRenderAssetManager& LevelManager : LevelRenderAssetManagers)
+	for (FLevelRenderAssetManager* LevelManager : LevelRenderAssetManagers)
 	{
-		Levels.Push(LevelManager.GetLevel());
-		LevelManager.Remove(nullptr);
+		if (LevelManager!=nullptr)
+		{
+			Levels.Push(LevelManager->GetLevel());
+			LevelManager->Remove(nullptr);
+		}
 	}
 
 	LevelRenderAssetManagers.Empty();
@@ -1847,9 +1938,12 @@ bool FRenderAssetStreamingManager::HandleStreamingManagerMemoryCommand( const TC
 	MemSize += LevelRenderAssetManagers.GetAllocatedSize();
 	MemSize += AsyncWork->GetTask().StreamingData.GetAllocatedSize();
 
-	for (const FLevelRenderAssetManager& LevelManager : LevelRenderAssetManagers)
+	for (const FLevelRenderAssetManager* LevelManager : LevelRenderAssetManagers)
 	{
-		MemSize += LevelManager.GetAllocatedSize();
+		if (LevelManager!=nullptr)
+		{
+			MemSize += LevelManager->GetAllocatedSize();
+		}
 	}
 
 	Ar.Logf(TEXT("StreamingManagerTexture: %.2f KB used"), MemSize / 1024.0f);
diff --git a/Engine/Source/Runtime/Engine/Private/Streaming/StreamingManagerTexture.h b/Engine/Source/Runtime/Engine/Private/Streaming/StreamingManagerTexture.h
--- a/Engine/Source/Runtime/Engine/Private/Streaming/StreamingManagerTexture.h
+++ b/Engine/Source/Runtime/Engine/Private/Streaming/StreamingManagerTexture.h
@@ -190,6 +190,7 @@ private:
 		void ProcessRemovedRenderAssets();
 		void ProcessAddedRenderAssets();
 		void ConditionalUpdateStaticData();
+		void ProcessLevelsToReferenceToStreamedTextures();
 
 		/** Adds new textures/meshes and level data on the gamethread (while the worker thread isn't active). */
 		void PrepareAsyncTask( bool bProcessEverything );
@@ -351,7 +352,7 @@ private:
 	int32 CurrentPendingMipCopyRequestIdx;
 
 	/** Level data */
-	TIndirectArray<FLevelRenderAssetManager> LevelRenderAssetManagers;
+	TArray<FLevelRenderAssetManager*> LevelRenderAssetManagers;
 
 	/** Stages [0,N-2] is non-threaded data collection, Stage N-1 is wait-for-AsyncWork-and-finalize. */
 	int32					ProcessingStage;
diff --git a/Engine/Source/Runtime/Engine/Private/Streaming/StreamingTexture.h b/Engine/Source/Runtime/Engine/Private/Streaming/StreamingTexture.h
--- a/Engine/Source/Runtime/Engine/Private/Streaming/StreamingTexture.h
+++ b/Engine/Source/Runtime/Engine/Private/Streaming/StreamingTexture.h
@@ -288,6 +288,9 @@ public:
 	/** (5) */
 	uint32			bCachedForceFullyLoadHeuristic : 1;
 
+	/** (5) */
+	TBitArray<>		LevelIndexUsage;
+
 private:
 	FORCEINLINE_DEBUGGABLE void StreamWantedMips_Internal(FRenderAssetStreamingManager& Manager, bool bUseCachedData);
 };
diff --git a/Engine/Source/Runtime/Engine/Private/Streaming/TextureStreamingHelpers.cpp b/Engine/Source/Runtime/Engine/Private/Streaming/TextureStreamingHelpers.cpp
--- a/Engine/Source/Runtime/Engine/Private/Streaming/TextureStreamingHelpers.cpp
+++ b/Engine/Source/Runtime/Engine/Private/Streaming/TextureStreamingHelpers.cpp
@@ -235,6 +235,12 @@ TAutoConsoleVariable<float> CVarStreamingMaxTextureUVDensity(
 	TEXT("Component with bigger entries become handled as dynamic component.\n"),
 	ECVF_Default);
 
+TAutoConsoleVariable<int32> CVarStreamingMipCalculationEnablePerLevelList(
+	TEXT("r.Streaming.MipCalculationEnablePerLevelList"),
+	1,
+	TEXT("If non-zero, Mip size computation for streamed texture will use levels referenced with it (instead of iterating thorugh every levels).\n"),
+	ECVF_Default);
+
 ENGINE_API TAutoConsoleVariable<int32> CVarFramesForFullUpdate(
 	TEXT("r.Streaming.FramesForFullUpdate"),
 	5,
@@ -278,6 +284,7 @@ void FRenderAssetStreamingSettings::Update()
 	MaxTextureUVDensity = CVarStreamingMaxTextureUVDensity.GetValueOnAnyThread();
 	bUseMaterialData = bUseNewMetrics && CVarStreamingUseMaterialData.GetValueOnAnyThread() != 0;
 	HiddenPrimitiveScale = bUseNewMetrics ? CVarStreamingHiddenPrimitiveScale.GetValueOnAnyThread() : 1.f;
+	bMipCalculationEnablePerLevelList = CVarStreamingMipCalculationEnablePerLevelList.GetValueOnAnyThread() != 0;
 
 	MaterialQualityLevel = (int32)GetCachedScalabilityCVars().MaterialQualityLevel;
 
diff --git a/Engine/Source/Runtime/Engine/Private/Streaming/TextureStreamingHelpers.h b/Engine/Source/Runtime/Engine/Private/Streaming/TextureStreamingHelpers.h
--- a/Engine/Source/Runtime/Engine/Private/Streaming/TextureStreamingHelpers.h
+++ b/Engine/Source/Runtime/Engine/Private/Streaming/TextureStreamingHelpers.h
@@ -89,6 +89,7 @@ struct FRenderAssetStreamingSettings
 	float MaxTextureUVDensity;
 	int32 MaterialQualityLevel;
 	int32 FramesForFullUpdate;
+	bool bMipCalculationEnablePerLevelList;
 
 	bool bStressTest;
 	static int32 ExtraIOLatency;
