Added DISABLE_REPLICATED_PROPERTY macro to explicitely disable the replication of a networked property in GetLifetimeReplicatedProps.
This is needed since we will start warning when we detect replicated properties not registered in GetLifetimeReplicatedProps

Added DISABLE_ALL_CLASS_REPLICATED_PROPERTIES macro to disable the replication of ALL replicated properties of a given class

Added RESET_REPLIFETIME macro that allows you to reenable a networked property that was Disabled in a base class.

diff --git a/Engine/Source/Runtime/CoreUObject/Public/UObject/CoreNetTypes.h b/Engine/Source/Runtime/CoreUObject/Public/UObject/CoreNetTypes.h
--- a/Engine/Source/Runtime/CoreUObject/Public/UObject/CoreNetTypes.h
+++ b/Engine/Source/Runtime/CoreUObject/Public/UObject/CoreNetTypes.h
@@ -22,8 +22,9 @@ enum ELifetimeCondition
 	COND_ReplayOnly = 10					UMETA(DisplayName = "Replay Only"),						// This property will only send to the replay connection
 	COND_SimulatedOnlyNoReplay = 11			UMETA(DisplayName = "Simulated Only No Replay"),		// This property will send to actors only, but not to replay connections
 	COND_SimulatedOrPhysicsNoReplay = 12	UMETA(DisplayName = "Simulated Or Physics No Replay"),	// This property will send to simulated Or bRepPhysics actors, but not to replay connections
-	COND_SkipReplay = 13					UMETA(DisplayName = "Skip Replay"),	// This property will not send to the replay connection
-	COND_Max = 14							UMETA(Hidden)
+	COND_SkipReplay = 13					UMETA(DisplayName = "Skip Replay"),						// This property will not send to the replay connection
+	COND_Never = 15							UMETA(Hidden),											// This property will never be replicated
+	COND_Max = 16							UMETA(Hidden)
 };
 
 
diff --git a/Engine/Source/Runtime/Engine/Private/RepLayout.cpp b/Engine/Source/Runtime/Engine/Private/RepLayout.cpp
--- a/Engine/Source/Runtime/Engine/Private/RepLayout.cpp
+++ b/Engine/Source/Runtime/Engine/Private/RepLayout.cpp
@@ -560,6 +560,7 @@ TStaticBitArray<COND_Max> BuildConditionMapFromRepFlags(FReplicationFlags RepFla
 	ConditionMap[COND_SkipReplay] = !bIsReplay;
 
 	ConditionMap[COND_Custom] = true;
+	ConditionMap[COND_Never] = false;
 
 	return ConditionMap;
 }
@@ -626,14 +627,14 @@ void FReplicationChangelistMgr::CountBytes(FArchive& Ar) const
 
 uint16 FRepLayout::CompareProperties_r(
 	FSendingRepState* RESTRICT RepState,
-	const int32 CmdStart,
-	const int32 CmdEnd,
-	const uint8* RESTRICT ShadowData,
-	const uint8* RESTRICT Data,
-	TArray<uint16>& Changed,
-	uint16 Handle,
-	const bool bIsInitial,
-	const bool bForceFail) const
+	const int32				CmdStart,
+	const int32				CmdEnd,
+	const uint8* RESTRICT	ShadowData,
+	const uint8* RESTRICT	Data,
+	TArray<uint16>&			Changed,
+	uint16					Handle,
+	const bool				bIsInitial,
+	const bool				bForceFail) const
 {
 	check(ShadowData);
 
@@ -708,13 +709,13 @@ uint16 FRepLayout::CompareProperties_r(
 
 void FRepLayout::CompareProperties_Array_r(
 	FSendingRepState* RESTRICT RepState,
-	const uint8* RESTRICT ShadowData,
-	const uint8* RESTRICT Data,
-	TArray<uint16>& Changed,
-	const uint16 CmdIndex,
-	const uint16 Handle,
-	const bool bIsInitial,
-	const bool bForceFail
+	const uint8* RESTRICT	ShadowData,
+	const uint8* RESTRICT	Data,
+	TArray<uint16>&			Changed,
+	const uint16			CmdIndex,
+	const uint16			Handle,
+	const bool				bIsInitial,
+	const bool				bForceFail
 	) const
 {
 	const FRepLayoutCmd& Cmd = Cmds[CmdIndex];
@@ -766,9 +767,9 @@ void FRepLayout::CompareProperties_Array_r(
 
 bool FRepLayout::CompareProperties(
 	FSendingRepState* RESTRICT RepState,
-	FRepChangelistState* RESTRICT RepChangelistState,
-	const uint8* RESTRICT Data,
-	const FReplicationFlags& RepFlags) const
+	FRepChangelistState* RESTRICT	RepChangelistState,
+	const uint8* RESTRICT			Data,
+	const FReplicationFlags&		RepFlags) const
 {
 	SCOPE_CYCLE_COUNTER(STAT_NetReplicateDynamicPropCompareTime);
 
@@ -900,8 +901,8 @@ bool FRepLayout::ReplicateProperties(
 			if (Pruned.Num() > 0)
 			{
 				SendProperties_BackwardsCompatible(RepState, ChangeTracker, Data, OwningChannel->Connection, Writer, Pruned);
-				return true;
-			}
+			return true;
+		}
 		}
 
 		return false;
@@ -1045,8 +1046,8 @@ PRAGMA_ENABLE_DEPRECATION_WARNINGS
 
 		if (Changed.Num() > 0)
 		{
-			SendProperties_BackwardsCompatible(RepState, ChangeTracker, Data, OwningChannel->Connection, Writer, Changed);
-		}
+		SendProperties_BackwardsCompatible(RepState, ChangeTracker, Data, OwningChannel->Connection, Writer, Changed);
+	}
 	}
 	else if (Changed.Num() > 0)
 	{
@@ -2096,13 +2097,13 @@ void FRepLayout::SendProperties_BackwardsCompatible_r(
 
 void FRepLayout::SendAllProperties_BackwardsCompatible_r(
 	FSendingRepState* RESTRICT RepState,
-	FNetBitWriter& Writer,
-	const bool bDoChecksum,
-	UPackageMapClient* PackageMapClient,
-	FNetFieldExportGroup* NetFieldExportGroup,
-	const int32 CmdStart,
-	const int32 CmdEnd, 
-	const uint8* SourceData) const
+	FNetBitWriter&						Writer,
+	const bool							bDoChecksum,
+	UPackageMapClient*					PackageMapClient,
+	FNetFieldExportGroup*				NetFieldExportGroup,
+	const int32							CmdStart,
+	const int32							CmdEnd, 
+	const uint8*						SourceData) const
 {
 	FNetBitWriter TempWriter(Writer.PackageMap, 0);
 
@@ -4178,12 +4179,12 @@ void FRepLayout::InitFromClass(UClass* InObjectClass, const UNetConnection* Serv
 	SCOPE_CYCLE_UOBJECT(ObjectClass, InObjectClass);
 
 	const bool bIsObjectActor = InObjectClass->IsChildOf(AActor::StaticClass());
-	RoleIndex                 = -1;
-	RemoteRoleIndex           = -1;
-	FirstNonCustomParent      = -1;
+	RoleIndex = -1;
+	RemoteRoleIndex = -1;
+	FirstNonCustomParent = -1;
 
-	int32 RelativeHandle      = 0;
-	int32 LastOffset          = -1;
+	int32 RelativeHandle = 0;
+	int32 LastOffset = -1;
 
 	InObjectClass->SetUpRuntimeReplicationData();
 	Parents.Empty(InObjectClass->ClassReps.Num());
@@ -4227,22 +4228,22 @@ void FRepLayout::InitFromClass(UClass* InObjectClass, const UNetConnection* Serv
 
 		if (bIsObjectActor)
 		{
-			// Find Role/RemoteRole property indexes so we can swap them on the client
-			if (Property->GetFName() == NAME_Role)
-			{
-				check(RoleIndex == -1);
-				check(Parents[ParentHandle].CmdEnd == Parents[ParentHandle].CmdStart + 1);
-				RoleIndex = ParentHandle;
-			}
+		// Find Role/RemoteRole property indexes so we can swap them on the client
+		if (Property->GetFName() == NAME_Role)
+		{
+			check(RoleIndex == -1);
+			check(Parents[ParentHandle].CmdEnd == Parents[ParentHandle].CmdStart + 1);
+			RoleIndex = ParentHandle;
+		}
 
-			if (Property->GetFName() == NAME_RemoteRole)
-			{
-				check(RemoteRoleIndex == -1);
-				check(Parents[ParentHandle].CmdEnd == Parents[ParentHandle].CmdStart + 1);
-				RemoteRoleIndex = ParentHandle;
-			}
+		if (Property->GetFName() == NAME_RemoteRole)
+		{
+			check(RemoteRoleIndex == -1);
+			check(Parents[ParentHandle].CmdEnd == Parents[ParentHandle].CmdStart + 1);
+			RemoteRoleIndex = ParentHandle;
 		}
 	}
+	}
 
 	// Make sure it either found both, or didn't find either
 	check((RoleIndex == -1) == (RemoteRoleIndex == -1));
@@ -4296,7 +4297,7 @@ void FRepLayout::InitFromClass(UClass* InObjectClass, const UNetConnection* Serv
 		{
 			// We handle remote role specially, since it can change between connections when downgraded
 			// So we force it on the conditional list
-			check(LifetimeProps[i].Condition == COND_None);
+			check(LifetimeProps[i].Condition == COND_None || LifetimeProps[i].Condition == COND_Never);
 			LifetimeProps[i].Condition = COND_Custom;
 			continue;
 		}
diff --git a/Engine/Source/Runtime/Engine/Private/UnrealNetwork.cpp b/Engine/Source/Runtime/Engine/Private/UnrealNetwork.cpp
--- a/Engine/Source/Runtime/Engine/Private/UnrealNetwork.cpp
+++ b/Engine/Source/Runtime/Engine/Private/UnrealNetwork.cpp
@@ -7,3 +7,131 @@ FOnWriteGameSpecificDemoHeader FNetworkReplayDelegates::OnWriteGameSpecificDemoH
 FOnProcessGameSpecificDemoHeader FNetworkReplayDelegates::OnProcessGameSpecificDemoHeader;
 
 // ----------------------------------------------------------------
+
+void RegisterReplicatedLifetimeProperty(const UProperty* ReplicatedProperty, TArray< FLifetimeProperty >& OutLifetimeProps, ELifetimeCondition InCondition, ELifetimeRepNotifyCondition InRepNotifyCondition)
+{
+	if (!ReplicatedProperty) 
+	{
+		check(false);
+		return;
+	}
+
+	for ( int32 i = 0; i < ReplicatedProperty->ArrayDim; i++ )
+	{
+		const uint16 RepIndex = ReplicatedProperty->RepIndex + i;
+		FLifetimeProperty* RegisteredPropertyPtr = OutLifetimeProps.FindByPredicate([&RepIndex](const FLifetimeProperty& Var) { return Var.RepIndex == RepIndex; });
+
+		FLifetimeProperty LifetimeProp(RepIndex, InCondition, InRepNotifyCondition);
+
+		if (RegisteredPropertyPtr)
+		{
+			// Disabled properties can be re-enabled via DOREPLIFETIME
+			if (RegisteredPropertyPtr->Condition == COND_Never)
+			{
+				// Copy the new conditions since disabling a property doesn't set other conditions.
+				(*RegisteredPropertyPtr) = LifetimeProp;
+			}
+			else
+			{
+				// Conditions should be identical when calling DOREPLIFETIME twice on the same variable.
+				checkf((*RegisteredPropertyPtr) == LifetimeProp, TEXT("Property %s was registered twice with different conditions (old:%d) (new:%d)"), *ReplicatedProperty->GetName(), RegisteredPropertyPtr->Condition, InCondition);
+			}
+		}
+		else
+		{
+			OutLifetimeProps.Add(LifetimeProp);
+		}
+	}
+}
+
+void SetReplicatedPropertyToDisabled(const UProperty* ReplicatedProperty, TArray< FLifetimeProperty >& OutLifetimeProps)
+{
+	check(ReplicatedProperty);
+
+	for (int32 i = 0; i < ReplicatedProperty->ArrayDim; i++)
+	{
+		const uint16 RepIndex = ReplicatedProperty->RepIndex + i;
+		FLifetimeProperty* RegisteredPropertyPtr = OutLifetimeProps.FindByPredicate([&RepIndex](const FLifetimeProperty& Var) { return Var.RepIndex == RepIndex; });
+
+		if (RegisteredPropertyPtr)
+		{
+			RegisteredPropertyPtr->Condition = COND_Never;
+		}
+		else
+		{
+			OutLifetimeProps.Add(FLifetimeProperty(RepIndex, COND_Never));
+		}
+	}
+}
+
+void DisableReplicatedLifetimeProperty(const UClass* ThisClass, const UClass* PropertyClass, FName PropertyName, TArray< FLifetimeProperty >& OutLifetimeProps)
+{
+	const UProperty* ReplicatedProperty = GetReplicatedProperty(ThisClass, PropertyClass, PropertyName);
+	if (!ReplicatedProperty)
+	{
+		return;
+	}
+
+	SetReplicatedPropertyToDisabled(ReplicatedProperty, OutLifetimeProps);
+}
+
+void ResetReplicatedLifetimeProperty(const UClass* ThisClass, const UClass* PropertyClass, FName PropertyName, ELifetimeCondition LifetimeCondition, TArray< FLifetimeProperty >& OutLifetimeProps)
+{
+	const UProperty* ReplicatedProperty = GetReplicatedProperty(ThisClass, PropertyClass, PropertyName);
+	if (!ReplicatedProperty)
+	{
+		return;
+	}
+
+	for (int32 i = 0; i < ReplicatedProperty->ArrayDim; i++)
+	{
+		uint16 RepIndex = ReplicatedProperty->RepIndex + i;
+		FLifetimeProperty* RegisteredPropertyPtr = OutLifetimeProps.FindByPredicate([&RepIndex](const FLifetimeProperty& Var) { return Var.RepIndex == RepIndex; });
+
+		// Set the new condition
+		if (RegisteredPropertyPtr)
+		{
+			RegisteredPropertyPtr->Condition = LifetimeCondition;
+		}
+		else
+		{
+			OutLifetimeProps.Add(FLifetimeProperty(RepIndex, LifetimeCondition));
+		}
+	}
+}
+
+void DisableAllReplicatedPropertiesOfClass(const UClass* ThisClass, const UClass* ClassToDisable, EFieldIteratorFlags::SuperClassFlags SuperClassBehavior, TArray< FLifetimeProperty >& OutLifetimeProps)
+{
+	if (!ThisClass->IsChildOf(ClassToDisable))
+	{
+		ensureMsgf(false, TEXT("Attempting to disable replicated properties of '%s' but current class '%s' is not a child of '%s'"), *ClassToDisable->GetName(), *ThisClass->GetName(), *ClassToDisable->GetName());
+		return;
+	}
+
+	for (TFieldIterator<UProperty> It(ClassToDisable, SuperClassBehavior); It; ++It)
+	{
+		const UProperty* Prop = *It;
+		if (Prop && Prop->PropertyFlags & CPF_Net)
+		{
+			SetReplicatedPropertyToDisabled(Prop, OutLifetimeProps);
+		}
+	}
+}
+
+void DeprecatedChangeCondition(const  UProperty* ReplicatedProperty, TArray<FLifetimeProperty>& OutLifetimeProps, ELifetimeCondition InCondition)
+{
+	bool bFound = false;
+	for (int32 i = 0; i < OutLifetimeProps.Num(); i++)
+	{
+		if (OutLifetimeProps[i].RepIndex == ReplicatedProperty->RepIndex)
+		{
+			for ( int32 j = 0; j < ReplicatedProperty->ArrayDim; j++ )
+			{
+				OutLifetimeProps[i + j].Condition = InCondition;
+			}
+			bFound = true;
+			break;
+		}
+	}
+	check( bFound );
+}
\ No newline at end of file
diff --git a/Engine/Source/Runtime/Engine/Public/Net/UnrealNetwork.h b/Engine/Source/Runtime/Engine/Public/Net/UnrealNetwork.h
--- a/Engine/Source/Runtime/Engine/Public/Net/UnrealNetwork.h
+++ b/Engine/Source/Runtime/Engine/Public/Net/UnrealNetwork.h
@@ -59,7 +59,7 @@ struct ENGINE_API FNetworkReplayDelegates
 -----------------------------------------------------------------------------*/
 
 /** wrapper to find replicated properties that also makes sure they're valid */
-static UProperty* GetReplicatedProperty(UClass* CallingClass, UClass* PropClass, const FName& PropName)
+static UProperty* GetReplicatedProperty(const UClass* CallingClass, const UClass* PropClass, const FName& PropName)
 {
 #if !(UE_BUILD_SHIPPING || UE_BUILD_TEST)
 	if (!CallingClass->IsChildOf(PropClass))
@@ -79,11 +79,8 @@ static UProperty* GetReplicatedProperty(UClass* CallingClass, UClass* PropClass,
 
 #define DOREPLIFETIME(c,v) \
 { \
-	static UProperty* sp##v = GetReplicatedProperty(StaticClass(), c::StaticClass(),GET_MEMBER_NAME_CHECKED(c,v)); \
-	for ( int32 i = 0; i < sp##v->ArrayDim; i++ )							\
-	{																		\
-		OutLifetimeProps.AddUnique( FLifetimeProperty( sp##v->RepIndex + i ) );	\
-	}																		\
+	UProperty* ReplicatedProperty = GetReplicatedProperty(StaticClass(), c::StaticClass(),GET_MEMBER_NAME_CHECKED(c,v)); \
+	RegisterReplicatedLifetimeProperty(ReplicatedProperty, OutLifetimeProps, COND_None); \
 }
 
 /** This macro is used by nativized code (DynamicClasses), so the Property may be recreated. */
@@ -104,51 +101,68 @@ static UProperty* GetReplicatedProperty(UClass* CallingClass, UClass* PropClass,
 
 #define DOREPLIFETIME_CONDITION(c,v,cond) \
 { \
-	static UProperty* sp##v = GetReplicatedProperty(StaticClass(), c::StaticClass(),GET_MEMBER_NAME_CHECKED(c,v)); \
-	for ( int32 i = 0; i < sp##v->ArrayDim; i++ )										\
-	{																					\
-		OutLifetimeProps.AddUnique( FLifetimeProperty( sp##v->RepIndex + i, cond ) );	\
-	}																					\
+	UProperty* ReplicatedProperty = GetReplicatedProperty(StaticClass(), c::StaticClass(),GET_MEMBER_NAME_CHECKED(c,v)); \
+	RegisterReplicatedLifetimeProperty(ReplicatedProperty, OutLifetimeProps, cond); \
 }
 
 /** Allows gamecode to specify RepNotify condition: REPNOTIFY_OnChanged (default) or REPNOTIFY_Always for when repnotify function is called  */
 #define DOREPLIFETIME_CONDITION_NOTIFY(c,v,cond, rncond) \
 { \
-	static UProperty* sp##v = GetReplicatedProperty(StaticClass(), c::StaticClass(),GET_MEMBER_NAME_CHECKED(c,v)); \
-	for ( int32 i = 0; i < sp##v->ArrayDim; i++ )										\
-	{																					\
-		OutLifetimeProps.AddUnique( FLifetimeProperty( sp##v->RepIndex + i, cond, rncond) );	\
-	}																					\
+	UProperty* ReplicatedProperty = GetReplicatedProperty(StaticClass(), c::StaticClass(),GET_MEMBER_NAME_CHECKED(c,v)); \
+	RegisterReplicatedLifetimeProperty(ReplicatedProperty, OutLifetimeProps, cond, rncond); \
 }
 
 #define DOREPLIFETIME_ACTIVE_OVERRIDE(c,v,active)	\
 {													\
-	static UProperty* sp##v = GetReplicatedProperty(StaticClass(), c::StaticClass(),GET_MEMBER_NAME_CHECKED(c,v)); \
+	UProperty* sp##v = GetReplicatedProperty(StaticClass(), c::StaticClass(),GET_MEMBER_NAME_CHECKED(c,v)); \
 	for ( int32 i = 0; i < sp##v->ArrayDim; i++ )											\
 	{																						\
 		ChangedPropertyTracker.SetCustomIsActiveOverride( sp##v->RepIndex + i, active );	\
 	}																						\
 }
 
+UE_DEPRECATED(4.30, "Please use the RESET_REPLIFETIME_CONDITION macro")
+ENGINE_API void DeprecatedChangeCondition(UProperty* ReplicatedProperty, TArray< FLifetimeProperty >& OutLifetimeProps, ELifetimeCondition InCondition);
+
 #define DOREPLIFETIME_CHANGE_CONDITION(c,v,cond) \
 { \
-	static UProperty* sp##v = GetReplicatedProperty(StaticClass(), c::StaticClass(),GET_MEMBER_NAME_CHECKED(c,v));	\
-	bool bFound = false;																							\
-	for ( int32 i = 0; i < OutLifetimeProps.Num(); i++ )															\
-	{																												\
-		if ( OutLifetimeProps[i].RepIndex == sp##v->RepIndex )														\
-		{																											\
-			for ( int32 j = 0; j < sp##v->ArrayDim; j++ )															\
-			{																										\
-				OutLifetimeProps[i + j].Condition = cond;															\
-			}																										\
-			bFound = true;																							\
-			break;																									\
-		}																											\
-	}																												\
-	check( bFound );																								\
+	UProperty* sp##v = GetReplicatedProperty(StaticClass(), c::StaticClass(),GET_MEMBER_NAME_CHECKED(c,v));			\
+	DeprecatedChangeCondition(sp##v, OutLifetimeProps, cond);														\
 }
 
+ENGINE_API void RegisterReplicatedLifetimeProperty(const UProperty* ReplicatedProperty, TArray< FLifetimeProperty >& OutLifetimeProps, ELifetimeCondition InCondition, ELifetimeRepNotifyCondition InRepNotifyCondition = REPNOTIFY_OnChanged);
+
+/*-----------------------------------------------------------------------------
+	Disable macros.
+	Use these macros to state that properties tagged replicated 
+	are voluntarely not replicated. This silences an error about missing
+	registered properties when class replication is started
+-----------------------------------------------------------------------------*/
+
+/** Use this macro in GetLifetimeReplicatedProps to flag a replicated property as not-replicated */
+#define DISABLE_REPLICATED_PROPERTY(c,v) \
+DisableReplicatedLifetimeProperty(StaticClass(), c::StaticClass(), GET_MEMBER_NAME_CHECKED(c,v), OutLifetimeProps);
+
+/** Use this macro in GetLifetimeReplicatedProps to flag all replicated properties of a class as not-replicated.
+    Use the EFieldIteratorFlags enum to disable all inherited properties or only those of the class specified
+*/
+#define DISABLE_ALL_CLASS_REPLICATED_PROPERTIES(c, SuperClassBehavior) \
+DisableAllReplicatedPropertiesOfClass(StaticClass(), c::StaticClass(), SuperClassBehavior, OutLifetimeProps);
+
+ENGINE_API void DisableReplicatedLifetimeProperty(const UClass* ThisClass, const UClass* PropertyClass, FName PropertyName, TArray< FLifetimeProperty >& OutLifetimeProps);
+ENGINE_API void DisableAllReplicatedPropertiesOfClass(const UClass* ThisClass, const UClass* ClassToDisable, EFieldIteratorFlags::SuperClassFlags SuperClassBehavior, TArray< FLifetimeProperty >& OutLifetimeProps);
+
+/*-----------------------------------------------------------------------------
+	Reset macros.
+	Use these to change the replication settings of an inherited property
+-----------------------------------------------------------------------------*/
+
+#define RESET_REPLIFETIME(c,v)  ResetReplicatedLifetimeProperty(StaticClass(), c::StaticClass(), GET_MEMBER_NAME_CHECKED(c,v), COND_None, OutLifetimeProps);
+
+#define RESET_REPLIFETIME_CONDITION(c,v,cond) ResetReplicatedLifetimeProperty(StaticClass(), c::StaticClass(), GET_MEMBER_NAME_CHECKED(c,v), cond, OutLifetimeProps);
+
+ENGINE_API void ResetReplicatedLifetimeProperty(const UClass* ThisClass, const UClass* PropertyClass, FName PropertyName, ELifetimeCondition LifetimeCondition, TArray< FLifetimeProperty >& OutLifetimeProps);
+
 /*-----------------------------------------------------------------------------
 	RPC Parameter Validation Helpers
 -----------------------------------------------------------------------------*/
