Deprecated and hidden WorldInfo.bEnableNavigationSystem #UE4

diff --git a/Engine/Source/Runtime/Engine/Classes/GameFramework/WorldSettings.h b/Engine/Source/Runtime/Engine/Classes/GameFramework/WorldSettings.h
--- a/Engine/Source/Runtime/Engine/Classes/GameFramework/WorldSettings.h
+++ b/Engine/Source/Runtime/Engine/Classes/GameFramework/WorldSettings.h
@@ -438,7 +438,8 @@ protected:
 	 *	navigation functionality won't be accessible).
 	 *	This flag is now deprecated. Use NavigationSystemConfig property to 
 	 *	determine navigation system's properties or existence */
-	UPROPERTY(VisibleAnywhere, BlueprintReadOnly, config, Category=World, AdvancedDisplay, meta=(DisplayName = "DEPRECATED_bEnableNavigationSystem"))
+	UE_DEPRECATED(4.22, "This member will be removed. Please use NavigationSystemConfig instead.")
+	UPROPERTY(BlueprintReadOnly, config, Category = World, meta=(DisplayName = "DEPRECATED_bEnableNavigationSystem"))
 	uint8 bEnableNavigationSystem:1;
 
 public:
diff --git a/Engine/Source/Runtime/Engine/Private/WorldSettings.cpp b/Engine/Source/Runtime/Engine/Private/WorldSettings.cpp
--- a/Engine/Source/Runtime/Engine/Private/WorldSettings.cpp
+++ b/Engine/Source/Runtime/Engine/Private/WorldSettings.cpp
@@ -65,7 +65,9 @@ AWorldSettings::AWorldSettings(const FObjectInitializer& ObjectInitializer)
 	static FConstructorStatics ConstructorStatics;
 
 	bEnableWorldBoundsChecks = true;
+PRAGMA_DISABLE_DEPRECATION_WARNINGS
 	bEnableNavigationSystem = true;
+PRAGMA_ENABLE_DEPRECATION_WARNINGS
 	NavigationSystemConfig = nullptr;
 	bEnableAISystem = true;
 	bEnableWorldComposition = false;
@@ -449,6 +451,8 @@ void AWorldSettings::PostLoad()
 
 #endif// WITH_EDITOR
 
+PRAGMA_DISABLE_DEPRECATION_WARNINGS
+	// temporarily using deprecated bEnableNavigationSystem for backwards compatibility
 	if (bEnableNavigationSystem && NavigationSystemConfig == nullptr)
 	{
 		ULevel* Level = GetLevel();
@@ -462,6 +466,7 @@ void AWorldSettings::PostLoad()
 			bEnableNavigationSystem = false;
 		}
 	}
+PRAGMA_ENABLE_DEPRECATION_WARNINGS
 }
 
 bool AWorldSettings::IsNavigationSystemEnabled() const
