Fixed bug where FColor::MakeRandomColor would return white.

 The method is used in two major places:
 - Debug skeletons when rendering bones in the viewport.
 - Module Editor Colors in Cascade.

Both cases do now produce random colors as opposed to white.

#5497

diff --git a/Engine/Source/Runtime/Core/Private/Math/Color.cpp b/Engine/Source/Runtime/Core/Private/Math/Color.cpp
--- a/Engine/Source/Runtime/Core/Private/Math/Color.cpp
+++ b/Engine/Source/Runtime/Core/Private/Math/Color.cpp
@@ -353,7 +353,7 @@ FLinearColor FLinearColor::LerpUsingHSV( const FLinearColor& From, const FLinear
 FLinearColor FLinearColor::MakeRandomColor()
 {
 	const uint8 Hue = (uint8)(FMath::FRand()*255.f);
-	return FLinearColor::MakeFromHSV8(Hue, 0, 255);
+	return FLinearColor::MakeFromHSV8(Hue, 255, 255);
 }
 
 FColor FColor::MakeRandomColor()
diff --git a/Engine/Source/Runtime/Engine/Private/SkeletalMesh.cpp b/Engine/Source/Runtime/Engine/Private/SkeletalMesh.cpp
--- a/Engine/Source/Runtime/Engine/Private/SkeletalMesh.cpp
+++ b/Engine/Source/Runtime/Engine/Private/SkeletalMesh.cpp
@@ -4161,7 +4161,7 @@ void FSkeletalMeshSceneProxy::DebugDrawSkeleton(int32 ViewIndex, FMeshElementCol
 	{
 		FRandomStream Stream((int32)InUID);
 		const uint8 Hue = (uint8)(Stream.FRand()*255.f);
-		return FLinearColor::MakeFromHSV8(Hue, 0, 255);
+		return FLinearColor::MakeFromHSV8(Hue, 255, 255);
 	};
 
 	FPrimitiveDrawInterface* PDI = Collector.GetPDI(ViewIndex);
