Fixed bug with Particle Cutouts where an alpha threshold of 0 would disable the feature. An alpha threshold of 0 now culls regions with exactly zero alpha.

The fix simply removes the check for an alpha threshold of 0. Enabling particle cutouts now simply checks whether the cutout texture exists.

diff --git a/Engine/Source/Runtime/Engine/Private/Particles/ParticleSystemRender.cpp b/Engine/Source/Runtime/Engine/Private/Particles/ParticleSystemRender.cpp
--- a/Engine/Source/Runtime/Engine/Private/Particles/ParticleSystemRender.cpp
+++ b/Engine/Source/Runtime/Engine/Private/Particles/ParticleSystemRender.cpp
@@ -1026,7 +1026,7 @@ FParticleVertexFactoryBase *FDynamicSpriteEmitterData::CreateVertexFactory(ERHIF
 	FParticleSpriteVertexFactory *VertexFactory = new FParticleSpriteVertexFactory(InFeatureLevel);
 	VertexFactory->SetParticleFactoryType(PVFT_Sprite);
 	const FParticleRequiredModule * RequiredModule = GetSourceData()->RequiredModule;
-	VertexFactory->SetNumVertsInInstanceBuffer(RequiredModule->bCutoutTexureIsValid && RequiredModule->AlphaThreshold ? RequiredModule->NumBoundingVertices : 4);
+	VertexFactory->SetNumVertsInInstanceBuffer(RequiredModule->bCutoutTexureIsValid ? RequiredModule->NumBoundingVertices : 4);
 	VertexFactory->SetUsesDynamicParameter(bUsesDynamicParameter, bUsesDynamicParameter ? GetDynamicParameterVertexStride(): 0);
 	VertexFactory->InitResource();
 	return VertexFactory;
@@ -1193,7 +1193,7 @@ void FDynamicSpriteEmitterData::GetDynamicMeshElementsEmitter(const FParticleSys
 					SpriteVertexFactory->SetInstanceBuffer(Allocation.VertexBuffer, Allocation.VertexOffset, InstanceBufferStride, bInstanced);
 					SpriteVertexFactory->SetDynamicParameterBuffer(DynamicParameterAllocation.VertexBuffer, DynamicParameterAllocation.VertexOffset, GetDynamicParameterVertexStride(), bInstanced);
 
-					if (SourceData->RequiredModule->bCutoutTexureIsValid && SourceData->RequiredModule->AlphaThreshold)
+					if (SourceData->RequiredModule->bCutoutTexureIsValid)
 					{
 						SpriteVertexFactory->SetCutoutParameters(SourceData->RequiredModule->NumBoundingVertices, SourceData->RequiredModule->BoundingGeometryBufferSRV);
 					}
