Copy - Fix for unbuilt ResourceClusters causing instanced meshes to not find their FMeshMapBuildData during a build

diff --git a/Engine/Source/Runtime/Engine/Classes/Engine/MapBuildDataRegistry.h b/Engine/Source/Runtime/Engine/Classes/Engine/MapBuildDataRegistry.h
--- a/Engine/Source/Runtime/Engine/Classes/Engine/MapBuildDataRegistry.h
+++ b/Engine/Source/Runtime/Engine/Classes/Engine/MapBuildDataRegistry.h
@@ -298,6 +298,7 @@ public:
 	ENGINE_API FMeshMapBuildData& AllocateMeshBuildData(const FGuid& MeshId, bool bMarkDirty);
 	ENGINE_API const FMeshMapBuildData* GetMeshBuildData(FGuid MeshId) const;
 	ENGINE_API FMeshMapBuildData* GetMeshBuildData(FGuid MeshId);
+	ENGINE_API FMeshMapBuildData* GetMeshBuildDataDuringBuild(FGuid MeshId);
 
 	/** 
 	 * Allocates a new FPrecomputedLightVolumeData from the registry.
diff --git a/Engine/Source/Runtime/Engine/Private/LightMap.cpp b/Engine/Source/Runtime/Engine/Private/LightMap.cpp
--- a/Engine/Source/Runtime/Engine/Private/LightMap.cpp
+++ b/Engine/Source/Runtime/Engine/Private/LightMap.cpp
@@ -270,7 +270,7 @@ struct FLightMapAllocation
 	{
 		if (InstanceIndex >= 0 && Registry)
 		{
-			FMeshMapBuildData* MeshBuildData = Registry->GetMeshBuildData(MapBuildDataId);
+			FMeshMapBuildData* MeshBuildData = Registry->GetMeshBuildDataDuringBuild(MapBuildDataId);
 			check(MeshBuildData);
 
 			UInstancedStaticMeshComponent* Component = CastChecked<UInstancedStaticMeshComponent>(Primitive);
diff --git a/Engine/Source/Runtime/Engine/Private/MapBuildData.cpp b/Engine/Source/Runtime/Engine/Private/MapBuildData.cpp
--- a/Engine/Source/Runtime/Engine/Private/MapBuildData.cpp
+++ b/Engine/Source/Runtime/Engine/Private/MapBuildData.cpp
@@ -492,6 +492,11 @@ FMeshMapBuildData* UMapBuildDataRegistry::GetMeshBuildData(FGuid MeshId)
 	return FoundData;
 }
 
+FMeshMapBuildData* UMapBuildDataRegistry::GetMeshBuildDataDuringBuild(FGuid MeshId)
+{
+	return MeshBuildData.Find(MeshId);
+}
+
 FPrecomputedLightVolumeData& UMapBuildDataRegistry::AllocateLevelPrecomputedLightVolumeBuildData(const FGuid& LevelId)
 {
 	check(LevelId.IsValid());
diff --git a/Engine/Source/Runtime/Engine/Private/ShadowMap.cpp b/Engine/Source/Runtime/Engine/Private/ShadowMap.cpp
--- a/Engine/Source/Runtime/Engine/Private/ShadowMap.cpp
+++ b/Engine/Source/Runtime/Engine/Private/ShadowMap.cpp
@@ -97,7 +97,7 @@ struct FShadowMapAllocation
 	{
 		if (InstanceIndex >= 0 && Registry)
 		{
-			FMeshMapBuildData* MeshBuildData = Registry->GetMeshBuildData(MapBuildDataId);
+			FMeshMapBuildData* MeshBuildData = Registry->GetMeshBuildDataDuringBuild(MapBuildDataId);
 			check(MeshBuildData);
 
 			UInstancedStaticMeshComponent* Component = CastChecked<UInstancedStaticMeshComponent>(Primitive);
