Fix a few minor comment inaccuracies
Integrates changes from PR #5526

diff --git a/Engine/Source/Runtime/Core/Public/Containers/UnrealString.h b/Engine/Source/Runtime/Core/Public/Containers/UnrealString.h
--- a/Engine/Source/Runtime/Core/Public/Containers/UnrealString.h
+++ b/Engine/Source/Runtime/Core/Public/Containers/UnrealString.h
@@ -1139,31 +1139,31 @@ public:
 		return Data.Num() ? Data.Num() - 1 : 0;
 	}
 
-	/** @return the left most given number of characters */
+	/** Returns the left most given number of characters */
 	FORCEINLINE FString Left( int32 Count ) const
 	{
 		return FString( FMath::Clamp(Count,0,Len()), **this );
 	}
 
-	/** @return the left most characters from the string chopping the given number of characters from the end */
+	/** Returns the left most characters from the string chopping the given number of characters from the end */
 	FORCEINLINE FString LeftChop( int32 Count ) const
 	{
 		return FString( FMath::Clamp(Len()-Count,0,Len()), **this );
 	}
 
-	/** @return the string to the right of the specified location, counting back from the right (end of the word). */
+	/** Returns the string to the right of the specified location, counting back from the right (end of the word). */
 	FORCEINLINE FString Right( int32 Count ) const
 	{
 		return FString( **this + Len()-FMath::Clamp(Count,0,Len()) );
 	}
 
-	/** @return the string to the right of the specified location, counting forward from the left (from the beginning of the word). */
+	/** Returns the string to the right of the specified location, counting forward from the left (from the beginning of the word). */
 	FORCEINLINE FString RightChop( int32 Count ) const
 	{
 		return FString( **this + Len()-FMath::Clamp(Len()-Count,0,Len()) );
 	}
 
-	/** @return the substring from Start position for Count characters. */
+	/** Returns the substring from Start position for Count characters. */
 	FORCEINLINE FString Mid( int32 Start, int32 Count=MAX_int32 ) const
 	{
 		check(Count >= 0);
@@ -1342,7 +1342,7 @@ public:
 		return true;
 	}
 
-	/** @return a new string with the characters of this converted to uppercase */
+	/** Returns a new string with the characters of this converted to uppercase */
 	FString ToUpper() const &;
 
 	/**
@@ -1354,7 +1354,7 @@ public:
 	/** Converts all characters in this string to uppercase */
 	void ToUpperInline();
 
-	/** @return a new string with the characters of this converted to lowercase */
+	/** Returns a new string with the characters of this converted to lowercase */
 	FString ToLower() const &;
 
 	/**
@@ -1372,7 +1372,7 @@ public:
 	/** Pad the right of this string for ChCount characters */
 	FString RightPad( int32 ChCount ) const;
 	
-	/** @return true if the string only contains numeric characters */
+	/** Returns true if the string only contains numeric characters */
 	bool IsNumeric() const;
 	
 	/** Removes spaces from the string.  I.E. "Spaces Are Cool" --> "SpacesAreCool". */
@@ -1427,7 +1427,7 @@ public:
 	 */
 	static FString Format(const TCHAR* InFormatString, const TArray<FStringFormatArg>& InOrderedArguments);
 
-	// @return string with Ch character
+	/** Returns a string containing only the Ch character */
 	static FString Chr( TCHAR Ch );
 
 	/**
@@ -1946,7 +1946,7 @@ inline int32 StringToBytes( const FString& String, uint8* OutBytes, int32 MaxBuf
 	return NumBytes;
 }
 
-/** @return Char value of Nibble */
+/** Returns Char value of Nibble */
 inline TCHAR NibbleToTChar(uint8 Num)
 {
 	if (Num > 9)
diff --git a/Engine/Source/Runtime/Core/Public/Misc/Paths.h b/Engine/Source/Runtime/Core/Public/Misc/Paths.h
--- a/Engine/Source/Runtime/Core/Public/Misc/Paths.h
+++ b/Engine/Source/Runtime/Core/Public/Misc/Paths.h
@@ -267,39 +267,25 @@ public:
 	UE_DEPRECATED(4.18, "FPaths::GameLogDir() has been superseded by FPaths::ProjectLogDir().")
 	static FORCEINLINE FString GameLogDir() { return ProjectLogDir(); }
 
-	/**
-	 * @return The directory for automation save files
-	 */
+	/** Returns the directory for automation save files */
 	static FString AutomationDir();
 
-	/**
-	 * @return The directory for automation save files that are meant to be deleted every run
-	 */
+	/** Returns the directory for automation save files that are meant to be deleted every run */
 	static FString AutomationTransientDir();
 
-	/**
-	* @return The directory for automation log files.
-	*/
+	/** Returns the directory for automation log files */
 	static FString AutomationLogDir();
 
-	/**
-	 * @return The directory for local files used in cloud emulation or support
-	 */
+	/** Returns the directory for local files used in cloud emulation or support */
 	static FString CloudDir();
 
-		/**
-	 * @return The directory that contains subfolders for developer-specific content
-	 */
+	/** Returns the directory that contains subfolders for developer-specific content */
 	static FString GameDevelopersDir();
 
-	/**
-	 * @return The directory that contains developer-specific content for the current user
-	 */
+	/** Returns The directory that contains developer-specific content for the current user */
 	static FString GameUserDeveloperDir();
 
-	/**
-	 * @return The directory for temp files used for diffing
-	 */
+	/** Returns the directory for temp files used for diffing */
 	static FString DiffDir();
 
 	/** 
@@ -350,19 +336,13 @@ public:
 	 */
 	static FString GameAgnosticSavedDir();
 
-	/**
-	 * @return The directory where engine source code files are kept
-	 */
+	/** Returns the directory where engine source code files are kept */
 	static FString EngineSourceDir();
 
-	/**
-	 * @return The directory where game source code files are kept
-	 */
+	/** Returns the directory where game source code files are kept */
 	static FString GameSourceDir();
 
-	/**
-	 * @return The directory where feature packs are kept
-	 */
+	/** Returns the directory where feature packs are kept */
 	static FString FeaturePackDir();
 
 	/**
@@ -416,16 +396,16 @@ public:
 	/** Sets the extension of the given filename (like ChangeExtension, but also applies the extension if the file doesn't have one) */
 	static FString SetExtension(const FString& InPath, const FString& InNewExtension);
 
-	/** @return true if this file was found, false otherwise */
+	/** Returns true if this file was found, false otherwise */
 	static bool FileExists(const FString& InPath);
 
-	/** @return true if this directory was found, false otherwise */
+	/** Returns true if this directory was found, false otherwise */
 	static bool DirectoryExists(const FString& InPath);
 
-	/** @return true if this path represents a drive */
+	/** Returns true if this path represents a root drive or volume */
 	static bool IsDrive(const FString& InPath);
 
-	/** @return true if this path is relative */
+	/** Returns true if this path is relative to another path */
 	static bool IsRelative(const FString& InPath);
 
 	/** Convert all / and \ to TEXT("/") */
@@ -477,10 +457,11 @@ public:
 	static void MakePlatformFilename(FString& InPath);
 
 	/** 
-	 * Assuming both paths (or filenames) are relative to the base dir, find the relative path to the InPath.
+	 * Assuming both paths (or filenames) are relative to the same base dir, modifies InPath to be relative to InRelativeTo
 	 *
-	 * @Param InPath Path to make this path relative to.
-	 * @return Path relative to InPath.
+	 * @param InPath Path to change to be relative to InRelativeTo
+	 * @param InRelativeTo Path to use as the new relative base
+	 * @returns true if InPath was changed to be relative
 	 */
 	static bool MakePathRelativeTo( FString& InPath, const TCHAR* InRelativeTo );
 
@@ -539,17 +520,12 @@ public:
 
 	/**
 	* Returns a string containing all invalid characters as dictated by the operating system
-	* 
-	* @return FString
 	*/
 	static const FString& GetInvalidFileSystemChars();
 
 	/**
 	*	Returns a string that is safe to use as a filename because all items in
 	*	GetInvalidFileSystemChars() are removed
-	* 
-	* @param  InPath
-	* @return FString
 	*/
 	static FString MakeValidFileName(const FString& InString, const TCHAR InReplacementChar=0);
 
diff --git a/Engine/Source/Runtime/Engine/Classes/Camera/PlayerCameraManager.h b/Engine/Source/Runtime/Engine/Classes/Camera/PlayerCameraManager.h
--- a/Engine/Source/Runtime/Engine/Classes/Camera/PlayerCameraManager.h
+++ b/Engine/Source/Runtime/Engine/Classes/Camera/PlayerCameraManager.h
@@ -522,10 +522,10 @@ public:
 
 	friend struct FTViewTarget;
 public:
-	/** @return the current ViewTarget. */
+	/** Returns the current ViewTarget. */
 	AActor* GetViewTarget() const;
 
-	/** @return the ViewTarget if it is an APawn, or nullptr otherwise */
+	/** Returns the ViewTarget if it is an APawn, or nullptr otherwise */
 	class APawn* GetViewTargetPawn() const;
 
 	//~ Begin AActor Interface
@@ -619,7 +619,7 @@ public:
 	 */
 	virtual void InitializeFor(class APlayerController* PC);
 	
-	/** @return Returns the camera's current full FOV angle, in degrees. */
+	/** Returns the camera's current full FOV angle, in degrees. */
 	UFUNCTION(BlueprintCallable, Category = "Camera")
 	virtual float GetFOVAngle() const;
 	
@@ -632,10 +632,10 @@ public:
 	/** Unlocks the FOV. */
 	virtual void UnlockFOV();
 
-	/** @return Returns true if this camera is using an orthographic perspective. */
+	/** Returns true if this camera is using an orthographic perspective. */
 	virtual bool IsOrthographic() const;
 
-	/** @return Returns the current orthographic width for the camera. */
+	/** Returns the current orthographic width for the camera. */
 	virtual float GetOrthoWidth() const;
 
 	/** 
@@ -657,11 +657,11 @@ public:
 	 */
 	void GetCameraViewPoint(FVector& OutCamLoc, FRotator& OutCamRot) const;
 	
-	/** @return Returns camera's current rotation. */
+	/** Returns camera's current rotation. */
 	UFUNCTION(BlueprintCallable, Category = "Camera", meta=(Keywords="View Direction"))
 	FRotator GetCameraRotation() const;
 
-	/** @return Returns camera's current location. */
+	/** Returns camera's current location. */
 	UFUNCTION(BlueprintCallable, Category = "Camera", meta=(Keywords="View Position"))
 	FVector GetCameraLocation() const;
 	
@@ -735,7 +735,7 @@ public:
 	// Camera Lens Effects
 	//
 	
-	/** @return Returns first instance of a lens effect of the given class. */
+	/** Returns first instance of a lens effect of the given class. */
 	virtual class AEmitterCameraLensEffectBase* FindCameraLensEffect(TSubclassOf<class AEmitterCameraLensEffectBase> LensEffectEmitterClass);
 	
 	/** 
@@ -856,7 +856,7 @@ public:
 	UFUNCTION(BlueprintCallable, Category = "Camera Animation")
 	virtual void StopAllCameraAnims(bool bImmediate = false);
 
-	/** @return Returns first existing instance of the specified camera anim, or NULL if none exists. */
+	/** Returns first existing instance of the specified camera anim, or NULL if none exists. */
 	class UCameraAnimInst* FindInstanceOfCameraAnim(class UCameraAnim const* Anim) const;
 
 protected:
@@ -864,7 +864,7 @@ protected:
 	void InitTempCameraActor(class ACameraActor* CamActor, class UCameraAnimInst const* AnimInstToInitFor) const;
 	void ApplyAnimToCamera(class ACameraActor const* AnimatedCamActor, class UCameraAnimInst const* AnimInst, FMinimalViewInfo& InOutPOV);
 
-	/** @return Returns an available CameraAnimInst, or NULL if no more are available. */
+	/** Returns an available CameraAnimInst, or NULL if no more are available. */
 	class UCameraAnimInst* AllocCameraAnimInst();
 
 	/** 
@@ -917,6 +917,6 @@ private:
 	FVector GetActorLocation() const { return Super::GetActorLocation(); }
 
 public:
-	/** @return Returns TransformComponent subobject */
+	/** Returns TransformComponent subobject */
 	class USceneComponent* GetTransformComponent() const { return TransformComponent; }
 };
diff --git a/Engine/Source/Runtime/Engine/Classes/Components/AudioComponent.h b/Engine/Source/Runtime/Engine/Classes/Components/AudioComponent.h
--- a/Engine/Source/Runtime/Engine/Classes/Components/AudioComponent.h
+++ b/Engine/Source/Runtime/Engine/Classes/Components/AudioComponent.h
@@ -381,7 +381,7 @@ public:
 	UFUNCTION(BlueprintCallable, Category = "Audio|Components|Audio")
 	void SetPaused(bool bPause);
 
-	/** @return true if this component is currently playing a SoundCue. */
+	/** Returns true if this component is currently playing a SoundCue. */
 	UFUNCTION(BlueprintCallable, Category="Audio|Components|Audio")
 	virtual bool IsPlaying() const;
 
diff --git a/Engine/Source/Runtime/Engine/Classes/GameFramework/CharacterMovementComponent.h b/Engine/Source/Runtime/Engine/Classes/GameFramework/CharacterMovementComponent.h
--- a/Engine/Source/Runtime/Engine/Classes/GameFramework/CharacterMovementComponent.h
+++ b/Engine/Source/Runtime/Engine/Classes/GameFramework/CharacterMovementComponent.h
@@ -1196,12 +1196,12 @@ public:
 	virtual void AddRadialImpulse(const FVector& Origin, float Radius, float Strength, enum ERadialImpulseFalloff Falloff, bool bVelChange) override;
 	//END UMovementComponent Interface
 
-	/** @return true if the character is in the 'Walking' movement mode. */
+	/** Returns true if the character is in the 'Walking' movement mode. */
 	UFUNCTION(BlueprintCallable, Category="Pawn|Components|CharacterMovement")
 	bool IsWalking() const;
 
 	/**
-	 * @return true if currently performing a movement update.
+	 * Returns true if currently performing a movement update.
 	 * @see bMovementInProgress
 	 */
 	bool IsMovementInProgress() const { return bMovementInProgress; }
@@ -1314,7 +1314,7 @@ public:
 	/** Jump onto shore from water */
 	virtual void JumpOutOfWater(FVector WallNormal);
 
-	/** @return how far to rotate character during the time interval DeltaTime. */
+	/** Returns how far to rotate character during the time interval DeltaTime. */
 	virtual FRotator GetDeltaRotation(float DeltaTime) const;
 
 	/**
@@ -1384,35 +1384,34 @@ public:
 	UFUNCTION(BlueprintCallable, Category="Pawn|Components|CharacterMovement")
 	virtual float GetMaxJumpHeightWithJumpTime() const;
 
-	/** @return Maximum acceleration for the current state. */
+	/** Returns maximum acceleration for the current state. */
 	UFUNCTION(BlueprintCallable, Category = "Pawn|Components|CharacterMovement")
 	virtual float GetMinAnalogSpeed() const;
 	
-	/** @return Maximum acceleration for the current state, based on MaxAcceleration and any additional modifiers. */
 	UE_DEPRECATED(4.3, "GetModifiedMaxAcceleration() is deprecated, apply your own modifiers to GetMaxAcceleration() if desired.")
 	virtual float GetModifiedMaxAcceleration() const;
 	
-	/** @return Maximum acceleration for the current state, based on MaxAcceleration and any additional modifiers. */
+	/** Returns maximum acceleration for the current state, based on MaxAcceleration and any additional modifiers. */
 	UFUNCTION(BlueprintCallable, Category="Pawn|Components|CharacterMovement", meta=(DeprecatedFunction, DisplayName="GetModifiedMaxAcceleration", ScriptName="GetModifiedMaxAcceleration", DeprecationMessage="GetModifiedMaxAcceleration() is deprecated, apply your own modifiers to GetMaxAcceleration() if desired."))
 	virtual float K2_GetModifiedMaxAcceleration() const;
 
-	/** @return Maximum acceleration for the current state. */
+	/** Returns maximum acceleration for the current state. */
 	UFUNCTION(BlueprintCallable, Category="Pawn|Components|CharacterMovement")
 	virtual float GetMaxAcceleration() const;
 
-	/** @return Maximum deceleration for the current state when braking (ie when there is no acceleration). */
+	/** Returns maximum deceleration for the current state when braking (ie when there is no acceleration). */
 	UFUNCTION(BlueprintCallable, Category="Pawn|Components|CharacterMovement")
 	virtual float GetMaxBrakingDeceleration() const;
 
-	/** @return Current acceleration, computed from input vector each update. */
+	/** Returns current acceleration, computed from input vector each update. */
 	UFUNCTION(BlueprintCallable, Category="Pawn|Components|CharacterMovement", meta=(Keywords="Acceleration GetAcceleration"))
 	FVector GetCurrentAcceleration() const;
 
-	/** @return Modifier [0..1] based on the magnitude of the last input vector, which is used to modify the acceleration and max speed during movement. */
+	/** Returns modifier [0..1] based on the magnitude of the last input vector, which is used to modify the acceleration and max speed during movement. */
 	UFUNCTION(BlueprintCallable, Category="Pawn|Components|CharacterMovement")
 	float GetAnalogInputModifier() const;
 	
-	/** @return true if we can step up on the actor in the given FHitResult. */
+	/** Returns true if we can step up on the actor in the given FHitResult. */
 	virtual bool CanStepUp(const FHitResult& Hit) const;
 
 	/** Struct updated by StepUp() to return result of final step down, if applicable. */
@@ -1589,10 +1588,10 @@ public:
 	 */
 	virtual void UnCrouch(bool bClientSimulation = false);
 
-	/** @return true if the character is allowed to crouch in the current state. By default it is allowed when walking or falling, if CanEverCrouch() is true. */
+	/** Returns true if the character is allowed to crouch in the current state. By default it is allowed when walking or falling, if CanEverCrouch() is true. */
 	virtual bool CanCrouchInCurrentState() const;
 	
-	/** @return true if there is a suitable floor SideStep from current position. */
+	/** Returns true if there is a suitable floor SideStep from current position. */
 	virtual bool CheckLedgeDirection(const FVector& OldLocation, const FVector& SideStep, const FVector& GravDir) const;
 
 	/** 
@@ -1640,7 +1639,7 @@ public:
 	
 	virtual void SetUpdatedComponent(USceneComponent* NewUpdatedComponent) override;
 	
-	/** @Return MovementMode string */
+	/** Returns MovementMode string */
 	virtual FString GetMovementName() const;
 
 	/** 
@@ -1687,10 +1686,10 @@ public:
 	/** Check if swimming pawn just ran into edge of the pool and should jump out. */
 	virtual bool CheckWaterJump(FVector CheckPoint, FVector& WallNormal);
 
-	/** @return whether this pawn is currently allowed to walk off ledges */
+	/** Returns whether this pawn is currently allowed to walk off ledges */
 	virtual bool CanWalkOffLedges() const;
 
-	/** @return The distance from the edge of the capsule within which we don't allow the character to perch on the edge of a surface. */
+	/** Returns The distance from the edge of the capsule within which we don't allow the character to perch on the edge of a surface. */
 	UFUNCTION(BlueprintCallable, Category="Pawn|Components|CharacterMovement")
 	float GetPerchRadiusThreshold() const;
 
@@ -2299,7 +2298,7 @@ public:
 	UPROPERTY(Transient)
 	FRootMotionSourceGroup CurrentRootMotion;
 
-	/** @return true if we have Root Motion from any source to use in PerformMovement() physics. */
+	/** Returns true if we have Root Motion from any source to use in PerformMovement() physics. */
 	bool HasRootMotionSources() const;
 
 	/** Apply a RootMotionSource to current root motion 
@@ -2350,7 +2349,7 @@ public:
 	UPROPERTY(Transient)
 	FVector AnimRootMotionVelocity;
 
-	/** @return true if we have Root Motion from animation to use in PerformMovement() physics. 
+	/** Returns true if we have Root Motion from animation to use in PerformMovement() physics. 
 		Not valid outside of the scope of that function. Since RootMotion is extracted and used in it. */
 	bool HasAnimRootMotion() const
 	{
@@ -2573,7 +2572,7 @@ public:
 	/** Set the properties describing the position, etc. of the moved pawn at the start of the move. */
 	virtual void SetInitialPosition(ACharacter* C);
 
-	/** @Return true if this move is an "important" move that should be sent again if not acked by the server */
+	/** Returns true if this move is an "important" move that should be sent again if not acked by the server */
 	virtual bool IsImportantMove(const FSavedMovePtr& LastAckedMove) const;
 	
 	/** Returns starting position if we were to revert the move, either absolute StartLocation, or StartRelativeLocation offset from MovementBase's current location (since we want to try to move forward at this time). */
@@ -2588,7 +2587,7 @@ public:
 	/** Set the properties describing the final position, etc. of the moved pawn. */
 	virtual void PostUpdate(ACharacter* C, EPostUpdateMode PostUpdateMode);
 	
-	/** @Return true if this move can be combined with NewMove for replication without changing any behavior */
+	/** Returns true if this move can be combined with NewMove for replication without changing any behavior */
 	virtual bool CanCombineWith(const FSavedMovePtr& NewMove, ACharacter* InCharacter, float MaxDelta) const;
 
 	/** Combine this move with an older move and update relevant state. */
@@ -2597,7 +2596,7 @@ public:
 	/** Called before ClientUpdatePosition uses this SavedMove to make a predictive correction	 */
 	virtual void PrepMoveFor(ACharacter* C);
 
-	/** @returns a byte containing encoded special movement information (jumping, crouching, etc.)	 */
+	/** Returns a byte containing encoded special movement information (jumping, crouching, etc.)	 */
 	virtual uint8 GetCompressedFlags() const;
 
 	// Bit masks used by GetCompressedFlags() to encode movement information.
@@ -2875,14 +2874,10 @@ public:
 	/** Creation time of this prediction data, used to contextualize LifetimeRawTimeDiscrepancy */
 	float WorldCreationTime;
 
-	/** 
-	 * @return Time delta to use for the current ServerMove(). Takes into account time discrepancy resolution if active.
-	 */
+	/** Returns time delta to use for the current ServerMove(). Takes into account time discrepancy resolution if active. */
 	float GetServerMoveDeltaTime(float ClientTimeStamp, float ActorTimeDilation) const;
 
-	/** 
-	 * @return Base time delta to use for a ServerMove, default calculation (no time discrepancy resolution)
-	 */
+	/** Returns base time delta to use for a ServerMove, default calculation (no time discrepancy resolution) */
 	float GetBaseServerMoveDeltaTime(float ClientTimeStamp, float ActorTimeDilation) const;
 
 };
diff --git a/Engine/Source/Runtime/Engine/Classes/GameFramework/GameMode.h b/Engine/Source/Runtime/Engine/Classes/GameFramework/GameMode.h
--- a/Engine/Source/Runtime/Engine/Classes/GameFramework/GameMode.h
+++ b/Engine/Source/Runtime/Engine/Classes/GameFramework/GameMode.h
@@ -87,14 +87,14 @@ protected:
 	/** Called when the state transitions to WaitingToStart */
 	virtual void HandleMatchIsWaitingToStart();
 
-	/** @return True if ready to Start Match. Games should override this */
+	/** Returns true if ready to Start Match. Games should override this */
 	UFUNCTION(BlueprintNativeEvent, Category="Game")
 	bool ReadyToStartMatch();
 
 	/** Called when the state transitions to InProgress */
 	virtual void HandleMatchHasStarted();
 
-	/** @return true if ready to End Match. Games should override this */
+	/** Returns true if ready to End Match. Games should override this */
 	UFUNCTION(BlueprintNativeEvent, Category="Game")
 	bool ReadyToEndMatch();
 
diff --git a/Engine/Source/Runtime/Engine/Classes/GameFramework/GameUserSettings.h b/Engine/Source/Runtime/Engine/Classes/GameFramework/GameUserSettings.h
--- a/Engine/Source/Runtime/Engine/Classes/GameFramework/GameUserSettings.h
+++ b/Engine/Source/Runtime/Engine/Classes/GameFramework/GameUserSettings.h
@@ -278,15 +278,15 @@ public:
 	/** Loads the resolution settings before is object is available */
 	static void PreloadResolutionSettings();
 
-	/** @return The default resolution when no resolution is set */
+	/** Returns the default resolution when no resolution is set */
 	UFUNCTION(BlueprintCallable, Category=Settings)
 	static FIntPoint GetDefaultResolution();
 
-	/** @return The default window position when no position is set */
+	/** Returns the default window position when no position is set */
 	UFUNCTION(BlueprintCallable, Category=Settings)
 	static FIntPoint GetDefaultWindowPosition();
 
-	/** @return The default window mode when no mode is set */
+	/** Returns the default window mode when no mode is set */
 	UFUNCTION(BlueprintCallable, Category=Settings)
 	static EWindowMode::Type GetDefaultWindowMode();
 
diff --git a/Engine/Source/Runtime/Engine/Classes/GameFramework/MovementComponent.h b/Engine/Source/Runtime/Engine/Classes/GameFramework/MovementComponent.h
--- a/Engine/Source/Runtime/Engine/Classes/GameFramework/MovementComponent.h
+++ b/Engine/Source/Runtime/Engine/Classes/GameFramework/MovementComponent.h
@@ -200,27 +200,25 @@ public:
 
 	//~ End ActorComponent Interface
 
-	/** @return gravity that affects this component */
+	/** Returns gravity that affects this component */
 	UFUNCTION(BlueprintCallable, Category="Components|Movement")
 	virtual float GetGravityZ() const;
 
-	/** @return Maximum speed of component in current movement mode. */
+	/** Returns maximum speed of component in current movement mode. */
 	UFUNCTION(BlueprintCallable, Category="Components|Movement")
 	virtual float GetMaxSpeed() const;
 
-	/** @return a scalar applied to the maximum velocity that the component can currently move. */
 	UE_DEPRECATED(4.3, "GetMaxSpeedModifier() is deprecated, apply your own modifiers to GetMaxSpeed() if desired.")
 	virtual float GetMaxSpeedModifier() const;
 	
-	/** @return a scalar applied to the maximum velocity that the component can currently move. */
+	/** Returns a scalar applied to the maximum velocity that the component can currently move. */
 	UFUNCTION(BlueprintCallable, Category="Components|Movement", meta=(DeprecatedFunction, DisplayName="GetMaxSpeedModifier", ScriptName="GetMaxSpeedModifier", DeprecationMessage="GetMaxSpeedModifier() is deprecated, apply your own modifiers to GetMaxSpeed() if desired."))
 	virtual float K2_GetMaxSpeedModifier() const;
 
-	/** @return the result of GetMaxSpeed() * GetMaxSpeedModifier(). */
 	UE_DEPRECATED(4.3, "GetModifiedMaxSpeed() is deprecated, use GetMaxSpeed() instead.")
 	virtual float GetModifiedMaxSpeed() const;
 
-	/** @return the result of GetMaxSpeed() * GetMaxSpeedModifier(). */
+	/** Returns the result of GetMaxSpeed() * GetMaxSpeedModifier(). */
 	UFUNCTION(BlueprintCallable, Category="Components|Movement", meta=(DisplayName="GetModifiedMaxSpeed", ScriptName="GetModifiedMaxSpeed", DeprecationMessage="GetModifiedMaxSpeed() is deprecated, apply your own modifiers to GetMaxSpeed() if desired."))
 	virtual float K2_GetModifiedMaxSpeed() const;
 
@@ -243,7 +241,7 @@ public:
 	 */
 	virtual bool ShouldSkipUpdate(float DeltaTime) const;
 
-	/** @return PhysicsVolume this MovementComponent is using, or the world's default physics volume if none. **/
+	/** Returns the PhysicsVolume this MovementComponent is using, or the world's default physics volume if none. **/
 	UFUNCTION(BlueprintCallable, Category="Components|Movement")
 	virtual APhysicsVolume* GetPhysicsVolume() const;
 
@@ -255,7 +253,7 @@ public:
 	UFUNCTION(BlueprintCallable, Category="Components|Movement")
 	virtual void SetUpdatedComponent(USceneComponent* NewUpdatedComponent);
 
-	/** return true if it's in PhysicsVolume with water flag **/
+	/** Returns true if it's in PhysicsVolume with water flag **/
 	virtual bool IsInWater() const;
 
 	/** Update tick registration state, determined by bAutoUpdateTickRegistration. Called by SetUpdatedComponent. */
@@ -425,7 +423,7 @@ public:
 	UFUNCTION(BlueprintCallable, Category = "Components|Movement|Planar")
 	virtual void SetPlaneConstraintEnabled(bool bEnabled);
 
-	/** @return The normal of the plane that constrains movement, enforced if the plane constraint is enabled. */
+	/** Returns the normal of the plane that constrains movement, enforced if the plane constraint is enabled. */
 	UFUNCTION(BlueprintCallable, Category="Components|Movement|Planar")
 	const FVector& GetPlaneConstraintNormal() const;
 
diff --git a/Engine/Source/Runtime/Engine/Classes/GameFramework/NavMovementComponent.h b/Engine/Source/Runtime/Engine/Classes/GameFramework/NavMovementComponent.h
--- a/Engine/Source/Runtime/Engine/Classes/GameFramework/NavMovementComponent.h
+++ b/Engine/Source/Runtime/Engine/Classes/GameFramework/NavMovementComponent.h
@@ -79,13 +79,13 @@ public:
 	void UpdateNavAgent(const AActor& InOwner);
 	void UpdateNavAgent(const UCapsuleComponent& CapsuleComponent);
 
-	/** @returns location of controlled actor - meaning center of collision bounding box */
+	/** Returnss location of controlled actor - meaning center of collision bounding box */
 	FORCEINLINE FVector GetActorLocation() const { return UpdatedComponent ? UpdatedComponent->GetComponentLocation() : FVector(FLT_MAX); }
-	/** @returns location of controlled actor's "feet" meaning center of bottom of collision bounding box */
+	/** Returnss location of controlled actor's "feet" meaning center of bottom of collision bounding box */
 	FORCEINLINE FVector GetActorFeetLocation() const { return UpdatedComponent ? (UpdatedComponent->GetComponentLocation() - FVector(0,0,UpdatedComponent->Bounds.BoxExtent.Z)) : FNavigationSystem::InvalidLocation; }
-	/** @returns based location of controlled actor */
+	/** Returnss based location of controlled actor */
 	virtual FBasedPosition GetActorFeetLocationBased() const;
-	/** @returns navigation location of controlled actor */
+	/** Returnss navigation location of controlled actor */
 	FORCEINLINE FVector GetActorNavLocation() const { INavAgentInterface* MyOwner = Cast<INavAgentInterface>(GetOwner()); return MyOwner ? MyOwner->GetNavAgentLocation() : FNavigationSystem::InvalidLocation; }
 
 	/** path following: request new velocity */
@@ -98,7 +98,7 @@ public:
 	 *  (e.g. performing scripted move and can't stop) */
 	virtual bool CanStopPathFollowing() const;
 
-	/** @returns braking distance for acceleration driven path following */
+	/** Returnss braking distance for acceleration driven path following */
 	virtual float GetPathFollowingBrakingDistance(float MaxSpeed) const;
 
 	void SetPathFollowingAgent(IPathFollowingAgentInterface* InPathFollowingAgent) { PathFollowingComp = Cast<UObject>(InPathFollowingAgent); }
@@ -113,56 +113,56 @@ public:
 
 	FORCEINLINE bool UseAccelerationForPathFollowing() const { return bUseAccelerationForPaths; }
 
-	/** @returns the NavAgentProps(const) */
+	/** Returnss the NavAgentProps(const) */
 	FORCEINLINE const FNavAgentProperties& GetNavAgentPropertiesRef() const { return NavAgentProps; }
-	/** @returns the NavAgentProps */
+	/** Returnss the NavAgentProps */
 	FORCEINLINE FNavAgentProperties& GetNavAgentPropertiesRef() { return NavAgentProps; }
 
 	/** Resets runtime movement state to character's movement capabilities */
 	void ResetMoveState() { MovementState = NavAgentProps; }
 
-	/** @return true if path following can start */
+	/** Returns true if path following can start */
 	virtual bool CanStartPathFollowing() const { return true; }
 
-	/** @return true if component can crouch */
+	/** Returns true if component can crouch */
 	FORCEINLINE bool CanEverCrouch() const { return NavAgentProps.bCanCrouch; }
 
-	/** @return true if component can jump */
+	/** Returns true if component can jump */
 	FORCEINLINE bool CanEverJump() const { return NavAgentProps.bCanJump; }
 
-	/** @return true if component can move along the ground (walk, drive, etc) */
+	/** Returns true if component can move along the ground (walk, drive, etc) */
 	FORCEINLINE bool CanEverMoveOnGround() const { return NavAgentProps.bCanWalk; }
 
-	/** @return true if component can swim */
+	/** Returns true if component can swim */
 	FORCEINLINE bool CanEverSwim() const { return NavAgentProps.bCanSwim; }
 
-	/** @return true if component can fly */
+	/** Returns true if component can fly */
 	FORCEINLINE bool CanEverFly() const { return NavAgentProps.bCanFly; }
 
-	/** @return true if component is allowed to jump */
+	/** Returns true if component is allowed to jump */
 	FORCEINLINE bool IsJumpAllowed() const { return CanEverJump() && MovementState.bCanJump; }
 
-	/** @param bAllowed true if component is allowed to jump */
+	/** Sets rather this component is allowed to jump */
 	FORCEINLINE void SetJumpAllowed(bool bAllowed) { MovementState.bCanJump = bAllowed; }
 
 
-	/** @return true if currently crouching */ 
+	/** Returns true if currently crouching */ 
 	UFUNCTION(BlueprintCallable, Category="AI|Components|NavMovement")
 	virtual bool IsCrouching() const;
 
-	/** @return true if currently falling (not flying, in a non-fluid volume, and not on the ground) */ 
+	/** Returns true if currently falling (not flying, in a non-fluid volume, and not on the ground) */ 
 	UFUNCTION(BlueprintCallable, Category="AI|Components|NavMovement")
 	virtual bool IsFalling() const;
 
-	/** @return true if currently moving on the ground (e.g. walking or driving) */ 
+	/** Returns true if currently moving on the ground (e.g. walking or driving) */ 
 	UFUNCTION(BlueprintCallable, Category="AI|Components|NavMovement")
 	virtual bool IsMovingOnGround() const;
 	
-	/** @return true if currently swimming (moving through a fluid volume) */
+	/** Returns true if currently swimming (moving through a fluid volume) */
 	UFUNCTION(BlueprintCallable, Category="AI|Components|NavMovement")
 	virtual bool IsSwimming() const;
 
-	/** @return true if currently flying (moving through a non-fluid volume without resting on the ground) */
+	/** Returns true if currently flying (moving through a non-fluid volume without resting on the ground) */
 	UFUNCTION(BlueprintCallable, Category="AI|Components|NavMovement")
 	virtual bool IsFlying() const;
 };
diff --git a/Engine/Source/Runtime/Engine/Classes/GameFramework/Pawn.h b/Engine/Source/Runtime/Engine/Classes/GameFramework/Pawn.h
--- a/Engine/Source/Runtime/Engine/Classes/GameFramework/Pawn.h
+++ b/Engine/Source/Runtime/Engine/Classes/GameFramework/Pawn.h
@@ -189,7 +189,7 @@ public:
 	/** Return true if yaw is within AllowedYawError of desired yaw */
 	virtual bool ReachedDesiredRotation();
 
-	/** @return The half-height of the default Pawn, scaled by the component scale. By default returns the half-height of the RootComponent, regardless of whether it is registered or collidable. */
+	/** Returns The half-height of the default Pawn, scaled by the component scale. By default returns the half-height of the RootComponent, regardless of whether it is registered or collidable. */
 	virtual float GetDefaultHalfHeight() const;
 
 	/** See if this actor is currently being controlled */
@@ -282,7 +282,7 @@ public:
 	virtual void PostEditChangeProperty(struct FPropertyChangedEvent& PropertyChangedEvent) override;
 #endif
 
-	/** @return vector direction of gravity */
+	/** Returns vector direction of gravity */
 	FVector GetGravityDirection();
 
 	/** Make sure pawn properties are back to default. */
@@ -296,7 +296,7 @@ public:
 
 	/** 
 	 * Called when this Pawn is possessed. Only called on the server (or in standalone).
-	 *@param C The controller possessing this pawn
+	 * @param NewController The controller possessing this pawn
 	 */
 	virtual void PossessedBy(AController* NewController);
 
@@ -308,15 +308,15 @@ public:
 	UFUNCTION(BlueprintImplementableEvent, meta=(DisplayName= "Unpossessed"))
 	void ReceiveUnpossessed(AController* OldController);
 
-	/** @return true if controlled by a local (not network) Controller.	 */
+	/** Returns true if controlled by a local (not network) Controller.	 */
 	UFUNCTION(BlueprintPure, Category=Pawn)
 	virtual bool IsLocallyControlled() const;
   
-	/** @return true if controlled by a human player (possessed by a PlayerController).	 */
+	/** Returns true if controlled by a human player (possessed by a PlayerController).	 */
 	UFUNCTION(BlueprintPure, Category=Pawn)
 	virtual bool IsPlayerControlled() const;
 
-	/** @return true if controlled by a bot.	 */
+	/** Returns true if controlled by a bot.	 */
 	UFUNCTION(BlueprintPure, Category = Pawn)
 	virtual bool IsBotControlled() const;
 
@@ -326,7 +326,7 @@ public:
 	 */
 	virtual FRotator GetViewRotation() const;
 
-	/** @return	Pawn's eye location */
+	/** Returns	Pawn's eye location */
 	virtual FVector GetPawnViewLocation() const;
 
 	/**
@@ -337,7 +337,7 @@ public:
 	UFUNCTION(BlueprintCallable, Category=Pawn)
 	virtual FRotator GetBaseAimRotation() const;
 
-	/** @return true if player is viewing this Pawn in FreeCam */
+	/** Returns true if player is viewing this Pawn in FreeCam */
 	virtual bool InFreeCam() const;
 
 	/** Tell client that the Pawn is begin restarted. Calls Restart(). */
diff --git a/Engine/Source/Runtime/Engine/Classes/GameFramework/PlayerController.h b/Engine/Source/Runtime/Engine/Classes/GameFramework/PlayerController.h
--- a/Engine/Source/Runtime/Engine/Classes/GameFramework/PlayerController.h
+++ b/Engine/Source/Runtime/Engine/Classes/GameFramework/PlayerController.h
@@ -1699,7 +1699,7 @@ public:
 	/** Gives the PlayerController an opportunity to cleanup any changes it applied to the game viewport, primarily for the touch interface */
 	virtual void CleanupGameViewport();
 
-	/** @Returns true if input should be frozen (whether UnFreeze timer is active) */
+	/** Returns true if input should be frozen (whether UnFreeze timer is active) */
 	virtual void AcknowledgePossession(class APawn* P);
 
 	/** Clean up when a Pawn's player is leaving a game. Base implementation destroys the pawn. */
@@ -1720,7 +1720,7 @@ public:
 	 */
 	virtual void ViewAPlayer(int32 dir);
 
-	/** @return true if this controller thinks it's able to restart. Called from GameModeBase::PlayerCanRestart */
+	/** Returns true if this controller thinks it's able to restart. Called from GameModeBase::PlayerCanRestart */
 	UFUNCTION(BlueprintCallable, Category = "Game")
 	virtual bool CanRestartPlayer();
 
@@ -1787,12 +1787,12 @@ public:
 	 */
 	virtual bool DefaultCanUnpause();
 
-	/** @return true if game is currently paused. */
+	/** Returns true if game is currently paused. */
 	bool IsPaused() const;
 
 	bool InputEnabled() const { return bInputEnabled; }
 
-	/** @return true if we fully tick when paused (and if our tick function is enabled when paused).	 */
+	/** Returns true if we fully tick when paused (and if our tick function is enabled when paused).	 */
 	bool ShouldPerformFullTickWhenPaused() const;
 
 	/** returns whether the client has completely loaded the server's current world (valid on server only) */
@@ -1897,7 +1897,7 @@ public:
 	 */
 	virtual void InitInputSystem();
 
-	/** @returns true if input should be frozen (whether UnFreeze timer is active) */
+	/** Returns true if input should be frozen (whether UnFreeze timer is active) */
 	virtual bool IsFrozen();
 
 	/**
diff --git a/Engine/Source/Runtime/Engine/Classes/Kismet/BlueprintPathsLibrary.h b/Engine/Source/Runtime/Engine/Classes/Kismet/BlueprintPathsLibrary.h
--- a/Engine/Source/Runtime/Engine/Classes/Kismet/BlueprintPathsLibrary.h
+++ b/Engine/Source/Runtime/Engine/Classes/Kismet/BlueprintPathsLibrary.h
@@ -283,45 +283,31 @@ class ENGINE_API UBlueprintPathsLibrary : public UBlueprintFunctionLibrary
 	UFUNCTION(BlueprintPure, Category = "Utilities|Paths")
 	static FString ProjectLogDir();
 
-	/**
-	* @return The directory for automation save files
-	*/
+	/** Returns the directory for automation save files */
 	UFUNCTION(BlueprintPure, Category = "Utilities|Paths")
 	static FString AutomationDir();
 
-	/**
-	* @return The directory for automation save files that are meant to be deleted every run
-	*/
+	/** Returns the directory for automation save files that are meant to be deleted every run */
 	UFUNCTION(BlueprintPure, Category = "Utilities|Paths")
 	static FString AutomationTransientDir();
 
-	/**
-	* @return The directory for automation log files.
-	*/
+	/** Returns the directory for automation log files */
 	UFUNCTION(BlueprintPure, Category = "Utilities|Paths")
 	static FString AutomationLogDir();
 
-	/**
-	* @return The directory for local files used in cloud emulation or support
-	*/
+	/** Returns the directory for local files used in cloud emulation or support */
 	UFUNCTION(BlueprintPure, Category = "Utilities|Paths")
 	static FString CloudDir();
 
-	/**
-	* @return The directory that contains subfolders for developer-specific content
-	*/
+	/** Returns the directory that contains subfolders for developer-specific content */
 	UFUNCTION(BlueprintPure, Category = "Utilities|Paths")
 	static FString GameDevelopersDir();
 
-	/**
-	* @return The directory that contains developer-specific content for the current user
-	*/
+	/** Returns the directory that contains developer-specific content for the current user */
 	UFUNCTION(BlueprintPure, Category = "Utilities|Paths")
 	static FString GameUserDeveloperDir();
 
-	/**
-	* @return The directory for temp files used for diffing
-	*/
+	/** Returns the directory for temp files used for diffing */
 	UFUNCTION(BlueprintPure, Category = "Utilities|Paths")
 	static FString DiffDir();
 
@@ -376,21 +362,15 @@ class ENGINE_API UBlueprintPathsLibrary : public UBlueprintFunctionLibrary
 	UFUNCTION(BlueprintPure, Category = "Utilities|Paths")
 	static FString GameAgnosticSavedDir();
 
-	/**
-	* @return The directory where engine source code files are kept
-	*/
+	/** Returns the directory where engine source code files are kept */
 	UFUNCTION(BlueprintPure, Category = "Utilities|Paths")
 	static FString EngineSourceDir();
 
-	/**
-	* @return The directory where game source code files are kept
-	*/
+	/** Returns the directory where game source code files are kept */
 	UFUNCTION(BlueprintPure, Category = "Utilities|Paths")
 	static FString GameSourceDir();
 
-	/**
-	* @return The directory where feature packs are kept
-	*/
+	/** Returns the directory where feature packs are kept */
 	UFUNCTION(BlueprintPure, Category = "Utilities|Paths")
 	static FString FeaturePackDir();
 
@@ -448,19 +428,19 @@ class ENGINE_API UBlueprintPathsLibrary : public UBlueprintFunctionLibrary
 	UFUNCTION(BlueprintPure, Category = "Utilities|Paths")
 	static FString SetExtension(const FString& InPath, const FString& InNewExtension);
 
-	/** @return true if this file was found, false otherwise */
+	/** Returns true if this file was found, false otherwise */
 	UFUNCTION(BlueprintPure, Category = "Utilities|Paths")
 	static bool FileExists(const FString& InPath);
 
-	/** @return true if this directory was found, false otherwise */
+	/** Returns true if this directory was found, false otherwise */
 	UFUNCTION(BlueprintPure, Category = "Utilities|Paths")
 	static bool DirectoryExists(const FString& InPath);
 
-	/** @return true if this path represents a drive */
+	/** Returns true if this path represents a root drive or volume  */
 	UFUNCTION(BlueprintPure, Category = "Utilities|Paths")
 	static bool IsDrive(const FString& InPath);
 
-	/** @return true if this path is relative */
+	/** Returns true if this path is relative to another path */
 	UFUNCTION(BlueprintPure, Category = "Utilities|Paths")
 	static bool IsRelative(const FString& InPath);
 
@@ -519,12 +499,15 @@ class ENGINE_API UBlueprintPathsLibrary : public UBlueprintFunctionLibrary
 	UFUNCTION(BlueprintPure, Category = "Utilities|Paths")
 	static void MakePlatformFilename(const FString& InPath, FString& OutPath);
 
+
 	/**
-	* Assuming both paths (or filenames) are relative to the base dir, find the relative path to the InPath.
-	*
-	* @Param InPath Path to make this path relative to.
-	* @return Path relative to InPath.
-	*/
+	 * Assuming both paths (or filenames) are relative to the same base dir, converts InPath to be relative to InRelativeTo
+	 *
+	 * @param InPath Path to change to be relative to InRelativeTo
+	 * @param InRelativeTo Path to use as the new relative base
+	 * @param InPath New path relative to InRelativeTo
+	 * @returns true if OutPath was changed to be relative
+	 */
 	UFUNCTION(BlueprintPure, Category = "Utilities|Paths")
 	static bool MakePathRelativeTo(const FString& InPath, const FString& InRelativeTo, FString& OutPath);
 
@@ -561,11 +544,7 @@ class ENGINE_API UBlueprintPathsLibrary : public UBlueprintFunctionLibrary
 	UFUNCTION(BlueprintPure, Category = "Utilities|Paths")
 	static FString CreateTempFilename(const FString& Path, const FString& Prefix = TEXT(""), const FString& Extension = TEXT(".tmp"));
 
-	/**
-	* Returns a string containing all invalid characters as dictated by the operating system
-	*
-	* @return FString
-	*/
+	/** Returns a string containing all invalid characters as dictated by the operating system */
 	UFUNCTION(BlueprintPure, Category = "Utilities|Paths")
 	static FString GetInvalidFileSystemChars();
 
@@ -577,7 +556,6 @@ class ENGINE_API UBlueprintPathsLibrary : public UBlueprintFunctionLibrary
 	*
 	* @param  InString
 	* @param  InReplacementChar
-	* @return FString
 	*/
 	UFUNCTION(BlueprintPure, Category = "Utilities|Paths")
 	static FString MakeValidFileName(const FString& InString, const FString& InReplacementChar = TEXT(""));
@@ -587,9 +565,8 @@ class ENGINE_API UBlueprintPathsLibrary : public UBlueprintFunctionLibrary
 	* Note that this is a different set of restrictions to those imposed by FPackageName
 	*
 	* @param InPath - path to validate
-	* @param OutReason - optional parameter to fill with the failure reason
-	* @return bDidSucceed - Whether the path could be validated
-	* @return OutReason - If the Paths was invalid then this will contain the failure reason
+	* @param OutReason - If validation fails, this is filled with the failure reason
+	* @param bDidSucceed - Whether the path could be validated
 	*/
 	UFUNCTION(BlueprintPure, Category = "Utilities|Paths")
 	static void ValidatePath(const FString& InPath, bool& bDidSucceed, FText& OutReason);
diff --git a/Engine/Source/Runtime/Engine/Classes/Kismet/KismetInputLibrary.h b/Engine/Source/Runtime/Engine/Classes/Kismet/KismetInputLibrary.h
--- a/Engine/Source/Runtime/Engine/Classes/Kismet/KismetInputLibrary.h
+++ b/Engine/Source/Runtime/Engine/Classes/Kismet/KismetInputLibrary.h
@@ -34,7 +34,7 @@ class ENGINE_API UKismetInputLibrary : public UBlueprintFunctionLibrary
 	 * Test if the input key are equal (A == B)
 	 * @param A - The key to compare against
 	 * @param B - The key to compare
-	 * @returns True if the key are equal, false otherwise
+	 * Returns true if the key are equal, false otherwise
 	 */
 	UFUNCTION(BlueprintPure, meta=(DisplayName = "Equal (Key)", CompactNodeTitle = "=="), Category="Utilities|Key")
 	static bool EqualEqual_KeyKey(FKey A, FKey B);
@@ -43,159 +43,133 @@ class ENGINE_API UKismetInputLibrary : public UBlueprintFunctionLibrary
 	* Test if the input chords are equal (A == B)
 	* @param A - The chord to compare against
 	* @param B - The chord to compare
-	* @returns True if the chords are equal, false otherwise
+	* Returns true if the chords are equal, false otherwise
 	*/
 	UFUNCTION( BlueprintPure, meta = ( DisplayName = "Equal (Input Chord)", CompactNodeTitle = "==" ), Category = "Utilities|Key" )
 	static bool EqualEqual_InputChordInputChord( FInputChord A, FInputChord B );
 	
 	/**
-	 * @returns True if the key is a modifier key: Ctrl, Command, Alt, Shift
+	 * Returns true if the key is a modifier key: Ctrl, Command, Alt, Shift
 	 */
 	UFUNCTION(BlueprintPure, meta = (DisplayName = "Is Modifier Key"), Category="Utilities|Key")
 	static bool Key_IsModifierKey(const FKey& Key);
 	
 	/**
-	 * @returns True if the key is a gamepad button
+	 * Returns true if the key is a gamepad button
 	 */
 	UFUNCTION(BlueprintPure, meta = (DisplayName = "Is Gamepad Key"), Category="Utilities|Key")
 	static bool Key_IsGamepadKey(const FKey& Key);
 	
 	/**
-	 * @returns True if the key is a mouse button
+	 * Returns true if the key is a mouse button
 	 */
 	UFUNCTION(BlueprintPure, meta = (DisplayName = "Is Mouse Button"), Category="Utilities|Key")
 	static bool Key_IsMouseButton(const FKey& Key);
 	
 	/**
-	 * @returns True if the key is a keyboard button
+	 * Returns true if the key is a keyboard button
 	 */
 	UFUNCTION(BlueprintPure, meta = (DisplayName = "Is Keyboard Key"), Category="Utilities|Key")
 	static bool Key_IsKeyboardKey(const FKey& Key);
 	
 	/**
-	 * @returns True if the key is a float axis
+	 * Returns true if the key is a float axis
 	 */
 	UFUNCTION(BlueprintPure, meta = (DisplayName = "Is Float Axis"), Category="Utilities|Key")
 	static bool Key_IsFloatAxis(const FKey& Key);
 	
 	/**
-	 * @returns True if the key is a vector axis
+	 * Returns true if the key is a vector axis
 	 */
 	UFUNCTION(BlueprintPure, meta = (DisplayName = "Is Vector Axis"), Category="Utilities|Key")
 	static bool Key_IsVectorAxis(const FKey& Key);
 	
 	/**
-	 * @returns True if this is a valid key.
+	 * Returns true if this is a valid key.
 	 */
 	UFUNCTION(BlueprintPure, meta = (DisplayName = "Is Valid"), Category = "Utilities|Key")
 	static bool Key_IsValid(const FKey& Key);
 
 	/**
-	 * @returns The display name of the key.
+	 * Returns the display name of the key.
 	 */
 	UFUNCTION(BlueprintPure, meta = (DisplayName = "Get Key Display Name"), Category="Utilities|Key")
 	static FText Key_GetDisplayName(const FKey& Key);
 
 	/**
 	 * Returns whether or not this character is an auto-repeated keystroke
-	 *
-	 * @return  True if this character is a repeat
 	 */
 	UFUNCTION(BlueprintPure, meta=( DisplayName = "Is Repeat" ), Category="Utilities|InputEvent")
 	static bool InputEvent_IsRepeat(const FInputEvent& Input);
 
 	/**
 	 * Returns true if either shift key was down when this event occurred
-	 *
-	 * @return  True if shift is pressed
 	 */
 	UFUNCTION(BlueprintPure, meta=( DisplayName = "Is Shift Down" ), Category="Utilities|InputEvent")
 	static bool InputEvent_IsShiftDown(const FInputEvent& Input);
 
 	/**
 	 * Returns true if left shift key was down when this event occurred
-	 *
-	 * @return True if left shift is pressed.
 	 */
 	UFUNCTION(BlueprintPure, meta=( DisplayName = "Is Left Shift Down" ), Category="Utilities|InputEvent")
 	static bool InputEvent_IsLeftShiftDown(const FInputEvent& Input);
 
 	/**
 	 * Returns true if right shift key was down when this event occurred
-	 *
-	 * @return True if right shift is pressed.
 	 */
 	UFUNCTION(BlueprintPure, meta=( DisplayName = "Is Right Shift Down" ), Category="Utilities|InputEvent")
 	static bool InputEvent_IsRightShiftDown(const FInputEvent& Input);
 
 	/**
 	 * Returns true if either control key was down when this event occurred
-	 *
-	 * @return  True if control is pressed
 	 */
 	UFUNCTION(BlueprintPure, meta=( DisplayName = "Is Control Down" ), Category="Utilities|InputEvent")
 	static bool InputEvent_IsControlDown(const FInputEvent& Input);
 
 	/**
 	 * Returns true if left control key was down when this event occurred
-	 *
-	 * @return  True if left control is pressed
 	 */
 	UFUNCTION(BlueprintPure, meta=( DisplayName = "Is Left Control Down" ), Category="Utilities|InputEvent")
 	static bool InputEvent_IsLeftControlDown(const FInputEvent& Input);
 
 	/**
 	 * Returns true if left control key was down when this event occurred
-	 *
-	 * @return  True if left control is pressed
 	 */
 	UFUNCTION(BlueprintPure, meta=( DisplayName = "Is Right Control Down" ), Category="Utilities|InputEvent")
 	static bool InputEvent_IsRightControlDown(const FInputEvent& Input);
 
 	/**
 	 * Returns true if either alt key was down when this event occurred
-	 *
-	 * @return  True if alt is pressed
 	 */
 	UFUNCTION(BlueprintPure, meta=( DisplayName = "Is Alt Down" ), Category="Utilities|InputEvent")
 	static bool InputEvent_IsAltDown(const FInputEvent& Input);
 
 	/**
 	 * Returns true if left alt key was down when this event occurred
-	 *
-	 * @return  True if left alt is pressed
 	 */
 	UFUNCTION(BlueprintPure, meta=( DisplayName = "Is Left Alt Down" ), Category="Utilities|InputEvent")
 	static bool InputEvent_IsLeftAltDown(const FInputEvent& Input);
 
 	/**
 	 * Returns true if right alt key was down when this event occurred
-	 *
-	 * @return  True if right alt is pressed
 	 */
 	UFUNCTION(BlueprintPure, meta=( DisplayName = "Is Right Alt Down" ), Category="Utilities|InputEvent")
 	static bool InputEvent_IsRightAltDown(const FInputEvent& Input);
 
 	/**
 	 * Returns true if either command key was down when this event occurred
-	 *
-	 * @return  True if command is pressed
 	 */
 	UFUNCTION(BlueprintPure, meta=( DisplayName = "Is Command Down" ), Category="Utilities|InputEvent")
 	static bool InputEvent_IsCommandDown(const FInputEvent& Input);
 
 	/**
 	 * Returns true if left command key was down when this event occurred
-	 *
-	 * @return  True if left command is pressed
 	 */
 	UFUNCTION(BlueprintPure, meta=( DisplayName = "Is Left Command Down" ), Category="Utilities|InputEvent")
 	static bool InputEvent_IsLeftCommandDown(const FInputEvent& Input);
 
 	/**
 	 * Returns true if right command key was down when this event occurred
-	 *
-	 * @return  True if right command is pressed
 	 */
 	UFUNCTION(BlueprintPure, meta=( DisplayName = "Is Right Command Down" ), Category="Utilities|InputEvent")
 	static bool InputEvent_IsRightCommandDown(const FInputEvent& Input);
@@ -215,15 +189,15 @@ class ENGINE_API UKismetInputLibrary : public UBlueprintFunctionLibrary
 	UFUNCTION(BlueprintPure, Category = "Utilities|FAnalogInputEvent")
 	static float GetAnalogValue(const FAnalogInputEvent& Input);
 
-	/** @return The position of the cursor in screen space */
+	/** Returns The position of the cursor in screen space */
 	UFUNCTION(BlueprintPure, meta=( DisplayName = "Get Screen Space Position" ), Category="Utilities|PointerEvent")
 	static FVector2D PointerEvent_GetScreenSpacePosition(const FPointerEvent& Input);
 
-	/** @return The position of the cursor in screen space last time we handled an input event */
+	/** Returns the position of the cursor in screen space last time we handled an input event */
 	UFUNCTION(BlueprintPure, meta=( DisplayName = "Get Last Screen Space Position" ), Category="Utilities|PointerEvent")
 	static FVector2D PointerEvent_GetLastScreenSpacePosition(const FPointerEvent& Input);
 
-	/** @return the distance the mouse traveled since the last event was handled. */
+	/** Returns the distance the mouse traveled since the last event was handled. */
 	UFUNCTION(BlueprintPure, meta=( DisplayName = "Get Cursor Delta" ), Category="Utilities|PointerEvent")
 	static FVector2D PointerEvent_GetCursorDelta(const FPointerEvent& Input);
 
@@ -239,27 +213,27 @@ class ENGINE_API UKismetInputLibrary : public UBlueprintFunctionLibrary
 	UFUNCTION(BlueprintPure, meta=( DisplayName = "Get Wheel Delta" ), Category="Utilities|PointerEvent")
 	static float PointerEvent_GetWheelDelta(const FPointerEvent& Input);
 
-	/** @return The index of the user that caused the event */
+	/** Returns the index of the user that caused the event */
 	UFUNCTION(BlueprintPure, meta=( DisplayName = "Get User Index" ), Category="Utilities|PointerEvent")
 	static int32 PointerEvent_GetUserIndex(const FPointerEvent& Input);
 
-	/** @return The unique identifier of the pointer (e.g., finger index) */
+	/** Returns the unique identifier of the pointer (e.g., finger index) */
 	UFUNCTION(BlueprintPure, meta=( DisplayName = "Get Pointer Index" ), Category="Utilities|PointerEvent")
 	static int32 PointerEvent_GetPointerIndex(const FPointerEvent& Input);
 
-	/** @return The index of the touch pad that generated this event (for platforms with multiple touch pads per user) */
+	/** Returns the index of the touch pad that generated this event (for platforms with multiple touch pads per user) */
 	UFUNCTION(BlueprintPure, meta=( DisplayName = "Get Touchpad Index" ), Category="Utilities|PointerEvent")
 	static int32 PointerEvent_GetTouchpadIndex(const FPointerEvent& Input);
 
-	/** @return Is this event a result from a touch (as opposed to a mouse) */
+	/** Returns true if this event a result from a touch (as opposed to a mouse) */
 	UFUNCTION(BlueprintPure, meta=( DisplayName = "Is Touch Event" ), Category="Utilities|PointerEvent")
 	static bool PointerEvent_IsTouchEvent(const FPointerEvent& Input);
 
-	/** @return The type of touch gesture */
+	/** Returns the type of touch gesture */
 	UFUNCTION(BlueprintPure, Category="Utilities|PointerEvent")
 	static ESlateGesture PointerEvent_GetGestureType(const FPointerEvent& Input);
 
-	/** @return The change in gesture value since the last gesture event of the same type. */
+	/** Returns the change in gesture value since the last gesture event of the same type. */
 	UFUNCTION(BlueprintPure, meta=( DisplayName = "Get Gesture Delta" ), Category="Utilities|PointerEvent")
 	static FVector2D PointerEvent_GetGestureDelta(const FPointerEvent& Input);
 };
diff --git a/Engine/Source/Runtime/Engine/Classes/Kismet/KismetStringLibrary.h b/Engine/Source/Runtime/Engine/Classes/Kismet/KismetStringLibrary.h
--- a/Engine/Source/Runtime/Engine/Classes/Kismet/KismetStringLibrary.h
+++ b/Engine/Source/Runtime/Engine/Classes/Kismet/KismetStringLibrary.h
@@ -459,27 +459,31 @@ class ENGINE_API UKismetStringLibrary : public UBlueprintFunctionLibrary
 	UFUNCTION(BlueprintPure, Category = "Utilities|String")
 	static bool Split(const FString& SourceString, const FString& InStr, FString& LeftS, FString& RightS, ESearchCase::Type SearchCase = ESearchCase::IgnoreCase, ESearchDir::Type SearchDir = ESearchDir::FromStart);
 	
-	/** @return the left most given number of characters */
+	/** Returns the left most given number of characters */
 	UFUNCTION(BlueprintPure, Category = "Utilities|String")
 	static FString Left(const FString& SourceString, int32 Count);
 	
-	/** @return the left most characters from the string chopping the given number of characters from the end */
+	/** Returns the left most characters from the string chopping the given number of characters from the end */
 	UFUNCTION(BlueprintPure, Category = "Utilities|String")
 	static FString LeftChop(const FString& SourceString, int32 Count);
 	
-	/** @return the string to the right of the specified location, counting back from the right (end of the word). */
+	/** Returns the string to the right of the specified location, counting back from the right (end of the word). */
 	UFUNCTION(BlueprintPure, Category = "Utilities|String")
 	static FString Right(const FString& SourceString, int32 Count);
 
-	/** @return the string to the right of the specified location, counting forward from the left (from the beginning of the word). */
+	/** Returns the string to the right of the specified location, counting forward from the left (from the beginning of the word). */
 	UFUNCTION(BlueprintPure, Category = "Utilities|String")
 	static FString RightChop(const FString& SourceString, int32 Count);
 
-	/** @return the substring from Start position for Count characters. */
+	/** Returns the substring from Start position for Count characters. */
 	UFUNCTION(BlueprintPure, Category = "Utilities|String")
 	static FString Mid(const FString& SourceString, int32 Start, int32 Count);
 
-	/** Convert a number of seconds into minutes:seconds.milliseconds format string */
+	/**
+	 * Convert a number of seconds into minutes:seconds.milliseconds format string (including leading zeroes)
+	 *
+	 * @return A new string built from the seconds parameter
+	 */
 	UFUNCTION(BlueprintPure,  Category = "Utilities|String")
 	static FString TimeSecondsToString(float InSeconds);
 };
diff --git a/Engine/Source/Runtime/SlateCore/Public/Input/Events.h b/Engine/Source/Runtime/SlateCore/Public/Input/Events.h
--- a/Engine/Source/Runtime/SlateCore/Public/Input/Events.h
+++ b/Engine/Source/Runtime/SlateCore/Public/Input/Events.h
@@ -210,8 +210,6 @@ public:
 
 	/**
 	 * Returns true if left shift key was down when this event occurred
-	 *
-	 * @return True if left shift is pressed.
 	 */
 	bool IsLeftShiftDown() const
 	{
@@ -220,8 +218,6 @@ public:
 
 	/**
 	 * Returns true if right shift key was down when this event occurred
-	 *
-	 * @return True if right shift is pressed.
 	 */
 	bool IsRightShiftDown() const
 	{
@@ -230,8 +226,6 @@ public:
 
 	/**
 	 * Returns true if either control key was down when this event occurred
-	 *
-	 * @return  True if control is pressed
 	 */
 	bool IsControlDown() const
 	{
@@ -240,8 +234,6 @@ public:
 
 	/**
 	 * Returns true if left control key was down when this event occurred
-	 *
-	 * @return  True if left control is pressed
 	 */
 	bool IsLeftControlDown() const
 	{
@@ -250,8 +242,6 @@ public:
 
 	/**
 	 * Returns true if right control key was down when this event occurred
-	 *
-	 * @return  True if right control is pressed
 	 */
 	bool IsRightControlDown() const
 	{
@@ -260,8 +250,6 @@ public:
 
 	/**
 	 * Returns true if either alt key was down when this event occurred
-	 *
-	 * @return  True if alt is pressed
 	 */
 	bool IsAltDown() const
 	{
@@ -270,8 +258,6 @@ public:
 
 	/**
 	 * Returns true if left alt key was down when this event occurred
-	 *
-	 * @return  True if left alt is pressed
 	 */
 	bool IsLeftAltDown() const
 	{
@@ -280,8 +266,6 @@ public:
 
 	/**
 	 * Returns true if right alt key was down when this event occurred
-	 *
-	 * @return  True if right alt is pressed
 	 */
 	bool IsRightAltDown() const
 	{
@@ -290,8 +274,6 @@ public:
 
 	/**
 	 * Returns true if either command key was down when this event occurred
-	 *
-	 * @return  True if command is pressed
 	 */
 	bool IsCommandDown() const
 	{
@@ -300,8 +282,6 @@ public:
 
 	/**
 	 * Returns true if left command key was down when this event occurred
-	 *
-	 * @return  True if left command is pressed
 	 */
 	bool IsLeftCommandDown() const
 	{
@@ -310,8 +290,6 @@ public:
 
 	/**
 	 * Returns true if right command key was down when this event occurred
-	 *
-	 * @return  True if right command is pressed
 	 */
 	bool IsRightCommandDown() const
 	{
@@ -320,8 +298,6 @@ public:
 
 	/**
 	 * Returns true if caps lock was on when this event occurred
-	 * 
-	 * @return True if caps lock is on
 	 */
 	bool AreCapsLocked() const
 	{
@@ -329,7 +305,7 @@ public:
 	}
 
 	/**
-	 * @return the complete set of modifier keys
+	 * Returns the complete set of modifier keys
 	 */
 	const FModifierKeysState& GetModifierKeys() const
 	{
@@ -338,8 +314,6 @@ public:
 
 	/**
 	* Returns the index of the user that generated this event.
-	*
-	* @return The index of the user that caused the event
 	*/
 	uint32 GetUserIndex() const
 	{
@@ -800,13 +774,13 @@ public:
 	
 public:
 
-	/** @return The position of the cursor in screen space */
+	/** Returns The position of the cursor in screen space */
 	const FVector2D& GetScreenSpacePosition() const { return ScreenSpacePosition; }
 
-	/** @return The position of the cursor in screen space last time we handled an input event */
+	/** Returns the position of the cursor in screen space last time we handled an input event */
 	const FVector2D& GetLastScreenSpacePosition() const { return LastScreenSpacePosition; }
 
-	/** @return the distance the mouse traveled since the last event was handled. */
+	/** Returns the distance the mouse traveled since the last event was handled. */
 	FVector2D GetCursorDelta() const { return CursorDelta; }
 
 	/** Mouse buttons that are currently pressed */
@@ -818,37 +792,37 @@ public:
 	/** How much did the mouse wheel turn since the last mouse event */
 	float GetWheelDelta() const { return WheelOrGestureDelta.Y; }
 
-	/** @return The index of the user that caused the event */
+	/** Returns the index of the user that caused the event */
 	int32 GetUserIndex() const { return UserIndex; }
 
-	/** @return The unique identifier of the pointer (e.g., finger index) */
+	/** Returns the unique identifier of the pointer (e.g., finger index) */
 	uint32 GetPointerIndex() const { return PointerIndex; }
 
-	/** @return The index of the touch pad that generated this event (for platforms with multiple touch pads per user) */
+	/** Returns the index of the touch pad that generated this event (for platforms with multiple touch pads per user) */
 	uint32 GetTouchpadIndex() const { return TouchpadIndex; }
 
-	/** @return the force of a touch (1.0f is mapped to an general touch force, < 1 is "light", > 1 is "heavy", and 10 is the max force possible) */
+	/** Returns the force of a touch (1.0f is mapped to an general touch force, < 1 is "light", > 1 is "heavy", and 10 is the max force possible) */
 	float GetTouchForce() const { return Force; }
 
-	/** @return Is this event a result from a touch (as opposed to a mouse) */
+	/** Is this event a result from a touch (as opposed to a mouse) */
 	bool IsTouchEvent() const { return bIsTouchEvent; }
 
-	/** @return Is this event a special force-change touch event */
+	/** Is this event a special force-change touch event */
 	bool IsTouchForceChangedEvent() const { return bIsTouchForceChanged; }
 
-	/** @return Is this event a special first-move touch event */
+	/** Is this event a special first-move touch event */
 	bool IsTouchFirstMoveEvent() const { return bIsTouchFirstMove; }
 
-	/** @return The type of touch gesture */
+	/** Returns the type of touch gesture */
 	EGestureEvent GetGestureType() const { return GestureType; }
 
-	/** @return The change in gesture value since the last gesture event of the same type. */
+	/** Returns the change in gesture value since the last gesture event of the same type. */
 	const FVector2D& GetGestureDelta() const { return WheelOrGestureDelta; }
 
-	/** @return Is the gesture delta inverted */
+	/** Is the gesture delta inverted */
 	bool IsDirectionInvertedFromDevice() const { return bIsDirectionInvertedFromDevice; }
 
-	/** @return the full set of pressed buttons */
+	/** Returns the full set of pressed buttons */
 	const TSet<FKey>& GetPressedButtons() const { return PressedButtons; }
 
 	/** We override the assignment operator to allow generated code to compile with the const ref member. */
@@ -955,19 +929,19 @@ public:
 
 public:
 
-	/** @return The index of the user that caused the event */
+	/** Returns the index of the user that caused the event */
 	uint32 GetUserIndex() const { return UserIndex; }
 
-	/** @return Current tilt of the device/controller */
+	/** Returns the current tilt of the device/controller */
 	const FVector& GetTilt() const { return Tilt; }
 
-	/** @return Rotation speed */
+	/** Returns otation speed */
 	const FVector& GetRotationRate() const { return RotationRate; }
 
-	/** @return Gravity vector (pointing down into the ground) */
+	/** Returns the gravity vector (pointing down into the ground) */
 	const FVector& GetGravity() const { return Gravity; }
 
-	/** @return 3D acceleration of the device */
+	/** Returns the 3D acceleration of the device */
 	const FVector& GetAcceleration() const { return Acceleration; }
 
 private:
@@ -1022,10 +996,10 @@ public:
 
 public:
 
-	/** @return the type of navigation request (Left, Right, Up, Down) */
+	/** Returns the type of navigation request (Left, Right, Up, Down) */
 	EUINavigation GetNavigationType() const { return NavigationType; }
 
-	/** @return the genesis of the navigation request (Keyboard, Controller, User) */
+	/** Returns the genesis of the navigation request (Keyboard, Controller, User) */
 	ENavigationGenesis GetNavigationGenesis() const { return NavigationGenesis; }
 
 private:
diff --git a/Engine/Source/Runtime/UMG/Public/Blueprint/IUserListEntry.h b/Engine/Source/Runtime/UMG/Public/Blueprint/IUserListEntry.h
--- a/Engine/Source/Runtime/UMG/Public/Blueprint/IUserListEntry.h
+++ b/Engine/Source/Runtime/UMG/Public/Blueprint/IUserListEntry.h
@@ -19,15 +19,15 @@ class UMG_API INativeUserListEntry : public IInterface
 	GENERATED_IINTERFACE_BODY()
 
 public:
-	/** @return True if the item represented by this entry is currently selected. */
+	/** Returns true if the item represented by this entry is currently selected. */
 	UFUNCTION(BlueprintCallable, Category = UserListEntry)
 	virtual bool IsListItemSelected() const = 0;
 
-	/** @return True if the item represented by this entry is currently expanded and showing its children. Tree view entries only. */
+	/** Returns true if the item represented by this entry is currently expanded and showing its children. Tree view entries only. */
 	UFUNCTION(BlueprintCallable, Category = UserListEntry)
 	virtual bool IsListItemExpanded() const = 0;
 
-	/** @return The list view that contains this entry. */
+	/** Returns the list view that contains this entry. */
 	UFUNCTION(BlueprintCallable, Category = UserListEntry)
 	virtual UListViewBase* GetOwningListView() const = 0;
 
diff --git a/Engine/Source/Runtime/UMG/Public/Blueprint/SlateBlueprintLibrary.h b/Engine/Source/Runtime/UMG/Public/Blueprint/SlateBlueprintLibrary.h
--- a/Engine/Source/Runtime/UMG/Public/Blueprint/SlateBlueprintLibrary.h
+++ b/Engine/Source/Runtime/UMG/Public/Blueprint/SlateBlueprintLibrary.h
@@ -43,11 +43,11 @@ public:
 	UFUNCTION(BlueprintPure, Category="User Interface|Geometry")
 	static FVector2D LocalToAbsolute(const FGeometry& Geometry, FVector2D LocalCoordinate);
 
-	/** @return the size of the geometry in local space. */
+	/** Returns the size of the geometry in local space. */
 	UFUNCTION(BlueprintPure, Category="User Interface|Geometry")
 	static FVector2D GetLocalSize(const FGeometry& Geometry);
 
-	/** @return the size of the geometry in absolute space. */
+	/** Returns the size of the geometry in absolute space. */
 	UFUNCTION(BlueprintPure, Category="User Interface|Geometry")
 	static FVector2D GetAbsoluteSize(const FGeometry& Geometry);
 
@@ -67,7 +67,7 @@ public:
 	UFUNCTION(BlueprintPure, Category = "User Interface|Geometry")
 	static FVector2D TransformVectorLocalToAbsolute(const FGeometry& Geometry, FVector2D LocalVector);
 
-	/** @return Whether brushes A and B are identical. */
+	/** Returns whether brushes A and B are identical. */
 	UFUNCTION(BlueprintPure, meta = (DisplayName = "Equal (SlateBrush)", CompactNodeTitle = "=="), Category = "SlateBrush")
 	static bool EqualEqual_SlateBrush(const FSlateBrush& A, const FSlateBrush& B);
 
diff --git a/Engine/Source/Runtime/UMG/Public/Components/CheckBox.h b/Engine/Source/Runtime/UMG/Public/Components/CheckBox.h
--- a/Engine/Source/Runtime/UMG/Public/Components/CheckBox.h
+++ b/Engine/Source/Runtime/UMG/Public/Components/CheckBox.h
@@ -117,7 +117,7 @@ public:
 	UFUNCTION(BlueprintCallable, Category="Widget")
 	bool IsChecked() const;
 
-	/** @return the full current checked state. */
+	/** Returns the full current checked state. */
 	UFUNCTION(BlueprintCallable, Category="Widget")
 	ECheckBoxState GetCheckedState() const;
 
diff --git a/Engine/Source/Runtime/UMG/Public/Components/ComboBoxString.h b/Engine/Source/Runtime/UMG/Public/Components/ComboBoxString.h
--- a/Engine/Source/Runtime/UMG/Public/Components/ComboBoxString.h
+++ b/Engine/Source/Runtime/UMG/Public/Components/ComboBoxString.h
@@ -133,7 +133,7 @@ public:
 	UFUNCTION(BlueprintCallable, Category="ComboBox")
 	int32 GetSelectedIndex() const;
 
-	/** @return The number of options */
+	/** Returns the number of options */
 	UFUNCTION(BlueprintCallable, Category="ComboBox")
 	int32 GetOptionCount() const;
 
diff --git a/Engine/Source/Runtime/UMG/Public/Components/ListView.h b/Engine/Source/Runtime/UMG/Public/Components/ListView.h
--- a/Engine/Source/Runtime/UMG/Public/Components/ListView.h
+++ b/Engine/Source/Runtime/UMG/Public/Components/ListView.h
@@ -89,7 +89,7 @@ public:
 	UFUNCTION(BlueprintCallable, Category = ListView)
 	void SetSelectionMode(TEnumAsByte<ESelectionMode::Type> SelectionMode);
 
-	/** @return True if a refresh is pending and the list will be rebuilt on the next tick */
+	/** Returns true if a refresh is pending and the list will be rebuilt on the next tick */
 	UFUNCTION(BlueprintCallable, Category = ListView)
 	bool IsRefreshPending() const;
 
diff --git a/Engine/Source/Runtime/UMG/Public/Components/MenuAnchor.h b/Engine/Source/Runtime/UMG/Public/Components/MenuAnchor.h
--- a/Engine/Source/Runtime/UMG/Public/Components/MenuAnchor.h
+++ b/Engine/Source/Runtime/UMG/Public/Components/MenuAnchor.h
@@ -76,23 +76,23 @@ public:
 	UFUNCTION(BlueprintCallable, Category="Menu Anchor")
 	void Close();
 
-	/** @return true if the popup is open; false otherwise. */
+	/** Returns true if the popup is open; false otherwise. */
 	UFUNCTION(BlueprintCallable, Category="Menu Anchor")
 	bool IsOpen() const;
 
 	/**
-	 * @return true if we should open the menu due to a click. Sometimes we should not, if
+	 * Returns true if we should open the menu due to a click. Sometimes we should not, if
 	 * the same MouseDownEvent that just closed the menu is about to re-open it because it 
 	 * happens to land on the button.
 	 */
 	UFUNCTION(BlueprintCallable, Category="Menu Anchor")
 	bool ShouldOpenDueToClick() const;
 
-	/** @return The current menu position */
+	/** Returns the current menu position */
 	UFUNCTION(BlueprintCallable, Category="Menu Anchor")
 	FVector2D GetMenuPosition() const;
 
-	/** @return Whether this menu has open submenus */
+	/** Returns whether this menu has open submenus */
 	UFUNCTION(BlueprintCallable, Category="Menu Anchor")
 	bool HasOpenSubMenus() const;
 
diff --git a/Engine/Source/Runtime/UMG/Public/Components/PanelWidget.h b/Engine/Source/Runtime/UMG/Public/Components/PanelWidget.h
--- a/Engine/Source/Runtime/UMG/Public/Components/PanelWidget.h
+++ b/Engine/Source/Runtime/UMG/Public/Components/PanelWidget.h
@@ -43,7 +43,7 @@ public:
 	UFUNCTION(BlueprintCallable, Category="Widget|Panel")
 	int32 GetChildIndex(const UWidget* Content) const;
 
-	/** @return true if panel contains this widget */
+	/** Returns true if panel contains this widget */
 	UFUNCTION(BlueprintCallable, Category="Widget|Panel")
 	bool HasChild(UWidget* Content) const;
 
@@ -99,7 +99,7 @@ public:
 	UFUNCTION(BlueprintCallable, Category="Widget|Panel")
 	bool RemoveChild(UWidget* Content);
 
-	/** @return true if there are any child widgets in the panel */
+	/** Returns true if there are any child widgets in the panel */
 	UFUNCTION(BlueprintCallable, Category="Widget|Panel")
 	bool HasAnyChildren() const;
 
@@ -110,13 +110,13 @@ public:
 	/** The slots in the widget holding the child widgets of this panel. */
 	const TArray<UPanelSlot*>& GetSlots() const;
 
-	/** @returns true if the panel supports more than one child. */
+	/** Returns true if the panel supports more than one child. */
 	bool CanHaveMultipleChildren() const
 	{
 		return bCanHaveMultipleChildren;
 	}
 
-	/** @returns true if the panel can accept another child widget. */
+	/** Returns true if the panel can accept another child widget. */
 	bool CanAddMoreChildren() const
 	{
 		return CanHaveMultipleChildren() || GetChildrenCount() == 0;
diff --git a/Engine/Source/Runtime/UMG/Public/Components/Widget.h b/Engine/Source/Runtime/UMG/Public/Components/Widget.h
--- a/Engine/Source/Runtime/UMG/Public/Components/Widget.h
+++ b/Engine/Source/Runtime/UMG/Public/Components/Widget.h
@@ -386,7 +386,7 @@ public:
 
 #if WITH_EDITOR
 
-	/** @return is this widget locked */
+	/** Is this widget locked in the designer UI */
 	bool IsLockedInDesigner() const
 	{
 		return bLockedInDesigner;
@@ -453,7 +453,7 @@ public:
 	UFUNCTION(BlueprintCallable, Category="Widget")
 	void ResetCursor();
 
-	/** @return true if the widget is Visible, HitTestInvisible or SelfHitTestInvisible. */
+	/** Returns true if the widget is Visible, HitTestInvisible or SelfHitTestInvisible. */
 	UFUNCTION(BlueprintCallable, Category="Widget")
 	bool IsVisible() const;
 
@@ -485,7 +485,7 @@ public:
 	UFUNCTION(BlueprintCallable, Category="Widget")
 	void ForceVolatile(bool bForce);
 
-	/** @return true if the widget is currently being hovered by a pointer device */
+	/** Returns true if the widget is currently being hovered by a pointer device */
 	UFUNCTION(BlueprintCallable, Category="Widget")
 	virtual bool IsHovered() const;
 
@@ -517,19 +517,19 @@ public:
 	UFUNCTION(BlueprintCallable, Category="Widget")
 	void SetKeyboardFocus();
 
-	/** @return true if this widget is focused by a specific user. */
+	/** Returns true if this widget is focused by a specific user. */
 	UFUNCTION(BlueprintCallable, Category="Widget")
 	bool HasUserFocus(APlayerController* PlayerController) const;
 
-	/** @return true if this widget is focused by any user. */
+	/** Returns true if this widget is focused by any user. */
 	UFUNCTION(BlueprintCallable, Category="Widget")
 	bool HasAnyUserFocus() const;
 
-	/** @return true if any descendant widget is focused by any user. */
+	/** Returns true if any descendant widget is focused by any user. */
 	UFUNCTION(BlueprintCallable, Category="Widget", meta=(DisplayName="HasAnyUserFocusedDescendants"))
 	bool HasFocusedDescendants() const;
 
-	/** @return true if any descendant widget is focused by a specific user. */
+	/** Returns true if any descendant widget is focused by a specific user. */
 	UFUNCTION(BlueprintCallable, Category="Widget")
 	bool HasUserFocusedDescendants(APlayerController* PlayerController) const;
 	
@@ -715,7 +715,7 @@ public:
 		return ( DesignerFlags&FlagToCheck ) != 0;
 	}
 
-	/** @return The friendly name of the widget to display in the editor */
+	/** Returns the friendly name of the widget to display in the editor */
 	const FString& GetDisplayLabel() const
 	{
 		return DisplayLabel;
@@ -724,7 +724,7 @@ public:
 	/** Sets the friendly name of the widget to display in the editor */
 	void SetDisplayLabel(const FString& DisplayLabel);
 
-	/** @return the category name of the widget */
+	/** Returns the category name of the widget */
 	const FString& GetCategoryName() const;
 
 	/** Sets the category name of the widget */
diff --git a/Engine/Source/Runtime/UMG/Public/Components/WidgetComponent.h b/Engine/Source/Runtime/UMG/Public/Components/WidgetComponent.h
--- a/Engine/Source/Runtime/UMG/Public/Components/WidgetComponent.h
+++ b/Engine/Source/Runtime/UMG/Public/Components/WidgetComponent.h
@@ -132,37 +132,37 @@ public:
 	 */
 	TTuple<FVector, FVector2D> GetCylinderHitLocation(FVector WorldHitLocation, FVector WorldHitDirection) const;
 
-	/** @return Gets the last local location that was hit */
+	/** Gets the last local location that was hit */
 	FVector2D GetLastLocalHitLocation() const
 	{
 		return LastLocalHitLocation;
 	}
 	
-	/** @return The class of the user widget displayed by this component */
+	/** Returns the class of the user widget displayed by this component */
 	TSubclassOf<UUserWidget> GetWidgetClass() const { return WidgetClass; }
 
-	/** @return The user widget object displayed by this component */
+	/** Returns the user widget object displayed by this component */
 	UFUNCTION(BlueprintCallable, Category=UserInterface, meta=(UnsafeDuringActorConstruction=true))
 	UUserWidget* GetUserWidgetObject() const;
 
-	/** @return Returns the Slate widget that was assigned to this component, if any */
+	/** Returns the Slate widget that was assigned to this component, if any */
 	const TSharedPtr<SWidget>& GetSlateWidget() const;
 
-	/** @return List of widgets with their geometry and the cursor position transformed into this Widget component's space. */
+	/** Returns the list of widgets with their geometry and the cursor position transformed into this Widget component's space. */
 	TArray<FWidgetAndPointer> GetHitWidgetPath(FVector WorldHitLocation, bool bIgnoreEnabledStatus, float CursorRadius = 0.0f);
 
-	/** @return List of widgets with their geometry and the cursor position transformed into this Widget space. The widget space is expressed as a Vector2D. */
+	/** Returns the list of widgets with their geometry and the cursor position transformed into this Widget space. The widget space is expressed as a Vector2D. */
 	TArray<FWidgetAndPointer> GetHitWidgetPath(FVector2D WidgetSpaceHitCoordinate, bool bIgnoreEnabledStatus, float CursorRadius = 0.0f);
 
-	/** @return The render target to which the user widget is rendered */
+	/** Returns the render target to which the user widget is rendered */
 	UFUNCTION(BlueprintCallable, Category=UserInterface)
 	UTextureRenderTarget2D* GetRenderTarget() const;
 
-	/** @return The dynamic material instance used to render the user widget */
+	/** Returns the dynamic material instance used to render the user widget */
 	UFUNCTION(BlueprintCallable, Category=UserInterface)
 	UMaterialInstanceDynamic* GetMaterialInstance() const;
 
-	/** @return The window containing the user widget content */
+	/** Returns the window containing the user widget content */
 	TSharedPtr<SWindow> GetSlateWindow() const;
 
 	/**  
@@ -201,11 +201,11 @@ public:
 	UFUNCTION(BlueprintCallable, Category=UserInterface)
 	ULocalPlayer* GetOwnerPlayer() const;
 
-	/** @return The "specified" draw size of the quad in the world */
+	/** Returns the "specified" draw size of the quad in the world */
 	UFUNCTION(BlueprintCallable, Category=UserInterface)
 	FVector2D GetDrawSize() const;
 
-	/** @return The "actual" draw size of the quad in the world */
+	/** Returns the "actual" draw size of the quad in the world */
 	UFUNCTION(BlueprintCallable, Category = UserInterface)
 	FVector2D GetCurrentDrawSize() const;
 
@@ -259,7 +259,7 @@ public:
 	/** Sets how much opacity from the UI widget's texture alpha is used when rendering to the viewport (0.0-1.0) */
 	void SetOpacityFromTexture( const float NewOpacityFromTexture );
 
-	/** @return The pivot point where the UI is rendered about the origin. */
+	/** Returns the pivot point where the UI is rendered about the origin. */
 	UFUNCTION(BlueprintCallable, Category = UserInterface)
 	FVector2D GetPivot() const { return Pivot; }
 
@@ -366,7 +366,7 @@ protected:
 	/** Draws the current widget to the render target if possible. */
 	virtual void DrawWidgetToRenderTarget(float DeltaTime);
 
-	/** @return the width of the widget component taking GeometryMode into account. */
+	/** Returns the width of the widget component taking GeometryMode into account. */
 	float ComputeComponentWidth() const;
 
 	void UpdateMaterialInstance();
