Ensure default materials are cached first when switching feature levels.

diff --git a/Engine/Source/Runtime/Engine/Private/Materials/Material.cpp b/Engine/Source/Runtime/Engine/Private/Materials/Material.cpp
--- a/Engine/Source/Runtime/Engine/Private/Materials/Material.cpp
+++ b/Engine/Source/Runtime/Engine/Private/Materials/Material.cpp
@@ -5668,6 +5668,12 @@ void UMaterial::AllMaterialsCacheResourceShadersForRendering(bool bUpdateProgres
 	GetObjectsOfClass(UMaterial::StaticClass(), MaterialArray, true, RF_ClassDefaultObject, EInternalObjectFlags::None);
 	float TaskIncrement = (float)100.0f / MaterialArray.Num();
 
+	// ensure default materials are cached first. Default materials must be available to fallback to during async compile.
+ 	MaterialArray.Sort([](const UObject& L, const UObject& R)
+ 	{
+ 		return ((const UMaterial&)L).IsDefaultMaterial() > ((const UMaterial&)R).IsDefaultMaterial();
+	});
+
 	for (UObject* MaterialObj : MaterialArray)
 	{
 		UMaterial* Material = (UMaterial*)MaterialObj;
