Fixed a crash in 'stat ai' occuring if UWorld.ControllerList contains nullptrs #UE4

diff --git a/Engine/Source/Runtime/Engine/Private/UnrealEngine.cpp b/Engine/Source/Runtime/Engine/Private/UnrealEngine.cpp
--- a/Engine/Source/Runtime/Engine/Private/UnrealEngine.cpp
+++ b/Engine/Source/Runtime/Engine/Private/UnrealEngine.cpp
@@ -15950,7 +15950,7 @@ int32 UEngine::RenderStatAI(UWorld* World, FViewport* Viewport, FCanvas* Canvas,
 	for (FConstControllerIterator Iterator = World->GetControllerIterator(); Iterator; ++Iterator)
 	{
 		AController* Controller = Iterator->Get();
-		if (!Cast<APlayerController>(Controller))
+		if (Controller && !Cast<APlayerController>(Controller))
 		{
 			++NumAI;
 			if (Controller->GetPawn() != NULL && World->GetTimeSeconds() - Controller->GetPawn()->GetLastRenderTime() < 0.08f)
