Fixed another crash in FStreamingTextureManager

diff --git a/Engine/Source/Runtime/Engine/Private/Streaming/StreamingManagerTexture.cpp b/Engine/Source/Runtime/Engine/Private/Streaming/StreamingManagerTexture.cpp
--- a/Engine/Source/Runtime/Engine/Private/Streaming/StreamingManagerTexture.cpp
+++ b/Engine/Source/Runtime/Engine/Private/Streaming/StreamingManagerTexture.cpp
@@ -780,7 +780,7 @@ void FRenderAssetStreamingManager::RemoveStaticReferences(const UPrimitiveCompon
 		ULevel* Level = Primitive->GetComponentLevel();
 		for (FLevelRenderAssetManager* LevelManager : LevelRenderAssetManagers)
 		{
-			if (!Level || (LevelManager!=nullptr && LevelManager->GetLevel() == Level))
+			if (LevelManager != nullptr && (!Level || LevelManager->GetLevel() == Level))
 			{
 				LevelManager->RemoveComponentReferences(Primitive, RemovedRenderAssets);
 			}
@@ -825,7 +825,7 @@ void FRenderAssetStreamingManager::NotifyPrimitiveDetached( const UPrimitiveComp
 		{
 			for (FLevelRenderAssetManager* LevelManager : LevelRenderAssetManagers)
 			{
-				if (!Level || (LevelManager!=nullptr && LevelManager->GetLevel() == Level))
+				if (LevelManager != nullptr && (!Level || LevelManager->GetLevel() == Level))
 				{
 					LevelManager->RemoveComponentReferences(Primitive, RemovedRenderAssets);
 				}
