Add CountApplicableToMin to listtextures to see how many textures would be affected by changing the minimum required streaming LODs

diff --git a/Engine/Source/Runtime/Engine/Private/UnrealEngine.cpp b/Engine/Source/Runtime/Engine/Private/UnrealEngine.cpp
--- a/Engine/Source/Runtime/Engine/Private/UnrealEngine.cpp
+++ b/Engine/Source/Runtime/Engine/Private/UnrealEngine.cpp
@@ -4819,7 +4819,8 @@ bool UEngine::HandleListTexturesCommand( const TCHAR* Cmd, FOutputDevice& Ar )
 			}
 		}
 	}
-
+	const int32 MinMips = UTexture2D::GetMinTextureResidentMipCount();
+	int32 NumApplicableToMinSize = 0;
 	// Collect textures.
 	TArray<FSortedTexture> SortedTextures;
 	for( TObjectIterator<UTexture> It; It; ++It )
@@ -4856,6 +4857,11 @@ bool UEngine::HandleListTexturesCommand( const TCHAR* Cmd, FOutputDevice& Ar )
 			bIsStreamingTexture = Texture2D->GetStreamingIndex() != INDEX_NONE;
 			UsageCount			= TextureToUsageMap.FindRef(Texture2D);
 			bIsForced			= Texture2D->ShouldMipLevelsBeForcedResident() && bIsStreamingTexture;
+
+			if ((NumMips >= MinMips) && bIsStreamingTexture)
+			{
+				NumApplicableToMinSize++;
+			}
 		}
 		else if (TextureCube != nullptr)
 		{
@@ -4904,8 +4910,8 @@ bool UEngine::HandleListTexturesCommand( const TCHAR* Cmd, FOutputDevice& Ar )
 	FormatMaxAllowedSizes.AddZeroed(PF_MAX);
 
 	// Display.
-	int32 TotalMaxAllowedSize = 0;
-	int32 TotalCurrentSize	= 0;
+	uint64 TotalMaxAllowedSize = 0;
+	uint64 TotalCurrentSize	= 0;
 
 	if (bCSV)
 	{
@@ -4973,7 +4979,7 @@ bool UEngine::HandleListTexturesCommand( const TCHAR* Cmd, FOutputDevice& Ar )
 		TotalCurrentSize	+= SortedTexture.CurrentSize;
 	}
 
-	Ar.Logf(TEXT("Total size: InMem= %.2f MB  OnDisk= %.2f MB  Count=%d"), (double)TotalCurrentSize / 1024. / 1024., (double)TotalMaxAllowedSize / 1024. / 1024., SortedTextures.Num() );
+	Ar.Logf(TEXT("Total size: InMem= %.2f MB  OnDisk= %.2f MB  Count=%d, CountApplicableToMin=%d"), (double)TotalCurrentSize / 1024. / 1024., (double)TotalMaxAllowedSize / 1024. / 1024., SortedTextures.Num(), NumApplicableToMinSize);
 	for (int32 i = 0; i < PF_MAX; ++i)
 	{
 		if (FormatCurrentSizes[i] > 0 || FormatMaxAllowedSizes[i] > 0)
@@ -4995,12 +5001,19 @@ bool UEngine::HandleListTexturesCommand( const TCHAR* Cmd, FOutputDevice& Ar )
 bool UEngine::HandleListStaticMeshesCommand(const TCHAR* Cmd, FOutputDevice& Ar)
 {
 	const bool bAlphaSort = FParse::Param(Cmd, TEXT("ALPHASORT"));
-	const bool bMobileSort = FParse::Param(Cmd, TEXT("MOBILESORT"));	
+	const bool bMobileSort = FParse::Param(Cmd, TEXT("MOBILESORT"));
 	const bool bUsedComponents = FParse::Param(Cmd, TEXT("usedcomponents"));
+	const bool bUsage = FParse::Param(Cmd, TEXT("-?"));
+
+
 
 	//non-editor builds literally don't have the data to determine mobile lods or vert data.  The data prints out incorrectly and is confusing, just remove it.
 	const bool bHasMobileColumns = (bool)WITH_EDITORONLY_DATA;
-	Ar.Logf(TEXT("Listing all static meshes. Optional params: \n-alphasort: sort alphabetically \n-mobilesort: sort by mobile verts \n-usedcomponents: print the all components used by each mesh"));
+	Ar.Logf(TEXT("Listing all static meshes."));
+	if (bUsage)
+	{
+		Ar.Logf(TEXT("\n Optional params: \n-alphasort: sort alphabetically \n-mobilesort: sort by mobile verts \n-usedcomponents: print the all components used by each mesh"));
+	}
 
 	//Collect usage counts
 	TMap<UStaticMesh*, TArray<UStaticMeshComponent*>> UsageList;
