Fix a render asset streaming deadlock

diff --git a/Engine/Source/Runtime/Engine/Private/Streaming/RenderAssetUpdate.h b/Engine/Source/Runtime/Engine/Private/Streaming/RenderAssetUpdate.h
--- a/Engine/Source/Runtime/Engine/Private/Streaming/RenderAssetUpdate.h
+++ b/Engine/Source/Runtime/Engine/Private/Streaming/RenderAssetUpdate.h
@@ -50,6 +50,7 @@ public:
 		: PendingFirstMip(INDEX_NONE)
 		, RequestedMips(INDEX_NONE)
 		, ScheduledTaskCount(0)
+		, LockOwningThreadID(FPlatformTLS::GetCurrentThreadId())
 		, bIsCancelled(false)
 		, TaskState(TS_Locked) // The object is created in the locked state to follow the Tick path
 		, PendingTaskState(TS_None)
@@ -174,6 +175,9 @@ public:
 	{
 		DECLARE_SCOPE_CYCLE_COUNTER(TEXT("FRenderAssetUpdate::DoLock"), STAT_FRenderAssetUpdate_DoLock, STATGROUP_StreamingDetails);
 
+		// Can't lock twice on the same thread or we will deadlock
+		check(LockOwningThreadID != FPlatformTLS::GetCurrentThreadId());
+
 		// Acquire the lock
 		int32 CachedTaskState = TS_None;
 		do
@@ -190,12 +194,16 @@ public:
 
 		// If we just acquired the lock, nothing should be in the process.
 		ensure(PendingTaskState == TS_None);
+		LockOwningThreadID = FPlatformTLS::GetCurrentThreadId();
 		PendingTaskState = (ETaskState)CachedTaskState;
 	}
 
 	/** Release any lock on the object, allowing other thread to modify it. */
 	void DoUnlock()
 	{
+		// Make sure lock and unlock happens on the same thread
+		check(LockOwningThreadID == FPlatformTLS::GetCurrentThreadId());
+
 		ensure(TaskState == TS_Locked && PendingTaskState != TS_Locked);
 
 		ETaskState CachedPendingTaskState = PendingTaskState;
@@ -203,6 +211,7 @@ public:
 		// Reset the pending task state first to prevent racing condition that could fail ensure(PendingTaskState == TS_None) in DoLock()
 		PendingTaskState = TS_None;
 		TaskState = CachedPendingTaskState;
+		LockOwningThreadID = InvalidLockOwningThreadID;
 	}
 
 	/** Get the number of requested mips for this update, ignoring cancellation attempts. */
@@ -289,6 +298,18 @@ private:
 	{
 		DECLARE_SCOPE_CYCLE_COUNTER(TEXT("FRenderAssetUpdate::DoConditionalLock"), STAT_FRenderAssetUpdate_DoConditionalLock, STATGROUP_StreamingDetails);
 
+		// Can't lock twice on the same thread or we will deadlock
+		if (LockOwningThreadID == FPlatformTLS::GetCurrentThreadId())
+		{
+			// We are trying to execute the task on the current thread but failed. Ask for a reschedule on next Tick.
+			// It is safe to modify PendingTaskState here because current thread has the lock
+			if (PendingTaskState == TS_Scheduled)
+			{
+				PendingTaskState = TS_Pending;
+			}
+			return false;
+		}
+
 		// Acquire the lock if there is work to do and if it is allowed to wait for the lock
 		int32 CachedTaskState = TS_None;
 		do
@@ -310,6 +331,7 @@ private:
 		} while (CachedTaskState == TS_Locked || FPlatformAtomics::InterlockedCompareExchange(&TaskState, TS_Locked, CachedTaskState) != CachedTaskState);
 
 		ensure(PendingTaskState == TS_None);
+		LockOwningThreadID = FPlatformTLS::GetCurrentThreadId();
 		PendingTaskState = (ETaskState)CachedTaskState;
 
 		return true;
@@ -409,6 +431,9 @@ private:
 	typedef FAutoDeleteAsyncTask<FMipUpdateTask> FAsyncMipUpdateTask;
 
 protected:
+	/** A special value to indicate that no thread is holding the lock. */
+	static const uint32 InvalidLockOwningThreadID = 0xffffffff;
+
 	/** The index of mip that will end as being the first mip of the intermediate (future) texture/mesh. */
 	int32 PendingFirstMip;
 	/** The total number of mips of the intermediate (future) texture/mesh. */
@@ -420,6 +445,13 @@ protected:
 	/** The number of scheduled ticks (and exceptionally other work from FCancelIORequestsTask) from the renderthread and async thread. Used to prevent deleting the object while it could be accessed. */
 	volatile int32 ScheduledTaskCount;
 
+	/**
+	 * The TLS ID of the thread holding the lock (TS_Locked).
+	 * Used to prevent calling Tick inside Tick on the same thread, which causes deadlock.
+	 * This requires a single call to Tick runs on the same thread from start to end (a.k.a. it doesn't work with Fiber)
+	 */
+	uint32 LockOwningThreadID;
+
 	/** Whether the task has been cancelled because the update could not proceed or because the user called Abort(). */
 	bool bIsCancelled;
 
