UE4 - Patched up a panic sync load of audio. Not a full solution.

diff --git a/Engine/Source/Runtime/CoreUObject/Private/Serialization/BulkData.cpp b/Engine/Source/Runtime/CoreUObject/Private/Serialization/BulkData.cpp
--- a/Engine/Source/Runtime/CoreUObject/Private/Serialization/BulkData.cpp
+++ b/Engine/Source/Runtime/CoreUObject/Private/Serialization/BulkData.cpp
@@ -111,7 +111,7 @@ FOwnedBulkDataPtr::~FOwnedBulkDataPtr()
 const void* FOwnedBulkDataPtr::GetPointer()
 {
 	// return the pointer that the caller can use
-	return AllocatedData ? AllocatedData : MappedRegion->GetMappedPtr();
+	return AllocatedData ? AllocatedData : (MappedRegion ? MappedRegion->GetMappedPtr() : nullptr);
 }
 
 bool FUntypedBulkData::FAllocatedPtr::MapFile(const TCHAR *InFilename, int64 Offset, int64 Size)
diff --git a/Engine/Source/Runtime/Engine/Private/SoundWave.cpp b/Engine/Source/Runtime/Engine/Private/SoundWave.cpp
--- a/Engine/Source/Runtime/Engine/Private/SoundWave.cpp
+++ b/Engine/Source/Runtime/Engine/Private/SoundWave.cpp
@@ -800,6 +800,19 @@ void USoundWave::InitAudioResource( FByteBulkData& CompressedData )
 			check(!OwnedBulkDataPtr);
 			OwnedBulkDataPtr = CompressedData.StealFileMapping();
 			ResourceData = (const uint8*)OwnedBulkDataPtr->GetPointer();
+			if (!ResourceData)
+			{
+				UE_LOG(LogAudio, Error, TEXT("Soundwave '%s' was not loaded when it should have been, forcing a sync load."), *GetFullName());
+
+				delete OwnedBulkDataPtr;
+				CompressedData.ForceBulkDataResident();
+				OwnedBulkDataPtr = CompressedData.StealFileMapping();
+				ResourceData = (const uint8*)OwnedBulkDataPtr->GetPointer();
+				if (!ResourceData)
+				{
+					UE_LOG(LogAudio, Fatal, TEXT("Soundwave '%s' failed to load even after forcing a sync load."), *GetFullName());
+				}
+			}
 		}
 	}
 }
