Curve update issue

diff --git a/Engine/Source/Runtime/Engine/Private/Components/SkeletalMeshComponent.cpp b/Engine/Source/Runtime/Engine/Private/Components/SkeletalMeshComponent.cpp
--- a/Engine/Source/Runtime/Engine/Private/Components/SkeletalMeshComponent.cpp
+++ b/Engine/Source/Runtime/Engine/Private/Components/SkeletalMeshComponent.cpp
@@ -973,6 +973,14 @@ void USkeletalMeshComponent::TickAnimation(float DeltaTime, bool bNeedsValidRoot
 	SCOPED_NAMED_EVENT(USkeletalMeshComponent_TickAnimation, FColor::Yellow);
 	SCOPE_CYCLE_COUNTER(STAT_AnimGameThreadTime);
 	SCOPE_CYCLE_COUNTER(STAT_AnimTickTime);
+
+	// if curves have to be refreshed before updating animation
+	if (!AreRequiredCurvesUpToDate())
+	{
+		QUICK_SCOPE_CYCLE_COUNTER(STAT_USkeletalMeshComponent_RefreshBoneTransforms_RecalcRequiredCurves);
+		RecalcRequiredCurves();
+	}
+
 	if (SkeletalMesh != nullptr)
 	{
 		// We're about to UpdateAnimation, this will potentially queue events that we'll need to dispatch.
