Framework - Splitting the functions for GetPixelBoundingBox and GetPixelPoint, into static ones that don't require you to rebuild the FSceneViewProjectionData which is expensive, instead if you have the projection data, you can now run through these claculations a lot faster.

diff --git a/Engine/Source/Runtime/Engine/Classes/Engine/LocalPlayer.h b/Engine/Source/Runtime/Engine/Classes/Engine/LocalPlayer.h
--- a/Engine/Source/Runtime/Engine/Classes/Engine/LocalPlayer.h
+++ b/Engine/Source/Runtime/Engine/Classes/Engine/LocalPlayer.h
@@ -449,7 +449,8 @@ public:
 	 * @param	OutUpperRight	The Upper Right corner of the pixel space box
 	 * @return  False if there is no viewport, or if the box is behind the camera completely
 	 */
-	bool GetPixelBoundingBox(const FBox& ActorBox, FVector2D& OutLowerLeft, FVector2D& OutUpperRight, const FVector2D* OptionalAllotedSize = NULL);
+	bool GetPixelBoundingBox(const FBox& ActorBox, FVector2D& OutLowerLeft, FVector2D& OutUpperRight, const FVector2D* OptionalAllotedSize = nullptr);
+	static bool GetPixelBoundingBox(const FSceneViewProjectionData& ProjectionData, const FBox& ActorBox, FVector2D& OutLowerLeft, FVector2D& OutUpperRight, const FVector2D* OptionalAllotedSize = nullptr);
 
 	/**
 	 * This function will give you a point in Pixel Space from a World Space position
@@ -458,8 +459,8 @@ public:
 	 * @param	OutPoint	The point in pixel space
 	 * @return  False if there is no viewport, or if the box is behind the camera completely
 	 */
-	bool GetPixelPoint(const FVector& InPoint, FVector2D& OutPoint, const FVector2D* OptionalAllotedSize = NULL);
-
+	bool GetPixelPoint(const FVector& InPoint, FVector2D& OutPoint, const FVector2D* OptionalAllotedSize = nullptr);
+	static bool GetPixelPoint(const FSceneViewProjectionData& ProjectionData, const FVector& InPoint, FVector2D& OutPoint, const FVector2D* OptionalAllotedSize = nullptr);
 
 	/**
 	 * Helper function for deriving various bits of data needed for projection
diff --git a/Engine/Source/Runtime/Engine/Private/LocalPlayer.cpp b/Engine/Source/Runtime/Engine/Private/LocalPlayer.cpp
--- a/Engine/Source/Runtime/Engine/Private/LocalPlayer.cpp
+++ b/Engine/Source/Runtime/Engine/Private/LocalPlayer.cpp
@@ -868,7 +868,7 @@ FSceneView* ULocalPlayer::CalcSceneView( class FSceneViewFamily* ViewFamily,
 
 bool ULocalPlayer::GetPixelBoundingBox(const FBox& ActorBox, FVector2D& OutLowerLeft, FVector2D& OutUpperRight, const FVector2D* OptionalAllotedSize)
 {
-		//@TODO: CAMERA: This has issues with aspect-ratio constrained cameras
+	//@TODO: CAMERA: This has issues with aspect-ratio constrained cameras
 	if ((ViewportClient != NULL) && (ViewportClient->Viewport != NULL) && (PlayerController != NULL))
 	{
 		// get the projection data
@@ -878,71 +878,75 @@ bool ULocalPlayer::GetPixelBoundingBox(const FBox& ActorBox, FVector2D& OutLower
 			return false;
 		}
 
-		// if we passed in an optional size, use it for the viewrect
-		FIntRect ViewRect = ProjectionData.GetConstrainedViewRect();
-		if (OptionalAllotedSize != NULL)
-		{
-			ViewRect.Min = FIntPoint(0,0);
-			ViewRect.Max = FIntPoint(OptionalAllotedSize->X, OptionalAllotedSize->Y);
-		}
+		return ULocalPlayer::GetPixelBoundingBox(ProjectionData, ActorBox, OutLowerLeft, OutUpperRight, OptionalAllotedSize);
+	}
+	else
+	{
+		return false;
+	}
+}
 
-		// transform the box
-		const int32 NumOfVerts = 8;
-		FVector Vertices[NumOfVerts] =
-		{
-			FVector(ActorBox.Min),
-			FVector(ActorBox.Min.X, ActorBox.Min.Y, ActorBox.Max.Z),
-			FVector(ActorBox.Min.X, ActorBox.Max.Y, ActorBox.Min.Z),
-			FVector(ActorBox.Max.X, ActorBox.Min.Y, ActorBox.Min.Z),
-			FVector(ActorBox.Max.X, ActorBox.Max.Y, ActorBox.Min.Z),
-			FVector(ActorBox.Max.X, ActorBox.Min.Y, ActorBox.Max.Z),
-			FVector(ActorBox.Min.X, ActorBox.Max.Y, ActorBox.Max.Z),
-			FVector(ActorBox.Max)
-		};
+bool ULocalPlayer::GetPixelBoundingBox(const FSceneViewProjectionData& ProjectionData, const FBox& ActorBox, FVector2D& OutLowerLeft, FVector2D& OutUpperRight, const FVector2D* OptionalAllotedSize)
+{
+	// if we passed in an optional size, use it for the viewrect
+	FIntRect ViewRect = ProjectionData.GetConstrainedViewRect();
+	if (OptionalAllotedSize != NULL)
+	{
+		ViewRect.Min = FIntPoint(0, 0);
+		ViewRect.Max = FIntPoint(OptionalAllotedSize->X, OptionalAllotedSize->Y);
+	}
+
+	// transform the box
+	const int32 NumOfVerts = 8;
+	FVector Vertices[NumOfVerts] =
+	{
+		FVector(ActorBox.Min),
+		FVector(ActorBox.Min.X, ActorBox.Min.Y, ActorBox.Max.Z),
+		FVector(ActorBox.Min.X, ActorBox.Max.Y, ActorBox.Min.Z),
+		FVector(ActorBox.Max.X, ActorBox.Min.Y, ActorBox.Min.Z),
+		FVector(ActorBox.Max.X, ActorBox.Max.Y, ActorBox.Min.Z),
+		FVector(ActorBox.Max.X, ActorBox.Min.Y, ActorBox.Max.Z),
+		FVector(ActorBox.Min.X, ActorBox.Max.Y, ActorBox.Max.Z),
+		FVector(ActorBox.Max)
+	};
 
-		// create the view projection matrix
-		const FMatrix ViewProjectionMatrix = ProjectionData.ComputeViewProjectionMatrix();
+	// create the view projection matrix
+	const FMatrix ViewProjectionMatrix = ProjectionData.ComputeViewProjectionMatrix();
 
-		int SuccessCount = 0;
-		OutLowerLeft = FVector2D(FLT_MAX, FLT_MAX);
-		OutUpperRight = FVector2D(FLT_MIN, FLT_MIN);
-		for (int i = 0; i < NumOfVerts; ++i)
-		{
-			//grab the point in screen space
-			const FVector4 ScreenPoint = ViewProjectionMatrix.TransformFVector4( FVector4( Vertices[i], 1.0f) );
+	int SuccessCount = 0;
+	OutLowerLeft = FVector2D(FLT_MAX, FLT_MAX);
+	OutUpperRight = FVector2D(FLT_MIN, FLT_MIN);
+	for (int i = 0; i < NumOfVerts; ++i)
+	{
+		//grab the point in screen space
+		const FVector4 ScreenPoint = ViewProjectionMatrix.TransformFVector4(FVector4(Vertices[i], 1.0f));
 
-			if (ScreenPoint.W > 0.0f)
-			{
-				float InvW = 1.0f / ScreenPoint.W;
-				FVector2D PixelPoint = FVector2D( ViewRect.Min.X + (0.5f + ScreenPoint.X * 0.5f * InvW) * ViewRect.Width(),
-												  ViewRect.Min.Y + (0.5f - ScreenPoint.Y * 0.5f * InvW) * ViewRect.Height());
+		if (ScreenPoint.W > 0.0f)
+		{
+			float InvW = 1.0f / ScreenPoint.W;
+			FVector2D PixelPoint = FVector2D(ViewRect.Min.X + (0.5f + ScreenPoint.X * 0.5f * InvW) * ViewRect.Width(),
+				ViewRect.Min.Y + (0.5f - ScreenPoint.Y * 0.5f * InvW) * ViewRect.Height());
 
-				PixelPoint.X = FMath::Clamp<float>(PixelPoint.X, 0, ViewRect.Width());
-				PixelPoint.Y = FMath::Clamp<float>(PixelPoint.Y, 0, ViewRect.Height());
+			PixelPoint.X = FMath::Clamp<float>(PixelPoint.X, 0, ViewRect.Width());
+			PixelPoint.Y = FMath::Clamp<float>(PixelPoint.Y, 0, ViewRect.Height());
 
-				OutLowerLeft.X = FMath::Min(OutLowerLeft.X, PixelPoint.X);
-				OutLowerLeft.Y = FMath::Min(OutLowerLeft.Y, PixelPoint.Y);
+			OutLowerLeft.X = FMath::Min(OutLowerLeft.X, PixelPoint.X);
+			OutLowerLeft.Y = FMath::Min(OutLowerLeft.Y, PixelPoint.Y);
 
-				OutUpperRight.X = FMath::Max(OutUpperRight.X, PixelPoint.X);
-				OutUpperRight.Y = FMath::Max(OutUpperRight.Y, PixelPoint.Y);
+			OutUpperRight.X = FMath::Max(OutUpperRight.X, PixelPoint.X);
+			OutUpperRight.Y = FMath::Max(OutUpperRight.Y, PixelPoint.Y);
 
-				++SuccessCount;
-			}
+			++SuccessCount;
 		}
-
-		// make sure we are calculating with more than one point;
-		return SuccessCount >= 2;
-	}
-	else
-	{
-		return false;
 	}
+
+	// make sure we are calculating with more than one point;
+	return SuccessCount >= 2;
 }
 
 bool ULocalPlayer::GetPixelPoint(const FVector& InPoint, FVector2D& OutPoint, const FVector2D* OptionalAllotedSize)
 {
 	//@TODO: CAMERA: This has issues with aspect-ratio constrained cameras
-	bool bInFrontOfCamera = true;
 	if ((ViewportClient != NULL) && (ViewportClient->Viewport != NULL) && (PlayerController != NULL))
 	{
 		// get the projection data
@@ -952,33 +956,43 @@ bool ULocalPlayer::GetPixelPoint(const FVector& InPoint, FVector2D& OutPoint, co
 			return false;
 		}
 
-		// if we passed in an optional size, use it for the viewrect
-		FIntRect ViewRect = ProjectionData.GetConstrainedViewRect();
-		if (OptionalAllotedSize != NULL)
-		{
-			ViewRect.Min = FIntPoint(0,0);
-			ViewRect.Max = FIntPoint(OptionalAllotedSize->X, OptionalAllotedSize->Y);
-		}
+		return ULocalPlayer::GetPixelPoint(ProjectionData, InPoint, OutPoint, OptionalAllotedSize);
+	}
 
-		// create the view projection matrix
-		const FMatrix ViewProjectionMatrix = ProjectionData.ComputeViewProjectionMatrix();
+	return false;
+}
 
-		//@TODO: CAMERA: Validate this code!
-		// grab the point in screen space
-		FVector4 ScreenPoint = ViewProjectionMatrix.TransformFVector4( FVector4( InPoint, 1.0f) );
+bool ULocalPlayer::GetPixelPoint(const FSceneViewProjectionData& ProjectionData, const FVector& InPoint, FVector2D& OutPoint, const FVector2D* OptionalAllotedSize)
+{
+	bool bInFrontOfCamera = true;
 
-		ScreenPoint.W = (ScreenPoint.W == 0) ? KINDA_SMALL_NUMBER : ScreenPoint.W;
+	// if we passed in an optional size, use it for the viewrect
+	FIntRect ViewRect = ProjectionData.GetConstrainedViewRect();
+	if (OptionalAllotedSize != NULL)
+	{
+		ViewRect.Min = FIntPoint(0, 0);
+		ViewRect.Max = FIntPoint(OptionalAllotedSize->X, OptionalAllotedSize->Y);
+	}
 
-		float InvW = 1.0f / ScreenPoint.W;
-		OutPoint = FVector2D(ViewRect.Min.X + (0.5f + ScreenPoint.X * 0.5f * InvW) * ViewRect.Width(),
-				             ViewRect.Min.Y + (0.5f - ScreenPoint.Y * 0.5f * InvW) * ViewRect.Height());
+	// create the view projection matrix
+	const FMatrix ViewProjectionMatrix = ProjectionData.ComputeViewProjectionMatrix();
 
-		if (ScreenPoint.W < 0.0f)
-		{
-			bInFrontOfCamera = false;
-			OutPoint = FVector2D(ViewRect.Max) - OutPoint;
-		}
+	//@TODO: CAMERA: Validate this code!
+	// grab the point in screen space
+	FVector4 ScreenPoint = ViewProjectionMatrix.TransformFVector4(FVector4(InPoint, 1.0f));
+
+	ScreenPoint.W = (ScreenPoint.W == 0) ? KINDA_SMALL_NUMBER : ScreenPoint.W;
+
+	float InvW = 1.0f / ScreenPoint.W;
+	OutPoint = FVector2D(ViewRect.Min.X + (0.5f + ScreenPoint.X * 0.5f * InvW) * ViewRect.Width(),
+		ViewRect.Min.Y + (0.5f - ScreenPoint.Y * 0.5f * InvW) * ViewRect.Height());
+
+	if (ScreenPoint.W < 0.0f)
+	{
+		bInFrontOfCamera = false;
+		OutPoint = FVector2D(ViewRect.Max) - OutPoint;
 	}
+
 	return bInFrontOfCamera;
 }
 
@@ -1090,7 +1104,6 @@ bool ULocalPlayer::GetProjectionData(FViewport* Viewport, EStereoscopicPass Ster
 		ProjectionData.SetViewRectangle(FIntRect(X, Y, X + SizeX, Y + SizeY));
 	}
 
-
 	return true;
 }
 
