- Fixed crash with editor builds and some missing sounds. Finishes gil's previous checkin

diff --git a/Engine/Source/Runtime/Engine/Private/AudioDecompress.cpp b/Engine/Source/Runtime/Engine/Private/AudioDecompress.cpp
--- a/Engine/Source/Runtime/Engine/Private/AudioDecompress.cpp
+++ b/Engine/Source/Runtime/Engine/Private/AudioDecompress.cpp
@@ -692,7 +692,7 @@ void FAsyncAudioDecompressWorker::DoWork()
 
 		if (Wave->DecompressionType == DTYPE_Native)
 		{
-			UE_CLOG(Wave->OwnedBulkDataPtr->GetMappedRegion(), LogAudio, Warning, TEXT("Mapped audio (%s) was discarded after decompression. This is not ideal as it takes more load time and doesn't save memory."), *Wave->GetName());
+			UE_CLOG(Wave->OwnedBulkDataPtr && Wave->OwnedBulkDataPtr->GetMappedRegion(), LogAudio, Warning, TEXT("Mapped audio (%s) was discarded after decompression. This is not ideal as it takes more load time and doesn't save memory."), *Wave->GetName());
 			// Delete the compressed data
 			Wave->RemoveAudioResource();
 		}
diff --git a/Engine/Source/Runtime/Engine/Private/SoundWave.cpp b/Engine/Source/Runtime/Engine/Private/SoundWave.cpp
--- a/Engine/Source/Runtime/Engine/Private/SoundWave.cpp
+++ b/Engine/Source/Runtime/Engine/Private/SoundWave.cpp
@@ -797,6 +797,10 @@ void USoundWave::InitAudioResource( FByteBulkData& CompressedData )
 		ResourceSize = CompressedData.GetBulkDataSize();
 		if( ResourceSize > 0 )
 		{
+#if WITH_EDITOR
+			check(!ResourceData);
+			CompressedData.GetCopy( ( void** )&ResourceData, true );
+#else
 			check(!OwnedBulkDataPtr);
 			OwnedBulkDataPtr = CompressedData.StealFileMapping();
 			ResourceData = (const uint8*)OwnedBulkDataPtr->GetPointer();
@@ -813,6 +817,7 @@ void USoundWave::InitAudioResource( FByteBulkData& CompressedData )
 					UE_LOG(LogAudio, Fatal, TEXT("Soundwave '%s' failed to load even after forcing a sync load."), *GetFullName());
 				}
 			}
+#endif
 		}
 	}
 }
@@ -824,8 +829,15 @@ bool USoundWave::InitAudioResource(FName Format)
 		FByteBulkData* Bulk = GetCompressedData(Format, GetPlatformCompressionOverridesForCurrentPlatform());
 		if (Bulk)
 		{
+#if WITH_EDITOR
+			ResourceSize = Bulk->GetBulkDataSize();
+			check(ResourceSize > 0);
+			check(!ResourceData);
+			Bulk->GetCopy((void**)&ResourceData, true);
+#else
 			InitAudioResource(*Bulk);
 			check(ResourceSize > 0);
+#endif
 		}
 	}
 
@@ -834,10 +846,19 @@ bool USoundWave::InitAudioResource(FName Format)
 
 void USoundWave::RemoveAudioResource()
 {
+#if WITH_EDITOR
+	if(ResourceData)
+	{
+		FMemory::Free((void*)ResourceData);
+		ResourceSize = 0;
+		ResourceData = NULL;
+	}
+#else
 	delete OwnedBulkDataPtr;
 	OwnedBulkDataPtr = nullptr;
 	ResourceData = nullptr;
 		ResourceSize = 0;
+#endif
 }
 
 #if WITH_EDITOR
