UE4 - Manual merge from 4977115
- Remove shader pipelines on mobile platforms to reduce memory
#rnx

diff --git a/Engine/Source/Runtime/Engine/Private/Materials/MaterialShader.cpp b/Engine/Source/Runtime/Engine/Private/Materials/MaterialShader.cpp
--- a/Engine/Source/Runtime/Engine/Private/Materials/MaterialShader.cpp
+++ b/Engine/Source/Runtime/Engine/Private/Materials/MaterialShader.cpp
@@ -940,7 +940,7 @@ FShader* FMaterialShaderType::FinishCompileShader(
 	// This allows FShaders to share compiled bytecode and RHI shader references
 	FShaderResource* Resource = FShaderResource::FindOrCreateShaderResource(CurrentJob.Output, SpecificType, /* PermutationId = */ 0);
 
-	if (ShaderPipelineType && !ShaderPipelineType->ShouldOptimizeUnusedOutputs())
+	if (ShaderPipelineType && !ShaderPipelineType->ShouldOptimizeUnusedOutputs(CurrentJob.Input.Target.GetPlatform()))
 	{
 		// If sharing shaders in this pipeline, remove it from the type/id so it uses the one in the shared shadermap list
 		ShaderPipelineType = nullptr;
@@ -1538,7 +1538,7 @@ void FMaterialShaderMap::Compile(
 			for (TLinkedList<FShaderPipelineType*>::TIterator ShaderPipelineIt(FShaderPipelineType::GetTypeList());ShaderPipelineIt;ShaderPipelineIt.Next())
 			{
 				const FShaderPipelineType* Pipeline = *ShaderPipelineIt;
-				if (Pipeline->IsMaterialTypePipeline() && Pipeline->HasTessellation() == bHasTessellation)
+				if (Pipeline->IsMaterialTypePipeline() && Pipeline->HasTessellation() == bHasTessellation && RHISupportsShaderPipelines(InPlatform))
 				{
 					auto& StageTypes = Pipeline->GetStages();
 					TArray<FMaterialShaderType*> ShaderStagesToCompile;
@@ -1561,7 +1561,7 @@ void FMaterialShaderMap::Compile(
 					{
 						// Verify that the shader map Id contains inputs for any shaders that will be put into this shader map
 						check(InShaderMapId.ContainsShaderPipelineType(Pipeline));
-						if (Pipeline->ShouldOptimizeUnusedOutputs())
+						if (Pipeline->ShouldOptimizeUnusedOutputs(InPlatform))
 						{
 							NumShaders += ShaderStagesToCompile.Num();
 							FMaterialShaderType::BeginCompileShaderPipeline(CompilingId, InPlatform, Material, MaterialEnvironment, Pipeline, ShaderStagesToCompile, NewJobs);
@@ -2036,7 +2036,7 @@ void FMaterialShaderMap::LoadMissingShadersFromMemory(const FMaterial* Material)
 					FMaterialShaderType* ShaderType = (FMaterialShaderType*)Shader->GetMaterialShaderType();
 					if (!HasShader(ShaderType, /* PermutationId = */ 0))
 					{
-						FShaderId ShaderId(MaterialShaderMapHash, PipelineType->ShouldOptimizeUnusedOutputs() ? PipelineType : nullptr, nullptr, ShaderType, /** PermutationId = */ 0, FShaderTarget(ShaderType->GetFrequency(), GetShaderPlatform()));
+						FShaderId ShaderId(MaterialShaderMapHash, PipelineType->ShouldOptimizeUnusedOutputs(GetShaderPlatform()) ? PipelineType : nullptr, nullptr, ShaderType, /** PermutationId = */ 0, FShaderTarget(ShaderType->GetFrequency(), GetShaderPlatform()));
 						FShader* FoundShader = ShaderType->FindShaderById(ShaderId);
 						if (FoundShader)
 						{
diff --git a/Engine/Source/Runtime/Engine/Private/Materials/MeshMaterialShader.cpp b/Engine/Source/Runtime/Engine/Private/Materials/MeshMaterialShader.cpp
--- a/Engine/Source/Runtime/Engine/Private/Materials/MeshMaterialShader.cpp
+++ b/Engine/Source/Runtime/Engine/Private/Materials/MeshMaterialShader.cpp
@@ -152,7 +152,7 @@ FShader* FMeshMaterialShaderType::FinishCompileShader(
 	// This allows FShaders to share compiled bytecode and RHI shader references
 	FShaderResource* Resource = FShaderResource::FindOrCreateShaderResource(CurrentJob.Output, SpecificType, /* PermutationId = */ 0);
 
-	if (ShaderPipelineType && !ShaderPipelineType->ShouldOptimizeUnusedOutputs())
+	if (ShaderPipelineType && !ShaderPipelineType->ShouldOptimizeUnusedOutputs(CurrentJob.Input.Target.GetPlatform()))
 	{
 		// If sharing shaders in this pipeline, remove it from the type/id so it uses the one in the shared shadermap list
 		ShaderPipelineType = nullptr;
@@ -253,7 +253,7 @@ uint32 FMeshMaterialShaderMap::BeginCompile(
 				// Verify that the shader map Id contains inputs for any shaders that will be put into this shader map
 				check(InShaderMapId.ContainsShaderPipelineType(Pipeline));
 
-				if (Pipeline->ShouldOptimizeUnusedOutputs())
+				if (Pipeline->ShouldOptimizeUnusedOutputs(InPlatform))
 				{
 					NumShadersPerVF += NumShaderStagesToCompile;
 					TArray<FMeshMaterialShaderType*> ShaderStagesToCompile;
@@ -432,7 +432,7 @@ void FMeshMaterialShaderMap::LoadMissingShadersFromMemory(
 					FMeshMaterialShaderType* ShaderType = (FMeshMaterialShaderType*)Shader->GetMeshMaterialShaderType();
 					if (!HasShader(ShaderType, /* PermutationId = */ 0))
 					{
-						const FShaderId ShaderId(MaterialShaderMapHash, PipelineType->ShouldOptimizeUnusedOutputs() ? PipelineType : nullptr, VertexFactoryType, ShaderType, /** PermutationId = */ 0, FShaderTarget(ShaderType->GetFrequency(), InPlatform));
+						const FShaderId ShaderId(MaterialShaderMapHash, PipelineType->ShouldOptimizeUnusedOutputs(InPlatform) ? PipelineType : nullptr, VertexFactoryType, ShaderType, /** PermutationId = */ 0, FShaderTarget(ShaderType->GetFrequency(), InPlatform));
 						FShader* FoundShader = ShaderType->FindShaderById(ShaderId);
 						if (FoundShader)
 						{
diff --git a/Engine/Source/Runtime/Engine/Private/ShaderCompiler/ShaderCompiler.cpp b/Engine/Source/Runtime/Engine/Private/ShaderCompiler/ShaderCompiler.cpp
--- a/Engine/Source/Runtime/Engine/Private/ShaderCompiler/ShaderCompiler.cpp
+++ b/Engine/Source/Runtime/Engine/Private/ShaderCompiler/ShaderCompiler.cpp
@@ -3708,7 +3708,7 @@ FShader* FGlobalShaderTypeCompiler::FinishCompileShader(FGlobalShaderType* Shade
 		FShaderResource* Resource = FShaderResource::FindOrCreateShaderResource(CurrentJob.Output, SpecificType, SpecificPermutationId);
 		check(Resource);
 
-		if (ShaderPipelineType && !ShaderPipelineType->ShouldOptimizeUnusedOutputs())
+		if (ShaderPipelineType && !ShaderPipelineType->ShouldOptimizeUnusedOutputs(CurrentJob.Input.Target.GetPlatform()))
 		{
 			// If sharing shaders in this pipeline, remove it from the type/id so it uses the one in the shared shadermap list
 			ShaderPipelineType = nullptr;
@@ -3869,7 +3869,7 @@ void VerifyGlobalShaders(EShaderPlatform Platform, bool bLoadedFromCacheFile)
 						UE_LOG(LogShaders, Log, TEXT("	%s"), Pipeline->GetName());
 					}
 
-					if (Pipeline->ShouldOptimizeUnusedOutputs())
+					if (Pipeline->ShouldOptimizeUnusedOutputs(Platform))
 					{
 						// Make a pipeline job with all the stages
 						FGlobalShaderTypeCompiler::BeginCompileShaderPipeline(Platform, Pipeline, ShaderStages, GlobalShaderJobs);
@@ -4494,7 +4494,7 @@ static inline FShader* ProcessCompiledJob(FShaderCompileJob* SingleJob, const FS
 	{
 		// Add the new global shader instance to the global shader map if it's a shared shader
 		EShaderPlatform Platform = (EShaderPlatform)SingleJob->Input.Target.Platform;
-		if (!Pipeline || !Pipeline->ShouldOptimizeUnusedOutputs())
+		if (!Pipeline || !Pipeline->ShouldOptimizeUnusedOutputs(Platform))
 		{
 			GGlobalShaderMap[Platform]->AddShader(GlobalShaderType, SingleJob->PermutationId, Shader);
 			// Add this shared pipeline to the list
diff --git a/Engine/Source/Runtime/RHI/Public/RHIDefinitions.h b/Engine/Source/Runtime/RHI/Public/RHIDefinitions.h
--- a/Engine/Source/Runtime/RHI/Public/RHIDefinitions.h
+++ b/Engine/Source/Runtime/RHI/Public/RHIDefinitions.h
@@ -1111,6 +1111,11 @@ inline bool RHISupportsNativeShaderLibraries(const EShaderPlatform Platform)
 	return IsMetalPlatform(Platform);
 }
 
+inline bool RHISupportsShaderPipelines(EShaderPlatform Platform)
+{
+	return !IsMobilePlatform(Platform);
+}
+
 
 // Return what the expected number of samplers will be supported by a feature level
 // Note that since the Feature Level is pretty orthogonal to the RHI/HW, this is not going to be perfect
diff --git a/Engine/Source/Runtime/RenderCore/Private/Shader.cpp b/Engine/Source/Runtime/RenderCore/Private/Shader.cpp
--- a/Engine/Source/Runtime/RenderCore/Private/Shader.cpp
+++ b/Engine/Source/Runtime/RenderCore/Private/Shader.cpp
@@ -1904,7 +1904,7 @@ void DumpShaderStats(EShaderPlatform Platform, EShaderFrequency Frequency)
 					}
 				}
 
-				if (!PipelineType->ShouldOptimizeUnusedOutputs() && bFound)
+				if (!PipelineType->ShouldOptimizeUnusedOutputs(Platform) && bFound)
 				{
 					++NumSharedPipelines;
 				}
@@ -1984,7 +1984,7 @@ void DumpShaderPipelineStats(EShaderPlatform Platform)
 
 		// Write a row for the shader type.
 		ShaderTypeViewer.AddColumn(Type->GetName());
-		ShaderTypeViewer.AddColumn(Type->ShouldOptimizeUnusedOutputs() ? TEXT("U") : TEXT("S"));
+		ShaderTypeViewer.AddColumn(Type->ShouldOptimizeUnusedOutputs(Platform) ? TEXT("U") : TEXT("S"));
 
 		for (int32 Index = 0; Index < SF_NumFrequencies - 1; ++Index)
 		{
diff --git a/Engine/Source/Runtime/RenderCore/Public/Shader.h b/Engine/Source/Runtime/RenderCore/Public/Shader.h
--- a/Engine/Source/Runtime/RenderCore/Public/Shader.h
+++ b/Engine/Source/Runtime/RenderCore/Public/Shader.h
@@ -1531,7 +1531,11 @@ public:
 	bool IsGlobalTypePipeline() const { return Stages[0]->GetGlobalShaderType() != nullptr; }
 	bool IsMaterialTypePipeline() const { return Stages[0]->GetMaterialShaderType() != nullptr; }
 	bool IsMeshMaterialTypePipeline() const { return Stages[0]->GetMeshMaterialShaderType() != nullptr; }
-	bool ShouldOptimizeUnusedOutputs() const { return bShouldOptimizeUnusedOutputs; }
+
+	FORCEINLINE bool ShouldOptimizeUnusedOutputs(EShaderPlatform Platform) const
+	{
+		return bShouldOptimizeUnusedOutputs && RHISupportsShaderPipelines(Platform);
+	}
 
 	/** Gets a list of FShaderTypes & PipelineTypes whose source file no longer matches what that type was compiled with */
 	static void GetOutdatedTypes(TArray<FShaderType*>& OutdatedShaderTypes, TArray<const FShaderPipelineType*>& ShaderPipelineTypesToFlush, TArray<const FVertexFactoryType*>& OutdatedFactoryTypes);
@@ -2049,11 +2053,11 @@ public:
 		for (auto Pair : ShaderPipelines)
 		{
 			FShaderPipeline* Pipeline = Pair.Value;
-			if (Pipeline->PipelineType->ShouldOptimizeUnusedOutputs() && Filter == FShaderPipeline::EOnlyShared)
+			if (Pipeline->PipelineType->ShouldOptimizeUnusedOutputs(Platform) && Filter == FShaderPipeline::EOnlyShared)
 			{
 				continue;
 			}
-			else if (!Pipeline->PipelineType->ShouldOptimizeUnusedOutputs() && Filter == FShaderPipeline::EOnlyUnique)
+			else if (!Pipeline->PipelineType->ShouldOptimizeUnusedOutputs(Platform) && Filter == FShaderPipeline::EOnlyUnique)
 			{
 				continue;
 			}
diff --git a/Engine/Source/Runtime/Renderer/Private/DepthRendering.cpp b/Engine/Source/Runtime/Renderer/Private/DepthRendering.cpp
--- a/Engine/Source/Runtime/Renderer/Private/DepthRendering.cpp
+++ b/Engine/Source/Runtime/Renderer/Private/DepthRendering.cpp
@@ -77,10 +77,10 @@ IMPLEMENT_SHADERPIPELINE_TYPE_VS(DepthPosOnlyNoPixelPipeline, TDepthOnlyVS<true>
 IMPLEMENT_SHADERPIPELINE_TYPE_VSPS(DepthPipeline, TDepthOnlyVS<false>, FDepthOnlyPS, true);
 
 
-bool UseShaderPipelines()
+static FORCEINLINE bool UseShaderPipelines(ERHIFeatureLevel::Type InFeatureLevel)
 {
 	static const auto* CVar = IConsoleManager::Get().FindTConsoleVariableDataInt(TEXT("r.ShaderPipelines"));
-	return CVar && CVar->GetValueOnAnyThread() != 0;
+	return RHISupportsShaderPipelines(GShaderPlatformForFeatureLevel[InFeatureLevel]) && CVar && CVar->GetValueOnAnyThread() != 0;
 }
 
 template <bool bPositionOnly>
@@ -97,7 +97,7 @@ void GetDepthPassShaders(
 {
 	if (bPositionOnly && !bUsesMobileColorValue)
 	{
-		ShaderPipeline = UseShaderPipelines() ? Material.GetShaderPipeline(&DepthPosOnlyNoPixelPipeline, VertexFactoryType) : nullptr;
+		ShaderPipeline = UseShaderPipelines(FeatureLevel) ? Material.GetShaderPipeline(&DepthPosOnlyNoPixelPipeline, VertexFactoryType) : nullptr;
 		VertexShader = ShaderPipeline
 			? ShaderPipeline->GetShader<TDepthOnlyVS<bPositionOnly> >()
 			: Material.GetShader<TDepthOnlyVS<bPositionOnly> >(VertexFactoryType);
@@ -124,7 +124,7 @@ void GetDepthPassShaders(
 		{
 			HullShader = nullptr;
 			DomainShader = nullptr;
-			bool bUseShaderPipelines = UseShaderPipelines();
+			bool bUseShaderPipelines = UseShaderPipelines(FeatureLevel);
 			if (bNeedsPixelShader)
 			{
 				ShaderPipeline = bUseShaderPipelines ? Material.GetShaderPipeline(&DepthPipeline, VertexFactoryType, false) : nullptr;
diff --git a/Engine/Source/Runtime/Renderer/Private/MobileTranslucentRendering.cpp b/Engine/Source/Runtime/Renderer/Private/MobileTranslucentRendering.cpp
--- a/Engine/Source/Runtime/Renderer/Private/MobileTranslucentRendering.cpp
+++ b/Engine/Source/Runtime/Renderer/Private/MobileTranslucentRendering.cpp
@@ -262,7 +262,6 @@ public:
 };
 
 IMPLEMENT_MATERIAL_SHADER_TYPE(, FOpacityOnlyPS, TEXT("/Engine/Private/MobileOpacityShaders.usf"), TEXT("MainPS"), SF_Pixel);
-IMPLEMENT_SHADERPIPELINE_TYPE_VSPS(MobileOpacityPipeline, FOpacityOnlyVS, FOpacityOnlyPS, true);
 
 bool FMobileSceneRenderer::RenderInverseOpacity(FRHICommandListImmediate& RHICmdList, const FViewInfo& View)
 {
diff --git a/Engine/Source/Runtime/Renderer/Private/VelocityRendering.cpp b/Engine/Source/Runtime/Renderer/Private/VelocityRendering.cpp
--- a/Engine/Source/Runtime/Renderer/Private/VelocityRendering.cpp
+++ b/Engine/Source/Runtime/Renderer/Private/VelocityRendering.cpp
@@ -640,7 +640,7 @@ void GetVelocityPassShaders(
 	}
 
 	static const auto* CVar = IConsoleManager::Get().FindTConsoleVariableDataInt(TEXT("r.ShaderPipelines"));
-	const bool bUseShaderPipelines = !bNeedsHSDS && CVar && CVar->GetValueOnAnyThread() != 0;
+	const bool bUseShaderPipelines = RHISupportsShaderPipelines(GShaderPlatformForFeatureLevel[FeatureLevel]) && !bNeedsHSDS && CVar && CVar->GetValueOnAnyThread() != 0;
 
 	FShaderPipeline* ShaderPipeline = bUseShaderPipelines ? Material.GetShaderPipeline(&VelocityPipeline, VertexFactoryType, false) : nullptr;
 	if (ShaderPipeline)
