Remove World from the object collection in ParticleSystemManager. It was introducing a memory reference and was breaking some editor functionnality. The world is already handle properly via OnWorldInit & OnWorldCleanup & OnPreWorldFinishDestroy.

diff --git a/Engine/Plugins/FX/Niagara/Source/Niagara/Private/NiagaraWorldManager.cpp b/Engine/Plugins/FX/Niagara/Source/Niagara/Private/NiagaraWorldManager.cpp
--- a/Engine/Plugins/FX/Niagara/Source/Niagara/Private/NiagaraWorldManager.cpp
+++ b/Engine/Plugins/FX/Niagara/Source/Niagara/Private/NiagaraWorldManager.cpp
@@ -71,6 +71,8 @@ FNiagaraWorldManager* FNiagaraWorldManager::Get(UWorld* World)
 
 void FNiagaraWorldManager::AddReferencedObjects(FReferenceCollector& Collector)
 {
+	// World doesn't need to be added to the reference list. It will be handle via OnWorldInit & OnWorldCleanup & OnPreWorldFinishDestroy in INiagaraModule
+
 	Collector.AddReferencedObjects(ParameterCollections);
 }
 
diff --git a/Engine/Source/Runtime/Engine/Private/Particles/ParticleSystemManager.cpp b/Engine/Source/Runtime/Engine/Private/Particles/ParticleSystemManager.cpp
--- a/Engine/Source/Runtime/Engine/Private/Particles/ParticleSystemManager.cpp
+++ b/Engine/Source/Runtime/Engine/Private/Particles/ParticleSystemManager.cpp
@@ -232,7 +232,7 @@ void FParticleSystemWorldManager::AddReferencedObjects(FReferenceCollector& Coll
 	if (World)
 	{
 		//UE_LOG(LogParticles, Warning, TEXT("| AddReferencedObjects - %d - 0x%p - %s |"), ManagedPSCs.Num(), World, *World->GetFullName());
-		Collector.AddReferencedObject(World);
+		// World doesn't need to be added to the reference list. It will be handle via OnWorldInit & OnWorldCleanup & OnPreWorldFinishDestroy
 
 		for (int32 PSCIndex = 0; PSCIndex < ManagedPSCs.Num(); ++PSCIndex)
 		{
