Allow custom collision channel names to be exposed to Python

diff --git a/Engine/Source/Runtime/Engine/Classes/Engine/CollisionProfile.h b/Engine/Source/Runtime/Engine/Classes/Engine/CollisionProfile.h
--- a/Engine/Source/Runtime/Engine/Classes/Engine/CollisionProfile.h
+++ b/Engine/Source/Runtime/Engine/Classes/Engine/CollisionProfile.h
@@ -285,7 +285,7 @@ private:
 	/**
 	 * Fill up ResponseToChannels for ProfileList data from config
 	 */
-	void FillProfileData(TArray<FCollisionResponseTemplate>& ProfileList, const UEnum* CollisionChannelEnum, const FString& KeyName, TArray<FCustomProfile>& EditProfileList);
+	void FillProfileData(TArray<FCollisionResponseTemplate>& ProfileList, const UEnum* CollisionChannelEnum, TArray<FCustomProfile>& EditProfileList);
 
 	/**
 	 * Load Custom Responses to the Template. Returns true if all are found and customized. False otherwise
diff --git a/Engine/Source/Runtime/Engine/Private/Collision/CollisionProfile.cpp b/Engine/Source/Runtime/Engine/Private/Collision/CollisionProfile.cpp
--- a/Engine/Source/Runtime/Engine/Private/Collision/CollisionProfile.cpp
+++ b/Engine/Source/Runtime/Engine/Private/Collision/CollisionProfile.cpp
@@ -294,10 +294,11 @@ void UCollisionProfile::LoadProfileConfig(bool bForceInit)
 	check (Enum);
 	UStruct* Struct = FCollisionResponseContainer::StaticStruct(); 
 	check (Struct);
-	const FString KeyName = TEXT("DisplayName");
-	const FString TraceType = TEXT("TraceQuery");
-	const FString TraceValue = TEXT("1");
-	const FString HiddenMeta = TEXT("Hidden");
+	static const FString DisplayNameKey = TEXT("DisplayName");
+	static const FString ScriptNameKey = TEXT("ScriptName");
+	static const FString TraceType = TEXT("TraceQuery");
+	static const FString TraceValue = TEXT("1");
+	static const FString HiddenMeta = TEXT("Hidden");
 
 	// need to initialize displaynames separate
 	int32 NumEnum = Enum->NumEnums();
@@ -307,7 +308,7 @@ void UCollisionProfile::LoadProfileConfig(bool bForceInit)
 	ObjectTypeMapping.Empty();
 
 	// first go through enum entry, and add suffix to displaynames
-	int32 PrefixLen = FString(TEXT("ECC_")).Len();
+	static const FString Prefix = TEXT("ECC_");
 
 	// need to have mapping table between ECollisionChannel and EObjectTypeQuery/ETraceTypeQuery
 	UEnum* ObjectTypeEnum = StaticEnum<EObjectTypeQuery>();
@@ -317,7 +318,7 @@ void UCollisionProfile::LoadProfileConfig(bool bForceInit)
 	for ( int32 EnumIndex=0; EnumIndex<NumEnum; ++EnumIndex )
 	{
 		FString EnumName = Enum->GetNameStringByIndex(EnumIndex);
-		EnumName = EnumName.RightChop(PrefixLen);
+		EnumName = EnumName.RightChop(Prefix.Len());
 		FName DisplayName = FName(*EnumName);
 
 		if ( IS_VALID_COLLISIONCHANNEL(EnumIndex) )
@@ -334,7 +335,8 @@ void UCollisionProfile::LoadProfileConfig(bool bForceInit)
 
 #if WITH_EDITOR
 			// clear displayname since we're setting it in below
-			Enum->RemoveMetaData(*KeyName, EnumIndex);
+			Enum->RemoveMetaData(*DisplayNameKey, EnumIndex);
+			Enum->RemoveMetaData(*ScriptNameKey, EnumIndex);
 			if (Enum->HasMetaData(*HiddenMeta, EnumIndex) == false)
 			{
 				Enum->SetMetaData(*HiddenMeta, NULL, EnumIndex);
@@ -405,8 +407,17 @@ void UCollisionProfile::LoadProfileConfig(bool bForceInit)
 				ChannelDisplayNames[EnumIndex] = FName(*DisplayValue);
 
 #if WITH_EDITOR
+				// Create the name to use when exposing this channel for scripting
+				FString ScriptValue = DisplayValue;
+				ScriptValue.RemoveSpacesInline();
+				if (!ScriptValue.StartsWith(Prefix))
+				{
+					ScriptValue.InsertAt(0, Prefix);
+				}
+
 				// set displayvalue for this enum entry
-				Enum->SetMetaData(*KeyName, *DisplayValue, EnumIndex);
+				Enum->SetMetaData(*DisplayNameKey, *DisplayValue, EnumIndex);
+				Enum->SetMetaData(*ScriptNameKey, *ScriptValue, EnumIndex);
 				// also need to remove "Hidden"
 				Enum->RemoveMetaData(*HiddenMeta, EnumIndex);
 #endif 
@@ -441,7 +452,8 @@ void UCollisionProfile::LoadProfileConfig(bool bForceInit)
 				UField* Field = FindField<UField>(Struct, FName(*VariableName));
 				// I verified up in the class, this can't happen
 				check (Field);
-				Field->SetMetaData(*KeyName, *DisplayValue);
+				Field->SetMetaData(*DisplayNameKey, *DisplayValue);
+				Field->SetMetaData(*ScriptNameKey, *ScriptValue);
 #endif
 			}
 			else
@@ -473,7 +485,8 @@ void UCollisionProfile::LoadProfileConfig(bool bForceInit)
 		const FString& Hidden = Enum->GetMetaData(*HiddenMeta, EnumIndex);
 		if ( Hidden.IsEmpty() )
 		{
-			const FString& DisplayName = Enum->GetMetaData(*KeyName, EnumIndex);
+			const FString& DisplayName = Enum->GetMetaData(*DisplayNameKey, EnumIndex);
+			const FString& ScriptName = Enum->GetMetaData(*ScriptNameKey, EnumIndex);
 			if ( !DisplayName.IsEmpty() )
 			{
 				// find out trace type or object type
@@ -483,7 +496,8 @@ void UCollisionProfile::LoadProfileConfig(bool bForceInit)
 					if (ensure(TraceTypeIndex != INDEX_NONE))
 					{
 						TraceTypeEnum->RemoveMetaData(*HiddenMeta, TraceTypeIndex);
-						TraceTypeEnum->SetMetaData(*KeyName, *DisplayName, TraceTypeIndex);
+						TraceTypeEnum->SetMetaData(*DisplayNameKey, *DisplayName, TraceTypeIndex);
+						TraceTypeEnum->SetMetaData(*ScriptNameKey, *ScriptName, TraceTypeIndex);
 						++TraceTypeMappingCount;
 					}
 				}
@@ -493,7 +507,8 @@ void UCollisionProfile::LoadProfileConfig(bool bForceInit)
 					if (ensure(ObjectTypeIndex != INDEX_NONE))
 					{
 						ObjectTypeEnum->RemoveMetaData(*HiddenMeta, ObjectTypeIndex);
-						ObjectTypeEnum->SetMetaData(*KeyName, *DisplayName, ObjectTypeIndex);
+						ObjectTypeEnum->SetMetaData(*DisplayNameKey, *DisplayName, ObjectTypeIndex);
+						ObjectTypeEnum->SetMetaData(*ScriptNameKey, *ScriptName, ObjectTypeIndex);
 						++ObjectTypeMappingCount;
 					}
 				}
@@ -530,7 +545,7 @@ void UCollisionProfile::LoadProfileConfig(bool bForceInit)
 	// 2. Second loads all set up back to ResponseToChannels
 	// this does a lot of iteration, but this only happens once loaded, so it's better to be convenient than efficient
 	// fill up Profiles data
-	FillProfileData(Profiles, Enum, KeyName, EditProfiles);
+	FillProfileData(Profiles, Enum, EditProfiles);
 
 	// 3. It loads redirect data  - now time to load profile redirect
 	ProfileRedirectsMap.Empty();
@@ -573,7 +588,7 @@ void UCollisionProfile::LoadProfileConfig(bool bForceInit)
 #endif
 }
 
-void UCollisionProfile::FillProfileData(TArray<FCollisionResponseTemplate>& ProfileList, const UEnum* CollisionChannelEnum, const FString& KeyName, TArray<FCustomProfile>& EditProfileList)
+void UCollisionProfile::FillProfileData(TArray<FCollisionResponseTemplate>& ProfileList, const UEnum* CollisionChannelEnum, TArray<FCustomProfile>& EditProfileList)
 {
 	// first go through if same name is found later, delete previous entry, this way if game overrides, 
 	// we delete engine version and use game version
