Fixed HierarchicalInstancedStaticMeshComponent's bounds not being calculated properly for AI navigation purposes #UE4

diff --git a/Engine/Source/Runtime/Engine/Classes/Components/HierarchicalInstancedStaticMeshComponent.h b/Engine/Source/Runtime/Engine/Classes/Components/HierarchicalInstancedStaticMeshComponent.h
--- a/Engine/Source/Runtime/Engine/Classes/Components/HierarchicalInstancedStaticMeshComponent.h
+++ b/Engine/Source/Runtime/Engine/Classes/Components/HierarchicalInstancedStaticMeshComponent.h
@@ -266,6 +266,7 @@ protected:
 
 	virtual void GetNavigationPerInstanceTransforms(const FBox& AreaBox, TArray<FTransform>& InstanceData) const override;
 	virtual void PartialNavigationUpdate(int32 InstanceIdx) override;
+	virtual FBox GetNavigationBounds() const override;
 	void FlushAccumulatedNavigationUpdates();
 	mutable FBox AccumulatedNavigationDirtyArea;
 
diff --git a/Engine/Source/Runtime/Engine/Private/HierarchicalInstancedStaticMesh.cpp b/Engine/Source/Runtime/Engine/Private/HierarchicalInstancedStaticMesh.cpp
--- a/Engine/Source/Runtime/Engine/Private/HierarchicalInstancedStaticMesh.cpp
+++ b/Engine/Source/Runtime/Engine/Private/HierarchicalInstancedStaticMesh.cpp
@@ -3114,6 +3114,11 @@ void UHierarchicalInstancedStaticMeshComponent::PartialNavigationUpdate(int32 In
 	}
 }
 
+FBox UHierarchicalInstancedStaticMeshComponent::GetNavigationBounds() const
+{
+	return CalcBounds(GetComponentTransform()).GetBox();
+}
+
 void UHierarchicalInstancedStaticMeshComponent::FlushAccumulatedNavigationUpdates()
 {
 	if (AccumulatedNavigationDirtyArea.IsValid)
