Add delta replay checkpoints (demo.WithDeltaCheckpoints, disabled by default), implemented for local file and http replay streamers

diff --git a/Engine/Source/Runtime/Engine/Classes/Engine/Channel.h b/Engine/Source/Runtime/Engine/Classes/Engine/Channel.h
--- a/Engine/Source/Runtime/Engine/Classes/Engine/Channel.h
+++ b/Engine/Source/Runtime/Engine/Classes/Engine/Channel.h
@@ -71,6 +71,7 @@ class ENGINE_API UChannel
 	uint32				SentClosingBunch:1;	// Set when sending closing bunch to avoid recursion in send-failure-close case.
 	uint32				bPooled:1;			// Set when placed in the actor channel pool
 	uint32				OpenedLocally:1;	// Whether channel was opened locally or by remote.
+	uint32				bOpenedForCheckpoint:1;	// Whether channel was opened by replay checkpoint recording
 	int32				ChIndex;			// Index of this channel.
 	FPacketIdRange		OpenPacketId;		// If OpenedLocally is true, this is the packet we sent the bOpen bunch on. Otherwise, it's the packet we received the bOpen bunch on.
 	UE_DEPRECATED(4.22, "ChType has been deprecated in favor of ChName.")
diff --git a/Engine/Source/Runtime/Engine/Classes/Engine/DemoNetDriver.h b/Engine/Source/Runtime/Engine/Classes/Engine/DemoNetDriver.h
--- a/Engine/Source/Runtime/Engine/Classes/Engine/DemoNetDriver.h
+++ b/Engine/Source/Runtime/Engine/Classes/Engine/DemoNetDriver.h
@@ -162,6 +162,7 @@ enum class EReplayHeaderFlags : uint32
 	None				= 0,
 	ClientRecorded		= ( 1 << 0 ),
 	HasStreamingFixes	= ( 1 << 1 ),
+	DeltaCheckpoints	= ( 1 << 2 ),
 };
 
 ENUM_CLASS_FLAGS(EReplayHeaderFlags);
@@ -407,13 +408,22 @@ class ENGINE_API UDemoNetDriver : public UNetDriver
 
 	/** Keeps track of NetGUIDs that were deleted, so we can skip them when saving checkpoints. Only used while recording. */
 	TSet< FNetworkGUID >							DeletedNetStartupActorGUIDs;
-	
+
+	/** Delta net startup actors that need to be destroyed after checkpoints are loaded */
+	TSet< FString >									DeltaDeletedNetStartupActors;
+
+	/** GUID list of actors deleted before the checkpoint was recorded, including dynamic actors */
+	TSet<FNetworkGUID>								DeletedActorGuids;
+
+	/** Delta checkpoint list for DeletedActorGuids */
+	TSet<FNetworkGUID>								DeltaDeletedActorGuids;
+
 	/** 
 	 * Net startup actors that need to be rolled back during scrubbing by being destroyed and re-spawned 
 	 * NOTE - DeletedNetStartupActors will take precedence here, and destroy the actor instead
 	 */
 	UPROPERTY(transient)
-	TMap< FString, FRollbackNetStartupActorInfo >	RollbackNetStartupActors;
+	TMap<FString, FRollbackNetStartupActorInfo>		RollbackNetStartupActors;
 
 	double				LastCheckpointTime;					// Last time a checkpoint was saved
 
@@ -792,6 +802,12 @@ public:
 		return bHasLevelStreamingFixes;
 	}
 
+	/** Returns whether or not this replay was recorded / is playing with delta checkpoints. */
+	FORCEINLINE bool HasDeltaCheckpoints() const 
+	{
+		return bHasDeltaCheckpoints;
+	}
+
 	/**
 	 * Called when a new ActorChannel is opened, before the Actor is notified.
 	 *
@@ -908,6 +924,9 @@ private:
 	// Whether or not the Streaming Level Fixes are enabled for capture or playback.
 	bool bHasLevelStreamingFixes;
 
+	// Checkpoints are delta compressed
+	bool bHasDeltaCheckpoints;
+
 	// Levels that are currently pending for fast forward.
 	// Using raw pointers, because we manually keep when levels are added and removed.
 	TMap<class ULevel*, TSet<TWeakObjectPtr<class AActor>>> LevelsPendingFastForward;
@@ -1093,6 +1112,8 @@ protected:
 	bool bIsWaitingForHeaderDownload;
 	bool bIsWaitingForStream;
 
+	int64 MaxArchiveReadPos;
+
 private:
 
 	FString ActiveReplayName;
diff --git a/Engine/Source/Runtime/Engine/Classes/Engine/NetConnection.h b/Engine/Source/Runtime/Engine/Classes/Engine/NetConnection.h
--- a/Engine/Source/Runtime/Engine/Classes/Engine/NetConnection.h
+++ b/Engine/Source/Runtime/Engine/Classes/Engine/NetConnection.h
@@ -134,6 +134,13 @@ namespace EClientLoginState
 	}
 };
 
+/** Type of property data resend used by replay checkpoints */
+ enum class EResendAllDataState : uint8
+ {
+	 None,
+	 SinceOpen,
+	 SinceCheckpoint
+ };
 
 // Delegates
 #if !UE_BUILD_SHIPPING
@@ -615,8 +622,16 @@ public:
 	 *	This will also act as if it needs to re-open all the channels, etc.
 	 *   NOTE - This doesn't force all exports to happen again though, it will only export new stuff, so keep that in mind.
 	 */
+	UE_DEPRECATED(4.23, "Use ResendAllDataState instead.")
 	bool bResendAllDataSinceOpen;
 
+	EResendAllDataState ResendAllDataState;
+
+	/** Set of guids we may need to ignore when processing a delta checkpoint */
+	TSet<FNetworkGUID> IgnoredGuids;
+
+	/** Set of channels we may need to ignore when processing a delta checkpoint */
+	TSet<int32> IgnoredChannels;
 
 #if !UE_BUILD_SHIPPING
 	/** Delegate for hooking ReceivedRawPacket */
@@ -1124,6 +1139,8 @@ private:
 	
 	/** True if we've hit the actor channel limit and logged a warning about it */
 	bool bHasWarnedAboutChannelLimit;
+
+	FNetworkGUID GetActorGUIDFromOpenBunch(FInBunch& Bunch);
 };
 
 
diff --git a/Engine/Source/Runtime/Engine/Classes/Engine/NetworkObjectList.h b/Engine/Source/Runtime/Engine/Classes/Engine/NetworkObjectList.h
--- a/Engine/Source/Runtime/Engine/Classes/Engine/NetworkObjectList.h
+++ b/Engine/Source/Runtime/Engine/Classes/Engine/NetworkObjectList.h
@@ -48,6 +48,9 @@ struct FNetworkObjectInfo
 	/** Force this object to be considered relevant for at least one update */
 	uint8 bForceRelevantNextUpdate : 1;
 
+	/** Should this object be considered for replay checkpoint writes */
+	uint8 bDirtyForReplay : 1;
+
 	FNetworkObjectInfo()
 		: Actor(nullptr)
 		, NextUpdateTime(0.0)
@@ -55,7 +58,8 @@ struct FNetworkObjectInfo
 		, OptimalNetUpdateDelta(0.0f)
 		, LastNetUpdateTime(0.0f)
 		, bPendingNetUpdate(false)
-		, bForceRelevantNextUpdate(false) {}
+		, bForceRelevantNextUpdate(false)
+		, bDirtyForReplay(false) {}
 
 	FNetworkObjectInfo(AActor* InActor)
 		: Actor(InActor)
@@ -65,7 +69,8 @@ struct FNetworkObjectInfo
 		, OptimalNetUpdateDelta(0.0f) 
 		, LastNetUpdateTime(0.0f)
 		, bPendingNetUpdate(false)
-		, bForceRelevantNextUpdate(false) {}
+		, bForceRelevantNextUpdate(false)
+		, bDirtyForReplay(false) {}
 
 	void CountBytes(FArchive& Ar) const;
 };
diff --git a/Engine/Source/Runtime/Engine/Classes/Engine/PackageMapClient.h b/Engine/Source/Runtime/Engine/Classes/Engine/PackageMapClient.h
--- a/Engine/Source/Runtime/Engine/Classes/Engine/PackageMapClient.h
+++ b/Engine/Source/Runtime/Engine/Classes/Engine/PackageMapClient.h
@@ -36,27 +36,29 @@ PRAGMA_DISABLE_DEPRECATION_WARNINGS	// 4.22, Name, Type
 class ENGINE_API FNetFieldExport
 {
 public:
-	FNetFieldExport() : bExported( false ), Handle( 0 ), CompatibleChecksum( 0 ), bIncompatible( false )
+	FNetFieldExport() : Handle( 0 ), CompatibleChecksum( 0 ), bExported( false ), bDirtyForReplay( true ), bIncompatible( false )
 	{
 	}
 
 	UE_DEPRECATED(4.22, "Type is no longer required, please use other constructor.")
 	FNetFieldExport( const uint32 InHandle, const uint32 InCompatibleChecksum, const FString InName, FString InType ) :
-		bExported( false ),
 		Handle( InHandle ),
 		CompatibleChecksum( InCompatibleChecksum ),
 		ExportName( *InName ),
 		Name( InName ),
 		Type( InType ),
+		bExported( false ),
+		bDirtyForReplay( true ),
 		bIncompatible( false )
 	{
 	}
 
 	FNetFieldExport( const uint32 InHandle, const uint32 InCompatibleChecksum, const FName& InName ) :
-		bExported( false ),
 		Handle( InHandle ),
 		CompatibleChecksum( InCompatibleChecksum ),
 		ExportName( InName ),
+		bExported( false ),
+		bDirtyForReplay( true ),
 		bIncompatible( false )
 	{
 	}
@@ -97,7 +99,6 @@ public:
 
 	void CountBytes(FArchive& Ar) const;
 
-	bool			bExported;
 	uint32			Handle;
 	uint32			CompatibleChecksum;
 	FName			ExportName;
@@ -105,6 +106,8 @@ public:
 	FString			Name;
 	UE_DEPRECATED(4.22, "Type is deprecated.")
 	FString			Type;
+	bool			bExported;
+	bool			bDirtyForReplay;
 
 	// Transient properties
 	mutable bool	bIncompatible;		// If true, we've already determined that this property isn't compatible. We use this to curb warning spam.
@@ -114,11 +117,12 @@ PRAGMA_ENABLE_DEPRECATION_WARNINGS
 class ENGINE_API FNetFieldExportGroup
 {
 public:
-	FNetFieldExportGroup() : PathNameIndex( 0 ) { }
+	FNetFieldExportGroup() : PathNameIndex( 0 ), bDirtyForReplay( true ) { }
 
 	FString						PathName;
 	uint32						PathNameIndex;
 	TArray< FNetFieldExport >	NetFieldExports;
+	bool						bDirtyForReplay;
 
 	friend FArchive& operator<<( FArchive& Ar, FNetFieldExportGroup& C )
 	{
@@ -162,7 +166,7 @@ public:
 class FNetGuidCacheObject
 {
 public:
-	FNetGuidCacheObject() : NetworkChecksum( 0 ), ReadOnlyTimestamp( 0 ), bNoLoad( 0 ), bIgnoreWhenMissing( 0 ), bIsPending( 0 ), bIsBroken( 0 )
+	FNetGuidCacheObject() : NetworkChecksum( 0 ), ReadOnlyTimestamp( 0 ), bNoLoad( 0 ), bIgnoreWhenMissing( 0 ), bIsPending( 0 ), bIsBroken( 0 ), bDirtyForReplay( 1 )
 	{
 	}
 
@@ -179,21 +183,31 @@ public:
 	uint8						bIgnoreWhenMissing	: 1;	// Don't warn when this asset can't be found or loaded
 	uint8						bIsPending			: 1;	// This object is waiting to be fully loaded
 	uint8						bIsBroken			: 1;	// If this object failed to load, then we set this to signify that we should stop trying
+	uint8						bDirtyForReplay		: 1;	// If this object has been modified, used by replay checkpoints
+};
+
+enum class EAppendNetExportFlags : uint32
+{
+	None = 0,
+	ForceExportDirtyGroups = (1 << 0),
 };
 
+ENUM_CLASS_FLAGS(EAppendNetExportFlags);
+
+
 class ENGINE_API FNetGUIDCache
 {
 public:
 	FNetGUIDCache( UNetDriver * InDriver );
 
-	enum class ENetworkChecksumMode
+	enum class ENetworkChecksumMode : uint8
 	{
 		None			= 0,		// Don't use checksums
 		SaveAndUse		= 1,		// Save checksums in stream, and use to validate while loading packages
 		SaveButIgnore	= 2,		// Save checksums in stream, but ignore when loading packages
 	};
 
-	enum class EAsyncLoadMode
+	enum class EAsyncLoadMode : uint8
 	{
 		UseCVar			= 0,		// Use CVar (net.AllowAsyncLoading) to determine if we should async load
 		ForceDisable	= 1,		// Disable async loading
@@ -243,8 +257,6 @@ public:
 	TSet< FNetworkGUID >							ImportedNetGuids;
 	TMap< FNetworkGUID, TSet< FNetworkGUID > >		PendingOuterNetGuids;
 
-	bool											IsExportingNetGUIDBunch;
-
 	UNetDriver *									Driver;
 
 	TMap< FName, FNetworkGUID >						PendingAsyncPackages;
@@ -252,6 +264,8 @@ public:
 	ENetworkChecksumMode							NetworkChecksumMode;
 	EAsyncLoadMode									AsyncLoadMode;
 
+	bool											IsExportingNetGUIDBunch;
+
 private:
 
 	friend class UPackageMapClient;
@@ -380,6 +394,7 @@ public:
 	void								TrackNetFieldExport( FNetFieldExportGroup* NetFieldExportGroup, const int32 NetFieldExportHandle );
 	TSharedPtr< FNetFieldExportGroup >	GetNetFieldExportGroupChecked( const FString& PathName ) const;
 	void								SerializeNetFieldExportGroupMap( FArchive& Ar, bool bClearPendingExports=true );
+	void								SerializeNetFieldExportDelta(FArchive& Ar);
 
 	TUniquePtr<TGuardValue<bool>> ScopedIgnoreReceivedExportGUIDs()
 	{
@@ -394,6 +409,8 @@ protected:
 	void AppendNetFieldExports( FArchive& Archive );
 	void ReceiveNetFieldExports( FArchive& Archive );
 
+	void AppendNetFieldExportsInternal( FArchive& Archive, const TSet<uint64>& InNetFieldExports, EAppendNetExportFlags Flags );
+
 	void AppendNetExportGUIDs( FArchive& Archive );
 	void ReceiveNetExportGUIDs( FArchive& Archive );
 
diff --git a/Engine/Source/Runtime/Engine/Private/DataChannel.cpp b/Engine/Source/Runtime/Engine/Private/DataChannel.cpp
--- a/Engine/Source/Runtime/Engine/Private/DataChannel.cpp
+++ b/Engine/Source/Runtime/Engine/Private/DataChannel.cpp
@@ -830,7 +830,7 @@ void UActorChannel::AppendExportBunches( TArray<FOutBunch *>& OutExportBunches )
 	Super::AppendExportBunches( OutExportBunches );
 
 	// We don't want to append QueuedExportBunches to these bunches, since these were for queued RPC's, and we don't want to record RPC's during bResendAllDataSinceOpen
-	if ( !Connection->bResendAllDataSinceOpen )
+	if ( Connection->ResendAllDataState == EResendAllDataState::None )
 	{
 		// Let the profiler know about exported GUID bunches
 		for ( const FOutBunch* ExportBunch : QueuedExportBunches )
@@ -857,7 +857,7 @@ void UActorChannel::AppendMustBeMappedGuids( FOutBunch* Bunch )
 	}
 
 	// We don't want to append QueuedMustBeMappedGuidsInLastBunch to these bunches, since these were for queued RPC's, and we don't want to record RPC's during bResendAllDataSinceOpen
-	if ( !Connection->bResendAllDataSinceOpen )
+	if ( Connection->ResendAllDataState == EResendAllDataState::None )
 	{
 		if ( QueuedMustBeMappedGuidsInLastBunch.Num() > 0 )
 		{
@@ -890,10 +890,21 @@ FPacketIdRange UChannel::SendBunch( FOutBunch* Bunch, bool Merge )
 	check( !Bunch->bHasPackageMapExports );
 
 	// Set bunch flags.
-	if( ( OpenPacketId.First==INDEX_NONE || Connection->bResendAllDataSinceOpen ) && OpenedLocally )
+	if( ( OpenPacketId.First==INDEX_NONE || (Connection->ResendAllDataState != EResendAllDataState::None) ) && OpenedLocally )
 	{
-		Bunch->bOpen = 1;
-		OpenTemporary = !Bunch->bReliable;
+		bool bOpenBunch = true;
+
+		if (Connection->ResendAllDataState == EResendAllDataState::SinceCheckpoint)
+		{
+			bOpenBunch = !bOpenedForCheckpoint;
+			bOpenedForCheckpoint = true;
+		}
+		
+		if (bOpenBunch)
+		{
+			Bunch->bOpen = 1;
+			OpenTemporary = !Bunch->bReliable;
+		}
 	}
 
 	// If channel was opened temporarily, we are never allowed to send reliable packets on it.
@@ -1088,7 +1099,7 @@ FPacketIdRange UChannel::SendBunch( FOutBunch* Bunch, bool Merge )
 	}
 
 	// Update open range if necessary
-	if (Bunch->bOpen && !Connection->bResendAllDataSinceOpen)
+	if (Bunch->bOpen && (Connection->ResendAllDataState == EResendAllDataState::None))
 	{
 		OpenPacketId = PacketIdRange;		
 	}
@@ -1112,7 +1123,7 @@ PRAGMA_ENABLE_DEPRECATION_WARNINGS
 // OUtbunch is a bunch that was new'd by the network system or NULL. It should never be one created on the stack
 FOutBunch* UChannel::PrepBunch(FOutBunch* Bunch, FOutBunch* OutBunch, bool Merge)
 {
-	if ( Connection->bResendAllDataSinceOpen )
+	if ( Connection->ResendAllDataState != EResendAllDataState::None )
 	{
 		return Bunch;
 	}
@@ -1176,7 +1187,7 @@ FOutBunch* UChannel::PrepBunch(FOutBunch* Bunch, FOutBunch* OutBunch, bool Merge
 
 int32 UChannel::SendRawBunch(FOutBunch* OutBunch, bool Merge)
 {
-	if ( Connection->bResendAllDataSinceOpen )
+	if ( Connection->ResendAllDataState != EResendAllDataState::None )
 	{
 		check( OpenPacketId.First != INDEX_NONE );
 		check( OpenPacketId.Last != INDEX_NONE );
@@ -1187,9 +1198,14 @@ int32 UChannel::SendRawBunch(FOutBunch* OutBunch, bool Merge)
 	OutBunch->ReceivedAck = 0;
 	int32 PacketId = Connection->SendRawBunch(*OutBunch, Merge);
 	if( OpenPacketId.First==INDEX_NONE && OpenedLocally )
+	{
 		OpenPacketId = FPacketIdRange(PacketId);
+	}
+
 	if( OutBunch->bClose )
+	{
 		SetClosingFlag();
+	}
 
 	return PacketId;
 }
@@ -1254,6 +1270,7 @@ void UChannel::AddedToChannelPool()
 	bPendingDormancy = false;
 	bPausedUntilReliableACK = false;
 	SentClosingBunch = false;
+	bOpenedForCheckpoint = false;
 	ChIndex = 0;
 	OpenedLocally = false;
 	OpenPacketId = FPacketIdRange();
@@ -2766,7 +2783,7 @@ int64 UActorChannel::ReplicateActor()
 	FReplicationFlags RepFlags;
 
 	// Send initial stuff.
-	if( OpenPacketId.First != INDEX_NONE && !Connection->bResendAllDataSinceOpen )
+	if( OpenPacketId.First != INDEX_NONE && (Connection->ResendAllDataState == EResendAllDataState::None) )
 	{
 		if( !SpawnAcked && OpenAcked )
 		{
@@ -2782,7 +2799,15 @@ int64 UActorChannel::ReplicateActor()
 	}
 	else
 	{
-		RepFlags.bNetInitial = true;
+		if (Connection->ResendAllDataState == EResendAllDataState::SinceCheckpoint)
+		{
+			RepFlags.bNetInitial = !bOpenedForCheckpoint;
+		}
+		else
+		{
+			RepFlags.bNetInitial = true;
+		}
+
 		Bunch.bClose = Actor->bNetTemporary;
 		Bunch.bReliable = true; // Net temporary sends need to be reliable as well to force them to retry
 	}
@@ -2838,7 +2863,7 @@ int64 UActorChannel::ReplicateActor()
 		// The SubObjects
 		WroteSomethingImportant |= Actor->ReplicateSubobjects(this, &Bunch, &RepFlags);
 
-		if (Connection->bResendAllDataSinceOpen)
+		if (Connection->ResendAllDataState != EResendAllDataState::None)
 		{
 			if (WroteSomethingImportant)
 			{
diff --git a/Engine/Source/Runtime/Engine/Private/DataReplication.cpp b/Engine/Source/Runtime/Engine/Private/DataReplication.cpp
--- a/Engine/Source/Runtime/Engine/Private/DataReplication.cpp
+++ b/Engine/Source/Runtime/Engine/Private/DataReplication.cpp
@@ -204,8 +204,9 @@ void FObjectReplicator::InitRecentProperties(uint8* Source)
 
 					SerializeCustomDeltaProperty( Connection, Source, *It, ArrayIdx, DeltaState, NewState, OldState );
 
-					// Store the initial delta state in case we need it for when we're asked to resend all data since channel was first openeded (bResendAllDataSinceOpen)
-					CDOCustomDeltaState.Add( It->RepIndex + ArrayIdx, NewState );
+					// Store the initial delta state in case we need it for when we're asked to resend all data since channel was first opened (bResendAllDataSinceOpen)
+					CDOCustomDeltaState.Add( It->RepIndex + ArrayIdx, NewState ); 
+					CheckpointCustomDeltaState.Add( It->RepIndex + ArrayIdx, NewState );
 				}
 			}
 		}
@@ -340,11 +341,13 @@ void FObjectReplicator::CleanUp()
 
 	// Cleanup custom delta state
 	RecentCustomDeltaState.Empty();
+	CheckpointCustomDeltaState.Empty();
 
 	LifetimeCustomDeltaProperties.Empty();
 	LifetimeCustomDeltaPropertyConditions.Empty();
 
 	RepState = nullptr;
+	CheckpointRepState = nullptr;
 }
 
 void FObjectReplicator::StartReplicating( class UActorChannel * InActorChannel )
@@ -1160,17 +1163,35 @@ void FObjectReplicator::ReplicateCustomDeltaProperties( FNetBitWriter & Bunch, F
 
 		TempBitWriter.Reset();
 
-		if ( Connection->bResendAllDataSinceOpen )
+		if (Connection->ResendAllDataState != EResendAllDataState::None)
 		{
-			// If we are resending data since open, we don't want to affect the current state of channel/replication, so just do the minimum and send the data, and return
-			// In this case, we'll send all of the properties since the CDO, so use the initial CDO delta state
-			TSharedPtr<INetDeltaBaseState>& OldState = CDOCustomDeltaState.FindChecked( RetireIndex );
+			if (Connection->ResendAllDataState == EResendAllDataState::SinceCheckpoint)
+			{
+				TSharedPtr<INetDeltaBaseState> OldState = CheckpointCustomDeltaState.FindChecked(RetireIndex);
+
+				if (!SerializeCustomDeltaProperty(OwningChannelConnection, (void*)Object, It, Index, TempBitWriter, NewState, OldState))
+				{
+					continue;
+				}
 
-			if ( SerializeCustomDeltaProperty( OwningChannelConnection, ( void* )Object, It, Index, TempBitWriter, NewState, OldState ) )
+				// update checkpoint with new state
+				CheckpointCustomDeltaState.Add(RetireIndex, NewState);
+			}
+			else
 			{
-				// Write property header and payload to the bunch
-				WritePropertyHeaderAndPayload( Object, It, NetFieldExportGroup, Bunch, TempBitWriter );
+				// If we are resending data since open, we don't want to affect the current state of channel/replication, so just do the minimum and send the data, and return
+				// In this case, we'll send all of the properties since the CDO, so use the initial CDO delta state
+				TSharedPtr<INetDeltaBaseState>& OldState = CDOCustomDeltaState.FindChecked(RetireIndex);
+
+				if (!SerializeCustomDeltaProperty(OwningChannelConnection, (void*)Object, It, Index, TempBitWriter, NewState, OldState))
+				{
+					continue;
+				}
 			}
+
+			// Write property header and payload to the bunch
+			WritePropertyHeaderAndPayload( Object, It, NetFieldExportGroup, Bunch, TempBitWriter );
+
 			continue;
 		}
 
@@ -1242,7 +1263,7 @@ bool FObjectReplicator::ReplicateProperties( FOutBunch & Bunch, FReplicationFlag
 	//		and all the work could just be done in a single place.
 
 	// Update change list (this will re-use work done by previous connections)
-	FSendingRepState* SendingRepState = RepState->GetSendingRepState();
+	FSendingRepState* SendingRepState = ((Connection->ResendAllDataState == EResendAllDataState::SinceCheckpoint) && CheckpointRepState.IsValid()) ? CheckpointRepState->GetSendingRepState() : RepState->GetSendingRepState();
 	RepLayout->UpdateChangelistMgr(*ChangelistMgr, SendingRepState, Object, Connection->Driver->ReplicationFrame, RepFlags, OwningChannel->bForceCompareProperties);
 
 	// Replicate properties in the layout
@@ -1251,7 +1272,7 @@ bool FObjectReplicator::ReplicateProperties( FOutBunch & Bunch, FReplicationFlag
 	// Replicate all the custom delta properties (fast arrays, etc)
 	ReplicateCustomDeltaProperties(Writer, RepFlags);
 
-	if ( OwningChannelConnection->bResendAllDataSinceOpen )
+	if ( Connection->ResendAllDataState != EResendAllDataState::None )
 	{
 		// If we are resending data since open, we don't want to affect the current state of channel/replication, so just send the data, and return
 		const bool WroteImportantData = Writer.GetNumBits() != 0;
@@ -1259,6 +1280,12 @@ bool FObjectReplicator::ReplicateProperties( FOutBunch & Bunch, FReplicationFlag
 		if ( WroteImportantData )
 		{
 			OwningChannel->WriteContentBlockPayload( Object, Bunch, bHasRepLayout, Writer );
+
+			if (Connection->ResendAllDataState == EResendAllDataState::SinceCheckpoint)
+			{
+				UpdateCheckpoint();
+			}
+
 			return true;
 		}
 
@@ -1905,6 +1932,28 @@ void FObjectReplicator::WritePropertyHeaderAndPayload(
 	NETWORK_PROFILER( GNetworkProfiler.TrackWritePropertyHeader( Property, HeaderBits, nullptr ) );
 }
 
+void FObjectReplicator::UpdateCheckpoint()
+{
+	TArray<uint16> CheckpointChangelist;
+
+	if (CheckpointRepState.IsValid())
+	{
+		CheckpointChangelist = MoveTemp(CheckpointRepState->GetSendingRepState()->LifetimeChangelist);
+	}
+	else
+	{
+		CheckpointChangelist = RepState->GetSendingRepState()->LifetimeChangelist;
+	}
+
+	// Update rep state
+	TSharedPtr<FRepChangedPropertyTracker> RepChangedPropertyTracker = Connection->Driver->FindOrCreateRepChangedPropertyTracker( GetObject() );
+
+	CheckpointRepState = RepLayout->CreateRepState((const uint8*)GetObject(), RepChangedPropertyTracker, ECreateRepStateFlags::SkipCreateReceivingState);
+
+	// Keep current set of changed properties
+	CheckpointRepState->GetSendingRepState()->LifetimeChangelist = MoveTemp(CheckpointChangelist);
+}
+
 FScopedActorRoleSwap::FScopedActorRoleSwap(AActor* InActor)
 	: Actor(InActor)
 {
@@ -1926,6 +1975,4 @@ FScopedActorRoleSwap::~FScopedActorRoleSwap()
 	{
 		Actor->SwapRoles();
 	}
-}
-
-
+}
\ No newline at end of file
diff --git a/Engine/Source/Runtime/Engine/Private/DemoNetDriver.cpp b/Engine/Source/Runtime/Engine/Private/DemoNetDriver.cpp
--- a/Engine/Source/Runtime/Engine/Private/DemoNetDriver.cpp
+++ b/Engine/Source/Runtime/Engine/Private/DemoNetDriver.cpp
@@ -70,6 +70,7 @@ static FAutoConsoleVariableRef CVarDemoSaveRollbackActorState( TEXT( "demo.SaveR
 
 static TAutoConsoleVariable<int32> CVarWithLevelStreamingFixes(TEXT("demo.WithLevelStreamingFixes"), 0, TEXT("If 1, provides fixes for level streaming (but breaks backwards compatibility)."));
 static TAutoConsoleVariable<int32> CVarWithDemoTimeBurnIn(TEXT("demo.WithTimeBurnIn"), 0, TEXT("If true, adds an on screen message with the current DemoTime and Changelist."));
+static TAutoConsoleVariable<int32> CVarWithDeltaCheckpoints(TEXT("demo.WithDeltaCheckpoints"), 0, TEXT("If true, record checkpoints as a delta from the previous checkpoint."));
 
 #if !(UE_BUILD_SHIPPING || UE_BUILD_TEST)
 static TAutoConsoleVariable<int32> CVarDemoForceFailure( TEXT( "demo.ForceFailure" ), 0, TEXT( "" ) );
@@ -145,12 +146,12 @@ namespace DemoNetDriverRecordingPrivate
 }
 
 // Helps manage packets, and any associations with streaming levels or exported GUIDs / fields.
-class FScopedPacketManager
+class FScopedPacketManager : public FNoncopyable
 {
 public:
 	FScopedPacketManager(UDemoNetConnection& InConnection, const uint32 InSeenLevelIndex):
 		Connection(InConnection),
-		Packets(Connection.bResendAllDataSinceOpen ? InConnection.QueuedCheckpointPackets : InConnection.QueuedDemoPackets),
+		Packets((Connection.ResendAllDataState != EResendAllDataState::None) ? InConnection.QueuedCheckpointPackets : InConnection.QueuedDemoPackets),
 		SeenLevelIndex(InSeenLevelIndex)
 	{
 		FlushNetChecked(Connection);
@@ -205,8 +206,9 @@ public:
  * This relies on serialization always used a fixed number of bytes for primitive types,
  * and Sane implementations of Seek and Tell.
  */
-struct FScopedStoreArchiveOffset
+class FScopedStoreArchiveOffset : public FNoncopyable
 {
+public:
 	typedef UDemoNetDriver::FArchivePos FArchivePos;
 
 	FScopedStoreArchiveOffset(FArchive& InAr):
@@ -233,6 +235,32 @@ private:
 	const FArchivePos StartPosition;
 };
 
+class FScopedIgnoreDeletedGuids : public FNoncopyable
+{
+public:
+	FScopedIgnoreDeletedGuids(UNetConnection* InConnection, TSet<FNetworkGUID>& DeletedActorGuids)
+		: ServerConnection(InConnection)
+	{
+		if (ServerConnection)
+		{
+			ServerConnection->IgnoredGuids = MoveTemp(DeletedActorGuids);
+			ServerConnection->IgnoredChannels.Empty();
+		}
+	}
+
+	~FScopedIgnoreDeletedGuids()
+	{
+		if (ServerConnection)
+		{
+			ServerConnection->IgnoredGuids.Empty();
+			ServerConnection->IgnoredChannels.Empty();
+		}
+	}
+
+private:
+	UNetConnection* ServerConnection;
+};
+
 class FJumpToLiveReplayTask : public FQueuedReplayTask
 {
 public:
@@ -316,8 +344,10 @@ public:
 		// Clamp time
 		Driver->DemoCurrentTime = FMath::Clamp(Driver->DemoCurrentTime, 0.0f, Driver->DemoTotalTime - 0.01f);
 
+		EReplayCheckpointType CheckpointType = Driver->HasDeltaCheckpoints() ? EReplayCheckpointType::Delta : EReplayCheckpointType::Full;
+
 		// Tell the streamer to start going to this time
-		Driver->ReplayStreamer->GotoTimeInMS(Driver->DemoCurrentTime * 1000, FGotoCallback::CreateSP(this, &FGotoTimeInSecondsTask::CheckpointReady));
+		Driver->ReplayStreamer->GotoTimeInMS(Driver->DemoCurrentTime * 1000, FGotoCallback::CreateSP(this, &FGotoTimeInSecondsTask::CheckpointReady), CheckpointType);
 
 		// Pause channels while we wait (so the world is paused while we wait for the new stream location to load)
 		Driver->PauseChannels( true );
@@ -452,7 +482,9 @@ public:
 			const float LastProcessedPacketTime = FPendingTaskHelper::GetLastProcessedPacketTime(Driver.Get());
 			GotoTime = LastProcessedPacketTime * 1000;
 
-			Driver->ReplayStreamer->GotoTimeInMS(GotoTime, FGotoCallback::CreateSP(this, &FFastForwardLevelsTask::CheckpointReady));
+			EReplayCheckpointType CheckpointType = Driver->HasDeltaCheckpoints() ? EReplayCheckpointType::Delta : EReplayCheckpointType::Full;
+
+			Driver->ReplayStreamer->GotoTimeInMS(GotoTime, FGotoCallback::CreateSP(this, &FFastForwardLevelsTask::CheckpointReady), CheckpointType);
 
 			// Pause channels while we wait (so the world is paused while we wait for the new stream location to load)
 			Driver->PauseChannels(true);
@@ -503,7 +535,7 @@ public:
 
 		if (!Result.WasSuccessful())
 		{
-			UE_LOG( LogDemo, Warning, TEXT( "FFastForwardLevelsTask::CheckpointReady: Faled to get checkpoint." ) );
+			UE_LOG( LogDemo, Warning, TEXT( "FFastForwardLevelsTask::CheckpointReady: Failed to get checkpoint." ) );
 		}
 	}
 
@@ -514,7 +546,7 @@ private:
 	TOptional<FGotoResult> GotoResult;
 };
 
-class FScopedForceUnicodeInArchive
+class FScopedForceUnicodeInArchive : public FNoncopyable
 {
 public:
 	FScopedForceUnicodeInArchive(FArchive& InArchive)
@@ -562,6 +594,7 @@ UDemoNetDriver::UDemoNetDriver(const FObjectInitializer& ObjectInitializer)
 	bRecordMapChanges = false;
 	bIsWaitingForHeaderDownload = false;
 	bIsWaitingForStream = false;
+	MaxArchiveReadPos = 0;
 
 	LevelIntervals.Reserve(512);
 
@@ -860,7 +893,16 @@ bool UDemoNetDriver::ReadPlaybackDemoHeader(FString& Error)
 	}
 	
 	// Check whether or not we need to process streaming level fixes.
-	bHasLevelStreamingFixes = !!(PlaybackDemoHeader.HeaderFlags & EReplayHeaderFlags::HasStreamingFixes);
+	bHasLevelStreamingFixes = EnumHasAnyFlags(PlaybackDemoHeader.HeaderFlags, EReplayHeaderFlags::HasStreamingFixes);
+	// Or delta checkpoints
+	bHasDeltaCheckpoints = EnumHasAnyFlags(PlaybackDemoHeader.HeaderFlags, EReplayHeaderFlags::DeltaCheckpoints);
+
+	if (bHasDeltaCheckpoints && !ReplayStreamer->IsCheckpointTypeSupported(EReplayCheckpointType::Delta))
+	{
+		UE_LOG(LogDemo, Error, TEXT("UDemoNetDriver::InitConnect: Replay has delta checkpoints but streamer does not support them."));
+		GameInstance->HandleDemoPlaybackFailure(EDemoPlayFailure::Generic, FString(EDemoPlayFailure::ToString(EDemoPlayFailure::Generic)));
+		return false;
+	}
 
 	// Set network version on connection
 	ServerConnection->EngineNetworkProtocolVersion = PlaybackDemoHeader.EngineNetworkProtocolVersion;
@@ -973,6 +1015,7 @@ bool UDemoNetDriver::InitListen( FNetworkNotify* InNotify, FURL& ListenURL, bool
 	// We'll only check these CVars here, because we don't want to break the replay if they change part way through recording.
 	// During playback the CVars won't be used. Instead, we'll rely on the DemoPacketHeader value.
 	bHasLevelStreamingFixes = !!CVarWithLevelStreamingFixes.GetValueOnAnyThread();
+	bHasDeltaCheckpoints = !!CVarWithDeltaCheckpoints.GetValueOnAnyThread() && ReplayStreamer->IsCheckpointTypeSupported(EReplayCheckpointType::Delta);
 
 	// Recording, local machine is server, demo stream acts "as if" it's a client.
 	UDemoNetConnection* Connection = NewObject<UDemoNetConnection>();
@@ -1192,11 +1235,17 @@ bool UDemoNetDriver::WriteNetworkDemoHeader(FString& Error)
 	{
 		DemoHeader.HeaderFlags |= EReplayHeaderFlags::ClientRecorded;
 	}
+
 	if (HasLevelStreamingFixes())
 	{
 		DemoHeader.HeaderFlags |= EReplayHeaderFlags::HasStreamingFixes;
 	}
 
+	if (HasDeltaCheckpoints())
+	{
+		DemoHeader.HeaderFlags |= EReplayHeaderFlags::DeltaCheckpoints;
+	}
+
 	DemoHeader.Guid = FGuid::NewGuid();
 
 	// Write the header
@@ -1572,7 +1621,7 @@ bool UDemoNetDriver::DemoReplicateActor(AActor* Actor, UNetConnection* Connectio
 			// Close the channel if this actor shouldn't have one
 			if (!bShouldHaveChannel)
 			{
-				if (!Connection->bResendAllDataSinceOpen)		// Don't close the channel if we're forcing them to re-open for checkpoints
+				if (Connection->ResendAllDataState == EResendAllDataState::None)		// Don't close the channel if we're forcing them to re-open for checkpoints
 				{
 					Channel->Close(EChannelCloseReason::Destroyed);
 				}
@@ -1580,7 +1629,7 @@ bool UDemoNetDriver::DemoReplicateActor(AActor* Actor, UNetConnection* Connectio
 		}
 	}
 
-	if ( bMustReplicate && Connection->Driver->OutBunches == OriginalOutBunches )
+	if ( bMustReplicate && !HasDeltaCheckpoints() && Connection->Driver->OutBunches == OriginalOutBunches )
 	{
 		UE_LOG( LogDemo, Error, TEXT( "DemoReplicateActor: bMustReplicate is true but nothing was sent: %s" ), Actor ? *Actor->GetName() : TEXT( "NULL" ) );
 	}
@@ -1593,15 +1642,22 @@ void UDemoNetDriver::SerializeGuidCache(TSharedPtr<FNetGUIDCache> InGuidCache, F
 	int32 NumValues = 0;
 	int32 UnloadedValues = 0;
 
-	int64 CountPos = CheckpointArchive->Tell();
+	FArchivePos CountPos = CheckpointArchive->Tell();
 
 	*CheckpointArchive << NumValues;
 
+	const bool bDeltaCheckpoint = HasDeltaCheckpoints();
+
 	for ( auto It = InGuidCache->ObjectLookup.CreateIterator(); It; ++It )
 	{
 		FNetworkGUID& NetworkGUID = It.Key();
 		FNetGuidCacheObject& CacheObject = It.Value();
 
+		if (bDeltaCheckpoint && !CacheObject.bDirtyForReplay)
+		{
+			continue;
+		}
+
 		if (NetworkGUID.IsValid())
 		{
 			const UObject* Object = CacheObject.Object.Get();
@@ -1629,6 +1685,8 @@ void UDemoNetDriver::SerializeGuidCache(TSharedPtr<FNetGUIDCache> InGuidCache, F
 
 				*CheckpointArchive << Flags;
 
+				CacheObject.bDirtyForReplay = false;
+
 				++NumValues;
 
 				const bool bUnloaded = Object == nullptr || !Object->IsNameStableForNetworking();
@@ -1640,7 +1698,7 @@ void UDemoNetDriver::SerializeGuidCache(TSharedPtr<FNetGUIDCache> InGuidCache, F
 		}
 	}
 
-	int64 Pos = CheckpointArchive->Tell();
+	FArchivePos Pos = CheckpointArchive->Tell();
 	CheckpointArchive->Seek(CountPos);
 	*CheckpointArchive << NumValues;
 	CheckpointArchive->Seek(Pos);
@@ -1682,6 +1740,8 @@ void UDemoNetDriver::SaveCheckpoint()
 
 	check( CheckpointSaveContext.CheckpointSaveState == ECheckpointSaveState_Idle );
 	
+	const bool bDeltaCheckpoint = HasDeltaCheckpoints();
+
 	if (HasLevelStreamingFixes())
 	{
 		SCOPED_NAMED_EVENT(UDemoNetDriver_ReplayLevelSortAndAssign, FColor::Purple);
@@ -1711,9 +1771,17 @@ void UDemoNetDriver::SaveCheckpoint()
 					if (Actor)
 					{
 						// Validate that we do not pickup any extra actors
-						if (AllObjectsSet.Find(Actor) != nullptr)
+						const TSharedPtr<FNetworkObjectInfo>* ObjectInfo = AllObjectsSet.Find(Actor);
+						if (ObjectInfo != nullptr && ObjectInfo->IsValid())
 						{
-							ActorArray.Add( { ChannelPair.Key, Actor->GetOuter() } );
+							FNetworkObjectInfo* Info = ObjectInfo->Get();
+
+							if (!bDeltaCheckpoint || Info->bDirtyForReplay)
+							{
+								ActorArray.Add({ ChannelPair.Key, Actor->GetOuter() });
+
+								Info->bDirtyForReplay = false;
+							}
 						}
 					}
 				}
@@ -1751,7 +1819,12 @@ void UDemoNetDriver::SaveCheckpoint()
 
 			if ( ClientConnections[0]->FindActorChannelRef( Actor ) )
 			{
-				CheckpointSaveContext.PendingCheckpointActors.Add( { Actor, -1 } );
+				if (!bDeltaCheckpoint || ObjectInfo->bDirtyForReplay)
+				{
+					CheckpointSaveContext.PendingCheckpointActors.Add({ Actor, -1 });
+
+					ObjectInfo->bDirtyForReplay = false;
+				}
 			}
 		}
 	}
@@ -1840,6 +1913,8 @@ void UDemoNetDriver::TickCheckpoint()
 		// Save package map ack status in case we export stuff during the checkpoint (so we can restore the connection back to what it was before we saved the checkpoint)
 		PackageMapClient->OverridePackageMapExportAckStatus( &CheckpointSaveContext.CheckpointAckState );
 
+		const bool bDeltaCheckpoint = HasDeltaCheckpoints();
+
 		while (bExecuteNextState && (CheckpointSaveContext.CheckpointSaveState != ECheckpointSaveState_Finalize) && !(Params.CheckpointMaxUploadTimePerFrame > 0 && CurrentTime - Params.StartCheckpointTime > Params.CheckpointMaxUploadTimePerFrame))
 		{
 			switch (CheckpointSaveContext.CheckpointSaveState)
@@ -1864,8 +1939,7 @@ void UDemoNetDriver::TickCheckpoint()
 
 					{
 						// Re-use the existing connection to record all properties that have changed since channels were first opened
-						// Set bResendAllDataSinceOpen to true to signify that we want to do this
-						TGuardValue<bool> ResendAllData(ClientConnection->bResendAllDataSinceOpen, true);
+						TGuardValue<EResendAllDataState> ResendAllData(ClientConnection->ResendAllDataState, bDeltaCheckpoint ? EResendAllDataState::SinceCheckpoint : EResendAllDataState::SinceOpen);
 
 						// Can't use conditionally create here, because NumActorsToProcess will be empty when HasLevelStreamingFixes is false.
 						TUniquePtr<FScopedPacketManager> PacketManager;
@@ -1935,8 +2009,19 @@ void UDemoNetDriver::TickCheckpoint()
 
 						*CheckpointArchive << CurrentLevelIndex;
 
-						// Save deleted startup actors	
-						*CheckpointArchive << DeletedNetStartupActors;
+						// Save deleted startup actors
+						if (bDeltaCheckpoint)
+						{
+							*CheckpointArchive << DeltaDeletedNetStartupActors;
+							DeltaDeletedNetStartupActors.Empty();
+
+							*CheckpointArchive << DeltaDeletedActorGuids;
+							DeltaDeletedActorGuids.Empty();
+						}
+						else
+						{
+							*CheckpointArchive << DeletedNetStartupActors;
+						}
 
 						CheckpointSaveContext.CheckpointSaveState = ECheckpointSaveState_SerializeGuidCache;
 					}
@@ -1968,7 +2053,14 @@ void UDemoNetDriver::TickCheckpoint()
 						SCOPED_NAMED_EVENT(UDemoNetDriver_SerializeNetFieldExportGroupMap, FColor::Green);
 
 						// Save the compatible rep layout map
-						PackageMapClient->SerializeNetFieldExportGroupMap( *CheckpointArchive );
+						if (bDeltaCheckpoint)
+						{
+							PackageMapClient->SerializeNetFieldExportDelta(*CheckpointArchive);
+						}
+						else
+						{
+							PackageMapClient->SerializeNetFieldExportGroupMap(*CheckpointArchive);
+						}
 
 						CheckpointSaveContext.CheckpointSaveState = ECheckpointSaveState_SerializeDemoFrameFromQueuedDemoPackets;
 					}
@@ -2459,6 +2551,8 @@ void UDemoNetDriver::TickDemoRecordFrame( float DeltaSeconds )
 			FDemoActorPriority DemoActorPriority;
 			FActorPriority& ActorPriority = DemoActorPriority.ActorPriority;
 
+			const bool bDeltaCheckpoint = HasDeltaCheckpoints();
+
 			for (const TSharedPtr<FNetworkObjectInfo>& ObjectInfo : ActiveObjectSet)
 			{
 				FNetworkObjectInfo* ActorInfo = ObjectInfo.Get();
@@ -2539,6 +2633,8 @@ void UDemoNetDriver::TickDemoRecordFrame( float DeltaSeconds )
 
 					PrioritizedActors.Add(DemoActorPriority);
 
+					ActorInfo->bDirtyForReplay = bDeltaCheckpoint;
+
 					if (bIsRelevant)
 					{
 						ActorInfo->LastNetUpdateTime = Time;
@@ -3034,7 +3130,7 @@ bool UDemoNetDriver::ReadDemoFrameIntoPlaybackPackets( FArchive& Ar, TArray<FPla
 
 		const EReadPacketMode ReadPacketMode = bAppendPackets ? EReadPacketMode::Default : EReadPacketMode::SkipData;
 
-		while (true)
+		while ((MaxArchiveReadPos == 0) || (Ar.Tell() < MaxArchiveReadPos))
 		{
 			if (HasLevelStreamingFixes())
 			{
@@ -3652,7 +3748,7 @@ void UDemoNetDriver::TickDemoPlayback( float DeltaSeconds )
 	}
 
 	// Advance demo time by seconds passed if we're not paused
-	if ( World->GetWorldSettings()->Pauser == nullptr )
+	if ( World && World->GetWorldSettings() && World->GetWorldSettings()->Pauser == nullptr )
 	{
 		DemoCurrentTime += DeltaSeconds;
 	}
@@ -3990,8 +4086,16 @@ void UDemoNetDriver::RespawnNecessaryNetStartupActors(TArray<AActor*>& SpawnedAc
 			continue;
 		}
 
+		const FString LevelPackageName = GetLevelPackageName(*RollbackActor.Level);
+
 		// skip rollback actors in streamed out levels (pending gc)
-		if (!LevelStatusesByName.Contains(GetLevelPackageName(*RollbackActor.Level)))
+		if (!LevelStatusesByName.Contains(LevelPackageName))
+		{
+			continue;
+		}
+
+		FLevelStatus& LevelStatus = GetLevelStatus(LevelPackageName);
+		if (!LevelStatus.bIsReady)
 		{
 			continue;
 		}
@@ -4311,16 +4415,34 @@ bool UDemoNetDriver::FastForwardLevels(const FGotoResult& GotoResult)
 
 			TGuardValue<bool> LoadingCheckpointGuard(bIsLoadingCheckpoint, true);
 
-			FArchivePos PacketOffset = 0;
-			*CheckpointArchive << PacketOffset;
-			CheckpointArchive->Seek(PacketOffset + CheckpointArchive->Tell());
+			const bool bDeltaCheckpoint = HasDeltaCheckpoints();
 
-			if (!ReadPacketsHelper.ReadPackets(*CheckpointArchive) && ReadPacketsHelper.IsError())
+			do 
 			{
-				UE_LOG(LogDemo, Warning, TEXT("UDemoNetDriver::FastForwardLevels: Failed to read packets from Checkpoint."));
-				NotifyDemoPlaybackFailure(EDemoPlayFailure::Serialization);
-				return false;
-			}
+				FArchivePos MaxArchivePos = 0;
+
+				if (bDeltaCheckpoint)
+				{
+					uint32 CheckpointSize = 0;
+					*CheckpointArchive << CheckpointSize;
+
+					MaxArchivePos = CheckpointArchive->Tell() + CheckpointSize;
+				}
+
+				TGuardValue<int64> MaxArchivePosGuard(MaxArchiveReadPos, MaxArchivePos);
+
+				FArchivePos PacketOffset = 0;
+				*CheckpointArchive << PacketOffset;
+				CheckpointArchive->Seek(PacketOffset + CheckpointArchive->Tell());
+
+				if (!ReadPacketsHelper.ReadPackets(*CheckpointArchive) && ReadPacketsHelper.IsError())
+				{
+					UE_LOG(LogDemo, Warning, TEXT("UDemoNetDriver::FastForwardLevels: Failed to read packets from Checkpoint."));
+					NotifyDemoPlaybackFailure(EDemoPlayFailure::Generic);
+					return false;
+				}
+			} 
+			while (!CheckpointArchive->IsError() && (CheckpointArchive->Tell() < CheckpointArchive->TotalSize()));
 		}
 
 		// Next, read in streaming data (if any is available)
@@ -4404,6 +4526,8 @@ bool UDemoNetDriver::FastForwardLevels(const FGotoResult& GotoResult)
 			TWeakObjectPtr<UNetConnection> Connection;
 		} ScopedIgnoreChannels(ServerConnection);
 
+		FScopedIgnoreDeletedGuids IgnoreGuids(ServerConnection, DeletedActorGuids);
+
 		// Process all the packets we need.
 		for (FPlaybackPacket& Packet : ReadPacketsHelper.Packets)
 		{
@@ -4514,6 +4638,17 @@ bool UDemoNetDriver::LoadCheckpoint(const FGotoResult& GotoResult)
 
 	int32 LevelForCheckpoint = 0;
 
+	const bool bDeltaCheckpoint = HasDeltaCheckpoints();
+
+	if (bDeltaCheckpoint)
+	{
+		if (GotoCheckpointArchive->TotalSize() > 0)
+		{
+			uint32 CheckpointSize = 0;
+			*GotoCheckpointArchive << CheckpointSize;
+		}
+	}
+
 	if (HasLevelStreamingFixes())
 	{
 		// Make sure to read the packet offset, even though we won't use it here.
@@ -4776,6 +4911,7 @@ bool UDemoNetDriver::LoadCheckpoint(const FGotoResult& GotoResult)
 	{
 		// Make sure this is empty so that RespawnNecessaryNetStartupActors will respawn them
 		DeletedNetStartupActors.Empty();
+		DeletedActorGuids.Empty();
 
 		// Re-create all startup actors that were destroyed but should exist beyond this point
 		TArray<AActor*> SpawnedActors;
@@ -4800,14 +4936,108 @@ bool UDemoNetDriver::LoadCheckpoint(const FGotoResult& GotoResult)
 		return true;
 	}
 
-	// Load net startup actors that need to be destroyed
-	if ( PlaybackDemoHeader.Version >= HISTORY_DELETED_STARTUP_ACTORS )
+	GotoCheckpointArchive->Seek(0);
+
+	DeletedNetStartupActors.Empty();
+	DeletedActorGuids.Empty();
+
+	do
 	{
-		*GotoCheckpointArchive << DeletedNetStartupActors;
-	}
+		FArchivePos MaxArchivePos = 0;
+
+		if (bDeltaCheckpoint)
+		{
+			uint32 CheckpointSize = 0;
+			*GotoCheckpointArchive << CheckpointSize;
+
+			MaxArchivePos = GotoCheckpointArchive->Tell() + CheckpointSize;
+		}
+
+		TGuardValue<int64> MaxArchivePosGuard(MaxArchiveReadPos, MaxArchivePos);
+			
+		if (HasLevelStreamingFixes())
+		{
+			FArchivePos PacketOffset = 0;
+			*GotoCheckpointArchive << PacketOffset;
+		}
+
+		if (PlaybackDemoHeader.Version >= HISTORY_MULTIPLE_LEVELS)
+		{
+			int32 LevelIndex;
+			*GotoCheckpointArchive << LevelIndex;
+		}
+
+		// Load net startup actors that need to be destroyed
+		if ( PlaybackDemoHeader.Version >= HISTORY_DELETED_STARTUP_ACTORS )
+		{
+			if (bDeltaCheckpoint)
+			{
+				TSet<FString> DeltaActors;
+				*GotoCheckpointArchive << DeltaActors;
+
+				DeletedNetStartupActors.Append(DeltaActors);
+
+				TSet<FNetworkGUID> DeltaGuids;
+				*GotoCheckpointArchive << DeltaGuids;
+
+				DeletedActorGuids.Append(DeltaGuids);
+			}
+			else
+			{
+				DeletedNetStartupActors.Empty();
+				DeletedActorGuids.Empty();
+
+				*GotoCheckpointArchive << DeletedNetStartupActors;
+			}
+		}
+
+		int32 NumValues = 0;
+		*GotoCheckpointArchive << NumValues;
+
+		for ( int32 i = 0; i < NumValues; i++ )
+		{
+			FNetworkGUID Guid;
+
+			*GotoCheckpointArchive << Guid;
+
+			FNetGuidCacheObject CacheObject;
+
+			FString PathName;
+
+			*GotoCheckpointArchive << CacheObject.OuterGUID;
+			*GotoCheckpointArchive << PathName;
+			*GotoCheckpointArchive << CacheObject.NetworkChecksum;
+
+			// Remap the pathname to handle client-recorded replays
+			GEngine->NetworkRemapPath(this, PathName, true);
+
+			CacheObject.PathName = FName( *PathName );
 
+			uint8 Flags = 0;
+			*GotoCheckpointArchive << Flags;
+
+			CacheObject.bNoLoad = ( Flags & ( 1 << 0 ) ) ? true : false;
+			CacheObject.bIgnoreWhenMissing = ( Flags & ( 1 << 1 ) ) ? true : false;		
+
+			GuidCache->ObjectLookup.Add( Guid, CacheObject );
+		}
+
+		// Read in the compatible rep layouts in this checkpoint
+		if (bDeltaCheckpoint)
+		{
+			CastChecked<UPackageMapClient>(ServerConnection->PackageMap)->SerializeNetFieldExportDelta(*GotoCheckpointArchive);
+		}
+		else
+		{
+			CastChecked<UPackageMapClient>(ServerConnection->PackageMap)->SerializeNetFieldExportGroupMap(*GotoCheckpointArchive);
+		}
+
+		ReadDemoFrameIntoPlaybackPackets( *GotoCheckpointArchive );
+	}
+	while (!GotoCheckpointArchive->IsError() && (GotoCheckpointArchive->Tell() < GotoCheckpointArchive->TotalSize()));
+	
 	// Destroy startup actors that shouldn't exist past this checkpoint
-	for ( FActorIterator It( GetWorld() ); It; ++It )
+	for ( FActorIterator It( World ); It; ++It )
 	{
 		const FString FullName = It->GetFullName();
 		if ( DeletedNetStartupActors.Contains( FullName ) )
@@ -4826,50 +5056,14 @@ bool UDemoNetDriver::LoadCheckpoint(const FGotoResult& GotoResult)
 			UE_LOG(LogDemo, Verbose, TEXT("LoadCheckpoint: deleting startup actor %s"), *FullName);
 
 			// Delete the actor
-			GetWorld()->DestroyActor( *It, true );
+			World->DestroyActor( *It, true );
 		}
 	}
 
 	// Re-create all startup actors that were destroyed but should exist beyond this point
 	TArray<AActor*> SpawnedActors;
 	RespawnNecessaryNetStartupActors(SpawnedActors);
-
-	int32 NumValues = 0;
-	*GotoCheckpointArchive << NumValues;
-
-	for ( int32 i = 0; i < NumValues; i++ )
-	{
-		FNetworkGUID Guid;
-		
-		*GotoCheckpointArchive << Guid;
 		
-		FNetGuidCacheObject CacheObject;
-
-		FString PathName;
-
-		*GotoCheckpointArchive << CacheObject.OuterGUID;
-		*GotoCheckpointArchive << PathName;
-		*GotoCheckpointArchive << CacheObject.NetworkChecksum;
-
-		// Remap the pathname to handle client-recorded replays
-		GEngine->NetworkRemapPath(this, PathName, true);
-
-		CacheObject.PathName = FName( *PathName );
-
-		uint8 Flags = 0;
-		*GotoCheckpointArchive << Flags;
-
-		CacheObject.bNoLoad = ( Flags & ( 1 << 0 ) ) ? true : false;
-		CacheObject.bIgnoreWhenMissing = ( Flags & ( 1 << 1 ) ) ? true : false;		
-
-		GuidCache->ObjectLookup.Add( Guid, CacheObject );
-	}
-
-	// Read in the compatible rep layouts in this checkpoint
-	( ( UPackageMapClient* )ServerConnection->PackageMap )->SerializeNetFieldExportGroupMap( *GotoCheckpointArchive );
-
-	ReadDemoFrameIntoPlaybackPackets( *GotoCheckpointArchive );
-
 	if ( PlaybackPackets.Num() > 0 )
 	{
 		DemoCurrentTime = PlaybackPackets.Last().TimeSeconds;
@@ -4891,7 +5085,10 @@ bool UDemoNetDriver::LoadCheckpoint(const FGotoResult& GotoResult)
 		CleanupOutstandingRewindActors();
 	}
 
-	ProcessAllPlaybackPackets();
+	{
+		FScopedIgnoreDeletedGuids IgnoreGuids(ServerConnection, DeletedActorGuids);
+		ProcessAllPlaybackPackets();
+	}
 
 	bDemoPlaybackDone = false;
 	bIsLoadingCheckpoint = false;
@@ -4926,7 +5123,7 @@ bool UDemoNetDriver::IsSavingCheckpoint() const
 		UNetConnection* const NetConnection = ClientConnections[0];
 		if (NetConnection)
 		{
-			return NetConnection->bResendAllDataSinceOpen;
+			return (NetConnection->ResendAllDataState != EResendAllDataState::None);
 		}
 	}
 
@@ -5152,7 +5349,7 @@ void UDemoNetConnection::LowLevelSend(void* Data, int32 CountBits, FOutPacketTra
 
 	TrackSendForProfiler( Data, CountBytes );
 
-	if ( bResendAllDataSinceOpen )
+	if ( ResendAllDataState != EResendAllDataState::None )
 	{
 		// This path is only active for a checkpoint saving out, we need to queue in separate list
 		new( QueuedCheckpointPackets )FQueuedDemoPacket( ( uint8* )Data, CountBits, Traits );
@@ -5413,12 +5610,15 @@ void UDemoNetDriver::NotifyActorDestroyed( AActor* Actor, bool IsSeamlessTravel
 		}
 		else
 		{
+			const FString FullName = Actor->GetFullName();
+
 			// This was deleted due to a game interaction, which isn't supported for Rewindable actors (while recording).
 			// However, since the actor is going to be deleted imminently, we need to track it.
-			UE_CLOG(bActorRewindable, LogDemo, Warning, TEXT("Replay Rewindable Actor destroyed during recording. Replay may show artifacts (%s)"), *Actor->GetFullName());
+			UE_CLOG(bActorRewindable, LogDemo, Warning, TEXT("Replay Rewindable Actor destroyed during recording. Replay may show artifacts (%s)"), *FullName);
 
-			UE_LOG(LogDemo, VeryVerbose, TEXT("NotifyActyorDestroyed: adding actor to deleted startup list: %s"), *Actor->GetFullName());
-			DeletedNetStartupActors.Add( Actor->GetFullName() );
+			UE_LOG(LogDemo, VeryVerbose, TEXT("NotifyActyorDestroyed: adding actor to deleted startup list: %s"), *FullName);
+			DeletedNetStartupActors.Add( FullName );
+			DeltaDeletedNetStartupActors.Add( FullName );
 
 			FNetworkGUID NetGUID = GuidCache->NetGUIDLookup.FindRef(Actor);
 			if (NetGUID.IsValid())
@@ -5428,6 +5628,16 @@ void UDemoNetDriver::NotifyActorDestroyed( AActor* Actor, bool IsSeamlessTravel
 		}
 	}
 
+	if (bIsRecording && !IsSeamlessTravel)
+	{
+		FNetworkGUID NetGUID = GuidCache->NetGUIDLookup.FindRef(Actor);
+		if (NetGUID.IsValid() && !NetGUID.IsDefault())
+		{
+			DeletedActorGuids.Add(NetGUID);
+			DeltaDeletedActorGuids.Add(NetGUID);
+		}
+	}
+
 	TUniquePtr<FScopedPacketManager> PacketManager(ConditionallyCreatePacketManager(*Actor->GetLevel()));
 	Super::NotifyActorDestroyed( Actor, IsSeamlessTravel );
 }
diff --git a/Engine/Source/Runtime/Engine/Private/NetConnection.cpp b/Engine/Source/Runtime/Engine/Private/NetConnection.cpp
--- a/Engine/Source/Runtime/Engine/Private/NetConnection.cpp
+++ b/Engine/Source/Runtime/Engine/Private/NetConnection.cpp
@@ -151,7 +151,10 @@ UNetConnection::UNetConnection(const FObjectInitializer& ObjectInitializer)
 ,	InitInReliable		( 0 )
 ,	EngineNetworkProtocolVersion( FNetworkVersion::GetEngineNetworkProtocolVersion() )
 ,	GameNetworkProtocolVersion( FNetworkVersion::GetGameNetworkProtocolVersion() )
+PRAGMA_DISABLE_DEPRECATION_WARNINGS
 ,	bResendAllDataSinceOpen( false )
+PRAGMA_ENABLE_DEPRECATION_WARNINGS
+,	ResendAllDataState( EResendAllDataState::None )
 #if !UE_BUILD_SHIPPING
 ,	ReceivedRawPacketDel()
 #endif
@@ -1525,6 +1528,32 @@ bool UNetConnection::ReadPacketInfo(FBitReader& Reader)
 	return true;
 }
 
+FNetworkGUID UNetConnection::GetActorGUIDFromOpenBunch(FInBunch& Bunch)
+{
+	// NOTE: This could break if this is a PartialBunch and the ActorGUID wasn't serialized.
+	//			Seems unlikely given the aggressive Flushing + increased MTU on InternalAck.
+
+	// Any GUIDs / Exports will have been read already for InternalAck connections,
+	// but we may have to skip over must-be-mapped GUIDs before we can read the actor GUID.
+
+	if (Bunch.bHasMustBeMappedGUIDs)
+	{
+		uint16 NumMustBeMappedGUIDs = 0;
+		Bunch << NumMustBeMappedGUIDs;
+
+		for (int32 i = 0; i < NumMustBeMappedGUIDs; i++)
+		{
+			FNetworkGUID NetGUID;
+			Bunch << NetGUID;
+		}
+	}
+
+	FNetworkGUID ActorGUID;
+	Bunch << ActorGUID;
+
+	return ActorGUID;
+}
+
 void UNetConnection::ReceivedPacket( FBitReader& Reader )
 {
 	SCOPED_NAMED_EVENT(UNetConnection_ReceivedPacket, FColor::Green);
@@ -1876,28 +1905,16 @@ PRAGMA_ENABLE_DEPRECATION_WARNINGS
 
 				if (bNewlyOpenedActorChannel)
 				{
-					// NOTE: This could break if this is a PartialBunch and the ActorGUID wasn't serialized.
-					//			Seems unlikely given the aggressive Flushing + increased MTU on InternalAck.
+					FNetworkGUID ActorGUID = GetActorGUIDFromOpenBunch(Bunch);
 
-					// Any GUIDs / Exports will have been read already for InternalAck connections,
-					// but we may have to skip over must-be-mapped GUIDs before we can read the actor GUID.
-
-					if (Bunch.bHasMustBeMappedGUIDs)
+					if (!Bunch.IsError())
 					{
-						uint16 NumMustBeMappedGUIDs = 0;
-						Bunch << NumMustBeMappedGUIDs;
-
-						for (int32 i = 0; i < NumMustBeMappedGUIDs; i++)
-						{
-							FNetworkGUID NetGUID;
-							Bunch << NetGUID;
-						}
+						IgnoringChannels.Add(Bunch.ChIndex, ActorGUID);
+					}
+					else
+					{
+						UE_LOG(LogNetTraffic, Error, TEXT("UNetConnection::ReceivedPacket: Unable to read actor GUID for ignored bunch. (Channel %d)"), Bunch.ChIndex);
 					}
-
-					FNetworkGUID ActorGUID;
-					Bunch << ActorGUID;
-
-					IgnoringChannels.Add(Bunch.ChIndex, ActorGUID);
 				}
 
 				if (IgnoringChannels.Contains(Bunch.ChIndex))
@@ -1919,6 +1936,8 @@ PRAGMA_ENABLE_DEPRECATION_WARNINGS
 							}
 						}
 					}
+
+					UE_LOG(LogNetTraffic, Log, TEXT("Ignoring bunch for already open channel: %i"), Bunch.ChIndex);
 					continue;
 				}
 			}
@@ -1986,6 +2005,42 @@ PRAGMA_ENABLE_DEPRECATION_WARNINGS
 					continue;
 				}
 
+				// peek for guid
+				if (InternalAck && (Bunch.ChName == NAME_Actor))
+				{
+					if (Bunch.bOpen && (!Bunch.bPartial || Bunch.bPartialInitial))
+					{
+						FBitReaderMark Mark(Bunch);
+						FNetworkGUID ActorGUID = GetActorGUIDFromOpenBunch(Bunch);
+						Mark.Pop(Bunch);
+
+						if (ActorGUID.IsValid() && !ActorGUID.IsDefault())
+						{
+							if (IgnoredGuids.Contains(ActorGUID))
+							{
+								UE_LOG(LogNetTraffic, Verbose, TEXT("Adding Channel: %i to ignore list, ignoring guid: %s"), Bunch.ChIndex, *ActorGUID.ToString());
+								IgnoredChannels.Add(Bunch.ChIndex);
+								continue;
+							}
+							else
+							{
+								if (IgnoredChannels.Remove(Bunch.ChIndex))
+								{
+									UE_LOG(LogNetTraffic, Verbose, TEXT("Removing Channel: %i from ignore list, got new guid: %s"), Bunch.ChIndex, *ActorGUID.ToString());
+								}
+							}
+						}
+					}
+					else
+					{
+						if (IgnoredChannels.Contains(Bunch.ChIndex))
+						{
+							UE_LOG(LogNetTraffic, Verbose, TEXT("Ignoring bunch on channel: %i"), Bunch.ChIndex);
+							continue;
+						}
+					}
+				}
+
 				// Reliable (either open or later), so create new channel.
 				UE_LOG(LogNetTraffic, Log, TEXT("      Bunch Create %i: ChName %s, ChSequence: %i, bReliable: %i, bPartial: %i, bPartialInitial: %i, bPartialFinal: %i"), Bunch.ChIndex, *Bunch.ChName.ToString(), Bunch.ChSequence, (int)Bunch.bReliable, (int)Bunch.bPartial, (int)Bunch.bPartialInitial, (int)Bunch.bPartialFinal );
 				Channel = CreateChannelByName( Bunch.ChName, EChannelCreateFlags::None, Bunch.ChIndex );
diff --git a/Engine/Source/Runtime/Engine/Private/PackageMapClient.cpp b/Engine/Source/Runtime/Engine/Private/PackageMapClient.cpp
--- a/Engine/Source/Runtime/Engine/Private/PackageMapClient.cpp
+++ b/Engine/Source/Runtime/Engine/Private/PackageMapClient.cpp
@@ -1340,11 +1340,57 @@ void UPackageMapClient::ReceiveExportData(FArchive& Archive)
 	ReceiveNetExportGUIDs(Archive);
 }
 
+void UPackageMapClient::SerializeNetFieldExportDelta(FArchive& Ar)
+{
+	if (Ar.IsSaving())
+	{
+		TSet<uint64> DeltaNetFieldExports;
+		
+		for ( auto It = GuidCache->NetFieldExportGroupMap.CreateIterator(); It; ++It )
+		{
+			// Save out the export group
+			TSharedPtr<FNetFieldExportGroup> ExportGroup = It.Value();
+			if (ExportGroup.IsValid())
+			{
+				for ( int32 i = 0; i < ExportGroup->NetFieldExports.Num(); i++ )
+				{
+					if (ExportGroup->NetFieldExports[i].bExported && ExportGroup->NetFieldExports[i].bDirtyForReplay)
+					{
+						check(ExportGroup->PathNameIndex != 0);
+
+						const uint64 CmdHandle = ((uint64)ExportGroup->PathNameIndex) << 32 | (uint64)i;
+
+						check(i == ExportGroup->NetFieldExports[i].Handle);
+
+						DeltaNetFieldExports.Add(CmdHandle);
+
+						ExportGroup->NetFieldExports[i].bDirtyForReplay = false;
+					}
+				}
+			}
+		}
+
+		AppendNetFieldExportsInternal(Ar, DeltaNetFieldExports, EAppendNetExportFlags::ForceExportDirtyGroups);
+
+		NetFieldExports.Empty();
+	}
+	else
+	{
+		ReceiveNetFieldExports(Ar);
+	}
+}
+
 void UPackageMapClient::AppendNetFieldExports(FArchive& Archive)
+{
+	AppendNetFieldExportsInternal(Archive, NetFieldExports, EAppendNetExportFlags::None);
+	NetFieldExports.Empty();
+}
+
+void UPackageMapClient::AppendNetFieldExportsInternal(FArchive& Archive, const TSet<uint64>& InNetFieldExports, EAppendNetExportFlags Flags)
 {
 	check(Connection->InternalAck);
 
-	uint32 NetFieldCount = NetFieldExports.Num();
+	uint32 NetFieldCount = InNetFieldExports.Num();
 	Archive.SerializeIntPacked(NetFieldCount);
 
 	if (0 == NetFieldCount)
@@ -1355,7 +1401,7 @@ void UPackageMapClient::AppendNetFieldExports(FArchive& Archive)
 	TArray< uint32, TInlineAllocator<64> > ExportedPathInThisBunchAlready;
 	ExportedPathInThisBunchAlready.Reserve(NetFieldCount);
 
-	for (const uint64 FieldExport : NetFieldExports)
+	for (const uint64 FieldExport : InNetFieldExports)
 	{
 		// Parse the path name index and cmd index out of the uint64
 		uint32 PathNameIndex = FieldExport >> 32;
@@ -1369,7 +1415,9 @@ void UPackageMapClient::AppendNetFieldExports(FArchive& Archive)
 		check(NetFieldExportHandle == NetFieldExportGroup->NetFieldExports[NetFieldExportHandle].Handle);
 
 		// Export the path if we need to
-		uint32 NeedsExport = (OverrideAckState->NetFieldExportGroupPathAcked.Contains(PathNameIndex) || ExportedPathInThisBunchAlready.Contains(PathNameIndex)) ? 0 : 1;
+		const bool bForceExportDirty = EnumHasAnyFlags(Flags, EAppendNetExportFlags::ForceExportDirtyGroups) && NetFieldExportGroup->bDirtyForReplay;
+
+		uint32 NeedsExport = ((bForceExportDirty || !OverrideAckState->NetFieldExportGroupPathAcked.Contains(PathNameIndex)) && !ExportedPathInThisBunchAlready.Contains(PathNameIndex)) ? 1 : 0;
 
 		Archive.SerializeIntPacked(PathNameIndex);
 		Archive.SerializeIntPacked(NeedsExport);
@@ -1382,6 +1430,11 @@ void UPackageMapClient::AppendNetFieldExports(FArchive& Archive)
 			Archive.SerializeIntPacked(NumExports);
 
 			ExportedPathInThisBunchAlready.Add(PathNameIndex);
+
+			if (bForceExportDirty)
+			{
+				NetFieldExportGroup->bDirtyForReplay = false;
+			}
 		}
 
 		Archive << NetFieldExportGroup->NetFieldExports[NetFieldExportHandle];
@@ -1389,8 +1442,6 @@ void UPackageMapClient::AppendNetFieldExports(FArchive& Archive)
 		OverrideAckState->NetFieldExportGroupPathAcked.Add( PathNameIndex );
 		OverrideAckState->NetFieldExportAcked.Add( FieldExport );
 	}
-
-	NetFieldExports.Empty();
 }
 
 void UPackageMapClient::ReceiveNetFieldExports(FArchive& Archive)
@@ -2015,7 +2066,11 @@ void UPackageMapClient::Serialize(FArchive& Ar)
 //	FNetGUIDCache
 //----------------------------------------------------------------------------------------
 
-FNetGUIDCache::FNetGUIDCache( UNetDriver * InDriver ) : IsExportingNetGUIDBunch( false ), Driver( InDriver ), NetworkChecksumMode( ENetworkChecksumMode::SaveAndUse ), AsyncLoadMode( EAsyncLoadMode::UseCVar )
+FNetGUIDCache::FNetGUIDCache(UNetDriver* InDriver) 
+	: Driver(InDriver)
+	, NetworkChecksumMode(ENetworkChecksumMode::SaveAndUse)
+	, AsyncLoadMode(EAsyncLoadMode::UseCVar)
+	, IsExportingNetGUIDBunch(false)
 {
 	UniqueNetIDs[0] = UniqueNetIDs[1] = 0;
 	UniqueNetFieldExportGroupPathIndex = 0;
diff --git a/Engine/Source/Runtime/Engine/Private/RepLayout.cpp b/Engine/Source/Runtime/Engine/Private/RepLayout.cpp
--- a/Engine/Source/Runtime/Engine/Private/RepLayout.cpp
+++ b/Engine/Source/Runtime/Engine/Private/RepLayout.cpp
@@ -883,7 +883,7 @@ bool FRepLayout::ReplicateProperties(
 		FilterChangeList(InactiveChangelist, RepState->InactiveParents, RepState->InactiveChangelist, NewlyActiveChangelist);
 	}
 
-	if (OwningChannel->Connection->bResendAllDataSinceOpen)
+	if (OwningChannel->Connection->ResendAllDataState != EResendAllDataState::None)
 	{
 		check(OwningChannel->Connection->InternalAck);
 
diff --git a/Engine/Source/Runtime/Engine/Public/Net/DataReplication.h b/Engine/Source/Runtime/Engine/Public/Net/DataReplication.h
--- a/Engine/Source/Runtime/Engine/Public/Net/DataReplication.h
+++ b/Engine/Source/Runtime/Engine/Public/Net/DataReplication.h
@@ -78,6 +78,7 @@ public:
 	TArray<FPropertyRetirement>						Retirement;					// Property retransmission.
 	TMap<int32, TSharedPtr<INetDeltaBaseState> >	RecentCustomDeltaState;		// This is the delta state we need to compare with when determining what to send to a client for custom delta properties
 	TMap<int32, TSharedPtr<INetDeltaBaseState> >	CDOCustomDeltaState;		// Same as RecentCustomDeltaState, but this will always remain as the initial CDO version. We use this to send all properties since channel was first opened (for bResendAllDataSinceOpen)
+	TMap<int32, TSharedPtr<INetDeltaBaseState> >	CheckpointCustomDeltaState;	// Same as RecentCustomDeltaState, but will represent the state at the last checkpoint
 
 	TArray< int32 >									LifetimeCustomDeltaProperties;
 	TArray< ELifetimeCondition >					LifetimeCustomDeltaPropertyConditions;
@@ -96,6 +97,7 @@ public:
 
 	TSharedPtr< FRepLayout >						RepLayout;
 	TUniquePtr< FRepState > 						RepState;
+	TUniquePtr< FRepState >							CheckpointRepState;
 
 	TSet< FNetworkGUID >							ReferencedGuids;
 	int32											TrackedGuidMemoryBytes;
@@ -164,6 +166,9 @@ public:
 	void	ReplicateCustomDeltaProperties( FNetBitWriter & Bunch, FReplicationFlags RepFlags );
 	bool	ReplicateProperties( FOutBunch & Bunch, FReplicationFlags RepFlags );
 	void	PostSendBunch(FPacketIdRange & PacketRange, uint8 bReliable);
+
+	/** Updates the custom delta state for a replay delta checkpoint */
+	void	UpdateCheckpoint();
 	
 	bool	ReceivedBunch( FNetBitReader& Bunch, const FReplicationFlags& RepFlags, const bool bHasRepLayout, bool& bOutHasUnmapped );
 	bool	ReceivedRPC(FNetBitReader& Reader, const FReplicationFlags& RepFlags, const FFieldNetCache* FieldCache, const bool bCanDelayRPC, bool& bOutDelayRPC, TSet<FNetworkGUID>& OutUnmappedGuids);
diff --git a/Engine/Source/Runtime/NetworkReplayStreaming/HttpNetworkReplayStreaming/Private/HttpNetworkReplayStreaming.cpp b/Engine/Source/Runtime/NetworkReplayStreaming/HttpNetworkReplayStreaming/Private/HttpNetworkReplayStreaming.cpp
--- a/Engine/Source/Runtime/NetworkReplayStreaming/HttpNetworkReplayStreaming/Private/HttpNetworkReplayStreaming.cpp
+++ b/Engine/Source/Runtime/NetworkReplayStreaming/HttpNetworkReplayStreaming/Private/HttpNetworkReplayStreaming.cpp
@@ -12,6 +12,8 @@
 #include "Engine/World.h"
 #include "UnrealEngine.h"
 #include "Engine/LocalPlayer.h"
+#include "Serialization/MemoryReader.h"
+#include "Serialization/MemoryWriter.h"
 
 DEFINE_LOG_CATEGORY_STATIC( LogHttpReplay, Log, All );
 
@@ -216,6 +218,7 @@ FHttpNetworkReplayStreamer::FHttpNetworkReplayStreamer() :
 	HighPriorityEndTime( 0 ),
 	StreamerLastError( ENetworkReplayError::None ),
 	DownloadCheckpointIndex( -1 ),
+	DeltaDownloadCheckpointIndex( -1 ),
 	LastGotoTimeInMS( -1 ),
 	TotalUploadBytes( 0 ),
 	RefreshViewerFails( 0 )
@@ -758,17 +761,21 @@ bool FQueuedGotoFakeCheckpoint::PreProcess( FHttpNetworkReplayStreamer* Streamer
 	Streamer->CheckpointArchive.Buffer.Empty();
 	Streamer->CheckpointArchive.Pos = 0;
 
-	// Completely reset our stream (we're going to start downloading from the start of the checkpoint)
-	Streamer->StreamArchive.Buffer.Empty();
-	Streamer->StreamArchive.Pos				= 0;
-	Streamer->StreamArchive.bAtEndOfReplay	= false;
+	if (!Streamer->IsDataAvailableForTimeRange(0, Streamer->LastGotoTimeInMS))
+	{
+		// Completely reset our stream (we're going to start downloading from the start of the checkpoint)
+		Streamer->StreamArchive.Buffer.Empty();
 
-	// Reset our stream range
-	Streamer->StreamTimeRangeStart	= 0;
-	Streamer->StreamTimeRangeEnd	= 0;
+		// Reset our stream range
+		Streamer->StreamTimeRangeStart	= 0;
+		Streamer->StreamTimeRangeEnd	= 0;
+
+		// Reset chunk index
+		Streamer->StreamChunkIndex		= 0;
+	}
 
-	// Reset chunk index
-	Streamer->StreamChunkIndex		= 0;
+	Streamer->StreamArchive.Pos				= 0;
+	Streamer->StreamArchive.bAtEndOfReplay	= false;
 
 	FGotoCallback& GotoCheckpointDelegate = Streamer->GotoCheckpointDelegate;
 	if (GotoCheckpointDelegate.IsBound())
@@ -788,9 +795,16 @@ bool FQueuedGotoFakeCheckpoint::PreProcess( FHttpNetworkReplayStreamer* Streamer
 	return true;
 }
 
-void FHttpNetworkReplayStreamer::GotoCheckpointIndex( const int32 CheckpointIndex, const FGotoCallback& Delegate )
+void FHttpNetworkReplayStreamer::GotoCheckpointIndex( const int32 CheckpointIndex, const FGotoCallback& Delegate, EReplayCheckpointType CheckpointType )
 {
-	if ( GotoCheckpointDelegate.IsBound() || DownloadCheckpointIndex != -1 )
+	if ( !IsCheckpointTypeSupported(CheckpointType) )
+	{
+		UE_LOG( LogHttpReplay, Warning, TEXT( "FHttpNetworkReplayStreamer::GotoCheckpointIndex. Unsupported checkpoint type." ) );
+		Delegate.ExecuteIfBound(FGotoResult());
+		return;
+	}
+
+	if ( GotoCheckpointDelegate.IsBound() || DownloadCheckpointIndex != -1 || DeltaDownloadCheckpointIndex != -1 )
 	{
 		// If we're currently going to a checkpoint now, ignore this request
 		UE_LOG( LogHttpReplay, Warning, TEXT( "FHttpNetworkReplayStreamer::GotoCheckpointIndex. Busy processing another checkpoint." ) );
@@ -805,13 +819,11 @@ void FHttpNetworkReplayStreamer::GotoCheckpointIndex( const int32 CheckpointInde
 		return;
 	}
 
-	check( DownloadCheckpointIndex == -1 );
-
 	if ( CheckpointIndex == -1 )
 	{
 		GotoCheckpointDelegate = Delegate;
 		SetHighPriorityTimeRange( 0, LastGotoTimeInMS );
-		LastChunkTime = 0;		// Force the next chunk to start downloading immediately in case LastGotoTimeInMS is 0 (which would effectively disable high priority mode immediately)
+		LastChunkTime = 0;		// Force the next chunk to start downloading immediately
 		AddCustomRequestToQueue( TSharedPtr< FQueuedHttpRequest >( new FQueuedGotoFakeCheckpoint() ) );
 		return;
 	}
@@ -826,10 +838,19 @@ void FHttpNetworkReplayStreamer::GotoCheckpointIndex( const int32 CheckpointInde
 	TSharedRef<IHttpRequest> HttpRequest = FHttpModule::Get().CreateRequest();
 
 	// Download the next stream chunk
-	HttpRequest->SetURL( FString::Printf( TEXT( "%sevent/%s" ), *ServerURL,  *CheckpointList.ReplayEvents[CheckpointIndex].ID ) );
+	HttpRequest->SetURL( FString::Printf( TEXT( "%sevent/%s" ), *ServerURL,  *CheckpointList.ReplayEvents[DeltaDownloadCheckpointIndex].ID ) );
 	HttpRequest->SetVerb( TEXT( "GET" ) );
 
-	HttpRequest->OnProcessRequestComplete().BindRaw( this, &FHttpNetworkReplayStreamer::HttpDownloadCheckpointFinished );
+	if (CheckpointType == EReplayCheckpointType::Delta)
+	{
+		DeltaDownloadCheckpointIndex = 0;
+
+		HttpRequest->OnProcessRequestComplete().BindRaw( this, &FHttpNetworkReplayStreamer::HttpDownloadCheckpointDeltaFinished );
+	}
+	else
+	{
+		HttpRequest->OnProcessRequestComplete().BindRaw( this, &FHttpNetworkReplayStreamer::HttpDownloadCheckpointFinished );
+	}
 
 	GotoCheckpointDelegate	= Delegate;
 	DownloadCheckpointIndex = CheckpointIndex;
@@ -920,8 +941,15 @@ void FHttpNetworkReplayStreamer::HttpRequestEventDataFinished(FHttpRequestPtr Ht
 	RequestEventDataCompleteDelegate.ExecuteIfBound(Result);
 }
 
-void FHttpNetworkReplayStreamer::GotoTimeInMS( const uint32 TimeInMS, const FGotoCallback& Delegate )
+void FHttpNetworkReplayStreamer::GotoTimeInMS( const uint32 TimeInMS, const FGotoCallback& Delegate, EReplayCheckpointType CheckpointType )
 {
+	if ( !IsCheckpointTypeSupported(CheckpointType) )
+	{
+		UE_LOG( LogHttpReplay, Log, TEXT( "FHttpNetworkReplayStreamer::GotoTimeInMS. Unsupported checkpoint type." ) );
+		Delegate.ExecuteIfBound( FGotoResult() );
+		return;
+	}
+
 	if ( LastGotoTimeInMS != -1 || DownloadCheckpointIndex != -1 )
 	{
 		// If we're processing requests, be on the safe side and cancel the scrub
@@ -968,7 +996,7 @@ void FHttpNetworkReplayStreamer::GotoTimeInMS( const uint32 TimeInMS, const FGot
 		}
 	}
 
-	GotoCheckpointIndex( CheckpointIndex, Delegate );
+	GotoCheckpointIndex(CheckpointIndex, Delegate, CheckpointType);
 }
 
 void FHttpNetworkReplayStreamer::FlushCheckpointInternal( uint32 TimeInMS )
@@ -2518,6 +2546,166 @@ bool FHttpNetworkReplayStreamer::IsStreaming() const
 	return StreamerState != EStreamerState::Idle;
 }
 
+void FHttpNetworkReplayStreamer::HttpDownloadCheckpointDeltaFinished( FHttpRequestPtr HttpRequest, FHttpResponsePtr HttpResponse, bool bSucceeded )
+{
+	UE_LOG( LogHttpReplay, Verbose, TEXT( "FHttpNetworkReplayStreamer::HttpDownloadCheckpointDeltaFinished." ) );
+
+	RequestFinished( EStreamerState::StreamingDown, EQueuedHttpRequestType::DownloadingCheckpoint, HttpRequest );
+
+	check( StreamArchive.IsLoading() );
+	check( GotoCheckpointDelegate.IsBound() );
+	check( DownloadCheckpointIndex >= 0 );
+	check( DeltaDownloadCheckpointIndex >= 0);
+
+	// If we get here after StopStreaming was called, then assume this operation should be cancelled
+	// A more correct fix would be to actually cancel this in-flight request when StopStreaming is called
+	// But for now, this is a safe change, and can co-exist with the more proper fix
+	if ( bStopStreamingCalled )
+	{
+		GotoCheckpointDelegate.ExecuteIfBound(FGotoResult());
+		GotoCheckpointDelegate	= FGotoCallback();
+		DownloadCheckpointIndex = -1;
+		DeltaDownloadCheckpointIndex = -1;
+		LastGotoTimeInMS		= -1;
+		return;
+	}
+
+	if ( bSucceeded && HttpResponse->GetResponseCode() == EHttpResponseCodes::Ok )
+	{
+		if ( HttpResponse->GetContent().Num() == 0 )
+		{
+			UE_LOG( LogHttpReplay, Warning, TEXT( "FHttpNetworkReplayStreamer::HttpDownloadCheckpointDeltaFinished. Checkpoint empty." ) );
+			GotoCheckpointDelegate.ExecuteIfBound(FGotoResult());
+			GotoCheckpointDelegate = FGotoCallback();
+			DownloadCheckpointIndex = -1;
+			DeltaDownloadCheckpointIndex = -1;
+			LastGotoTimeInMS		= -1;
+			return;
+		}
+
+		// Add response to response cache
+		AddResponseToCache( HttpRequest, HttpResponse );
+
+		// clear archive when receiving first chunk, then append
+		if (DeltaDownloadCheckpointIndex == 0)
+		{
+			CheckpointArchive.Buffer.Empty();
+		}
+
+		uint32 CheckpointSize = 0;
+
+		FMemoryReader ContentAr(HttpResponse->GetContent(), true);
+		ContentAr << CheckpointSize;
+
+		FMemoryWriter CheckpointAr(CheckpointArchive.Buffer, true);
+		CheckpointAr.Seek(CheckpointAr.TotalSize());
+		CheckpointAr << CheckpointSize;
+
+		// Get the checkpoint data
+		if (!DecompressResponse(HttpResponse, CheckpointArchive.Buffer))
+		{
+			GotoCheckpointDelegate.ExecuteIfBound(FGotoResult());
+			GotoCheckpointDelegate	= FGotoCallback();
+			DownloadCheckpointIndex = -1;
+			DeltaDownloadCheckpointIndex = -1;
+			LastGotoTimeInMS		= -1;
+			return;
+		}
+
+		if (DeltaDownloadCheckpointIndex < DownloadCheckpointIndex)
+		{
+			DeltaDownloadCheckpointIndex++;
+			check(CheckpointList.ReplayEvents.IsValidIndex(DeltaDownloadCheckpointIndex));
+
+			TSharedRef<IHttpRequest> NextHttpRequest = FHttpModule::Get().CreateRequest();
+
+			// Download the next stream chunk
+			NextHttpRequest->SetURL( FString::Printf( TEXT( "%sevent/%s" ), *ServerURL,  *CheckpointList.ReplayEvents[DeltaDownloadCheckpointIndex].ID ) );
+			NextHttpRequest->SetVerb( TEXT( "GET" ) );
+
+			NextHttpRequest->OnProcessRequestComplete().BindRaw( this, &FHttpNetworkReplayStreamer::HttpDownloadCheckpointDeltaFinished );
+
+			AddRequestToQueue( EQueuedHttpRequestType::DownloadingCheckpoint, NextHttpRequest );
+			return;
+		}
+		else
+		{
+			CheckpointArchive.Pos = 0;
+
+			// Completely reset our stream (we're going to start downloading from the start of the checkpoint)
+			StreamArchive.Buffer.Empty();
+			StreamArchive.Pos				= 0;
+			StreamArchive.bAtEndOfReplay	= false;
+
+			// Reset any time we were waiting on in the past
+			HighPriorityEndTime	= 0;
+
+			// Reset our stream range
+			StreamTimeRangeStart	= 0;
+			StreamTimeRangeEnd		= 0;
+
+			// Set the next chunk to be right after this checkpoint (which was stored in the metadata)
+			StreamChunkIndex = FCString::Atoi( *CheckpointList.ReplayEvents[ DownloadCheckpointIndex ].Metadata );
+
+			LastChunkTime = 0;		// Force the next chunk to start downloading immediately
+
+			// If we want to fast forward past the end of a stream, clamp to the checkpoint
+			if ( LastGotoTimeInMS >= 0 && StreamChunkIndex >= NumTotalStreamChunks )
+			{
+				UE_LOG( LogHttpReplay, Warning, TEXT( "FHttpNetworkReplayStreamer::HttpDownloadCheckpointFinished. Clamped to checkpoint: %i" ), LastGotoTimeInMS );
+
+				StreamTimeRangeStart	= CheckpointList.ReplayEvents[DownloadCheckpointIndex].Time1;
+				StreamTimeRangeEnd		= CheckpointList.ReplayEvents[DownloadCheckpointIndex].Time1;
+				LastGotoTimeInMS		= -1;
+			}
+
+			if ( LastGotoTimeInMS >= 0 )
+			{
+				// If we are fine scrubbing, make sure to wait on the part of the stream that is needed to do this in one frame
+				SetHighPriorityTimeRange( CheckpointList.ReplayEvents[ DownloadCheckpointIndex ].Time1, LastGotoTimeInMS );
+
+				// Subtract off checkpoint time so we pass in the leftover to the engine to fast forward through for the fine scrubbing part
+				LastGotoTimeInMS -= CheckpointList.ReplayEvents[ DownloadCheckpointIndex ].Time1;
+			}
+
+			// Notify game code of success
+			FGotoResult Result;
+			Result.Result = EStreamingOperationResult::Success;
+			Result.ExtraTimeMS = LastGotoTimeInMS;
+			GotoCheckpointDelegate.ExecuteIfBound(Result);
+
+			UE_LOG( LogHttpReplay, Verbose, TEXT( "FHttpNetworkReplayStreamer::HttpDownloadCheckpointDeltaFinished. SUCCESS. StreamChunkIndex: %i" ), StreamChunkIndex );
+		}
+	}
+	else
+	{
+		// Oops, something went wrong, notify game code of failure
+		UE_LOG( LogHttpReplay, Error, TEXT( "FHttpNetworkReplayStreamer::HttpDownloadCheckpointDeltaFinished. FAILED, %s" ), *BuildRequestErrorString( HttpRequest, HttpResponse ) );
+		GotoCheckpointDelegate.ExecuteIfBound(FGotoResult());
+	}
+
+	// Reset things
+	GotoCheckpointDelegate	= FGotoCallback();
+	DownloadCheckpointIndex = -1;
+	DeltaDownloadCheckpointIndex = -1;
+	LastGotoTimeInMS		= -1;
+}
+
+bool FHttpNetworkReplayStreamer::IsCheckpointTypeSupported(EReplayCheckpointType CheckpointType) const
+{
+	bool bSupported = false;
+
+	switch (CheckpointType)
+	{
+	case EReplayCheckpointType::Full:
+	case EReplayCheckpointType::Delta:
+		bSupported = true;
+		break;
+	}
+
+	return bSupported;
+}
+
 IMPLEMENT_MODULE( FHttpNetworkReplayStreamingFactory, HttpNetworkReplayStreaming )
 
 TSharedPtr< INetworkReplayStreamer > FHttpNetworkReplayStreamingFactory::CreateReplayStreamer()
diff --git a/Engine/Source/Runtime/NetworkReplayStreaming/HttpNetworkReplayStreaming/Public/HttpNetworkReplayStreaming.h b/Engine/Source/Runtime/NetworkReplayStreaming/HttpNetworkReplayStreaming/Public/HttpNetworkReplayStreaming.h
--- a/Engine/Source/Runtime/NetworkReplayStreaming/HttpNetworkReplayStreaming/Public/HttpNetworkReplayStreaming.h
+++ b/Engine/Source/Runtime/NetworkReplayStreaming/HttpNetworkReplayStreaming/Public/HttpNetworkReplayStreaming.h
@@ -212,8 +212,8 @@ public:
 	virtual FArchive*	GetStreamingArchive() override;
 	virtual FArchive*	GetCheckpointArchive() override;
 	virtual void		FlushCheckpoint( const uint32 TimeInMS ) override;
-	virtual void		GotoCheckpointIndex( const int32 CheckpointIndex, const FGotoCallback& Delegate ) override;
-	virtual void		GotoTimeInMS( const uint32 TimeInMS, const FGotoCallback& Delegate ) override;
+	virtual void		GotoCheckpointIndex( const int32 CheckpointIndex, const FGotoCallback& Delegate, EReplayCheckpointType CheckpointType ) override;
+	virtual void		GotoTimeInMS( const uint32 TimeInMS, const FGotoCallback& Delegate, EReplayCheckpointType CheckpointType ) override;
 	virtual void		UpdateTotalDemoTime( uint32 TimeInMS ) override;
 	virtual uint32		GetTotalDemoTime() const override { return TotalDemoTimeInMS; }
 	virtual bool		IsDataAvailable() const override;
@@ -260,6 +260,8 @@ public:
 		return EStreamingOperationResult::Unsupported;
 	}
 
+	virtual bool IsCheckpointTypeSupported(EReplayCheckpointType CheckpointType) const override;
+
 	/** FHttpNetworkReplayStreamer */
 	void UploadHeader();
 	void FlushStream();
@@ -292,6 +294,9 @@ public:
 	virtual bool DecompressResponse(FHttpResponsePtr HttpResponse, TArray<uint8>& ResultBuffer) const;
 	virtual bool CompressRequest(FHttpRequestPtr HttpRequest, const TArray<uint8>& RequestBuffer) const;
 
+	void InternalGotoTimeInMS(const uint32 TimeInMS, const FGotoCallback& Delegate, bool bDelta);
+	void InternalGotoCheckpointIndex(const int32 CheckpointIndex, const FGotoCallback& Delegate, const FHttpRequestCompleteDelegate& RequestDelegate);
+
 	/** EStreamerState - Overall state of the streamer */
 	enum class EStreamerState
 	{
@@ -307,6 +312,7 @@ public:
 	void HttpDownloadHeaderFinished( FHttpRequestPtr HttpRequest, FHttpResponsePtr HttpResponse, bool bSucceeded, FDownloadHeaderCallback Delegate);
 	void HttpDownloadFinished( FHttpRequestPtr HttpRequest, FHttpResponsePtr HttpResponse, bool bSucceeded, int32 RequestedStreamChunkIndex, bool bStreamWasLive );
 	void HttpDownloadCheckpointFinished( FHttpRequestPtr HttpRequest, FHttpResponsePtr HttpResponse, bool bSucceeded );
+	void HttpDownloadCheckpointDeltaFinished( FHttpRequestPtr HttpRequest, FHttpResponsePtr HttpResponse, bool bSucceeded );
 	void HttpRefreshViewerFinished( FHttpRequestPtr HttpRequest, FHttpResponsePtr HttpResponse, bool bSucceeded );
 	void HttpStartUploadingFinished( FHttpRequestPtr HttpRequest, FHttpResponsePtr HttpResponse, bool bSucceeded );
 	void HttpStopUploadingFinished( FHttpRequestPtr HttpRequest, FHttpResponsePtr HttpResponse, bool bSucceeded );
@@ -357,6 +363,7 @@ public:
 	FStartStreamingCallback			StartStreamingDelegate;		// Delegate passed in to StartStreaming
 	FGotoCallback					GotoCheckpointDelegate;
 	int32							DownloadCheckpointIndex;
+	int32							DeltaDownloadCheckpointIndex;
 	int64							LastGotoTimeInMS;
 
 	FReplayEventList					CheckpointList;
diff --git a/Engine/Source/Runtime/NetworkReplayStreaming/InMemoryNetworkReplayStreaming/Private/InMemoryNetworkReplayStreaming.cpp b/Engine/Source/Runtime/NetworkReplayStreaming/InMemoryNetworkReplayStreaming/Private/InMemoryNetworkReplayStreaming.cpp
--- a/Engine/Source/Runtime/NetworkReplayStreaming/InMemoryNetworkReplayStreaming/Private/InMemoryNetworkReplayStreaming.cpp
+++ b/Engine/Source/Runtime/NetworkReplayStreaming/InMemoryNetworkReplayStreaming/Private/InMemoryNetworkReplayStreaming.cpp
@@ -425,7 +425,7 @@ void FInMemoryNetworkReplayStreamer::FlushCheckpoint(const uint32 TimeInMS)
 	FoundReplay->StreamChunks.Add(NewChunk);
 }
 
-void FInMemoryNetworkReplayStreamer::GotoCheckpointIndex(const int32 CheckpointIndex, const FGotoCallback& Delegate)
+void FInMemoryNetworkReplayStreamer::GotoCheckpointIndex(const int32 CheckpointIndex, const FGotoCallback& Delegate, EReplayCheckpointType CheckpointType)
 {
 	GotoCheckpointIndexInternal(CheckpointIndex, Delegate, -1);
 }
@@ -509,7 +509,7 @@ FInMemoryReplay* FInMemoryNetworkReplayStreamer::GetCurrentReplayChecked() const
 	return FoundReplay;
 }
 
-void FInMemoryNetworkReplayStreamer::GotoTimeInMS(const uint32 TimeInMS, const FGotoCallback& Delegate)
+void FInMemoryNetworkReplayStreamer::GotoTimeInMS(const uint32 TimeInMS, const FGotoCallback& Delegate, EReplayCheckpointType CheckpointType)
 {
 	int32 CheckpointIndex = -1;
 
diff --git a/Engine/Source/Runtime/NetworkReplayStreaming/InMemoryNetworkReplayStreaming/Public/InMemoryNetworkReplayStreaming.h b/Engine/Source/Runtime/NetworkReplayStreaming/InMemoryNetworkReplayStreaming/Public/InMemoryNetworkReplayStreaming.h
--- a/Engine/Source/Runtime/NetworkReplayStreaming/InMemoryNetworkReplayStreaming/Public/InMemoryNetworkReplayStreaming.h
+++ b/Engine/Source/Runtime/NetworkReplayStreaming/InMemoryNetworkReplayStreaming/Public/InMemoryNetworkReplayStreaming.h
@@ -130,8 +130,8 @@ public:
 	virtual FArchive* GetStreamingArchive() override;
 	virtual FArchive* GetCheckpointArchive() override;
 	virtual void FlushCheckpoint( const uint32 TimeInMS ) override;
-	virtual void GotoCheckpointIndex( const int32 CheckpointIndex, const FGotoCallback& Delegate ) override;
-	virtual void GotoTimeInMS( const uint32 TimeInMS, const FGotoCallback& Delegate ) override;
+	virtual void GotoCheckpointIndex( const int32 CheckpointIndex, const FGotoCallback& Delegate, EReplayCheckpointType CheckpointType ) override;
+	virtual void GotoTimeInMS( const uint32 TimeInMS, const FGotoCallback& Delegate, EReplayCheckpointType CheckpointType ) override;
 	virtual void UpdateTotalDemoTime( uint32 TimeInMS ) override;
 	virtual uint32 GetTotalDemoTime() const override;
 	virtual bool IsDataAvailable() const override;
@@ -193,6 +193,8 @@ public:
 		return EStreamingOperationResult::Unsupported;
 	}
 
+	virtual bool IsCheckpointTypeSupported(EReplayCheckpointType CheckpointType) const override { return (CheckpointType == EReplayCheckpointType::Full); }
+
 private:
 	bool IsNamedStreamLive( const FString& StreamName ) const;
 
diff --git a/Engine/Source/Runtime/NetworkReplayStreaming/LocalFileNetworkReplayStreaming/Private/LocalFileNetworkReplayStreaming.cpp b/Engine/Source/Runtime/NetworkReplayStreaming/LocalFileNetworkReplayStreaming/Private/LocalFileNetworkReplayStreaming.cpp
--- a/Engine/Source/Runtime/NetworkReplayStreaming/LocalFileNetworkReplayStreaming/Private/LocalFileNetworkReplayStreaming.cpp
+++ b/Engine/Source/Runtime/NetworkReplayStreaming/LocalFileNetworkReplayStreaming/Private/LocalFileNetworkReplayStreaming.cpp
@@ -1678,7 +1678,7 @@ void FLocalFileNetworkReplayStreamer::FlushCheckpointInternal(const uint32 TimeI
 	CheckpointAr.Pos = 0;	
 }
 
-void FLocalFileNetworkReplayStreamer::GotoCheckpointIndex(const int32 CheckpointIndex, const FGotoCallback& Delegate)
+void FLocalFileNetworkReplayStreamer::GotoCheckpointIndex(const int32 CheckpointIndex, const FGotoCallback& Delegate, EReplayCheckpointType CheckpointType)
 {
 	if (IsFileRequestPendingOrInProgress(EQueuedLocalFileRequestType::ReadingCheckpoint))
 	{
@@ -1743,8 +1743,185 @@ void FLocalFileNetworkReplayStreamer::GotoCheckpointIndex(const int32 Checkpoint
 		return;
 	}
 
-	AddCachedFileRequestToQueue<FGotoResult>(EQueuedLocalFileRequestType::ReadingCheckpoint, CurrentReplayInfo.Checkpoints[CheckpointIndex].ChunkIndex,
-		[this, CheckpointIndex](TLocalFileRequestCommonData<FGotoResult>& RequestData)
+	if (CheckpointType == EReplayCheckpointType::Delta)
+	{
+		AddDelegateFileRequestToQueue<FGotoResult>(EQueuedLocalFileRequestType::ReadingCheckpoint,
+			[this, CheckpointIndex](TLocalFileRequestCommonData<FGotoResult>& RequestData)
+		{
+			// If we get here after StopStreaming was called, then assume this operation should be cancelled
+			// A more correct fix would be to actually cancel this in-flight request when StopStreaming is called
+			// But for now, this is a safe change, and can co-exist with the more proper fix
+			if (bStopStreamingCalled)
+			{
+				return;
+			}
+
+			RequestData.DataBuffer.Empty();
+
+			const FString FullDemoFilename = GetDemoFullFilename(CurrentStreamName);
+
+			TSharedPtr<FArchive> LocalFileAr = CreateLocalFileReader(FullDemoFilename);
+			if (LocalFileAr.IsValid())
+			{
+				if (ReadReplayInfo(*LocalFileAr, RequestData.ReplayInfo))
+				{
+					TArray<uint8> CheckpointData;
+
+					// read all the checkpoints
+					for (int32 i = 0; i <= CheckpointIndex; ++i)
+					{
+						if (DeltaCheckpointCache.Contains(i))
+						{
+							FMemoryWriter Writer(RequestData.DataBuffer, true, true);
+							uint32 CheckpointSize = DeltaCheckpointCache[i]->RequestData.Num();
+							Writer << CheckpointSize;
+
+							RequestData.DataBuffer.Append(DeltaCheckpointCache[i]->RequestData);
+						}
+						else
+						{
+							LocalFileAr->Seek(RequestData.ReplayInfo.Checkpoints[i].EventDataOffset);
+
+							CheckpointData.Reset();
+							CheckpointData.AddUninitialized(RequestData.ReplayInfo.Checkpoints[i].SizeInBytes);
+
+							LocalFileAr->Serialize(CheckpointData.GetData(), CheckpointData.Num());
+
+							// Get the checkpoint data
+							if (RequestData.ReplayInfo.bCompressed)
+							{
+								if (SupportsCompression())
+								{
+									SCOPE_CYCLE_COUNTER(STAT_LocalReplay_DecompressTime);
+
+									TArray<uint8> UncompressedData;
+
+									if (!DecompressBuffer(CheckpointData, UncompressedData))
+									{
+										UE_LOG(LogLocalFileReplay, Error, TEXT("FLocalFileNetworkReplayStreamer::GotoCheckpointIndexDelta. DecompressBuffer FAILED."));
+										RequestData.DataBuffer.Empty();
+										return;
+									}
+
+									FMemoryWriter Writer(RequestData.DataBuffer, true, true);
+									uint32 CheckpointSize = UncompressedData.Num();
+									Writer << CheckpointSize;
+
+									RequestData.DataBuffer.Append(UncompressedData);
+
+									DeltaCheckpointCache.Add(i, MakeShareable(new FCachedFileRequest(UncompressedData, 0)));
+								}
+								else
+								{
+									UE_LOG(LogLocalFileReplay, Error, TEXT("FLocalFileNetworkReplayStreamer::GotoCheckpointIndexDelta. Compressed checkpoint but streamer does not support compression."));
+									RequestData.DataBuffer.Empty();
+									return;
+								}
+							}
+							else
+							{
+								FMemoryWriter Writer(RequestData.DataBuffer, true, true);
+								uint32 CheckpointSize = CheckpointData.Num();
+								Writer << CheckpointSize;
+
+								RequestData.DataBuffer.Append(CheckpointData);
+
+								DeltaCheckpointCache.Add(i, MakeShareable(new FCachedFileRequest(CheckpointData, 0)));
+							}
+						}
+					}
+				}
+
+				LocalFileAr = nullptr;
+			}
+		},
+			[this, CheckpointIndex, Delegate](TLocalFileRequestCommonData<FGotoResult>& RequestData)
+		{
+			if (bStopStreamingCalled)
+			{
+				Delegate.ExecuteIfBound(RequestData.DelegateResult);
+				LastGotoTimeInMS = -1;
+				return;
+			}
+
+			if (RequestData.DataBuffer.Num() == 0)
+			{
+				UE_LOG(LogLocalFileReplay, Warning, TEXT("FLocalFileNetworkReplayStreamer::GotoCheckpointIndexDelta. Checkpoint empty."));
+				Delegate.ExecuteIfBound(RequestData.DelegateResult);
+				LastGotoTimeInMS = -1;
+				return;
+			}
+
+			CheckpointAr.Buffer = MoveTemp(RequestData.DataBuffer);
+			CheckpointAr.Pos = 0;
+
+			int32 DataChunkIndex = FCString::Atoi(*CurrentReplayInfo.Checkpoints[CheckpointIndex].Metadata);
+			check(CurrentReplayInfo.DataChunks.IsValidIndex(DataChunkIndex));
+
+			bool bIsDataAvailableForTimeRange = IsDataAvailableForTimeRange(CurrentReplayInfo.Checkpoints[CheckpointIndex].Time1, LastGotoTimeInMS);
+
+			if (!bIsDataAvailableForTimeRange)
+			{
+				// Completely reset our stream (we're going to start loading from the start of the checkpoint)
+				StreamAr.Buffer.Empty();
+				StreamAr.Pos = 0;
+				StreamAr.bAtEndOfReplay = false;
+
+				// Reset any time we were waiting on in the past
+				HighPriorityEndTime = 0;
+
+				StreamDataOffset = CurrentReplayInfo.DataChunks[DataChunkIndex].StreamOffset;
+
+				// Reset our stream range
+				StreamTimeRange = TInterval<uint32>(0, 0);
+
+				// Set the next chunk to be right after this checkpoint (which was stored in the metadata)
+				StreamChunkIndex = DataChunkIndex;
+
+				LastChunkTime = 0;		// Force the next chunk to start loading immediately in case LastGotoTimeInMS is 0 (which would effectively disable high priority mode immediately)
+			}
+			else
+			{
+				// set stream position back to the correct location
+				StreamAr.Pos = CurrentReplayInfo.DataChunks[DataChunkIndex].StreamOffset - StreamDataOffset;
+				check(StreamAr.Pos >= 0 && StreamAr.Pos <= StreamAr.Buffer.Num());
+				StreamAr.bAtEndOfReplay = false;
+			}
+
+			// If we want to fast forward past the end of a stream (and we set a new chunk to stream), clamp to the checkpoint
+			if (LastGotoTimeInMS >= 0 && StreamChunkIndex >= CurrentReplayInfo.DataChunks.Num() && !bIsDataAvailableForTimeRange)
+			{
+				UE_LOG(LogLocalFileReplay, Warning, TEXT("FLocalFileNetworkReplayStreamer::GotoCheckpointIndexDelta. Clamped to checkpoint: %i"), LastGotoTimeInMS);
+
+				StreamTimeRange = TInterval<uint32>(CurrentReplayInfo.Checkpoints[CheckpointIndex].Time1, CurrentReplayInfo.Checkpoints[CheckpointIndex].Time1);
+				LastGotoTimeInMS = -1;
+			}
+
+			if (LastGotoTimeInMS >= 0)
+			{
+				// If we are fine scrubbing, make sure to wait on the part of the stream that is needed to do this in one frame
+				SetHighPriorityTimeRange(CurrentReplayInfo.Checkpoints[CheckpointIndex].Time1, LastGotoTimeInMS);
+
+				// Subtract off starting time so we pass in the leftover to the engine to fast forward through for the fine scrubbing part
+				LastGotoTimeInMS -= CurrentReplayInfo.Checkpoints[CheckpointIndex].Time1;
+			}
+
+			// Notify game code of success
+			RequestData.DelegateResult.Result = EStreamingOperationResult::Success;
+			RequestData.DelegateResult.ExtraTimeMS = LastGotoTimeInMS;
+
+			Delegate.ExecuteIfBound(RequestData.DelegateResult);
+
+			UE_LOG(LogLocalFileReplay, Verbose, TEXT("FLocalFileNetworkReplayStreamer::GotoCheckpointIndexDelta. SUCCESS. StreamChunkIndex: %i"), StreamChunkIndex);
+
+			// Reset things
+			LastGotoTimeInMS = -1;
+		});
+	}
+	else
+	{
+		AddCachedFileRequestToQueue<FGotoResult>(EQueuedLocalFileRequestType::ReadingCheckpoint, CurrentReplayInfo.Checkpoints[CheckpointIndex].ChunkIndex,
+			[this, CheckpointIndex](TLocalFileRequestCommonData<FGotoResult>& RequestData)
 		{
 			// If we get here after StopStreaming was called, then assume this operation should be cancelled
 			// A more correct fix would be to actually cancel this in-flight request when StopStreaming is called
@@ -1801,7 +1978,7 @@ void FLocalFileNetworkReplayStreamer::GotoCheckpointIndex(const int32 Checkpoint
 				LocalFileAr = nullptr;
 			}
 		},
-		[this, CheckpointIndex, Delegate](TLocalFileRequestCommonData<FGotoResult>& RequestData)
+			[this, CheckpointIndex, Delegate](TLocalFileRequestCommonData<FGotoResult>& RequestData)
 		{
 			if (bStopStreamingCalled)
 			{
@@ -1812,7 +1989,7 @@ void FLocalFileNetworkReplayStreamer::GotoCheckpointIndex(const int32 Checkpoint
 
 			if (RequestData.DataBuffer.Num() == 0)
 			{
-				UE_LOG( LogLocalFileReplay, Warning, TEXT( "FLocalFileNetworkReplayStreamer::GotoCheckpointIndex. Checkpoint empty." ) );
+				UE_LOG(LogLocalFileReplay, Warning, TEXT("FLocalFileNetworkReplayStreamer::GotoCheckpointIndex. Checkpoint empty."));
 				Delegate.ExecuteIfBound(RequestData.DelegateResult);
 				LastGotoTimeInMS = -1;
 				return;
@@ -1860,14 +2037,14 @@ void FLocalFileNetworkReplayStreamer::GotoCheckpointIndex(const int32 Checkpoint
 			}
 			else if (LastGotoTimeInMS >= 0)
 			{
-				UE_LOG( LogLocalFileReplay, Warning, TEXT("FLocalFileNetworkReplayStreamer::GotoCheckpointIndex. Clamped to checkpoint: %i"), LastGotoTimeInMS );
+				UE_LOG(LogLocalFileReplay, Warning, TEXT("FLocalFileNetworkReplayStreamer::GotoCheckpointIndex. Clamped to checkpoint: %i"), LastGotoTimeInMS);
 
 				// If we want to fast forward past the end of a stream, clamp to the checkpoint
 				StreamTimeRange = TInterval<uint32>(Checkpoint.Time1, Checkpoint.Time1);
 				LastGotoTimeInMS = -1;
 			}
 
-			if ( LastGotoTimeInMS >= 0 )
+			if (LastGotoTimeInMS >= 0)
 			{
 				// If we are fine scrubbing, make sure to wait on the part of the stream that is needed to do this in one frame
 				SetHighPriorityTimeRange(Checkpoint.Time1, LastGotoTimeInMS);
@@ -1887,11 +2064,12 @@ void FLocalFileNetworkReplayStreamer::GotoCheckpointIndex(const int32 Checkpoint
 			UE_LOG(LogLocalFileReplay, Verbose, TEXT("FLocalFileNetworkReplayStreamer::GotoCheckpointIndex. SUCCESS. StreamChunkIndex: %i"), StreamChunkIndex);
 
 			// Reset things
-			LastGotoTimeInMS		= -1;
+			LastGotoTimeInMS = -1;
 		});
+	}
 }
 
-void FLocalFileNetworkReplayStreamer::GotoTimeInMS(const uint32 TimeInMS, const FGotoCallback& Delegate)
+void FLocalFileNetworkReplayStreamer::GotoTimeInMS(const uint32 TimeInMS, const FGotoCallback& Delegate, EReplayCheckpointType CheckpointType)
 {
 	if (IsFileRequestPendingOrInProgress(EQueuedLocalFileRequestType::ReadingCheckpoint) || LastGotoTimeInMS != -1)
 	{
@@ -1930,7 +2108,7 @@ void FLocalFileNetworkReplayStreamer::GotoTimeInMS(const uint32 TimeInMS, const
 		}
 	}
 
-	GotoCheckpointIndex( CheckpointIndex, Delegate );
+	GotoCheckpointIndex(CheckpointIndex, Delegate, CheckpointType);
 }
 
 bool FLocalFileNetworkReplayStreamer::HasPendingFileRequests() const
@@ -2639,7 +2817,7 @@ void FLocalFileNetworkReplayStreamer::AddRequestToCache(int32 ChunkIndex, const
 	}
 
 	// Add to cache (or freshen existing entry)
-	RequestCache.Add(ChunkIndex, FCachedFileRequest(RequestData, FPlatformTime::Seconds()));
+	RequestCache.Add(ChunkIndex, MakeShareable(new FCachedFileRequest(RequestData, FPlatformTime::Seconds())));
 
 	// Anytime we add something to cache, make sure it's within budget
 	CleanupRequestCache();
@@ -2654,16 +2832,20 @@ void FLocalFileNetworkReplayStreamer::CleanupRequestCache()
 		int32 OldestKey = INDEX_NONE;
 		uint32 TotalSize = 0;
 
-		for (auto It = RequestCache.CreateIterator(); It; ++It)
+		for (const auto& RequestPair : RequestCache)
 		{
-			if ((OldestKey == INDEX_NONE) || It.Value().LastAccessTime < OldestTime)
+			const TSharedPtr<FCachedFileRequest>& Request = RequestPair.Value;
+			if (Request.IsValid())
 			{
-				OldestTime = It.Value().LastAccessTime;
-				OldestKey = It.Key();
-			}
+				if ((OldestKey == INDEX_NONE) || Request->LastAccessTime < OldestTime)
+				{
+					OldestTime = Request->LastAccessTime;
+					OldestKey = RequestPair.Key;
+				}
 
-			// Accumulate total cache size
-			TotalSize += It.Value().RequestData.Num();
+				// Accumulate total cache size
+				TotalSize += Request->RequestData.Num();
+			}
 		}
 
 		check(OldestKey != INDEX_NONE);
@@ -2823,6 +3005,21 @@ const void FLocalFileNetworkReplayStreamer::GetUserIndicesFromUserStrings(const
 	}
 }
 
+bool FLocalFileNetworkReplayStreamer::IsCheckpointTypeSupported(EReplayCheckpointType CheckpointType) const
+{
+	bool bSupported = false;
+
+	switch (CheckpointType)
+	{
+	case EReplayCheckpointType::Full:
+	case EReplayCheckpointType::Delta:
+		bSupported = true;
+		break;
+	}
+
+	return bSupported;
+}
+
 IMPLEMENT_MODULE(FLocalFileNetworkReplayStreamingFactory, LocalFileNetworkReplayStreaming)
 
 TSharedPtr<INetworkReplayStreamer> FLocalFileNetworkReplayStreamingFactory::CreateReplayStreamer() 
diff --git a/Engine/Source/Runtime/NetworkReplayStreaming/LocalFileNetworkReplayStreaming/Public/LocalFileNetworkReplayStreaming.h b/Engine/Source/Runtime/NetworkReplayStreaming/LocalFileNetworkReplayStreaming/Public/LocalFileNetworkReplayStreaming.h
--- a/Engine/Source/Runtime/NetworkReplayStreaming/LocalFileNetworkReplayStreaming/Public/LocalFileNetworkReplayStreaming.h
+++ b/Engine/Source/Runtime/NetworkReplayStreaming/LocalFileNetworkReplayStreaming/Public/LocalFileNetworkReplayStreaming.h
@@ -121,15 +121,15 @@ struct FLocalFileReplayInfo
 class LOCALFILENETWORKREPLAYSTREAMING_API FLocalFileStreamFArchive : public FArchive
 {
 public:
-	FLocalFileStreamFArchive() : Pos( 0 ), bAtEndOfReplay( false ) {}
+	FLocalFileStreamFArchive() : Pos(0), bAtEndOfReplay(false) {}
 
-	virtual void	Serialize( void* V, int64 Length ) override;
+	virtual void	Serialize(void* V, int64 Length) override;
 	virtual int64	Tell() override;
 	virtual int64	TotalSize() override;
-	virtual void	Seek( int64 InPos ) override;
+	virtual void	Seek(int64 InPos) override;
 	virtual bool	AtEnd() override;
 
-	TArray< uint8 >	Buffer;
+	TArray<uint8>	Buffer;
 	int32			Pos;
 	bool			bAtEndOfReplay;
 };
@@ -373,8 +373,8 @@ public:
 
 	virtual bool GetCachedRequest() override
 	{
-		FCachedFileRequest* CachedRequest = this->Streamer->RequestCache.Find(CacheKey);
-		if (CachedRequest != nullptr)
+		TSharedPtr<FCachedFileRequest> CachedRequest = this->Streamer->RequestCache.FindRef(CacheKey);
+		if (CachedRequest.IsValid())
 		{
 			// If we have this response in the cache, process it now
 			CachedRequest->LastAccessTime = FPlatformTime::Seconds();
@@ -405,8 +405,8 @@ public:
 	virtual FArchive* GetStreamingArchive() override;
 	virtual FArchive* GetCheckpointArchive() override;
 	virtual void FlushCheckpoint(const uint32 TimeInMS) override;
-	virtual void GotoCheckpointIndex(const int32 CheckpointIndex, const FGotoCallback& Delegate) override;
-	virtual void GotoTimeInMS(const uint32 TimeInMS, const FGotoCallback& Delegate) override;
+	virtual void GotoCheckpointIndex(const int32 CheckpointIndex, const FGotoCallback& Delegate, EReplayCheckpointType CheckpointType) override;
+	virtual void GotoTimeInMS(const uint32 TimeInMS, const FGotoCallback& Delegate, EReplayCheckpointType CheckpointType) override;
 	virtual void UpdateTotalDemoTime(uint32 TimeInMS) override;
 	virtual uint32 GetTotalDemoTime() const override { return CurrentReplayInfo.LengthInMS; }
 	virtual bool IsDataAvailable() const override;
@@ -443,6 +443,8 @@ public:
 	virtual void RefreshHeader() override;
 	virtual void DownloadHeader(const FDownloadHeaderCallback& Delegate) override;
 
+	virtual bool IsCheckpointTypeSupported(EReplayCheckpointType CheckpointType) const override;
+
 	virtual bool SupportsCompression() const { return false; }
 	virtual int32 GetDecompressedSize(FArchive& InCompressed) const { return 0; }
 	virtual bool DecompressBuffer(const TArray<uint8>& InCompressed, TArray< uint8 >& OutBuffer) const { return false; }
@@ -511,7 +513,10 @@ public:
 	}
 
 	/** Map of chunk index to cached value */
-	TMap<int32, FCachedFileRequest> RequestCache;
+	TMap<int32, TSharedPtr<FCachedFileRequest>> RequestCache;
+
+	/** Map of checkpoint index to cached value */
+	TMap<int32, TSharedPtr<FCachedFileRequest>> DeltaCheckpointCache;
 
 protected:
 
diff --git a/Engine/Source/Runtime/NetworkReplayStreaming/NetworkReplayStreaming/Public/NetworkReplayStreaming.h b/Engine/Source/Runtime/NetworkReplayStreaming/NetworkReplayStreaming/Public/NetworkReplayStreaming.h
--- a/Engine/Source/Runtime/NetworkReplayStreaming/NetworkReplayStreaming/Public/NetworkReplayStreaming.h
+++ b/Engine/Source/Runtime/NetworkReplayStreaming/NetworkReplayStreaming/Public/NetworkReplayStreaming.h
@@ -446,6 +446,11 @@ struct FRenameReplayResult : public FStreamingResultBase
 DECLARE_DELEGATE_OneParam(FRenameReplayCallback, const FRenameReplayResult&);
 /** End KeepReplay Types */
 
+enum class EReplayCheckpointType : uint8
+{
+	Full,
+	Delta
+};
 
 /**
  * Generic interface for network replay streaming
@@ -468,9 +473,15 @@ public:
 	virtual FArchive* GetCheckpointArchive() = 0;
 	virtual void FlushCheckpoint(const uint32 TimeInMS) = 0;
 
-	virtual void GotoCheckpointIndex(const int32 CheckpointIndex, const FGotoCallback& Delegate) = 0;
+	UE_DEPRECATED(4.23, "Please use the version of GotoCheckpointIndex that accepts a FGotoCallback delegate and an EReplayCheckpointType")
+	virtual void GotoCheckpointIndex(const int32 CheckpointIndex, const FGotoCallback& Delegate) { GotoCheckpointIndex(CheckpointIndex, Delegate, EReplayCheckpointType::Full);	}
+	virtual void GotoCheckpointIndex(const int32 CheckpointIndex, const FGotoCallback& Delegate, EReplayCheckpointType CheckpointType) = 0;
+
+	UE_DEPRECATED(4.23, "Please use the version of GotoCheckpointIndex that accepts a FGotoCallback delegate and an EReplayCheckpointType.")
+	virtual void GotoTimeInMS(const uint32 TimeInMS, const FGotoCallback& Delegate) { GotoTimeInMS(TimeInMS, Delegate, EReplayCheckpointType::Full); }
+	virtual void GotoTimeInMS(const uint32 TimeInMS, const FGotoCallback& Delegate, EReplayCheckpointType CheckpointType) = 0;
 
-	virtual void GotoTimeInMS(const uint32 TimeInMS, const FGotoCallback& Delegate) = 0;
+	virtual bool IsCheckpointTypeSupported(EReplayCheckpointType CheckpointType) const = 0;
 
 	virtual void UpdateTotalDemoTime(uint32 TimeInMS) = 0;
 	virtual uint32 GetTotalDemoTime() const = 0;
diff --git a/Engine/Source/Runtime/NetworkReplayStreaming/NullNetworkReplayStreaming/Private/NullNetworkReplayStreaming.cpp b/Engine/Source/Runtime/NetworkReplayStreaming/NullNetworkReplayStreaming/Private/NullNetworkReplayStreaming.cpp
--- a/Engine/Source/Runtime/NetworkReplayStreaming/NullNetworkReplayStreaming/Private/NullNetworkReplayStreaming.cpp
+++ b/Engine/Source/Runtime/NetworkReplayStreaming/NullNetworkReplayStreaming/Private/NullNetworkReplayStreaming.cpp
@@ -558,7 +558,7 @@ void FNullNetworkReplayStreamer::FlushCheckpoint(const uint32 TimeInMS)
 	++CurrentCheckpointIndex;
 }
 
-void FNullNetworkReplayStreamer::GotoCheckpointIndex(const int32 CheckpointIndex, const FGotoCallback& Delegate)
+void FNullNetworkReplayStreamer::GotoCheckpointIndex(const int32 CheckpointIndex, const FGotoCallback& Delegate, EReplayCheckpointType CheckpointType)
 {
 	GotoCheckpointIndexInternal(CheckpointIndex, Delegate, -1);
 }
@@ -641,7 +641,7 @@ void FNullNetworkReplayStreamer::UpdateReplayInfoIfValid()
 	}
 }
 
-void FNullNetworkReplayStreamer::GotoTimeInMS(const uint32 TimeInMS, const FGotoCallback& Delegate)
+void FNullNetworkReplayStreamer::GotoTimeInMS(const uint32 TimeInMS, const FGotoCallback& Delegate, EReplayCheckpointType CheckpointType)
 {
 	// Enumerate all the events in the events folder, since we need to know what times the checkpoints correlate with
 	TArray<FNullCheckpointListItem> Checkpoints;
diff --git a/Engine/Source/Runtime/NetworkReplayStreaming/NullNetworkReplayStreaming/Public/NullNetworkReplayStreaming.h b/Engine/Source/Runtime/NetworkReplayStreaming/NullNetworkReplayStreaming/Public/NullNetworkReplayStreaming.h
--- a/Engine/Source/Runtime/NetworkReplayStreaming/NullNetworkReplayStreaming/Public/NullNetworkReplayStreaming.h
+++ b/Engine/Source/Runtime/NetworkReplayStreaming/NullNetworkReplayStreaming/Public/NullNetworkReplayStreaming.h
@@ -51,8 +51,8 @@ public:
 	virtual FArchive* GetStreamingArchive() override;
 	virtual FArchive* GetCheckpointArchive() override;
 	virtual void FlushCheckpoint( const uint32 TimeInMS ) override;
-	virtual void GotoCheckpointIndex( const int32 CheckpointIndex, const FGotoCallback& Delegate ) override;
-	virtual void GotoTimeInMS( const uint32 TimeInMS, const FGotoCallback& Delegate ) override;
+	virtual void GotoCheckpointIndex( const int32 CheckpointIndex, const FGotoCallback& Delegate, EReplayCheckpointType CheckpointType ) override;
+	virtual void GotoTimeInMS( const uint32 TimeInMS, const FGotoCallback& Delegate, EReplayCheckpointType CheckpointType ) override;
 	virtual void UpdateTotalDemoTime( uint32 TimeInMS ) override;
 	virtual uint32 GetTotalDemoTime() const override { return ReplayInfo.LengthInMS; }
 	virtual bool IsDataAvailable() const override;
@@ -113,6 +113,8 @@ public:
 
 	virtual EStreamingOperationResult GetDemoPath(FString& DemoPath) const override;
 
+	virtual bool IsCheckpointTypeSupported(EReplayCheckpointType CheckpointType) const override { return (CheckpointType == EReplayCheckpointType::Full); }
+
 private:
 	bool IsNamedStreamLive( const FString& StreamName ) const;
 
