Some memory reporting improvements

diff --git a/Engine/Source/Runtime/Engine/Private/SkeletalMeshLODRenderData.cpp b/Engine/Source/Runtime/Engine/Private/SkeletalMeshLODRenderData.cpp
--- a/Engine/Source/Runtime/Engine/Private/SkeletalMeshLODRenderData.cpp
+++ b/Engine/Source/Runtime/Engine/Private/SkeletalMeshLODRenderData.cpp
@@ -675,6 +675,14 @@ void FSkeletalMeshLODRenderData::Serialize(FArchive& Ar, UObject* Owner, int32 I
 			SkelMeshOwner->NeedCPUData(Idx);
 	}
 
+	if (FPlatformProperties::RequiresCookedData())
+	{
+		if (bNeedsCPUAccess && SkelMeshOwner)
+		{
+			UE_LOG(LogStaticMesh, Log, TEXT("[%s] Skeletal Mesh is marked for CPU read."), *SkelMeshOwner->GetName());
+		}
+	}
+
 	if (StripFlags.IsDataStrippedForServer() || StripFlags.IsClassDataStripped(MinLodStripFlag))
 	{
 		TArray<FSkelMeshRenderSection> DummySections;
diff --git a/Engine/Source/Runtime/Engine/Private/StaticMesh.cpp b/Engine/Source/Runtime/Engine/Private/StaticMesh.cpp
--- a/Engine/Source/Runtime/Engine/Private/StaticMesh.cpp
+++ b/Engine/Source/Runtime/Engine/Private/StaticMesh.cpp
@@ -322,6 +322,14 @@ void FStaticMeshLODResources::SerializeBuffers(FArchive& Ar, UStaticMesh* OwnerS
 	// TODO: Not needed in uncooked games either after PostLoad!
 	bool bNeedsCPUAccess = !FPlatformProperties::RequiresCookedData() || bMeshCPUAcces;
 
+	if (FPlatformProperties::RequiresCookedData())
+	{
+		if (bNeedsCPUAccess && OwnerStaticMesh)
+		{
+			UE_LOG(LogStaticMesh, Log, TEXT("[%s] Mesh is marked for CPU read."), *OwnerStaticMesh->GetName());
+		}
+	}
+
 	bHasWireframeIndices = false;
 	bHasAdjacencyInfo = false;
 	bHasDepthOnlyIndices = false;
diff --git a/Engine/Source/Runtime/Engine/Public/GenericOctree.h b/Engine/Source/Runtime/Engine/Public/GenericOctree.h
--- a/Engine/Source/Runtime/Engine/Public/GenericOctree.h
+++ b/Engine/Source/Runtime/Engine/Public/GenericOctree.h
@@ -705,7 +705,7 @@ public:
 	bool IsValidElementId(FOctreeElementId ElementId) const;
 
 	/** Writes stats for the octree to the log. */
-	void DumpStats();
+	void DumpStats() const;
 
 	SIZE_T GetSizeBytes() const
 	{
diff --git a/Engine/Source/Runtime/Engine/Public/GenericOctree.inl b/Engine/Source/Runtime/Engine/Public/GenericOctree.inl
--- a/Engine/Source/Runtime/Engine/Public/GenericOctree.inl
+++ b/Engine/Source/Runtime/Engine/Public/GenericOctree.inl
@@ -297,7 +297,7 @@ bool TOctree<ElementType,OctreeSemantics>::IsValidElementId(FOctreeElementId Ele
 };
 
 template<typename ElementType,typename OctreeSemantics>
-void TOctree<ElementType,OctreeSemantics>::DumpStats()
+void TOctree<ElementType,OctreeSemantics>::DumpStats() const
 {
 	int32 NumNodes = 0;
 	int32 NumLeaves = 0;
