Added describe to log message when receiving NMT_ChannelActorFailed to provide more information.

diff --git a/Engine/Source/Runtime/Engine/Private/DataChannel.cpp b/Engine/Source/Runtime/Engine/Private/DataChannel.cpp
--- a/Engine/Source/Runtime/Engine/Private/DataChannel.cpp
+++ b/Engine/Source/Runtime/Engine/Private/DataChannel.cpp
@@ -1397,7 +1397,7 @@ void UControlChannel::ReceivedBunch( FInBunch& Bunch )
 
 				if (FNetControlMessage<NMT_ActorChannelFailure>::Receive(Bunch, ChannelIndex))
 				{
-					UE_LOG(LogNet, Log, TEXT("Server connection received: %s"), FNetControlMessageInfo::GetName(MessageType));
+					UE_LOG(LogNet, Log, TEXT("Server connection received: %s %s"), FNetControlMessageInfo::GetName(MessageType), *Describe());
 
 					// Check if Channel index provided by client is valid and within range of channel on server
 					if (ChannelIndex >= 0 && ChannelIndex < Connection->Channels.Num())
