Allow ExecuteConsoleCommand to route to the engine as a fallback if there are no player controllers active

diff --git a/Engine/Source/Runtime/Engine/Private/KismetSystemLibrary.cpp b/Engine/Source/Runtime/Engine/Private/KismetSystemLibrary.cpp
--- a/Engine/Source/Runtime/Engine/Private/KismetSystemLibrary.cpp
+++ b/Engine/Source/Runtime/Engine/Private/KismetSystemLibrary.cpp
@@ -302,10 +302,15 @@ void UKismetSystemLibrary::ExecuteConsoleCommand(UObject* WorldContextObject, co
 {
 	// First, try routing through the primary player
 	APlayerController* TargetPC = Player ? Player : UGameplayStatics::GetPlayerController(WorldContextObject, 0);
-	if( TargetPC )
+	if (TargetPC)
 	{
 		TargetPC->ConsoleCommand(Command, true);
 	}
+	else if (GEngine)
+	{
+		UWorld* World = GEngine->GetWorldFromContextObject(WorldContextObject, EGetWorldErrorMode::ReturnNull);
+		GEngine->Exec(World, *Command);
+	}
 }
 
 float UKismetSystemLibrary::GetConsoleVariableFloatValue(const FString& VariableName)
