Spectate option to un-possess or not current pawn

#rnx

diff --git a/Engine/Source/Runtime/Engine/Classes/GameFramework/PlayerController.h b/Engine/Source/Runtime/Engine/Classes/GameFramework/PlayerController.h
--- a/Engine/Source/Runtime/Engine/Classes/GameFramework/PlayerController.h
+++ b/Engine/Source/Runtime/Engine/Classes/GameFramework/PlayerController.h
@@ -1864,6 +1864,9 @@ protected:
 	/** Destroys the SpectatorPawn and sets it to nullptr. */
 	virtual void DestroySpectatorPawn();
 
+	/** Useful to spectate other pawn without un-possessing the current pawn */
+	virtual bool ShouldKeepCurrentPawnUponSpectating() const { return false; }
+
 private:
 	/** The pawn used when spectating (nullptr if not spectating). */
 	UPROPERTY()
diff --git a/Engine/Source/Runtime/Engine/Private/PlayerController.cpp b/Engine/Source/Runtime/Engine/Private/PlayerController.cpp
--- a/Engine/Source/Runtime/Engine/Private/PlayerController.cpp
+++ b/Engine/Source/Runtime/Engine/Private/PlayerController.cpp
@@ -4736,7 +4736,7 @@ void APlayerController::EndPlayingState()
 
 void APlayerController::BeginSpectatingState()
 {
-	if (GetPawn() != NULL && Role == ROLE_Authority)
+	if (GetPawn() != NULL && Role == ROLE_Authority && ShouldKeepCurrentPawnUponSpectating() == false)
 	{
 		UnPossess();
 	}
