Fix for crash in FSoundEffectBase (if provided Inpreset is null)

diff --git a/Engine/Source/Runtime/Engine/Private/SoundEffectBase.cpp b/Engine/Source/Runtime/Engine/Private/SoundEffectBase.cpp
--- a/Engine/Source/Runtime/Engine/Private/SoundEffectBase.cpp
+++ b/Engine/Source/Runtime/Engine/Private/SoundEffectBase.cpp
@@ -27,9 +27,17 @@ void FSoundEffectBase::SetEnabled(const bool bInIsEnabled)
 
 void FSoundEffectBase::SetPreset(USoundEffectPreset* Inpreset)
 {
-	Preset = Inpreset;
-	Preset->AddEffectInstance(this);
-	bChanged = true;
+	if (Preset != Inpreset)
+	{
+		ClearPreset();
+
+		Preset = Inpreset;
+		if (Preset)
+		{
+			Preset->AddEffectInstance(this);
+			bChanged = true;
+		}
+	}
 }
 
 void FSoundEffectBase::ClearPreset()
@@ -63,7 +71,7 @@ void FSoundEffectBase::EffectCommand(TFunction<void()> Command)
 }
 void FSoundEffectBase::PumpPendingMessages()
 {
-	// Pumps the commadn queue
+	// Pumps the command queue
 	TFunction<void()> Command;
 	while (CommandQueue.Dequeue(Command))
 	{
