Prevent UNetConnection CDOs from allocating memory for Channels, Sequence Receipts, etc.
Also, make sure more code is properly tracked as LLM Networking.

diff --git a/Engine/Source/Runtime/Engine/Private/GameInstance.cpp b/Engine/Source/Runtime/Engine/Private/GameInstance.cpp
--- a/Engine/Source/Runtime/Engine/Private/GameInstance.cpp
+++ b/Engine/Source/Runtime/Engine/Private/GameInstance.cpp
@@ -918,6 +918,8 @@ const TArray<class ULocalPlayer*>& UGameInstance::GetLocalPlayers() const
 
 void UGameInstance::StartRecordingReplay(const FString& Name, const FString& FriendlyName, const TArray<FString>& AdditionalOptions)
 {
+	LLM_SCOPE(ELLMTag::Networking);
+
 	if ( FParse::Param( FCommandLine::Get(),TEXT( "NOREPLAYS" ) ) )
 	{
 		UE_LOG( LogDemo, Warning, TEXT( "UGameInstance::StartRecordingReplay: Rejected due to -noreplays option" ) );
@@ -1026,6 +1028,8 @@ void UGameInstance::StopRecordingReplay()
 
 bool UGameInstance::PlayReplay(const FString& Name, UWorld* WorldOverride, const TArray<FString>& AdditionalOptions)
 {
+	LLM_SCOPE(ELLMTag::Networking);
+
 	UWorld* CurrentWorld = WorldOverride != nullptr ? WorldOverride : GetWorld();
 
 	if ( CurrentWorld == nullptr )
diff --git a/Engine/Source/Runtime/Engine/Private/NetConnection.cpp b/Engine/Source/Runtime/Engine/Private/NetConnection.cpp
--- a/Engine/Source/Runtime/Engine/Private/NetConnection.cpp
+++ b/Engine/Source/Runtime/Engine/Private/NetConnection.cpp
@@ -164,19 +164,28 @@ PRAGMA_ENABLE_DEPRECATION_WARNINGS
 ,	HasDirtyAcks(0u)
 ,	bHasWarnedAboutChannelLimit(false)
 {
+	// This isn't ideal, because it won't capture memory derived classes are creating dynamically.
+	// The allocations could *probably* be moved somewhere else (like InitBase), but that
+	// causes failure to connect for some reason, and for now this is easier.
+	LLM_SCOPE(ELLMTag::Networking);
+
 	MaxChannelSize = CVarMaxChannelSize.GetValueOnAnyThread();
 	if (MaxChannelSize <= 0)
 	{
 		UE_LOG(LogNet, Warning, TEXT("CVarMaxChannelSize of %d is less than or equal to 0, using the default number of channels."), MaxChannelSize);
 		MaxChannelSize = DEFAULT_MAX_CHANNEL_SIZE;
 	}
+	
+	
+	if (!HasAnyFlags(EObjectFlags::RF_ClassDefaultObject | EObjectFlags::RF_ArchetypeObject))
+	{
+		Channels.AddDefaulted(MaxChannelSize);
+		OutReliable.AddDefaulted(MaxChannelSize);
+		InReliable.AddDefaulted(MaxChannelSize);
+		PendingOutRec.AddDefaulted(MaxChannelSize);
 
-	Channels.AddDefaulted(MaxChannelSize);
-	OutReliable.AddDefaulted(MaxChannelSize);
-	InReliable.AddDefaulted(MaxChannelSize);
-	PendingOutRec.AddDefaulted(MaxChannelSize);
-
-	PacketNotify.Init(InPacketId, OutPacketId);
+		PacketNotify.Init(InPacketId, OutPacketId);
+	}	
 }
 
 /**
@@ -501,7 +510,6 @@ void UNetConnection::Serialize( FArchive& Ar )
 	{
 		// TODO: We don't currently track:
 		//		StatelessConnectComponents
-		//		PacketHandlers
 		//		AnalyticsVars
 		//		AnalyticsData
 		//		Histogram data.
@@ -561,6 +569,14 @@ PRAGMA_ENABLE_DEPRECATION_WARNINGS
 		GRANULAR_NETWORK_MEMORY_TRACKING_TRACK("LastOut", LastOut.CountMemory(Ar));
 		GRANULAR_NETWORK_MEMORY_TRACKING_TRACK("SendBunchHeader", SendBunchHeader.CountMemory(Ar));
 
+		GRANULAR_NETWORK_MEMORY_TRACKING_TRACK("PacketHandler",
+			if (Handler.IsValid())
+			{
+				// PacketHandler already counts its size.
+				Handler->CountBytes(Ar);
+			}
+		);
+
 #if DO_ENABLE_NET_TEST
 		GRANULAR_NETWORK_MEMORY_TRACKING_TRACK("Delayed",
 			Delayed.CountBytes(Ar);
diff --git a/Engine/Source/Runtime/Engine/Private/PendingNetGame.cpp b/Engine/Source/Runtime/Engine/Private/PendingNetGame.cpp
--- a/Engine/Source/Runtime/Engine/Private/PendingNetGame.cpp
+++ b/Engine/Source/Runtime/Engine/Private/PendingNetGame.cpp
@@ -31,6 +31,8 @@ UPendingNetGame::UPendingNetGame(const FObjectInitializer& ObjectInitializer)
 
 void UPendingNetGame::InitNetDriver()
 {
+	LLM_SCOPE(ELLMTag::Networking);
+
 	if (!GDisallowNetworkTravel)
 	{
 		NETWORK_PROFILER(GNetworkProfiler.TrackSessionChange(true, URL));
diff --git a/Engine/Source/Runtime/Engine/Private/UnrealEngine.cpp b/Engine/Source/Runtime/Engine/Private/UnrealEngine.cpp
--- a/Engine/Source/Runtime/Engine/Private/UnrealEngine.cpp
+++ b/Engine/Source/Runtime/Engine/Private/UnrealEngine.cpp
@@ -10992,6 +10992,7 @@ UNetDriver* CreateNetDriver_Local(UEngine* Engine, FWorldContext& Context, FName
 
 UNetDriver* UEngine::CreateNetDriver(UWorld *InWorld, FName NetDriverDefinition)
 {
+	LLM_SCOPE(ELLMTag::Networking);
 	return CreateNetDriver_Local(this, GetWorldContextFromWorldChecked(InWorld), NetDriverDefinition);
 }
 
@@ -11630,10 +11631,15 @@ EBrowseReturnVal::Type UEngine::Browse( FWorldContext& WorldContext, FURL URL, F
 				CollectGarbage(GARBAGE_COLLECTION_KEEPFLAGS);
 			}
 
-			// Just reload the RPC world (as we have no real map to load)
-			WorldContext.PendingNetGame = NewObject<UPendingNetGame>();
-			WorldContext.PendingNetGame->Initialize(WorldContext.LastURL);
-			WorldContext.PendingNetGame->InitNetDriver();
+			{
+				LLM_SCOPE(ELLMTag::Networking);
+
+				// Just reload the RPC world (as we have no real map to load)
+				WorldContext.PendingNetGame = NewObject<UPendingNetGame>();
+				WorldContext.PendingNetGame->Initialize(WorldContext.LastURL);
+				WorldContext.PendingNetGame->InitNetDriver();
+			}
+
 			if (!WorldContext.PendingNetGame->NetDriver)
 			{
 				// UPendingNetGame will set the appropriate error code and connection lost type, so
diff --git a/Engine/Source/Runtime/Engine/Private/World.cpp b/Engine/Source/Runtime/Engine/Private/World.cpp
--- a/Engine/Source/Runtime/Engine/Private/World.cpp
+++ b/Engine/Source/Runtime/Engine/Private/World.cpp
@@ -5089,6 +5089,8 @@ void UWorld::SendChallengeControlMessage(const FEncryptionKeyResponse& Response,
 bool UWorld::Listen( FURL& InURL )
 {
 #if WITH_SERVER_CODE
+	LLM_SCOPE(ELLMTag::Networking);
+
 	if( NetDriver )
 	{
 		GEngine->BroadcastNetworkFailure(this, NetDriver, ENetworkFailure::NetDriverAlreadyExists);
