Added "r.UITextureLODBias" in the scalability settings to add an extra bias to UI textures with mips.

Added 2 new texture filter modes :  TF_FirstMipNearest and TF_FirstMipBilinear. Those always sample the first (highest) mip.
When running with cooked content, TF_FirstMipNearest and TF_FirstMipBilinear will drop the all mips but the first (freeing mip tail memory).
In uncooked game, the visuals will be valid but memory of the mip trail is still allocated.
In the editor, the mip tail can still be visualized.

Fixed the texture editor streaming method description to actually reflect the state of the texture.

diff --git a/Engine/Source/Editor/TextureEditor/Private/TextureEditorToolkit.cpp b/Engine/Source/Editor/TextureEditor/Private/TextureEditorToolkit.cpp
--- a/Engine/Source/Editor/TextureEditor/Private/TextureEditorToolkit.cpp
+++ b/Engine/Source/Editor/TextureEditor/Private/TextureEditorToolkit.cpp
@@ -465,7 +465,7 @@ void FTextureEditorToolkit::PopulateQuickInfo( )
 	}
 
 	SizeText->SetText(FText::Format(NSLOCTEXT("TextureEditor", "QuickInfo_ResourceSize", "Resource Size: {0} Kb"), FText::AsNumber(Size, &SizeOptions)));
-	MethodText->SetText(FText::Format(NSLOCTEXT("TextureEditor", "QuickInfo_Method", "Method: {0}"), Texture->NeverStream ? NSLOCTEXT("TextureEditor", "QuickInfo_MethodNotStreamed", "Not Streamed") : NSLOCTEXT("TextureEditor", "QuickInfo_MethodStreamed", "Streamed")));
+	MethodText->SetText(FText::Format(NSLOCTEXT("TextureEditor", "QuickInfo_Method", "Method: {0}"), !Texture->bIsStreamable ? NSLOCTEXT("TextureEditor", "QuickInfo_MethodNotStreamed", "Not Streamed") : NSLOCTEXT("TextureEditor", "QuickInfo_MethodStreamed", "Streamed")));
 	LODBiasText->SetText(FText::Format(NSLOCTEXT("TextureEditor", "QuickInfo_LODBias", "Combined LOD Bias: {0}"), FText::AsNumber(Texture->GetCachedLODBias())));
 
 	int32 TextureFormatIndex = PF_MAX;
diff --git a/Engine/Source/Runtime/Engine/Classes/Engine/Texture.h b/Engine/Source/Runtime/Engine/Classes/Engine/Texture.h
--- a/Engine/Source/Runtime/Engine/Classes/Engine/Texture.h
+++ b/Engine/Source/Runtime/Engine/Classes/Engine/Texture.h
@@ -47,6 +47,8 @@ enum TextureFilter
 	TF_Nearest UMETA(DisplayName="Nearest"),
 	TF_Bilinear UMETA(DisplayName="Bi-linear"),
 	TF_Trilinear UMETA(DisplayName="Tri-linear"),
+	TF_FirstMipNearest  UMETA(DisplayName = "First mip nearest"),
+	TF_FirstMipBilinear  UMETA(DisplayName = "First mip bi-linear"),
 	/** Use setting from the Texture Group. */
 	TF_Default UMETA(DisplayName="Default (from Texture Group)"),
 	TF_MAX,
diff --git a/Engine/Source/Runtime/Engine/Private/Texture2D.cpp b/Engine/Source/Runtime/Engine/Private/Texture2D.cpp
--- a/Engine/Source/Runtime/Engine/Private/Texture2D.cpp
+++ b/Engine/Source/Runtime/Engine/Private/Texture2D.cpp
@@ -503,6 +503,20 @@ void UTexture2D::UpdateResource()
 	CachePlatformData();
 	// clear all the cooked cached platform data if the source could have changed... 
 	ClearAllCachedCookedPlatformData();
+#else
+	// Note that using TF_FirstMip disables texture streaming, because the mip data becomes lost.
+	// Also, the cleanup of the platform data must go between UpdateCachedLODBias() and UpdateResource().
+	if ((Filter == TF_FirstMipNearest || Filter == TF_FirstMipBilinear) && PlatformData && FPlatformProperties::RequiresCookedData())
+	{
+		const int32 FirstMip = FMath::Clamp(0, GetCachedLODBias(), PlatformData->Mips.Num() - 1);
+		// Remove any mips after the first mip.
+		PlatformData->Mips.RemoveAt(FirstMip + 1, PlatformData->Mips.Num() - FirstMip - 1);
+		// Remove any mips before the first mip.
+		PlatformData->Mips.RemoveAt(0, FirstMip);
+		// Update the texture size for the memory usage metrics.
+		PlatformData->SizeX = PlatformData->Mips[0].SizeX;
+		PlatformData->SizeY = PlatformData->Mips[0].SizeY;
+	}
 #endif // #if WITH_EDITOR
 
 	// Route to super.
@@ -827,6 +841,7 @@ FTextureResource* UTexture2D::CreateResource()
 					!NeverStream && 
 					(NumMips > 1) && 
 					(LODGroup != TEXTUREGROUP_UI) && 
+					(Filter != TF_FirstMipNearest && Filter != TF_FirstMipBilinear) &&
 					!bTemporarilyDisableStreaming;
 	
 	if (bIsStreamable && NumMips > 0)
@@ -1470,15 +1485,13 @@ uint32 FTexture2DResource::GetSizeY() const
 /** Returns the default mip bias for this texture. */
 int32 FTexture2DResource::GetDefaultMipMapBias() const
 {
-	if ( Owner->LODGroup == TEXTUREGROUP_UI )
+	check(Owner);
+	if ((Owner->LODGroup == TEXTUREGROUP_UI && CVarForceHighestMipOnUITexturesEnabled.GetValueOnAnyThread() > 0) || 
+		(!GIsEditor && (Owner->Filter == TF_FirstMipNearest || Owner->Filter == TF_FirstMipBilinear))) // Don't use in editor otherwise it prevents visualizing mips.
 	{
-		if ( CVarForceHighestMipOnUITexturesEnabled.GetValueOnAnyThread() > 0 )
-		{
-			const TIndirectArray<FTexture2DMipMap>& OwnerMips = Owner->GetPlatformMips();
-			return -OwnerMips.Num();
-		}
+		const TIndirectArray<FTexture2DMipMap>& OwnerMips = Owner->GetPlatformMips();
+		return -OwnerMips.Num();
 	}
-	
 	return 0;
 }
 
diff --git a/Engine/Source/Runtime/Engine/Private/TextureLODSettings.cpp b/Engine/Source/Runtime/Engine/Private/TextureLODSettings.cpp
--- a/Engine/Source/Runtime/Engine/Private/TextureLODSettings.cpp
+++ b/Engine/Source/Runtime/Engine/Private/TextureLODSettings.cpp
@@ -3,6 +3,16 @@
 #include "Engine/TextureLODSettings.h"
 #include "Engine/Texture2D.h"
 
+
+int32 GUITextureLODBias = 0;
+FAutoConsoleVariableRef CVarUITextureLODBias(
+	TEXT("r.UITextureLODBias"),
+	GUITextureLODBias,
+	TEXT("Extra LOD bias to apply to UI textures. (default=0)"),
+	ECVF_Scalability
+);
+
+
 void FTextureLODGroup::SetupGroup()
 {
 	// editor would never want to use smaller mips based on memory (could affect cooking, etc!)
@@ -150,6 +160,11 @@ int32 UTextureLODSettings::CalculateLODBias(int32 Width, int32 Height, int32 Max
 		// Considering them again here would apply them twice.
 		UsedLODBias	+= LODBias + LODGroupInfo.LODBias;
 	}
+	
+	if (LODGroup == TEXTUREGROUP_UI)
+	{
+		UsedLODBias += GUITextureLODBias;
+	}
 
 	int32 MinLOD		= LODGroupInfo.MinLODMipCount;
 	int32 MaxLOD		= LODGroupInfo.MaxLODMipCount;
@@ -278,9 +293,17 @@ ETextureSamplerFilter UTextureLODSettings::GetSamplerFilter(const UTexture* Text
 
 	switch(Texture->Filter)
 	{
-		case TF_Nearest: Filter = ETextureSamplerFilter::Point; break;
-		case TF_Bilinear: Filter = ETextureSamplerFilter::Bilinear; break;
-		case TF_Trilinear: Filter = ETextureSamplerFilter::Trilinear; break;
+		case TF_FirstMipNearest:
+		case TF_Nearest: 
+			Filter = ETextureSamplerFilter::Point; 
+			break;
+		case TF_FirstMipBilinear:
+		case TF_Bilinear: 
+			Filter = ETextureSamplerFilter::Bilinear; 
+			break;
+		case TF_Trilinear: 
+			Filter = ETextureSamplerFilter::Trilinear; 
+			break;
 
 		// TF_Default
 		default:
diff --git a/Engine/Source/Runtime/SlateRHIRenderer/Private/SlateRHIRenderingPolicy.cpp b/Engine/Source/Runtime/SlateRHIRenderer/Private/SlateRHIRenderingPolicy.cpp
--- a/Engine/Source/Runtime/SlateRHIRenderer/Private/SlateRHIRenderingPolicy.cpp
+++ b/Engine/Source/Runtime/SlateRHIRenderer/Private/SlateRHIRenderingPolicy.cpp
@@ -1335,9 +1335,17 @@ ETextureSamplerFilter FSlateRHIRenderingPolicy::GetSamplerFilter(const UTexture*
 
 	switch (Texture->Filter)
 	{
-	case TF_Nearest: Filter = ETextureSamplerFilter::Point; break;
-	case TF_Bilinear: Filter = ETextureSamplerFilter::Bilinear; break;
-	case TF_Trilinear: Filter = ETextureSamplerFilter::Trilinear; break;
+	case TF_FirstMipNearest:
+	case TF_Nearest: 
+		Filter = ETextureSamplerFilter::Point; 
+		break;
+	case TF_FirstMipBilinear:
+	case TF_Bilinear:
+		Filter = ETextureSamplerFilter::Bilinear; 
+		break;
+	case TF_Trilinear: 
+		Filter = ETextureSamplerFilter::Trilinear; 
+		break;
 
 		// TF_Default
 	default:
