Fixing occlusion interpolation time.

diff --git a/Engine/Source/Runtime/Engine/Private/ActiveSound.cpp b/Engine/Source/Runtime/Engine/Private/ActiveSound.cpp
--- a/Engine/Source/Runtime/Engine/Private/ActiveSound.cpp
+++ b/Engine/Source/Runtime/Engine/Private/ActiveSound.cpp
@@ -905,6 +905,7 @@ void FActiveSound::CheckOcclusion(const FVector ListenerLocation, const FVector
 		{
 			InterpolationTime = AttenuationSettingsPtr->OcclusionInterpolationTime;
 		}
+		bHasCheckedOcclusion = true;
 	}
 
 	if (bIsOccluded)
