Removing the EInternalObjectFlags::Async flag for UBodySetup instances that have been created off-game thread via UShapeComponent::CreateShapeBodySetupIfNeeded #UE4

This fixes a memory leak resulting from UShapeComponent::CreateShapeBodySetupIfNeeded sometimes being called as part of async navigation collision gathering.

diff --git a/Engine/Source/Runtime/Engine/Private/Components/ShapeComponent.cpp b/Engine/Source/Runtime/Engine/Private/Components/ShapeComponent.cpp
--- a/Engine/Source/Runtime/Engine/Private/Components/ShapeComponent.cpp
+++ b/Engine/Source/Runtime/Engine/Private/Components/ShapeComponent.cpp
@@ -131,6 +131,12 @@ void UShapeComponent::CreateShapeBodySetupIfNeeded()
 
 		// If this component is in GC cluster, make sure we add the body setup to it to
 		ShapeBodySetup->AddToCluster(this);
+		// if we got created outside of game thread, but got added to a cluster, 
+		// we no longer need the Async flag
+		if (ShapeBodySetup->HasAnyInternalFlags(EInternalObjectFlags::Async) && GUObjectClusters.GetObjectCluster(ShapeBodySetup))
+		{
+			ShapeBodySetup->ClearInternalFlags(EInternalObjectFlags::Async);
+		}
 		
 		ShapeBodySetup->CollisionTraceFlag = CTF_UseSimpleAsComplex;
 		AddShapeToGeomArray<ShapeElemType>();
