SoundNode hash fix-up.  Some sound nodes were exhibiting hash collisions in rare instances.
udn: https://udn.unrealengine.com/questions/481812/audio-crossfade-hash-bug.html

diff --git a/Engine/Source/Runtime/Engine/Private/SoundNode.cpp b/Engine/Source/Runtime/Engine/Private/SoundNode.cpp
--- a/Engine/Source/Runtime/Engine/Private/SoundNode.cpp
+++ b/Engine/Source/Runtime/Engine/Private/SoundNode.cpp
@@ -67,13 +67,21 @@ UEdGraphNode* USoundNode::GetGraphNode() const
 UPTRINT USoundNode::GetNodeWaveInstanceHash(const UPTRINT ParentWaveInstanceHash, const USoundNode* ChildNode, const uint32 ChildIndex)
 {
 	checkf(ChildIndex < MAX_ALLOWED_CHILD_NODES, TEXT("Too many children (%d) in SoundCue '%s'"), ChildIndex, *CastChecked<USoundCue>(ChildNode->GetOuter())->GetFullName());
-	return ((ParentWaveInstanceHash << ChildIndex) ^ (UPTRINT)ChildNode);
+
+	return GetNodeWaveInstanceHash(ParentWaveInstanceHash, reinterpret_cast<const UPTRINT>(ChildNode), ChildIndex);
 }
 
 UPTRINT USoundNode::GetNodeWaveInstanceHash(const UPTRINT ParentWaveInstanceHash, const UPTRINT ChildNodeHash, const uint32 ChildIndex)
 {
-	checkf(ChildIndex < MAX_ALLOWED_CHILD_NODES, TEXT("Too many children (%d) in SoundCue"), ChildIndex);
+#define USE_NEW_SOUNDCUE_NODE_HASH 1
+#if USE_NEW_SOUNDCUE_NODE_HASH
+	const uint32 ChildHash = PointerHash(reinterpret_cast<const void*>(ChildNodeHash), GetTypeHash(ChildIndex));
+	const uint32 Hash = PointerHash(reinterpret_cast<const void*>(ParentWaveInstanceHash), ChildHash);
+
+	return static_cast<UPTRINT>(Hash);
+#else
 	return ((ParentWaveInstanceHash << ChildIndex) ^ ChildNodeHash);
+#endif // USE_NEW_SOUNDCUE_NODE_HASH
 }
 
 void USoundNode::ParseNodes( FAudioDevice* AudioDevice, const UPTRINT NodeWaveInstanceHash, FActiveSound& ActiveSound, const FSoundParseParameters& ParseParams, TArray<FWaveInstance*>& WaveInstances )
