UE4 - fix case of include

diff --git a/Engine/Source/Runtime/Engine/Classes/Animation/AnimSequence.h b/Engine/Source/Runtime/Engine/Classes/Animation/AnimSequence.h
--- a/Engine/Source/Runtime/Engine/Classes/Animation/AnimSequence.h
+++ b/Engine/Source/Runtime/Engine/Classes/Animation/AnimSequence.h
@@ -17,7 +17,7 @@
 #include "Animation/AnimSequenceBase.h"
 
 #include "Async/MappedFileHandle.h"
-#include "HAL/PlatformFileManager.h"
+#include "HAL/PlatformFilemanager.h"
 #include "GenericPlatform/GenericPlatformFile.h"
 #include "Misc/Paths.h"
 #include "Serialization/BulkData.h"
