Null ptr fixes for APlayerCameraManager

diff --git a/Engine/Source/Runtime/Engine/Private/PlayerCameraManager.cpp b/Engine/Source/Runtime/Engine/Private/PlayerCameraManager.cpp
--- a/Engine/Source/Runtime/Engine/Private/PlayerCameraManager.cpp
+++ b/Engine/Source/Runtime/Engine/Private/PlayerCameraManager.cpp
@@ -1167,7 +1167,8 @@ AEmitterCameraLensEffectBase* APlayerCameraManager::FindCameraLensEffect(TSubcla
 	for (int32 i = 0; i < CameraLensEffects.Num(); ++i)
 	{
 		AEmitterCameraLensEffectBase* LensEffect = CameraLensEffects[i];
-		if ( !LensEffect->IsPendingKill() &&
+		if (LensEffect &&
+			!LensEffect->IsPendingKill() &&
 			( (LensEffect->GetClass() == LensEffectEmitterClass) ||
 			(LensEffect->EmittersToTreatAsSame.Find(LensEffectEmitterClass) != INDEX_NONE) ||
 			(GetDefault<AEmitterCameraLensEffectBase>(LensEffectEmitterClass)->EmittersToTreatAsSame.Find(LensEffect->GetClass()) != INDEX_NONE ) ) )
@@ -1185,7 +1186,8 @@ AEmitterCameraLensEffectBase* APlayerCameraManager::AddCameraLensEffect(TSubclas
 	if (LensEffectEmitterClass != NULL)
 	{
 		AEmitterCameraLensEffectBase* LensEffect = NULL;
-		if (!GetDefault<AEmitterCameraLensEffectBase>(LensEffectEmitterClass)->bAllowMultipleInstances)
+		const AEmitterCameraLensEffectBase* LensEffectClassDefaultObject = GetDefault<AEmitterCameraLensEffectBase>(LensEffectEmitterClass);
+		if (LensEffectClassDefaultObject && !LensEffectClassDefaultObject->bAllowMultipleInstances)
 		{
 			LensEffect = FindCameraLensEffect(LensEffectEmitterClass);
 
