Remove FMaterialShaderMap::VertexFactoryMap since it not used anywhere. Reduce class size by 80 bytes

diff --git a/Engine/Source/Runtime/Engine/Public/MaterialShared.h b/Engine/Source/Runtime/Engine/Public/MaterialShared.h
--- a/Engine/Source/Runtime/Engine/Public/MaterialShared.h
+++ b/Engine/Source/Runtime/Engine/Public/MaterialShared.h
@@ -849,7 +849,6 @@ public:
 		return sizeof(*this)
 			+ MeshShaderMaps.GetAllocatedSize()
 			+ OrderedMeshShaderMaps.GetAllocatedSize()
-			+ VertexFactoryMap.GetAllocatedSize()
 			+ MaterialCompilationOutput.UniformExpressionSet.GetAllocatedSize()
 #if ALLOW_SHADERMAP_DEBUG_DATA
 			+ FriendlyName.GetAllocatedSize()
@@ -951,9 +950,6 @@ private:
 	/** The static parameter set that this shader map was compiled with */
 	FMaterialShaderMapId ShaderMapId;
 
-	/** A map from vertex factory type to the material's cached shaders for that vertex factory type. */
-	TMap<FVertexFactoryType*,FMeshMaterialShaderMap*> VertexFactoryMap;
-
 	/** Uniform expressions generated from the material compile. */
 	FMaterialCompilationOutput MaterialCompilationOutput;
 
