Fixed the nav bounds calculation bug in InstancedStaticMeshComponents #UE4

Same issue as was reported for HierarchicalInstancedStaticMeshComponents in UE-69472, similar fix.
Thanks to [REDACTED_HANDLE] on github for noticing and reporting that.

diff --git a/Engine/Source/Runtime/Engine/Classes/Components/InstancedStaticMeshComponent.h b/Engine/Source/Runtime/Engine/Classes/Components/InstancedStaticMeshComponent.h
--- a/Engine/Source/Runtime/Engine/Classes/Components/InstancedStaticMeshComponent.h
+++ b/Engine/Source/Runtime/Engine/Classes/Components/InstancedStaticMeshComponent.h
@@ -263,6 +263,7 @@ public:
 
 	//~ Begin UNavRelevantInterface Interface
 	virtual void GetNavigationData(FNavigationRelevantData& Data) const override;
+	virtual FBox GetNavigationBounds() const override;
 	//~ End UPrimitiveComponent Interface
 
 	//~ Begin UObject Interface
diff --git a/Engine/Source/Runtime/Engine/Private/InstancedStaticMesh.cpp b/Engine/Source/Runtime/Engine/Private/InstancedStaticMesh.cpp
--- a/Engine/Source/Runtime/Engine/Private/InstancedStaticMesh.cpp
+++ b/Engine/Source/Runtime/Engine/Private/InstancedStaticMesh.cpp
@@ -2227,6 +2227,11 @@ void UInstancedStaticMeshComponent::GetNavigationData(FNavigationRelevantData& D
 	}
 }
 
+FBox UInstancedStaticMeshComponent::GetNavigationBounds() const
+{
+	return CalcBounds(GetComponentTransform()).GetBox();
+}
+
 void UInstancedStaticMeshComponent::GetNavigationPerInstanceTransforms(const FBox& AreaBox, TArray<FTransform>& InstanceData) const
 {
 	for (const auto& InstancedData : PerInstanceSMData)
