fix a check failure caused by data race

diff --git a/Engine/Source/Runtime/Engine/Classes/Engine/StreamableRenderAsset.h b/Engine/Source/Runtime/Engine/Classes/Engine/StreamableRenderAsset.h
--- a/Engine/Source/Runtime/Engine/Classes/Engine/StreamableRenderAsset.h
+++ b/Engine/Source/Runtime/Engine/Classes/Engine/StreamableRenderAsset.h
@@ -188,7 +188,6 @@ public:
 #if WITH_EDITOR
 		return GetNumResidentMips();
 #else
-		check(CachedNumResidentLODs == GetNumResidentMips());
 		return CachedNumResidentLODs;
 #endif
 	}
@@ -198,7 +197,6 @@ public:
 #if WITH_EDITOR
 		return IsReadyForStreaming();
 #else
-		check(bCachedReadyForStreaming == IsReadyForStreaming());
 		return bCachedReadyForStreaming;
 #endif
 	}
diff --git a/Engine/Source/Runtime/Engine/Private/Streaming/StaticMeshUpdate.cpp b/Engine/Source/Runtime/Engine/Private/Streaming/StaticMeshUpdate.cpp
--- a/Engine/Source/Runtime/Engine/Private/Streaming/StaticMeshUpdate.cpp
+++ b/Engine/Source/Runtime/Engine/Private/Streaming/StaticMeshUpdate.cpp
@@ -208,6 +208,7 @@ void FStaticMeshStreamIn::DoFinishUpdate(const FContext& Context)
 				IntermediateBuffersArray[LODIdx].TransferBuffers(LODResource, Batcher);
 			}
 		}
+		check(Mesh->GetCachedNumResidentLODs() == RenderData->LODResources.Num() - RenderData->CurrentFirstLODIdx);
 		RenderData->CurrentFirstLODIdx = PendingFirstMip;
 		Mesh->SetCachedNumResidentLODs(static_cast<uint8>(RenderData->LODResources.Num() - PendingFirstMip));
 	}
@@ -244,6 +245,7 @@ void FStaticMeshStreamOut::DoReleaseBuffers(const FContext& Context)
 	if (!IsCancelled() && Mesh && RenderData)
 	{
 		check(CurrentFirstLODIdx == RenderData->CurrentFirstLODIdx && PendingFirstMip > CurrentFirstLODIdx);
+		check(Mesh->GetCachedNumResidentLODs() == RenderData->LODResources.Num() - RenderData->CurrentFirstLODIdx);
 
 		RenderData->CurrentFirstLODIdx = PendingFirstMip;
 		Mesh->SetCachedNumResidentLODs(static_cast<uint8>(RenderData->LODResources.Num() - PendingFirstMip));
diff --git a/Engine/Source/Runtime/Engine/Private/Texture2D.cpp b/Engine/Source/Runtime/Engine/Private/Texture2D.cpp
--- a/Engine/Source/Runtime/Engine/Private/Texture2D.cpp
+++ b/Engine/Source/Runtime/Engine/Private/Texture2D.cpp
@@ -1641,6 +1641,8 @@ void FTexture2DResource::UpdateTexture(FTexture2DRHIRef& InTextureRHI, int32 InN
 			TextureRHI->DoNoDeferDelete();
 		}
 
+		check(Owner->GetCachedNumResidentLODs() == NumMips - CurrentFirstMip);
+
 		TextureRHI		= InTextureRHI;
 		Texture2DRHI	= InTextureRHI;
 		CurrentFirstMip = InNewFirstMip;
