Fix Role and RemoteRole getting stuck on clients.

diff --git a/Engine/Source/Runtime/Engine/Private/DataReplication.cpp b/Engine/Source/Runtime/Engine/Private/DataReplication.cpp
--- a/Engine/Source/Runtime/Engine/Private/DataReplication.cpp
+++ b/Engine/Source/Runtime/Engine/Private/DataReplication.cpp
@@ -400,7 +400,6 @@ void FObjectReplicator::StartReplicating( class UActorChannel * InActorChannel )
 		}
 	}
 }
-
 static FORCEINLINE void ValidateRetirementHistory( const FPropertyRetirement & Retire, const UObject* Object )
 {
 #if !UE_BUILD_SHIPPING
@@ -1264,7 +1263,7 @@ bool FObjectReplicator::ReplicateProperties( FOutBunch & Bunch, FReplicationFlag
 
 	// Update change list (this will re-use work done by previous connections)
 	FSendingRepState* SendingRepState = ((Connection->ResendAllDataState == EResendAllDataState::SinceCheckpoint) && CheckpointRepState.IsValid()) ? CheckpointRepState->GetSendingRepState() : RepState->GetSendingRepState();
-	RepLayout->UpdateChangelistMgr(*ChangelistMgr, SendingRepState, Object, Connection->Driver->ReplicationFrame, RepFlags, OwningChannel->bForceCompareProperties);
+	RepLayout->UpdateChangelistMgr(SendingRepState, *ChangelistMgr, Object, Connection->Driver->ReplicationFrame, RepFlags, OwningChannel->bForceCompareProperties);
 
 	// Replicate properties in the layout
 	const bool bHasRepLayout = RepLayout->ReplicateProperties(SendingRepState, ChangelistMgr->GetRepChangelistState(), (uint8*)Object, ObjectClass, OwningChannel, Writer, RepFlags);
@@ -1975,4 +1974,4 @@ FScopedActorRoleSwap::~FScopedActorRoleSwap()
 	{
 		Actor->SwapRoles();
 	}
-}
\ No newline at end of file
+}
diff --git a/Engine/Source/Runtime/Engine/Private/RepLayout.cpp b/Engine/Source/Runtime/Engine/Private/RepLayout.cpp
--- a/Engine/Source/Runtime/Engine/Private/RepLayout.cpp
+++ b/Engine/Source/Runtime/Engine/Private/RepLayout.cpp
@@ -689,8 +689,8 @@ FReceivingRepState::FReceivingRepState(FRepStateStaticBuffer&& InStaticBuffer) :
 }
 
 void FRepLayout::UpdateChangelistMgr(
-	FReplicationChangelistMgr& InChangelistMgr,
 	FSendingRepState* RESTRICT RepState,
+	FReplicationChangelistMgr& InChangelistMgr,
 	const UObject* InObject,
 	const uint32 ReplicationFrame,
 	const FReplicationFlags& RepFlags,
@@ -718,17 +718,16 @@ void FReplicationChangelistMgr::CountBytes(FArchive& Ar) const
 	RepChangelistState.CountBytes(Ar);
 }
 
-
 uint16 FRepLayout::CompareProperties_r(
 	FSendingRepState* RESTRICT RepState,
-	const int32				CmdStart,
-	const int32				CmdEnd,
-	const uint8* RESTRICT	ShadowData,
-	const uint8* RESTRICT	Data,
-	TArray<uint16>&			Changed,
-	uint16					Handle,
-	const bool				bIsInitial,
-	const bool				bForceFail) const
+	const int32 CmdStart,
+	const int32 CmdEnd,
+	const uint8* RESTRICT ShadowData,
+	const uint8* RESTRICT Data,
+	TArray<uint16>& Changed,
+	uint16 Handle,
+	const bool bIsInitial,
+	const bool bForceFail) const
 {
 	check(ShadowData);
 
@@ -761,7 +760,7 @@ uint16 FRepLayout::CompareProperties_r(
 			}
 
 			// Once we hit an array, start using a stack based approach
-			CompareProperties_Array_r(RepState, ShadowData + Cmd.ShadowOffset, ( const uint8* )Data + Cmd.Offset, Changed, CmdIndex, Handle, bIsInitial, bForceFail );
+			CompareProperties_Array_r(RepState, ShadowData + Cmd.ShadowOffset, (const uint8*)Data + Cmd.Offset, Changed, CmdIndex, Handle, bIsInitial, bForceFail);
 			CmdIndex = Cmd.EndCmd - 1;		// The -1 to handle the ++ in the for loop
 			continue;
 		}
@@ -803,14 +802,13 @@ uint16 FRepLayout::CompareProperties_r(
 
 void FRepLayout::CompareProperties_Array_r(
 	FSendingRepState* RESTRICT RepState,
-	const uint8* RESTRICT	ShadowData,
-	const uint8* RESTRICT	Data,
-	TArray<uint16>&			Changed,
-	const uint16			CmdIndex,
-	const uint16			Handle,
-	const bool				bIsInitial,
-	const bool				bForceFail
-	) const
+	const uint8* RESTRICT ShadowData,
+	const uint8* RESTRICT Data,
+	TArray<uint16>& Changed,
+	const uint16 CmdIndex,
+	const uint16 Handle,
+	const bool bIsInitial,
+	const bool bForceFail) const
 {
 	const FRepLayoutCmd& Cmd = Cmds[CmdIndex];
 
@@ -821,7 +819,7 @@ void FRepLayout::CompareProperties_Array_r(
 	const uint16 ShadowArrayNum = ShadowArray->Num();
 
 	// Make the shadow state match the actual state at the time of compare
-	FScriptArrayHelper StoredArrayHelper((UArrayProperty*)Cmd.Property, ShadowData);
+	FScriptArrayHelper StoredArrayHelper((UArrayProperty*)Cmd.Property, ShadowArray);
 	StoredArrayHelper.Resize(ArrayNum);
 
 	TArray<uint16> ChangedLocal;
@@ -861,9 +859,9 @@ void FRepLayout::CompareProperties_Array_r(
 
 bool FRepLayout::CompareProperties(
 	FSendingRepState* RESTRICT RepState,
-	FRepChangelistState* RESTRICT	RepChangelistState,
-	const uint8* RESTRICT			Data,
-	const FReplicationFlags&		RepFlags) const
+	FRepChangelistState* RESTRICT RepChangelistState,
+	const uint8* RESTRICT Data,
+	const FReplicationFlags& RepFlags) const
 {
 	SCOPE_CYCLE_COUNTER(STAT_NetReplicateDynamicPropCompareTime);
 
@@ -1140,8 +1138,8 @@ PRAGMA_ENABLE_DEPRECATION_WARNINGS
 
 		if (Changed.Num() > 0)
 		{
-		SendProperties_BackwardsCompatible(RepState, ChangeTracker, Data, OwningChannel->Connection, Writer, Changed);
-	}
+			SendProperties_BackwardsCompatible(RepState, ChangeTracker, Data, OwningChannel->Connection, Writer, Changed);
+		}
 	}
 	else if (Changed.Num() > 0)
 	{
@@ -1850,7 +1848,7 @@ void FRepLayout::SendProperties_r(
 			{
 				Data = reinterpret_cast<const uint8*>(&(RepState->SavedRemoteRole));
 			}
-		}		
+		}
 
 #if !(UE_BUILD_SHIPPING || UE_BUILD_TEST)
 		if (GDoReplicationContextString> 0)
@@ -2181,7 +2179,7 @@ void FRepLayout::SendProperties_BackwardsCompatible_r(
 			{
 				Data = reinterpret_cast<const uint8*>(&(RepState->SavedRemoteRole));
 			}
-		}		
+		}
 
 		WriteProperty_BackwardsCompatible(Writer, Cmd, HandleIterator.CmdIndex, Owner, Data, bDoChecksum);
 	}
@@ -2191,13 +2189,13 @@ void FRepLayout::SendProperties_BackwardsCompatible_r(
 
 void FRepLayout::SendAllProperties_BackwardsCompatible_r(
 	FSendingRepState* RESTRICT RepState,
-	FNetBitWriter&						Writer,
-	const bool							bDoChecksum,
-	UPackageMapClient*					PackageMapClient,
-	FNetFieldExportGroup*				NetFieldExportGroup,
-	const int32							CmdStart,
-	const int32							CmdEnd, 
-	const uint8*						SourceData) const
+	FNetBitWriter& Writer,
+	const bool bDoChecksum,
+	UPackageMapClient* PackageMapClient,
+	FNetFieldExportGroup* NetFieldExportGroup,
+	const int32 CmdStart,
+	const int32 CmdEnd, 
+	const uint8* SourceData) const
 {
 	FNetBitWriter TempWriter(Writer.PackageMap, 0);
 
diff --git a/Engine/Source/Runtime/Engine/Public/Net/RepLayout.h b/Engine/Source/Runtime/Engine/Public/Net/RepLayout.h
--- a/Engine/Source/Runtime/Engine/Public/Net/RepLayout.h
+++ b/Engine/Source/Runtime/Engine/Public/Net/RepLayout.h
@@ -593,12 +593,6 @@ public:
 	/** The unique list of properties that have changed since the channel was first opened */
 	TArray<uint16> LifetimeChangelist;
 
-	/**
-	 * A map tracking which replication conditions are currently active.
-	 * @see ELifetimeCondition.
-	 */
-	TStaticBitArray<COND_Max> ConditionMap;
-
 	/**
 	 * Properties which are inactive through conditions have their changes stored here, so they can be 
 	 * applied if/when the property becomes active.
@@ -609,6 +603,12 @@ public:
 
 	/** Cached set of inactive parent commands. */
 	TBitArray<> InactiveParents;
+
+	// Cache off the RemoteRole and Role per connection to avoid issues with
+	// FScopedRoleDowngrade. See UE-66313 (among others).
+
+	ENetRole SavedRemoteRole = ROLE_MAX;
+	ENetRole SavedRole = ROLE_MAX;
 };
 
 
@@ -1548,8 +1548,8 @@ public:
 		bool& bOutGuidsChanged) const;
 
 	void UpdateChangelistMgr(
-		FReplicationChangelistMgr& InChangelistMgr,
 		FSendingRepState* RESTRICT RepState,
+		FReplicationChangelistMgr& InChangelistMgr,
 		const UObject* InObject,
 		const uint32 ReplicationFrame,
 		const FReplicationFlags& RepFlags,
@@ -1933,4 +1933,4 @@ private:
 
 	/** Shared comparison to default state for multicast rpc */
 	TBitArray<> SharedInfoRPCParentsChanged;	
-};
\ No newline at end of file
+};
