Fix merge build issue

diff --git a/Engine/Source/Runtime/Engine/Public/Net/RepLayout.h b/Engine/Source/Runtime/Engine/Public/Net/RepLayout.h
--- a/Engine/Source/Runtime/Engine/Public/Net/RepLayout.h
+++ b/Engine/Source/Runtime/Engine/Public/Net/RepLayout.h
@@ -603,12 +603,6 @@ public:
 
 	/** Cached set of inactive parent commands. */
 	TBitArray<> InactiveParents;
-
-	// Cache off the RemoteRole and Role per connection to avoid issues with
-	// FScopedRoleDowngrade. See UE-66313 (among others).
-
-	ENetRole SavedRemoteRole = ROLE_MAX;
-	ENetRole SavedRole = ROLE_MAX;
 };
 
 
