UE4.22 - Use Vertex declaration cache
#rnx

diff --git a/Engine/Plugins/Experimental/AR/PointCloud/Source/PointCloud/Private/PointCloudVertexFactory.cpp b/Engine/Plugins/Experimental/AR/PointCloud/Source/PointCloud/Private/PointCloudVertexFactory.cpp
--- a/Engine/Plugins/Experimental/AR/PointCloud/Source/PointCloud/Private/PointCloudVertexFactory.cpp
+++ b/Engine/Plugins/Experimental/AR/PointCloud/Source/PointCloud/Private/PointCloudVertexFactory.cpp
@@ -69,7 +69,7 @@ public:
 	{
 		FVertexDeclarationElementList Elements;
 		Elements.Add(FVertexElement(0, 0, VET_Float4, 0, sizeof(FVector4)));
-		VertexDeclarationRHI = RHICreateVertexDeclaration(Elements);
+		VertexDeclarationRHI = PipelineStateCache::GetOrCreateVertexDeclaration(Elements);
 	}
 
 	virtual void ReleaseRHI() override
diff --git a/Engine/Plugins/FX/Niagara/Source/NiagaraVertexFactories/Private/NiagaraRibbonVertexFactory.cpp b/Engine/Plugins/FX/Niagara/Source/NiagaraVertexFactories/Private/NiagaraRibbonVertexFactory.cpp
--- a/Engine/Plugins/FX/Niagara/Source/NiagaraVertexFactories/Private/NiagaraRibbonVertexFactory.cpp
+++ b/Engine/Plugins/FX/Niagara/Source/NiagaraVertexFactories/Private/NiagaraRibbonVertexFactory.cpp
@@ -155,7 +155,7 @@ public:
 		// Create the vertex declaration for rendering the factory normally.
 		// This is done in InitDynamicRHI instead of InitRHI to allow FNiagaraRibbonVertexFactory::InitRHI
 		// to rely on it being initialized, since InitDynamicRHI is called before InitRHI.
-		VertexDeclarationRHI = RHICreateVertexDeclaration(Elements);
+		VertexDeclarationRHI = PipelineStateCache::GetOrCreateVertexDeclaration(Elements);
 	}
 
 	virtual void ReleaseDynamicRHI()
diff --git a/Engine/Plugins/FX/Niagara/Source/NiagaraVertexFactories/Private/NiagaraSpriteVertexFactory.cpp b/Engine/Plugins/FX/Niagara/Source/NiagaraVertexFactories/Private/NiagaraSpriteVertexFactory.cpp
--- a/Engine/Plugins/FX/Niagara/Source/NiagaraVertexFactories/Private/NiagaraSpriteVertexFactory.cpp
+++ b/Engine/Plugins/FX/Niagara/Source/NiagaraVertexFactories/Private/NiagaraSpriteVertexFactory.cpp
@@ -183,7 +183,7 @@ public:
 		// Create the vertex declaration for rendering the factory normally.
 		// This is done in InitDynamicRHI instead of InitRHI to allow FParticleSpriteVertexFactory::InitRHI
 		// to rely on it being initialized, since InitDynamicRHI is called before InitRHI.
-		VertexDeclarationRHI = RHICreateVertexDeclaration(Elements);
+		VertexDeclarationRHI = PipelineStateCache::GetOrCreateVertexDeclaration(Elements);
 	}
 
 	virtual void ReleaseDynamicRHI()
diff --git a/Engine/Plugins/Media/WebMMedia/Source/WebMMedia/Private/Player/WebMVideoDecoder.cpp b/Engine/Plugins/Media/WebMMedia/Source/WebMMedia/Private/Player/WebMVideoDecoder.cpp
--- a/Engine/Plugins/Media/WebMMedia/Source/WebMMedia/Private/Player/WebMVideoDecoder.cpp
+++ b/Engine/Plugins/Media/WebMMedia/Source/WebMMedia/Private/Player/WebMVideoDecoder.cpp
@@ -34,7 +34,7 @@ namespace
 			uint16 Stride = sizeof(FMediaElementVertex);
 			Elements.Add(FVertexElement(0, STRUCT_OFFSET(FMediaElementVertex, Position), VET_Float4, 0, Stride));
 			Elements.Add(FVertexElement(0, STRUCT_OFFSET(FMediaElementVertex, TextureCoordinate), VET_Float2, 1, Stride));
-			VertexDeclarationRHI = RHICreateVertexDeclaration(Elements);
+			VertexDeclarationRHI = PipelineStateCache::GetOrCreateVertexDeclaration(Elements);
 
 			TResourceArray<FMediaElementVertex> Vertices;
 			Vertices.AddUninitialized(4);
diff --git a/Engine/Plugins/Runtime/Oculus/OculusVR/Source/OculusHMD/Private/OculusHMD_StressTester.cpp b/Engine/Plugins/Runtime/Oculus/OculusVR/Source/OculusHMD/Private/OculusHMD_StressTester.cpp
--- a/Engine/Plugins/Runtime/Oculus/OculusVR/Source/OculusHMD/Private/OculusHMD_StressTester.cpp
+++ b/Engine/Plugins/Runtime/Oculus/OculusVR/Source/OculusHMD/Private/OculusHMD_StressTester.cpp
@@ -84,7 +84,7 @@ public:
 		uint32 Stride = sizeof(FTextureVertex);
 		Elements.Add(FVertexElement(0, STRUCT_OFFSET(FTextureVertex, Position), VET_Float4, 0, Stride));
 		Elements.Add(FVertexElement(0, STRUCT_OFFSET(FTextureVertex, UV), VET_Float2, 1, Stride));
-		VertexDeclarationRHI = RHICreateVertexDeclaration(Elements);
+		VertexDeclarationRHI = PipelineStateCache::GetOrCreateVertexDeclaration(Elements);
 	}
 
 	virtual void ReleaseRHI() override
diff --git a/Engine/Source/Runtime/Apple/MetalRHI/Private/MetalCommands.cpp b/Engine/Source/Runtime/Apple/MetalRHI/Private/MetalCommands.cpp
--- a/Engine/Source/Runtime/Apple/MetalRHI/Private/MetalCommands.cpp
+++ b/Engine/Source/Runtime/Apple/MetalRHI/Private/MetalCommands.cpp
@@ -29,7 +29,7 @@ public:
 	{
 		FVertexDeclarationElementList Elements;
 		Elements.Add(FVertexElement(0, 0, VET_Float4, 0, sizeof(FVector4)));
-		VertexDeclarationRHI = RHICreateVertexDeclaration(Elements);
+		VertexDeclarationRHI = PipelineStateCache::GetOrCreateVertexDeclaration(Elements);
 	}
 	virtual void ReleaseRHI() override
 	{
diff --git a/Engine/Source/Runtime/Engine/Private/Particles/ParticleBeamTrailVertexFactory.cpp b/Engine/Source/Runtime/Engine/Private/Particles/ParticleBeamTrailVertexFactory.cpp
--- a/Engine/Source/Runtime/Engine/Private/Particles/ParticleBeamTrailVertexFactory.cpp
+++ b/Engine/Source/Runtime/Engine/Private/Particles/ParticleBeamTrailVertexFactory.cpp
@@ -8,6 +8,7 @@
 #include "ParticleHelper.h"
 #include "ShaderParameterUtils.h"
 #include "MeshMaterialShader.h"
+#include "PipelineStateCache.h"
 
 IMPLEMENT_GLOBAL_SHADER_PARAMETER_STRUCT(FParticleBeamTrailUniformParameters,"BeamTrailVF");
 
@@ -88,7 +89,7 @@ public:
 		// Create the vertex declaration for rendering the factory normally.
 		// This is done in InitDynamicRHI instead of InitRHI to allow FParticleBeamTrailVertexFactory::InitRHI
 		// to rely on it being initialized, since InitDynamicRHI is called before InitRHI.
-		VertexDeclarationRHI = RHICreateVertexDeclaration(Elements);
+		VertexDeclarationRHI = PipelineStateCache::GetOrCreateVertexDeclaration(Elements);
 	}
 
 	virtual void ReleaseDynamicRHI()
diff --git a/Engine/Source/Runtime/Engine/Private/Particles/ParticleCurveTexture.cpp b/Engine/Source/Runtime/Engine/Private/Particles/ParticleCurveTexture.cpp
--- a/Engine/Source/Runtime/Engine/Private/Particles/ParticleCurveTexture.cpp
+++ b/Engine/Source/Runtime/Engine/Private/Particles/ParticleCurveTexture.cpp
@@ -154,7 +154,7 @@ public:
 			Offset += sizeof(FVector2D);
 		}
 
-		VertexDeclarationRHI = RHICreateVertexDeclaration(Elements);
+		VertexDeclarationRHI = PipelineStateCache::GetOrCreateVertexDeclaration(Elements);
 	}
 
 	virtual void ReleaseRHI() override
diff --git a/Engine/Source/Runtime/Engine/Private/Particles/ParticleGpuSimulation.cpp b/Engine/Source/Runtime/Engine/Private/Particles/ParticleGpuSimulation.cpp
--- a/Engine/Source/Runtime/Engine/Private/Particles/ParticleGpuSimulation.cpp
+++ b/Engine/Source/Runtime/Engine/Private/Particles/ParticleGpuSimulation.cpp
@@ -676,7 +676,7 @@ public:
 		/** The stream to read the texture coordinates from. */
 		Elements.Add(FVertexElement(0, 0, VET_Float2, 0, sizeof(FVector2D), false));
 
-		VertexDeclarationRHI = RHICreateVertexDeclaration(Elements);
+		VertexDeclarationRHI = PipelineStateCache::GetOrCreateVertexDeclaration(Elements);
 	}
 
 	/**
@@ -1337,7 +1337,7 @@ public:
 		FVertexDeclarationElementList Elements;
 		// TexCoord.
 		Elements.Add(FVertexElement(0, 0, VET_Float2, 0, sizeof(FVector2D), /*bUseInstanceIndex=*/ false));
-		VertexDeclarationRHI = RHICreateVertexDeclaration( Elements );
+		VertexDeclarationRHI = PipelineStateCache::GetOrCreateVertexDeclaration( Elements );
 	}
 
 	virtual void ReleaseRHI() override
@@ -1366,7 +1366,7 @@ public:
 		Elements.Add(FVertexElement(0, 0, VET_Float2, 0, sizeof(FVector2D), /*bUseInstanceIndex=*/ false));
 		// TileOffsets
 		Elements.Add(FVertexElement(1, 0, VET_Float2, 1, sizeof(FVector2D), /*bUseInstanceIndex=*/ true));
-		VertexDeclarationRHI = RHICreateVertexDeclaration( Elements );
+		VertexDeclarationRHI = PipelineStateCache::GetOrCreateVertexDeclaration( Elements );
 	}
 
 	virtual void ReleaseRHI() override
@@ -1867,7 +1867,7 @@ public:
 			Offset += sizeof(FVector2D);
 		}
 
-		VertexDeclarationRHI = RHICreateVertexDeclaration( Elements );
+		VertexDeclarationRHI = PipelineStateCache::GetOrCreateVertexDeclaration( Elements );
 	}
 
 	virtual void ReleaseRHI() override
@@ -2080,7 +2080,7 @@ public:
 	{
 		FVertexDeclarationElementList Elements;
 		Elements.Add(FVertexElement(0, 0, VET_Float2, 0, sizeof(FVector2D)));
-		VertexDeclarationRHI = RHICreateVertexDeclaration( Elements );
+		VertexDeclarationRHI = PipelineStateCache::GetOrCreateVertexDeclaration( Elements );
 	}
 
 	virtual void ReleaseRHI() override
diff --git a/Engine/Source/Runtime/Engine/Private/Particles/ParticleVertexFactory.cpp b/Engine/Source/Runtime/Engine/Private/Particles/ParticleVertexFactory.cpp
--- a/Engine/Source/Runtime/Engine/Private/Particles/ParticleVertexFactory.cpp
+++ b/Engine/Source/Runtime/Engine/Private/Particles/ParticleVertexFactory.cpp
@@ -194,7 +194,7 @@ public:
 		// Create the vertex declaration for rendering the factory normally.
 		// This is done in InitDynamicRHI instead of InitRHI to allow FParticleSpriteVertexFactory::InitRHI
 		// to rely on it being initialized, since InitDynamicRHI is called before InitRHI.
-		VertexDeclarationRHI = RHICreateVertexDeclaration(Elements);
+		VertexDeclarationRHI = PipelineStateCache::GetOrCreateVertexDeclaration(Elements);
 	}
 
 	virtual void ReleaseDynamicRHI()
diff --git a/Engine/Source/Runtime/Engine/Private/VectorFieldVisualization.cpp b/Engine/Source/Runtime/Engine/Private/VectorFieldVisualization.cpp
--- a/Engine/Source/Runtime/Engine/Private/VectorFieldVisualization.cpp
+++ b/Engine/Source/Runtime/Engine/Private/VectorFieldVisualization.cpp
@@ -72,7 +72,7 @@ public:
 	{
 		FVertexDeclarationElementList Elements;
 		Elements.Add(FVertexElement(0, 0, VET_Float4, 0, sizeof(FVector4)));
-		VertexDeclarationRHI = RHICreateVertexDeclaration(Elements);
+		VertexDeclarationRHI = PipelineStateCache::GetOrCreateVertexDeclaration(Elements);
 	}
 
 	virtual void ReleaseRHI() override
diff --git a/Engine/Source/Runtime/Engine/Public/BatchedElements.h b/Engine/Source/Runtime/Engine/Public/BatchedElements.h
--- a/Engine/Source/Runtime/Engine/Public/BatchedElements.h
+++ b/Engine/Source/Runtime/Engine/Public/BatchedElements.h
@@ -14,6 +14,7 @@
 #include "HitProxies.h"
 #include "SceneView.h"
 #include "StaticBoundShaderState.h"
+#include "PipelineStateCache.h"
 
 struct FBatchedPoint;
 struct FMeshPassProcessorRenderState;
@@ -67,7 +68,7 @@ public:
 		Elements.Add(FVertexElement(0,STRUCT_OFFSET(FSimpleElementVertex,TextureCoordinate),VET_Float2,1,Stride));
 		Elements.Add(FVertexElement(0,STRUCT_OFFSET(FSimpleElementVertex,Color),VET_Float4,2,Stride));
 		Elements.Add(FVertexElement(0,STRUCT_OFFSET(FSimpleElementVertex,HitProxyIdColor),VET_Color,3,Stride));
-		VertexDeclarationRHI = RHICreateVertexDeclaration(Elements);
+		VertexDeclarationRHI = PipelineStateCache::GetOrCreateVertexDeclaration(Elements);
 	}
 
 	virtual void ReleaseRHI() override
diff --git a/Engine/Source/Runtime/Landscape/Private/LandscapeEditProcedural.cpp b/Engine/Source/Runtime/Landscape/Private/LandscapeEditProcedural.cpp
--- a/Engine/Source/Runtime/Landscape/Private/LandscapeEditProcedural.cpp
+++ b/Engine/Source/Runtime/Landscape/Private/LandscapeEditProcedural.cpp
@@ -135,7 +135,7 @@ public:
 		uint32 Stride = sizeof(FLandscapeProceduralVertex);
 		Elements.Add(FVertexElement(0, STRUCT_OFFSET(FLandscapeProceduralVertex, Position), VET_Float2, 0, Stride));
 		Elements.Add(FVertexElement(0, STRUCT_OFFSET(FLandscapeProceduralVertex, UV), VET_Float2, 1, Stride));
-		VertexDeclarationRHI = RHICreateVertexDeclaration(Elements);
+		VertexDeclarationRHI = PipelineStateCache::GetOrCreateVertexDeclaration(Elements);
 	}
 
 	virtual void ReleaseRHI()
diff --git a/Engine/Source/Runtime/RenderCore/Private/RenderUtils.cpp b/Engine/Source/Runtime/RenderCore/Private/RenderUtils.cpp
--- a/Engine/Source/Runtime/RenderCore/Private/RenderUtils.cpp
+++ b/Engine/Source/Runtime/RenderCore/Private/RenderUtils.cpp
@@ -6,6 +6,7 @@
 #include "RenderResource.h"
 #include "RHIStaticStates.h"
 #include "RenderGraphUtils.h"
+#include "PipelineStateCache.h"
 
 #if WITH_EDITOR
 #include "Misc/CoreMisc.h"
@@ -821,7 +822,7 @@ public:
 	{
 		FVertexDeclarationElementList Elements;
 		Elements.Add(FVertexElement(0, 0, VET_Float4, 0, sizeof(FVector4)));
-		VertexDeclarationRHI = RHICreateVertexDeclaration(Elements);
+		VertexDeclarationRHI = PipelineStateCache::GetOrCreateVertexDeclaration(Elements);
 	}
 	virtual void ReleaseRHI() override
 	{
@@ -844,7 +845,7 @@ public:
 	{
 		FVertexDeclarationElementList Elements;
 		Elements.Add(FVertexElement(0, 0, VET_Float3, 0, sizeof(FVector)));
-		VertexDeclarationRHI = RHICreateVertexDeclaration(Elements);
+		VertexDeclarationRHI = PipelineStateCache::GetOrCreateVertexDeclaration(Elements);
 	}
 	virtual void ReleaseRHI() override
 	{
@@ -867,7 +868,7 @@ public:
 	{
 		FVertexDeclarationElementList Elements;
 		Elements.Add(FVertexElement(0, 0, VET_Float2, 0, sizeof(FVector2D)));
-		VertexDeclarationRHI = RHICreateVertexDeclaration(Elements);
+		VertexDeclarationRHI = PipelineStateCache::GetOrCreateVertexDeclaration(Elements);
 	}
 	virtual void ReleaseRHI() override
 	{
diff --git a/Engine/Source/Runtime/RenderCore/Private/ShaderPipelineCache.cpp b/Engine/Source/Runtime/RenderCore/Private/ShaderPipelineCache.cpp
--- a/Engine/Source/Runtime/RenderCore/Private/ShaderPipelineCache.cpp
+++ b/Engine/Source/Runtime/RenderCore/Private/ShaderPipelineCache.cpp
@@ -540,7 +540,7 @@ bool FShaderPipelineCache::Precompile(FRHICommandListImmediate& RHICmdList, ESha
 	{
 		FGraphicsPipelineStateInitializer GraphicsInitializer;
 		
-		auto VertexDesc = RHICmdList.CreateVertexDeclaration(PSO.GraphicsDesc.VertexDescriptor);
+		FRHIVertexDeclaration* VertexDesc = PipelineStateCache::GetOrCreateVertexDeclaration(PSO.GraphicsDesc.VertexDescriptor);
 		GraphicsInitializer.BoundShaderState.VertexDeclarationRHI = VertexDesc;
 		
 		FVertexShaderRHIRef VertexShader;
diff --git a/Engine/Source/Runtime/RenderCore/Private/VertexFactory.cpp b/Engine/Source/Runtime/RenderCore/Private/VertexFactory.cpp
--- a/Engine/Source/Runtime/RenderCore/Private/VertexFactory.cpp
+++ b/Engine/Source/Runtime/RenderCore/Private/VertexFactory.cpp
@@ -7,6 +7,7 @@
 #include "VertexFactory.h"
 #include "Serialization/MemoryWriter.h"
 #include "UObject/DebugSerializationFlags.h"
+#include "PipelineStateCache.h"
 
 uint32 FVertexFactoryType::NextHashIndex = 0;
 bool FVertexFactoryType::bInitializedSerializationHistory = false;
@@ -273,12 +274,12 @@ FVertexElement FVertexFactory::AccessPositionStreamComponent(const FVertexStream
 void FVertexFactory::InitDeclaration(FVertexDeclarationElementList& Elements)
 {
 	// Create the vertex declaration for rendering the factory normally.
-	Declaration = RHICreateVertexDeclaration(Elements);
+	Declaration = PipelineStateCache::GetOrCreateVertexDeclaration(Elements);
 }
 
 void FVertexFactory::InitPositionDeclaration(const FVertexDeclarationElementList& Elements)
 {
-	PositionDeclaration = RHICreateVertexDeclaration(Elements);
+	PositionDeclaration = PipelineStateCache::GetOrCreateVertexDeclaration(Elements);
 }
 
 FVertexFactoryParameterRef::FVertexFactoryParameterRef(FVertexFactoryType* InVertexFactoryType,const FShaderParameterMap& ParameterMap, EShaderFrequency InShaderFrequency, EShaderPlatform InShaderPlatform)
diff --git a/Engine/Source/Runtime/RenderCore/Public/CommonRenderResources.h b/Engine/Source/Runtime/RenderCore/Public/CommonRenderResources.h
--- a/Engine/Source/Runtime/RenderCore/Public/CommonRenderResources.h
+++ b/Engine/Source/Runtime/RenderCore/Public/CommonRenderResources.h
@@ -9,6 +9,7 @@
 #include "RenderResource.h"
 #include "GlobalShader.h"
 #include "ShaderParameterStruct.h"
+#include "PipelineStateCache.h"
 
 
 /** The vertex data used to filter a texture. */
@@ -33,7 +34,7 @@ public:
 		uint32 Stride = sizeof(FFilterVertex);
 		Elements.Add(FVertexElement(0, STRUCT_OFFSET(FFilterVertex, Position), VET_Float4, 0, Stride));
 		Elements.Add(FVertexElement(0, STRUCT_OFFSET(FFilterVertex, UV), VET_Float2, 1, Stride));
-		VertexDeclarationRHI = RHICreateVertexDeclaration(Elements);
+		VertexDeclarationRHI = PipelineStateCache::GetOrCreateVertexDeclaration(Elements);
 	}
 
 	virtual void ReleaseRHI()
@@ -56,7 +57,7 @@ public:
 	virtual void InitRHI()
 	{
 		FVertexDeclarationElementList Elements;
-		VertexDeclarationRHI = RHICreateVertexDeclaration(Elements);
+		VertexDeclarationRHI = PipelineStateCache::GetOrCreateVertexDeclaration(Elements);
 	}
 
 	virtual void ReleaseRHI()
diff --git a/Engine/Source/Runtime/Renderer/Private/AtmosphereRendering.cpp b/Engine/Source/Runtime/Renderer/Private/AtmosphereRendering.cpp
--- a/Engine/Source/Runtime/Renderer/Private/AtmosphereRendering.cpp
+++ b/Engine/Source/Runtime/Renderer/Private/AtmosphereRendering.cpp
@@ -268,7 +268,7 @@ public:
 	{
 		FVertexDeclarationElementList Elements;
 		Elements.Add(FVertexElement(0, 0, VET_Float2, 0, sizeof(FVector2D)));
-		VertexDeclarationRHI = RHICreateVertexDeclaration(Elements);
+		VertexDeclarationRHI = PipelineStateCache::GetOrCreateVertexDeclaration(Elements);
 	}
 
 	virtual void ReleaseRHI() override
diff --git a/Engine/Source/Runtime/Renderer/Private/CapsuleShadowRendering.cpp b/Engine/Source/Runtime/Renderer/Private/CapsuleShadowRendering.cpp
--- a/Engine/Source/Runtime/Renderer/Private/CapsuleShadowRendering.cpp
+++ b/Engine/Source/Runtime/Renderer/Private/CapsuleShadowRendering.cpp
@@ -763,7 +763,7 @@ public:
 		FVertexDeclarationElementList Elements;
 		uint32 Stride = sizeof(FVector2D);
 		Elements.Add(FVertexElement(0,0,VET_Float2,0,Stride,false));
-		VertexDeclarationRHI = RHICreateVertexDeclaration(Elements);
+		VertexDeclarationRHI = PipelineStateCache::GetOrCreateVertexDeclaration(Elements);
 	}
 
 	virtual void ReleaseRHI()
diff --git a/Engine/Source/Runtime/Renderer/Private/FogRendering.cpp b/Engine/Source/Runtime/Renderer/Private/FogRendering.cpp
--- a/Engine/Source/Runtime/Renderer/Private/FogRendering.cpp
+++ b/Engine/Source/Runtime/Renderer/Private/FogRendering.cpp
@@ -247,7 +247,7 @@ public:
 	{
 		FVertexDeclarationElementList Elements;
 		Elements.Add(FVertexElement(0, 0, VET_Float2, 0, sizeof(FVector2D)));
-		VertexDeclarationRHI = RHICreateVertexDeclaration(Elements);
+		VertexDeclarationRHI = PipelineStateCache::GetOrCreateVertexDeclaration(Elements);
 	}
 
 	virtual void ReleaseRHI() override
diff --git a/Engine/Source/Runtime/Renderer/Private/GPUBenchmark.cpp b/Engine/Source/Runtime/Renderer/Private/GPUBenchmark.cpp
--- a/Engine/Source/Runtime/Renderer/Private/GPUBenchmark.cpp
+++ b/Engine/Source/Runtime/Renderer/Private/GPUBenchmark.cpp
@@ -179,7 +179,7 @@ struct FVertexThroughputDeclaration : public FRenderResource
 			{ 0, 4 * sizeof(FVector4), VET_Float4, 4, sizeof(FBenchmarkVertex) },
 		};
 
-		DeclRHI = RHICreateVertexDeclaration(Elements);
+		DeclRHI = PipelineStateCache::GetOrCreateVertexDeclaration(Elements);
 	}
 
 	virtual void ReleaseRHI() override
diff --git a/Engine/Source/Runtime/Renderer/Private/PostProcess/PostProcessHMD.cpp b/Engine/Source/Runtime/Renderer/Private/PostProcess/PostProcessHMD.cpp
--- a/Engine/Source/Runtime/Renderer/Private/PostProcess/PostProcessHMD.cpp
+++ b/Engine/Source/Runtime/Renderer/Private/PostProcess/PostProcessHMD.cpp
@@ -35,7 +35,7 @@ public:
 		Elements.Add(FVertexElement(0, STRUCT_OFFSET(FDistortionVertex, TexB), VET_Float2, 3, Stride));
 		Elements.Add(FVertexElement(0, STRUCT_OFFSET(FDistortionVertex, VignetteFactor), VET_Float1, 4, Stride));
 		Elements.Add(FVertexElement(0, STRUCT_OFFSET(FDistortionVertex, TimewarpFactor), VET_Float1, 5, Stride));
-		VertexDeclarationRHI = RHICreateVertexDeclaration(Elements);
+		VertexDeclarationRHI = PipelineStateCache::GetOrCreateVertexDeclaration(Elements);
 	}
 
 	virtual void ReleaseRHI() override
diff --git a/Engine/Source/Runtime/Renderer/Private/PostProcess/PostProcessMaterial.cpp b/Engine/Source/Runtime/Renderer/Private/PostProcess/PostProcessMaterial.cpp
--- a/Engine/Source/Runtime/Renderer/Private/PostProcess/PostProcessMaterial.cpp
+++ b/Engine/Source/Runtime/Renderer/Private/PostProcess/PostProcessMaterial.cpp
@@ -192,7 +192,7 @@ public:
 		FVertexDeclarationElementList Elements;
 		uint32 Stride = sizeof(FFilterVertex);
 		Elements.Add(FVertexElement(0,STRUCT_OFFSET(FFilterVertex,Position),VET_Float4,0,Stride));
-		VertexDeclarationRHI = RHICreateVertexDeclaration(Elements);
+		VertexDeclarationRHI = PipelineStateCache::GetOrCreateVertexDeclaration(Elements);
 	}
 	
 	virtual void ReleaseRHI()
diff --git a/Engine/Source/Runtime/SlateRHIRenderer/Private/SlateShaders.cpp b/Engine/Source/Runtime/SlateRHIRenderer/Private/SlateShaders.cpp
--- a/Engine/Source/Runtime/SlateRHIRenderer/Private/SlateShaders.cpp
+++ b/Engine/Source/Runtime/SlateRHIRenderer/Private/SlateShaders.cpp
@@ -2,6 +2,7 @@
 
 #include "SlateShaders.h"
 #include "Rendering/RenderingCommon.h"
+#include "PipelineStateCache.h"
 
 /** Flag to determine if we are running with a color vision deficiency shader on */
 EColorVisionDeficiency GSlateColorDeficiencyType = EColorVisionDeficiency::NormalVision;
@@ -62,7 +63,7 @@ void FSlateVertexDeclaration::InitRHI()
 	Elements.Add(FVertexElement(0, STRUCT_OFFSET(FSlateVertex, Color), VET_Color, 3, Stride));
 	Elements.Add(FVertexElement(0, STRUCT_OFFSET(FSlateVertex, PixelSize), VET_UShort2, 4, Stride));
 
-	VertexDeclarationRHI = RHICreateVertexDeclaration(Elements);
+	VertexDeclarationRHI = PipelineStateCache::GetOrCreateVertexDeclaration(Elements);
 }
 
 void FSlateVertexDeclaration::ReleaseRHI()
@@ -84,7 +85,7 @@ void FSlateInstancedVertexDeclaration::InitRHI()
 	Elements.Add(FVertexElement(0, STRUCT_OFFSET(FSlateVertex, Color), VET_Color, 3, Stride));
 	Elements.Add(FVertexElement(1, 0, VET_Float4, 4, sizeof(FVector4), true));
 	
-	VertexDeclarationRHI = RHICreateVertexDeclaration(Elements);
+	VertexDeclarationRHI = PipelineStateCache::GetOrCreateVertexDeclaration(Elements);
 }
 
 void FSlateElementPS::ModifyCompilationEnvironment(const FGlobalShaderPermutationParameters& Parameters, FShaderCompilerEnvironment& OutEnvironment)
@@ -103,7 +104,7 @@ void FSlateMaskingVertexDeclaration::InitRHI()
 	uint32 Stride = sizeof(uint32);
 	Elements.Add(FVertexElement(0, 0, VET_UByte4, 0, Stride));
 
-	VertexDeclarationRHI = RHICreateVertexDeclaration(Elements);
+	VertexDeclarationRHI = PipelineStateCache::GetOrCreateVertexDeclaration(Elements);
 }
 
 void FSlateMaskingVertexDeclaration::ReleaseRHI()
diff --git a/Engine/Source/Runtime/UtilityShaders/Public/MediaShaders.h b/Engine/Source/Runtime/UtilityShaders/Public/MediaShaders.h
--- a/Engine/Source/Runtime/UtilityShaders/Public/MediaShaders.h
+++ b/Engine/Source/Runtime/UtilityShaders/Public/MediaShaders.h
@@ -7,6 +7,7 @@
 #include "RenderResource.h"
 #include "Shader.h"
 #include "GlobalShader.h"
+#include "PipelineStateCache.h"
 
 
 namespace MediaShaders
@@ -98,7 +99,7 @@ public:
 		uint16 Stride = sizeof(FMediaElementVertex);
 		Elements.Add(FVertexElement(0, STRUCT_OFFSET(FMediaElementVertex, Position), VET_Float4, 0, Stride));
 		Elements.Add(FVertexElement(0, STRUCT_OFFSET(FMediaElementVertex, TextureCoordinate), VET_Float2, 1, Stride));
-		VertexDeclarationRHI = RHICreateVertexDeclaration(Elements);
+		VertexDeclarationRHI = PipelineStateCache::GetOrCreateVertexDeclaration(Elements);
 	}
 
 	virtual void ReleaseRHI() override
