Fix ensure/crash in replication graph due to server not cleaning up world settings during seamless travel

diff --git a/Engine/Source/Runtime/Engine/Private/NetworkDriver.cpp b/Engine/Source/Runtime/Engine/Private/NetworkDriver.cpp
--- a/Engine/Source/Runtime/Engine/Private/NetworkDriver.cpp
+++ b/Engine/Source/Runtime/Engine/Private/NetworkDriver.cpp
@@ -5021,10 +5021,7 @@ void UNetDriver::CleanupWorldForSeamlessTravel()
 				// skips over AWorldSettings actors for an unknown reason.
 				if (Level->GetWorldSettings())
 				{
-					if (ServerConnection != nullptr)
-					{
-						NotifyActorLevelUnloaded(Level->GetWorldSettings());
-					}
+					NotifyActorLevelUnloaded(Level->GetWorldSettings());
 				}
 			}
 		}
