Make sure that if an Object's RepFlags change, or if a properties become Active, we send those to clients.

PR #5353: UE-43135 (Contributed by [REDACTED_PERSON])

diff --git a/Engine/Source/Runtime/Engine/Private/DataReplication.cpp b/Engine/Source/Runtime/Engine/Private/DataReplication.cpp
--- a/Engine/Source/Runtime/Engine/Private/DataReplication.cpp
+++ b/Engine/Source/Runtime/Engine/Private/DataReplication.cpp
@@ -1126,7 +1126,7 @@ void FObjectReplicator::ReplicateCustomDeltaProperties( FNetBitWriter & Bunch, F
 	UNetConnection * OwningChannelConnection = OwningChannel->Connection;
 
 	// Initialize a map of which conditions are valid
-	const TStaticBitArray<COND_Max> ConditionMap = BuildConditionMapFromRepFlags(RepFlags);
+	const TStaticBitArray<COND_Max> ConditionMap = FSendingRepState::BuildConditionMap(RepFlags);
 
 	// Make sure net field export group is registered
 	FNetFieldExportGroup* NetFieldExportGroup = OwningChannel->GetOrCreateNetFieldExportGroupForClassNetCache( Object );
diff --git a/Engine/Source/Runtime/Engine/Private/RepLayout.cpp b/Engine/Source/Runtime/Engine/Private/RepLayout.cpp
--- a/Engine/Source/Runtime/Engine/Private/RepLayout.cpp
+++ b/Engine/Source/Runtime/Engine/Private/RepLayout.cpp
@@ -993,8 +993,8 @@ bool FRepLayout::ReplicateProperties(
 			if (Pruned.Num() > 0)
 			{
 				SendProperties_BackwardsCompatible(RepState, ChangeTracker, Data, OwningChannel->Connection, Writer, Pruned);
-			return true;
-		}
+				return true;
+			}
 		}
 
 		return false;
@@ -1052,7 +1052,7 @@ bool FRepLayout::ReplicateProperties(
 
 PRAGMA_DISABLE_DEPRECATION_WARNINGS
 		MergeChangeList(Data, HistoryItem.Changed, Temp, Changed);
-PRAGMA_ENABLE_DEPRECATION_WARNINGS
+PRAGMA_ENABLE_DEPRECATION_WARNINGS		
 	}
 
 	// Merge in newly active properties so they can be sent.
@@ -1110,7 +1110,7 @@ PRAGMA_ENABLE_DEPRECATION_WARNINGS
 
 	const int32 NumBits = Writer.GetNumBits();
 
-	// Filter out the final changelist into Active and Inaction.
+	// Filter out the final changelist into Active and Inactive.
 	TArray<uint16> UnfilteredChanged = MoveTemp(Changed);
 	TArray<uint16> NewlyInactiveChangelist;
 	FilterChangeList(UnfilteredChanged, RepState->InactiveParents, NewlyInactiveChangelist, Changed);
@@ -5041,13 +5041,48 @@ void FRepLayout::BuildHandleToCmdIndexTable_r(
 	}
 }
 
+TStaticBitArray<COND_Max> FSendingRepState::BuildConditionMap(const FReplicationFlags& RepFlags)
+{
+	TStaticBitArray<COND_Max> ConditionMap;
+
+	// Setup condition map
+	const bool bIsInitial = RepFlags.bNetInitial ? true : false;
+	const bool bIsOwner = RepFlags.bNetOwner ? true : false;
+	const bool bIsSimulated = RepFlags.bNetSimulated ? true : false;
+	const bool bIsPhysics = RepFlags.bRepPhysics ? true : false;
+	const bool bIsReplay = RepFlags.bReplay ? true : false;
+
+	ConditionMap[COND_None] = true;
+	ConditionMap[COND_InitialOnly] = bIsInitial;
+
+	ConditionMap[COND_OwnerOnly] = bIsOwner;
+	ConditionMap[COND_SkipOwner] = !bIsOwner;
+
+	ConditionMap[COND_SimulatedOnly] = bIsSimulated;
+	ConditionMap[COND_SimulatedOnlyNoReplay] = bIsSimulated && !bIsReplay;
+	ConditionMap[COND_AutonomousOnly] = !bIsSimulated;
+
+	ConditionMap[COND_SimulatedOrPhysics] = bIsSimulated || bIsPhysics;
+	ConditionMap[COND_SimulatedOrPhysicsNoReplay] = (bIsSimulated || bIsPhysics) && !bIsReplay;
+
+	ConditionMap[COND_InitialOrOwner] = bIsInitial || bIsOwner;
+	ConditionMap[COND_ReplayOrOwner] = bIsReplay || bIsOwner;
+	ConditionMap[COND_ReplayOnly] = bIsReplay;
+	ConditionMap[COND_SkipReplay] = !bIsReplay;
+
+	ConditionMap[COND_Custom] = true;
+	ConditionMap[COND_Never] = false;
+
+	return ConditionMap;
+}
+
 void FRepLayout::RebuildConditionalProperties(
 	FSendingRepState* RESTRICT RepState,
 	const FReplicationFlags& RepFlags) const
 {
 	SCOPE_CYCLE_COUNTER(STAT_NetRebuildConditionalTime);
 	
-	TStaticBitArray<COND_Max> ConditionMap = BuildConditionMapFromRepFlags(RepFlags);
+	TStaticBitArray<COND_Max> ConditionMap = FSendingRepState::BuildConditionMap(RepFlags);
 	for (auto It = TBitArray<>::FIterator(RepState->InactiveParents); It; ++It)
 	{
 		It.GetValue() = !ConditionMap[Parents[It.GetIndex()].Condition];
@@ -5109,8 +5144,8 @@ TUniquePtr<FRepState> FRepLayout::CreateRepState(
 
 		// Start out the conditional props based on a default RepFlags struct
 		// It will rebuild if it ever changes
-		RepState->SendingRepState->InactiveParents.Init(false, Parents.Num());
 		RebuildConditionalProperties(RepState->SendingRepState.Get(), FReplicationFlags());
+		RepState->SendingRepState->InactiveParents.Init(false, Parents.Num());
 	}
 	
 	if (!EnumHasAnyFlags(Flags, ECreateRepStateFlags::SkipCreateReceivingState))
@@ -5138,7 +5173,6 @@ void FRepLayout::InitRepStateStaticBuffer(FRepStateStaticBuffer& ShadowData, con
 	CopyProperties(ShadowData, Source);
 }
 
-
 void FRepLayout::ConstructProperties(FRepStateStaticBuffer& InShadowData) const
 {
 	uint8* ShadowData = InShadowData.GetData();
diff --git a/Engine/Source/Runtime/Engine/Public/Net/RepLayout.h b/Engine/Source/Runtime/Engine/Public/Net/RepLayout.h
--- a/Engine/Source/Runtime/Engine/Public/Net/RepLayout.h
+++ b/Engine/Source/Runtime/Engine/Public/Net/RepLayout.h
@@ -536,11 +536,20 @@ private:
 		LastChangelistIndex(0),
 		LastCompareIndex(0),
 		InactiveChangelist({0})
-	{}
+	{}	
 
 public:
 
 	void CountBytes(FArchive& Ar) const;
+	
+	/**
+	 * Builds a new ConditionMap given the input RepFlags.
+	 * This can be used to determine whether or not a given property should be
+	 * considered enabled / disabled based on ELifetimeCondition.
+	 *
+	 * TODO: This doesn't have to be part of FRepState.
+	 */
+	static TStaticBitArray<COND_Max> BuildConditionMap(const FReplicationFlags& InFlags);
 
 	/** Whether or not FRepLayout::OpenAcked has been called with this FRepState. */
 	bool bOpenAckedCalled;
@@ -593,6 +602,7 @@ public:
 	/** The unique list of properties that have changed since the channel was first opened */
 	TArray<uint16> LifetimeChangelist;
 
+
 	/**
 	 * Properties which are inactive through conditions have their changes stored here, so they can be 
 	 * applied if/when the property becomes active.
@@ -605,7 +615,6 @@ public:
 	TBitArray<> InactiveParents;
 };
 
-
 /** Replication State that is unique Per Object Per Net Connection. */
 class FRepState : public FNoncopyable
 {
@@ -1927,4 +1936,4 @@ private:
 
 	/** Shared comparison to default state for multicast rpc */
 	TBitArray<> SharedInfoRPCParentsChanged;	
-};
+};
\ No newline at end of file
