Merging CL#5013294

Support for AnimClass Property Track by exposing SetAnimClass,GetAnimClass.
Added redirects for k2_SetAnimInstanceClass (and previous SetAnimInstanceClass).

Copy from Dev-Sequencer

diff --git a/Engine/Config/BaseEngine.ini b/Engine/Config/BaseEngine.ini
--- a/Engine/Config/BaseEngine.ini
+++ b/Engine/Config/BaseEngine.ini
@@ -594,7 +594,6 @@ bUseStreamingPause=false
 +StructRedirects=(OldName="Timecode",NewName="/Script/CoreUObject.Timecode")
 
 ; 4.21
-+FunctionRedirects=(OldName="SkeletalMeshComponent.SetAnimInstanceClass",NewName="SkeletalMeshComponent.K2_SetAnimInstanceClass")
 +EnumRedirects=(OldName="EMeshComponentUpdateFlag",NewName="/Script/Engine.EVisibilityBasedAnimTickOption")
 +PropertyRedirects=(OldName="SkinnedMeshComponent.MeshComponentUpdateFlag", NewName="VisibilityBasedAnimTickOption")
 
@@ -649,6 +648,10 @@ bUseStreamingPause=false
 
 +FunctionRedirects=(OldName="PlayerController.ClientPlayForceFeedback",NewName="PlayerController.K2_ClientPlayForceFeedback")
 
+; 4.23
++FunctionRedirects=(OldName="SkeletalMeshComponent.SetAnimInstanceClass",NewName="SkeletalMeshComponent.SetAnimClass")
++FunctionRedirects=(OldName="SkeletalMeshComponent.K2_SetAnimInstanceClass",NewName="SkeletalMeshComponent.SetAnimClass")
+
 [CoreUObject.Metadata]
 ; Note: UnrealHeaderTool should be rerun after making changes to MetadataRedirects to catch any keys specified in class headers
 +MetadataRedirects=(OldKey="K2Protected", NewKey="BlueprintProtected")
diff --git a/Engine/Source/Editor/UnrealEd/Classes/Animation/DebugSkelMeshComponent.h b/Engine/Source/Editor/UnrealEd/Classes/Animation/DebugSkelMeshComponent.h
--- a/Engine/Source/Editor/UnrealEd/Classes/Animation/DebugSkelMeshComponent.h
+++ b/Engine/Source/Editor/UnrealEd/Classes/Animation/DebugSkelMeshComponent.h
@@ -302,7 +302,7 @@ class UNREALED_API UDebugSkelMeshComponent : public USkeletalMeshComponent
 	//~ Begin SkeletalMeshComponent Interface
 	virtual void InitAnim(bool bForceReinit) override;
 	virtual bool IsWindEnabled() const override { return true; }
-	virtual void K2_SetAnimInstanceClass(class UClass* NewClass) override;
+	virtual void SetAnimClass(class UClass* NewClass) override;
 	//~ End SkeletalMeshComponent Interface
 	// Preview.
 	// @todo document
diff --git a/Engine/Source/Editor/UnrealEd/Private/Animation/DebugSkelMeshComponent.cpp b/Engine/Source/Editor/UnrealEd/Private/Animation/DebugSkelMeshComponent.cpp
--- a/Engine/Source/Editor/UnrealEd/Private/Animation/DebugSkelMeshComponent.cpp
+++ b/Engine/Source/Editor/UnrealEd/Private/Animation/DebugSkelMeshComponent.cpp
@@ -295,7 +295,7 @@ void UDebugSkelMeshComponent::InitAnim(bool bForceReinit)
 	}
 }
 
-void UDebugSkelMeshComponent::K2_SetAnimInstanceClass(class UClass* NewClass)
+void UDebugSkelMeshComponent::SetAnimClass(class UClass* NewClass)
 {
 	// Override this to do nothing and warn the user
 	UE_LOG(LogAnimation, Warning, TEXT("Attempting to destroy an animation preview actor, skipping."));
diff --git a/Engine/Source/Runtime/Engine/Classes/Components/SkeletalMeshComponent.h b/Engine/Source/Runtime/Engine/Classes/Components/SkeletalMeshComponent.h
--- a/Engine/Source/Runtime/Engine/Classes/Components/SkeletalMeshComponent.h
+++ b/Engine/Source/Runtime/Engine/Classes/Components/SkeletalMeshComponent.h
@@ -701,10 +701,17 @@ public:
 	/** Array of physical interactions for the frame. This is a temporary solution for a more permanent force system and should not be used directly*/
 	TArray<FPendingRadialForces> PendingRadialForces;
 
-	/** Set the anim instance class. Clears and re-initializes the anim instance with the new class and sets animation mode to 'AnimationBlueprint' */
-	UFUNCTION(BlueprintCallable, Category="Components|SkeletalMesh", meta=(Keywords = "AnimBlueprint", DisplayName = "Set Anim Instance Class"))
+	UE_DEPRECATED(4.23, "This function is deprecated. Please use SetAnimClass instead. ")
 	virtual void K2_SetAnimInstanceClass(class UClass* NewClass);
 
+	/** Set the anim instance class. Clears and re-initializes the anim instance with the new class and sets animation mode to 'AnimationBlueprint' */
+	UFUNCTION(BlueprintCallable, Category = "Components|SkeletalMesh", meta = (Keywords = "AnimBlueprint", DisplayName = "Set Anim Instance Class"))
+	virtual void SetAnimClass(class UClass* NewClass);
+
+	/** Get the anim instance class via getter callable by sequencer.  */
+	UFUNCTION(BlueprintInternalUseOnly)
+	class UClass*  GetAnimClass();
+
 	/** Set the anim instance class. Clears and re-initializes the anim instance with the new class and sets animation mode to 'AnimationBlueprint' */
 	void SetAnimInstanceClass(class UClass* NewClass);
 
diff --git a/Engine/Source/Runtime/Engine/Private/Components/SkeletalMeshComponent.cpp b/Engine/Source/Runtime/Engine/Private/Components/SkeletalMeshComponent.cpp
--- a/Engine/Source/Runtime/Engine/Private/Components/SkeletalMeshComponent.cpp
+++ b/Engine/Source/Runtime/Engine/Private/Components/SkeletalMeshComponent.cpp
@@ -2625,6 +2625,16 @@ void USkeletalMeshComponent::K2_SetAnimInstanceClass(class UClass* NewClass)
 	SetAnimInstanceClass(NewClass);
 }
 
+void USkeletalMeshComponent::SetAnimClass(class UClass* NewClass)
+{
+	SetAnimInstanceClass(NewClass);
+}
+
+class UClass* USkeletalMeshComponent::GetAnimClass()
+{
+	return AnimClass;
+}
+
 void USkeletalMeshComponent::SetAnimInstanceClass(class UClass* NewClass)
 {
 	if (NewClass != nullptr)
