* Fix towhook constraint constantly breaking when attaching to other vehicles.
* Add hooks for constraint breakage

diff --git a/Engine/Source/Runtime/Engine/Classes/PhysicsEngine/ConstraintInstance.h b/Engine/Source/Runtime/Engine/Classes/PhysicsEngine/ConstraintInstance.h
--- a/Engine/Source/Runtime/Engine/Classes/PhysicsEngine/ConstraintInstance.h
+++ b/Engine/Source/Runtime/Engine/Classes/PhysicsEngine/ConstraintInstance.h
@@ -485,6 +485,9 @@ public:
 	/** Create physics engine constraint using physx actors. */
 	void InitConstraint_AssumesLocked(const FPhysicsActorHandle& ActorRef1, const FPhysicsActorHandle& ActorRef2, float InScale, FOnConstraintBroken InConstraintBrokenDelegate = FOnConstraintBroken());
 
+	/** Set teh constraint broken delegate. */
+	void SetConstraintBrokenDelegate(FOnConstraintBroken InConstraintBrokenDelegate);
+
 	/** Terminate physics engine constraint */
 	void TermConstraint();
 
diff --git a/Engine/Source/Runtime/Engine/Private/PhysicsEngine/ConstraintInstance.cpp b/Engine/Source/Runtime/Engine/Private/PhysicsEngine/ConstraintInstance.cpp
--- a/Engine/Source/Runtime/Engine/Private/PhysicsEngine/ConstraintInstance.cpp
+++ b/Engine/Source/Runtime/Engine/Private/PhysicsEngine/ConstraintInstance.cpp
@@ -464,6 +464,11 @@ void FConstraintInstance::InitConstraint_AssumesLocked(const FPhysicsActorHandle
 	EnsureSleepingActorsStaySleeping_AssumesLocked(ActorRef1, ActorRef2);
 }
 
+void FConstraintInstance::SetConstraintBrokenDelegate(FOnConstraintBroken InConstraintBrokenDelegate)
+{
+	OnConstraintBrokenDelegate = InConstraintBrokenDelegate;
+}
+
 void FConstraintProfileProperties::Update_AssumesLocked(const FPhysicsConstraintHandle& InConstraintRef, float AverageMass, float UseScale) const
 {
 	// flags and projection settings
