Update RepLayout to not pass raw pointers around, and instead use a variants of TRepDataBufferBase everywhere.

review-5002983 [REDACTED_HANDLE] [REDACTED_HANDLE] [REDACTED_HANDLE] [REDACTED_HANDLE] [REDACTED_HANDLE] [REDACTED_HANDLE]

diff --git a/Engine/Source/Runtime/Engine/Private/RepLayout.cpp b/Engine/Source/Runtime/Engine/Private/RepLayout.cpp
--- a/Engine/Source/Runtime/Engine/Private/RepLayout.cpp
+++ b/Engine/Source/Runtime/Engine/Private/RepLayout.cpp
@@ -235,11 +235,12 @@ static FORCEINLINE void SerializeGenericChecksum(FBitArchive& Ar)
 	check(Checksum == 0xABADF00D);
 }
 
+template<ERepDataBufferType DataType>
 static void SerializeReadWritePropertyChecksum(
-	const FRepLayoutCmd&	Cmd,
-	const int32				CurCmdIndex,
-	const uint8*			Data,
-	FBitArchive&			Ar)
+	const FRepLayoutCmd& Cmd,
+	const int32 CurCmdIndex,
+	const TConstRepDataBuffer<DataType> Data,
+	FBitArchive& Ar)
 {
 	// Serialize various attributes that will mostly ensure we are working on the same property
 	const uint32 NameHash = GetTypeHash(Cmd.Property->GetName());
@@ -274,7 +275,7 @@ static void SerializeReadWritePropertyChecksum(
 	//	it also needs to write the same blob it just read as well.
 	FBitWriter Writer(0, true);
 
-	Cmd.Property->NetSerializeItem(Writer, NULL, const_cast<uint8*>(Data));
+	Cmd.Property->NetSerializeItem(Writer, NULL, const_cast<uint8*>(Data.Data));
 
 	if (Ar.IsSaving())
 	{
@@ -637,7 +638,8 @@ void FRepStateStaticBuffer::CountBytes(FArchive& Ar) const
 
 				for (int32 i = 0; i < Array->Num(); ++i)
 				{
-					CountBytes_Command(Parent, InCmdIndex + 1, Cmd.EndCmd, ShadowArrayData + (Cmd.ElementSize * i));
+					const int32 ArrayElementOffset = Cmd.ElementSize * i;
+					CountBytes_Command(Parent, InCmdIndex + 1, Cmd.EndCmd, ShadowArrayData + ArrayElementOffset);
 				}
 			}
 			else if (ERepLayoutCmdType::PropertyString == Cmd.Type)
@@ -722,8 +724,8 @@ uint16 FRepLayout::CompareProperties_r(
 	FSendingRepState* RESTRICT RepState,
 	const int32 CmdStart,
 	const int32 CmdEnd,
-	const uint8* RESTRICT ShadowData,
-	const uint8* RESTRICT Data,
+	FRepShadowDataBuffer ShadowData,
+	const FConstRepObjectDataBuffer Data,
 	TArray<uint16>& Changed,
 	uint16 Handle,
 	const bool bIsInitial,
@@ -760,7 +762,7 @@ uint16 FRepLayout::CompareProperties_r(
 			}
 
 			// Once we hit an array, start using a stack based approach
-			CompareProperties_Array_r(RepState, ShadowData + Cmd.ShadowOffset, (const uint8*)Data + Cmd.Offset, Changed, CmdIndex, Handle, bIsInitial, bForceFail);
+			CompareProperties_Array_r(RepState, ShadowData + Cmd, Data + Cmd, Changed, CmdIndex, Handle, bIsInitial, bForceFail);
 			CmdIndex = Cmd.EndCmd - 1;		// The -1 to handle the ++ in the for loop
 			continue;
 		}
@@ -774,7 +776,7 @@ uint16 FRepLayout::CompareProperties_r(
 		// In that case, just allow this to fail and perform the old logic.
 		if (RepState && Cmd.ParentIndex == RoleIndex)
 		{
-			const ENetRole ObjectRole = *(const ENetRole*)(Data + Cmd.Offset);
+			const ENetRole ObjectRole = *(const ENetRole*)(Data + Cmd).Data;
 			if (bForceFail || RepState->SavedRole != ObjectRole)
 			{
 				RepState->SavedRole = ObjectRole;
@@ -783,16 +785,16 @@ uint16 FRepLayout::CompareProperties_r(
 		}
 		else if (RepState && Cmd.ParentIndex == RemoteRoleIndex)
 		{
-			const ENetRole ObjectRemoteRole = *(const ENetRole*)(Data + Cmd.Offset);
+			const ENetRole ObjectRemoteRole = *(const ENetRole*)(Data + Cmd).Data;
 			if (bForceFail || RepState->SavedRemoteRole != ObjectRemoteRole)
 			{
 				RepState->SavedRemoteRole = ObjectRemoteRole;
 				Changed.Add(Handle);
 			}
 		}
-		else if (bForceFail || !PropertiesAreIdentical(Cmd, (const void*)(ShadowData + Cmd.ShadowOffset), (const void*)(Data + Cmd.Offset)))
+		else if (bForceFail || !PropertiesAreIdentical(Cmd, (ShadowData + Cmd).Data, (Data + Cmd).Data))
 		{
-			StoreProperty(Cmd, (void*)(ShadowData + Cmd.ShadowOffset), (const void*)(Data + Cmd.Offset));
+			StoreProperty(Cmd, (ShadowData + Cmd).Data, (Data + Cmd.Offset).Data);
 			Changed.Add(Handle);
 		}
 	}
@@ -802,8 +804,8 @@ uint16 FRepLayout::CompareProperties_r(
 
 void FRepLayout::CompareProperties_Array_r(
 	FSendingRepState* RESTRICT RepState,
-	const uint8* RESTRICT ShadowData,
-	const uint8* RESTRICT Data,
+	FRepShadowDataBuffer ShadowData,
+	const FConstRepObjectDataBuffer Data,
 	TArray<uint16>& Changed,
 	const uint16 CmdIndex,
 	const uint16 Handle,
@@ -812,8 +814,8 @@ void FRepLayout::CompareProperties_Array_r(
 {
 	const FRepLayoutCmd& Cmd = Cmds[CmdIndex];
 
-	FScriptArray* ShadowArray = (FScriptArray*)ShadowData;
-	FScriptArray* Array = (FScriptArray*)Data;
+	FScriptArray* ShadowArray = (FScriptArray*)ShadowData.Data;
+	FScriptArray* Array = (FScriptArray*)Data.Data;
 
 	const uint16 ArrayNum = Array->Num();
 	const uint16 ShadowArrayNum = ShadowArray->Num();
@@ -826,16 +828,15 @@ void FRepLayout::CompareProperties_Array_r(
 
 	uint16 LocalHandle = 0;
 
-	Data = (uint8*)Array->GetData();
-	ShadowData = (uint8*)ShadowArray->GetData();
+	const FConstRepObjectDataBuffer ArrayData(Array->GetData());
+	FRepShadowDataBuffer ShadowArrayData(ShadowArray->GetData());
 
 	for (int32 i = 0; i < ArrayNum; i++)
 	{
-		const int32 ElementOffset = i * Cmd.ElementSize;
-
+		const int32 ArrayElementOffset = i * Cmd.ElementSize;
 		const bool bNewForceFail = bForceFail || i >= ShadowArrayNum;
 
-		LocalHandle = CompareProperties_r(RepState, CmdIndex + 1, Cmd.EndCmd - 1, ShadowData + ElementOffset, Data + ElementOffset, ChangedLocal, LocalHandle, bIsInitial, bNewForceFail);
+		LocalHandle = CompareProperties_r(RepState, CmdIndex + 1, Cmd.EndCmd - 1, ShadowArrayData + ArrayElementOffset, ArrayData + ArrayElementOffset, ChangedLocal, LocalHandle, bIsInitial, bNewForceFail);
 	}
 
 	if (ChangedLocal.Num())
@@ -860,7 +861,7 @@ void FRepLayout::CompareProperties_Array_r(
 bool FRepLayout::CompareProperties(
 	FSendingRepState* RESTRICT RepState,
 	FRepChangelistState* RESTRICT RepChangelistState,
-	const uint8* RESTRICT Data,
+	const FConstRepObjectDataBuffer Data,
 	const FReplicationFlags& RepFlags) const
 {
 	SCOPE_CYCLE_COUNTER(STAT_NetReplicateDynamicPropCompareTime);
@@ -912,18 +913,16 @@ bool FRepLayout::CompareProperties(
 		TArray<uint16>& FirstChangelistRef = RepChangelistState->ChangeHistory[FirstHistoryIndex].Changed;
 		TArray<uint16> SecondChangelistCopy = MoveTemp(RepChangelistState->ChangeHistory[SecondHistoryIndex].Changed);
 
-PRAGMA_DISABLE_DEPRECATION_WARNINGS
 		MergeChangeList(Data, FirstChangelistRef, SecondChangelistCopy, RepChangelistState->ChangeHistory[SecondHistoryIndex].Changed);
-PRAGMA_ENABLE_DEPRECATION_WARNINGS
 	}
 
 	return true;
 }
 
 static FORCEINLINE void WritePropertyHandle(
-	FNetBitWriter&	Writer,
-	uint16			Handle,
-	bool			bDoChecksum)
+	FNetBitWriter& Writer,
+	uint16 Handle,
+	bool bDoChecksum)
 {
 	const int NumStartingBits = Writer.GetNumBits();
 
@@ -945,7 +944,7 @@ static FORCEINLINE void WritePropertyHandle(
 bool FRepLayout::ReplicateProperties(
 	FSendingRepState* RESTRICT RepState,
 	FRepChangelistState* RESTRICT RepChangelistState,
-	const uint8* RESTRICT Data,
+	const FConstRepObjectDataBuffer Data,
 	UClass* ObjectClass,
 	UActorChannel* OwningChannel,
 	FNetBitWriter& Writer,
@@ -1049,20 +1048,14 @@ bool FRepLayout::ReplicateProperties(
 		FRepChangedHistory& HistoryItem = RepChangelistState->ChangeHistory[HistoryIndex];
 
 		TArray<uint16> Temp = MoveTemp(Changed);
-
-PRAGMA_DISABLE_DEPRECATION_WARNINGS
 		MergeChangeList(Data, HistoryItem.Changed, Temp, Changed);
-PRAGMA_ENABLE_DEPRECATION_WARNINGS		
 	}
 
 	// Merge in newly active properties so they can be sent.
 	if (NewlyActiveChangelist.Num() > 0)
 	{
 		TArray<uint16> Temp = MoveTemp(Changed);
-
-PRAGMA_DISABLE_DEPRECATION_WARNINGS
 		MergeChangeList(Data, NewlyActiveChangelist, Temp, Changed);
-PRAGMA_ENABLE_DEPRECATION_WARNINGS
 	}
 
 	// We're all caught up now
@@ -1080,10 +1073,7 @@ PRAGMA_ENABLE_DEPRECATION_WARNINGS
 			for (int32 i = 0; i < RepState->PreOpenAckHistory.Num(); i++)
 			{
 				TArray<uint16> Temp = MoveTemp(Changed);
-
-PRAGMA_DISABLE_DEPRECATION_WARNINGS
 				MergeChangeList(Data, RepState->PreOpenAckHistory[i].Changed, Temp, Changed);
-PRAGMA_ENABLE_DEPRECATION_WARNINGS
 			}
 			RepState->PreOpenAckHistory.Empty();
 		}
@@ -1119,10 +1109,8 @@ PRAGMA_ENABLE_DEPRECATION_WARNINGS
 	if (NewlyInactiveChangelist.Num() > 1)
 	{
 		TArray<uint16> Temp = MoveTemp(RepState->InactiveChangelist);
-
-PRAGMA_DISABLE_DEPRECATION_WARNINGS
 		MergeChangeList(Data, NewlyInactiveChangelist, Temp, RepState->InactiveChangelist);
-PRAGMA_ENABLE_DEPRECATION_WARNINGS
+
 	}
 
 	// Send the final merged change list
@@ -1131,10 +1119,7 @@ PRAGMA_ENABLE_DEPRECATION_WARNINGS
 		// Remember all properties that have changed since this channel was first opened in case we need it (for bResendAllDataSinceOpen)
 		// We use UnfilteredChanged so LifetimeChangelist contains all properties, regardless of Active state.
 		TArray<uint16> Temp = MoveTemp(RepState->LifetimeChangelist);
-
-PRAGMA_DISABLE_DEPRECATION_WARNINGS
 		MergeChangeList(Data, UnfilteredChanged, Temp, RepState->LifetimeChangelist);
-PRAGMA_ENABLE_DEPRECATION_WARNINGS
 
 		if (Changed.Num() > 0)
 		{
@@ -1162,7 +1147,7 @@ PRAGMA_ENABLE_DEPRECATION_WARNINGS
 void FRepLayout::UpdateChangelistHistory(
 	FSendingRepState* RepState,
 	UClass* ObjectClass,
-	const uint8* RESTRICT Data,
+	const FConstRepObjectDataBuffer Data,
 	UNetConnection* Connection,
 	TArray<uint16>* OutMerged) const
 {
@@ -1200,10 +1185,7 @@ void FRepLayout::UpdateChangelistHistory(
 				// Merge in nak'd change lists
 				check(OutMerged != NULL);
 				TArray<uint16> Temp = MoveTemp(*OutMerged);
-
-PRAGMA_DISABLE_DEPRECATION_WARNINGS
 				MergeChangeList(Data, HistoryItem.Changed, Temp, *OutMerged);
-PRAGMA_ENABLE_DEPRECATION_WARNINGS
 
 				HistoryItem.Changed.Empty();
 
@@ -1476,10 +1458,10 @@ public:
 };
 
 void FRepLayout::MergeChangeList_r(
-	FRepHandleIterator&		RepHandleIterator1,
-	FRepHandleIterator&		RepHandleIterator2,
-	const uint8* RESTRICT	SourceData,
-	TArray<uint16>&			OutChanged) const
+	FRepHandleIterator& RepHandleIterator1,
+	FRepHandleIterator& RepHandleIterator2,
+	const FConstRepObjectDataBuffer SourceData,
+	TArray<uint16>& OutChanged) const
 {
 	while (true)
 	{
@@ -1570,35 +1552,33 @@ void FRepLayout::MergeChangeList_r(
 
 		if (Cmd.Type == ERepLayoutCmdType::DynamicArray)
 		{
-			const uint8* Data = SourceData + ArrayOffset + Cmd.Offset;
-
-			const FScriptArray* Array = (FScriptArray *)Data;
+			const FConstRepObjectDataBuffer Data = (SourceData + Cmd) + ArrayOffset;
+			const FScriptArray* Array = (FScriptArray *)Data.Data;
+			const FConstRepObjectDataBuffer ArrayData(Array->GetData());
 
 			FScopedIteratorArrayTracker ArrayTracker1(ActiveIterator1);
 			FScopedIteratorArrayTracker ArrayTracker2(ActiveIterator2);
 
 			const int32 OriginalChangedNum	= OutChanged.AddUninitialized();
 
-			const uint8* NewData = (uint8*)Array->GetData();
-
 			TArray<FHandleToCmdIndex>& ArrayHandleToCmdIndex = ActiveIterator1 ? *ActiveIterator1->HandleToCmdIndex[Cmd.RelativeHandle - 1].HandleToCmdIndex : *ActiveIterator2->HandleToCmdIndex[Cmd.RelativeHandle - 1].HandleToCmdIndex; //-V595
 
 			if (!ActiveIterator1)
 			{
 				FRepHandleIterator ArrayIterator2(ActiveIterator2->ChangelistIterator, Cmds, ArrayHandleToCmdIndex, Cmd.ElementSize, Array->Num(), CmdIndex + 1, Cmd.EndCmd - 1);
-				PruneChangeList_r(ArrayIterator2, NewData, OutChanged);
+				PruneChangeList_r(ArrayIterator2, ArrayData, OutChanged);
 			}
 			else if (!ActiveIterator2)
 			{
 				FRepHandleIterator ArrayIterator1(ActiveIterator1->ChangelistIterator, Cmds, ArrayHandleToCmdIndex, Cmd.ElementSize, Array->Num(), CmdIndex + 1, Cmd.EndCmd - 1);
-				PruneChangeList_r(ArrayIterator1, NewData, OutChanged);
+				PruneChangeList_r(ArrayIterator1, ArrayData, OutChanged);
 			}
 			else
 			{
 				FRepHandleIterator ArrayIterator1(ActiveIterator1->ChangelistIterator, Cmds, ArrayHandleToCmdIndex, Cmd.ElementSize, Array->Num(), CmdIndex + 1, Cmd.EndCmd - 1);
 				FRepHandleIterator ArrayIterator2(ActiveIterator2->ChangelistIterator, Cmds, ArrayHandleToCmdIndex, Cmd.ElementSize, Array->Num(), CmdIndex + 1, Cmd.EndCmd - 1);
 
-				MergeChangeList_r(ArrayIterator1, ArrayIterator2, NewData, OutChanged);
+				MergeChangeList_r(ArrayIterator1, ArrayIterator2, ArrayData, OutChanged);
 			}
 
 			// Patch in the jump offset
@@ -1611,9 +1591,9 @@ void FRepLayout::MergeChangeList_r(
 }
 
 void FRepLayout::PruneChangeList_r(
-	FRepHandleIterator&		RepHandleIterator,
-	const uint8* RESTRICT	SourceData,
-	TArray<uint16>&			OutChanged) const
+	FRepHandleIterator& RepHandleIterator,
+	const FConstRepObjectDataBuffer SourceData,
+	TArray<uint16>& OutChanged) const
 {
 	while (RepHandleIterator.NextHandle())
 	{
@@ -1626,20 +1606,18 @@ void FRepLayout::PruneChangeList_r(
 
 		if (Cmd.Type == ERepLayoutCmdType::DynamicArray)
 		{
-			const uint8* Data = SourceData + ArrayOffset + Cmd.Offset;
-
-			const FScriptArray* Array = (FScriptArray *)Data;
+			const FConstRepObjectDataBuffer Data = (SourceData + Cmd) + ArrayOffset;
+			const FScriptArray* Array = (FScriptArray *)Data.Data;
+			const FConstRepObjectDataBuffer ArrayData(Array->GetData());
 
 			FScopedIteratorArrayTracker ArrayTracker(&RepHandleIterator);
 
 			const int32 OriginalChangedNum = OutChanged.AddUninitialized();
 
-			const uint8* NewData = (uint8*)Array->GetData();
-
 			TArray<FHandleToCmdIndex>& ArrayHandleToCmdIndex = *RepHandleIterator.HandleToCmdIndex[Cmd.RelativeHandle - 1].HandleToCmdIndex;
 
 			FRepHandleIterator ArrayIterator(RepHandleIterator.ChangelistIterator, Cmds, ArrayHandleToCmdIndex, Cmd.ElementSize, Array->Num(), CmdIndex + 1, Cmd.EndCmd - 1);
-			PruneChangeList_r(ArrayIterator, NewData, OutChanged);
+			PruneChangeList_r(ArrayIterator, ArrayData, OutChanged);
 
 			// Patch in the jump offset
 			OutChanged[OriginalChangedNum] = OutChanged.Num() - (OriginalChangedNum + 1);
@@ -1651,10 +1629,10 @@ void FRepLayout::PruneChangeList_r(
 }
 
 void FRepLayout::FilterChangeList(
-	const TArray<uint16>&	Changelist,
-	const TBitArray<>&		InactiveParents,
-	TArray<uint16>&			OutInactiveProperties,
-	TArray<uint16>&			OutActiveProperties) const
+	const TArray<uint16>& Changelist,
+	const TBitArray<>& InactiveParents,
+	TArray<uint16>& OutInactiveProperties,
+	TArray<uint16>& OutActiveProperties) const
 {
 	FChangelistIterator ChangelistIterator(Changelist, 0);
 	FRepHandleIterator HandleIterator(ChangelistIterator, Cmds, BaseHandleToCmdIndex, 0, 1, 0, Cmds.Num() - 1);
@@ -1692,9 +1670,9 @@ void FRepLayout::FilterChangeList(
 }
 
 void FRepLayout::FilterChangeListToActive(
-	const TArray<uint16>&	Changelist,
-	const TBitArray<>&		InactiveParents,
-	TArray<uint16>&			OutProperties) const
+	const TArray<uint16>& Changelist,
+	const TBitArray<>& InactiveParents,
+	TArray<uint16>& OutProperties) const
 {
 	FChangelistIterator ChangelistIterator(Changelist, 0);
 	FRepHandleIterator HandleIterator(ChangelistIterator, Cmds, BaseHandleToCmdIndex, 0, 1, 0, Cmds.Num() - 1);
@@ -1734,13 +1712,13 @@ void FRepLayout::FilterChangeListToActive(
 }
 
 const FRepSerializedPropertyInfo* FRepSerializationSharedInfo::WriteSharedProperty(
-	const FRepLayoutCmd&	Cmd,
-	const FGuid&			PropertyGuid,
-	const int32				CmdIndex,
-	const uint16			Handle,
-	const uint8* RESTRICT	Data,
-	const bool				bWriteHandle,
-	const bool				bDoChecksum)
+	const FRepLayoutCmd& Cmd,
+	const FGuid& PropertyGuid,
+	const int32 CmdIndex,
+	const uint16 Handle,
+	const FConstRepObjectDataBuffer Data,
+	const bool bWriteHandle,
+	const bool bDoChecksum)
 {
 #if !(UE_BUILD_SHIPPING || UE_BUILD_TEST)
 	check(!SharedPropertyInfo.ContainsByPredicate([&](const FRepSerializedPropertyInfo& Info)
@@ -1763,7 +1741,7 @@ const FRepSerializedPropertyInfo* FRepSerializationSharedInfo::WriteSharedProper
 	SharedPropInfo.PropBitOffset = SerializedProperties->GetNumBits();
 
 	// This property changed, so send it
-	Cmd.Property->NetSerializeItem(*SerializedProperties, nullptr, (void*)Data);
+	Cmd.Property->NetSerializeItem(*SerializedProperties, nullptr, const_cast<uint8*>(Data.Data));
 
 	const int64 NumPropEndBits = SerializedProperties->GetNumBits();
 
@@ -1787,7 +1765,7 @@ void FRepLayout::SendProperties_r(
 	FNetBitWriter& Writer,
 	const bool bDoChecksum,
 	FRepHandleIterator& HandleIterator,
-	const uint8* RESTRICT SourceData,
+	const FConstRepObjectDataBuffer SourceData,
 	const int32 ArrayDepth,
 	const FRepSerializationSharedInfo& SharedInfo) const
 {
@@ -1798,13 +1776,14 @@ void FRepLayout::SendProperties_r(
 
 		UE_LOG(LogRepProperties, VeryVerbose, TEXT("SendProperties_r: Parent=%d, Cmd=%d, ArrayIndex=%d"), Cmd.ParentIndex, HandleIterator.CmdIndex, HandleIterator.ArrayIndex);
 		
-		const uint8* Data = SourceData + HandleIterator.ArrayOffset + Cmd.Offset;
+		FConstRepObjectDataBuffer Data = (SourceData + Cmd) + HandleIterator.ArrayOffset;
 
 		if (Cmd.Type == ERepLayoutCmdType::DynamicArray)
 		{
 			WritePropertyHandle(Writer, HandleIterator.Handle, bDoChecksum);
 
-			const FScriptArray* Array = (FScriptArray *)Data;
+			const FScriptArray* Array = (FScriptArray *)Data.Data;
+			const FConstRepObjectDataBuffer ArrayData(Array->GetData());
 
 			// Write array num
 			uint16 ArrayNum = Array->Num();
@@ -1819,8 +1798,6 @@ void FRepLayout::SendProperties_r(
 
 			const int32 OldChangedIndex = HandleIterator.ChangelistIterator.ChangedIndex;
 
-			const uint8* NewData = (uint8*)Array->GetData();
-
 			TArray<FHandleToCmdIndex>& ArrayHandleToCmdIndex = *HandleIterator.HandleToCmdIndex[Cmd.RelativeHandle - 1].HandleToCmdIndex;
 
 			FRepHandleIterator ArrayHandleIterator(HandleIterator.ChangelistIterator, Cmds, ArrayHandleToCmdIndex, Cmd.ElementSize, ArrayNum, HandleIterator.CmdIndex + 1, Cmd.EndCmd - 1);
@@ -1828,7 +1805,7 @@ void FRepLayout::SendProperties_r(
 			check(ArrayHandleIterator.ArrayElementSize> 0);
 			check(ArrayHandleIterator.NumHandlesPerElement> 0);
 
-			SendProperties_r(RepState, ChangedTracker, Writer, bDoChecksum, ArrayHandleIterator, NewData, ArrayDepth + 1, SharedInfo);
+			SendProperties_r(RepState, ChangedTracker, Writer, bDoChecksum, ArrayHandleIterator, ArrayData, ArrayDepth + 1, SharedInfo);
 
 			check(HandleIterator.ChangelistIterator.ChangedIndex - OldChangedIndex == ArrayChangedCount);				// Make sure we read correct amount
 			check(HandleIterator.ChangelistIterator.Changed[HandleIterator.ChangelistIterator.ChangedIndex] == 0);	// Make sure we are at the end
@@ -1861,7 +1838,7 @@ void FRepLayout::SendProperties_r(
 
 		if ((GNetSharedSerializedData != 0) && ((Cmd.Flags & ERepLayoutFlags::IsSharedSerialization) != ERepLayoutFlags::None))
 		{
-			FGuid PropertyGuid(HandleIterator.CmdIndex, HandleIterator.ArrayIndex, ArrayDepth, (int32)((PTRINT)Data & 0xFFFFFFFF));
+			FGuid PropertyGuid(HandleIterator.CmdIndex, HandleIterator.ArrayIndex, ArrayDepth, (int32)((PTRINT)Data.Data & 0xFFFFFFFF));
 
 			SharedPropInfo = SharedInfo.SharedPropertyInfo.FindByPredicate([&](const FRepSerializedPropertyInfo& Info) 
 			{ 
@@ -1881,7 +1858,7 @@ void FRepLayout::SendProperties_r(
 				UE_LOG(LogRepProperties, VeryVerbose, TEXT("SerializeProperties_r: Verify SharedSerialization, NetSerializeItem"));
 
 				WritePropertyHandle(Writer, HandleIterator.Handle, bDoChecksum);
-				Cmd.Property->NetSerializeItem(Writer, Writer.PackageMap, (void*)Data);
+				Cmd.Property->NetSerializeItem(Writer, Writer.PackageMap, const_cast<uint8*>(Data.Data));
 
 #ifdef ENABLE_PROPERTY_CHECKSUMS
 				if (bDoChecksum)
@@ -1921,7 +1898,7 @@ void FRepLayout::SendProperties_r(
 			const int32 NumStartBits = Writer.GetNumBits();
 
 			// This property changed, so send it
-			Cmd.Property->NetSerializeItem(Writer, Writer.PackageMap, (void*)Data);
+			Cmd.Property->NetSerializeItem(Writer, Writer.PackageMap, const_cast<uint8*>(Data.Data));
 			UE_LOG(LogRepProperties, VeryVerbose, TEXT("SerializeProperties_r: NetSerializeItem"));
 
 			const int32 NumEndBits = Writer.GetNumBits();
@@ -1948,7 +1925,7 @@ void FRepLayout::SendProperties_r(
 void FRepLayout::SendProperties(
 	FSendingRepState* RESTRICT RepState,
 	FRepChangedPropertyTracker* ChangedTracker,
-	const uint8* RESTRICT Data,
+	const FConstRepObjectDataBuffer Data,
 	UClass* ObjectClass,
 	FNetBitWriter& Writer,
 	TArray<uint16>& Changed,
@@ -2034,12 +2011,12 @@ TSharedPtr<FNetFieldExportGroup> FRepLayout::CreateNetfieldExportGroup() const
 }
 
 static FORCEINLINE void WriteProperty_BackwardsCompatible(
-	FNetBitWriter&			Writer,
-	const FRepLayoutCmd&	Cmd,
-	const int32				CmdIndex,
-	const UObject*			Owner,
-	const uint8*			Data,
-	const bool				bDoChecksum)
+	FNetBitWriter& Writer,
+	const FRepLayoutCmd& Cmd,
+	const int32 CmdIndex,
+	const UObject* Owner,
+	const FConstRepObjectDataBuffer Data,
+	const bool bDoChecksum)
 {
 #if !(UE_BUILD_SHIPPING || UE_BUILD_TEST)
 	if (GDoReplicationContextString> 0)
@@ -2053,7 +2030,7 @@ static FORCEINLINE void WriteProperty_BackwardsCompatible(
 	FNetBitWriter TempWriter(Writer.PackageMap, 0);
 
 	// This property changed, so send it
-	Cmd.Property->NetSerializeItem(TempWriter, TempWriter.PackageMap, (void*)Data);
+	Cmd.Property->NetSerializeItem(TempWriter, TempWriter.PackageMap, const_cast<uint8*>(Data.Data));
 	UE_LOG(LogRepProperties, VeryVerbose, TEXT("WriteProperty_BackwardsCompatible: (Temp) NetSerializeItem"));
 
 	uint32 NumBits = TempWriter.GetNumBits();
@@ -2088,7 +2065,7 @@ void FRepLayout::SendProperties_BackwardsCompatible_r(
 	FNetBitWriter& Writer,
 	const bool bDoChecksum,
 	FRepHandleIterator& HandleIterator,
-	const uint8* RESTRICT SourceData) const
+	const FConstRepObjectDataBuffer SourceData) const
 {
 	int32 OldIndex = -1;
 
@@ -2101,7 +2078,7 @@ void FRepLayout::SendProperties_BackwardsCompatible_r(
 
 		UE_LOG(LogRepProperties, VeryVerbose, TEXT("SendProperties_BackwardsCompatible_r: Parent=%d, Cmd=%d, ArrayIndex=%d"), Cmd.ParentIndex, HandleIterator.CmdIndex, HandleIterator.ArrayIndex);
 
-		const uint8* Data = SourceData + HandleIterator.ArrayOffset + Cmd.Offset;
+		FConstRepObjectDataBuffer Data = (SourceData + Cmd) + HandleIterator.ArrayOffset;
 
 		PackageMapClient->TrackNetFieldExport(NetFieldExportGroup, HandleIterator.CmdIndex);
 
@@ -2122,7 +2099,8 @@ void FRepLayout::SendProperties_BackwardsCompatible_r(
 
 		if (Cmd.Type == ERepLayoutCmdType::DynamicArray)
 		{
-			const FScriptArray* Array = (FScriptArray *)Data;
+			const FScriptArray* Array = (FScriptArray *)Data.Data;
+			const FConstRepObjectDataBuffer ArrayData(Array->GetData());
 
 			uint32 ArrayNum = Array->Num();
 
@@ -2133,8 +2111,6 @@ void FRepLayout::SendProperties_BackwardsCompatible_r(
 
 			const int32 OldChangedIndex = HandleIterator.ChangelistIterator.ChangedIndex;
 
-			const uint8* NewData = (uint8*)Array->GetData();
-
 			TArray<FHandleToCmdIndex>& ArrayHandleToCmdIndex = *HandleIterator.HandleToCmdIndex[Cmd.RelativeHandle - 1].HandleToCmdIndex;
 
 			FRepHandleIterator ArrayHandleIterator(HandleIterator.ChangelistIterator, Cmds, ArrayHandleToCmdIndex, Cmd.ElementSize, ArrayNum, HandleIterator.CmdIndex + 1, Cmd.EndCmd - 1);
@@ -2151,7 +2127,7 @@ void FRepLayout::SendProperties_BackwardsCompatible_r(
 			if (ArrayNum> 0)
 			{
 				UE_LOG(LogRepProperties, VeryVerbose, TEXT("SendProperties_BackwardsCompatible_r: (Temp) Array Recurse Properties"), ArrayNum);
-				SendProperties_BackwardsCompatible_r(RepState, PackageMapClient, NetFieldExportGroup, ChangedTracker, TempWriter, bDoChecksum, ArrayHandleIterator, NewData);
+				SendProperties_BackwardsCompatible_r(RepState, PackageMapClient, NetFieldExportGroup, ChangedTracker, TempWriter, bDoChecksum, ArrayHandleIterator, ArrayData);
 			}
 
 			uint32 EndArrayIndex = 0;
@@ -2195,7 +2171,7 @@ void FRepLayout::SendAllProperties_BackwardsCompatible_r(
 	FNetFieldExportGroup* NetFieldExportGroup,
 	const int32 CmdStart,
 	const int32 CmdEnd, 
-	const uint8* SourceData) const
+	const FConstRepObjectDataBuffer SourceData) const
 {
 	FNetBitWriter TempWriter(Writer.PackageMap, 0);
 
@@ -2211,11 +2187,12 @@ void FRepLayout::SendAllProperties_BackwardsCompatible_r(
 
 		WritePropertyHandle_BackwardsCompatible(Writer, CmdIndex + 1, bDoChecksum);
 
-		const uint8* Data = SourceData + Cmd.Offset;
+		FConstRepObjectDataBuffer Data = SourceData + Cmd;
 
 		if (Cmd.Type == ERepLayoutCmdType::DynamicArray)
 		{			
-			const FScriptArray* Array = (FScriptArray *)Data;
+			const FScriptArray* Array = (FScriptArray *)Data.Data;
+			const FConstRepObjectDataBuffer ArrayData(Array->GetData());
 
 			TempWriter.Reset();
 
@@ -2231,7 +2208,8 @@ void FRepLayout::SendAllProperties_BackwardsCompatible_r(
 				TempWriter.SerializeIntPacked(ArrayIndex);
 
 				UE_LOG(LogRepProperties, VeryVerbose, TEXT("SendAllProperties_BackwardsCompatible_r: (Temp) ArrayIndex=%d"), ArrayIndex);
-				SendAllProperties_BackwardsCompatible_r(RepState, TempWriter, bDoChecksum, PackageMapClient, NetFieldExportGroup, CmdIndex + 1, Cmd.EndCmd - 1, ((const uint8*)Array->GetData()) + Cmd.ElementSize * i);
+				const int32 ArrayElementOffset = Cmd.ElementSize * i;
+				SendAllProperties_BackwardsCompatible_r(RepState, TempWriter, bDoChecksum, PackageMapClient, NetFieldExportGroup, CmdIndex + 1, Cmd.EndCmd - 1, ArrayData + ArrayElementOffset);
 			}
 
 			uint32 EndArrayIndex = 0;
@@ -2256,7 +2234,7 @@ void FRepLayout::SendAllProperties_BackwardsCompatible_r(
 			{
 				Data = reinterpret_cast<const uint8*>(&(RepState->SavedRemoteRole));
 			}
-		}		
+		}
 
 		WriteProperty_BackwardsCompatible(Writer, Cmd, CmdIndex, Owner, Data, bDoChecksum);
 	}
@@ -2267,7 +2245,7 @@ void FRepLayout::SendAllProperties_BackwardsCompatible_r(
 void FRepLayout::SendProperties_BackwardsCompatible(
 	FSendingRepState* RESTRICT RepState,
 	FRepChangedPropertyTracker* ChangedTracker,
-	const uint8* RESTRICT Data,
+	const FConstRepObjectDataBuffer Data,
 	UNetConnection* Connection,
 	FNetBitWriter& Writer,
 	TArray<uint16>& Changed) const
@@ -2394,7 +2372,7 @@ static bool ReceivePropertyHelper(
 #ifdef ENABLE_PROPERTY_CHECKSUMS
 	if (bDoChecksum)
 	{
-		SerializeReadWritePropertyChecksum(Cmd, CmdIndex, Data + SwappedCmd, Bunch);
+		SerializeReadWritePropertyChecksum(Cmd, CmdIndex, FConstRepObjectDataBuffer(Data + SwappedCmd), Bunch);
 	}
 #endif
 
@@ -2647,7 +2625,7 @@ bool FRepLayout::ReceiveProperties(
 	UActorChannel* OwningChannel,
 	UClass* InObjectClass,
 	FReceivingRepState* RESTRICT RepState,
-	void* RESTRICT Data,
+	FRepObjectDataBuffer Data,
 	FNetBitReader& InBunch,
 	bool& bOutHasUnmapped,
 	bool& bOutGuidsChanged,
@@ -2683,7 +2661,7 @@ bool FRepLayout::ReceiveProperties(
 	ReceivePropertiesImpl.ReadNextHandle();
 
 	// Read all properties
-	ReceivePropertiesImpl.ProcessCmds((uint8*)Data, RepState->StaticBuffer.GetData());
+	ReceivePropertiesImpl.ProcessCmds(Data, RepState->StaticBuffer.GetData());
 
 	// Make sure we're waiting on the last NULL terminator
 	if (ReceivePropertiesImpl.WaitingHandle != 0)
@@ -2708,7 +2686,7 @@ bool FRepLayout::ReceiveProperties(
 bool FRepLayout::ReceiveProperties_BackwardsCompatible(
 	UNetConnection* Connection,
 	FReceivingRepState* RESTRICT RepState,
-	void* RESTRICT Data,
+	FRepObjectDataBuffer Data,
 	FNetBitReader& InBunch,
 	bool& bOutHasUnmapped,
 	const bool bEnableRepNotifies,
@@ -2727,7 +2705,7 @@ bool FRepLayout::ReceiveProperties_BackwardsCompatible(
 
 	UE_LOG(LogRepProperties, VeryVerbose, TEXT("ReceiveProperties_BackwardsCompatible: Owner=%s, NetFieldExportGroupFound=%d"), *OwnerPathName, !!NetFieldExportGroup.IsValid());
 
-	return ReceiveProperties_BackwardsCompatible_r(RepState, NetFieldExportGroup.Get(), InBunch, 0, Cmds.Num() - 1, (bEnableRepNotifies && RepState) ? RepState->StaticBuffer.GetData() : nullptr, (uint8*)Data, (uint8*)Data, RepState ? &RepState->GuidReferencesMap : nullptr, bOutHasUnmapped, bOutGuidsChanged);
+	return ReceiveProperties_BackwardsCompatible_r(RepState, NetFieldExportGroup.Get(), InBunch, 0, Cmds.Num() - 1, (bEnableRepNotifies && RepState) ? RepState->StaticBuffer.GetData() : nullptr, Data, Data, RepState ? &RepState->GuidReferencesMap : nullptr, bOutHasUnmapped, bOutGuidsChanged);
 }
 
 int32 FRepLayout::FindCompatibleProperty(
@@ -2762,9 +2740,9 @@ bool FRepLayout::ReceiveProperties_BackwardsCompatible_r(
 	FNetBitReader& Reader,
 	const int32 CmdStart,
 	const int32 CmdEnd,
-	uint8* RESTRICT ShadowData,
-	uint8* RESTRICT OldData,
-	uint8* RESTRICT Data,
+	FRepShadowDataBuffer ShadowData,
+	FRepObjectDataBuffer OldData,
+	FRepObjectDataBuffer Data,
 	FGuidReferencesMap* GuidReferencesMap,
 	bool& bOutHasUnmapped,
 	bool& bOutGuidsChanged) const
@@ -2888,10 +2866,10 @@ bool FRepLayout::ReceiveProperties_BackwardsCompatible_r(
 				return false;
 			}
 
-			const int32 AbsOffset = (Data - OldData) + Cmd.Offset;
+			const int32 AbsOffset = (Data.Data - OldData.Data) + Cmd.Offset;
 
-			FScriptArray* DataArray = (FScriptArray*)(Data + Cmd.Offset);
-			FScriptArray* ShadowArray = ShadowData ? (FScriptArray*)(ShadowData + Cmd.ShadowOffset) : nullptr;
+			FScriptArray* DataArray = (FScriptArray*)(Data + Cmd).Data;
+			FScriptArray* ShadowArray = ShadowData ? (FScriptArray*)(ShadowData + Cmd).Data : nullptr;
 
 			const int32 ShadowArrayNum = ShadowArray ? ShadowArray->Num() : INDEX_NONE;
 
@@ -2909,7 +2887,7 @@ bool FRepLayout::ReceiveProperties_BackwardsCompatible_r(
 				CmdIndex,
 				&LocalShadowData,
 				&LocalData,
-				ShadowData != nullptr ? &RepState->RepNotifies : nullptr);
+				ShadowData ? &RepState->RepNotifies : nullptr);
 
 			// Read until we read all array elements
 			while (true)
@@ -2958,10 +2936,10 @@ bool FRepLayout::ReceiveProperties_BackwardsCompatible_r(
 					return false;
 				}
 
-				const int32 ElementOffset = Index * Cmd.ElementSize;
+				const int32 ArrayElementOffset = Index * Cmd.ElementSize;
 
-				FRepObjectDataBuffer ElementData = LocalData + ElementOffset;
-				FRepShadowDataBuffer ElementShadowData = LocalShadowData ? LocalShadowData + ElementOffset : FRepShadowDataBuffer((uint8*)nullptr);
+				FRepObjectDataBuffer ElementData = LocalData + ArrayElementOffset;
+				FRepShadowDataBuffer ElementShadowData = LocalShadowData ? LocalShadowData + ArrayElementOffset : nullptr;
 
 				if (!ReceiveProperties_BackwardsCompatible_r(RepState, NetFieldExportGroup, TempReader, CmdIndex + 1, Cmd.EndCmd - 1, ElementShadowData, LocalData, ElementData, NewGuidReferencesArray, bOutHasUnmapped, bOutGuidsChanged))
 				{
@@ -2983,9 +2961,9 @@ bool FRepLayout::ReceiveProperties_BackwardsCompatible_r(
 		}
 		else
 		{
-			const int32 ElementOffset = (Data - OldData);
+			const int32 ElementOffset = (Data.Data - OldData.Data);
 
-			if (ReceivePropertyHelper(TempReader, GuidReferencesMap, ElementOffset, ShadowData, Data, ShadowData != nullptr ? &RepState->RepNotifies : nullptr, Parents, Cmds, CmdIndex, false, bOutGuidsChanged, false))
+			if (ReceivePropertyHelper(TempReader, GuidReferencesMap, ElementOffset, ShadowData, Data, ShadowData ? &RepState->RepNotifies : nullptr, Parents, Cmds, CmdIndex, false, bOutGuidsChanged, false))
 			{
 				bOutHasUnmapped = true;
 			}
@@ -3011,9 +2989,9 @@ FGuidReferences::~FGuidReferences()
 }
 
 void FRepLayout::GatherGuidReferences_r(
-	FGuidReferencesMap*	GuidReferencesMap,
-	TSet<FNetworkGUID>&	OutReferencedGuids,
-	int32&				OutTrackedGuidMemoryBytes) const
+	FGuidReferencesMap* GuidReferencesMap,
+	TSet<FNetworkGUID>& OutReferencedGuids,
+	int32& OutTrackedGuidMemoryBytes) const
 {
 	for (const auto& GuidReferencePair : *GuidReferencesMap)
 	{
@@ -3085,8 +3063,8 @@ void FRepLayout::UpdateUnmappedObjects_r(
 	FGuidReferencesMap* GuidReferencesMap,
 	UObject* OriginalObject,
 	UPackageMap* PackageMap, 
-	uint8* RESTRICT ShadowData, 
-	uint8* RESTRICT Data, 
+	FRepShadowDataBuffer ShadowData, 
+	FRepObjectDataBuffer Data, 
 	const int32 MaxAbsOffset,
 	bool& bOutSomeObjectsWereMapped,
 	bool& bOutHasMoreUnmapped) const
@@ -3111,12 +3089,15 @@ void FRepLayout::UpdateUnmappedObjects_r(
 		{
 			check(Cmd.Type == ERepLayoutCmdType::DynamicArray);
 			
-			FScriptArray* StoredArray = (FScriptArray*)(ShadowData + Cmd.ShadowOffset);
-			FScriptArray* Array = (FScriptArray*)(Data + AbsOffset);
+			FScriptArray* ShadowArray = (FScriptArray*)(ShadowData + Cmd).Data;
+			FScriptArray* Array = (FScriptArray*)(Data + AbsOffset).Data;
 			
-			const int32 NewMaxOffset = FMath::Min(StoredArray->Num() * Cmd.ElementSize, Array->Num() * Cmd.ElementSize);
+			FRepShadowDataBuffer ShadowArrayData(ShadowArray->GetData());
+			FRepObjectDataBuffer ArrayData(Array->GetData());
+			
+			const int32 NewMaxOffset = FMath::Min(ShadowArray->Num() * Cmd.ElementSize, Array->Num() * Cmd.ElementSize);
 
-			UpdateUnmappedObjects_r(RepState, GuidReferences.Array, OriginalObject, PackageMap, (uint8*)StoredArray->GetData(), (uint8*)Array->GetData(), NewMaxOffset, bOutSomeObjectsWereMapped, bOutHasMoreUnmapped);
+			UpdateUnmappedObjects_r(RepState, GuidReferences.Array, OriginalObject, PackageMap, ShadowArrayData, ArrayData, NewMaxOffset, bOutSomeObjectsWereMapped, bOutHasMoreUnmapped);
 			continue;
 		}
 
@@ -3164,7 +3145,7 @@ void FRepLayout::UpdateUnmappedObjects_r(
 			// Copy current value over so we can check to see if it changed
 			if (Parent.Property->HasAnyPropertyFlags(CPF_RepNotify))
 			{
-				StoreProperty(Cmd, ShadowData + Cmd.ShadowOffset, Data + AbsOffset);
+				StoreProperty(Cmd, ShadowData + Cmd, Data + AbsOffset);
 			}
 
 			// Initialize the reader with the stored buffer that we need to read from
@@ -3176,7 +3157,7 @@ void FRepLayout::UpdateUnmappedObjects_r(
 			// Check to see if this property changed
 			if (Parent.Property->HasAnyPropertyFlags(CPF_RepNotify))
 			{
-				if (Parent.RepNotifyCondition == REPNOTIFY_Always || !PropertiesAreIdentical(Cmd, ShadowData + Cmd.ShadowOffset, Data + AbsOffset))
+				if (Parent.RepNotifyCondition == REPNOTIFY_Always || !PropertiesAreIdentical(Cmd, ShadowData + Cmd, Data + AbsOffset))
 				{
 					// If this properties needs an OnRep, queue that up to be handled later
 					RepState->RepNotifies.AddUnique(Parent.Property);
@@ -3305,20 +3286,22 @@ static void ValidateWithChecksum_DynamicArray_r(
 		UE_LOG(LogRep, Fatal, TEXT("ValidateWithChecksum_AnyArray_r: Array element sizes different! %s %i / %i"), *Cmd.Property->GetFullName(), ElementSize, Cmd.ElementSize);
 	}
 
-	uint8* LocalData = (uint8*)Array->GetData();
+	const TConstRepDataBuffer<DataType> ArrayData(Array->GetData());
 	for (int32 i = 0; i < ArrayNum - 1; i++)
 	{
-		ValidateWithChecksum_r(CmdIt, TConstRepDataBuffer<DataType>(LocalData + i * ElementSize), Ar);
+		const int32 ArrayElementsOffset = i * ElementSize;
+		ValidateWithChecksum_r(CmdIt, Data + ArrayElementsOffset, Ar);
 		CmdIt -= ArraySubCommands;
 	}
 
-	ValidateWithChecksum_r(CmdIt, TConstRepDataBuffer<DataType>(LocalData + (ArrayNum - 1) * ElementSize), Ar);
+	const int32 ArrayElementOffset = (ArrayNum - 1) * ElementSize;
+	ValidateWithChecksum_r<DataType>(CmdIt, ArrayData + ArrayElementOffset, Ar);
 }
 
 template<ERepDataBufferType DataType>
 static void ValidateWithChecksum_r(
 	TArray<FRepLayoutCmd>::TConstIterator& CmdIt,
-	TConstRepDataBuffer<DataType> Data, 
+	const TConstRepDataBuffer<DataType> Data, 
 	FBitArchive& Ar)
 {
 	for (; CmdIt->Type != ERepLayoutCmdType::Return; ++CmdIt)
@@ -3352,7 +3335,7 @@ uint32 FRepLayout::GenerateChecksum(const FRepState* RepState) const
 }
 
 void FRepLayout::PruneChangeList(
-	const void* RESTRICT Data,
+	const FConstRepObjectDataBuffer Data,
 	const TArray<uint16>& Changed,
 	TArray<uint16>& PrunedChanged) const
 {
@@ -3364,14 +3347,14 @@ void FRepLayout::PruneChangeList(
 	{
 		FChangelistIterator ChangelistIterator(Changed, 0);
 		FRepHandleIterator HandleIterator(ChangelistIterator, Cmds, BaseHandleToCmdIndex, 0, 1, 0, Cmds.Num() - 1);
-		PruneChangeList_r(HandleIterator, (uint8*)Data, PrunedChanged);
+		PruneChangeList_r(HandleIterator, Data, PrunedChanged);
 	}
 
 	PrunedChanged.Add(0);
 }
 
 void FRepLayout::MergeChangeList(
-	const uint8* RESTRICT Data,
+	const FConstRepObjectDataBuffer Data,
 	const TArray<uint16>& Dirty1,
 	const TArray<uint16>& Dirty2,
 	TArray<uint16>& MergedDirty) const
@@ -3385,7 +3368,7 @@ void FRepLayout::MergeChangeList(
 		{
 			FChangelistIterator ChangelistIterator(Dirty1, 0);
 			FRepHandleIterator HandleIterator(ChangelistIterator, Cmds, BaseHandleToCmdIndex, 0, 1, 0, Cmds.Num() - 1);
-			PruneChangeList_r(HandleIterator, (uint8*)Data, MergedDirty);
+			PruneChangeList_r(HandleIterator, Data, MergedDirty);
 		}
 		else
 		{
@@ -3395,7 +3378,7 @@ void FRepLayout::MergeChangeList(
 			FChangelistIterator ChangelistIterator2(Dirty2, 0);
 			FRepHandleIterator HandleIterator2(ChangelistIterator2, Cmds, BaseHandleToCmdIndex, 0, 1, 0, Cmds.Num() - 1);
 
-			MergeChangeList_r(HandleIterator1, HandleIterator2, (uint8*)Data, MergedDirty);
+			MergeChangeList_r(HandleIterator1, HandleIterator2, Data, MergedDirty);
 		}
 	}
 
@@ -3403,14 +3386,14 @@ void FRepLayout::MergeChangeList(
 }
 
 void FRepLayout::SanityCheckChangeList_DynamicArray_r(
-	const int32				CmdIndex, 
-	const uint8* RESTRICT	Data, 
-	TArray<uint16> &		Changed,
-	int32 &					ChangedIndex) const
+	const int32 CmdIndex, 
+	const FConstRepObjectDataBuffer Data, 
+	TArray<uint16>& Changed,
+	int32& ChangedIndex) const
 {
-	const FRepLayoutCmd& Cmd = Cmds[ CmdIndex ];
+	const FRepLayoutCmd& Cmd = Cmds[CmdIndex];
 
-	FScriptArray * Array = (FScriptArray *)Data;
+	FScriptArray * Array = (FScriptArray *)Data.Data;
 
 	// Read the jump offset
 	// We won't need to actually jump over anything because we expect the change list to be pruned once we get here
@@ -3419,13 +3402,14 @@ void FRepLayout::SanityCheckChangeList_DynamicArray_r(
 
 	const int32 OldChangedIndex = ChangedIndex;
 
-	Data = (uint8*)Array->GetData();
+	const FConstRepObjectDataBuffer ArrayData = (uint8*)Array->GetData();
 
 	uint16 LocalHandle = 0;
 
 	for (int32 i = 0; i < Array->Num(); i++)
 	{
-		LocalHandle = SanityCheckChangeList_r(CmdIndex + 1, Cmd.EndCmd - 1, Data + i * Cmd.ElementSize, Changed, ChangedIndex, LocalHandle);
+		const int32 ArrayElementOffset = i * Cmd.ElementSize;
+		LocalHandle = SanityCheckChangeList_r(CmdIndex + 1, Cmd.EndCmd - 1, ArrayData + ArrayElementOffset, Changed, ChangedIndex, LocalHandle);
 	}
 
 	check(ChangedIndex - OldChangedIndex == ArrayChangedCount);	// Make sure we read correct amount
@@ -3435,12 +3419,12 @@ void FRepLayout::SanityCheckChangeList_DynamicArray_r(
 }
 
 uint16 FRepLayout::SanityCheckChangeList_r(
-	const int32				CmdStart, 
-	const int32				CmdEnd, 
-	const uint8* RESTRICT	Data, 
-	TArray<uint16> &		Changed,
-	int32 &					ChangedIndex,
-	uint16					Handle 
+	const int32 CmdStart, 
+	const int32 CmdEnd, 
+	const FConstRepObjectDataBuffer Data, 
+	TArray<uint16>& Changed,
+	int32& ChangedIndex,
+	uint16 Handle 
 	) const
 {
 	for (int32 CmdIndex = CmdStart; CmdIndex < CmdEnd; CmdIndex++)
@@ -3457,7 +3441,7 @@ uint16 FRepLayout::SanityCheckChangeList_r(
 			{
 				const int32 LastChangedArrayHandle = Changed[ChangedIndex];
 				ChangedIndex++;
-				SanityCheckChangeList_DynamicArray_r(CmdIndex, Data + Cmd.Offset, Changed, ChangedIndex);
+				SanityCheckChangeList_DynamicArray_r(CmdIndex, Data + Cmd, Changed, ChangedIndex);
 				check(Changed[ChangedIndex] == 0 || Changed[ChangedIndex] > LastChangedArrayHandle);
 			}
 			CmdIndex = Cmd.EndCmd - 1;	// Jump past children of this array (the -1 because of the ++ in the for loop)
@@ -3475,7 +3459,7 @@ uint16 FRepLayout::SanityCheckChangeList_r(
 	return Handle;
 }
 
-void FRepLayout::SanityCheckChangeList(const uint8* RESTRICT Data, TArray<uint16> & Changed) const
+void FRepLayout::SanityCheckChangeList(const FConstRepObjectDataBuffer Data, TArray<uint16> & Changed) const
 {
 	int32 ChangedIndex = 0;
 	SanityCheckChangeList_r(0, Cmds.Num() - 1, Data, Changed, ChangedIndex, 0);
@@ -3752,20 +3736,19 @@ static FName NAME_Vector_NetQuantize(TEXT("Vector_NetQuantize"));
 static FName NAME_UniqueNetIdRepl(TEXT("UniqueNetIdRepl"));
 static FName NAME_RepMovement(TEXT("RepMovement"));
 
-uint32 FRepLayout::AddPropertyCmd(
-	UProperty*				Property,
-	int32					Offset,
-	int32					RelativeHandle,
-	int32					ParentIndex,
-	uint32					ParentChecksum,
-	int32					StaticArrayIndex,
-	const UNetConnection*	ServerConnection)
+static uint32 AddPropertyCmd(
+	TArray<FRepLayoutCmd>& Cmds,
+	UProperty* Property,
+	int32 Offset,
+	int32 RelativeHandle,
+	int32 ParentIndex,
+	uint32 ParentChecksum,
+	int32 StaticArrayIndex,
+	const UNetConnection* ServerConnection)
 {
 	SCOPE_CYCLE_COUNTER(STAT_RepLayout_AddPropertyCmd);
 
-	const int32 Index = Cmds.AddZeroed();
-
-	FRepLayoutCmd & Cmd = Cmds[Index];
+	FRepLayoutCmd & Cmd = Cmds.AddZeroed_GetRef();
 
 	Cmd.Property = Property;
 	Cmd.Type = ERepLayoutCmdType::Property;		// Initially set to generic type
@@ -3882,18 +3865,17 @@ uint32 FRepLayout::AddPropertyCmd(
 	return Cmd.CompatibleChecksum;
 }
 
-uint32 FRepLayout::AddArrayCmd(
-	UArrayProperty*			Property,
-	int32					Offset,
-	int32					RelativeHandle,
-	int32					ParentIndex,
-	uint32					ParentChecksum,
-	int32					StaticArrayIndex,
-	const UNetConnection*	ServerConnection)
+static FORCEINLINE uint32 AddArrayCmd(
+	TArray<FRepLayoutCmd>& Cmds,
+	UArrayProperty* Property,
+	int32 Offset,
+	int32 RelativeHandle,
+	int32 ParentIndex,
+	uint32 ParentChecksum,
+	int32 StaticArrayIndex,
+	const UNetConnection* ServerConnection)
 {
-	const int32 Index = Cmds.AddZeroed();
-
-	FRepLayoutCmd & Cmd = Cmds[Index];
+	FRepLayoutCmd& Cmd = Cmds.AddZeroed_GetRef();
 
 	Cmd.Type = ERepLayoutCmdType::DynamicArray;
 	Cmd.Property = Property;
@@ -3906,23 +3888,40 @@ uint32 FRepLayout::AddArrayCmd(
 	return Cmd.CompatibleChecksum;
 }
 
-void FRepLayout::AddReturnCmd()
+static FORCEINLINE void AddReturnCmd(TArray<FRepLayoutCmd>& Cmds)
 {
-	const int32 Index = Cmds.AddZeroed();
-	
-	FRepLayoutCmd & Cmd = Cmds[Index];
+	Cmds.AddZeroed_GetRef().Type = ERepLayoutCmdType::Return;
+}
+
+enum class ERepBuildType
+{
+	Class,
+	Function,
+	Struct
+};
 
-	Cmd.Type = ERepLayoutCmdType::Return;
+template<ERepBuildType BuildType>
+static FORCEINLINE const int32 GetOffsetForProperty(UProperty& Property)
+{
+	return Property.GetOffset_ForGC();
 }
 
-int32 FRepLayout::InitFromProperty_r(
-	UProperty*				Property,
-	int32					Offset,
-	int32					RelativeHandle,
-	int32					ParentIndex,
-	uint32					ParentChecksum,
-	int32					StaticArrayIndex,
-	const UNetConnection*	ServerConnection)
+template<>
+const FORCEINLINE int32 GetOffsetForProperty<ERepBuildType::Function>(UProperty& Property)
+{
+	return Property.GetOffset_ForUFunction();
+}
+
+template<ERepBuildType BuildType>
+static int32 InitFromProperty_r(
+	TArray<FRepLayoutCmd>& Cmds,
+	UProperty* Property,
+	int32 Offset,
+	int32 RelativeHandle,
+	int32 ParentIndex,
+	uint32 ParentChecksum,
+	int32 StaticArrayIndex,
+	const UNetConnection* ServerConnection)
 {
 	UArrayProperty * ArrayProp = Cast<UArrayProperty>(Property);
 
@@ -3932,11 +3931,11 @@ int32 FRepLayout::InitFromProperty_r(
 
 		RelativeHandle++;
 
-		const uint32 ArrayChecksum = AddArrayCmd(ArrayProp, Offset + ArrayProp->GetOffset_ForGC(), RelativeHandle, ParentIndex, ParentChecksum, StaticArrayIndex, ServerConnection);
+		const uint32 ArrayChecksum = AddArrayCmd(Cmds, ArrayProp, Offset + GetOffsetForProperty<BuildType>(*ArrayProp), RelativeHandle, ParentIndex, ParentChecksum, StaticArrayIndex, ServerConnection);
 
-		InitFromProperty_r(ArrayProp->Inner, 0, 0, ParentIndex, ArrayChecksum, 0, ServerConnection);
+		InitFromProperty_r<BuildType>(Cmds, ArrayProp->Inner, 0, 0, ParentIndex, ArrayChecksum, 0, ServerConnection);
 		
-		AddReturnCmd();
+		AddReturnCmd(Cmds);
 
 		Cmds[CmdStart].EndCmd = Cmds.Num();		// Patch in the offset to jump over our array inner elements
 
@@ -3958,11 +3957,12 @@ int32 FRepLayout::InitFromProperty_r(
 		if (Struct->StructFlags & STRUCT_NetSerializeNative)
 		{
 			RelativeHandle++;
-			AddPropertyCmd(Property, Offset + Property->GetOffset_ForGC(), RelativeHandle, ParentIndex, ParentChecksum, StaticArrayIndex, ServerConnection);
+			AddPropertyCmd(Cmds, Property, Offset + GetOffsetForProperty<BuildType>(*Property), RelativeHandle, ParentIndex, ParentChecksum, StaticArrayIndex, ServerConnection);
 			return RelativeHandle;
 		}
 
-		// Track properties so me can ensure they are sorted by offsets at the end
+		// Track properties so we can ensure they are sorted by offsets at the end
+		// TODO: Do these actually need to be sorted?
 		TArray<UProperty*> NetProperties;
 
 		for (TFieldIterator<UProperty> It(Struct); It; ++It)
@@ -3978,15 +3978,18 @@ int32 FRepLayout::InitFromProperty_r(
 		// Sort NetProperties by memory offset
 		struct FCompareUFieldOffsets
 		{
-			FORCEINLINE bool operator()(UProperty & A, UProperty & B) const
+			FORCEINLINE bool operator()(UProperty& A, UProperty& B) const
 			{
+				const int32 AOffset = GetOffsetForProperty<BuildType>(A);
+				const int32 BOffset = GetOffsetForProperty<BuildType>(B);
+
 				// Ensure stable sort
-				if (A.GetOffset_ForGC() == B.GetOffset_ForGC())
+				if (AOffset == BOffset)
 				{
 					return A.GetName() < B.GetName();
 				}
 
-				return A.GetOffset_ForGC() < B.GetOffset_ForGC();
+				return AOffset < BOffset;
 			}
 		};
 
@@ -3998,7 +4001,8 @@ int32 FRepLayout::InitFromProperty_r(
 		{
 			for (int32 j = 0; j < NetProperties[i]->ArrayDim; j++)
 			{
-				RelativeHandle = InitFromProperty_r(NetProperties[i], Offset + StructProp->GetOffset_ForGC() + j * NetProperties[i]->ElementSize, RelativeHandle, ParentIndex, StructChecksum, j, ServerConnection);
+				const int32 ArrayElementOffset = j * NetProperties[i]->ElementSize;
+				RelativeHandle = InitFromProperty_r<BuildType>(Cmds, NetProperties[i], Offset + GetOffsetForProperty<BuildType>(*StructProp) + ArrayElementOffset, RelativeHandle, ParentIndex, StructChecksum, j, ServerConnection);
 			}
 		}
 		return RelativeHandle;
@@ -4007,18 +4011,18 @@ int32 FRepLayout::InitFromProperty_r(
 	// Add actual property
 	RelativeHandle++;
 
-	AddPropertyCmd(Property, Offset + Property->GetOffset_ForGC(), RelativeHandle, ParentIndex, ParentChecksum, StaticArrayIndex, ServerConnection);
+	AddPropertyCmd(Cmds, Property, Offset + GetOffsetForProperty<BuildType>(*Property), RelativeHandle, ParentIndex, ParentChecksum, StaticArrayIndex, ServerConnection);
 
 	return RelativeHandle;
 }
 
-uint16 FRepLayout::AddParentProperty(UProperty* Property, int32 ArrayIndex)
+static FORCEINLINE uint16 AddParentProperty(TArray<FRepParentCmd>& Parents, UProperty* Property, int32 ArrayIndex)
 {
 	return Parents.Emplace(Property, ArrayIndex);
 }
 
 /** Setup some flags on our parent properties, so we can handle them properly later.*/
-static void SetupRepStructFlags(FRepParentCmd& Parent, const bool bSkipCustomDeltaCheck)
+static FORCEINLINE void SetupRepStructFlags(FRepParentCmd& Parent, const bool bSkipCustomDeltaCheck)
 {
 	if (UStructProperty* StructProperty = Cast<UStructProperty>(Parent.Property))
 	{
@@ -4036,25 +4040,11 @@ static void SetupRepStructFlags(FRepParentCmd& Parent, const bool bSkipCustomDel
 			Parent.Flags |= ERepParentFlags::IsNetSerialize;
 		}
 	}
-}
 
-enum class ERepBuildShadowOffsetsType
-{
-	Class,
-	Function,
-	Struct
-};
-
-template<ERepBuildShadowOffsetsType ShadowType>
-static const int32 GetOffsetForProperty(UProperty* Property)
-{
-	return Property->GetOffset_ForGC();
-}
-
-template<>
-const int32 GetOffsetForProperty<ERepBuildShadowOffsetsType::Function>(UProperty* Property)
-{
-	return Property->GetOffset_ForUFunction();
+	if (EnumHasAnyFlags(Parent.Property->PropertyFlags, CPF_ZeroConstructor))
+	{
+		Parent.Flags |= ERepParentFlags::IsZeroConstructible;
+	}
 }
 
 /**
@@ -4150,12 +4140,12 @@ static void BuildShadowOffsets_r(TArray<FRepLayoutCmd>::TIterator& CmdIt, int32&
 	}
 }
 
-template<ERepBuildShadowOffsetsType ShadowType>
+template<ERepBuildType ShadowType>
 static void BuildShadowOffsets(UStruct* Owner, TArray<FRepParentCmd>& Parents, TArray<FRepLayoutCmd>& Cmds, int32& ShadowOffset, ERepLayoutState& LayoutState)
 {
 	SCOPE_CYCLE_COUNTER(STAT_RepLayout_BuildShadowOffsets);
 
-	if (ShadowType == ERepBuildShadowOffsetsType::Class && !!GUsePackedShadowBuffers)
+	if (ShadowType == ERepBuildType::Class && !!GUsePackedShadowBuffers)
 	{
 		ShadowOffset = 0;
 		LayoutState = Parents.Num() > 0 ? ERepLayoutState::Normal : ERepLayoutState::Empty;
@@ -4202,7 +4192,7 @@ static void BuildShadowOffsets(UStruct* Owner, TArray<FRepParentCmd>& Parents, T
 
 				if (Parent.Property->ArrayDim > 1 || EnumHasAnyFlags(Parent.Flags, ERepParentFlags::IsStructProperty))
 				{
-					const int32 ArrayStartParentOffset = GetOffsetForProperty<ShadowType>(Parent.Property);
+					const int32 ArrayStartParentOffset = GetOffsetForProperty<ShadowType>(*Parent.Property);
 
 					ShadowOffset = Align(ShadowOffset, IndexAndAlignment.Alignment);
 
@@ -4211,7 +4201,7 @@ static void BuildShadowOffsets(UStruct* Owner, TArray<FRepParentCmd>& Parents, T
 						const FParentCmdIndexAndAlignment& NextIndexAndAlignment = IndexAndAlignmentArray[i];
 						FRepParentCmd& NextParent = Parents[NextIndexAndAlignment.Index];
 
-						NextParent.ShadowOffset = ShadowOffset + (GetOffsetForProperty<ShadowType>(NextParent.Property) - ArrayStartParentOffset);
+						NextParent.ShadowOffset = ShadowOffset + (GetOffsetForProperty<ShadowType>(*NextParent.Property) - ArrayStartParentOffset);
 
 						for (auto CmdIt = Cmds.CreateIterator() + NextParent.CmdStart; CmdIt.GetIndex() < NextParent.CmdEnd; ++CmdIt)
 						{
@@ -4248,7 +4238,7 @@ static void BuildShadowOffsets(UStruct* Owner, TArray<FRepParentCmd>& Parents, T
 
 		for (auto ParentIt = Parents.CreateIterator(); ParentIt; ++ParentIt)
 		{
-			ParentIt->ShadowOffset = GetOffsetForProperty<ShadowType>(ParentIt->Property);
+			ParentIt->ShadowOffset = GetOffsetForProperty<ShadowType>(*ParentIt->Property);
 		}
 
 		for (auto CmdIt = Cmds.CreateIterator(); CmdIt; ++CmdIt)
@@ -4288,15 +4278,16 @@ void FRepLayout::InitFromClass(UClass* InObjectClass, const UNetConnection* Serv
 
 		check(Property->PropertyFlags & CPF_Net);
 
-		const int32 ParentHandle = AddParentProperty(Property, ArrayIdx);
+		const int32 ParentHandle = AddParentProperty(Parents, Property, ArrayIdx);
 
 		check(ParentHandle == i);
 		check(Parents[i].Property->RepIndex + Parents[i].ArrayIndex == i);
 
 		Parents[ParentHandle].CmdStart = Cmds.Num();
-		RelativeHandle = InitFromProperty_r(Property, Property->ElementSize * ArrayIdx, RelativeHandle, ParentHandle, 0, ArrayIdx, ServerConnection);
+		RelativeHandle = InitFromProperty_r<ERepBuildType::Class>(Cmds, Property, Property->ElementSize * ArrayIdx, RelativeHandle, ParentHandle, 0, ArrayIdx, ServerConnection);
 		Parents[ParentHandle].CmdEnd = Cmds.Num();
 		Parents[ParentHandle].Flags |= ERepParentFlags::IsConditional;
+		Parents[ParentHandle].Offset = GetOffsetForProperty<ERepBuildType::Class>(*Property);
 
 		if (Parents[i].CmdEnd > Parents[i].CmdStart)
 		{
@@ -4320,21 +4311,21 @@ void FRepLayout::InitFromClass(UClass* InObjectClass, const UNetConnection* Serv
 
 		if (bIsObjectActor)
 		{
-		// Find Role/RemoteRole property indexes so we can swap them on the client
-		if (Property->GetFName() == NAME_Role)
-		{
-			check(RoleIndex == -1);
-			check(Parents[ParentHandle].CmdEnd == Parents[ParentHandle].CmdStart + 1);
-			RoleIndex = ParentHandle;
-		}
+		    // Find Role/RemoteRole property indexes so we can swap them on the client
+		    if (Property->GetFName() == NAME_Role)
+		    {
+			    check(RoleIndex == -1);
+			    check(Parents[ParentHandle].CmdEnd == Parents[ParentHandle].CmdStart + 1);
+			    RoleIndex = ParentHandle;
+		    }
 
-		if (Property->GetFName() == NAME_RemoteRole)
-		{
-			check(RemoteRoleIndex == -1);
-			check(Parents[ParentHandle].CmdEnd == Parents[ParentHandle].CmdStart + 1);
-			RemoteRoleIndex = ParentHandle;
-		}
-	}
+		    if (Property->GetFName() == NAME_RemoteRole)
+		    {
+			    check(RemoteRoleIndex == -1);
+			    check(Parents[ParentHandle].CmdEnd == Parents[ParentHandle].CmdStart + 1);
+			    RemoteRoleIndex = ParentHandle;
+		    }
+	    }
 	}
 
 	// Make sure it either found both, or didn't find either
@@ -4347,7 +4338,7 @@ void FRepLayout::InitFromClass(UClass* InObjectClass, const UNetConnection* Serv
 		Parents[RemoteRoleIndex].RoleSwapIndex = RoleIndex;
 	}
 	
-	AddReturnCmd();
+	AddReturnCmd(Cmds);
 
 	// Initialize lifetime props
 	// Properties that replicate for the lifetime of the channel
@@ -4405,7 +4396,7 @@ void FRepLayout::InitFromClass(UClass* InObjectClass, const UNetConnection* Serv
 		BuildHandleToCmdIndexTable_r(0, Cmds.Num() - 1, BaseHandleToCmdIndex);
 	}
 
-	BuildShadowOffsets<ERepBuildShadowOffsetsType::Class>(InObjectClass, Parents, Cmds, ShadowDataBufferSize, LayoutState);
+	BuildShadowOffsets<ERepBuildType::Class>(InObjectClass, Parents, Cmds, ShadowDataBufferSize, LayoutState);
 
 	Owner = InObjectClass;
 }
@@ -4425,23 +4416,24 @@ void FRepLayout::InitFromFunction(UFunction* InFunction, const UNetConnection* S
 	{
 		for (int32 ArrayIdx = 0; ArrayIdx < It->ArrayDim; ++ArrayIdx)
 		{
-			const int32 ParentHandle = AddParentProperty(*It, ArrayIdx);
+			const int32 ParentHandle = AddParentProperty(Parents, *It, ArrayIdx);
 			Parents[ParentHandle].CmdStart = Cmds.Num();
-			RelativeHandle = InitFromProperty_r(*It, It->ElementSize * ArrayIdx, RelativeHandle, ParentHandle, 0, ArrayIdx, ServerConnection);
+			RelativeHandle = InitFromProperty_r<ERepBuildType::Function>(Cmds, *It, It->ElementSize * ArrayIdx, RelativeHandle, ParentHandle, 0, ArrayIdx, ServerConnection);
 			Parents[ParentHandle].CmdEnd = Cmds.Num();
+			Parents[ParentHandle].Offset = GetOffsetForProperty<ERepBuildType::Function>(**It);
 
 			SetupRepStructFlags(Parents[ParentHandle], /**bSkipCustomDeltaCheck=*/true);
 		}
 	}
 
-	AddReturnCmd();
+	AddReturnCmd(Cmds);
 
 	if (!ServerConnection || EnumHasAnyFlags(Flags, ECreateRepLayoutFlags::MaySendProperties))
 	{
 		BuildHandleToCmdIndexTable_r(0, Cmds.Num() - 1, BaseHandleToCmdIndex);
 	}
 
-	BuildShadowOffsets<ERepBuildShadowOffsetsType::Function>(InFunction, Parents, Cmds, ShadowDataBufferSize, LayoutState);
+	BuildShadowOffsets<ERepBuildType::Function>(InFunction, Parents, Cmds, ShadowDataBufferSize, LayoutState);
 
 	Owner = InFunction;
 
@@ -4468,39 +4460,40 @@ void FRepLayout::InitFromStruct(UStruct* InStruct, const UNetConnection* ServerC
 			
 		for (int32 ArrayIdx = 0; ArrayIdx < It->ArrayDim; ++ArrayIdx)
 		{
-			const int32 ParentHandle = AddParentProperty(*It, ArrayIdx);
+			const int32 ParentHandle = AddParentProperty(Parents, *It, ArrayIdx);
 			Parents[ParentHandle].CmdStart = Cmds.Num();
-			RelativeHandle = InitFromProperty_r(*It, It->ElementSize * ArrayIdx, RelativeHandle, ParentHandle, 0, ArrayIdx, ServerConnection);
+			RelativeHandle = InitFromProperty_r<ERepBuildType::Struct>(Cmds, *It, It->ElementSize * ArrayIdx, RelativeHandle, ParentHandle, 0, ArrayIdx, ServerConnection);
 			Parents[ParentHandle].CmdEnd = Cmds.Num();
+			Parents[ParentHandle].Offset = GetOffsetForProperty<ERepBuildType::Struct>(**It);
 
 			SetupRepStructFlags(Parents[ParentHandle], /**bSkipCustomDeltaCheck=*/true);
 		}
 	}
 
-	AddReturnCmd();
+	AddReturnCmd(Cmds);
 
 	if (!ServerConnection || EnumHasAnyFlags(Flags, ECreateRepLayoutFlags::MaySendProperties))
 	{
 		BuildHandleToCmdIndexTable_r(0, Cmds.Num() - 1, BaseHandleToCmdIndex);
 	}
 
-	BuildShadowOffsets<ERepBuildShadowOffsetsType::Struct>(InStruct, Parents, Cmds, ShadowDataBufferSize, LayoutState);
+	BuildShadowOffsets<ERepBuildType::Struct>(InStruct, Parents, Cmds, ShadowDataBufferSize, LayoutState);
 
 	Owner = InStruct;
 }
 
 void FRepLayout::SerializeProperties_DynamicArray_r(
-	FBitArchive&						Ar, 
-	UPackageMap*						Map,
-	const int32							CmdIndex,
-	uint8*								Data,
-	bool&								bHasUnmapped,
-	const int32							ArrayDepth,
-	const FRepSerializationSharedInfo&	SharedInfo) const
+	FBitArchive& Ar, 
+	UPackageMap* Map,
+	const int32 CmdIndex,
+	FRepObjectDataBuffer Data,
+	bool& bHasUnmapped,
+	const int32 ArrayDepth,
+	const FRepSerializationSharedInfo& SharedInfo) const
 {
 	const FRepLayoutCmd& Cmd = Cmds[ CmdIndex ];
 
-	FScriptArray * Array = (FScriptArray *)Data;
+	FScriptArray* Array = (FScriptArray*)Data.Data;
 
 	uint16 OutArrayNum = Array->Num();
 	Ar << OutArrayNum;
@@ -4532,29 +4525,30 @@ void FRepLayout::SerializeProperties_DynamicArray_r(
 		// When loading, we may need to resize the array to properly fit the number of elements.
 		if (Ar.IsLoading() && OutArrayNum != Array->Num())
 		{
-			FScriptArrayHelper ArrayHelper((UArrayProperty *)Cmd.Property, Data);
+			FScriptArrayHelper ArrayHelper((UArrayProperty*)Cmd.Property, Data);
 			ArrayHelper.Resize(OutArrayNum);
 		}
 
-		Data = (uint8*)Array->GetData();
+		FRepObjectDataBuffer ArrayData(Array->GetData());
 
 		for (int32 i = 0; i < Array->Num() && !Ar.IsError(); i++)
 		{
-			SerializeProperties_r(Ar, Map, CmdIndex + 1, Cmd.EndCmd - 1, Data + i * Cmd.ElementSize, bHasUnmapped, i, ArrayDepth, SharedInfo);
+			const int32 ArrayElementOffset = i * Cmd.ElementSize;
+			SerializeProperties_r(Ar, Map, CmdIndex + 1, Cmd.EndCmd - 1, ArrayData + ArrayElementOffset, bHasUnmapped, i, ArrayDepth, SharedInfo);
 		}
 	}	
 }
 
 void FRepLayout::SerializeProperties_r(
-	FBitArchive&						Ar, 
-	UPackageMap*						Map,
-	const int32							CmdStart, 
-	const int32							CmdEnd,
-	void*								Data,
-	bool&								bHasUnmapped,
-	const int32							ArrayIndex,
-	const int32							ArrayDepth,
-	const FRepSerializationSharedInfo&	SharedInfo) const
+	FBitArchive& Ar, 
+	UPackageMap* Map,
+	const int32 CmdStart, 
+	const int32 CmdEnd,
+	FRepObjectDataBuffer Data,
+	bool& bHasUnmapped,
+	const int32 ArrayIndex,
+	const int32 ArrayDepth,
+	const FRepSerializationSharedInfo& SharedInfo) const
 {
 	for (int32 CmdIndex = CmdStart; CmdIndex < CmdEnd && !Ar.IsError(); CmdIndex++)
 	{
@@ -4564,7 +4558,7 @@ void FRepLayout::SerializeProperties_r(
 
 		if (Cmd.Type == ERepLayoutCmdType::DynamicArray)
 		{
-			SerializeProperties_DynamicArray_r(Ar, Map, CmdIndex, (uint8*)Data + Cmd.Offset, bHasUnmapped, ArrayDepth + 1, SharedInfo);
+			SerializeProperties_DynamicArray_r(Ar, Map, CmdIndex, Data + Cmd, bHasUnmapped, ArrayDepth + 1, SharedInfo);
 			CmdIndex = Cmd.EndCmd - 1;		// The -1 to handle the ++ in the for loop
 			continue;
 		}
@@ -4580,7 +4574,7 @@ void FRepLayout::SerializeProperties_r(
 
 		if ((GNetSharedSerializedData != 0) && Ar.IsSaving() && ((Cmd.Flags & ERepLayoutFlags::IsSharedSerialization) != ERepLayoutFlags::None))
 		{
-			FGuid PropertyGuid(CmdIndex, ArrayIndex, ArrayDepth, (int32)((PTRINT)((uint8*)Data + Cmd.Offset) & 0xFFFFFFFF));
+			FGuid PropertyGuid(CmdIndex, ArrayIndex, ArrayDepth, (int32)((PTRINT)(Data + Cmd) & 0xFFFFFFFF));
 
 			SharedPropInfo = SharedInfo.SharedPropertyInfo.FindByPredicate([&](const FRepSerializedPropertyInfo& Info) 
 			{ 
@@ -4600,7 +4594,7 @@ void FRepLayout::SerializeProperties_r(
 
 				FBitWriterMark BitWriterMark(Writer);
 
-				Cmd.Property->NetSerializeItem(Writer, Map, (void*)((uint8*)Data + Cmd.Offset));
+				Cmd.Property->NetSerializeItem(Writer, Map, (Data + Cmd).Data);
 
 				TArray<uint8> StandardBuffer;
 				BitWriterMark.Copy(Writer, StandardBuffer);
@@ -4625,7 +4619,7 @@ void FRepLayout::SerializeProperties_r(
 		else
 		{
 			GNumSharedSerializationMiss++;
-			if (!Cmd.Property->NetSerializeItem(Ar, Map, (void*)((uint8*)Data + Cmd.Offset)))
+			if (!Cmd.Property->NetSerializeItem(Ar, Map, (Data + Cmd).Data))
 			{
 				bHasUnmapped = true;
 			}
@@ -4641,12 +4635,12 @@ void FRepLayout::SerializeProperties_r(
 }
 
 void FRepLayout::BuildChangeList_r(
-	const TArray<FHandleToCmdIndex>&	HandleToCmdIndex,
-	const int32							CmdStart,
-	const int32							CmdEnd,
-	uint8*								Data,
-	const int32							HandleOffset,
-	TArray<uint16>&						Changed) const
+	const TArray<FHandleToCmdIndex>& HandleToCmdIndex,
+	const int32 CmdStart,
+	const int32 CmdEnd,
+	const FConstRepObjectDataBuffer Data,
+	const int32 HandleOffset,
+	TArray<uint16>& Changed) const
 {
 	for (int32 CmdIndex = CmdStart; CmdIndex < CmdEnd; CmdIndex++)
 	{
@@ -4656,28 +4650,30 @@ void FRepLayout::BuildChangeList_r(
 
 		if (Cmd.Type == ERepLayoutCmdType::DynamicArray)
 		{			
-			FScriptArray* Array = (FScriptArray *)(Data + Cmd.Offset);
+			FScriptArray* Array = (FScriptArray *)(Data + Cmd).Data;
+			const FConstRepObjectDataBuffer ArrayData = Array->GetData();
 
 			TArray<uint16> ChangedLocal;
 
 			TArray<FHandleToCmdIndex>& ArrayHandleToCmdIndex = *HandleToCmdIndex[Cmd.RelativeHandle - 1].HandleToCmdIndex;
 
-			const int32 ArrayCmdStart			= CmdIndex + 1;
-			const int32 ArrayCmdEnd				= Cmd.EndCmd - 1;
-			const int32 NumHandlesPerElement	= ArrayHandleToCmdIndex.Num();
+			const int32 ArrayCmdStart = CmdIndex + 1;
+			const int32 ArrayCmdEnd = Cmd.EndCmd - 1;
+			const int32 NumHandlesPerElement = ArrayHandleToCmdIndex.Num();
 
 			check(NumHandlesPerElement > 0);
 
 			for (int32 i = 0; i < Array->Num(); i++)
 			{
-				BuildChangeList_r(ArrayHandleToCmdIndex, ArrayCmdStart, ArrayCmdEnd, ((uint8*)Array->GetData()) + Cmd.ElementSize * i, i * NumHandlesPerElement, ChangedLocal);
+				const int32 ArrayElementOffset = Cmd.ElementSize * i;
+				BuildChangeList_r(ArrayHandleToCmdIndex, ArrayCmdStart, ArrayCmdEnd, ArrayData + ArrayElementOffset, i * NumHandlesPerElement, ChangedLocal);
 			}
 
 			if (ChangedLocal.Num())
 			{
 				Changed.Add(Cmd.RelativeHandle + HandleOffset);	// Identify the array cmd handle
-				Changed.Add(ChangedLocal.Num());					// This is so we can jump over the array if we need to
-				Changed.Append(ChangedLocal);						// Append the change list under the array
+				Changed.Add(ChangedLocal.Num());				// This is so we can jump over the array if we need to
+				Changed.Append(ChangedLocal);					// Append the change list under the array
 				Changed.Add(0);									// Null terminator
 			}
 
@@ -4691,10 +4687,10 @@ void FRepLayout::BuildChangeList_r(
 
 
 void FRepLayout::BuildSharedSerialization(
-	const uint8* RESTRICT			Data,
-	TArray<uint16>&					Changed,
-	const bool						bWriteHandle,
-	FRepSerializationSharedInfo&	SharedInfo) const
+	const FConstRepObjectDataBuffer Data,
+	TArray<uint16>& Changed,
+	const bool bWriteHandle,
+	FRepSerializationSharedInfo& SharedInfo) const
 {
 #ifdef ENABLE_PROPERTY_CHECKSUMS
 	const bool bDoChecksum = (GDoPropertyChecksum == 1);
@@ -4705,59 +4701,59 @@ void FRepLayout::BuildSharedSerialization(
 	FChangelistIterator ChangelistIterator(Changed, 0);
 	FRepHandleIterator HandleIterator(ChangelistIterator, Cmds, BaseHandleToCmdIndex, 0, 1, 0, Cmds.Num() - 1);
 
-	BuildSharedSerialization_r(HandleIterator, (uint8*)Data, bWriteHandle, bDoChecksum, 0, SharedInfo);
+	BuildSharedSerialization_r(HandleIterator, Data, bWriteHandle, bDoChecksum, 0, SharedInfo);
 
 	SharedInfo.SetValid();
 }
 
 void FRepLayout::BuildSharedSerialization_r(
-	FRepHandleIterator&				HandleIterator,
-	const uint8* RESTRICT			SourceData,
-	const bool						bWriteHandle,
-	const bool						bDoChecksum,
-	const int32						ArrayDepth,
-	FRepSerializationSharedInfo&	SharedInfo) const
+	FRepHandleIterator& HandleIterator,
+	const FConstRepObjectDataBuffer SourceData,
+	const bool bWriteHandle,
+	const bool bDoChecksum,
+	const int32 ArrayDepth,
+	FRepSerializationSharedInfo& SharedInfo) const
 {
 	while (HandleIterator.NextHandle())
 	{
-		const int32 CmdIndex	= HandleIterator.CmdIndex;
+		const int32 CmdIndex = HandleIterator.CmdIndex;
 		const int32 ArrayOffset = HandleIterator.ArrayOffset;
 
 		const FRepLayoutCmd& Cmd = Cmds[CmdIndex];
 		const FRepParentCmd& ParentCmd = Parents[Cmd.ParentIndex];
 
-		const uint8* Data = SourceData + ArrayOffset + Cmd.Offset;
+		const FConstRepObjectDataBuffer Data = SourceData + ArrayOffset + Cmd;
 
 		if (Cmd.Type == ERepLayoutCmdType::DynamicArray)
 		{
-			const FScriptArray* Array = (FScriptArray *)Data;
-			const uint8* NewData = (uint8*)Array->GetData();
+			const FScriptArray* Array = (FScriptArray *)Data.Data;
+			const FConstRepObjectDataBuffer ArrayData = (uint8*)Array->GetData();
 
 			FScopedIteratorArrayTracker ArrayTracker(&HandleIterator);
 
 			TArray<FHandleToCmdIndex>& ArrayHandleToCmdIndex = *HandleIterator.HandleToCmdIndex[Cmd.RelativeHandle - 1].HandleToCmdIndex;
 
 			FRepHandleIterator ArrayIterator(HandleIterator.ChangelistIterator, Cmds, ArrayHandleToCmdIndex, Cmd.ElementSize, Array->Num(), CmdIndex + 1, Cmd.EndCmd - 1);
-			BuildSharedSerialization_r(ArrayIterator, NewData, bWriteHandle, bDoChecksum, ArrayDepth + 1, SharedInfo);
+			BuildSharedSerialization_r(ArrayIterator, ArrayData, bWriteHandle, bDoChecksum, ArrayDepth + 1, SharedInfo);
 			continue;
 		}
 
-		if ((Cmd.Flags & ERepLayoutFlags::IsSharedSerialization) != ERepLayoutFlags::None)
+		if (EnumHasAnyFlags(Cmd.Flags, ERepLayoutFlags::IsSharedSerialization))
 		{
-			SharedInfo.WriteSharedProperty(Cmd, FGuid(HandleIterator.CmdIndex, HandleIterator.ArrayIndex, ArrayDepth, (int32)((PTRINT)Data & 0xFFFFFFFF)), HandleIterator.CmdIndex, HandleIterator.Handle, Data, bWriteHandle, bDoChecksum);
+			SharedInfo.WriteSharedProperty(Cmd, FGuid(HandleIterator.CmdIndex, HandleIterator.ArrayIndex, ArrayDepth, (int32)((PTRINT)Data.Data & 0xFFFFFFFF)), HandleIterator.CmdIndex, HandleIterator.Handle, Data.Data, bWriteHandle, bDoChecksum);
 		}
 	}
 }
 
 void FRepLayout::BuildSharedSerializationForRPC_DynamicArray_r(
-	const int32						CmdIndex,
-	uint8*							Data,
-	int32							ArrayDepth,
-	FRepSerializationSharedInfo&	SharedInfo)
+	const int32 CmdIndex,
+	const FConstRepObjectDataBuffer Data,
+	int32 ArrayDepth,
+	FRepSerializationSharedInfo& SharedInfo)
 {
 	const FRepLayoutCmd& Cmd = Cmds[ CmdIndex ];
 
-	FScriptArray * Array = (FScriptArray *)Data;	
+	FScriptArray* Array = (FScriptArray *)Data.Data;	
 	const int32 ArrayNum = Array->Num();
 
 	// Validate the maximum number of elements.
@@ -4771,21 +4767,22 @@ void FRepLayout::BuildSharedSerializationForRPC_DynamicArray_r(
 		return;
 	}
 
-	Data = (uint8*)Array->GetData();
+	const FConstRepObjectDataBuffer ArrayData(Array->GetData());
 
 	for (int32 i = 0; i < ArrayNum; i++)
 	{
-		BuildSharedSerializationForRPC_r(CmdIndex + 1, Cmd.EndCmd - 1, Data + i * Cmd.ElementSize, i, ArrayDepth, SharedInfo);
+		const int32 ArrayElementOffset = i * Cmd.ElementSize;
+		BuildSharedSerializationForRPC_r(CmdIndex + 1, Cmd.EndCmd - 1, ArrayData + ArrayElementOffset, i, ArrayDepth, SharedInfo);
 	}
 }
 
 void FRepLayout::BuildSharedSerializationForRPC_r(
-	const int32						CmdStart,
-	const int32						CmdEnd,
-	void*							Data,
-	int32							ArrayIndex,
-	int32							ArrayDepth,
-	FRepSerializationSharedInfo&	SharedInfo)
+	const int32 CmdStart,
+	const int32 CmdEnd,
+	const FConstRepObjectDataBuffer Data,
+	int32 ArrayIndex,
+	int32 ArrayDepth,
+	FRepSerializationSharedInfo& SharedInfo)
 {
 	for (int32 CmdIndex = CmdStart; CmdIndex < CmdEnd; CmdIndex++)
 	{
@@ -4795,21 +4792,21 @@ void FRepLayout::BuildSharedSerializationForRPC_r(
 
 		if (Cmd.Type == ERepLayoutCmdType::DynamicArray)
 		{
-			BuildSharedSerializationForRPC_DynamicArray_r(CmdIndex, (uint8*)Data + Cmd.Offset, ArrayDepth + 1, SharedInfo);
+			BuildSharedSerializationForRPC_DynamicArray_r(CmdIndex, Data + Cmd, ArrayDepth + 1, SharedInfo);
 			CmdIndex = Cmd.EndCmd - 1;		// The -1 to handle the ++ in the for loop
 			continue;
 		}
 
 		if (!Parents[Cmd.ParentIndex].Property->HasAnyPropertyFlags(CPF_OutParm) && ((Cmd.Flags & ERepLayoutFlags::IsSharedSerialization) != ERepLayoutFlags::None))
 		{
-			FGuid PropertyGuid(CmdIndex, ArrayIndex, ArrayDepth, (int32)((PTRINT)((uint8*)Data + Cmd.Offset) & 0xFFFFFFFF));
+			FGuid PropertyGuid(CmdIndex, ArrayIndex, ArrayDepth, (int32)((PTRINT)(Data + Cmd) & 0xFFFFFFFF));
 
-			SharedInfo.WriteSharedProperty(Cmd, PropertyGuid, CmdIndex, 0, (uint8*)Data + Cmd.Offset, false, false);
+			SharedInfo.WriteSharedProperty(Cmd, PropertyGuid, CmdIndex, 0, (Data + Cmd).Data, false, false);
 		}
 	}
 }
 
-void FRepLayout::BuildSharedSerializationForRPC(void* Data)
+void FRepLayout::BuildSharedSerializationForRPC(const FConstRepObjectDataBuffer Data)
 {
 	if ((GNetSharedSerializedData != 0) && !SharedInfoRPC.IsValid())
 	{
@@ -4829,7 +4826,7 @@ void FRepLayout::BuildSharedSerializationForRPC(void* Data)
 				// check for a complete match, including arrays
 				// (we're comparing against zero data here, since 
 				// that's the default.)
-				bSend = !Parents[i].Property->Identical_InContainer(Data, nullptr, Parents[i].ArrayIndex);
+				bSend = !Parents[i].Property->Identical_InContainer(Data.Data, nullptr, Parents[i].ArrayIndex);
 			}
 
 			if (bSend)
@@ -4852,10 +4849,10 @@ void FRepLayout::ClearSharedSerializationForRPC()
 }
 
 void FRepLayout::SendPropertiesForRPC(
-	UFunction*		Function,
-	UActorChannel*	Channel,
-	FNetBitWriter&	Writer,
-	void*			Data) const
+	UFunction* Function,
+	UActorChannel* Channel,
+	FNetBitWriter& Writer,
+	const FConstRepObjectDataBuffer Data) const
 {
 	check(Function == Owner);
 
@@ -4867,15 +4864,15 @@ void FRepLayout::SendPropertiesForRPC(
 
 			for (int32 i = 0; i < Parents.Num(); i++)
 			{
-				if (!Parents[i].Property->Identical_InContainer(Data, NULL, Parents[i].ArrayIndex))
+				if (!Parents[i].Property->Identical_InContainer(Data.Data, NULL, Parents[i].ArrayIndex))
 				{
-					BuildChangeList_r(BaseHandleToCmdIndex, Parents[i].CmdStart, Parents[i].CmdEnd, (uint8*)Data, 0, Changed);
+					BuildChangeList_r(BaseHandleToCmdIndex, Parents[i].CmdStart, Parents[i].CmdEnd, Data, 0, Changed);
 				}
 			}
 
 			Changed.Add(0); // Null terminator
 
-			SendProperties_BackwardsCompatible(nullptr, nullptr, (uint8*)Data, Channel->Connection, Writer, Changed);
+			SendProperties_BackwardsCompatible(nullptr, nullptr, Data, Channel->Connection, Writer, Changed);
 		}
 		else
 		{
@@ -4895,7 +4892,7 @@ void FRepLayout::SendPropertiesForRPC(
 						// check for a complete match, including arrays
 						// (we're comparing against zero data here, since 
 						// that's the default.)
-						Send = !Parents[i].Property->Identical_InContainer(Data, NULL, Parents[i].ArrayIndex);
+						Send = !Parents[i].Property->Identical_InContainer(Data.Data, NULL, Parents[i].ArrayIndex);
 					}
 
 					Writer.WriteBit(Send ? 1 : 0);
@@ -4904,7 +4901,7 @@ void FRepLayout::SendPropertiesForRPC(
 				if (Send)
 				{
 					bool bHasUnmapped = false;
-					SerializeProperties_r(Writer, Writer.PackageMap, Parents[i].CmdStart, Parents[i].CmdEnd, Data, bHasUnmapped, 0, 0, SharedInfoRPC);
+					SerializeProperties_r(Writer, Writer.PackageMap, Parents[i].CmdStart, Parents[i].CmdEnd, const_cast<uint8*>(Data.Data), bHasUnmapped, 0, 0, SharedInfoRPC);
 				}
 			}
 		}	
@@ -4912,21 +4909,21 @@ void FRepLayout::SendPropertiesForRPC(
 }
 
 void FRepLayout::ReceivePropertiesForRPC(
-	UObject*			Object,
-	UFunction*			Function,
-	UActorChannel*		Channel,
-	FNetBitReader&		Reader,
-	void*				Data,
+	UObject* Object,
+	UFunction* Function,
+	UActorChannel* Channel,
+	FNetBitReader& Reader,
+	FRepObjectDataBuffer Data,
 	TSet<FNetworkGUID>&	UnmappedGuids) const
 {
 	check(Function == Owner);
 
 	for (int32 i = 0; i < Parents.Num(); i++)
 	{
-		if (Parents[i].ArrayIndex == 0 && (Parents[i].Property->PropertyFlags & CPF_ZeroConstructor) == 0)
+		if (Parents[i].ArrayIndex == 0 && !EnumHasAnyFlags(Parents[i].Flags, ERepParentFlags::IsZeroConstructible))
 		{
 			// If this property needs to be constructed, make sure we do that
-			Parents[i].Property->InitializeValue((uint8*)Data + Parents[i].Property->GetOffset_ForUFunction());
+			Parents[i].Property->InitializeValue(Data + Parents[i]);
 		}
 	}
 
@@ -4995,11 +4992,11 @@ void FRepLayout::ReceivePropertiesForRPC(
 }
 
 void FRepLayout::SerializePropertiesForStruct(
-	UStruct*		Struct,
-	FBitArchive&	Ar,
-	UPackageMap*	Map,
-	void*			Data,
-	bool&			bHasUnmapped) const
+	UStruct* Struct,
+	FBitArchive& Ar,
+	UPackageMap* Map,
+	FRepObjectDataBuffer Data,
+	bool& bHasUnmapped) const
 {
 	check(Struct == Owner);
 
@@ -5017,9 +5014,9 @@ void FRepLayout::SerializePropertiesForStruct(
 }
 
 void FRepLayout::BuildHandleToCmdIndexTable_r(
-	const int32					CmdStart,
-	const int32					CmdEnd,
-	TArray<FHandleToCmdIndex>&	HandleToCmdIndex)
+	const int32 CmdStart,
+	const int32 CmdEnd,
+	TArray<FHandleToCmdIndex>& HandleToCmdIndex)
 {
 	for (int32 CmdIndex = CmdStart; CmdIndex < CmdEnd; CmdIndex++)
 	{
@@ -5091,7 +5088,6 @@ void FRepLayout::RebuildConditionalProperties(
 	RepState->RepFlags = RepFlags;
 }
 
-
 void FRepLayout::InitChangedTracker(FRepChangedPropertyTracker* ChangedTracker) const
 {
 	ChangedTracker->Parents.SetNum(Parents.Num());
@@ -5102,7 +5098,7 @@ void FRepLayout::InitChangedTracker(FRepChangedPropertyTracker* ChangedTracker)
 	}
 }
 
-FRepStateStaticBuffer FRepLayout::CreateShadowBuffer(const uint8* const Source) const
+FRepStateStaticBuffer FRepLayout::CreateShadowBuffer(const FConstRepObjectDataBuffer Source) const
 {
 	FRepStateStaticBuffer ShadowData(AsShared());
 
@@ -5124,7 +5120,7 @@ TSharedPtr<FReplicationChangelistMgr> FRepLayout::CreateReplicationChangelistMgr
 }
 
 TUniquePtr<FRepState> FRepLayout::CreateRepState(
-	const uint8* const Source,
+	const FConstRepObjectDataBuffer Source,
 	TSharedPtr<FRepChangedPropertyTracker>& InRepChangedPropertyTracker,
 	ECreateRepStateFlags Flags) const
 {
@@ -5165,7 +5161,7 @@ TUniquePtr<FRepState> FRepLayout::CreateRepState(
 	return RepState;
 }
 
-void FRepLayout::InitRepStateStaticBuffer(FRepStateStaticBuffer& ShadowData, const uint8* Source) const
+void FRepLayout::InitRepStateStaticBuffer(FRepStateStaticBuffer& ShadowData, const FConstRepObjectDataBuffer Source) const
 {
 	check(ShadowData.Buffer.Num() == 0);
 	ShadowData.Buffer.SetNumZeroed(ShadowDataBufferSize);
@@ -5175,7 +5171,7 @@ void FRepLayout::InitRepStateStaticBuffer(FRepStateStaticBuffer& ShadowData, con
 
 void FRepLayout::ConstructProperties(FRepStateStaticBuffer& InShadowData) const
 {
-	uint8* ShadowData = InShadowData.GetData();
+	FRepShadowDataBuffer ShadowData = InShadowData.GetData();
 
 	// Construct all items
 	for (const FRepParentCmd& Parent : Parents)
@@ -5184,14 +5180,14 @@ void FRepLayout::ConstructProperties(FRepStateStaticBuffer& InShadowData) const
 		if (Parent.ArrayIndex == 0)
 		{
 			check((Parent.ShadowOffset + Parent.Property->GetSize()) <= InShadowData.Num());
-			Parent.Property->InitializeValue(ShadowData + Parent.ShadowOffset);
+			Parent.Property->InitializeValue(ShadowData + Parent);
 		}
 	}
 }
 
-void FRepLayout::CopyProperties(FRepStateStaticBuffer& InShadowData, const uint8* const Source) const
+void FRepLayout::CopyProperties(FRepStateStaticBuffer& InShadowData, const FConstRepObjectDataBuffer Source) const
 {
-	uint8* ShadowData = InShadowData.GetData();
+	FRepShadowDataBuffer ShadowData = InShadowData.GetData();
 
 	// Init all items
 	for (const FRepParentCmd& Parent : Parents)
@@ -5200,14 +5196,14 @@ void FRepLayout::CopyProperties(FRepStateStaticBuffer& InShadowData, const uint8
 		if (Parent.ArrayIndex == 0)
 		{
 			check((Parent.ShadowOffset + Parent.Property->GetSize()) <= InShadowData.Num());
-			Parent.Property->CopyCompleteValue(ShadowData + Parent.ShadowOffset, Parent.Property->ContainerPtrToValuePtr<uint8>(Source));
+			Parent.Property->CopyCompleteValue(ShadowData + Parent, Source + Parent);
 		}
 	}
 }
 
 void FRepLayout::DestructProperties(FRepStateStaticBuffer& InShadowData) const
 {
-	uint8* ShadowData = InShadowData.GetData();
+	FRepShadowDataBuffer ShadowData = InShadowData.GetData();
 
 	// Destruct all items
 	for (const FRepParentCmd& Parent : Parents)
@@ -5216,7 +5212,7 @@ void FRepLayout::DestructProperties(FRepStateStaticBuffer& InShadowData) const
 		if (Parent.ArrayIndex == 0)
 		{
 			check((Parent.ShadowOffset + Parent.Property->GetSize()) <= InShadowData.Num());
-			Parent.Property->DestroyValue(ShadowData + Parent.ShadowOffset);
+			Parent.Property->DestroyValue(ShadowData + Parent);
 		}
 	}
 
diff --git a/Engine/Source/Runtime/Engine/Public/Net/RepLayout.h b/Engine/Source/Runtime/Engine/Public/Net/RepLayout.h
--- a/Engine/Source/Runtime/Engine/Public/Net/RepLayout.h
+++ b/Engine/Source/Runtime/Engine/Public/Net/RepLayout.h
@@ -46,8 +46,84 @@ enum class EReceivePropertiesFlags : uint32
 
 ENUM_CLASS_FLAGS(EReceivePropertiesFlags);
 
-/** Builds a Bitfield that flags whether or not given replication conditions are met. */
-extern TStaticBitArray<COND_Max> BuildConditionMapFromRepFlags(FReplicationFlags RepFlags);
+enum class ERepDataBufferType
+{
+	ObjectBuffer,	//! Indicates this buffer is a full object's memory.
+	ShadowBuffer	//! Indicates this buffer is a packed shadow buffer.
+};
+
+namespace UE4_RepLayout_Private
+{
+	/**
+	 * TRepDataBuffer and TConstRepDataBuffer act as wrapper around internal data
+	 * buffers that FRepLayout may use. This allows FRepLayout to properly interact
+	 * with memory buffers and apply commands to them more easily.
+	 */
+	template<ERepDataBufferType DataType, typename ConstOrNotType>
+	struct TRepDataBufferBase
+	{
+		static constexpr ERepDataBufferType Type = DataType;
+
+	private:
+
+		// To be consistent, we need to match the constness of the base type
+		// to the constness of the void* we use as parameters or as returns.
+		// This bit of code does that. Working inside out:
+		//		1. We remove the const / volatile qualifications from the ConstOrNotType.
+		//		2. We see if the CV-stripped type is the same as the ConstOrNotType.
+		//		3. If they are the same, then we know the input is not const, and so we'll use void*.
+		//			If they aren't the same, assume that ConstOrNotType is const, and use const void*.
+		typedef typename TChooseClass<TAreTypesEqual<ConstOrNotType, typename TRemoveCV<ConstOrNotType>::Type>::Value, void, void const>::Result ConstOrNotVoid;
+
+	public:
+
+		TRepDataBufferBase(ConstOrNotVoid* RESTRICT InDataBuffer) :
+			Data((ConstOrNotType* RESTRICT)InDataBuffer)
+		{}
+
+		friend TRepDataBufferBase operator+(TRepDataBufferBase InBuffer, int32 Offset)
+		{
+			return InBuffer.Data + Offset;
+		}
+
+		operator bool()
+		{
+			return Data != nullptr;
+		}
+
+		operator ConstOrNotType* ()
+		{
+			return Data;
+		}
+
+		operator ConstOrNotVoid* ()
+		{
+			return (ConstOrNotVoid*)Data;
+		}
+
+		ConstOrNotType* RESTRICT Data;
+	};
+
+	template<typename TLayoutCmdType, typename ConstOrNotType>
+	static TRepDataBufferBase<ERepDataBufferType::ObjectBuffer, ConstOrNotType> operator+(TRepDataBufferBase<ERepDataBufferType::ObjectBuffer, ConstOrNotType> InBuffer, const TLayoutCmdType& Cmd)
+	{
+		return InBuffer + Cmd.Offset;
+	}
+
+	template<typename TLayoutCmdType, typename ConstOrNotType>
+	static TRepDataBufferBase<ERepDataBufferType::ShadowBuffer, ConstOrNotType> operator+(const TRepDataBufferBase<ERepDataBufferType::ShadowBuffer, ConstOrNotType> InBuffer, const TLayoutCmdType& Cmd)
+	{
+		return InBuffer + Cmd.ShadowOffset;
+	}
+}
+
+template<ERepDataBufferType DataType> using TRepDataBuffer = UE4_RepLayout_Private::TRepDataBufferBase<DataType, uint8>;
+template<ERepDataBufferType DataType> using TConstRepDataBuffer = UE4_RepLayout_Private::TRepDataBufferBase<DataType, const uint8>;
+
+typedef TRepDataBuffer<ERepDataBufferType::ObjectBuffer> FRepObjectDataBuffer;
+typedef TRepDataBuffer<ERepDataBufferType::ShadowBuffer> FRepShadowDataBuffer;
+typedef TConstRepDataBuffer<ERepDataBufferType::ObjectBuffer> FConstRepObjectDataBuffer;
+typedef TConstRepDataBuffer<ERepDataBufferType::ShadowBuffer> FConstRepShadowDataBuffer;
 
 /** Stores meta data about a given Replicated property. */
 class FRepChangedParent
@@ -227,13 +303,13 @@ struct FRepSerializationSharedInfo
 	 * @param bDoChecksum		Whether or not we should do checksums. Only used if ENABLE_PROPERTY_CHECKSUMS is enabled.
 	 */
 	const FRepSerializedPropertyInfo* WriteSharedProperty(
-		const FRepLayoutCmd&	Cmd,
-		const FGuid&			PropertyGuid,
-		const int32				CmdIndex,
-		const uint16			Handle,
-		const uint8* RESTRICT	Data,
-		const bool				bWriteHandle,
-		const bool				bDoChecksum);
+		const FRepLayoutCmd& Cmd,
+		const FGuid& PropertyGuid,
+		const int32 CmdIndex,
+		const uint16 Handle,
+		const FConstRepObjectDataBuffer Data,
+		const bool bWriteHandle,
+		const bool bDoChecksum);
 
 	/** Metadata for properties in the shared data blob. */
 	TArray<FRepSerializedPropertyInfo> SharedPropertyInfo;
@@ -602,7 +678,6 @@ public:
 	/** The unique list of properties that have changed since the channel was first opened */
 	TArray<uint16> LifetimeChangelist;
 
-
 	/**
 	 * Properties which are inactive through conditions have their changes stored here, so they can be 
 	 * applied if/when the property becomes active.
@@ -703,7 +778,8 @@ enum class ERepParentFlags : uint32
 	IsConfig			= (1 << 2),	//! This property is defaulted from a config file
 	IsCustomDelta		= (1 << 3),	//! This property uses custom delta compression. Mutually exclusive with IsNetSerialize.
 	IsNetSerialize		= (1 << 4), //! This property uses a custom net serializer. Mutually exclusive with IsCustomDelta.
-	IsStructProperty	= (1 << 5)	//! This property is a UStructProperty.
+	IsStructProperty	= (1 << 5),	//! This property is a UStructProperty.
+	IsZeroConstructible	= (1 << 6)	//! This property is ZeroConstructible.
 };
 
 ENUM_CLASS_FLAGS(ERepParentFlags)
@@ -743,6 +819,9 @@ public:
 	 */
 	int32 ArrayIndex;
 
+	/** Absolute offset of property in Object Memory. */
+	int32 Offset;
+
 	/** Absolute offset of property in Shadow Memory. */
 	int32 ShadowOffset;
 
@@ -766,7 +845,7 @@ public:
 	int32 RoleSwapIndex;
 
 	ELifetimeCondition Condition;
-	ELifetimeRepNotifyCondition	RepNotifyCondition;
+	ELifetimeRepNotifyCondition RepNotifyCondition;
 
 	/**
 	 * Number of parameters that we need to pass to the RepNotify function (if any).
@@ -973,86 +1052,6 @@ public:
 	int32 ArrayOffset;
 };
 
-
-enum class ERepDataBufferType
-{
-	ObjectBuffer,	//! Indicates this buffer is a full object's memory.
-	ShadowBuffer	//! Indicates this buffer is a packed shadow buffer.
-};
-
-namespace UE4_RepLayout_Private
-{
-	/**
-	 * TRepDataBuffer and TConstRepDataBuffer act as wrapper around internal data
-	 * buffers that FRepLayout may use. This allows FRepLayout to properly interact
-	 * with memory buffers and apply commands to them more easily.
-	 */
-	template<ERepDataBufferType DataType, typename ConstOrNotType>
-	struct TRepDataBufferBase
-	{
-		static constexpr ERepDataBufferType Type = DataType;
-
-	private:
-
-		// To be consistent, we need to match the constness of the base type
-		// to the constness of the void* we use as parameters or as returns.
-		// This bit of code does that. Working inside out:
-		//		1. We remove the const / volatile qualifications from the ConstOrNotType.
-		//		2. We see if the CV-stripped type is the same as the ConstOrNotType.
-		//		3. If they are the same, then we know the input is not const, and so we'll use void*.
-		//			If they aren't the same, assume that ConstOrNotType is const, and use const void*.
-		typedef typename TChooseClass<TAreTypesEqual<ConstOrNotType, typename TRemoveCV<ConstOrNotType>::Type>::Value, void, void const>::Result ConstOrNotVoid;
-
-	public:
-
-		TRepDataBufferBase(ConstOrNotVoid* RESTRICT InDataBuffer) :
-			Data((ConstOrNotType* RESTRICT)InDataBuffer)
-		{}
-
-		friend TRepDataBufferBase operator+(TRepDataBufferBase InBuffer, int32 Offset)
-		{
-			return InBuffer.Data + Offset;
-		}
-
-		operator bool()
-		{
-			return Data != nullptr;
-		}
-
-		operator ConstOrNotType* ()
-		{
-			return Data;
-		}
-
-		operator ConstOrNotVoid* ()
-		{
-			return (ConstOrNotVoid*)Data;
-		}
-
-		ConstOrNotType* RESTRICT Data;
-	};
-
-	template<typename TLayoutCmdType, typename ConstOrNotType>
-	static TRepDataBufferBase<ERepDataBufferType::ObjectBuffer, ConstOrNotType> operator+(TRepDataBufferBase<ERepDataBufferType::ObjectBuffer, ConstOrNotType> InBuffer, const TLayoutCmdType& Cmd)
-	{
-		return InBuffer + Cmd.Offset;
-	}
-
-	template<typename TLayoutCmdType, typename ConstOrNotType>
-	static TRepDataBufferBase<ERepDataBufferType::ShadowBuffer, ConstOrNotType> operator+(const TRepDataBufferBase<ERepDataBufferType::ShadowBuffer, ConstOrNotType> InBuffer, const TLayoutCmdType& Cmd)
-	{
-		return InBuffer + Cmd.ShadowOffset;
-	}
-}
-
-template<ERepDataBufferType DataType> using TRepDataBuffer = UE4_RepLayout_Private::TRepDataBufferBase<DataType, uint8>;
-template<ERepDataBufferType DataType> using TConstRepDataBuffer = UE4_RepLayout_Private::TRepDataBufferBase<DataType, const uint8>;
-
-typedef TRepDataBuffer<ERepDataBufferType::ObjectBuffer> FRepObjectDataBuffer;
-typedef TRepDataBuffer<ERepDataBufferType::ShadowBuffer> FRepShadowDataBuffer;
-typedef TConstRepDataBuffer<ERepDataBufferType::ObjectBuffer> FConstRepObjectDataBuffer;
-typedef TConstRepDataBuffer<ERepDataBufferType::ShadowBuffer> FConstRepShadowDataBuffer;
-
 enum class ECreateRepLayoutFlags
 {
 	None,
@@ -1189,7 +1188,7 @@ public:
 	 *
 	 * @param Source	Memory buffer storing object property data.
 	 */
-	FRepStateStaticBuffer CreateShadowBuffer(const uint8* const Source) const;
+	FRepStateStaticBuffer CreateShadowBuffer(const FConstRepObjectDataBuffer Source) const;
 
 	/**
 	 * Creates and initializes a new FReplicationChangelistMgr.
@@ -1215,7 +1214,7 @@ public:
 	 *			Note, maybe a a FRepStateBase or FRepStateSending based on parameters.
 	 */
 	TUniquePtr<FRepState> CreateRepState(
-		const uint8* const Source,
+		const FConstRepObjectDataBuffer Source,
 		TSharedPtr<FRepChangedPropertyTracker>& InRepChangedPropertyTracker,
 		ECreateRepStateFlags Flags) const;
 
@@ -1239,7 +1238,7 @@ public:
 	bool ReplicateProperties(
 		FSendingRepState* RESTRICT RepState,
 		FRepChangelistState* RESTRICT RepChangelistState,
-		const uint8* RESTRICT Data,
+		const FConstRepObjectDataBuffer Data,
 		UClass* ObjectClass,
 		UActorChannel* OwningChannel,
 		FNetBitWriter& Writer,
@@ -1289,7 +1288,7 @@ public:
 		UActorChannel* OwningChannel,
 		UClass* InObjectClass,
 		FReceivingRepState* RESTRICT RepState,
-		void* RESTRICT Data,
+		FRepObjectDataBuffer Data,
 		FNetBitReader& InBunch,
 		bool& bOutHasUnmapped,
 		bool& bOutGuidsChanged,
@@ -1418,9 +1417,9 @@ public:
 		PostReplicate(RepState->GetSendingRepState(), PacketRange, bReliable);
 	}
 
-	void ReceivedNak(FRepState * RepState, int32 NakPacketId) const;
-	bool AllAcked(FRepState * RepState) const;
-	bool ReadyForDormancy(FRepState * RepState) const;
+	void ReceivedNak(FRepState* RepState, int32 NakPacketId) const;
+	bool AllAcked(FRepState* RepState) const;
+	bool ReadyForDormancy(FRepState* RepState) const;
 
 	template<ERepDataBufferType DataType>
 	void ValidateWithChecksum(TConstRepDataBuffer<DataType> Data, FBitArchive & Ar) const;
@@ -1437,20 +1436,15 @@ public:
 		PruneChangeList(Data, Changed, PrunedChanged);
 	}
 
-	/**
-	 * Combines two changelists, ensuring that handles are in the correct order, and arrays are properly structured.
-	 *
-	 * @param Data			Property Data for the Object / RepState the changelists refer to.
-	 * @param Dirty1		First changelist to merge. Must be non-empty, valid changelist.
-	 * @param Dirty2		Second changelist to merge. May be empty, or otherwise valid changelist.
-	 * @param MergedDirty	The combined changelist.
-	 */
 	UE_DEPRECATED(4.23, "This method will be made private in future versions.")
 	void MergeChangeList(
 		const uint8* RESTRICT Data,
 		const TArray<uint16>& Dirty1,
 		const TArray<uint16>& Dirty2,
-		TArray<uint16>& MergedDirty) const;
+		TArray<uint16>& MergedDirty) const
+	{
+		MergeChangeList(FConstRepObjectDataBuffer(Data), Dirty1, Dirty2, MergedDirty);
+	}
 
 	/**
 	 * Compare all properties between source and destination buffer, and optionally update the destination
@@ -1498,33 +1492,33 @@ public:
 
 	/** @see SendProperties. */
 	void ENGINE_API SendPropertiesForRPC(
-		UFunction*		Function,
-		UActorChannel*	Channel,
-		FNetBitWriter&	Writer,
-		void*			Data) const;
+		UFunction* Function,
+		UActorChannel* Channel,
+		FNetBitWriter& Writer,
+		const FConstRepObjectDataBuffer Data) const;
 
 	/** @see ReceiveProperties. */
 	void ReceivePropertiesForRPC(
-		UObject*			Object,
-		UFunction*			Function,
-		UActorChannel*		Channel,
-		FNetBitReader&		Reader,
-		void*				Data,
-		TSet<FNetworkGUID>&	UnmappedGuids) const;
+		UObject* Object,
+		UFunction* Function,
+		UActorChannel* Channel,
+		FNetBitReader& Reader,
+		FRepObjectDataBuffer Data,
+		TSet<FNetworkGUID>& UnmappedGuids) const;
 
 	/** Builds shared serialization state for a multicast rpc */
-	void ENGINE_API BuildSharedSerializationForRPC(void* Data);
+	void ENGINE_API BuildSharedSerializationForRPC(const FConstRepObjectDataBuffer Data);
 
 	/** Clears shared serialization state for a multicast rpc */
 	void ENGINE_API ClearSharedSerializationForRPC();
 
 	// Struct support
 	ENGINE_API void SerializePropertiesForStruct(
-		UStruct*		Struct,
-		FBitArchive&	Ar,
-		UPackageMap*	Map,
-		void*			Data,
-		bool&			bHasUnmapped) const;
+		UStruct* Struct,
+		FBitArchive& Ar,
+		UPackageMap* Map,
+		FRepObjectDataBuffer Data,
+		bool& bHasUnmapped) const;
 
 	/** Serializes all replicated properties of a UObject in or out of an archive (depending on what type of archive it is). */
 	ENGINE_API void SerializeObjectReplicatedProperties(UObject* Object, FBitArchive & Ar) const;
@@ -1535,7 +1529,7 @@ public:
 	void SendProperties_BackwardsCompatible(
 		FSendingRepState* RESTRICT RepState,
 		FRepChangedPropertyTracker* ChangedTracker,
-		const uint8* RESTRICT Data,
+		const FConstRepObjectDataBuffer Data,
 		UNetConnection* Connection,
 		FNetBitWriter& Writer,
 		TArray<uint16>& Changed) const;
@@ -1544,7 +1538,7 @@ public:
 	bool ReceiveProperties_BackwardsCompatible(
 		UNetConnection* Connection,
 		FReceivingRepState* RESTRICT RepState,
-		void* RESTRICT Data,
+		FRepObjectDataBuffer Data,
 		FNetBitReader& InBunch,
 		bool& bOutHasUnmapped,
 		const bool bEnableRepNotifies,
@@ -1622,7 +1616,7 @@ private:
 	bool CompareProperties(
 		FSendingRepState* RESTRICT RepState,
 		FRepChangelistState* RESTRICT RepChangelistState,
-		const uint8* RESTRICT Data,
+		const FConstRepObjectDataBuffer Data,
 		const FReplicationFlags& RepFlags) const;
 
 	/**
@@ -1644,7 +1638,7 @@ private:
 	void SendProperties(
 		FSendingRepState* RESTRICT RepState,
 		FRepChangedPropertyTracker* ChangedTracker,
-		const uint8* RESTRICT Data,
+		const FConstRepObjectDataBuffer Data,
 		UClass* ObjectClass,
 		FNetBitWriter& Writer,
 		TArray<uint16>& Changed,
@@ -1658,10 +1652,16 @@ private:
 	 * @param PrunedChanged		The resulting pruned changelist.
 	 */
 	void PruneChangeList(
-		const void* RESTRICT Data,
+		const FConstRepObjectDataBuffer Data,
 		const TArray<uint16>& Changed,
 		TArray<uint16>& PrunedChanged) const;
 
+	void MergeChangeList(
+		const FConstRepObjectDataBuffer Data,
+		const TArray<uint16>& Dirty1,
+		const TArray<uint16>& Dirty2,
+		TArray<uint16>& MergedDirty) const;
+
 	void RebuildConditionalProperties(
 		FSendingRepState* RESTRICT RepState,
 		const FReplicationFlags& RepFlags) const;
@@ -1669,7 +1669,7 @@ private:
 	void UpdateChangelistHistory(
 		FSendingRepState* RepState,
 		UClass* ObjectClass,
-		const uint8* RESTRICT Data,
+		const FConstRepObjectDataBuffer Data,
 		UNetConnection* Connection,
 		TArray<uint16>* OutMerged) const;
 
@@ -1681,7 +1681,7 @@ private:
 		FNetBitWriter& Writer,
 		const bool bDoChecksum,
 		FRepHandleIterator& HandleIterator,
-		const uint8* RESTRICT SourceData) const;
+		const FConstRepObjectDataBuffer Source) const;
 
 	void SendAllProperties_BackwardsCompatible_r(
 		FSendingRepState* RESTRICT RepState,
@@ -1691,7 +1691,7 @@ private:
 		FNetFieldExportGroup* NetFieldExportGroup,
 		const int32 CmdStart,
 		const int32 CmdEnd,
-		const uint8* SourceData) const;
+		const FConstRepObjectDataBuffer SourceData) const;
 
 	void SendProperties_r(
 		FSendingRepState* RESTRICT RepState,
@@ -1699,7 +1699,7 @@ private:
 		FNetBitWriter& Writer,
 		const bool bDoChecksum,
 		FRepHandleIterator& HandleIterator,
-		const uint8* RESTRICT SourceData,
+		const FConstRepObjectDataBuffer SourceData,
 		const int32	 ArrayDepth,
 		const FRepSerializationSharedInfo& SharedInfo) const;
 
@@ -1707,8 +1707,8 @@ private:
 		FSendingRepState* RESTRICT RepState,
 		const int32 CmdStart,
 		const int32 CmdEnd,
-		const uint8* RESTRICT CompareData,
-		const uint8* RESTRICT Data,
+		FRepShadowDataBuffer CompareData,
+		const FConstRepObjectDataBuffer Data,
 		TArray<uint16>& Changed,
 		uint16 Handle,
 		const bool bIsInitial,
@@ -1716,8 +1716,8 @@ private:
 
 	void CompareProperties_Array_r(
 		FSendingRepState* RESTRICT RepState,
-		const uint8* RESTRICT CompareData,
-		const uint8* RESTRICT Data,
+		FRepShadowDataBuffer ShadowData,
+		const FConstRepObjectDataBuffer Data,
 		TArray<uint16>& Changed,
 		const uint16 CmdIndex,
 		const uint16 Handle,
@@ -1725,39 +1725,39 @@ private:
 		const bool bForceFail) const;
 
 	void BuildSharedSerialization(
-		const uint8* RESTRICT			Data,
-		TArray<uint16>&					Changed,
-		const bool						bWriteHandle,
-		FRepSerializationSharedInfo&	SharedInfo) const;
+		const FConstRepObjectDataBuffer Data,
+		TArray<uint16>& Changed,
+		const bool bWriteHandle,
+		FRepSerializationSharedInfo& SharedInfo) const;
 
 	void BuildSharedSerialization_r(
-		FRepHandleIterator&				RepHandleIterator,
-		const uint8* RESTRICT			SourceData,
-		const bool						bWriteHandle,
-		const bool						bDoChecksum,
-		const int32						ArrayDepth,
-		FRepSerializationSharedInfo&	SharedInfo) const;
+		FRepHandleIterator& RepHandleIterator,
+		const FConstRepObjectDataBuffer SourceData,
+		const bool bWriteHandle,
+		const bool bDoChecksum,
+		const int32 ArrayDepth,
+		FRepSerializationSharedInfo& SharedInfo) const;
 
 	void BuildSharedSerializationForRPC_DynamicArray_r(
-		const int32						CmdIndex,
-		uint8*							Data,
-		int32							AarayDepth,
-		FRepSerializationSharedInfo&	SharedInfo);
+		const int32 CmdIndex,
+		const FConstRepObjectDataBuffer Data,
+		int32 ArrayDepth,
+		FRepSerializationSharedInfo& SharedInfo);
 
 	void BuildSharedSerializationForRPC_r(
-		const int32						CmdStart,
-		const int32						CmdEnd,
-		void*							Data,
-		int32							ArrayIndex,
-		int32							ArrayDepth,
-		FRepSerializationSharedInfo&	SharedInfo);
+		const int32 CmdStart,
+		const int32 CmdEnd,
+		const FConstRepObjectDataBuffer Data,
+		int32 ArrayIndex,
+		int32 ArrayDepth,
+		FRepSerializationSharedInfo& SharedInfo);
 
 	TSharedPtr<FNetFieldExportGroup> CreateNetfieldExportGroup() const;
 
 	int32 FindCompatibleProperty(
-		const int32		CmdStart,
-		const int32		CmdEnd,
-		const uint32	Checksum) const;
+		const int32 CmdStart,
+		const int32 CmdEnd,
+		const uint32 Checksum) const;
 
 	bool ReceiveProperties_BackwardsCompatible_r(
 		FReceivingRepState* RESTRICT RepState,
@@ -1765,9 +1765,9 @@ private:
 		FNetBitReader& Reader,
 		const int32 CmdStart,
 		const int32 CmdEnd,
-		uint8* RESTRICT ShadowData,
-		uint8* RESTRICT OldData,
-		uint8* RESTRICT Data,
+		FRepShadowDataBuffer ShadowData,
+		FRepObjectDataBuffer OldData,
+		FRepObjectDataBuffer Data,
 		FGuidReferencesMap* GuidReferencesMap,
 		bool& bOutHasUnmapped,
 		bool& bOutGuidsChanged) const;
@@ -1784,31 +1784,29 @@ private:
 		FGuidReferencesMap* GuidReferencesMap,
 		UObject* OriginalObject,
 		UPackageMap* PackageMap, 
-		uint8* RESTRICT StoredData, 
-		uint8* RESTRICT Data, 
+		FRepShadowDataBuffer ShadowData, 
+		FRepObjectDataBuffer Data, 
 		const int32 MaxAbsOffset,
 		bool& bOutSomeObjectsWereMapped,
 		bool& bOutHasMoreUnmapped) const;
 
 	void SanityCheckChangeList_DynamicArray_r(
-		const int32				CmdIndex, 
-		const uint8* RESTRICT	Data, 
-		TArray<uint16> &		Changed,
-		int32 &					ChangedIndex) const;
+		const int32 CmdIndex, 
+		const FConstRepObjectDataBuffer Data, 
+		TArray<uint16>& Changed,
+		int32& ChangedIndex) const;
 
 	uint16 SanityCheckChangeList_r(
-		const int32				CmdStart, 
-		const int32				CmdEnd, 
-		const uint8* RESTRICT	Data, 
-		TArray<uint16> &		Changed,
-		int32 &					ChangedIndex,
-		uint16					Handle) const;
-
-	void SanityCheckChangeList(const uint8* RESTRICT Data, TArray<uint16>& Changed) const;
+		const int32 CmdStart, 
+		const int32 CmdEnd, 
+		const FConstRepObjectDataBuffer Data, 
+		TArray<uint16>& Changed,
+		int32& ChangedIndex,
+		uint16 Handle) const;
 
-	uint16 AddParentProperty(UProperty* Property, int32 ArrayIndex);
+	void SanityCheckChangeList(const FConstRepObjectDataBuffer Data, TArray<uint16>& Changed) const;
 
-	int32 InitFromProperty_r(
+	uint32 AddPropertyCmd(
 		UProperty* Property,
 		int32 Offset,
 		int32 RelativeHandle,
@@ -1817,8 +1815,8 @@ private:
 		int32 StaticArrayIndex,
 		const UNetConnection* ServerConnection);
 
-	uint32 AddPropertyCmd(
-		UProperty* Property,
+	uint32 AddArrayCmd(
+		UArrayProperty* Property,
 		int32 Offset,
 		int32 RelativeHandle,
 		int32 ParentIndex,
@@ -1826,47 +1824,36 @@ private:
 		int32 StaticArrayIndex,
 		const UNetConnection* ServerConnection);
 
-	uint32 AddArrayCmd(
-		UArrayProperty*			Property,
-		int32					Offset,
-		int32					RelativeHandle,
-		int32					ParentIndex,
-		uint32					ParentChecksum,
-		int32					StaticArrayIndex,
-		const UNetConnection*	ServerConnection);
-
-	void AddReturnCmd();
-
 	void SerializeProperties_DynamicArray_r(
-		FBitArchive &						Ar, 
-		UPackageMap*						Map,
-		const int32							CmdIndex,
-		uint8*								Data,
-		bool &								bHasUnmapped,
-		const int32							ArrayDepth,
-		const FRepSerializationSharedInfo&	SharedInfo) const;
+		FBitArchive& Ar, 
+		UPackageMap* Map,
+		const int32 CmdIndex,
+		FRepObjectDataBuffer Data,
+		bool& bHasUnmapped,
+		const int32 ArrayDepth,
+		const FRepSerializationSharedInfo& SharedInfo) const;
 
 	void SerializeProperties_r(
-		FBitArchive&						Ar, 
-		UPackageMap*						Map,
-		const int32							CmdStart, 
-		const int32							CmdEnd, 
-		void*								Data,
-		bool&								bHasUnmapped,
-		const int32							ArrayIndex,
-		const int32							ArrayDepth,
-		const FRepSerializationSharedInfo&	SharedInfo) const;
+		FBitArchive& Ar, 
+		UPackageMap* Map,
+		const int32 CmdStart, 
+		const int32 CmdEnd, 
+		FRepObjectDataBuffer Data,
+		bool& bHasUnmapped,
+		const int32 ArrayIndex,
+		const int32 ArrayDepth,
+		const FRepSerializationSharedInfo& SharedInfo) const;
 
 	void MergeChangeList_r(
-		FRepHandleIterator&		RepHandleIterator1,
-		FRepHandleIterator&		RepHandleIterator2,
-		const uint8* RESTRICT	SourceData,
-		TArray<uint16>&			OutChanged) const;
+		FRepHandleIterator& RepHandleIterator1,
+		FRepHandleIterator& RepHandleIterator2,
+		const FConstRepObjectDataBuffer SourceData,
+		TArray<uint16>& OutChanged) const;
 
 	void PruneChangeList_r(
-		FRepHandleIterator&		RepHandleIterator,
-		const uint8* RESTRICT	SourceData,
-		TArray<uint16>&			OutChanged) const;
+		FRepHandleIterator& RepHandleIterator,
+		const FConstRepObjectDataBuffer SourceData,
+		TArray<uint16>& OutChanged) const;
 
 	/**
 	 * Splits a given Changelist into an Inactive Change List and an Active Change List.
@@ -1889,21 +1876,21 @@ private:
 		TArray<uint16>& OutActiveProperties) const;
 
 	void BuildChangeList_r(
-		const TArray<FHandleToCmdIndex>&	HandleToCmdIndex,
-		const int32							CmdStart,
-		const int32							CmdEnd,
-		uint8*								Data,
-		const int32							HandleOffset,
-		TArray<uint16>&						Changed) const;
+		const TArray<FHandleToCmdIndex>& HandleToCmdIndex,
+		const int32 CmdStart,
+		const int32 CmdEnd,
+		const FConstRepObjectDataBuffer Data,
+		const int32 HandleOffset,
+		TArray<uint16>& Changed) const;
 
 	void BuildHandleToCmdIndexTable_r(
-		const int32					CmdStart,
-		const int32					CmdEnd,
-		TArray<FHandleToCmdIndex>&	HandleToCmdIndex);
+		const int32 CmdStart,
+		const int32 CmdEnd,
+		TArray<FHandleToCmdIndex>& HandleToCmdIndex);
 
-	void InitRepStateStaticBuffer(FRepStateStaticBuffer& ShadowData, const uint8* Source) const;
+	void InitRepStateStaticBuffer(FRepStateStaticBuffer& ShadowData, const FConstRepObjectDataBuffer Source) const;
 	void ConstructProperties(FRepStateStaticBuffer& ShadowData) const;
-	void CopyProperties(FRepStateStaticBuffer& ShadowData, const uint8* const Source) const;
+	void CopyProperties(FRepStateStaticBuffer& ShadowData, const FConstRepObjectDataBuffer Source) const;
 	void DestructProperties(FRepStateStaticBuffer& RepStateStaticBuffer) const;
 
 	ERepLayoutState LayoutState;
