Fixup clang template issues from CL-5011314

diff --git a/Engine/Source/Runtime/Engine/Private/RepLayout.cpp b/Engine/Source/Runtime/Engine/Private/RepLayout.cpp
--- a/Engine/Source/Runtime/Engine/Private/RepLayout.cpp
+++ b/Engine/Source/Runtime/Engine/Private/RepLayout.cpp
@@ -2673,7 +2673,7 @@ bool FRepLayout::ReceiveProperties(
 #ifdef ENABLE_SUPER_CHECKSUMS
 	if (InBunch.ReadBit() == 1)
 	{
-		ValidateWithChecksum(FConstRepShadowDataBuffer(RepState->StaticBuffer.GetData()), InBunch);
+		ValidateWithChecksum<>(FConstRepShadowDataBuffer(RepState->StaticBuffer.GetData()), InBunch);
 	}
 #endif
 
@@ -3256,6 +3256,12 @@ void FRepLayout::CallRepNotifies(FReceivingRepState* RepState, UObject* Object)
 	RepState->RepNotifies.Empty();
 }
 
+template<ERepDataBufferType DataType>
+static void ValidateWithChecksum_r(
+	TArray<FRepLayoutCmd>::TConstIterator& CmdIt,
+	const TConstRepDataBuffer<DataType> Data,
+	FBitArchive& Ar);
+
 template<ERepDataBufferType DataType>
 static void ValidateWithChecksum_DynamicArray_r(
 	TArray<FRepLayoutCmd>::TConstIterator& CmdIt,
@@ -3290,16 +3296,16 @@ static void ValidateWithChecksum_DynamicArray_r(
 	for (int32 i = 0; i < ArrayNum - 1; i++)
 	{
 		const int32 ArrayElementsOffset = i * ElementSize;
-		ValidateWithChecksum_r(CmdIt, Data + ArrayElementsOffset, Ar);
+		ValidateWithChecksum_r<>(CmdIt, Data + ArrayElementsOffset, Ar);
 		CmdIt -= ArraySubCommands;
 	}
 
 	const int32 ArrayElementOffset = (ArrayNum - 1) * ElementSize;
-	ValidateWithChecksum_r<DataType>(CmdIt, ArrayData + ArrayElementOffset, Ar);
+	ValidateWithChecksum_r<>(CmdIt, ArrayData + ArrayElementOffset, Ar);
 }
 
 template<ERepDataBufferType DataType>
-static void ValidateWithChecksum_r(
+void ValidateWithChecksum_r(
 	TArray<FRepLayoutCmd>::TConstIterator& CmdIt,
 	const TConstRepDataBuffer<DataType> Data, 
 	FBitArchive& Ar)
@@ -3309,11 +3315,11 @@ static void ValidateWithChecksum_r(
 		const FRepLayoutCmd& Cmd = *CmdIt;
 		if (Cmd.Type == ERepLayoutCmdType::DynamicArray)
 		{
-			ValidateWithChecksum_DynamicArray_r(CmdIt, Data + Cmd, Ar);
+			ValidateWithChecksum_DynamicArray_r<>(CmdIt, Data + Cmd, Ar);
 		}
 		else
 		{
-			SerializeReadWritePropertyChecksum(Cmd, CmdIt.GetIndex() - 1, Data + Cmd, Ar);
+			SerializeReadWritePropertyChecksum<>(Cmd, CmdIt.GetIndex() - 1, Data + Cmd, Ar);
 		}
 	}
 }
@@ -3322,14 +3328,14 @@ template<ERepDataBufferType DataType>
 void FRepLayout::ValidateWithChecksum(TConstRepDataBuffer<DataType> Data, FBitArchive& Ar) const
 {
 	TArray<FRepLayoutCmd>::TConstIterator CmdIt = Cmds.CreateConstIterator();
-	ValidateWithChecksum_r(CmdIt, Data, Ar);
+	ValidateWithChecksum_r<>(CmdIt, Data, Ar);
 	check(CmdIt.GetIndex() == Cmds.Num());
 }
 
 uint32 FRepLayout::GenerateChecksum(const FRepState* RepState) const
 {
 	FBitWriter Writer(1024, true);
-	ValidateWithChecksum(FConstRepShadowDataBuffer(RepState->GetReceivingRepState()->StaticBuffer.GetData()), Writer);
+	ValidateWithChecksum<>(FConstRepShadowDataBuffer(RepState->GetReceivingRepState()->StaticBuffer.GetData()), Writer);
 
 	return FCrc::MemCrc32(Writer.GetData(), Writer.GetNumBytes(), 0);
 }
