Add configurable throttling of ServerUpdateCamera

diff --git a/Engine/Source/Runtime/Engine/Classes/GameFramework/GameNetworkManager.h b/Engine/Source/Runtime/Engine/Classes/GameFramework/GameNetworkManager.h
--- a/Engine/Source/Runtime/Engine/Classes/GameFramework/GameNetworkManager.h
+++ b/Engine/Source/Runtime/Engine/Classes/GameFramework/GameNetworkManager.h
@@ -157,6 +157,14 @@ class ENGINE_API AGameNetworkManager : public AInfo
 	UPROPERTY(GlobalConfig)
 	float ClientErrorUpdateRateLimit;
 
+	/** Minimum delay between calls to ServerUpdateCamera, in seconds. */
+	UPROPERTY(GlobalConfig)
+	float ClientNetCamUpdateDeltaTime;
+
+	/** Camera position change limit, when exceeded allows an immediate ServerUpdateCamera call. */
+	UPROPERTY(GlobalConfig)
+	float ClientNetCamUpdatePositionLimit;
+
 	//======================================================================================================================
 	// Movement Time Discrepancy settings for Characters (speed hack detection and prevention)
 
diff --git a/Engine/Source/Runtime/Engine/Private/GameNetworkManager.cpp b/Engine/Source/Runtime/Engine/Private/GameNetworkManager.cpp
--- a/Engine/Source/Runtime/Engine/Private/GameNetworkManager.cpp
+++ b/Engine/Source/Runtime/Engine/Private/GameNetworkManager.cpp
@@ -32,6 +32,8 @@ AGameNetworkManager::AGameNetworkManager(const FObjectInitializer& ObjectInitial
 	ClientNetSendMoveThrottleOverPlayerCount = 10;
 	ClientAuthorativePosition = false;
 	ClientErrorUpdateRateLimit = 0.0f;
+	ClientNetCamUpdateDeltaTime = 0.0166f;
+	ClientNetCamUpdatePositionLimit = 1000.0f;
 	bMovementTimeDiscrepancyDetection = false;
 	bMovementTimeDiscrepancyResolution = false;
 	MovementTimeDiscrepancyMaxTimeMargin = 0.25f;
diff --git a/Engine/Source/Runtime/Engine/Private/PlayerCameraManager.cpp b/Engine/Source/Runtime/Engine/Private/PlayerCameraManager.cpp
--- a/Engine/Source/Runtime/Engine/Private/PlayerCameraManager.cpp
+++ b/Engine/Source/Runtime/Engine/Private/PlayerCameraManager.cpp
@@ -20,6 +20,7 @@
 #include "Camera/CameraPhotography.h"
 #include "GameFramework/PlayerState.h"
 #include "IXRTrackingSystem.h" // for IsHeadTrackingAllowed()
+#include "GameFramework/GameNetworkManager.h"
 
 DEFINE_LOG_CATEGORY_STATIC(LogPlayerCameraManager, Log, All);
 
@@ -875,26 +876,45 @@ void APlayerCameraManager::UpdateCamera(float DeltaTime)
 	{
 		DoUpdateCamera(DeltaTime);
 
-		TimeSinceLastServerUpdateCamera += DeltaTime;
+		const float TimeDilation = FMath::Max(GetActorTimeDilation(), KINDA_SMALL_NUMBER);
+
+		TimeSinceLastServerUpdateCamera += (DeltaTime / TimeDilation);
 
 		if (bShouldSendClientSideCameraUpdate && IsNetMode(NM_Client))
 		{
 			SCOPE_CYCLE_COUNTER(STAT_ServerUpdateCamera);
 
+			const AGameNetworkManager* const GameNetworkManager = GetDefault<AGameNetworkManager>();
+			const float ClientNetCamUpdateDeltaTime = GameNetworkManager->ClientNetCamUpdateDeltaTime;
+			const float ClientNetCamUpdatePositionLimit = GameNetworkManager->ClientNetCamUpdatePositionLimit;
+
 			FMinimalViewInfo CurrentPOV = GetCameraCachePOV();
-			
-			if (!CurrentPOV.Equals(GetLastFrameCameraCachePOV()) || TimeSinceLastServerUpdateCamera > ServerUpdateCameraTimeout)
+			FMinimalViewInfo LastPOV = GetLastFrameCameraCachePOV();
+
+			FVector ClientCameraPosition = FRepMovement::RebaseOntoZeroOrigin(CurrentPOV.Location, this);
+			FVector PrevClientCameraPosition = FRepMovement::RebaseOntoZeroOrigin(LastPOV.Location, this);
+
+			const bool bPositionThreshold = (ClientCameraPosition - PrevClientCameraPosition).SizeSquared() > (ClientNetCamUpdatePositionLimit * ClientNetCamUpdatePositionLimit);
+
+			if (bPositionThreshold || (TimeSinceLastServerUpdateCamera > ClientNetCamUpdateDeltaTime))
 			{
 				// compress the rotation down to 4 bytes
 				int32 const ShortYaw = FRotator::CompressAxisToShort(CurrentPOV.Rotation.Yaw);
 				int32 const ShortPitch = FRotator::CompressAxisToShort(CurrentPOV.Rotation.Pitch);
 				int32 const CompressedRotation = (ShortYaw << 16) | ShortPitch;
 
-				FVector ClientCameraPosition = FRepMovement::RebaseOntoZeroOrigin(CurrentPOV.Location, this);
-				PCOwner->ServerUpdateCamera(ClientCameraPosition, CompressedRotation);
+				int32 const PrevShortYaw = FRotator::CompressAxisToShort(LastPOV.Rotation.Yaw);
+				int32 const PrevShortPitch = FRotator::CompressAxisToShort(LastPOV.Rotation.Pitch);
+				int32 const PrevCompressedRotation = (PrevShortYaw << 16) | PrevShortPitch;
+
+				if ((CompressedRotation != PrevCompressedRotation) || !ClientCameraPosition.Equals(PrevClientCameraPosition) || (TimeSinceLastServerUpdateCamera > ServerUpdateCameraTimeout))
+				{
+					PCOwner->ServerUpdateCamera(ClientCameraPosition, CompressedRotation);
 
-				TimeSinceLastServerUpdateCamera = 0.0f;
+					TimeSinceLastServerUpdateCamera = 0.0f;
+				}
 			}
+
 			bShouldSendClientSideCameraUpdate = false;
 		}
 	}
