Remove RepCmdIterators in favor of just inlining their logic in places.
They made code much harder to read and reason about, making things much more fragile.

diff --git a/Engine/Source/Runtime/Engine/Private/RepLayout.cpp b/Engine/Source/Runtime/Engine/Private/RepLayout.cpp
--- a/Engine/Source/Runtime/Engine/Private/RepLayout.cpp
+++ b/Engine/Source/Runtime/Engine/Private/RepLayout.cpp
@@ -359,213 +359,6 @@ static uint32 GetRepLayoutCmdCompatibleChecksum(
 	return CompatibleChecksum;
 }
 
-#define INIT_STACK(TStack)								\
-	void InitStack(TStack& StackState)					\
-
-#define SHOULD_PROCESS_NEXT_CMD()						\
-	bool ShouldProcessNextCmd()							\
-
-#define PROCESS_ARRAY_CMD(TStack)						\
-	void ProcessArrayCmd_r(								\
-	TStack&							PrevStackState,		\
-	TStack&							StackState,			\
-	const FRepLayoutCmd&			Cmd,				\
-	const int32						CmdIndex,			\
-	FShadowBuffer					ShadowData,			\
-	FDataBuffer						Data)				\
-
-
-#define PROCESS_CMD(TStack)								\
-	void ProcessCmd(									\
-	TStack&							StackState,			\
-	const FRepLayoutCmd&			Cmd,				\
-	const int32						CmdIndex,			\
-	FShadowBuffer					ShadowData,			\
-	FDataBuffer						Data)				\
-
-template<ERepDataBufferType DataType, ERepDataBufferType ShadowType>
-class TCmdIteratorBaseStackState
-{
-public:
-
-	using FCmdIteratorBaseStackState = TCmdIteratorBaseStackState<DataType, ShadowType> ;
-	using FShadowBuffer = TRepDataBuffer<ShadowType> ;
-	using FDataBuffer = TRepDataBuffer<DataType> ;
-
-	TCmdIteratorBaseStackState(
-		const int32 InCmdStart,
-		const int32 InCmdEnd,
-		FScriptArray* InShadowArray,
-		FScriptArray* InDataArray,
-		FShadowBuffer InShadowBaseData,
-		FDataBuffer InBaseData
-	):
-		CmdStart(InCmdStart),
-		CmdEnd(InCmdEnd),
-		ShadowArray(InShadowArray),
-		DataArray(InDataArray),
-		ShadowBaseData(InShadowBaseData),
-		BaseData(InBaseData)
-	{
-	}
-
-	const int32 CmdStart; 
-	const int32 CmdEnd;
-
-	FScriptArray* ShadowArray;
-	FScriptArray* DataArray;
-
-	FShadowBuffer ShadowBaseData;
-	FDataBuffer BaseData;
-};
-
-// This uses the "Curiously recurring template pattern" (CRTP) ideas
-template<typename TImpl, typename TStackState>
-class TRepLayoutCmdIterator
-{
-public:
-
-	using Super = TRepLayoutCmdIterator<TImpl, TStackState>;
-	using FCmdIteratorBaseStackState = typename TStackState::FCmdIteratorBaseStackState;
-	using FShadowBuffer = typename TStackState::FShadowBuffer;
-	using FDataBuffer = typename TStackState::FDataBuffer;
-
-	TRepLayoutCmdIterator(const TArray<FRepParentCmd>& InParents, const TArray<FRepLayoutCmd>& InCmds):
-		Parents(InParents),
-		Cmds(InCmds)
-	{}
-
-	void ProcessDataArrayElements_r(TStackState& StackState, const FRepLayoutCmd& ArrayCmd)
-	{
-		const int32 NumDataArrayElements = StackState.DataArray ? StackState.DataArray->Num() : 0;
-		const int32 NumShadowArrayElements = StackState.ShadowArray ? StackState.ShadowArray->Num() : 0;
-
-		// Loop using the number of elements in data array
-		for (int32 i = 0; i < NumDataArrayElements; i++)
-		{
-			const int32 ElementOffset = i * ArrayCmd.ElementSize;
-
-			// ShadowArray might be smaller than DataArray
-			FDataBuffer NewDataBuffer = StackState.BaseData + ElementOffset;
-			FShadowBuffer NewShadowBuffer = i < NumShadowArrayElements ? (StackState.ShadowBaseData + ElementOffset) : nullptr;
-
-			ProcessCmds_r(StackState, NewShadowBuffer, NewDataBuffer);
-		}
-	}
-
-	void ProcessShadowArrayElements_r(TStackState& StackState, const FRepLayoutCmd& ArrayCmd)
-	{
-		const int32 NumDataArrayElements = StackState.DataArray ? StackState.DataArray->Num()	: 0;
-		const int32 NumShadowArrayElements = StackState.ShadowArray ? StackState.ShadowArray->Num() : 0;
-
-		// Loop using the number of elements in shadow array
-		for (int32 i = 0; i < NumShadowArrayElements; i++)
-		{
-			const int32 ElementOffset = i * ArrayCmd.ElementSize;
-
-			// DataArray might be smaller than ShadowArray
-			FDataBuffer NewDataBuffer = i < NumDataArrayElements ? (StackState.BaseData + ElementOffset) : nullptr;
-			FShadowBuffer NewShadowBuffer = StackState.ShadowBaseData + ElementOffset;
-
-			ProcessCmds_r(StackState, NewShadowBuffer, NewDataBuffer);
-		}
-	}
-
-	void ProcessArrayCmd_r(
-		TStackState& PrevStackState,
-		const FRepLayoutCmd& Cmd,
-		const int32 CmdIndex,
-		FShadowBuffer ShadowData,
-		FDataBuffer Data)
-	{
-		check(ShadowData || Data);
-
-		FScriptArray* ShadowArray = (FScriptArray*)ShadowData.Data;
-		FScriptArray* DataArray = (FScriptArray*)Data.Data;
-
-		TStackState StackState(CmdIndex + 1, Cmd.EndCmd - 1, ShadowArray, DataArray, ShadowArray ? ShadowArray->GetData() : nullptr, DataArray ? DataArray->GetData() : nullptr);
-
-		static_cast<TImpl*>(this)->ProcessArrayCmd_r(PrevStackState, StackState, Cmd, CmdIndex, ShadowData, Data);
-	}
-
-	void ProcessCmds_r(
-		TStackState& StackState,
-		FShadowBuffer ShadowData,
-		FDataBuffer Data)
-	{
-		check(ShadowData || Data);
-
-		for (int32 CmdIndex = StackState.CmdStart; CmdIndex < StackState.CmdEnd; CmdIndex++)
-		{
-			const FRepLayoutCmd& Cmd = Cmds[CmdIndex];
-
-			check(Cmd.Type != ERepLayoutCmdType::Return);
-
-			if (Cmd.Type == ERepLayoutCmdType::DynamicArray)
-			{
-				if (static_cast<TImpl*>(this)->ShouldProcessNextCmd())
-				{
-					ProcessArrayCmd_r(StackState, Cmd, CmdIndex, ShadowData ? (ShadowData + Cmd) : nullptr, Data ? (Data + Cmd) : nullptr);
-				}
-				CmdIndex = Cmd.EndCmd - 1;	// Jump past children of this array (-1 for ++ in for loop)
-			}
-			else
-			{
-				if (static_cast<TImpl*>(this)->ShouldProcessNextCmd())
-				{
-					static_cast<TImpl*>(this)->ProcessCmd(StackState, Cmd, CmdIndex, ShadowData, Data);
-				}
-			}
-		}
-	}
-
-	void ProcessCmds(FDataBuffer Data, FShadowBuffer ShadowData)
-	{
-		TStackState StackState(0, Cmds.Num() - 1, NULL, NULL, ShadowData, Data);
-
-		static_cast<TImpl*>(this)->InitStack(StackState);
-
-		ProcessCmds_r(StackState, ShadowData, Data);
-	}
-
-	const TArray<FRepParentCmd>& Parents;
-	const TArray<FRepLayoutCmd>& Cmds;
-};
-
-TStaticBitArray<COND_Max> BuildConditionMapFromRepFlags(FReplicationFlags RepFlags)
-{
-	TStaticBitArray<COND_Max> ConditionMap;
-
-	// Setup condition map
-	const bool bIsInitial = RepFlags.bNetInitial ? true : false;
-	const bool bIsOwner = RepFlags.bNetOwner ? true : false;
-	const bool bIsSimulated = RepFlags.bNetSimulated ? true : false;
-	const bool bIsPhysics = RepFlags.bRepPhysics ? true : false;
-	const bool bIsReplay = RepFlags.bReplay ? true : false;
-
-	ConditionMap[COND_None] = true;
-	ConditionMap[COND_InitialOnly] = bIsInitial;
-
-	ConditionMap[COND_OwnerOnly] = bIsOwner;
-	ConditionMap[COND_SkipOwner] = !bIsOwner;
-
-	ConditionMap[COND_SimulatedOnly] = bIsSimulated;
-	ConditionMap[COND_SimulatedOnlyNoReplay] = bIsSimulated && !bIsReplay;
-	ConditionMap[COND_AutonomousOnly] = !bIsSimulated;
-
-	ConditionMap[COND_SimulatedOrPhysics] = bIsSimulated || bIsPhysics;
-	ConditionMap[COND_SimulatedOrPhysicsNoReplay] = (bIsSimulated || bIsPhysics) && !bIsReplay;
-
-	ConditionMap[COND_InitialOrOwner] = bIsInitial || bIsOwner;
-	ConditionMap[COND_ReplayOrOwner] = bIsReplay || bIsOwner;
-	ConditionMap[COND_ReplayOnly] = bIsReplay;
-	ConditionMap[COND_SkipReplay] = !bIsReplay;
-
-	ConditionMap[COND_Custom] = true;
-	ConditionMap[COND_Never] = false;
-
-	return ConditionMap;
-}
 
 void FRepStateStaticBuffer::CountBytes(FArchive& Ar) const
 {
@@ -1731,7 +1524,9 @@ const FRepSerializedPropertyInfo* FRepSerializationSharedInfo::WriteSharedProper
 
 	FRepSerializedPropertyInfo& SharedPropInfo = SharedPropertyInfo[InfoIndex];
 	SharedPropInfo.Guid = PropertyGuid;
-	SharedPropInfo.BitOffset = SerializedProperties->GetNumBits();
+	SharedPropInfo.BitOffset = SerializedProperties->GetNumBits(); 
+
+	UE_LOG(LogRepProperties, VeryVerbose, TEXT("WriteSharedProperty: Handle=%d, Guid=%s"), Handle, *PropertyGuid.ToString());
 
 	if (bWriteHandle)
 	{
@@ -1849,6 +1644,7 @@ void FRepLayout::SendProperties_r(
 		// Use shared serialization if was found
 		if (SharedPropInfo)
 		{
+			UE_LOG(LogRepProperties, VeryVerbose, TEXT("SerializeProperties_r: SharedSerialization - Handle=%d, Guid=%s"), HandleIterator.Handle, *SharedPropInfo->Guid.ToString());
 			GNumSharedSerializationHit++;
 #if !(UE_BUILD_SHIPPING || UE_BUILD_TEST)
 			if (GNetVerifyShareSerializedData != 0)
@@ -1871,8 +1667,7 @@ void FRepLayout::SendProperties_r(
 				BitWriterMark.Pop(Writer);
 
 				Writer.SerializeBitsWithOffset(SharedInfo.SerializedProperties->GetData(), SharedPropInfo->BitOffset, SharedPropInfo->BitLength);
-				UE_LOG(LogRepProperties, VeryVerbose, TEXT("SerializeProperties_r: SharedSerialization, BitOffset=%s, BitLength=%s"), SharedPropInfo->BitOffset, SharedPropInfo->BitLength);
-
+				
 				TArray<uint8> SharedBuffer;
 				BitWriterMark.Copy(Writer, SharedBuffer);
 
@@ -1885,7 +1680,6 @@ void FRepLayout::SendProperties_r(
 #endif
 			{
 				Writer.SerializeBitsWithOffset(SharedInfo.SerializedProperties->GetData(), SharedPropInfo->BitOffset, SharedPropInfo->BitLength);
-				UE_LOG(LogRepProperties, VeryVerbose, TEXT("SerializeProperties_r: SharedSerialization, BitOffset=%d, BitLength=%d"), SharedPropInfo->BitOffset, SharedPropInfo->BitLength);
 			}
 
 			NETWORK_PROFILER(GNetworkProfiler.TrackReplicateProperty(ParentCmd.Property, SharedPropInfo->PropBitLength, nullptr));
@@ -2297,24 +2091,6 @@ void FRepLayout::SendProperties_BackwardsCompatible(
 	}
 }
 
-class FReceivedPropertiesStackState : public TCmdIteratorBaseStackState<ERepDataBufferType::ObjectBuffer, ERepDataBufferType::ShadowBuffer>
-{
-public:
-	FReceivedPropertiesStackState(
-		const int32		InCmdStart,
-		const int32		InCmdEnd,
-		FScriptArray*	InShadowArray,
-		FScriptArray*	InDataArray,
-		FShadowBuffer	InShadowBaseData,
-		FDataBuffer		InBaseData
-	): 
-		FCmdIteratorBaseStackState(InCmdStart, InCmdEnd, InShadowArray, InDataArray, InShadowBaseData, InBaseData),
-		GuidReferencesMap(nullptr)
-	{}
-
-	FGuidReferencesMap* GuidReferencesMap;
-};
-
 static bool ReceivePropertyHelper(
 	FNetBitReader& Bunch, 
 	FGuidReferencesMap* GuidReferencesMap,
@@ -2502,124 +2278,197 @@ static FGuidReferencesMap* PrepReceivedArray(
 	return NewGuidReferencesArray ? NewGuidReferencesArray->Array : nullptr;
 }
 
-class FReceivePropertiesImpl : public TRepLayoutCmdIterator<FReceivePropertiesImpl, FReceivedPropertiesStackState>
+/** Struct containing parameters that don't change throughout recursion of ReceiveProperties_r */
+struct FReceivePropertiesSharedParams
 {
-public:
-	FReceivePropertiesImpl(
-		FNetBitReader& InBunch,
-		FReceivingRepState* RESTRICT InRepState,
-		bool bInDoChecksum,
-		const TArray<FRepParentCmd>& InParents,
-		const TArray<FRepLayoutCmd>& InCmds,
-		const EReceivePropertiesFlags InFlags
-	):
-        TRepLayoutCmdIterator(InParents, InCmds),
-		WaitingHandle(0),
-		CurrentHandle(0), 
-		Bunch(InBunch),
-		RepState(InRepState),
-		bDoChecksum(bInDoChecksum),
-		bHasUnmapped(false),
-		Flags(InFlags)
-	{}
-
-	void ReadNextHandle()
-	{
-		Bunch.SerializeIntPacked(WaitingHandle);
-
-#ifdef ENABLE_PROPERTY_CHECKSUMS
-		if (bDoChecksum)
-		{
-			SerializeGenericChecksum(Bunch);
-		}
-#endif
-	}
-
-	INIT_STACK(FReceivedPropertiesStackState)
-	{
-		StackState.GuidReferencesMap = &RepState->GuidReferencesMap;
-	}
+	const bool bDoChecksum;
+	const bool bSkipRoleSwap;
+	TArray<UProperty*>* RepNotifies;
+	FNetBitReader& Bunch;
+	bool& bOutHasUnmapped;
+	bool& bOutGuidsChanged;
+	const TArray<FRepParentCmd>& Parents;
+	const TArray<FRepLayoutCmd>& Cmds;
+	uint16 ReadHandle = 0;
+};
 
-	SHOULD_PROCESS_NEXT_CMD()
-	{
-		CurrentHandle++;
+/** Struct containing parameters that do change as we recurse into ReceiveProperties_r */
+struct FReceivePropertiesStackParams
+{
+	FRepObjectDataBuffer ObjectData;
+	FRepShadowDataBuffer ShadowData;
+	FGuidReferencesMap* GuidReferences;
+	const int32 CmdStart;
+	const int32 CmdEnd;
+	uint32 ArrayElementOffset = 0;
+	uint16 CurrentHandle = 0;
+};
 
-		if (CurrentHandle == WaitingHandle)
-		{
-			check(WaitingHandle != 0);
-			return true;
-		}
+static FORCEINLINE void ReadPropertyHandle(FReceivePropertiesSharedParams& Params)
+{
+	uint32 Handle = 0;
+	Params.Bunch.SerializeIntPacked(Handle);
 
-		return false;
-	}
+	UE_LOG(LogRepProperties, VeryVerbose, TEXT("ReadPropertyHandle: Handle=%d"), Handle);
 
-	PROCESS_ARRAY_CMD(FReceivedPropertiesStackState)
+#ifdef ENABLE_PROPERTY_CHECKSUMS
+	if (Params.bDoChecksum)
 	{
-		// Read array size
-		uint16 ArrayNum = 0;
-		Bunch << ArrayNum;
-
-		// Read the next property handle
-		ReadNextHandle();
+		SerializeGenericChecksum(Params.Bunch);
+	}
+#endif
 
-		const int32 AbsOffset = Data.Data - PrevStackState.BaseData.Data;
+	Params.ReadHandle = Handle;
+}
 
-		const FRepParentCmd& Parent = Parents[Cmd.ParentIndex];
+static bool ReceiveProperties_r(FReceivePropertiesSharedParams& Params, FReceivePropertiesStackParams& StackParams)
+{
+	// Note, it's never possible for the ObjectData to be nullptr.
+	// However, it is possible for the ShadowData to be nullptr.
+	// At the top level, ShadowData will always be valid.
+	// If RepNotifies aren't being used, PrepReceivedArray will ignore the current shadow data and just null out the next level's shadow data.
+	// If RepNotifies aren't being used, ReceivePropertyHelper will ignore the shadow data.
 
-		StackState.GuidReferencesMap = PrepReceivedArray(
-			ArrayNum,
-			StackState.ShadowArray,
-			StackState.DataArray,
-			PrevStackState.GuidReferencesMap,
-			AbsOffset,
-			Parent,
-			Cmd,
-			CmdIndex,
-			&StackState.ShadowBaseData,
-			&StackState.BaseData,
-			EnumHasAnyFlags(Flags, EReceivePropertiesFlags::RepNotifies) ? &RepState->RepNotifies : nullptr);
+	check(StackParams.GuidReferences != nullptr);
+	for (int32 CmdIndex = StackParams.CmdStart; CmdIndex < StackParams.CmdEnd; ++CmdIndex)
+	{
+		const FRepLayoutCmd& Cmd = Params.Cmds[CmdIndex];
+		check(ERepLayoutCmdType::Return != Cmd.Type);
 
-		// Save the old handle so we can restore it when we pop out of the array
-		const uint16 OldHandle = CurrentHandle;
+		++StackParams.CurrentHandle;
+		if (StackParams.CurrentHandle != Params.ReadHandle)
+		{
+			// Skip this property.
+			if (ERepLayoutCmdType::DynamicArray == Cmd.Type)
+			{
+				CmdIndex = Cmd.EndCmd - 1;
+			}
 
-		// Array children handles are always relative to their immediate parent
-		CurrentHandle = 0;
+			UE_LOG(LogRepProperties, VeryVerbose, TEXT("ReceiveProperties_r: Skipping Property Parent=%d, Cmd=%d, CurrentHandle=%d, ReadHandle=%d"),
+				Cmd.ParentIndex, CmdIndex, StackParams.CurrentHandle, Params.ReadHandle);
+		}
+		else
+		{
+			UE_LOG(LogRepProperties, VeryVerbose, TEXT("ReceiveProperties_r: Parent=%d, Cmd=%d, ArrayIndex=%d"), Cmd.ParentIndex, CmdIndex);
+				
+			if (ERepLayoutCmdType::DynamicArray == Cmd.Type)
+			{
+				const FRepParentCmd& Parent = Params.Parents[Cmd.ParentIndex];
+
+				// Don't worry about checking the ShadowData for nullptr here.
+				// We're either:
+				//	1. At the top level and it's valid
+				//	2. Tracking RepNotifies and it's valid.
+				//	3. We aren't tracking RepNotifies in which case it will be ignored.
+				FScriptArray* ShadowArray = (FScriptArray*)(StackParams.ShadowData + Cmd).Data;
+				FScriptArray* ObjectArray = (FScriptArray*)(StackParams.ObjectData + Cmd).Data;
+
+				// Setup a new Stack State for our array.
+				FReceivePropertiesStackParams ArrayStackParams{
+					nullptr,
+					nullptr,
+					nullptr,
+					CmdIndex + 1,
+					Cmd.EndCmd - 1,
+				};
+
+				// These buffers will track the dynamic array memory.
+				FRepObjectDataBuffer ObjectArrayBuffer = StackParams.ObjectData;
+				FRepShadowDataBuffer ShadowArrayBuffer = StackParams.ShadowData;
+
+				// Read the number of elements in the array, and resize as necessary.
+				uint16 ArrayNum = 0;
+				Params.Bunch << ArrayNum;
+
+				UE_LOG(LogRepProperties, VeryVerbose, TEXT("ReceiveProperties_r: ArrayNum=%d"), ArrayNum);
+
+				ArrayStackParams.GuidReferences = PrepReceivedArray(
+					ArrayNum,
+					ShadowArray,
+					ObjectArray,
+					StackParams.GuidReferences,
+
+					// We pass in the ArrayElementOffset + Offset here, because PrepReceivedArray expects the absolute offset.
+					StackParams.ArrayElementOffset + Cmd.Offset,
+					Params.Parents[Cmd.ParentIndex],
+					Cmd,
+					CmdIndex,
+					&ShadowArrayBuffer,
+					&ObjectArrayBuffer,
+					Params.RepNotifies);
+
+				// Read the next array handle.
+				ReadPropertyHandle(Params);
+
+				// It's possible that we've already hit the terminator.
+				// Maybe this was just a change of size of the array (like removing an element from the end).
+				if (0 != Params.ReadHandle)
+				{
+					const int32 ObjectArrayNum = ObjectArray->Num();
+					for (int32 i = 0; i < ObjectArrayNum; ++i)
+					{
+						const int32 ElementOffset = i * Cmd.ElementSize;
 
-		// Loop over array
-		ProcessDataArrayElements_r(StackState, Cmd);
+						ArrayStackParams.ObjectData = ObjectArrayBuffer + ElementOffset;
+						ArrayStackParams.ShadowData = ShadowArrayBuffer ? (ShadowArrayBuffer + ElementOffset) : nullptr;
+						ArrayStackParams.ArrayElementOffset = ElementOffset;
 
-		// Restore the current handle to what it was before we processed this array
-		CurrentHandle = OldHandle;
+						UE_LOG(LogRepProperties, VeryVerbose, TEXT("ReceivePropertiesArray_r: Recursing - Parent=%d, Cmd=%d, Index=%d"), Cmd.ParentIndex, CmdIndex, i);
+						if (!ReceiveProperties_r(Params, ArrayStackParams))
+						{
+							UE_LOG(LogRep, Error, TEXT("ReceiveProperties_r: Failed to receive property, Array Property - Property=%s, Parent=%d, Cmd=%d, Index=%d"), *Parent.CachedPropertyName.ToString(), Cmd.ParentIndex, CmdIndex, i);
+							return false;
+						}
+					}
 
-		// We should be waiting on the NULL terminator handle at this point
-		check(WaitingHandle == 0);
-		ReadNextHandle();
-	}
+					// Make sure we've hit the array terminator.
+					if (0 != Params.ReadHandle)
+					{
+						UE_LOG(LogRep, Warning, TEXT("ReceiveProperties_r: Failed to receive property, Array Property Improperly Terminated - Property=%s, Parent=%d, CmdIndex=%d, ReadHandle=%d"), *Parent.CachedPropertyName.ToString(), Cmd.ParentIndex, CmdIndex, Params.ReadHandle);
+						return false;
+					}
+				}
 
-	PROCESS_CMD(FReceivedPropertiesStackState)
-	{
-		check(StackState.GuidReferencesMap != NULL);
+				// Skip passed the inner array properties.
+				CmdIndex = Cmd.EndCmd - 1;
+			}
+			else
+			{
+				// Go ahead and receive the property.
+				if (ReceivePropertyHelper(
+					Params.Bunch,
+					StackParams.GuidReferences,
+					StackParams.ArrayElementOffset,
+					StackParams.ShadowData,
+					StackParams.ObjectData,
+					Params.RepNotifies,
+					Params.Parents,
+					Params.Cmds,
+					CmdIndex,
+					Params.bDoChecksum,
+					Params.bOutGuidsChanged,
+					Params.bSkipRoleSwap))
+				{
+					Params.bOutHasUnmapped = true;
+				}
+			}
 
-		const int32 ElementOffset = (Data.Data - StackState.BaseData.Data);
+			// TODO: Might be worth doing this before and after ReadNextHandle, or having ReadNextHandle check for errors?
+			if (Params.Bunch.IsError())
+			{
+				const FRepParentCmd& Parent = Params.Parents[Cmd.ParentIndex];
+				UE_LOG(LogRep, Error, TEXT("ReceiveProperties_r: Failed to receive property, BunchIsError - Property=%s, Parent=%d, Cmd=%d, ReadHandle=%d"), *Parent.CachedPropertyName.ToString(), Cmd.ParentIndex, CmdIndex, Params.ReadHandle);
+				return false;
+			}
 
-		if (ReceivePropertyHelper(Bunch, StackState.GuidReferencesMap, ElementOffset, ShadowData, Data, EnumHasAnyFlags(Flags, EReceivePropertiesFlags::RepNotifies) ? &RepState->RepNotifies : nullptr, Parents, Cmds, CmdIndex, bDoChecksum, bGuidsChanged, EnumHasAnyFlags(Flags, EReceivePropertiesFlags::SkipRoleSwap)))
-		{
-			bHasUnmapped = true;
+			// Read the next property handle to serialize.
+			// If we don't have any more properties, this could be a terminator.
+			ReadPropertyHandle(Params);
 		}
-
-		// Read the next property handle
-		ReadNextHandle();
 	}
 
-	uint32 WaitingHandle;
-	uint32 CurrentHandle;
-	FNetBitReader& Bunch;
-	FReceivingRepState*	RESTRICT RepState;
-	bool bDoChecksum;
-	bool bHasUnmapped;
-	bool bGuidsChanged;
-	EReceivePropertiesFlags Flags;
-};
+	return true;
+}
 
 bool FRepLayout::ReceiveProperties(
 	UActorChannel* OwningChannel,
@@ -2655,32 +2504,46 @@ bool FRepLayout::ReceiveProperties(
 	// If this is hit, it may mean the Client and Server have different properties!
 	check(LayoutState != ERepLayoutState::Empty);
 
-	FReceivePropertiesImpl ReceivePropertiesImpl(InBunch, RepState, bDoChecksum, Parents, Cmds, Flags);
-
-	// Read first handle
-	ReceivePropertiesImpl.ReadNextHandle();
+	FReceivePropertiesSharedParams Params{
+		bDoChecksum,
+		EnumHasAnyFlags(Flags, EReceivePropertiesFlags::SkipRoleSwap),
+		EnumHasAnyFlags(Flags, EReceivePropertiesFlags::RepNotifies) ? &RepState->RepNotifies : nullptr,
+		InBunch,
+		bOutHasUnmapped,
+		bOutGuidsChanged,
+		Parents,
+		Cmds
+	};
 
-	// Read all properties
-	ReceivePropertiesImpl.ProcessCmds(Data, RepState->StaticBuffer.GetData());
+	FReceivePropertiesStackParams StackParams{
+		FRepObjectDataBuffer(Data),
+		FRepShadowDataBuffer(RepState->StaticBuffer.GetData()),
+		&RepState->GuidReferencesMap,
+		0,
+		Cmds.Num() - 1
+	};
 
-	// Make sure we're waiting on the last NULL terminator
-	if (ReceivePropertiesImpl.WaitingHandle != 0)
+	// Read the first handle, and then start receiving properties.
+	ReadPropertyHandle(Params);
+	if (ReceiveProperties_r(Params, StackParams))
 	{
-		UE_LOG(LogRep, Warning, TEXT("Read out of sync."));
-		return false;
-	}
+		if (0 != Params.ReadHandle)
+		{
+			UE_LOG(LogRep, Error, TEXT("ReceiveProperties: Invalid property terminator handle - Handle=%d"), Params.ReadHandle);
+			return false;
+		}
 
 #ifdef ENABLE_SUPER_CHECKSUMS
-	if (InBunch.ReadBit() == 1)
-	{
-		ValidateWithChecksum<>(FConstRepShadowDataBuffer(RepState->StaticBuffer.GetData()), InBunch);
-	}
+		if (bDoChecksum)
+		{
+			ValidateWithChecksum<>(FConstRepShadowDataBuffer(RepState->StaticBuffer.GetData()), InBunch);
+		}
 #endif
 
-	bOutHasUnmapped = ReceivePropertiesImpl.bHasUnmapped;
-	bOutGuidsChanged = ReceivePropertiesImpl.bGuidsChanged;
+		return true;
+	}
 
-	return true;
+	return false;
 }
 
 bool FRepLayout::ReceiveProperties_BackwardsCompatible(
@@ -3143,7 +3006,7 @@ void FRepLayout::UpdateUnmappedObjects_r(
 			}
 
 			// Copy current value over so we can check to see if it changed
-			if (Parent.Property->HasAnyPropertyFlags(CPF_RepNotify))
+			if (INDEX_NONE != Parent.RepNotifyNumParams)
 			{
 				StoreProperty(Cmd, ShadowData + Cmd, Data + AbsOffset);
 			}
@@ -3155,7 +3018,7 @@ void FRepLayout::UpdateUnmappedObjects_r(
 			Cmd.Property->NetSerializeItem(Reader, PackageMap, Data + AbsOffset);
 
 			// Check to see if this property changed
-			if (Parent.Property->HasAnyPropertyFlags(CPF_RepNotify))
+			if (INDEX_NONE != Parent.RepNotifyNumParams)
 			{
 				if (Parent.RepNotifyCondition == REPNOTIFY_Always || !PropertiesAreIdentical(Cmd, ShadowData + Cmd, Data + AbsOffset))
 				{
@@ -3472,267 +3335,319 @@ void FRepLayout::SanityCheckChangeList(const FConstRepObjectDataBuffer Data, TAr
 	check(Changed[ChangedIndex] == 0);
 }
 
-template<ERepDataBufferType DataType, ERepDataBufferType ShadowType>
-class TDiffPropertiesImpl : public TRepLayoutCmdIterator<TDiffPropertiesImpl<DataType, ShadowType>, TCmdIteratorBaseStackState<DataType, ShadowType>>
+struct FDiffPropertiesSharedParams
 {
-public:
+	const ERepParentFlags PropertyFlags;
+	const EDiffPropertiesFlags DiffFlags;
+	TArray<UProperty*>* RepNotifies;
+	const TArray<FRepParentCmd>& Parents;
+	const TArray<FRepLayoutCmd>& Cmds;
+};
 
-	using Super = TRepLayoutCmdIterator<TDiffPropertiesImpl<DataType, ShadowType>, TCmdIteratorBaseStackState<DataType, ShadowType>>;
-	using FCmdIteratorBaseStackState = typename Super::FCmdIteratorBaseStackState;
-	using FShadowBuffer = typename Super::FShadowBuffer;
-	using FDataBuffer = typename Super::FDataBuffer;
+template<ERepDataBufferType DestinationType, ERepDataBufferType SourceType>
+struct TDiffPropertiesStackParams
+{
+	TRepDataBuffer<DestinationType> Destination;
+	TConstRepDataBuffer<SourceType> Source;
+	const uint16 StartCmd;
+	const uint16 EndCmd;
+};
 
-	TDiffPropertiesImpl(
-		const EDiffPropertiesFlags InFlags,
-		TArray<UProperty*>*	InRepNotifies,
-		const TArray<FRepParentCmd>& InParents,
-		const TArray<FRepLayoutCmd>& InCmds
-	): 
-		Super(InParents, InCmds),
-		Flags(InFlags),
-		ParentPropertyFlags(ERepParentFlags::IsLifetime),
-		RepNotifies(InRepNotifies),
-		bDifferent(false)
-	{
-		// Currently, only lifetime properties init from their defaults, so default to that,
-		// but also diff conditional properties if requested.
-		if ((Flags & EDiffPropertiesFlags::IncludeConditionalProperties) != EDiffPropertiesFlags::None)
-		{
-			ParentPropertyFlags |= ERepParentFlags::IsConditional;
-		}
-	}
+template<ERepDataBufferType DestinationType, ERepDataBufferType SourceType>
+static bool DiffProperties_r(FDiffPropertiesSharedParams& Params, TDiffPropertiesStackParams<DestinationType, SourceType>& StackParams)
+{
+	// Note, it's never possible for the Source or Destination data to be null.
+	// At the top level, both will always be valid (pointing to the memory of each element in each respective buffer).
+	// As we recurse, if we detect the size of the Arrays is different we'll either:
+	//	1. Bail out if we're not syncing properties.
+	//	2. Resize the Destination array to match the Source Array, guaranteeing both are allocated.
 
-	INIT_STACK(FCmdIteratorBaseStackState) { }
+	check(StackParams.Source);
+	check(StackParams.Destination);
 
-	SHOULD_PROCESS_NEXT_CMD() 
-	{ 
-		return true;
-	}
+	const bool bSyncProperties = EnumHasAnyFlags(Params.DiffFlags, EDiffPropertiesFlags::Sync);
+	bool bDifferent = false;
 
-	PROCESS_ARRAY_CMD(FCmdIteratorBaseStackState) 
+	for (uint16 CmdIndex = StackParams.StartCmd; CmdIndex < StackParams.EndCmd; ++CmdIndex)
 	{
-		if (StackState.DataArray->Num() != StackState.ShadowArray->Num())
+		const FRepLayoutCmd& Cmd = Params.Cmds[CmdIndex];
+		const FRepParentCmd& Parent = Params.Parents[Cmd.ParentIndex];
+
+		check(ERepLayoutCmdType::Return != Cmd.Type);
+
+		if (ERepLayoutCmdType::DynamicArray == Cmd.Type)
 		{
-			bDifferent = true;
+			// This will ensure that we've skipped passed the array's properties.
+			const uint16 ArrayStartCmd = CmdIndex + 1;
+			const uint16 ArrayEndCmd = Cmd.EndCmd - 1;
+			CmdIndex = ArrayEndCmd;
 
-			if ((Flags & EDiffPropertiesFlags::Sync) == EDiffPropertiesFlags::None)
-			{
-				UE_LOG(LogRep, Warning, TEXT("FDiffPropertiesImpl: Array sizes different: %s %i / %i"), *Cmd.Property->GetFullName(), StackState.DataArray->Num(), StackState.ShadowArray->Num());
-				return;
-			}
+			FScriptArray* SourceArray = (FScriptArray*)(StackParams.Source + Cmd).Data;
+			FScriptArray* DestinationArray = (FScriptArray*)(StackParams.Destination + Cmd).Data;
 
-			if ((this->Parents[Cmd.ParentIndex].Flags & ParentPropertyFlags) == ERepParentFlags::None)
+			if (SourceArray->Num() != DestinationArray->Num())
 			{
-				return;
-			}
-
-			// Make the shadow state match the actual state
-			FScriptArrayHelper ShadowArrayHelper((UArrayProperty *)Cmd.Property, ShadowData);
-			ShadowArrayHelper.Resize(StackState.DataArray->Num());
-		}
+				bDifferent = true;
+				if (!bSyncProperties)
+				{
+					UE_LOG(LogRep, Warning, TEXT("FDiffPropertiesImpl: Array sizes different: %s %i / %i"), *Cmd.Property->GetFullName(), SourceArray->Num(), DestinationArray->Num());
+					continue;
+				}
+				else if (EnumHasAnyFlags(Parent.Flags, Params.PropertyFlags))
+				{
+					continue;
+				}
 
-		StackState.BaseData = (uint8*)StackState.DataArray->GetData();
-		StackState.ShadowBaseData = (uint8*)StackState.ShadowArray->GetData();
+				// Make the destination state match the source state
+				FScriptArrayHelper DestinationArrayHelper((UArrayProperty *)Cmd.Property, DestinationArray);
+				DestinationArrayHelper.Resize(SourceArray->Num());
+			}
 
-		// Loop over array
-		this->ProcessDataArrayElements_r(StackState, Cmd);
-	}
+			decltype(StackParams.Destination) ArrayDestinationData(DestinationArray->GetData());
+			decltype(StackParams.Source) ArraySourceData(SourceArray->GetData());
 
-	PROCESS_CMD(FCmdIteratorBaseStackState) 
-	{
-		const FRepParentCmd& Parent = this->Parents[Cmd.ParentIndex];
+			typename TDecay<decltype(StackParams)>::Type ArrayStackParams{
+				ArrayDestinationData,
+				ArraySourceData,
+				ArrayStartCmd,
+				ArrayEndCmd
+			};
 
-		// Make the shadow state match the actual state at the time of send
-		if ((RepNotifies && Parent.RepNotifyCondition == REPNOTIFY_Always) || !PropertiesAreIdentical(Cmd, Data + Cmd, ShadowData + Cmd))
-		{
-			bDifferent = true;
+			for (int32 i = 0; i < SourceArray->Num(); ++i)
+			{
+				const int32 ElementOffset = i * Cmd.ElementSize;
 
-			if ((Flags & EDiffPropertiesFlags::Sync) == EDiffPropertiesFlags::None)
-			{			
-				UE_LOG(LogRep, Warning, TEXT("FDiffPropertiesImpl: Property different: %s"), *Cmd.Property->GetFullName());
-				return;
+				ArrayStackParams.Source = ArraySourceData + ElementOffset;
+				ArrayStackParams.Destination = ArrayDestinationData + ElementOffset;
+				bDifferent |= DiffProperties_r(Params, ArrayStackParams);
 			}
-
-			if ((Parent.Flags & ParentPropertyFlags) == ERepParentFlags::None)
+		}
+		else
+		{
+			// Make the shadow state match the actual state at the time of send
+			const bool bPropertyHasRepNotifies = Params.RepNotifies && INDEX_NONE != Parent.RepNotifyNumParams;
+			if ((bPropertyHasRepNotifies && Parent.RepNotifyCondition == REPNOTIFY_Always) || !PropertiesAreIdentical(Cmd, StackParams.Source + Cmd, StackParams.Destination + Cmd))
 			{
-				return;
-			}
+				bDifferent = true;
+				if (!bSyncProperties)
+				{
+					UE_LOG(LogRep, Warning, TEXT("DiffProperties_r: Property different: %s"), *Cmd.Property->GetFullName());
+					continue;
+				}
+				else if (!EnumHasAnyFlags(Parent.Flags, Params.PropertyFlags))
+				{
+					continue;
+				}
 
-			StoreProperty(Cmd, Data + Cmd, ShadowData + Cmd);
+				StoreProperty(Cmd, StackParams.Destination + Cmd, StackParams.Source + Cmd);
 
-			if (RepNotifies && Parent.Property->HasAnyPropertyFlags(CPF_RepNotify))
+				if (bPropertyHasRepNotifies)
+				{
+					Params.RepNotifies->AddUnique(Parent.Property);
+				}
+			}
+			else
 			{
-				RepNotifies->AddUnique(Parent.Property);
+				UE_CLOG(LogSkippedRepNotifies > 0, LogRep, Display, TEXT("FDiffPropertiesImpl: Skipping RepNotify because values are the same: %s"), *Cmd.Property->GetFullName());
 			}
 		}
-		else
-		{
-			UE_CLOG(LogSkippedRepNotifies> 0, LogRep, Display, TEXT("FDiffPropertiesImpl: Skipping RepNotify because values are the same: %s"), *Cmd.Property->GetFullName());
-		}
 	}
 
-	bool IsDifferent() const { return bDifferent; }
-
-private:
-	EDiffPropertiesFlags Flags;
-	ERepParentFlags ParentPropertyFlags;
-	TArray<UProperty*>* RepNotifies;
-	bool bDifferent;
-};
+	return bDifferent;
+}
 
-template<ERepDataBufferType DstType, ERepDataBufferType SrcType>
-class TDiffStablePropertiesImpl : public TRepLayoutCmdIterator<TDiffStablePropertiesImpl<DstType, SrcType>, TCmdIteratorBaseStackState<DstType, SrcType>>
+template<ERepDataBufferType DestinationType, ERepDataBufferType SourceType>
+bool FRepLayout::DiffProperties(
+	TArray<UProperty*>* RepNotifies,
+	TRepDataBuffer<DestinationType> Destination,
+	TConstRepDataBuffer<SourceType> Source,
+	const EDiffPropertiesFlags Flags) const
 {
-public:
 
-	using Super = TRepLayoutCmdIterator<TDiffStablePropertiesImpl<DstType, SrcType>, TCmdIteratorBaseStackState<DstType, SrcType>>;
-	using FCmdIteratorBaseStackState = typename Super::FCmdIteratorBaseStackState;
-	using FShadowBuffer = typename Super::FShadowBuffer;
-	using FDataBuffer = typename Super::FDataBuffer;
-
-	TDiffStablePropertiesImpl(
-		TArray<UProperty*>*				InRepNotifies,
-		TArray<UObject*>*				InObjReferences,
-		const TArray<FRepParentCmd>&	InParents,
-		const TArray<FRepLayoutCmd>&	InCmds
-	): 
-		Super(InParents, InCmds),
-		RepNotifies(InRepNotifies),
-		ObjReferences(InObjReferences),
-		bDifferent(false)
-	{}
-
-	INIT_STACK(FCmdIteratorBaseStackState) { }
-
-	SHOULD_PROCESS_NEXT_CMD() 
-	{ 
-		return true;
+	// Currently, only lifetime properties init from their defaults, so default to that,
+	// but also diff conditional properties if requested.
+	ERepParentFlags ParentPropertyFlags = ERepParentFlags::IsLifetime;
+	if (EnumHasAnyFlags(Flags, EDiffPropertiesFlags::IncludeConditionalProperties))
+	{
+		ParentPropertyFlags |= ERepParentFlags::IsConditional;
 	}
 
-	PROCESS_ARRAY_CMD(FCmdIteratorBaseStackState) 
-	{
-		if (StackState.DataArray->Num() != StackState.ShadowArray->Num())
-		{
-			bDifferent = true;
+	FDiffPropertiesSharedParams Params{
+		ParentPropertyFlags,
+		Flags,
+		RepNotifies,
+		Parents,
+		Cmds
+	};
 
-			if ((this->Parents[Cmd.ParentIndex].Flags & ERepParentFlags::IsLifetime) == ERepParentFlags::None)
-			{
-				// Currently, only lifetime properties init from their defaults
-				return;
-			}
+	TDiffPropertiesStackParams<DestinationType, SourceType> StackParams{
+		Destination,
+		Source,
+		0u,
+		static_cast<uint16>(Cmds.Num() - 1)
+	};
 
-			// Do not adjust source data, only the destination
-			FScriptArrayHelper ArrayHelper((UArrayProperty *)Cmd.Property, Data);
-			ArrayHelper.Resize(StackState.ShadowArray->Num());
-		}
+	return DiffProperties_r(Params, StackParams);
+}
 
-		StackState.BaseData = (uint8*)StackState.DataArray->GetData();
-		StackState.ShadowBaseData = (uint8*)StackState.ShadowArray->GetData();
+struct FDiffStablePropertiesSharedParams
+{
+	TArray<UProperty*>* RepNotifies;
+	TArray<UObject*>* ObjReferences;
+	const TArray<FRepParentCmd>& Parents;
+	const TArray<FRepLayoutCmd>& Cmds;
+};
 
-		// Loop over array
-		this->ProcessDataArrayElements_r(StackState, Cmd);
-	}
+template<ERepDataBufferType DestinationType, ERepDataBufferType SourceType>
+struct TDiffStablePropertiesStackParams
+{
+	TRepDataBuffer<DestinationType> Destination;
+	TConstRepDataBuffer<SourceType> Source;
+	const int32 StartCmd;
+	const int32 EndCmd;
+};
+
+template<ERepDataBufferType DestinationType, ERepDataBufferType SourceType>
+bool DiffStableProperties_r(FDiffStablePropertiesSharedParams& Params, TDiffStablePropertiesStackParams<DestinationType, SourceType>& StackParams)
+{
+	// Note, it's never possible for the Source or Destination data to be null.
+	// At the top level, both will always be valid (pointing to the memory of each element in each respective buffer).
+	// As we recurse, if we detect the size of the Arrays is different we'll either:
+	//	1. Bail out if we're not syncing properties.
+	//	2. Resize the Destination array to match the Source Array, guaranteeing both are allocated.
+
+	bool bDifferent = false;
 
-	PROCESS_CMD(FCmdIteratorBaseStackState) 
+	for (uint16 CmdIndex = StackParams.StartCmd; CmdIndex < StackParams.EndCmd; ++CmdIndex)
 	{
-		const FRepParentCmd& Parent = this->Parents[Cmd.ParentIndex];
+		const FRepLayoutCmd& Cmd = Params.Cmds[CmdIndex];
+		const FRepParentCmd& Parent = Params.Parents[Cmd.ParentIndex];
+
+		check(ERepLayoutCmdType::Return != Cmd.Type);
 
-		// Make the shadow state match the actual state at the time of send
-		if (!PropertiesAreIdentical(Cmd, Data + Cmd, ShadowData + Cmd))
+		if (ERepLayoutCmdType::DynamicArray == Cmd.Type)
 		{
-			bDifferent = true;
+			// This will ensure that we've skipped passed the array's properties.
+			const uint16 ArrayStartCmd = CmdIndex + 1;
+			const uint16 ArrayEndCmd = Cmd.EndCmd - 1;
+			CmdIndex = ArrayEndCmd;
 
-			if ((Parent.Flags & ERepParentFlags::IsLifetime) == ERepParentFlags::None)
+			FScriptArray* SourceArray = (FScriptArray*)(StackParams.Source + Cmd).Data;
+			FScriptArray* DestinationArray = (FScriptArray*)(StackParams.Destination + Cmd).Data;
+
+			if (SourceArray->Num() != DestinationArray->Num())
 			{
-				// Currently, only lifetime properties init from their defaults
-				return;
+				bDifferent = true;
+
+				if (!EnumHasAnyFlags(Parent.Flags, ERepParentFlags::IsLifetime))
+				{
+					// Currently, only lifetime properties init from their defaults
+					continue;
+				}
+
+				// Do not adjust source data, only the destination
+				FScriptArrayHelper DestinationArrayHelper((UArrayProperty *)Cmd.Property, DestinationArray);
+				DestinationArrayHelper.Resize(SourceArray->Num());
 			}
 
-			if (Cmd.Property->HasAnyPropertyFlags(CPF_Transient))
+			decltype(StackParams.Destination) ArrayDestinationData(DestinationArray->GetData());
+			decltype(StackParams.Source) ArraySourceData(SourceArray->GetData());
+
+			typename TDecay<decltype(StackParams)>::Type ArrayStackParams{
+				ArrayDestinationData,
+				ArraySourceData,
+				ArrayStartCmd,
+				ArrayEndCmd
+			};
+
+			for (int32 i = 0; i < SourceArray->Num(); ++i)
 			{
-				// skip transient properties
-				return;
-			}
+				const int32 ElementOffset = i * Cmd.ElementSize;
 
-			if (Cmd.Type == ERepLayoutCmdType::PropertyObject)
+				ArrayStackParams.Source = ArraySourceData + ElementOffset;
+				ArrayStackParams.Destination = ArrayDestinationData + ElementOffset;
+				bDifferent |= DiffStableProperties_r(Params, ArrayStackParams);
+			}
+		}
+		else
+		{
+			if (!PropertiesAreIdentical(Cmd, StackParams.Destination + Cmd, StackParams.Source + Cmd))
 			{
-				UObjectPropertyBase* ObjProperty = CastChecked<UObjectPropertyBase>(Cmd.Property);
-				if (ObjProperty)
+				bDifferent = true;
+
+				if (!EnumHasAnyFlags(Parent.Flags, ERepParentFlags::IsLifetime))
 				{
-					if (ObjProperty->PropertyClass && (ObjProperty->PropertyClass->IsChildOf(AActor::StaticClass()) || ObjProperty->PropertyClass->IsChildOf(UActorComponent::StaticClass())))
-					{
-						// skip actor and component references
-						return;
-					}
+					// Currently, only lifetime properties init from their defaults
+					continue;
+				}
+				else if (Cmd.Property->HasAnyPropertyFlags(CPF_Transient))
+				{
+					// skip transient properties
+					continue;
+				}
 
-					UObject* ObjValue = ObjProperty->GetObjectPropertyValue(ShadowData + Cmd);
-					if (ObjValue)
+				if (Cmd.Type == ERepLayoutCmdType::PropertyObject)
+				{
+					if (UObjectPropertyBase* ObjProperty = CastChecked<UObjectPropertyBase>(Cmd.Property))
 					{
-						const bool bStableForNetworking = (ObjValue->HasAnyFlags(RF_WasLoaded | RF_DefaultSubObject) || ObjValue->IsNative() || ObjValue->IsDefaultSubobject());
-						if (!bStableForNetworking)
+						if (ObjProperty->PropertyClass && (ObjProperty->PropertyClass->IsChildOf(AActor::StaticClass()) || ObjProperty->PropertyClass->IsChildOf(UActorComponent::StaticClass())))
 						{
-							// skip object references without a stable name
-							return;
+							// skip actor and component references
+							continue;
 						}
 
-						if (ObjReferences)
+						if (UObject* ObjValue = ObjProperty->GetObjectPropertyValue(StackParams.Source + Cmd))
 						{
-							ObjReferences->AddUnique(ObjValue);
+							const bool bStableForNetworking = (ObjValue->HasAnyFlags(RF_WasLoaded | RF_DefaultSubObject) || ObjValue->IsNative() || ObjValue->IsDefaultSubobject());
+							if (!bStableForNetworking)
+							{
+								// skip object references without a stable name
+								continue;
+							}
+
+							if (Params.ObjReferences)
+							{
+								Params.ObjReferences->AddUnique(ObjValue);
+							}
 						}
 					}
 				}
-			}
 
-			StoreProperty(Cmd, Data + Cmd, ShadowData + Cmd);
+				StoreProperty(Cmd, StackParams.Destination + Cmd, StackParams.Source + Cmd);
 
-			if (RepNotifies && Parent.Property->HasAnyPropertyFlags(CPF_RepNotify))
-			{
-				RepNotifies->AddUnique(Parent.Property);
+				if (Params.RepNotifies && INDEX_NONE != Parent.RepNotifyNumParams)
+				{
+					Params.RepNotifies->AddUnique(Parent.Property);
+				}
 			}
 		}
 	}
 
-	TArray<UProperty*>* RepNotifies;
-	TArray<UObject*>* ObjReferences;
-	bool bDifferent;
-};
-
-template<ERepDataBufferType DstType, ERepDataBufferType SrcType>
-bool FRepLayout::DiffProperties(
-	TArray<UProperty*>* RepNotifies,
-	TRepDataBuffer<DstType> Destination,
-	TConstRepDataBuffer<SrcType> Source,
-	const EDiffPropertiesFlags Flags) const
-{
-	if (ERepLayoutState::Empty == LayoutState)
-	{
-		return false;
-	}
-	else
-	{
-		TDiffPropertiesImpl<DstType, SrcType> DiffPropertiesImpl(Flags, RepNotifies, Parents, Cmds);
-		DiffPropertiesImpl.ProcessCmds(Destination, (uint8*)Source.Data);
-		return DiffPropertiesImpl.IsDifferent();
-	}
+	return bDifferent;
 }
 
-template<ERepDataBufferType DstType, ERepDataBufferType SrcType>
+template<ERepDataBufferType DestinationType, ERepDataBufferType SourceType>
 bool FRepLayout::DiffStableProperties(
 	TArray<UProperty*>* RepNotifies,
 	TArray<UObject*>* ObjReferences,
-	TRepDataBuffer<DstType> Destination,
-	TConstRepDataBuffer<SrcType> Source) const
+	TRepDataBuffer<DestinationType> Destination,
+	TConstRepDataBuffer<SourceType> Source) const
 {
-	if (ERepLayoutState::Empty == LayoutState)
-	{
-		return false;
-	}
-	else
-	{
-		TDiffStablePropertiesImpl<DstType, SrcType> DiffPropertiesImpl(RepNotifies, ObjReferences, Parents, Cmds);
-		DiffPropertiesImpl.ProcessCmds(Destination, (uint8*)Source.Data);
-		return DiffPropertiesImpl.bDifferent;
-	}
+	FDiffStablePropertiesSharedParams Params{
+		RepNotifies,
+		ObjReferences,
+		Parents,
+		Cmds
+	};
+
+	TDiffStablePropertiesStackParams<DestinationType, SourceType> StackParams{
+		Destination,
+		Source,
+		0,
+		Cmds.Num() - 1
+	};
+
+	return DiffStableProperties_r(Params, StackParams);
 }
 
 static FName NAME_Vector_NetQuantize100(TEXT("Vector_NetQuantize100"));
@@ -5316,6 +5231,41 @@ void FSendingRepState::CountBytes(FArchive& Ar) const
 	GRANULAR_NETWORK_MEMORY_TRACKING_TRACK("InactiveParents", InactiveParents.CountBytes(Ar));
 }
 
+TStaticBitArray<COND_Max> FSendingRepState::BuildConditionMapFromRepFlags(FReplicationFlags RepFlags)
+{
+	TStaticBitArray<COND_Max> ConditionMap;
+
+	// Setup condition map
+	const bool bIsInitial = RepFlags.bNetInitial ? true : false;
+	const bool bIsOwner = RepFlags.bNetOwner ? true : false;
+	const bool bIsSimulated = RepFlags.bNetSimulated ? true : false;
+	const bool bIsPhysics = RepFlags.bRepPhysics ? true : false;
+	const bool bIsReplay = RepFlags.bReplay ? true : false;
+
+	ConditionMap[COND_None] = true;
+	ConditionMap[COND_InitialOnly] = bIsInitial;
+
+	ConditionMap[COND_OwnerOnly] = bIsOwner;
+	ConditionMap[COND_SkipOwner] = !bIsOwner;
+
+	ConditionMap[COND_SimulatedOnly] = bIsSimulated;
+	ConditionMap[COND_SimulatedOnlyNoReplay] = bIsSimulated && !bIsReplay;
+	ConditionMap[COND_AutonomousOnly] = !bIsSimulated;
+
+	ConditionMap[COND_SimulatedOrPhysics] = bIsSimulated || bIsPhysics;
+	ConditionMap[COND_SimulatedOrPhysicsNoReplay] = (bIsSimulated || bIsPhysics) && !bIsReplay;
+
+	ConditionMap[COND_InitialOrOwner] = bIsInitial || bIsOwner;
+	ConditionMap[COND_ReplayOrOwner] = bIsReplay || bIsOwner;
+	ConditionMap[COND_ReplayOnly] = bIsReplay;
+	ConditionMap[COND_SkipReplay] = !bIsReplay;
+
+	ConditionMap[COND_Custom] = true;
+	ConditionMap[COND_Never] = false;
+
+	return ConditionMap;
+}
+
 void FRepState::CountBytes(FArchive& Ar) const
 {
 	GRANULAR_NETWORK_MEMORY_TRACKING_INIT(Ar, "FRepState::CountBytes");
diff --git a/Engine/Source/Runtime/Engine/Public/Net/RepLayout.h b/Engine/Source/Runtime/Engine/Public/Net/RepLayout.h
--- a/Engine/Source/Runtime/Engine/Public/Net/RepLayout.h
+++ b/Engine/Source/Runtime/Engine/Public/Net/RepLayout.h
@@ -1110,9 +1110,19 @@ enum class ERepLayoutState
  * Changelists are arrays of Property Handles that describe what Properties have changed, however they don't
  * track the actual values of the Properties.
  *
- * Property Handles are either Layout Command indices (in FRepLayout::Cmds), or array indices for Properties in
- * dynamic arrays.
- * Handles are 1-based, reserving 0 as a terminal case.
+ * Changelists can contain "sub-changelists" for arrays. Formally, they can be described as the following grammar:
+ *
+ *		Terminator			::=	0
+ *		Handle				::= Integer between 1 ~ 65535
+ *		Number				::= Integer between 0 ~ 65535
+ *		Changelist			::=	<Terminator> | <Handle><Changelist> | <Handle><Array-Changelist><Changelist>
+ *		Array-Changelist:	::= <Number><Changelist>
+ *
+ * An important distinction is that Handles do not have a 1:1 mapping with RepLayoutCommands.
+ * Handles are 1-based (as opposed to 0-based), and track a relative command index within a single
+ * level of a changelist. Each Array Command, regardless of the number of child Commands it has,
+ * will only be count as a single handle in its owning changelist. Each time we recurse into an Array-Changelist,
+ * our handles restart at 1 for that "depth", and they correspond to the Commands associated with the Array's element type.
  *
  * In order to generate Changelists, Layout Commands are sequentially applied that compare the values
  * of an object's cached state to a object's current state. Any properties that are found to be different
