Fix initialization of odd framed animation frame stripping during cooking

diff --git a/Engine/Source/Runtime/Engine/Private/Animation/AnimSequence.cpp b/Engine/Source/Runtime/Engine/Private/Animation/AnimSequence.cpp
--- a/Engine/Source/Runtime/Engine/Private/Animation/AnimSequence.cpp
+++ b/Engine/Source/Runtime/Engine/Private/Animation/AnimSequence.cpp
@@ -182,12 +182,17 @@ void FRequestAnimCompressionParams::InitFrameStrippingFromPlatform(const class I
 
 	if (UDeviceProfile* DeviceProfile = UDeviceProfileManager::Get().FindProfile(TargetPlatform->IniPlatformName()))
 	{
-		// if we don't prune, we assume all detail modes
 		int32 CVarPlatformFrameStrippingValue = 0;
 		if (DeviceProfile->GetConsolidatedCVarValue(StripFrameCVarName, CVarPlatformFrameStrippingValue))
 		{
 			bPerformFrameStripping = CVarPlatformFrameStrippingValue == 1;
 		}
+
+		int32 CVarPlatformOddAnimFrameStrippingValue = 0;
+		if (DeviceProfile->GetConsolidatedCVarValue(OddFrameStripStrippingCVarName, CVarPlatformOddAnimFrameStrippingValue))
+		{
+			bPerformFrameStrippingOnOddNumberedFrames = CVarPlatformOddAnimFrameStrippingValue == 1;
+		}
 	}
 #endif
 }
