Infinitely waiting for Audio Thread - Log spam when launching QAGame in Standalone Game
Fix: From Aaron - Move FAudioStreamingManager::UpdateResourceStreaming to TryLock but continue instead of scope lock.

diff --git a/Engine/Source/Runtime/Engine/Private/AudioStreaming.cpp b/Engine/Source/Runtime/Engine/Private/AudioStreaming.cpp
--- a/Engine/Source/Runtime/Engine/Private/AudioStreaming.cpp
+++ b/Engine/Source/Runtime/Engine/Private/AudioStreaming.cpp
@@ -532,63 +532,66 @@ void FAudioStreamingManager::UpdateResourceStreaming(float DeltaTime, bool bProc
 {
 	LLM_SCOPE(ELLMTag::Audio);
 
-	FScopeLock Lock(&CriticalSection);
-
-	for (auto& WavePair : StreamingSoundWaves)
+	if (CriticalSection.TryLock())
 	{
-		WavePair.Value->UpdateStreamingStatus();
-	}
+		for (auto& WavePair : StreamingSoundWaves)
+		{
+			WavePair.Value->UpdateStreamingStatus();
+		}
 
-	// Process any async file requests after updating the stream status
-	ProcessPendingAsyncFileResults();
+		// Process any async file requests after updating the stream status
+		ProcessPendingAsyncFileResults();
 
-	for (auto Source : StreamingSoundSources)
-	{
-		const FWaveInstance* WaveInstance = Source->GetWaveInstance();
-		USoundWave* Wave = WaveInstance ? WaveInstance->WaveData : nullptr;
-		if (Wave)
+		for (auto Source : StreamingSoundSources)
 		{
-			FStreamingWaveData** WaveDataPtr = StreamingSoundWaves.Find(Wave);
-	
-			if (WaveDataPtr && (*WaveDataPtr)->PendingChunkChangeRequestStatus.GetValue() == AudioState_ReadyFor_Requests)
+			const FWaveInstance* WaveInstance = Source->GetWaveInstance();
+			USoundWave* Wave = WaveInstance ? WaveInstance->WaveData : nullptr;
+			if (Wave)
 			{
-				FStreamingWaveData* WaveData = *WaveDataPtr;
-				// Request the chunk the source is using and the one after that
-				FWaveRequest& WaveRequest = GetWaveRequest(Wave);
-				const FSoundBuffer* SoundBuffer = Source->GetBuffer();
-				if (SoundBuffer)
+				FStreamingWaveData** WaveDataPtr = StreamingSoundWaves.Find(Wave);
+
+				if (WaveDataPtr && (*WaveDataPtr)->PendingChunkChangeRequestStatus.GetValue() == AudioState_ReadyFor_Requests)
 				{
-					int32 SourceChunk = SoundBuffer->GetCurrentChunkIndex();
-					if (SourceChunk >= 0 && SourceChunk < Wave->RunningPlatformData->NumChunks)
-					{
-						WaveRequest.RequiredIndices.AddUnique(SourceChunk);
-						WaveRequest.RequiredIndices.AddUnique((SourceChunk + 1) % Wave->RunningPlatformData->NumChunks);
-						WaveRequest.bPrioritiseRequest = true;
-					}
-					else
+					FStreamingWaveData* WaveData = *WaveDataPtr;
+					// Request the chunk the source is using and the one after that
+					FWaveRequest& WaveRequest = GetWaveRequest(Wave);
+					const FSoundBuffer* SoundBuffer = Source->GetBuffer();
+					if (SoundBuffer)
 					{
-						UE_LOG(LogAudio, Log, TEXT("Invalid chunk request curIndex=%d numChunks=%d\n"), SourceChunk, Wave->RunningPlatformData->NumChunks);
+						int32 SourceChunk = SoundBuffer->GetCurrentChunkIndex();
+						if (SourceChunk >= 0 && SourceChunk < Wave->RunningPlatformData->NumChunks)
+						{
+							WaveRequest.RequiredIndices.AddUnique(SourceChunk);
+							WaveRequest.RequiredIndices.AddUnique((SourceChunk + 1) % Wave->RunningPlatformData->NumChunks);
+							WaveRequest.bPrioritiseRequest = true;
+						}
+						else
+						{
+							UE_LOG(LogAudio, Log, TEXT("Invalid chunk request curIndex=%d numChunks=%d\n"), SourceChunk, Wave->RunningPlatformData->NumChunks);
+						}
 					}
 				}
 			}
 		}
-	}
-
-	for (auto Iter = WaveRequests.CreateIterator(); Iter; ++Iter)
-	{
-		USoundWave* Wave = Iter.Key();
-		FStreamingWaveData* WaveData = StreamingSoundWaves.FindRef(Wave);
 
-		if (WaveData && WaveData->PendingChunkChangeRequestStatus.GetValue() == AudioState_ReadyFor_Requests)
+		for (auto Iter = WaveRequests.CreateIterator(); Iter; ++Iter)
 		{
-			WaveData->UpdateChunkRequests(Iter.Value());
-			WaveData->UpdateStreamingStatus();
-			Iter.RemoveCurrent();
+			USoundWave* Wave = Iter.Key();
+			FStreamingWaveData* WaveData = StreamingSoundWaves.FindRef(Wave);
+
+			if (WaveData && WaveData->PendingChunkChangeRequestStatus.GetValue() == AudioState_ReadyFor_Requests)
+			{
+				WaveData->UpdateChunkRequests(Iter.Value());
+				WaveData->UpdateStreamingStatus();
+				Iter.RemoveCurrent();
+			}
 		}
-	}
 
-	// Process any async file requests after updating the streaming wave data stream statuses
-	ProcessPendingAsyncFileResults();
+		// Process any async file requests after updating the streaming wave data stream statuses
+		ProcessPendingAsyncFileResults();
+
+		CriticalSection.Unlock();
+	}
 }
 
 int32 FAudioStreamingManager::BlockTillAllRequestsFinished(float TimeLimit, bool)
