Remove old ShowSelectedLightmap debug command as it was unsafe and not publicly exposed

diff --git a/Engine/Config/BaseEngine.ini b/Engine/Config/BaseEngine.ini
--- a/Engine/Config/BaseEngine.ini
+++ b/Engine/Config/BaseEngine.ini
@@ -1088,10 +1088,6 @@ bUseVirtualCores = False
 ; MemoryToLeaveForTheOSInGB = 3
 ; MemoryUsedPerSCWProcessInGB = 0.4
 
-
-[DevOptions.Debug]
-ShowSelectedLightmap=False
-
 [LogFiles]
 PurgeLogsDays=5
 MaxLogFilesOnDisk=10
diff --git a/Engine/Source/Developer/LogVisualizer/Private/VisualLoggerCameraController.cpp b/Engine/Source/Developer/LogVisualizer/Private/VisualLoggerCameraController.cpp
--- a/Engine/Source/Developer/LogVisualizer/Private/VisualLoggerCameraController.cpp
+++ b/Engine/Source/Developer/LogVisualizer/Private/VisualLoggerCameraController.cpp
@@ -52,7 +52,7 @@ void AVisualLoggerCameraController::PostInitializeComponents()
 	Super::PostInitializeComponents();
 
 	// if hud is existing, delete it and create new hud for debug camera
-	if ( MyHUD != NULL )
+	if ( MyHUD != nullptr )
 	{
 		MyHUD->Destroy();
 	}
@@ -71,10 +71,10 @@ void AVisualLoggerCameraController::PostInitializeComponents()
 AVisualLoggerCameraController* AVisualLoggerCameraController::EnableCamera(UWorld* InWorld)
 {
 	APlayerController* PC = UGameplayStatics::GetPlayerController(InWorld, 0);
-	if (PC == NULL || PC->Player == NULL || PC->IsLocalPlayerController() == false
-		|| PC->GetWorld() == NULL)
+	if (PC == nullptr || PC->Player == nullptr || PC->IsLocalPlayerController() == false
+		|| PC->GetWorld() == nullptr)
 	{
-		return NULL;
+		return nullptr;
 	}
 	
 	if (Instance.IsValid() == false)
@@ -108,7 +108,7 @@ AVisualLoggerCameraController* AVisualLoggerCameraController::EnableCamera(UWorl
 		return Instance.Get();
 	}
 
-	return NULL;
+	return nullptr;
 }
 
 void AVisualLoggerCameraController::DisableCamera(UWorld* InWorld)
@@ -117,7 +117,7 @@ void AVisualLoggerCameraController::DisableCamera(UWorld* InWorld)
 	for (FConstPlayerControllerIterator Iterator = InWorld->GetPlayerControllerIterator(); Iterator; ++Iterator)
 	{
 		AVisualLoggerCameraController* VLogCam = Cast<AVisualLoggerCameraController>(Iterator->Get());
-		if (VLogCam && VLogCam->OriginalPlayer)
+		if (VLogCam && VLogCam->OriginalPlayer && VLogCam->OriginalControllerRef)
 		{
 			VLogCam->OriginalPlayer->SwitchController(VLogCam->OriginalControllerRef);
 			VLogCam->OnDeactivate(VLogCam->OriginalControllerRef);
@@ -129,7 +129,7 @@ void AVisualLoggerCameraController::DisableCamera(UWorld* InWorld)
 
 bool AVisualLoggerCameraController::IsEnabled(UWorld* InWorld)
 {
-	if (InWorld != NULL)
+	if (InWorld != nullptr)
 	{
 		for (FConstPlayerControllerIterator Iterator = InWorld->GetPlayerControllerIterator(); Iterator; ++Iterator)
 		{
diff --git a/Engine/Source/Runtime/Engine/Classes/Engine/DebugCameraController.h b/Engine/Source/Runtime/Engine/Classes/Engine/DebugCameraController.h
--- a/Engine/Source/Runtime/Engine/Classes/Engine/DebugCameraController.h
+++ b/Engine/Source/Runtime/Engine/Classes/Engine/DebugCameraController.h
@@ -28,15 +28,15 @@ class ENGINE_API ADebugCameraController
 	UPROPERTY(globalconfig)
 	uint32 bShowSelectedInfo:1;
 
-	/** @todo document */
+	/** Saves whether the FreezeRendering console command is active */
 	UPROPERTY()
 	uint32 bIsFrozenRendering:1;
 
-	/** @todo document */
+	/** Visualizes the frustum of the camera */
 	UPROPERTY()
 	class UDrawFrustumComponent* DrawFrustum;
 	
-	/** @todo document */
+	/** Sets whether to show information about the selected actor on the debug camera HUD.t */
 	UFUNCTION(exec)
 	virtual void ShowDebugSelectedInfo();
 
@@ -46,16 +46,16 @@ class ENGINE_API ADebugCameraController
 	/** Called when the user pressed the deselect key, just before the selected actor is cleared. */
 	void Unselect();
 
-	/** @todo document */
+	/** Speeds up camera movement */
 	void IncreaseCameraSpeed();
 
-	/** @todo document */
+	/** Slows down camera movement */
 	void DecreaseCameraSpeed();
 
-	/** @todo document */
+	/** Increases camera field of vision */
 	void IncreaseFOV();
 
-	/** @todo document */
+	/** Decreases camera field of vision */
 	void DecreaseFOV();
 
 	/** Toggles the display of debug info and input commands for the Debug Camera. */
@@ -70,19 +70,24 @@ class ENGINE_API ADebugCameraController
 
 public:
 
-	/** @todo document */
+	/** Currently selected actor, may be invalid */
+	UPROPERTY()
 	class AActor* SelectedActor;
 
+	/** Returns the currently selected actor, or null if it is invalid or not set */
 	UFUNCTION(BlueprintCallable, Category="Debug Camera")
 	AActor* GetSelectedActor() const;
 	
-	/** @todo document */
+	/** Currently selected component, may be invalid */
+	UPROPERTY()
 	class UPrimitiveComponent* SelectedComponent;
 
-	/** @todo document */
+	/** Controller that was active before this was spawned */
+	UPROPERTY()	
 	class APlayerController* OriginalControllerRef;
 
-	/** @todo document */
+	/** Player object that was active before this was spawned */
+	UPROPERTY()	
 	class UPlayer* OriginalPlayer;
 
 	/** Allows control over the speed of the spectator pawn. This scales the speed based on the InitialMaxSpeed. Use Set Pawn Movement Speed Scale during runtime */
diff --git a/Engine/Source/Runtime/Engine/Classes/Engine/Engine.h b/Engine/Source/Runtime/Engine/Classes/Engine/Engine.h
--- a/Engine/Source/Runtime/Engine/Classes/Engine/Engine.h
+++ b/Engine/Source/Runtime/Engine/Classes/Engine/Engine.h
@@ -2004,7 +2004,6 @@ public:
 	bool HandleProfileGPUHitchesCommand( const TCHAR* Cmd, FOutputDevice& Ar );
 	bool HandleShaderComplexityCommand( const TCHAR* Cmd, FOutputDevice& Ar );
 	bool HandleFreezeRenderingCommand( const TCHAR* Cmd, FOutputDevice& Ar, UWorld* InWorld );
-	bool HandleShowSelectedLightmapCommand( const TCHAR* Cmd, FOutputDevice& Ar );
 	bool HandleStartFPSChartCommand( const TCHAR* Cmd, FOutputDevice& Ar );
 	bool HandleStopFPSChartCommand( const TCHAR* Cmd, FOutputDevice& Ar, UWorld* InWorld );
 	bool HandleDumpLevelScriptActorsCommand( UWorld* InWorld, const TCHAR* Cmd, FOutputDevice& Ar );
diff --git a/Engine/Source/Runtime/Engine/Private/CheatManager.cpp b/Engine/Source/Runtime/Engine/Private/CheatManager.cpp
--- a/Engine/Source/Runtime/Engine/Private/CheatManager.cpp
+++ b/Engine/Source/Runtime/Engine/Private/CheatManager.cpp
@@ -593,7 +593,7 @@ void UCheatManager::EnableDebugCamera()
 void UCheatManager::DisableDebugCamera()
 {
 	ADebugCameraController* const DCC = Cast<ADebugCameraController>(GetOuter());
-	if (DCC)
+	if (DCC && DCC->OriginalPlayer && DCC->OriginalControllerRef)
 	{
 		DCC->OriginalPlayer->SwitchController(DCC->OriginalControllerRef);
 		DCC->OnDeactivate(DCC->OriginalControllerRef);
diff --git a/Engine/Source/Runtime/Engine/Private/DebugCameraController.cpp b/Engine/Source/Runtime/Engine/Private/DebugCameraController.cpp
--- a/Engine/Source/Runtime/Engine/Private/DebugCameraController.cpp
+++ b/Engine/Source/Runtime/Engine/Private/DebugCameraController.cpp
@@ -22,26 +22,15 @@
 #include "GameFramework/PlayerInput.h"
 #include "GameFramework/GameStateBase.h"
 
-/** The currently selected actor. */
-AActor* GDebugSelectedActor = NULL;
-/** The currently selected component in the actor. */
-ENGINE_API UPrimitiveComponent* GDebugSelectedComponent = NULL;
-/** The lightmap used by the currently selected component, if it's a static mesh component. */
-ENGINE_API FLightMap2D* GDebugSelectedLightmap = NULL;
-
-extern bool UntrackTexture( const FString& TextureName );
-extern bool TrackTexture( const FString& TextureName );
-
 static const float SPEED_SCALE_ADJUSTMENT = 0.5f;
 
-
 ADebugCameraController::ADebugCameraController(const FObjectInitializer& ObjectInitializer)
 	: Super(ObjectInitializer)
 {
-	SelectedActor = NULL;
-	SelectedComponent = NULL;
-	OriginalControllerRef = NULL;
-	OriginalPlayer = NULL;
+	SelectedActor = nullptr;
+	SelectedComponent = nullptr;
+	OriginalControllerRef = nullptr;
+	OriginalPlayer = nullptr;
 
 	SpeedScale = 1.f;
 	InitialMaxSpeed = 0.f;
@@ -49,7 +38,7 @@ ADebugCameraController::ADebugCameraController(const FObjectInitializer& ObjectI
 	InitialDecel = 0.f;
 
 	bIsFrozenRendering = false;
-	DrawFrustum = NULL;
+	DrawFrustum = nullptr;
 	bHidden = false;
 #if WITH_EDITORONLY_DATA
 	bHiddenEd = false;
@@ -148,73 +137,30 @@ AActor* ADebugCameraController::GetSelectedActor() const
 
 void ADebugCameraController::Select( FHitResult const& Hit )
 {
-	// First untrack the currently tracked lightmap.
-	UTexture2D* Texture2D = GDebugSelectedLightmap ? GDebugSelectedLightmap->GetTexture(0) : NULL;
-	if ( Texture2D )
-	{
-		UntrackTexture( Texture2D->GetName() );
-	}
-
 	// store selection
-
 	SelectedActor = Hit.GetActor();
 	SelectedComponent = Hit.Component.Get();
-	GDebugSelectedActor = SelectedActor;
-	GDebugSelectedComponent = SelectedComponent;
-	
-	// figure out lightmap
-	GDebugSelectedLightmap = NULL;
-	UStaticMeshComponent* StaticMeshComponent = Cast<UStaticMeshComponent>( SelectedComponent );
-	if ( StaticMeshComponent && StaticMeshComponent->LODData.Num() > 0 )
-	{
-		const FStaticMeshComponentLODInfo& LODInfo = StaticMeshComponent->LODData[0];
-		const FMeshMapBuildData* MeshMapBuildData = StaticMeshComponent->GetMeshMapBuildData(LODInfo);
-		if ( MeshMapBuildData && MeshMapBuildData->LightMap )
-		{
-			GDebugSelectedLightmap = MeshMapBuildData->LightMap->GetLightMap2D();
-			Texture2D = GDebugSelectedLightmap ? GDebugSelectedLightmap->GetTexture(0) : NULL;
-			if ( Texture2D )
-			{
-				extern bool TrackTexture( const FString& TextureName );
-				TrackTexture( Texture2D->GetName() );
-			}
-		}
-	}
-	
+
 	//BP Event
 	ReceiveOnActorSelected(SelectedActor, Hit.ImpactPoint, Hit.ImpactNormal, Hit);
 }
 
 
 void ADebugCameraController::Unselect()
-{
-	UTexture2D* Texture2D = GDebugSelectedLightmap ? GDebugSelectedLightmap->GetTexture(0) : NULL;
-	if ( Texture2D )
-	{
-		extern bool UntrackTexture( const FString& TextureName );
-		UntrackTexture( Texture2D->GetName() );
-	}
-	
-	SelectedActor = NULL;
-	SelectedComponent = NULL;
-	
-	GDebugSelectedActor = NULL;
-	GDebugSelectedComponent = NULL;
-	GDebugSelectedLightmap = NULL;
+{	
+	SelectedActor = nullptr;
+	SelectedComponent = nullptr;
 }
 
-
-
-
 FString ADebugCameraController::ConsoleCommand(const FString& Cmd,bool bWriteToLog)
 {
 	/**
 	 * This is the same as PlayerController::ConsoleCommand(), except with some extra code to 
 	 * give our regular PC a crack at handling the command.
 	 */
-	if (Player != NULL)
+	if (Player != nullptr)
 	{
-		UConsole* ViewportConsole = (GEngine->GameViewport != NULL) ? GEngine->GameViewport->ViewportConsole : NULL;
+		UConsole* ViewportConsole = (GEngine->GameViewport != nullptr) ? GEngine->GameViewport->ViewportConsole : nullptr;
 		FConsoleOutputDevice StrOut(ViewportConsole);
 	
 		const int32 CmdLen = Cmd.Len();
@@ -238,10 +184,10 @@ FString ADebugCameraController::ConsoleCommand(const FString& Cmd,bool bWriteToL
 
 		// Free temp arrays
 		FMemory::Free(CommandBuffer);
-		CommandBuffer=NULL;
+		CommandBuffer = nullptr;
 
 		FMemory::Free(Line);
-		Line=NULL;
+		Line = nullptr;
 
 		if (!bWriteToLog)
 		{
@@ -254,7 +200,7 @@ FString ADebugCameraController::ConsoleCommand(const FString& Cmd,bool bWriteToL
 
 void ADebugCameraController::UpdateHiddenComponents(const FVector& ViewLocation,TSet<FPrimitiveComponentId>& HiddenComponentsOut)
 {
-	if (OriginalControllerRef != NULL)
+	if (OriginalControllerRef != nullptr)
 	{
 		OriginalControllerRef->UpdateHiddenComponents(ViewLocation, HiddenComponentsOut);
 	}
@@ -262,10 +208,10 @@ void ADebugCameraController::UpdateHiddenComponents(const FVector& ViewLocation,
 
 ASpectatorPawn* ADebugCameraController::SpawnSpectatorPawn()
 {
-	ASpectatorPawn* SpawnedSpectator = NULL;
+	ASpectatorPawn* SpawnedSpectator = nullptr;
 
 	// Only spawned for the local player
-	if (GetSpectatorPawn() == NULL && IsLocalController())
+	if (GetSpectatorPawn() == nullptr && IsLocalController())
 	{
 		AGameStateBase const* const GameState = GetWorld()->GetGameState();
 		if (GameState)
@@ -298,7 +244,7 @@ ASpectatorPawn* ADebugCameraController::SpawnSpectatorPawn()
 		}
 	}
 
-	return SpawnedSpectator != NULL ? SpawnedSpectator : Super::SpawnSpectatorPawn();
+	return SpawnedSpectator != nullptr ? SpawnedSpectator : Super::SpawnSpectatorPawn();
 }
 
 void ADebugCameraController::SetSpectatorPawn(ASpectatorPawn* NewSpectatorPawn)
@@ -330,11 +276,8 @@ void ADebugCameraController::PostInitializeComponents()
 {
 	Super::PostInitializeComponents();
 
-	GDebugSelectedLightmap = nullptr;
-	GDebugSelectedActor = nullptr; 
-	
 	// if hud is existing, delete it and create new hud for debug camera
-	if ( MyHUD != NULL )
+	if ( MyHUD != nullptr )
 	{
 		MyHUD->Destroy();
 	}
@@ -371,7 +314,7 @@ void ADebugCameraController::OnActivate( APlayerController* OriginalPC )
 	}
 
 	// draw frustum of original camera (where you detached)
-	if (DrawFrustum == NULL)
+	if (DrawFrustum == nullptr)
 	{
 		DrawFrustum = NewObject<UDrawFrustumComponent>(OriginalPC->PlayerCameraManager);
 	}
@@ -406,8 +349,6 @@ void ADebugCameraController::AddCheats(bool bForce)
 #endif
 }
 
-
-
 void ADebugCameraController::OnDeactivate( APlayerController* RestoredPC )
 {
 	// restore FreezeRendering command state
@@ -423,8 +364,8 @@ void ADebugCameraController::OnDeactivate( APlayerController* RestoredPC )
 	RestoredPC->SetActorHiddenInGame(true);
 	RestoredPC->PlayerCameraManager->SetActorHiddenInGame(true);
 
-	OriginalControllerRef = NULL;
-	OriginalPlayer = NULL;
+	OriginalControllerRef = nullptr;
+	OriginalPlayer = nullptr;
 
 	ChangeState(NAME_Inactive);
 	GetWorld()->RemoveController(this);
@@ -433,7 +374,6 @@ void ADebugCameraController::OnDeactivate( APlayerController* RestoredPC )
 	ReceiveOnDeactivate(RestoredPC);
 }
 
-
 void ADebugCameraController::ToggleFreezeRendering()
 {
 	ConsoleCommand(TEXT("FreezeRendering"));
@@ -515,5 +455,3 @@ void ADebugCameraController::ToggleDisplay()
 		MyHUD->ShowHUD();
 	}
 }
-
-
diff --git a/Engine/Source/Runtime/Engine/Private/GameViewportClient.cpp b/Engine/Source/Runtime/Engine/Private/GameViewportClient.cpp
--- a/Engine/Source/Runtime/Engine/Private/GameViewportClient.cpp
+++ b/Engine/Source/Runtime/Engine/Private/GameViewportClient.cpp
@@ -69,13 +69,6 @@
 /** This variable allows forcing full screen of the first player controller viewport, even if there are multiple controllers plugged in and no cinematic playing. */
 bool GForceFullscreen = false;
 
-/** Whether to visualize the lightmap selected by the Debug Camera. */
-extern ENGINE_API bool GShowDebugSelectedLightmap;
-/** The currently selected component in the actor. */
-extern ENGINE_API UPrimitiveComponent* GDebugSelectedComponent;
-/** The lightmap used by the currently selected component, if it's a static mesh component. */
-extern ENGINE_API class FLightMap2D* GDebugSelectedLightmap;
-
 /** Delegate called at the end of the frame when a screenshot is captured */
 FOnScreenshotCaptured UGameViewportClient::ScreenshotCapturedDelegate;
 
@@ -118,29 +111,6 @@ static TAutoConsoleVariable<float> CVarSecondaryScreenPercentage( // TODO: make
 	TEXT(" 1: override secondary screen percentage."),
 	ECVF_Default);
 
-/**
- * Draw debug info on a game scene view.
- */
-class FGameViewDrawer : public FViewElementDrawer
-{
-public:
-	/**
-	 * Draws debug info using the given draw interface.
-	 */
-	virtual void Draw(const FSceneView* View,FPrimitiveDrawInterface* PDI)
-	{
-#if !(UE_BUILD_SHIPPING || UE_BUILD_TEST)
-		// Draw a wireframe sphere around the selected lightmap, if requested.
-		if ( GShowDebugSelectedLightmap && GDebugSelectedComponent && GDebugSelectedLightmap )
-		{
-			float Radius = GDebugSelectedComponent->Bounds.SphereRadius;
-			int32 Sides = FMath::Clamp<int32>( FMath::TruncToInt(Radius*Radius*4.0f*PI/(80.0f*80.0f)), 8, 200 );
-			DrawWireSphere( PDI, GDebugSelectedComponent->Bounds.Origin, FColor(255,130,0), GDebugSelectedComponent->Bounds.SphereRadius, Sides, SDPG_Foreground );
-		}
-#endif
-	}
-};
-
 UGameViewportClient::UGameViewportClient(const FObjectInitializer& ObjectInitializer)
 	: Super(ObjectInitializer)
 #if WITH_EDITOR
@@ -1129,8 +1099,6 @@ void UGameViewportClient::Draw(FViewport* InViewport, FCanvas* SceneCanvas)
 		SceneCanvas->SetStereoRendering(bStereoRendering);
 	}
 
-	FGameViewDrawer GameViewDrawer;
-
 	UWorld* MyWorld = GetWorld();
 
 	// create the view family for rendering the world scene to the viewport's render target
@@ -1225,7 +1193,7 @@ void UGameViewportClient::Draw(FViewport* InViewport, FCanvas* SceneCanvas)
 
 				EStereoscopicPass PassType = bStereoRendering ? GEngine->StereoRenderingDevice->GetViewPassForIndex(bStereoRendering, i) : eSSP_FULL;
 
-				FSceneView* View = LocalPlayer->CalcSceneView(&ViewFamily, ViewLocation, ViewRotation, InViewport, &GameViewDrawer, PassType);
+				FSceneView* View = LocalPlayer->CalcSceneView(&ViewFamily, ViewLocation, ViewRotation, InViewport, nullptr, PassType);
 
 				if (View)
 				{
diff --git a/Engine/Source/Runtime/Engine/Private/Streaming/TextureStreamingHelpers.cpp b/Engine/Source/Runtime/Engine/Private/Streaming/TextureStreamingHelpers.cpp
--- a/Engine/Source/Runtime/Engine/Private/Streaming/TextureStreamingHelpers.cpp
+++ b/Engine/Source/Runtime/Engine/Private/Streaming/TextureStreamingHelpers.cpp
@@ -320,9 +320,6 @@ void FRenderAssetStreamingSettings::Update()
     }
 }
 
-
-extern ENGINE_API UPrimitiveComponent* GDebugSelectedComponent;
-
 /** the float table {-1.0f,1.0f} **/
 float ENGINE_API GNegativeOneOneTable[2] = {-1.0f,1.0f};
 
diff --git a/Engine/Source/Runtime/Engine/Private/UnrealEngine.cpp b/Engine/Source/Runtime/Engine/Private/UnrealEngine.cpp
--- a/Engine/Source/Runtime/Engine/Private/UnrealEngine.cpp
+++ b/Engine/Source/Runtime/Engine/Private/UnrealEngine.cpp
@@ -277,11 +277,6 @@ void FEngineModule::ShutdownModule()
 */
 ENGINE_API UEngine*	GEngine = NULL;
 
-/**
-* Whether to visualize the light map selected by the Debug Camera.
-*/
-ENGINE_API bool GShowDebugSelectedLightmap = false;
-
 int32 GShowMaterialDrawEvents = 0;
 
 #if WANTS_DRAW_MESH_EVENTS
@@ -1448,8 +1443,6 @@ void UEngine::Init(IEngineLoop* InEngineLoop)
 		GConfig->GetBool(TEXT("/Script/Engine.Engine"), TEXT("bEnableOnScreenDebugMessages"), bTemp, GEngineIni);
 		bEnableOnScreenDebugMessages = bTemp ? true : false;
 		bEnableOnScreenDebugMessagesDisplay = bEnableOnScreenDebugMessages;
-
-		GConfig->GetBool(TEXT("DevOptions.Debug"), TEXT("ShowSelectedLightmap"), GShowDebugSelectedLightmap, GEngineIni);
 	}
 
 	// Update Script Maximum loop iteration count
@@ -3894,10 +3887,6 @@ bool UEngine::Exec( UWorld* InWorld, const TCHAR* Cmd, FOutputDevice& Ar )
 	{
 		return HandleFreezeRenderingCommand( Cmd, Ar, InWorld );
 	}
-	else if (FParse::Command(&Cmd, TEXT("ShowSelectedLightmap")))
-	{
-		return HandleShowSelectedLightmapCommand( Cmd, Ar );
-	}
 	else if( FParse::Command(&Cmd,TEXT("SHOWLOG")) )
 	{
 		return HandleShowLogCommand( Cmd, Ar );
@@ -4511,14 +4500,6 @@ bool UEngine::HandleFreezeRenderingCommand( const TCHAR* Cmd, FOutputDevice& Ar,
 	return true;
 }
 
-bool UEngine::HandleShowSelectedLightmapCommand( const TCHAR* Cmd, FOutputDevice& Ar )
-{
-	GShowDebugSelectedLightmap = !GShowDebugSelectedLightmap;
-	GConfig->SetBool(TEXT("DevOptions.Debug"), TEXT("ShowSelectedLightmap"), GShowDebugSelectedLightmap, GEngineIni);
-	Ar.Logf( TEXT( "Showing the selected lightmap: %s" ), GShowDebugSelectedLightmap ? TEXT("true") : TEXT("false") );
-	return true;
-}
-
 bool UEngine::HandleShaderComplexityCommand( const TCHAR* Cmd, FOutputDevice& Ar )
 {
 	FString FlagStr(FParse::Token(Cmd, 0));
diff --git a/Engine/Source/Runtime/Renderer/Private/LightMapRendering.h b/Engine/Source/Runtime/Renderer/Private/LightMapRendering.h
--- a/Engine/Source/Runtime/Renderer/Private/LightMapRendering.h
+++ b/Engine/Source/Runtime/Renderer/Private/LightMapRendering.h
@@ -16,8 +16,6 @@
 
 class FPrimitiveSceneProxy;
 
-extern ENGINE_API bool GShowDebugSelectedLightmap;
-extern ENGINE_API class FLightMap2D* GDebugSelectedLightmap;
 extern bool GVisualizeMipLevels;
 
 BEGIN_GLOBAL_SHADER_PARAMETER_STRUCT(FIndirectLightingCacheUniformParameters, )
diff --git a/Engine/Source/Runtime/Renderer/Private/SceneRendering.cpp b/Engine/Source/Runtime/Renderer/Private/SceneRendering.cpp
--- a/Engine/Source/Runtime/Renderer/Private/SceneRendering.cpp
+++ b/Engine/Source/Runtime/Renderer/Private/SceneRendering.cpp
@@ -93,9 +93,6 @@ FAutoConsoleVariableRef CVarDumpMeshDrawCommandMemoryStats(
 	ECVF_Scalability | ECVF_RenderThreadSafe
 	);
 
-extern ENGINE_API FLightMap2D* GDebugSelectedLightmap;
-extern ENGINE_API UPrimitiveComponent* GDebugSelectedComponent;
-
 DECLARE_GPU_STAT_NAMED(CustomDepth, TEXT("Custom Depth"));
 
 static TAutoConsoleVariable<int32> CVarCustomDepthTemporalAAJitter(
