[NO NON-PII COMMIT MESSAGE]

diff --git a/Engine/Source/Runtime/Engine/Classes/GameFramework/NavMovementComponent.h b/Engine/Source/Runtime/Engine/Classes/GameFramework/NavMovementComponent.h
--- a/Engine/Source/Runtime/Engine/Classes/GameFramework/NavMovementComponent.h
+++ b/Engine/Source/Runtime/Engine/Classes/GameFramework/NavMovementComponent.h
@@ -79,13 +79,13 @@ public:
 	void UpdateNavAgent(const AActor& InOwner);
 	void UpdateNavAgent(const UCapsuleComponent& CapsuleComponent);
 
-	/** Returnss location of controlled actor - meaning center of collision bounding box */
+	/** Returns location of controlled actor - meaning center of collision bounding box */
 	FORCEINLINE FVector GetActorLocation() const { return UpdatedComponent ? UpdatedComponent->GetComponentLocation() : FVector(FLT_MAX); }
-	/** Returnss location of controlled actor's "feet" meaning center of bottom of collision bounding box */
+	/** Returns location of controlled actor's "feet" meaning center of bottom of collision bounding box */
 	FORCEINLINE FVector GetActorFeetLocation() const { return UpdatedComponent ? (UpdatedComponent->GetComponentLocation() - FVector(0,0,UpdatedComponent->Bounds.BoxExtent.Z)) : FNavigationSystem::InvalidLocation; }
-	/** Returnss based location of controlled actor */
+	/** Returns based location of controlled actor */
 	virtual FBasedPosition GetActorFeetLocationBased() const;
-	/** Returnss navigation location of controlled actor */
+	/** Returns navigation location of controlled actor */
 	FORCEINLINE FVector GetActorNavLocation() const { INavAgentInterface* MyOwner = Cast<INavAgentInterface>(GetOwner()); return MyOwner ? MyOwner->GetNavAgentLocation() : FNavigationSystem::InvalidLocation; }
 
 	/** path following: request new velocity */
@@ -98,7 +98,7 @@ public:
 	 *  (e.g. performing scripted move and can't stop) */
 	virtual bool CanStopPathFollowing() const;
 
-	/** Returnss braking distance for acceleration driven path following */
+	/** Returns braking distance for acceleration driven path following */
 	virtual float GetPathFollowingBrakingDistance(float MaxSpeed) const;
 
 	void SetPathFollowingAgent(IPathFollowingAgentInterface* InPathFollowingAgent) { PathFollowingComp = Cast<UObject>(InPathFollowingAgent); }
@@ -113,9 +113,9 @@ public:
 
 	FORCEINLINE bool UseAccelerationForPathFollowing() const { return bUseAccelerationForPaths; }
 
-	/** Returnss the NavAgentProps(const) */
+	/** Returns the NavAgentProps(const) */
 	FORCEINLINE const FNavAgentProperties& GetNavAgentPropertiesRef() const { return NavAgentProps; }
-	/** Returnss the NavAgentProps */
+	/** Returns the NavAgentProps */
 	FORCEINLINE FNavAgentProperties& GetNavAgentPropertiesRef() { return NavAgentProps; }
 
 	/** Resets runtime movement state to character's movement capabilities */
@@ -142,10 +142,9 @@ public:
 	/** Returns true if component is allowed to jump */
 	FORCEINLINE bool IsJumpAllowed() const { return CanEverJump() && MovementState.bCanJump; }
 
-	/** Sets rather this component is allowed to jump */
+	/** Sets whether this component is allowed to jump */
 	FORCEINLINE void SetJumpAllowed(bool bAllowed) { MovementState.bCanJump = bAllowed; }
 
-
 	/** Returns true if currently crouching */ 
 	UFUNCTION(BlueprintCallable, Category="AI|Components|NavMovement")
 	virtual bool IsCrouching() const;
