Only modify TRenderAssetUpdate::LockOwiningThreadID when the lock is held

diff --git a/Engine/Source/Runtime/Engine/Private/Streaming/RenderAssetUpdate.h b/Engine/Source/Runtime/Engine/Private/Streaming/RenderAssetUpdate.h
--- a/Engine/Source/Runtime/Engine/Private/Streaming/RenderAssetUpdate.h
+++ b/Engine/Source/Runtime/Engine/Private/Streaming/RenderAssetUpdate.h
@@ -210,8 +210,8 @@ public:
 
 		// Reset the pending task state first to prevent racing condition that could fail ensure(PendingTaskState == TS_None) in DoLock()
 		PendingTaskState = TS_None;
-		TaskState = CachedPendingTaskState;
 		LockOwningThreadID = InvalidLockOwningThreadID;
+		TaskState = CachedPendingTaskState;
 	}
 
 	/** Get the number of requested mips for this update, ignoring cancellation attempts. */
@@ -450,7 +450,7 @@ protected:
 	 * Used to prevent calling Tick inside Tick on the same thread, which causes deadlock.
 	 * This requires a single call to Tick runs on the same thread from start to end (a.k.a. it doesn't work with Fiber)
 	 */
-	uint32 LockOwningThreadID;
+	volatile uint32 LockOwningThreadID;
 
 	/** Whether the task has been cancelled because the update could not proceed or because the user called Abort(). */
 	bool bIsCancelled;
