Swapped a log entry to verbose as it was firing for all character, weapon and vehicle skeletal meshes loaded.

diff --git a/Engine/Source/Runtime/Engine/Private/SkeletalMeshLODRenderData.cpp b/Engine/Source/Runtime/Engine/Private/SkeletalMeshLODRenderData.cpp
--- a/Engine/Source/Runtime/Engine/Private/SkeletalMeshLODRenderData.cpp
+++ b/Engine/Source/Runtime/Engine/Private/SkeletalMeshLODRenderData.cpp
@@ -679,7 +679,7 @@ void FSkeletalMeshLODRenderData::Serialize(FArchive& Ar, UObject* Owner, int32 I
 	{
 		if (bNeedsCPUAccess && SkelMeshOwner)
 		{
-			UE_LOG(LogStaticMesh, Log, TEXT("[%s] Skeletal Mesh is marked for CPU read."), *SkelMeshOwner->GetName());
+			UE_LOG(LogStaticMesh, Verbose, TEXT("[%s] Skeletal Mesh is marked for CPU read."), *SkelMeshOwner->GetName());
 		}
 	}
 
