#ue4 - Added LastClientAckRealTime to character movement client data, to track the last time (undilated) that a client receives any Ack from the server for movement, including good moves and corrections. Also clarified the comments on some other timestamps.

diff --git a/Engine/Source/Runtime/Engine/Classes/GameFramework/CharacterMovementComponent.h b/Engine/Source/Runtime/Engine/Classes/GameFramework/CharacterMovementComponent.h
--- a/Engine/Source/Runtime/Engine/Classes/GameFramework/CharacterMovementComponent.h
+++ b/Engine/Source/Runtime/Engine/Classes/GameFramework/CharacterMovementComponent.h
@@ -2677,12 +2677,15 @@ public:
 	FNetworkPredictionData_Client_Character(const UCharacterMovementComponent& ClientMovement);
 	virtual ~FNetworkPredictionData_Client_Character();
 
-	/** Client timestamp of last time it sent a servermove() to the server.  Used for holding off on sending movement updates to save bandwidth. */
+	/** Client timestamp of last time it sent a servermove() to the server. This is an increasing timestamp from the owning UWorld. Used for holding off on sending movement updates to save bandwidth. */
 	float ClientUpdateTime;
 
-	/** Current TimeStamp for sending new Moves to the Server. */
+	/** Current TimeStamp for sending new Moves to the Server. This time resets to zero at a frequency of MinTimeBetweenTimeStampResets. */
 	float CurrentTimeStamp;
 
+	/** Last World timestamp (undilated, real time) at which we received a server ack for a move. This could be either a good move or a correction from the server. */
+	float LastReceivedAckRealTime;
+
 	TArray<FSavedMovePtr> SavedMoves;		// Buffered moves pending position updates, orderd oldest to newest. Moves that have been acked by the server are removed.
 	TArray<FSavedMovePtr> FreeMoves;		// freed moves, available for buffering
 	FSavedMovePtr PendingMove;				// PendingMove already processed on client - waiting to combine with next movement to reduce client to server bandwidth
@@ -2785,7 +2788,7 @@ public:
 	int32 GetSavedMoveIndex(float TimeStamp) const;
 
 	/** Ack a given move. This move will become LastAckedMove, SavedMoves will be adjusted to only contain unAcked moves. */
-	void AckMove(int32 AckedMoveIndex);
+	void AckMove(int32 AckedMoveIndex, UCharacterMovementComponent& CharacterMovementComponent);
 
 	/** Allocate a new saved move. Subclasses should override this if they want to use a custom move class. */
 	virtual FSavedMovePtr AllocateNewMove();
diff --git a/Engine/Source/Runtime/Engine/Private/Components/CharacterMovementComponent.cpp b/Engine/Source/Runtime/Engine/Private/Components/CharacterMovementComponent.cpp
--- a/Engine/Source/Runtime/Engine/Private/Components/CharacterMovementComponent.cpp
+++ b/Engine/Source/Runtime/Engine/Private/Components/CharacterMovementComponent.cpp
@@ -9047,7 +9047,8 @@ void UCharacterMovementComponent::ClientAdjustPosition_Implementation
 		}
 		return;
 	}
-	ClientData->AckMove(MoveIndex);
+
+	ClientData->AckMove(MoveIndex, *this);
 	
 	FVector WorldShiftedNewLocation;
 	//  Received Location is relative to dynamic base
@@ -9351,7 +9352,8 @@ void UCharacterMovementComponent::ClientAckGoodMove_Implementation(float TimeSta
 		}
 		return;
 	}
-	ClientData->AckMove(MoveIndex);
+
+	ClientData->AckMove(MoveIndex, *this);
 }
 
 void UCharacterMovementComponent::CapsuleTouched(UPrimitiveComponent* OverlappedComp, AActor* Other, UPrimitiveComponent* OtherComp, int32 OtherBodyIndex, bool bFromSweep, const FHitResult& SweepResult )
@@ -9998,6 +10000,7 @@ PRAGMA_DISABLE_DEPRECATION_WARNINGS // For deprecated members of FNetworkPredict
 FNetworkPredictionData_Client_Character::FNetworkPredictionData_Client_Character(const UCharacterMovementComponent& ClientMovement)
 	: ClientUpdateTime(0.f)
 	, CurrentTimeStamp(0.f)
+	, LastReceivedAckRealTime(0.f)
 	, PendingMove(NULL)
 	, LastAckedMove(NULL)
 	, MaxFreeMoveCount(96)
@@ -10146,7 +10149,7 @@ int32 FNetworkPredictionData_Client_Character::GetSavedMoveIndex(float TimeStamp
 	return INDEX_NONE;
 }
 
-void FNetworkPredictionData_Client_Character::AckMove(int32 AckedMoveIndex) 
+void FNetworkPredictionData_Client_Character::AckMove(int32 AckedMoveIndex, UCharacterMovementComponent& CharacterMovementComponent) 
 {
 	// It is important that we know the move exists before we go deleting outdated moves.
 	// Timestamps are not guaranteed to be increasing order all the time, since they can be reset!
@@ -10172,6 +10175,11 @@ void FNetworkPredictionData_Client_Character::AckMove(int32 AckedMoveIndex)
 		const bool bAllowShrinking = false;
 		SavedMoves.RemoveAt(0, AckedMoveIndex + 1, bAllowShrinking);
 	}
+
+	if (const UWorld* const World = CharacterMovementComponent.GetWorld())
+	{
+		LastReceivedAckRealTime = World->GetRealTimeSeconds();
+	}
 }
 
 PRAGMA_DISABLE_DEPRECATION_WARNINGS // For deprecated members of FNetworkPredictionData_Server_Character
