Made NavigationSystem's coordinate transforms' initialization not static #UE4

Static initializtion of FNavigationSystem::CoordTypeTransformsTo was causing issues on PS4 since it was triggered before FTransform::Initialize got initialized resulting in CoordTypeTransformsTo containing 'zero' transforms.
I've also made the related code a lot more readable.

diff --git a/Engine/Source/Runtime/Engine/Classes/AI/NavigationSystemBase.h b/Engine/Source/Runtime/Engine/Classes/AI/NavigationSystemBase.h
--- a/Engine/Source/Runtime/Engine/Classes/AI/NavigationSystemBase.h
+++ b/Engine/Source/Runtime/Engine/Classes/AI/NavigationSystemBase.h
@@ -146,7 +146,12 @@ namespace FNavigationSystem
 	ENGINE_API TSubclassOf<UNavAreaBase> GetDefaultWalkableArea();
 	ENGINE_API TSubclassOf<UNavAreaBase> GetDefaultObstacleArea();
 
+	/**	Retrieves the transform the Navigation System is using to convert coords
+	 *	from FromCoordType to ToCoordType */
+	ENGINE_API const FTransform& GetCoordTransform(const ENavigationCoordSystem::Type FromCoordType, const ENavigationCoordSystem::Type ToCoordType);
+	UE_DEPRECATED(4.22, "FNavigationSystem::GetCoordTransformTo is deprecated. Use FNavigationSystem::GetCoordTransform instead")
 	ENGINE_API const FTransform& GetCoordTransformTo(const ENavigationCoordSystem::Type CoordType);
+	UE_DEPRECATED(4.22, "FNavigationSystem::GetCoordTransformFrom is deprecated. Use FNavigationSystem::GetCoordTransform instead")
 	ENGINE_API const FTransform& GetCoordTransformFrom(const ENavigationCoordSystem::Type CoordType);
 
 	ENGINE_API bool WantsComponentChangeNotifies();
@@ -231,8 +236,14 @@ public:
 	INavigationDataInterface* GetMainNavData(int) { return nullptr; }
 
 protected:
-	
+	/**	Sets the Transform the Navigation System will use when converting from FromCoordType
+	 *	to ToCoordType
+	 *	@param bAddInverse if true (default) will also set coord transform in 
+	 *		the reverse order using Transform.Inverse() */
+	static void SetCoordTransform(const ENavigationCoordSystem::Type FromCoordType, const ENavigationCoordSystem::Type ToCoordType, const FTransform& Transform, bool bAddInverse = true);
+	UE_DEPRECATED(4.22, "FNavigationSystem::SetCoordTransformTo is deprecated. Use FNavigationSystem::SetCoordTransform instead")
 	static void SetCoordTransformTo(const ENavigationCoordSystem::Type CoordType, const FTransform& Transform);
+	UE_DEPRECATED(4.22, "FNavigationSystem::SetCoordTransformFrom is deprecated. Use FNavigationSystem::SetCoordTransform instead")
 	static void SetCoordTransformFrom(const ENavigationCoordSystem::Type CoordType, const FTransform& Transform);
 	static void SetWantsComponentChangeNotifies(const bool bEnable);
 	static void SetDefaultWalkableArea(TSubclassOf<UNavAreaBase> InAreaClass);
diff --git a/Engine/Source/Runtime/Engine/Private/AI/Navigation/NavigationModifier.cpp b/Engine/Source/Runtime/Engine/Private/AI/Navigation/NavigationModifier.cpp
--- a/Engine/Source/Runtime/Engine/Private/AI/Navigation/NavigationModifier.cpp
+++ b/Engine/Source/Runtime/Engine/Private/AI/Navigation/NavigationModifier.cpp
@@ -468,14 +468,12 @@ void FAreaNavModifier::SetConvex(const FVector* InPoints, const int32 FirstIndex
 	TArray<FVector, TInlineAllocator<MaxConvexPoints>> HullVertices;
 	HullVertices.Empty(MaxConvexPoints);
 
-	const FTransform& NavCoordTransform = FNavigationSystem::GetCoordTransformTo(CoordType);
+	// we store verts in unreal coords
+	const FTransform& TotalTransform = LocalToWorld * FNavigationSystem::GetCoordTransform(CoordType, ENavigationCoordSystem::Unreal);
 
 	for (int32 i = FirstIndex; i < LastIndex; i++)
 	{
-		//const FVector Point = (CoordType == ENavigationCoordSystem::Recast) ? Recast2UnrealPoint(InPoints[i]) : InPoints[i];
-		const FVector Point = NavCoordTransform.TransformPosition(InPoints[i]);
-
-		FVector TransformedPoint = LocalToWorld.TransformPosition(Point);
+		FVector TransformedPoint = TotalTransform.TransformPosition(InPoints[i]);
 		ConvexData.MinZ = FMath::Min( ConvexData.MinZ, TransformedPoint.Z );
 		ConvexData.MaxZ = FMath::Max( ConvexData.MaxZ, TransformedPoint.Z );
 		TransformedPoint.Z = 0.f;
diff --git a/Engine/Source/Runtime/Engine/Private/AI/NavigationSystemBase.cpp b/Engine/Source/Runtime/Engine/Private/AI/NavigationSystemBase.cpp
--- a/Engine/Source/Runtime/Engine/Private/AI/NavigationSystemBase.cpp
+++ b/Engine/Source/Runtime/Engine/Private/AI/NavigationSystemBase.cpp
@@ -207,17 +207,38 @@ namespace FNavigationSystem
 	void UpdateLevelCollision(ULevel& Level) { Delegates.UpdateLevelCollision.Execute(Level); }
 #endif // WITH_EDITOR
 
-	FTransform CoordTypeTransformsTo[ENavigationCoordSystem::MAX] = { FTransform::Identity, FTransform::Identity };
-	FTransform CoordTypeTransformsFrom[ENavigationCoordSystem::MAX] = { FTransform::Identity, FTransform::Identity };
+	struct FCoordTransforms
+	{
+		FTransform& Get(const ENavigationCoordSystem::Type FromCoordType, const ENavigationCoordSystem::Type ToCoordType)
+		{
+			static FTransform CoordTypeTransforms[ENavigationCoordSystem::MAX][ENavigationCoordSystem::MAX] = {
+				{FTransform::Identity, FTransform::Identity}
+				, {FTransform::Identity, FTransform::Identity}
+			};
+
+			return CoordTypeTransforms[uint8(FromCoordType)][uint8(ToCoordType)];
+		}
+	};
+
+	FCoordTransforms& GetCoordTypeTransforms()
+	{
+		static FCoordTransforms CoordTypeTransforms;
+		return CoordTypeTransforms;
+	}
 
 	const FTransform& GetCoordTransformTo(const ENavigationCoordSystem::Type CoordType)
 	{
-		return CoordTypeTransformsTo[uint8(CoordType)];
+		return GetCoordTransform(ENavigationCoordSystem::Unreal, CoordType);
 	}
 
 	const FTransform& GetCoordTransformFrom(const ENavigationCoordSystem::Type CoordType)
 	{
-		return CoordTypeTransformsTo[uint8(CoordType)];
+		return GetCoordTransform(CoordType, ENavigationCoordSystem::Unreal);
+	}
+
+	const FTransform& GetCoordTransform(const ENavigationCoordSystem::Type FromCoordType, const ENavigationCoordSystem::Type ToCoordType)
+	{
+		return GetCoordTypeTransforms().Get(FromCoordType, ToCoordType);
 	}
 
 	UWorld* GetWorldFromContextObject(UObject* WorldContextObject)
@@ -333,12 +354,21 @@ UNavigationSystem::UNavigationSystem(const FObjectInitializer& ObjectInitializer
 
 void UNavigationSystemBase::SetCoordTransformTo(const ENavigationCoordSystem::Type CoordType, const FTransform& Transform)
 {
-	FNavigationSystem::CoordTypeTransformsTo[uint8(CoordType)] = Transform;
+	SetCoordTransform(ENavigationCoordSystem::Unreal, CoordType, Transform);
 }
 
 void UNavigationSystemBase::SetCoordTransformFrom(const ENavigationCoordSystem::Type CoordType, const FTransform& Transform)
 {
-	FNavigationSystem::CoordTypeTransformsTo[uint8(CoordType)] = Transform;
+	SetCoordTransform(CoordType, ENavigationCoordSystem::Unreal, Transform);
+}
+
+void UNavigationSystemBase::SetCoordTransform(const ENavigationCoordSystem::Type FromCoordType, const ENavigationCoordSystem::Type ToCoordType, const FTransform& Transform, bool bAddInverse)
+{
+	FNavigationSystem::GetCoordTypeTransforms().Get(FromCoordType, ToCoordType) = Transform;
+	if (bAddInverse)
+	{
+		FNavigationSystem::GetCoordTypeTransforms().Get(ToCoordType, FromCoordType) = Transform.Inverse();
+	}
 }
 
 void UNavigationSystemBase::SetWantsComponentChangeNotifies(const bool bEnable)
diff --git a/Engine/Source/Runtime/Engine/Public/AI/NavigationModifier.h b/Engine/Source/Runtime/Engine/Public/AI/NavigationModifier.h
--- a/Engine/Source/Runtime/Engine/Public/AI/NavigationModifier.h
+++ b/Engine/Source/Runtime/Engine/Public/AI/NavigationModifier.h
@@ -59,7 +59,7 @@ namespace ENavigationCoordSystem
 	enum Type
 	{
 		Unreal,
-		Recast,
+		Navigation,
 
 		MAX
 	};
@@ -142,6 +142,7 @@ protected:
 	uint8 bIsLowAreaModifier : 1;
 
 	void Init(const TSubclassOf<UNavAreaBase> InAreaClass);
+	/** @param CoordType specifies which coord system the input data is in */
 	void SetConvex(const FVector* InPoints, const int32 FirstIndex, const int32 LastIndex, ENavigationCoordSystem::Type CoordType, const FTransform& LocalToWorld);
 	void SetBox(const FBox& Box, const FTransform& LocalToWorld);
 };
diff --git a/Engine/Source/Runtime/NavigationSystem/Private/NavigationSystem.cpp b/Engine/Source/Runtime/NavigationSystem/Private/NavigationSystem.cpp
--- a/Engine/Source/Runtime/NavigationSystem/Private/NavigationSystem.cpp
+++ b/Engine/Source/Runtime/NavigationSystem/Private/NavigationSystem.cpp
@@ -416,7 +416,7 @@ UNavigationSystemV1::UNavigationSystemV1(const FObjectInitializer& ObjectInitial
 		SetDefaultObstacleArea(UNavArea_Obstacle::StaticClass());
 		
 		const FTransform RecastToUnrealTransfrom(Recast2UnrealMatrix());
-		SetCoordTransformFrom(ENavigationCoordSystem::Recast, RecastToUnrealTransfrom);
+		SetCoordTransformFrom(ENavigationCoordSystem::Navigation, RecastToUnrealTransfrom);
 	}
 
 #if WITH_EDITOR
