PR #5544: Add `GetDisableSplitscreenOverride` to viewport (Contributed by [REDACTED_PERSON])

diff --git a/Engine/Source/Runtime/Engine/Classes/Engine/GameViewportClient.h b/Engine/Source/Runtime/Engine/Classes/Engine/GameViewportClient.h
--- a/Engine/Source/Runtime/Engine/Classes/Engine/GameViewportClient.h
+++ b/Engine/Source/Runtime/Engine/Classes/Engine/GameViewportClient.h
@@ -370,6 +370,12 @@ public:
 	/** Allows game code to disable splitscreen (useful when in menus) */
 	void SetDisableSplitscreenOverride( const bool bDisabled );
 
+	/** Determines whether splitscreen is forced to be turned off */
+	bool GetDisableSplitscreenOverride() const
+	{
+		return bDisableSplitScreenOverride;
+	}
+
 	/** called before rending subtitles to allow the game viewport to determine the size of the subtitle area
 	 * @param Min top left bounds of subtitle region (0 to 1)
 	 * @param Max bottom right bounds of subtitle region (0 to 1)
