Potential fix for crash in anim evaluation due to overwriting bonecontainer while worker thread could still be active

diff --git a/Engine/Source/Runtime/Engine/Classes/Components/SkeletalMeshComponent.h b/Engine/Source/Runtime/Engine/Classes/Components/SkeletalMeshComponent.h
--- a/Engine/Source/Runtime/Engine/Classes/Components/SkeletalMeshComponent.h
+++ b/Engine/Source/Runtime/Engine/Classes/Components/SkeletalMeshComponent.h
@@ -1288,7 +1288,7 @@ public:
 	 *
 	 * @param	LODIndex	Index of LOD [0-(MaxLOD-1)]
 	 */
-	void RecalcRequiredBones(int32 LODIndex);
+	bool RecalcRequiredBones(int32 LODIndex);
 
 	/** Computes the required bones in this SkeletalMeshComponent based on current SkeletalMesh, LOD and PhysicsAsset
 	  * @param	LODIndex	Index of LOD [0-(MaxLOD-1)]
@@ -1301,6 +1301,21 @@ public:
 	*/
 	void RecalcRequiredCurves();
 
+private:
+	/**
+	 * Recalculates the RequiredBones array in this SkeletalMeshComponent based on current SkeletalMesh, LOD and PhysicsAsset.
+	 * Only modifies members of the SMC (safe to call while parallel task is in flight
+	 *
+	 * @param	LODIndex	Index of LOD [0-(MaxLOD-1)]
+	 */
+	bool RecalcRequiredBonesInternalSMC(int32 LODIndex);
+
+	/**
+	 * Causes the anim instances this component owns to recalc their required bones.
+	 * RecalcRequiredBonesInternalSMC must be called first
+	 */
+	void RecalcRequiredBonesInternalAnimInstances();
+
 public:
 	//~ Begin UObject Interface.
 	virtual void Serialize(FArchive& Ar) override;
diff --git a/Engine/Source/Runtime/Engine/Private/Components/SkeletalMeshComponent.cpp b/Engine/Source/Runtime/Engine/Private/Components/SkeletalMeshComponent.cpp
--- a/Engine/Source/Runtime/Engine/Private/Components/SkeletalMeshComponent.cpp
+++ b/Engine/Source/Runtime/Engine/Private/Components/SkeletalMeshComponent.cpp
@@ -1643,18 +1643,33 @@ void USkeletalMeshComponent::ComputeRequiredBones(TArray<FBoneIndexType>& OutReq
 	FAnimationRuntime::EnsureParentsPresent(OutFillComponentSpaceTransformsRequiredBones, SkeletalMesh->RefSkeleton);
 }
 
-void USkeletalMeshComponent::RecalcRequiredBones(int32 LODIndex)
+bool USkeletalMeshComponent::RecalcRequiredBones(int32 LODIndex)
+{
+	bool bSuccess = RecalcRequiredBonesInternalSMC(LODIndex);
+	RecalcRequiredBonesInternalAnimInstances();
+	return bSuccess;
+}
+
+bool USkeletalMeshComponent::RecalcRequiredBonesInternalSMC(int32 LODIndex)
 {
 	if (!SkeletalMesh)
 	{
-		return;
+		return false;
 	}
 
 	ComputeRequiredBones(RequiredBones, FillComponentSpaceTransformsRequiredBones, LODIndex, /*bIgnorePhysicsAsset=*/ false);
 
 	BoneSpaceTransforms = SkeletalMesh->RefSkeleton.GetRefBonePose();
 
-	// make sure animation requiredBone to mark as dirty
+	// Invalidate cached bones.
+	CachedBoneSpaceTransforms.Empty();
+	CachedComponentSpaceTransforms.Empty();
+	CachedCurve.Empty();
+	return true;
+}
+
+void USkeletalMeshComponent::RecalcRequiredBonesInternalAnimInstances()
+{
 	if (AnimScriptInstance)
 	{
 		AnimScriptInstance->RecalcRequiredBones();
@@ -1674,11 +1689,6 @@ void USkeletalMeshComponent::RecalcRequiredBones(int32 LODIndex)
 	// this should always happen
 	MarkRequiredCurveUpToDate();
 	bRequiredBonesUpToDate = true;
-
-	// Invalidate cached bones.
-	CachedBoneSpaceTransforms.Empty();
-	CachedComponentSpaceTransforms.Empty();
-	CachedCurve.Empty();
 }
 
 void USkeletalMeshComponent::MarkRequiredCurveUpToDate()
@@ -1948,13 +1958,7 @@ void USkeletalMeshComponent::RefreshBoneTransforms(FActorComponentTickFunction*
 	if (!bRequiredBonesUpToDate)
 	{
 		QUICK_SCOPE_CYCLE_COUNTER(STAT_USkeletalMeshComponent_RefreshBoneTransforms_RecalcRequiredBones);
-		RecalcRequiredBones(PredictedLODLevel);
-	}
-	// if curves have to be refreshed
-	else if (!AreRequiredCurvesUpToDate())
-	{
-		QUICK_SCOPE_CYCLE_COUNTER(STAT_USkeletalMeshComponent_RefreshBoneTransforms_RecalcRequiredCurves);
-		RecalcRequiredCurves();
+		RecalcRequiredBonesInternalSMC(PredictedLODLevel);
 	}
 
 	const bool bCachedShouldUseUpdateRateOptimizations = ShouldUseUpdateRateOptimizations() && AnimUpdateRateParams != nullptr;
@@ -1969,17 +1973,19 @@ void USkeletalMeshComponent::RefreshBoneTransforms(FActorComponentTickFunction*
 
 
 	TArray<uint16> const* CurrentAnimCurveUIDFinder = (AnimScriptInstance) ? &AnimScriptInstance->GetRequiredBones().GetUIDToArrayLookupTable() : nullptr;
+	const bool bAnimInstanceHasCurveUIDList = CurrentAnimCurveUIDFinder != nullptr;
+
 	const int32 CurrentCurveCount = (CurrentAnimCurveUIDFinder) ? FBlendedCurve::GetValidElementCount(CurrentAnimCurveUIDFinder) : 0;
 
 	const bool bInvalidCachedCurve = bDoEvaluationRateOptimization && 
-									CurrentAnimCurveUIDFinder != nullptr &&
+									 bAnimInstanceHasCurveUIDList &&
 									(CachedCurve.UIDToArrayIndexLUT != CurrentAnimCurveUIDFinder || CachedCurve.Num() != CurrentCurveCount);
 
 	const bool bShouldDoEvaluation = !bDoEvaluationRateOptimization || bInvalidCachedBones || bInvalidCachedCurve || (bExternalTickRateControlled && bExternalUpdate) || (bCachedShouldUseUpdateRateOptimizations && !AnimUpdateRateParams->ShouldSkipEvaluation());
 
 	const bool bShouldInterpolateSkippedFrames = (bExternalTickRateControlled && bExternalInterpolate) || (bCachedShouldUseUpdateRateOptimizations && AnimUpdateRateParams->ShouldInterpolateSkippedFrames());
 
-	const bool bShouldDoInterpolation = TickFunction != nullptr && bDoEvaluationRateOptimization && !bInvalidCachedBones && bShouldInterpolateSkippedFrames && CurrentAnimCurveUIDFinder != nullptr;
+	const bool bShouldDoInterpolation = TickFunction != nullptr && bDoEvaluationRateOptimization && !bInvalidCachedBones && bShouldInterpolateSkippedFrames && bAnimInstanceHasCurveUIDList;
 
 	const bool bShouldDoParallelInterpolation = bShouldDoInterpolation && CVarUseParallelAnimationInterpolation.GetValueOnGameThread() == 1;
 
@@ -1998,6 +2004,22 @@ void USkeletalMeshComponent::RefreshBoneTransforms(FActorComponentTickFunction*
 		return;
 	}
 
+	// Recalculate the RequiredBones array, if necessary
+	if (!bRequiredBonesUpToDate)
+	{
+		QUICK_SCOPE_CYCLE_COUNTER(STAT_USkeletalMeshComponent_RefreshBoneTransforms_RecalcRequiredBones);
+		RecalcRequiredBonesInternalAnimInstances();
+	}
+	// if curves have to be refreshed
+	else if (!AreRequiredCurvesUpToDate())
+	{
+		QUICK_SCOPE_CYCLE_COUNTER(STAT_USkeletalMeshComponent_RefreshBoneTransforms_RecalcRequiredCurves);
+		RecalcRequiredCurves();
+	}
+	
+	//Refresh UID list now curves have been updated
+	CurrentAnimCurveUIDFinder = (AnimScriptInstance) ? &AnimScriptInstance->GetRequiredBones().GetUIDToArrayLookupTable() : nullptr;
+
 	AnimEvaluationContext.SkeletalMesh = SkeletalMesh;
 	AnimEvaluationContext.AnimInstance = AnimScriptInstance;
 	AnimEvaluationContext.PostProcessAnimInstance = PostProcessAnimInstance;
