The files added in ShaderDebugInfo are now put into a unique folder based on the StaticParameters. This allows you to see and compare shaders from multiple permutations.
Also, in the d3d11 compiler, A description of the different parameters are appended at the end, for information

diff --git a/Engine/Source/Developer/Windows/ShaderFormatD3D/Private/D3DShaderCompiler.cpp b/Engine/Source/Developer/Windows/ShaderFormatD3D/Private/D3DShaderCompiler.cpp
--- a/Engine/Source/Developer/Windows/ShaderFormatD3D/Private/D3DShaderCompiler.cpp
+++ b/Engine/Source/Developer/Windows/ShaderFormatD3D/Private/D3DShaderCompiler.cpp
@@ -970,7 +970,9 @@ static bool CompileAndProcessD3DShader(FString& PreprocessedShaderSource, const
 				Line += TEXT("#if 0 /*DIRECT COMPILE*/\n");
 				Line += CreateShaderCompilerWorkerDirectCommandLine(Input);
 				Line += TEXT("\n#endif /*DIRECT COMPILE*/\n");
-
+				Line += TEXT("//");
+				Line += Input.DebugDescription;
+				Line += TEXT("\n");
 				FileWriter->Serialize(TCHAR_TO_ANSI(*Line), Line.Len());
 			}
 			FileWriter->Close();
diff --git a/Engine/Source/Runtime/Engine/Private/Materials/MaterialShader.cpp b/Engine/Source/Runtime/Engine/Private/Materials/MaterialShader.cpp
--- a/Engine/Source/Runtime/Engine/Private/Materials/MaterialShader.cpp
+++ b/Engine/Source/Runtime/Engine/Private/Materials/MaterialShader.cpp
@@ -861,7 +861,9 @@ FShaderCompileJob* FMaterialShaderType::BeginCompileShader(
 	FShaderCompilerEnvironment* MaterialEnvironment,
 	const FShaderPipelineType* ShaderPipeline,
 	EShaderPlatform Platform,
-	TArray<FShaderCommonCompileJob*>& NewJobs
+	TArray<FShaderCommonCompileJob*>& NewJobs,
+	const FString& DebugDescription,
+	const FString& DebugExtension
 	)
 {
 	FShaderCompileJob* NewJob = new FShaderCompileJob(ShaderMapId, nullptr, this, /* PermutationId = */ 0);
@@ -890,7 +892,10 @@ FShaderCompileJob* FMaterialShaderType::BeginCompileShader(
 		GetFunctionName(),
 		FShaderTarget(GetFrequency(),Platform),
 		NewJob,
-		NewJobs
+		NewJobs,
+		true,
+		DebugDescription,
+		DebugExtension
 		);
 	return NewJob;
 }
@@ -902,7 +907,9 @@ void FMaterialShaderType::BeginCompileShaderPipeline(
 	FShaderCompilerEnvironment* MaterialEnvironment,
 	const FShaderPipelineType* ShaderPipeline,
 	const TArray<FMaterialShaderType*>& ShaderStages,
-	TArray<FShaderCommonCompileJob*>& NewJobs)
+	TArray<FShaderCommonCompileJob*>& NewJobs,
+	const FString& DebugDescription,
+	const FString& DebugExtension)
 {
 	check(ShaderStages.Num() > 0);
 	check(ShaderPipeline);
@@ -913,7 +920,7 @@ void FMaterialShaderType::BeginCompileShaderPipeline(
 	for (int32 Index = 0; Index < ShaderStages.Num(); ++Index)
 	{
 		auto* ShaderStage = ShaderStages[Index];
-		ShaderStage->BeginCompileShader(ShaderMapId, Material, MaterialEnvironment, ShaderPipeline, Platform, NewPipelineJob->StageJobs);
+		ShaderStage->BeginCompileShader(ShaderMapId, Material, MaterialEnvironment, ShaderPipeline, Platform, NewPipelineJob->StageJobs, DebugDescription, DebugExtension);
 	}
 
 	NewJobs.Add(NewPipelineJob);
@@ -1448,6 +1455,16 @@ void FMaterialShaderMap::Compile(
 			);
   
 			UE_LOG(LogShaders, Display, TEXT("	%s"), *DebugDescription);
+
+			FSHA1 IdParameterSetHash;
+			IdParameterSetHash.Reset();
+			IdParameterSet.UpdateHash(IdParameterSetHash);
+			IdParameterSetHash.Final();
+			uint32* Hash = (uint32*)&IdParameterSetHash.m_digest[0];
+			FString DebugExtension = FString::Printf( TEXT("_%08x%08x"), Hash[0], Hash[1]);
+#else
+			FString DebugExtension = "";
+			FString DebugDescription = "";
 #endif
 
 			uint32 NumShaders = 0;
@@ -1491,7 +1508,9 @@ void FMaterialShaderMap::Compile(
 						Material,
 						MaterialEnvironment,
 						InPlatform,
-						NewJobs
+						NewJobs,
+						DebugDescription, 
+						DebugExtension
 						);
 					NumShaders += MeshShaders;
 					if (MeshShaders > 0)
@@ -1523,7 +1542,8 @@ void FMaterialShaderMap::Compile(
 							MaterialEnvironment,
 							nullptr,
 							InPlatform,
-							NewJobs
+							NewJobs,
+							DebugDescription, DebugExtension
 							);
 						check(!SharedShaderJobs.Find(ShaderType));
 						SharedShaderJobs.Add(ShaderType, Job);
@@ -1562,7 +1582,7 @@ void FMaterialShaderMap::Compile(
 						if (Pipeline->ShouldOptimizeUnusedOutputs(InPlatform))
 						{
 							NumShaders += ShaderStagesToCompile.Num();
-							FMaterialShaderType::BeginCompileShaderPipeline(CompilingId, InPlatform, Material, MaterialEnvironment, Pipeline, ShaderStagesToCompile, NewJobs);
+							FMaterialShaderType::BeginCompileShaderPipeline(CompilingId, InPlatform, Material, MaterialEnvironment, Pipeline, ShaderStagesToCompile, NewJobs, DebugDescription, DebugExtension);
 						}
 						else
 						{
diff --git a/Engine/Source/Runtime/Engine/Private/Materials/MeshMaterialShader.cpp b/Engine/Source/Runtime/Engine/Private/Materials/MeshMaterialShader.cpp
--- a/Engine/Source/Runtime/Engine/Private/Materials/MeshMaterialShader.cpp
+++ b/Engine/Source/Runtime/Engine/Private/Materials/MeshMaterialShader.cpp
@@ -48,7 +48,10 @@ FShaderCompileJob* FMeshMaterialShaderType::BeginCompileShader(
 	FShaderCompilerEnvironment* MaterialEnvironment,
 	FVertexFactoryType* VertexFactoryType,
 	const FShaderPipelineType* ShaderPipeline,
-	TArray<FShaderCommonCompileJob*>& NewJobs
+	TArray<FShaderCommonCompileJob*>& NewJobs,
+	FString DebugDescription,
+	FString DebugExtension
+
 	)
 {
 	FShaderCompileJob* NewJob = new FShaderCompileJob(ShaderMapId, VertexFactoryType, this, /* PermutationId = */ 0);
@@ -84,7 +87,9 @@ FShaderCompileJob* FMeshMaterialShaderType::BeginCompileShader(
 		FShaderTarget(GetFrequency(),Platform),
 		NewJob,
 		NewJobs,
-		bAllowDevelopmentShaderCompile
+		bAllowDevelopmentShaderCompile,
+		DebugDescription,
+		DebugExtension
 		);
 	return NewJob;
 }
@@ -97,7 +102,10 @@ void FMeshMaterialShaderType::BeginCompileShaderPipeline(
 	FVertexFactoryType* VertexFactoryType,
 	const FShaderPipelineType* ShaderPipeline,
 	const TArray<FMeshMaterialShaderType*>& ShaderStages,
-	TArray<FShaderCommonCompileJob*>& NewJobs)
+	TArray<FShaderCommonCompileJob*>& NewJobs,
+	FString DebugDescription,
+	FString DebugExtension
+	)
 {
 	check(ShaderStages.Num() > 0);
 	check(ShaderPipeline);
@@ -108,7 +116,7 @@ void FMeshMaterialShaderType::BeginCompileShaderPipeline(
 	for (int32 Index = 0; Index < ShaderStages.Num(); ++Index)
 	{
 		auto* ShaderStage = ShaderStages[Index];
-		ShaderStage->BeginCompileShader(ShaderMapId, Platform, Material, MaterialEnvironment, VertexFactoryType, ShaderPipeline, NewPipelineJob->StageJobs);
+		ShaderStage->BeginCompileShader(ShaderMapId, Platform, Material, MaterialEnvironment, VertexFactoryType, ShaderPipeline, NewPipelineJob->StageJobs, DebugDescription, DebugExtension);
 	}
 
 	NewJobs.Add(NewPipelineJob);
@@ -184,7 +192,9 @@ uint32 FMeshMaterialShaderMap::BeginCompile(
 	const FMaterial* Material,
 	FShaderCompilerEnvironment* MaterialEnvironment,
 	EShaderPlatform InPlatform,
-	TArray<FShaderCommonCompileJob*>& NewJobs
+	TArray<FShaderCommonCompileJob*>& NewJobs,
+	FString DebugDescription,
+	FString DebugExtension
 	)
 {
 	if (!VertexFactoryType)
@@ -218,7 +228,9 @@ uint32 FMeshMaterialShaderMap::BeginCompile(
 					MaterialEnvironment,
 					VertexFactoryType,
 					nullptr,
-					NewJobs
+					NewJobs,
+					DebugDescription, 
+					DebugExtension
 					);
 				check(!SharedShaderJobs.Find(ShaderType));
 				SharedShaderJobs.Add(ShaderType, Job);
@@ -268,7 +280,7 @@ uint32 FMeshMaterialShaderMap::BeginCompile(
 					}
 
 					// Make a pipeline job with all the stages
-					FMeshMaterialShaderType::BeginCompileShaderPipeline(ShaderMapId, InPlatform, Material, MaterialEnvironment, VertexFactoryType, Pipeline, ShaderStagesToCompile, NewJobs);
+					FMeshMaterialShaderType::BeginCompileShaderPipeline(ShaderMapId, InPlatform, Material, MaterialEnvironment, VertexFactoryType, Pipeline, ShaderStagesToCompile, NewJobs, DebugDescription, DebugExtension);
 				}
 				else
 				{
diff --git a/Engine/Source/Runtime/Engine/Private/ShaderCompiler/ShaderCompiler.cpp b/Engine/Source/Runtime/Engine/Private/ShaderCompiler/ShaderCompiler.cpp
--- a/Engine/Source/Runtime/Engine/Private/ShaderCompiler/ShaderCompiler.cpp
+++ b/Engine/Source/Runtime/Engine/Private/ShaderCompiler/ShaderCompiler.cpp
@@ -2824,7 +2824,10 @@ void GlobalBeginCompileShader(
 	FShaderTarget Target,
 	FShaderCompileJob* NewJob,
 	TArray<FShaderCommonCompileJob*>& NewJobs,
-	bool bAllowDevelopmentShaderCompile
+	bool bAllowDevelopmentShaderCompile,
+	const FString& DebugDescription,
+	const FString& DebugExtension
+
 	)
 {
 	COOK_STAT(ShaderCompilerCookStats::GlobalBeginCompileShaderCalls++);
@@ -2841,6 +2844,7 @@ void GlobalBeginCompileShader(
 	Input.DumpDebugInfoRootPath = GShaderCompilingManager->GetAbsoluteShaderDebugInfoDirectory() / Input.ShaderFormat.ToString();
 	// asset material name or "Global"
 	Input.DebugGroupName = DebugGroupName;
+	Input.DebugDescription = DebugDescription;
 
 	if (ShaderType->GetRootParametersMetadata())
 	{
@@ -2955,7 +2959,7 @@ void GlobalBeginCompileShader(
 	// Setup the debug info path if requested, or if this is a global shader and shader development mode is enabled
 	if (GDumpShaderDebugInfo != 0)
 	{
-		Input.DumpDebugInfoPath = Input.DumpDebugInfoRootPath / Input.DebugGroupName;
+		Input.DumpDebugInfoPath = Input.DumpDebugInfoRootPath / Input.DebugGroupName + DebugExtension,
 		
 		// Sanitize the name to be used as a path
 		// List mostly comes from set of characters not allowed by windows in a path.  Just try to rename a file and type one of these for the list.
diff --git a/Engine/Source/Runtime/Engine/Public/MaterialShaderType.h b/Engine/Source/Runtime/Engine/Public/MaterialShaderType.h
--- a/Engine/Source/Runtime/Engine/Public/MaterialShaderType.h
+++ b/Engine/Source/Runtime/Engine/Public/MaterialShaderType.h
@@ -125,7 +125,9 @@ public:
 		FShaderCompilerEnvironment* MaterialEnvironment,
 		const FShaderPipelineType* ShaderPipeline,
 		EShaderPlatform Platform,
-		TArray<FShaderCommonCompileJob*>& NewJobs
+		TArray<FShaderCommonCompileJob*>& NewJobs,
+		const FString& DebugDescription,
+		const FString& DebugExtension
 		);
 
 	static void BeginCompileShaderPipeline(
@@ -135,7 +137,10 @@ public:
 		FShaderCompilerEnvironment* MaterialEnvironment,
 		const FShaderPipelineType* ShaderPipeline,
 		const TArray<FMaterialShaderType*>& ShaderStages,
-		TArray<FShaderCommonCompileJob*>& NewJobs);
+		TArray<FShaderCommonCompileJob*>& NewJobs,
+		const FString& DebugDescription,
+		const FString& DebugExtension
+	);
 
 	/**
 	 * Either creates a new instance of this type or returns an equivalent existing shader.
diff --git a/Engine/Source/Runtime/Engine/Public/MaterialShared.h b/Engine/Source/Runtime/Engine/Public/MaterialShared.h
--- a/Engine/Source/Runtime/Engine/Public/MaterialShared.h
+++ b/Engine/Source/Runtime/Engine/Public/MaterialShared.h
@@ -661,7 +661,9 @@ public:
 		const FMaterial* Material,
 		FShaderCompilerEnvironment* MaterialEnvironment,
 		EShaderPlatform Platform,
-		TArray<FShaderCommonCompileJob*>& NewJobs
+		TArray<FShaderCommonCompileJob*>& NewJobs,
+		FString DebugDescription,
+		FString DebugExtension
 		);
 
 	/**
diff --git a/Engine/Source/Runtime/Engine/Public/MeshMaterialShaderType.h b/Engine/Source/Runtime/Engine/Public/MeshMaterialShaderType.h
--- a/Engine/Source/Runtime/Engine/Public/MeshMaterialShaderType.h
+++ b/Engine/Source/Runtime/Engine/Public/MeshMaterialShaderType.h
@@ -101,7 +101,9 @@ public:
 		FShaderCompilerEnvironment* MaterialEnvironment,
 		FVertexFactoryType* VertexFactoryType,
 		const FShaderPipelineType* ShaderPipeline,
-		TArray<FShaderCommonCompileJob*>& NewJobs
+		TArray<FShaderCommonCompileJob*>& NewJobs,
+		FString DebugDescription,
+		FString DebugExtension
 		);
 
 	static void BeginCompileShaderPipeline(
@@ -112,7 +114,9 @@ public:
 		FVertexFactoryType* VertexFactoryType,
 		const FShaderPipelineType* ShaderPipeline,
 		const TArray<FMeshMaterialShaderType*>& ShaderStages,
-		TArray<FShaderCommonCompileJob*>& NewJobs
+		TArray<FShaderCommonCompileJob*>& NewJobs,
+		FString DebugDescription,
+		FString DebugExtension
 		);
 
 	/**
diff --git a/Engine/Source/Runtime/Engine/Public/ShaderCompiler.h b/Engine/Source/Runtime/Engine/Public/ShaderCompiler.h
--- a/Engine/Source/Runtime/Engine/Public/ShaderCompiler.h
+++ b/Engine/Source/Runtime/Engine/Public/ShaderCompiler.h
@@ -611,7 +611,9 @@ extern ENGINE_API void GlobalBeginCompileShader(
 	FShaderTarget Target,
 	FShaderCompileJob* NewJob,
 	TArray<FShaderCommonCompileJob*>& NewJobs,
-	bool bAllowDevelopmentShaderCompile = true
+	bool bAllowDevelopmentShaderCompile = true,
+	const FString& DebugDescription = "",
+	const FString& DebugExtension = ""
 	);
 
 /** Implementation of the 'recompileshaders' console command.  Recompiles shaders at runtime based on various criteria. */
diff --git a/Engine/Source/Runtime/RenderCore/Public/ShaderCore.h b/Engine/Source/Runtime/RenderCore/Public/ShaderCore.h
--- a/Engine/Source/Runtime/RenderCore/Public/ShaderCore.h
+++ b/Engine/Source/Runtime/RenderCore/Public/ShaderCore.h
@@ -558,6 +558,9 @@ struct FShaderCompilerInput
 	// materialname or "Global" "for debugging and better error messages
 	FString DebugGroupName;
 
+	// Description of the configuration used when compiling. 
+	FString DebugDescription;
+
 	// Compilation Environment
 	FShaderCompilerEnvironment Environment;
 	TRefCountPtr<FShaderCompilerEnvironment> SharedEnvironment;
@@ -699,6 +702,7 @@ struct FShaderCompilerInput
 		Ar << Input.DumpDebugInfoRootPath;
 		Ar << Input.DumpDebugInfoPath;
 		Ar << Input.DebugGroupName;
+		Ar << Input.DebugDescription;
 		Ar << Input.Environment;
 		Ar << Input.ExtraSettings;
 		Ar << Input.RootParameterBindings;
