UE4.22 - Replace most of ENQUEUE_UNIQUE_RENDER_COMMAND_ONEPARAMETER
#rnx

diff --git a/Engine/Plugins/Compositing/OpenColorIO/Source/OpenColorIO/Private/OpenColorIOShared.cpp b/Engine/Plugins/Compositing/OpenColorIO/Source/OpenColorIO/Private/OpenColorIOShared.cpp
--- a/Engine/Plugins/Compositing/OpenColorIO/Source/OpenColorIO/Private/OpenColorIOShared.cpp
+++ b/Engine/Plugins/Compositing/OpenColorIO/Source/OpenColorIO/Private/OpenColorIOShared.cpp
@@ -119,9 +119,9 @@ void FOpenColorIOTransformResource::ReleaseShaderMap()
 	{
 		GameThreadShaderMap = nullptr;
 
-		ENQUEUE_UNIQUE_RENDER_COMMAND_ONEPARAMETER(
-			ReleaseShaderMap,
-			FOpenColorIOTransformResource*, ColorTransform, this,
+		FOpenColorIOTransformResource* ColorTransform = this;
+		ENQUEUE_RENDER_COMMAND(ReleaseShaderMap)(
+			[ColorTransform](FRHICommandListImmediate& RHICmdList)
 			{
 				ColorTransform->SetRenderingThreadShaderMap(nullptr);
 			});
diff --git a/Engine/Plugins/Experimental/GeometryCache/Source/GeometryCache/Private/GeometryCacheSceneProxy.cpp b/Engine/Plugins/Experimental/GeometryCache/Source/GeometryCache/Private/GeometryCacheSceneProxy.cpp
--- a/Engine/Plugins/Experimental/GeometryCache/Source/GeometryCache/Private/GeometryCacheSceneProxy.cpp
+++ b/Engine/Plugins/Experimental/GeometryCache/Source/GeometryCache/Private/GeometryCacheSceneProxy.cpp
@@ -133,19 +133,17 @@ FGeometryCacheSceneProxy::FGeometryCacheSceneProxy(UGeometryCacheComponent* Comp
 	{
 		// Update at least once after the scene proxy has been constructed
 		// Otherwise it is invisible until animation starts
-		ENQUEUE_UNIQUE_RENDER_COMMAND_ONEPARAMETER(
-			FGeometryCacheUpdateAnimation,
-			FGeometryCacheSceneProxy*, SceneProxy, this,
+		FGeometryCacheSceneProxy* SceneProxy = this;
+		ENQUEUE_RENDER_COMMAND(FGeometryCacheUpdateAnimation)(
+			[SceneProxy](FRHICommandListImmediate& RHICmdList)
 			{
 				SceneProxy->FrameUpdate();
 			});
 
-	#if RHI_RAYTRACING
-		if (IsRayTracingEnabled())
+#if RHI_RAYTRACING
 		{
-			ENQUEUE_UNIQUE_RENDER_COMMAND_ONEPARAMETER(
-				FGeometryCacheInitRayTracingGeometry,
-				FGeometryCacheSceneProxy*, SceneProxy, this,
+			ENQUEUE_RENDER_COMMAND(FGeometryCacheInitRayTracingGeometry)(
+				[SceneProxy](FRHICommandListImmediate& RHICmdList)
 				{
 					for (FGeomCacheTrackProxy* Section : SceneProxy->Tracks)
 					{
diff --git a/Engine/Plugins/Experimental/ImagePlate/Source/ImagePlate/Private/ImagePlateFileSequence.cpp b/Engine/Plugins/Experimental/ImagePlate/Source/ImagePlate/Private/ImagePlateFileSequence.cpp
--- a/Engine/Plugins/Experimental/ImagePlate/Source/ImagePlate/Private/ImagePlateFileSequence.cpp
+++ b/Engine/Plugins/Experimental/ImagePlate/Source/ImagePlate/Private/ImagePlateFileSequence.cpp
@@ -189,9 +189,8 @@ TFuture<void> FImagePlateSourceFrame::CopyTo(UTexture* DestinationTexture)
 	CommandData->SourceFrame = *this;
 	CommandData->DestinationTexture = DestinationTexture;
 
-	ENQUEUE_UNIQUE_RENDER_COMMAND_ONEPARAMETER(
-		CopySourceBufferToTexture,
-		FCommandDataPtr, CommandData, CommandData,
+	ENQUEUE_RENDER_COMMAND(CopySourceBufferToTexture)(
+		[CommandData](FRHICommandListImmediate& RHICmdList)
 		{
 			if (!CommandData->DestinationTexture->Resource || !CommandData->DestinationTexture->Resource->TextureRHI)
 			{
diff --git a/Engine/Plugins/FX/Niagara/Source/Niagara/Classes/NiagaraDataSet.h b/Engine/Plugins/FX/Niagara/Source/Niagara/Classes/NiagaraDataSet.h
--- a/Engine/Plugins/FX/Niagara/Source/Niagara/Classes/NiagaraDataSet.h
+++ b/Engine/Plugins/FX/Niagara/Source/Niagara/Classes/NiagaraDataSet.h
@@ -518,9 +518,9 @@ public:
 		}
 		else
 		{			
-			ENQUEUE_UNIQUE_RENDER_COMMAND_ONEPARAMETER(
-				ResetBuffersRT,
-				FNiagaraDataSet*, DataSet, this,
+			FNiagaraDataSet* DataSet = this;
+			ENQUEUE_RENDER_COMMAND(ResetBuffersRT)(
+				[DataSet](FRHICommandListImmediate& RHICmdList)
 				{
 					DataSet->ResetBuffersInternal();
 				}
diff --git a/Engine/Plugins/FX/Niagara/Source/Niagara/Classes/NiagaraScriptExecutionContext.h b/Engine/Plugins/FX/Niagara/Source/Niagara/Classes/NiagaraScriptExecutionContext.h
--- a/Engine/Plugins/FX/Niagara/Source/Niagara/Classes/NiagaraScriptExecutionContext.h
+++ b/Engine/Plugins/FX/Niagara/Source/Niagara/Classes/NiagaraScriptExecutionContext.h
@@ -141,9 +141,9 @@ struct FNiagaraComputeExecutionContext
 
 	void Reset()
 	{
-		ENQUEUE_UNIQUE_RENDER_COMMAND_ONEPARAMETER(
-			ResetRT,
-			FNiagaraComputeExecutionContext*, Context, this,
+		FNiagaraComputeExecutionContext* Context = this;
+		ENQUEUE_RENDER_COMMAND(ResetRT)(
+			[Context](FRHICommandListImmediate& RHICmdList)
 			{
 				Context->ResetInternal();
 			}
diff --git a/Engine/Plugins/FX/Niagara/Source/Niagara/Private/NiagaraEmitterInstance.cpp b/Engine/Plugins/FX/Niagara/Source/Niagara/Private/NiagaraEmitterInstance.cpp
--- a/Engine/Plugins/FX/Niagara/Source/Niagara/Private/NiagaraEmitterInstance.cpp
+++ b/Engine/Plugins/FX/Niagara/Source/Niagara/Private/NiagaraEmitterInstance.cpp
@@ -82,8 +82,9 @@ FNiagaraEmitterInstance::~FNiagaraEmitterInstance()
 				Batcher->Remove(GPUExecContext);
 			}
 
-			ENQUEUE_UNIQUE_RENDER_COMMAND_ONEPARAMETER(FDeleteComputeContextCommand,
-				FNiagaraComputeExecutionContext*, ExecContext, GPUExecContext,
+			FNiagaraComputeExecutionContext* ExecContext = GPUExecContext;
+			ENQUEUE_RENDER_COMMAND(FDeleteComputeContextCommand)(
+				[ExecContext](FRHICommandListImmediate& RHICmdList)
 				{
 					delete ExecContext;
 				});
@@ -93,8 +94,9 @@ FNiagaraEmitterInstance::~FNiagaraEmitterInstance()
 
 		if (ParticleDataSet != nullptr)
 		{
-			ENQUEUE_UNIQUE_RENDER_COMMAND_ONEPARAMETER(FDeleteParticleDataSetCommand,
-				FNiagaraDataSet*, DataSet, ParticleDataSet,
+			FNiagaraDataSet* DataSet = ParticleDataSet;
+			ENQUEUE_RENDER_COMMAND(FDeleteParticleDataSetCommand)(
+				[DataSet](FRHICommandListImmediate& RHICmdList)
 				{
 					delete DataSet;
 				});
diff --git a/Engine/Plugins/FX/Niagara/Source/Niagara/Private/NiagaraRenderer.cpp b/Engine/Plugins/FX/Niagara/Source/Niagara/Private/NiagaraRenderer.cpp
--- a/Engine/Plugins/FX/Niagara/Source/Niagara/Private/NiagaraRenderer.cpp
+++ b/Engine/Plugins/FX/Niagara/Source/Niagara/Private/NiagaraRenderer.cpp
@@ -80,9 +80,9 @@ NiagaraRenderer::~NiagaraRenderer()
 void NiagaraRenderer::Release()
 {
 	check(IsInGameThread());
-	ENQUEUE_UNIQUE_RENDER_COMMAND_ONEPARAMETER(
-		NiagaraRendererDeletion,
-		NiagaraRenderer*, Renderer, this,
+	NiagaraRenderer* Renderer = this;
+	ENQUEUE_RENDER_COMMAND(NiagaraRendererDeletion)(
+		[Renderer](FRHICommandListImmediate& RHICmdList)
 		{
 			delete Renderer;
 		}
diff --git a/Engine/Plugins/FX/Niagara/Source/NiagaraShader/Private/NiagaraShared.cpp b/Engine/Plugins/FX/Niagara/Source/NiagaraShader/Private/NiagaraShared.cpp
--- a/Engine/Plugins/FX/Niagara/Source/NiagaraShader/Private/NiagaraShared.cpp
+++ b/Engine/Plugins/FX/Niagara/Source/NiagaraShader/Private/NiagaraShared.cpp
@@ -154,9 +154,9 @@ void FNiagaraShaderScript::ReleaseShaderMap()
 	{
 		GameThreadShaderMap = nullptr;
 
-		ENQUEUE_UNIQUE_RENDER_COMMAND_ONEPARAMETER(
-			ReleaseShaderMap,
-			FNiagaraShaderScript*, Script, this,
+		FNiagaraShaderScript* Script = this;
+		ENQUEUE_RENDER_COMMAND(ReleaseShaderMap)(
+			[Script](FRHICommandListImmediate& RHICmdList)
 			{
 				Script->SetRenderingThreadShaderMap(nullptr);
 			});
diff --git a/Engine/Plugins/Lumin/MagicLeap/Source/MagicLeap/Private/MagicLeapCustomPresent.cpp b/Engine/Plugins/Lumin/MagicLeap/Source/MagicLeap/Private/MagicLeapCustomPresent.cpp
--- a/Engine/Plugins/Lumin/MagicLeap/Source/MagicLeap/Private/MagicLeapCustomPresent.cpp
+++ b/Engine/Plugins/Lumin/MagicLeap/Source/MagicLeap/Private/MagicLeapCustomPresent.cpp
@@ -416,8 +416,9 @@ void FMagicLeapCustomPresentOpenGL::UpdateViewport(const FViewport& Viewport, FR
 	RenderTargetTexture = *(reinterpret_cast<uint32_t*>(RT->GetNativeResource()));
 	InViewportRHI->SetCustomPresent(this);
 	
-	ENQUEUE_UNIQUE_RENDER_COMMAND_ONEPARAMETER(UpdateViewport_RT,
-		FMagicLeapCustomPresentOpenGL*, CustomPresent, this,
+	FMagicLeapCustomPresentOpenGL* CustomPresent = this;
+	ENQUEUE_RENDER_COMMAND(UpdateViewport_RT)(
+		[CustomPresent](FRHICommandListImmediate& RHICmdList)
 		{
 			CustomPresent->UpdateViewport_RenderThread();
 		}
@@ -641,8 +642,9 @@ void FMagicLeapCustomPresentVulkan::UpdateViewport(const FViewport& Viewport, FR
 
 	InViewportRHI->SetCustomPresent(this);
 	
-	ENQUEUE_UNIQUE_RENDER_COMMAND_ONEPARAMETER(UpdateViewport_RT,
-		FMagicLeapCustomPresentVulkan*, CustomPresent, this,
+	FMagicLeapCustomPresentVulkan* CustomPresent = this;
+	ENQUEUE_RENDER_COMMAND(UpdateViewport_RT)(
+		[CustomPresent](FRHICommandListImmediate& RHICmdList)
 		{
 			CustomPresent->UpdateViewport_RenderThread();
 		}
diff --git a/Engine/Plugins/Lumin/MagicLeap/Source/MagicLeap/Private/MagicLeapHMD.cpp b/Engine/Plugins/Lumin/MagicLeap/Source/MagicLeap/Private/MagicLeapHMD.cpp
--- a/Engine/Plugins/Lumin/MagicLeap/Source/MagicLeap/Private/MagicLeapHMD.cpp
+++ b/Engine/Plugins/Lumin/MagicLeap/Source/MagicLeap/Private/MagicLeapHMD.cpp
@@ -1158,8 +1158,9 @@ void FMagicLeapHMD::Startup()
 
 void FMagicLeapHMD::Shutdown()
 {
-	ENQUEUE_UNIQUE_RENDER_COMMAND_ONEPARAMETER(ShutdownRen,
-		FMagicLeapHMD*, Plugin, this,
+	FMagicLeapHMD* Plugin = this;
+	ENQUEUE_RENDER_COMMAND(ShutdownRendering)(
+		[Plugin](FRHICommandListImmediate& RHICmdList)
 		{
 			Plugin->ShutdownRendering();
 		});
@@ -1494,8 +1495,9 @@ void FMagicLeapHMD::ReleaseDevice()
 	// save any runtime configuration changes to the .ini
 	SaveToIni();
 
-	ENQUEUE_UNIQUE_RENDER_COMMAND_ONEPARAMETER(ReleaseDevice_RT,
-		FMagicLeapHMD*, Plugin, this,
+	FMagicLeapHMD* Plugin = this;
+	ENQUEUE_RENDER_COMMAND(ReleaseDevice_RT)(
+		[Plugin](FRHICommandListImmediate& RHICmdList)
 		{
 			Plugin->ReleaseDevice_RenderThread();
 		}
diff --git a/Engine/Plugins/Media/MediaIOFramework/Source/MediaIOCore/Private/MediaCapture.cpp b/Engine/Plugins/Media/MediaIOFramework/Source/MediaIOCore/Private/MediaCapture.cpp
--- a/Engine/Plugins/Media/MediaIOFramework/Source/MediaIOCore/Private/MediaCapture.cpp
+++ b/Engine/Plugins/Media/MediaIOFramework/Source/MediaIOCore/Private/MediaCapture.cpp
@@ -389,29 +389,24 @@ void UMediaCapture::InitializeResolveTarget(int32 InNumberOfBuffers)
 	check(CaptureFrames.Num() == 0);
 	CaptureFrames.AddDefaulted(InNumberOfBuffers);
 
-	auto RenderCommand = [this](FRHICommandListImmediate& RHICmdList)
-	{
-		FRHIResourceCreateInfo CreateInfo;
-		for (int32 Index = 0; Index < NumberOfCaptureFrame; ++Index)
-		{
-			CaptureFrames[Index].ReadbackTexture = RHICreateTexture2D(
-				DesiredOutputSize.X,
-				DesiredOutputSize.Y,
-				DesiredOutputPixelFormat,
-				1,
-				1,
-				TexCreate_CPUReadback,
-				CreateInfo
-			);
-		}
-		bResolvedTargetInitialized = true;
-	};
-
-	ENQUEUE_UNIQUE_RENDER_COMMAND_ONEPARAMETER(
-		MediaOutputCaptureFrameCreateTexture,
-		decltype(RenderCommand), InRenderCommand, RenderCommand,
+	UMediaCapture* This = this;
+	ENQUEUE_RENDER_COMMAND(MediaOutputCaptureFrameCreateTexture)(
+		[This](FRHICommandListImmediate& RHICmdList)
 		{
-			InRenderCommand(RHICmdList);
+			FRHIResourceCreateInfo CreateInfo;
+			for (int32 Index = 0; Index < This->NumberOfCaptureFrame; ++Index)
+			{
+				This->CaptureFrames[Index].ReadbackTexture = RHICreateTexture2D(
+					This->DesiredOutputSize.X,
+					This->DesiredOutputSize.Y,
+					This->DesiredOutputPixelFormat,
+					1,
+					1,
+					TexCreate_CPUReadback,
+					CreateInfo
+				);
+			}
+			This->bResolvedTargetInitialized = true;
 		});
 }
 
diff --git a/Engine/Plugins/Runtime/Oculus/OculusVR/Source/OculusMR/Private/OculusMR_CastingCameraActor.cpp b/Engine/Plugins/Runtime/Oculus/OculusVR/Source/OculusMR/Private/OculusMR_CastingCameraActor.cpp
--- a/Engine/Plugins/Runtime/Oculus/OculusVR/Source/OculusMR/Private/OculusMR_CastingCameraActor.cpp
+++ b/Engine/Plugins/Runtime/Oculus/OculusVR/Source/OculusMR/Private/OculusMR_CastingCameraActor.cpp
@@ -467,7 +467,7 @@ void AOculusMR_CastingCameraActor::UpdateCameraColorTexture(const ovrpSizei &fra
 		FTexture2DResource* DestTextureResource;
 		uint32 FrameWidth;
 		uint32 FrameHeight;
-	} UploadCameraTextureContext =
+	} Context =
 	{
 		SrcData,
 		Pitch,
@@ -476,9 +476,8 @@ void AOculusMR_CastingCameraActor::UpdateCameraColorTexture(const ovrpSizei &fra
 		(uint32) frameSize.h
 	};
 
-	ENQUEUE_UNIQUE_RENDER_COMMAND_ONEPARAMETER(
-		UpdateCameraColorTexture,
-		FUploadCameraTextureContext, Context, UploadCameraTextureContext,
+	ENQUEUE_RENDER_COMMAND(UpdateCameraColorTexture)(
+		[Context](FRHICommandListImmediate& RHICmdList)
 		{
 			const FUpdateTextureRegion2D UpdateRegion(
 				0, 0,		// Dest X, Y
@@ -523,7 +522,7 @@ void AOculusMR_CastingCameraActor::UpdateCameraDepthTexture(const ovrpSizei &fra
 		FTexture2DResource* DestTextureResource;
 		uint32 FrameWidth;
 		uint32 FrameHeight;
-	} UploadCameraTextureContext =
+	} Context =
 	{
 		SrcData,
 		Pitch,
@@ -532,9 +531,8 @@ void AOculusMR_CastingCameraActor::UpdateCameraDepthTexture(const ovrpSizei &fra
 		(uint32) frameSize.h
 	};
 
-	ENQUEUE_UNIQUE_RENDER_COMMAND_ONEPARAMETER(
-		UpdateCameraDepthTexture,
-		FUploadCameraTextureContext, Context, UploadCameraTextureContext,
+	ENQUEUE_RENDER_COMMAND(UpdateCameraDepthTexture)(
+		[Context](FRHICommandListImmediate& RHICmdList)
 		{
 			const FUpdateTextureRegion2D UpdateRegion(
 				0, 0,		// Dest X, Y
@@ -543,14 +541,14 @@ void AOculusMR_CastingCameraActor::UpdateCameraDepthTexture(const ovrpSizei &fra
 				Context.FrameHeight	    // Height
 			);
 
-	RHIUpdateTexture2D(
-		Context.DestTextureResource->GetTexture2DRHI(),	// Destination GPU texture
-		0,												// Mip map index
-		UpdateRegion,									// Update region
-		Context.CameraBufferPitch,						// Source buffer pitch
-		Context.CameraBuffer);							// Source buffer pointer
+			RHIUpdateTexture2D(
+				Context.DestTextureResource->GetTexture2DRHI(),	// Destination GPU texture
+				0,												// Mip map index
+				UpdateRegion,									// Update region
+				Context.CameraBufferPitch,						// Source buffer pitch
+				Context.CameraBuffer);							// Source buffer pointer
 
-	FMemory::Free(Context.CameraBuffer);
+			FMemory::Free(Context.CameraBuffer);
 		}
 	);
 }
diff --git a/Engine/Plugins/Runtime/Steam/SteamVR/Source/SteamVR/Private/SteamVRSplash.cpp b/Engine/Plugins/Runtime/Steam/SteamVR/Source/SteamVR/Private/SteamVRSplash.cpp
--- a/Engine/Plugins/Runtime/Steam/SteamVR/Source/SteamVR/Private/SteamVRSplash.cpp
+++ b/Engine/Plugins/Runtime/Steam/SteamVR/Source/SteamVR/Private/SteamVRSplash.cpp
@@ -25,8 +25,9 @@ void FSteamSplashTicker::UnregisterForMapLoad()
 
 void FSteamSplashTicker::OnPreLoadMap(const FString&)
 {
-	ENQUEUE_UNIQUE_RENDER_COMMAND_ONEPARAMETER(RegisterAsyncTick,
-		FTickableObjectRenderThread*, Ticker, this,
+	FTickableObjectRenderThread* Ticker = this;
+	ENQUEUE_RENDER_COMMAND(RegisterAsyncTick)(
+		[Ticker](FRHICommandListImmediate& RHICmdList)
 		{
 			Ticker->Register();
 		});
@@ -34,8 +35,9 @@ void FSteamSplashTicker::OnPreLoadMap(const FString&)
 
 void FSteamSplashTicker::OnPostLoadMap(UWorld*)
 {
-	ENQUEUE_UNIQUE_RENDER_COMMAND_ONEPARAMETER(UnregisterAsyncTick, 
-		FTickableObjectRenderThread*, Ticker, this,
+	FTickableObjectRenderThread* Ticker = this;
+	ENQUEUE_RENDER_COMMAND(UnregisterAsyncTick)(
+		[Ticker](FRHICommandListImmediate& RHICmdList)
 		{
 			Ticker->Unregister();
 		});
diff --git a/Engine/Source/Editor/MeshPaint/Private/TexturePaintHelpers.cpp b/Engine/Source/Editor/MeshPaint/Private/TexturePaintHelpers.cpp
--- a/Engine/Source/Editor/MeshPaint/Private/TexturePaintHelpers.cpp
+++ b/Engine/Source/Editor/MeshPaint/Private/TexturePaintHelpers.cpp
@@ -93,9 +93,8 @@ void TexturePaintHelpers::CopyTextureToRenderTargetTexture(UTexture* SourceTextu
 	// Tell the rendering thread to draw any remaining batched elements
 	Canvas.Flush_GameThread(true);
 	
-	ENQUEUE_UNIQUE_RENDER_COMMAND_ONEPARAMETER(
-		UpdateMeshPaintRTCommand,
-		FTextureRenderTargetResource*, RenderTargetResource, RenderTargetResource,
+	ENQUEUE_RENDER_COMMAND(UpdateMeshPaintRTCommand)(
+		[RenderTargetResource](FRHICommandListImmediate& RHICmdList)
 		{
 			// Copy (resolve) the rendered image from the frame buffer to its render target texture
 			RHICmdList.CopyToResolveTarget(
@@ -307,9 +306,8 @@ bool TexturePaintHelpers::GenerateSeamMask(UMeshComponent* MeshComponent, int32
 
 
 	{
-		ENQUEUE_UNIQUE_RENDER_COMMAND_ONEPARAMETER(
-			UpdateMeshPaintRTCommand5,
-			FTextureRenderTargetResource*, RenderTargetResource, RenderTargetResource,
+		ENQUEUE_RENDER_COMMAND(UpdateMeshPaintRTCommand5)(
+			[RenderTargetResource](FRHICommandListImmediate& RHICmdList)
 			{
 				// Copy (resolve) the rendered image from the frame buffer to its render target texture
 				RHICmdList.CopyToResolveTarget(
diff --git a/Engine/Source/Editor/MeshPaintMode/Private/PaintModePainter.cpp b/Engine/Source/Editor/MeshPaintMode/Private/PaintModePainter.cpp
--- a/Engine/Source/Editor/MeshPaintMode/Private/PaintModePainter.cpp
+++ b/Engine/Source/Editor/MeshPaintMode/Private/PaintModePainter.cpp
@@ -1617,9 +1617,8 @@ void FPaintModePainter::PaintTexture(const FMeshPaintParameters& InParams, TArra
 	}
 
 	{
-		ENQUEUE_UNIQUE_RENDER_COMMAND_ONEPARAMETER(
-			UpdateMeshPaintRTCommand1,
-			FTextureRenderTargetResource*, BrushRenderTargetResource, BrushRenderTargetResource,
+		ENQUEUE_RENDER_COMMAND(UpdateMeshPaintRTCommand1)(
+			[BrushRenderTargetResource](FRHICommandListImmediate& RHICmdList)
 			{
 				// Copy (resolve) the rendered image from the frame buffer to its render target texture
 				RHICmdList.CopyToResolveTarget(
@@ -1635,9 +1634,8 @@ void FPaintModePainter::PaintTexture(const FMeshPaintParameters& InParams, TArra
 		BrushMaskCanvas->Flush_GameThread(true);
 
 		{
-			ENQUEUE_UNIQUE_RENDER_COMMAND_ONEPARAMETER(
-				UpdateMeshPaintRTCommand2,
-				FTextureRenderTargetResource*, BrushMaskRenderTargetResource, BrushMaskRenderTargetResource,
+			ENQUEUE_RENDER_COMMAND(UpdateMeshPaintRTCommand2)(
+				[BrushMaskRenderTargetResource](FRHICommandListImmediate& RHICmdList)
 				{
 					// Copy (resolve) the rendered image from the frame buffer to its render target texture
 					RHICmdList.CopyToResolveTarget(
@@ -1737,9 +1735,8 @@ void FPaintModePainter::PaintTexture(const FMeshPaintParameters& InParams, TArra
 		}
 
 		{
-			ENQUEUE_UNIQUE_RENDER_COMMAND_ONEPARAMETER(
-				UpdateMeshPaintRTCommand3,
-				FTextureRenderTargetResource*, RenderTargetResource, RenderTargetResource,
+			ENQUEUE_RENDER_COMMAND(UpdateMeshPaintRTCommand3)(
+				[RenderTargetResource](FRHICommandListImmediate& RHICmdList)
 				{
 					// Copy (resolve) the rendered image from the frame buffer to its render target texture
 					RHICmdList.CopyToResolveTarget(
diff --git a/Engine/Source/Editor/UnrealEd/Private/AssetThumbnail.cpp b/Engine/Source/Editor/UnrealEd/Private/AssetThumbnail.cpp
--- a/Engine/Source/Editor/UnrealEd/Private/AssetThumbnail.cpp
+++ b/Engine/Source/Editor/UnrealEd/Private/AssetThumbnail.cpp
@@ -897,14 +897,14 @@ void FAssetThumbnailPool::ReleaseResources()
 		const TSharedRef<FThumbnailInfo>& Thumb = *ThumbIt;
 
 			// Release rendering resources
-			ENQUEUE_UNIQUE_RENDER_COMMAND_ONEPARAMETER( 
-				ReleaseThumbnailResources,
-				FThumbnailInfo_RenderThread, ThumbInfo, Thumb.Get(),
-			{
-				ThumbInfo.ThumbnailTexture->ClearTextureData();
-				ThumbInfo.ThumbnailTexture->ReleaseResource();
-				ThumbInfo.ThumbnailRenderTarget->ReleaseResource();
-			});
+			FThumbnailInfo_RenderThread ThumbInfo = Thumb.Get();
+			ENQUEUE_RENDER_COMMAND(ReleaseThumbnailResources)(
+				[ThumbInfo](FRHICommandListImmediate& RHICmdList)
+				{
+					ThumbInfo.ThumbnailTexture->ClearTextureData();
+					ThumbInfo.ThumbnailTexture->ReleaseResource();
+					ThumbInfo.ThumbnailRenderTarget->ReleaseResource();
+				});
 		}
 
 	// Wait for all resources to be released
@@ -1014,17 +1014,17 @@ void FAssetThumbnailPool::Tick( float DeltaTime )
 						FThumbnailRenderingInfo* RenderInfo = GUnrealEd->GetThumbnailManager()->GetRenderingInfo( Asset );
 						if ( RenderInfo != NULL && RenderInfo->Renderer != NULL )
 						{
-							ENQUEUE_UNIQUE_RENDER_COMMAND_ONEPARAMETER(
-								SyncSlateTextureCommand,
-								FThumbnailInfo_RenderThread, ThumbInfo, InfoRef.Get(),
-							{
-								if ( ThumbInfo.ThumbnailTexture->GetTypedResource() != ThumbInfo.ThumbnailRenderTarget->GetTextureRHI() )
+							FThumbnailInfo_RenderThread ThumbInfo = InfoRef.Get();
+							ENQUEUE_RENDER_COMMAND(SyncSlateTextureCommand)(
+								[ThumbInfo](FRHICommandListImmediate& RHICmdList)
 								{
-									ThumbInfo.ThumbnailTexture->ClearTextureData();
-									ThumbInfo.ThumbnailTexture->ReleaseDynamicRHI();
-									ThumbInfo.ThumbnailTexture->SetRHIRef(ThumbInfo.ThumbnailRenderTarget->GetTextureRHI(), ThumbInfo.Width, ThumbInfo.Height);
-								}
-							});
+									if ( ThumbInfo.ThumbnailTexture->GetTypedResource() != ThumbInfo.ThumbnailRenderTarget->GetTextureRHI() )
+									{
+										ThumbInfo.ThumbnailTexture->ClearTextureData();
+										ThumbInfo.ThumbnailTexture->ReleaseDynamicRHI();
+										ThumbInfo.ThumbnailTexture->SetRHIRef(ThumbInfo.ThumbnailRenderTarget->GetTextureRHI(), ThumbInfo.Width, ThumbInfo.Height);
+									}
+								});
 
 							if (InfoRef->LastUpdateTime <= 0.0f || RenderInfo->Renderer->AllowsRealtimeThumbnails(Asset))
 							{
@@ -1376,12 +1376,12 @@ void FAssetThumbnailPool::FreeThumbnail( const FName& ObjectPath, uint32 Width,
 			RealTimeThumbnails.Remove(ThumbnailInfo);
 			RealTimeThumbnailsToRender.Remove(ThumbnailInfo);
 
-			ENQUEUE_UNIQUE_RENDER_COMMAND_ONEPARAMETER( 
-				ReleaseThumbnailTextureData,
-				FSlateTexture2DRHIRef*, ThumbnailTexture, ThumbnailInfo->ThumbnailTexture,
-			{
-				ThumbnailTexture->ClearTextureData();
-			});
+			FSlateTexture2DRHIRef* ThumbnailTexture = ThumbnailInfo->ThumbnailTexture;
+			ENQUEUE_RENDER_COMMAND(ReleaseThumbnailTextureData)(
+				[ThumbnailTexture](FRHICommandListImmediate& RHICmdList)
+				{
+					ThumbnailTexture->ClearTextureData();
+				});
 
 			FreeThumbnails.Add( ThumbnailInfo );
 		}
diff --git a/Engine/Source/Editor/UnrealEd/Private/ObjectTools.cpp b/Engine/Source/Editor/UnrealEd/Private/ObjectTools.cpp
--- a/Engine/Source/Editor/UnrealEd/Private/ObjectTools.cpp
+++ b/Engine/Source/Editor/UnrealEd/Private/ObjectTools.cpp
@@ -3878,16 +3878,15 @@ namespace ThumbnailTools
 
 
 			{
-				ENQUEUE_UNIQUE_RENDER_COMMAND_ONEPARAMETER(
-					UpdateThumbnailRTCommand,
-					FTextureRenderTargetResource*, RenderTargetResource, RenderTargetResource,
-				{
-					// Copy (resolve) the rendered thumbnail from the render target to its texture
-					RHICmdList.CopyToResolveTarget(
-						RenderTargetResource->GetRenderTargetTexture(),		// Source texture
-						RenderTargetResource->TextureRHI,					// Dest texture
-						FResolveParams() );									// Resolve parameters
-				});
+				ENQUEUE_RENDER_COMMAND(UpdateThumbnailRTCommand)(
+					[RenderTargetResource](FRHICommandListImmediate& RHICmdList)
+					{
+						// Copy (resolve) the rendered thumbnail from the render target to its texture
+						RHICmdList.CopyToResolveTarget(
+							RenderTargetResource->GetRenderTargetTexture(),		// Source texture
+							RenderTargetResource->TextureRHI,					// Dest texture
+							FResolveParams() );									// Resolve parameters
+					});
 
 				if(OutThumbnail)
 				{
diff --git a/Engine/Source/Runtime/Engine/Private/TileRendering.cpp b/Engine/Source/Runtime/Engine/Private/TileRendering.cpp
--- a/Engine/Source/Runtime/Engine/Private/TileRendering.cpp
+++ b/Engine/Source/Runtime/Engine/Private/TileRendering.cpp
@@ -210,9 +210,9 @@ FCanvasTileRendererItem::FRenderData::FRenderData(ERHIFeatureLevel::Type InFeatu
 	, MaterialRenderProxy(InMaterialRenderProxy)
 	, Transform(InTransform)
 {
-
-	ENQUEUE_UNIQUE_RENDER_COMMAND_ONEPARAMETER(FCanvasTileRendererItemFTileVertexFactoryInit,
-		FTileVertexFactory*, TileVertexFactory, &VertexFactory,
+	FTileVertexFactory* TileVertexFactory = &VertexFactory;
+	ENQUEUE_RENDER_COMMAND(FCanvasTileRendererItemFTileVertexFactoryInit)(
+		[TileVertexFactory](FRHICommandListImmediate& RHICmdList)
 		{
 			FLocalVertexFactory::FDataType VertexData;
 			VertexData.NumTexCoords = 1;
@@ -258,8 +258,9 @@ FCanvasTileRendererItem::FRenderData::FRenderData(ERHIFeatureLevel::Type InFeatu
 
 FCanvasTileRendererItem::~FCanvasTileRendererItem()
 {
-	ENQUEUE_UNIQUE_RENDER_COMMAND_ONEPARAMETER(FCanvasTileRendererItemDeleteRenderData,
-		FCanvasTileRendererItem::FRenderData*, RenderData, Data,
+	FCanvasTileRendererItem::FRenderData* RenderData = Data;
+	ENQUEUE_RENDER_COMMAND(FCanvasTileRendererItemDeleteRenderData)(
+		[RenderData](FRHICommandListImmediate& RHICmdList)
 		{
 			delete RenderData;
 		});
diff --git a/Engine/Source/Runtime/Engine/Private/UnrealClient.cpp b/Engine/Source/Runtime/Engine/Private/UnrealClient.cpp
--- a/Engine/Source/Runtime/Engine/Private/UnrealClient.cpp
+++ b/Engine/Source/Runtime/Engine/Private/UnrealClient.cpp
@@ -63,7 +63,7 @@ bool FRenderTarget::ReadPixels(TArray< FColor >& OutImageData, FReadSurfaceDataF
 	};
 
 	OutImageData.Reset();
-	FReadSurfaceContext ReadSurfaceContext =
+	FReadSurfaceContext Context =
 	{
 		this,
 		&OutImageData,
@@ -71,17 +71,16 @@ bool FRenderTarget::ReadPixels(TArray< FColor >& OutImageData, FReadSurfaceDataF
 		InFlags
 	};
 
-	ENQUEUE_UNIQUE_RENDER_COMMAND_ONEPARAMETER(
-		ReadSurfaceCommand,
-		FReadSurfaceContext,Context,ReadSurfaceContext,
-	{
-		RHICmdList.ReadSurfaceData(
-			Context.SrcRenderTarget->GetRenderTargetTexture(),
-			Context.Rect,
-			*Context.OutData,
-			Context.Flags
-			);
-	});
+	ENQUEUE_RENDER_COMMAND(ReadSurfaceCommand)(
+		[Context](FRHICommandListImmediate& RHICmdList)
+		{
+			RHICmdList.ReadSurfaceData(
+				Context.SrcRenderTarget->GetRenderTargetTexture(),
+				Context.Rect,
+				*Context.OutData,
+				Context.Flags
+				);
+		});
 	FlushRenderingCommands();
 
 	return OutImageData.Num() > 0;
@@ -124,7 +123,7 @@ bool FRenderTarget::ReadFloat16Pixels(FFloat16Color* OutImageData,ECubeFace Cube
 	};
 	
 	TArray<FFloat16Color> SurfaceData;
-	FReadSurfaceFloatContext ReadSurfaceFloatContext =
+	FReadSurfaceFloatContext Context =
 	{
 		this,
 		&SurfaceData,
@@ -132,19 +131,18 @@ bool FRenderTarget::ReadFloat16Pixels(FFloat16Color* OutImageData,ECubeFace Cube
 		CubeFace	
 	};
 
-	ENQUEUE_UNIQUE_RENDER_COMMAND_ONEPARAMETER(
-		ReadSurfaceFloatCommand,
-		FReadSurfaceFloatContext,Context,ReadSurfaceFloatContext,
-	{
-		RHICmdList.ReadSurfaceFloatData(
-			Context.SrcRenderTarget->GetRenderTargetTexture(),
-			Context.Rect,
-			*Context.OutData,
-			Context.CubeFace,
-			0,
-			0
-			);
-	});
+	ENQUEUE_RENDER_COMMAND(ReadSurfaceFloatCommand)(
+		[Context](FRHICommandListImmediate& RHICmdList)
+		{
+			RHICmdList.ReadSurfaceFloatData(
+				Context.SrcRenderTarget->GetRenderTargetTexture(),
+				Context.Rect,
+				*Context.OutData,
+				Context.CubeFace,
+				0,
+				0
+				);
+		});
 	FlushRenderingCommands();
 
 	// Copy the surface data into the output array.
@@ -202,7 +200,7 @@ bool FRenderTarget::ReadLinearColorPixels(TArray<FLinearColor> &OutImageData, FR
 	};
 
 	OutImageData.Reset();
-	FReadSurfaceContext ReadSurfaceContext =
+	FReadSurfaceContext Context =
 	{
 		this,
 		&OutImageData,
@@ -210,9 +208,8 @@ bool FRenderTarget::ReadLinearColorPixels(TArray<FLinearColor> &OutImageData, FR
 		InFlags
 	};
 
-	ENQUEUE_UNIQUE_RENDER_COMMAND_ONEPARAMETER(
-		ReadSurfaceCommand,
-		FReadSurfaceContext, Context, ReadSurfaceContext,
+	ENQUEUE_RENDER_COMMAND(ReadSurfaceCommand)(
+		[Context](FRHICommandListImmediate& RHICmdList)
 		{
 			RHICmdList.ReadSurfaceData(
 			Context.SrcRenderTarget->GetRenderTargetTexture(),
@@ -1293,7 +1290,7 @@ struct FEndDrawingCommandParams
 };
 
 /**
- * Helper function used in ENQUEUE_UNIQUE_RENDER_COMMAND_ONEPARAMETER below. Needed to be split out due to
+ * Helper function used in ENQUEUE_RENDER_COMMAND below. Needed to be split out due to
  * use of macro and former already being one.
  *
  * @param Parameters	Parameters passed from the gamethread to the renderthread command.
@@ -1405,23 +1402,22 @@ void UPostProcessComponent::Serialize(FArchive& Ar)
 void FViewport::EnqueueBeginRenderFrame(const bool bShouldPresent)
 {
 	AdvanceFrameRenderPrerequisite();
-	ENQUEUE_UNIQUE_RENDER_COMMAND_ONEPARAMETER(
-		BeginDrawingCommand,
-		FViewport*,Viewport,this,
-	{
-		Viewport->BeginRenderFrame(RHICmdList);
-	});
+	FViewport* Viewport = this;
+	ENQUEUE_RENDER_COMMAND(BeginDrawingCommand)(
+		[Viewport](FRHICommandListImmediate& RHICmdList)
+		{
+			Viewport->BeginRenderFrame(RHICmdList);
+		});
 }
 
 
 void FViewport::EnqueueEndRenderFrame(const bool bLockToVsync, const bool bShouldPresent)
 {
 	FEndDrawingCommandParams Params = { this, (uint32)bLockToVsync, (uint32)GInputLatencyTimer.GameThreadTrigger, (uint32)(PresentAndStopMovieDelay > 0 ? 0 : bShouldPresent) };
-	ENQUEUE_UNIQUE_RENDER_COMMAND_ONEPARAMETER(
-		EndDrawingCommand,
-		FEndDrawingCommandParams, Parameters, Params,
+	ENQUEUE_RENDER_COMMAND(EndDrawingCommand)(
+		[Params](FRHICommandListImmediate& RHICmdList)
 		{
-			ViewportEndDrawing(RHICmdList, Parameters);
+			ViewportEndDrawing(RHICmdList, Params);
 		});
 }
 
@@ -1603,16 +1599,16 @@ const TArray<FColor>& FViewport::GetRawHitProxyData(FIntRect InRect)
 	{
 		EnqueueBeginRenderFrame(false);
 
-		ENQUEUE_UNIQUE_RENDER_COMMAND_ONEPARAMETER(
-			BeginDrawingCommandHitProxy,
-			FViewport*, Viewport, this,
+		FViewport* Viewport = this;
+		ENQUEUE_RENDER_COMMAND(BeginDrawingCommandHitProxy)(
+			[Viewport](FRHICommandListImmediate& RHICmdList)
 			{
-			// Set the hit proxy map's render target.
-			// Clear the hit proxy map to white, which is overloaded to mean no hit proxy.
-			FRHIRenderPassInfo RPInfo(Viewport->HitProxyMap.GetRenderTargetTexture(), ERenderTargetActions::Clear_Store);
-			RHICmdList.BeginRenderPass(RPInfo, TEXT("ClearHitProxyMap"));
-			RHICmdList.EndRenderPass();
-		});
+				// Set the hit proxy map's render target.
+				// Clear the hit proxy map to white, which is overloaded to mean no hit proxy.
+				FRHIRenderPassInfo RPInfo(Viewport->HitProxyMap.GetRenderTargetTexture(), ERenderTargetActions::Clear_Store);
+				RHICmdList.BeginRenderPass(RPInfo, TEXT("ClearHitProxyMap"));
+				RHICmdList.EndRenderPass();
+			});
 
 		// Let the viewport client draw its hit proxies.
 		UWorld* World = ViewportClient->GetWorld();
@@ -1623,18 +1619,17 @@ const TArray<FColor>& FViewport::GetRawHitProxyData(FIntRect InRect)
 		Canvas.Flush_GameThread();
 
 		//Resolve surface to texture.
-		ENQUEUE_UNIQUE_RENDER_COMMAND_ONEPARAMETER(
-			UpdateHitProxyRTCommand,
-			FHitProxyMap*, HitProxyMap, &HitProxyMap,
-		{
-			// Copy (resolve) the rendered thumbnail from the render target to its texture
-			RHICmdList.CopyToResolveTarget(HitProxyMap->GetRenderTargetTexture(), HitProxyMap->GetHitProxyTexture(), FResolveParams());
-			RHICmdList.CopyToResolveTarget(HitProxyMap->GetRenderTargetTexture(), HitProxyMap->GetHitProxyCPUTexture(), FResolveParams());
-		});
+		FHitProxyMap* HitProxyMapPtr = &HitProxyMap;
+		ENQUEUE_RENDER_COMMAND(UpdateHitProxyRTCommand)(
+			[HitProxyMapPtr](FRHICommandListImmediate& RHICmdList)
+			{
+				// Copy (resolve) the rendered thumbnail from the render target to its texture
+				RHICmdList.CopyToResolveTarget(HitProxyMapPtr->GetRenderTargetTexture(), HitProxyMapPtr->GetHitProxyTexture(), FResolveParams());
+				RHICmdList.CopyToResolveTarget(HitProxyMapPtr->GetRenderTargetTexture(), HitProxyMapPtr->GetHitProxyCPUTexture(), FResolveParams());
+			});
 
-		ENQUEUE_UNIQUE_RENDER_COMMAND_ONEPARAMETER(
-			EndDrawingCommand,
-			FViewport*, Viewport, this,
+		ENQUEUE_RENDER_COMMAND(EndDrawingCommand)(
+			[Viewport](FRHICommandListImmediate& RHICmdList)
 			{
 				Viewport->EndRenderFrame(RHICmdList, false, false);
 			});
@@ -1653,24 +1648,23 @@ const TArray<FColor>& FViewport::GetRawHitProxyData(FIntRect InRect)
 			TArray<FColor>* OutData;
 			FIntRect Rect;
 		};
-		FReadSurfaceContext ReadSurfaceContext =
+		FReadSurfaceContext Context =
 		{
 			this,
 			&CachedHitProxyData,
 			ViewportRect
 		};
 
-		ENQUEUE_UNIQUE_RENDER_COMMAND_ONEPARAMETER(
-			ReadSurfaceCommand,
-			FReadSurfaceContext, Context, ReadSurfaceContext,
+		ENQUEUE_RENDER_COMMAND(ReadSurfaceCommand)(
+			[Context](FRHICommandListImmediate& RHICmdList)
 			{
-			RHICmdList.ReadSurfaceData(
-			Context.Viewport->HitProxyMap.GetHitProxyCPUTexture(),
-			Context.Rect,
-			*Context.OutData,
-			FReadSurfaceDataFlags()
-			);
-		});
+				RHICmdList.ReadSurfaceData(
+				Context.Viewport->HitProxyMap.GetHitProxyCPUTexture(),
+				Context.Rect,
+				*Context.OutData,
+				FReadSurfaceDataFlags()
+				);
+			});
 		FlushRenderingCommands();
 	}
 
diff --git a/Engine/Source/Runtime/Engine/Private/UnrealEngine.cpp b/Engine/Source/Runtime/Engine/Private/UnrealEngine.cpp
--- a/Engine/Source/Runtime/Engine/Private/UnrealEngine.cpp
+++ b/Engine/Source/Runtime/Engine/Private/UnrealEngine.cpp
@@ -8210,7 +8210,12 @@ bool UEngine::PerformError(const TCHAR* Cmd, FOutputDevice& Ar)
 		{
 			Seconds = 1.0f;
 		}
-		ENQUEUE_UNIQUE_RENDER_COMMAND_ONEPARAMETER(CauseRenderThreadHitch, float, Length, Seconds, { SCOPE_CYCLE_COUNTER(STAT_IntentionalHitch); FPlatformProcess::Sleep(Length); });
+		ENQUEUE_RENDER_COMMAND(CauseRenderThreadHitch)(
+			[Seconds](FRHICommandListImmediate& RHICmdList)
+			{
+				SCOPE_CYCLE_COUNTER(STAT_IntentionalHitch);
+				FPlatformProcess::Sleep(Seconds);
+			});
 		return true;
 	}
 	else if (FParse::Command(&Cmd, TEXT("SPIN")))
@@ -8235,7 +8240,13 @@ bool UEngine::PerformError(const TCHAR* Cmd, FOutputDevice& Ar)
 		{
 			Seconds = 1.0f;
 		}
-		ENQUEUE_UNIQUE_RENDER_COMMAND_ONEPARAMETER(CauseRenderThreadHitch, float, Length, Seconds, { SCOPE_CYCLE_COUNTER(STAT_IntentionalHitch); double StartTime = FPlatformTime::Seconds(); while (FPlatformTime::Seconds() < StartTime + Length) {} });
+		ENQUEUE_RENDER_COMMAND(CauseRenderThreadHitch)(
+			[Seconds](FRHICommandListImmediate& RHICmdList)
+			{
+				SCOPE_CYCLE_COUNTER(STAT_IntentionalHitch);
+				double StartTime = FPlatformTime::Seconds();
+				while (FPlatformTime::Seconds() < StartTime + Seconds) {}
+			});
 		return true;
 	}
 	else if (FParse::Command(&Cmd, TEXT("LONGLOG")))
diff --git a/Engine/Source/Runtime/Landscape/Private/LandscapeEdit.cpp b/Engine/Source/Runtime/Landscape/Private/LandscapeEdit.cpp
--- a/Engine/Source/Runtime/Landscape/Private/LandscapeEdit.cpp
+++ b/Engine/Source/Runtime/Landscape/Private/LandscapeEdit.cpp
@@ -5949,9 +5949,8 @@ bool ALandscapeProxy::LandscapeExportHeightmapToRenderTarget(UTextureRenderTarge
 	// Tell the rendering thread to draw any remaining batched elements
 	Canvas.Flush_GameThread(true);
 
-	ENQUEUE_UNIQUE_RENDER_COMMAND_ONEPARAMETER(
-		DrawHeightmapRTCommand,
-		FTextureRenderTargetResource*, RenderTargetResource, RenderTargetResource,
+	ENQUEUE_RENDER_COMMAND(DrawHeightmapRTCommand)(
+		[RenderTargetResource](FRHICommandListImmediate& RHICmdList)
 		{
 			// Copy (resolve) the rendered image from the frame buffer to its render target texture
 			RHICmdList.CopyToResolveTarget(
diff --git a/Engine/Source/Runtime/Landscape/Private/LandscapeEditProcedural.cpp b/Engine/Source/Runtime/Landscape/Private/LandscapeEditProcedural.cpp
--- a/Engine/Source/Runtime/Landscape/Private/LandscapeEditProcedural.cpp
+++ b/Engine/Source/Runtime/Landscape/Private/LandscapeEditProcedural.cpp
@@ -1493,9 +1493,9 @@ void ALandscape::PrintDebugRTHeightmap(FString Context, UTextureRenderTarget2D*
 		return;
 	}
 
-	ENQUEUE_UNIQUE_RENDER_COMMAND_ONEPARAMETER(
-		HeightmapRTCanvasRenderTargetResolveCommand,
-		FTextureRenderTargetResource*, RenderTargetResource, InDebugRT->GameThread_GetRenderTargetResource(),
+	FTextureRenderTargetResource* RenderTargetResource = InDebugRT->GameThread_GetRenderTargetResource();
+	ENQUEUE_RENDER_COMMAND(HeightmapRTCanvasRenderTargetResolveCommand)(
+		[RenderTargetResource](FRHICommandListImmediate& RHICmdList)
 		{
 			// Copy (resolve) the rendered image from the frame buffer to its render target texture
 			RHICmdList.CopyToResolveTarget(RenderTargetResource->GetRenderTargetTexture(), RenderTargetResource->TextureRHI, FResolveParams());
diff --git a/Engine/Source/Runtime/Landscape/Private/LandscapeGrass.cpp b/Engine/Source/Runtime/Landscape/Private/LandscapeGrass.cpp
--- a/Engine/Source/Runtime/Landscape/Private/LandscapeGrass.cpp
+++ b/Engine/Source/Runtime/Landscape/Private/LandscapeGrass.cpp
@@ -584,9 +584,9 @@ public:
 		RenderTargetResource = RenderTargetTexture->GameThread_GetRenderTargetResource()->GetTextureRenderTarget2DResource();
 
 		// render
-		ENQUEUE_UNIQUE_RENDER_COMMAND_ONEPARAMETER(
-			FDrawSceneCommand,
-			FLandscapeGrassWeightExporter_RenderThread*, Exporter, this,
+		FLandscapeGrassWeightExporter_RenderThread* Exporter = this;
+		ENQUEUE_RENDER_COMMAND(FDrawSceneCommand)(
+			[Exporter](FRHICommandListImmediate& RHICmdList)
 			{
 				Exporter->RenderLandscapeComponentToTexture_RenderThread(RHICmdList);
 				FlushPendingDeleteRHIResources_RenderThread();
diff --git a/Engine/Source/Runtime/Landscape/Private/LandscapeRenderMobile.cpp b/Engine/Source/Runtime/Landscape/Private/LandscapeRenderMobile.cpp
--- a/Engine/Source/Runtime/Landscape/Private/LandscapeRenderMobile.cpp
+++ b/Engine/Source/Runtime/Landscape/Private/LandscapeRenderMobile.cpp
@@ -307,11 +307,11 @@ struct FLandscapeMobileRenderData
 			}
 			else
 			{
-				ENQUEUE_UNIQUE_RENDER_COMMAND_ONEPARAMETER(
-					InitCommand,
-					FLandscapeVertexBufferMobile*, VertexBuffer, VertexBuffer,
+				FLandscapeVertexBufferMobile* InVertexBuffer = VertexBuffer;
+				ENQUEUE_RENDER_COMMAND(InitCommand)(
+					[InVertexBuffer](FRHICommandListImmediate& RHICmdList)
 					{
-						delete VertexBuffer;
+						delete InVertexBuffer;
 					});
 			}
 		}
diff --git a/Engine/Source/Runtime/MoviePlayer/Private/DefaultGameMoviePlayer.cpp b/Engine/Source/Runtime/MoviePlayer/Private/DefaultGameMoviePlayer.cpp
--- a/Engine/Source/Runtime/MoviePlayer/Private/DefaultGameMoviePlayer.cpp
+++ b/Engine/Source/Runtime/MoviePlayer/Private/DefaultGameMoviePlayer.cpp
@@ -127,10 +127,12 @@ FDefaultGameMoviePlayer::~FDefaultGameMoviePlayer()
 	else if (GIsRHIInitialized)
 	{
 		// Even when uninitialized we must safely unregister the movie player on the render thread
-		ENQUEUE_UNIQUE_RENDER_COMMAND_ONEPARAMETER(UnregisterMoviePlayerTickable, FDefaultGameMoviePlayer*, MoviePlayer, this,
-		{
-			MoviePlayer->Unregister();
-		});
+		FDefaultGameMoviePlayer* InMoviePlayer = this;
+		ENQUEUE_RENDER_COMMAND(UnregisterMoviePlayerTickable)(
+			[InMoviePlayer](FRHICommandListImmediate& RHICmdList)
+			{
+				InMoviePlayer->Unregister();
+			});
 	}
 
 	FCoreDelegates::IsLoadingMovieCurrentlyPlaying.Unbind();
@@ -160,10 +162,12 @@ void FDefaultGameMoviePlayer::Initialize(FSlateRenderer& InSlateRenderer, TShare
 
 	UE_LOG(LogMoviePlayer, Log, TEXT("Initializing movie player"));
 
-	ENQUEUE_UNIQUE_RENDER_COMMAND_ONEPARAMETER(RegisterMoviePlayerTickable, FDefaultGameMoviePlayer*, MoviePlayer, this,
-	{
-		MoviePlayer->Register();
-	});
+	FDefaultGameMoviePlayer* InMoviePlayer = this;
+	ENQUEUE_RENDER_COMMAND(RegisterMoviePlayerTickable)(
+		[InMoviePlayer](FRHICommandListImmediate& RHICmdList)
+		{
+			InMoviePlayer->Register();
+		});
 
 	bInitialized = true;
 
@@ -262,10 +266,12 @@ void FDefaultGameMoviePlayer::Shutdown()
 	StopMovie();
 	WaitForMovieToFinish();
 
-	ENQUEUE_UNIQUE_RENDER_COMMAND_ONEPARAMETER(UnregisterMoviePlayerTickable, FDefaultGameMoviePlayer*, MoviePlayer, this,
-	{
-		MoviePlayer->Unregister();
-	});
+	FDefaultGameMoviePlayer* InMoviePlayer = this;
+	ENQUEUE_RENDER_COMMAND(UnregisterMoviePlayerTickable)(
+		[InMoviePlayer](FRHICommandListImmediate& RHICmdList)
+		{
+			InMoviePlayer->Unregister();
+		});
 
 	bInitialized = false;
 
diff --git a/Engine/Source/Runtime/PreLoadScreen/Private/PreLoadSlateThreading.cpp b/Engine/Source/Runtime/PreLoadScreen/Private/PreLoadSlateThreading.cpp
--- a/Engine/Source/Runtime/PreLoadScreen/Private/PreLoadSlateThreading.cpp
+++ b/Engine/Source/Runtime/PreLoadScreen/Private/PreLoadSlateThreading.cpp
@@ -212,9 +212,9 @@ void FPreLoadScreenSlateSynchMechanism::SlateThreadRunMainLoop()
             }
 
             //Queue up a render tick every time we tick on this sync thread.
-            ENQUEUE_UNIQUE_RENDER_COMMAND_ONEPARAMETER(
-                PreLoadScreenRenderTick,
-                FPreLoadScreenSlateSynchMechanism*, SyncMech, this,
+			FPreLoadScreenSlateSynchMechanism* SyncMech = this;
+			ENQUEUE_RENDER_COMMAND(PreLoadScreenRenderTick)(
+				[SyncMech](FRHICommandListImmediate& RHICmdList)
                 {
                     FPreLoadScreenManager* PreLoadManager = FPreLoadScreenManager::Get();
                     if (PreLoadManager && FPreLoadScreenManager::ShouldRender())
diff --git a/Engine/Source/Runtime/Renderer/Private/LightPropagationVolume.cpp b/Engine/Source/Runtime/Renderer/Private/LightPropagationVolume.cpp
--- a/Engine/Source/Runtime/Renderer/Private/LightPropagationVolume.cpp
+++ b/Engine/Source/Runtime/Renderer/Private/LightPropagationVolume.cpp
@@ -1512,12 +1512,11 @@ void FSceneViewState::DestroyLightPropagationVolume()
 		FLightPropagationVolume* LPV = LightPropagationVolume;
 		LightPropagationVolume = nullptr;
 
-		ENQUEUE_UNIQUE_RENDER_COMMAND_ONEPARAMETER(
-			DeleteLPV,
-			FLightPropagationVolume*, LPV, LPV,
-		{
-			LPV->Release();
-		}
+		ENQUEUE_RENDER_COMMAND(DeleteLPV)(
+			[LPV](FRHICommandListImmediate& RHICmdList)
+			{
+				LPV->Release();
+			}
 		);
 		bIsStereoView = false;
 	}
diff --git a/Engine/Source/Runtime/Renderer/Private/ReflectionEnvironmentCapture.cpp b/Engine/Source/Runtime/Renderer/Private/ReflectionEnvironmentCapture.cpp
--- a/Engine/Source/Runtime/Renderer/Private/ReflectionEnvironmentCapture.cpp
+++ b/Engine/Source/Runtime/Renderer/Private/ReflectionEnvironmentCapture.cpp
@@ -1027,13 +1027,13 @@ void FScene::UpdateAllReflectionCaptures(const TCHAR* CaptureReason, bool bVerif
 {
 	if (IsReflectionEnvironmentAvailable(GetFeatureLevel()))
 	{
-		ENQUEUE_UNIQUE_RENDER_COMMAND_ONEPARAMETER( 
-			CaptureCommand,
-			FScene*, Scene, this,
-		{
-			Scene->ReflectionSceneData.AllocatedReflectionCaptureState.Empty();
-			Scene->ReflectionSceneData.CubemapArraySlotsUsed.Reset();
-		});
+		FScene* Scene = this;
+		ENQUEUE_RENDER_COMMAND(CaptureCommand)(
+			[Scene](FRHICommandListImmediate& RHICmdList)
+			{
+				Scene->ReflectionSceneData.AllocatedReflectionCaptureState.Empty();
+				Scene->ReflectionSceneData.CubemapArraySlotsUsed.Reset();
+			});
 
 		// Only display status during building reflection captures, otherwise we may interrupt a editor widget manipulation of many captures
 		const int32 NumCapturesForStatus = bVerifyOnlyCapturing ? ReflectionSceneData.AllocatedReflectionCapturesGameThread.Num() : 0;
diff --git a/Engine/Source/Runtime/Renderer/Private/RendererScene.cpp b/Engine/Source/Runtime/Renderer/Private/RendererScene.cpp
--- a/Engine/Source/Runtime/Renderer/Private/RendererScene.cpp
+++ b/Engine/Source/Runtime/Renderer/Private/RendererScene.cpp
@@ -538,9 +538,9 @@ static void DoLazyStaticMeshUpdateCVarSinkFunction()
 			UWorld* World = *It;
 			if (World && World->Scene)
 			{
-				ENQUEUE_UNIQUE_RENDER_COMMAND_ONEPARAMETER(
-					UpdateDoLazyStaticMeshUpdate,
-					FScene*, Scene, (FScene*)(World->Scene),
+				FScene* Scene = (FScene*)(World->Scene);
+				ENQUEUE_RENDER_COMMAND(UpdateDoLazyStaticMeshUpdate)(
+					[Scene](FRHICommandListImmediate& RHICmdList)
 					{
 						Scene->UpdateDoLazyStaticMeshUpdate(RHICmdList);
 					});
@@ -2895,9 +2895,9 @@ void FScene::UpdateStaticDrawLists_RenderThread(FRHICommandListImmediate& RHICmd
 
 void FScene::UpdateStaticDrawLists()
 {
-	ENQUEUE_UNIQUE_RENDER_COMMAND_ONEPARAMETER(
-		FUpdateDrawLists,
-		FScene*, Scene, this,
+	FScene* Scene = this;
+	ENQUEUE_RENDER_COMMAND(FUpdateDrawLists)(
+		[Scene](FRHICommandListImmediate& RHICmdList)
 		{
 			Scene->UpdateStaticDrawLists_RenderThread(RHICmdList);
 		});
@@ -2937,9 +2937,9 @@ void FScene::Release()
 	GetRendererModule().RemoveScene(this);
 
 	// Send a command to the rendering thread to release the scene.
-	ENQUEUE_UNIQUE_RENDER_COMMAND_ONEPARAMETER(
-		FReleaseCommand,
-		FScene*,Scene,this,
+	FScene* Scene = this;
+	ENQUEUE_RENDER_COMMAND(FReleaseCommand)(
+		[Scene](FRHICommandListImmediate& RHICmdList)
 		{
 			delete Scene;
 		});
@@ -3402,9 +3402,9 @@ void UpdateStaticMeshesForMaterials(const TArray<const FMaterial*>& MaterialReso
 
 				if (bPrimitiveIsDependentOnMaterial)
 				{
-					ENQUEUE_UNIQUE_RENDER_COMMAND_ONEPARAMETER(
-						FUpdateStaticMeshesForMaterials,
-						FPrimitiveSceneProxy*, SceneProxy, PrimitiveComponent->SceneProxy,
+					FPrimitiveSceneProxy* SceneProxy = PrimitiveComponent->SceneProxy;
+					ENQUEUE_RENDER_COMMAND(FUpdateStaticMeshesForMaterials)(
+						[SceneProxy](FRHICommandListImmediate& RHICmdList)
 						{
 							SceneProxy->GetPrimitiveSceneInfo()->UpdateStaticMeshes(RHICmdList);
 						});
diff --git a/Engine/Source/Runtime/Renderer/Private/SceneRendering.cpp b/Engine/Source/Runtime/Renderer/Private/SceneRendering.cpp
--- a/Engine/Source/Runtime/Renderer/Private/SceneRendering.cpp
+++ b/Engine/Source/Runtime/Renderer/Private/SceneRendering.cpp
@@ -3319,9 +3319,8 @@ void FRendererModule::BeginRenderingViewFamily(FCanvas* Canvas, FSceneViewFamily
 
 		// We need to execute the pre-render view extensions before we do any view dependent work.
 		// Anything between here and FDrawSceneCommand will add to HMD view latency
-		ENQUEUE_UNIQUE_RENDER_COMMAND_ONEPARAMETER(
-			FViewExtensionPreDrawCommand,
-			FSceneRenderer*, SceneRenderer, SceneRenderer,
+		ENQUEUE_RENDER_COMMAND(FViewExtensionPreDrawCommand)(
+			[SceneRenderer](FRHICommandListImmediate& RHICmdList)
 			{
 				ViewExtensionPreRender_RenderThread(RHICmdList, SceneRenderer);
 			});
@@ -3337,13 +3336,12 @@ void FRendererModule::BeginRenderingViewFamily(FCanvas* Canvas, FSceneViewFamily
 
 		SceneRenderer->ViewFamily.DisplayInternalsData.Setup(World);
 
-		ENQUEUE_UNIQUE_RENDER_COMMAND_ONEPARAMETER(
-			FDrawSceneCommand,
-			FSceneRenderer*,SceneRenderer,SceneRenderer,
-		{
-			RenderViewFamily_RenderThread(RHICmdList, SceneRenderer);
-			FlushPendingDeleteRHIResources_RenderThread();
-		});
+		ENQUEUE_RENDER_COMMAND(FDrawSceneCommand)(
+			[SceneRenderer](FRHICommandListImmediate& RHICmdList)
+			{
+				RenderViewFamily_RenderThread(RHICmdList, SceneRenderer);
+				FlushPendingDeleteRHIResources_RenderThread();
+			});
 	}
 }
 
diff --git a/Engine/Source/Runtime/SlateRHIRenderer/Private/SlateRHIFontTexture.cpp b/Engine/Source/Runtime/SlateRHIRenderer/Private/SlateRHIFontTexture.cpp
--- a/Engine/Source/Runtime/SlateRHIRenderer/Private/SlateRHIFontTexture.cpp
+++ b/Engine/Source/Runtime/SlateRHIRenderer/Private/SlateRHIFontTexture.cpp
@@ -108,15 +108,16 @@ void FSlateFontAtlasRHI::ConditionalUpdateTexture()
 
 			BeginInitResource( FontTexture.Get() );
 
-			ENQUEUE_UNIQUE_RENDER_COMMAND_ONEPARAMETER( SlateUpdateFontAtlasTextureCommand,
-				FSlateFontAtlasRHI&, Atlas, *this,
-			{
-				uint32 DestStride;
-				uint8* TempData = (uint8*)RHILockTexture2D( Atlas.FontTexture->GetTypedResource(), 0, RLM_WriteOnly, /*out*/ DestStride, false );
-				// check( DestStride == Atlas.BytesPerPixel * Atlas.AtlasWidth ); // Temporarily disabling check
-				FMemory::Memcpy( TempData, Atlas.AtlasData.GetData(), Atlas.BytesPerPixel*Atlas.AtlasWidth*Atlas.AtlasHeight );
-				RHIUnlockTexture2D( Atlas.FontTexture->GetTypedResource(),0,false );
-			});
+			FSlateFontAtlasRHI* Atlas = this;
+			ENQUEUE_RENDER_COMMAND(SlateUpdateFontAtlasTextureCommand)(
+				[Atlas](FRHICommandListImmediate& RHICmdList)
+				{
+					uint32 DestStride;
+					uint8* TempData = (uint8*)RHILockTexture2D( Atlas->FontTexture->GetTypedResource(), 0, RLM_WriteOnly, /*out*/ DestStride, false );
+					// check( DestStride == Atlas.BytesPerPixel * Atlas.AtlasWidth ); // Temporarily disabling check
+					FMemory::Memcpy( TempData, Atlas->AtlasData.GetData(), Atlas->BytesPerPixel*Atlas->AtlasWidth*Atlas->AtlasHeight );
+					RHIUnlockTexture2D( Atlas->FontTexture->GetTypedResource(),0,false );
+				});
 		}
 
 		bNeedsUpdate = false;
@@ -139,11 +140,12 @@ FSlateFontTextureRHI::FSlateFontTextureRHI(const uint32 InWidth, const uint32 In
 
 		BeginInitResource(FontTexture.Get());
 
-		ENQUEUE_UNIQUE_RENDER_COMMAND_ONEPARAMETER(SlateUpdateFontTextureCommand,
-			FSlateFontTextureRHI&, FontTexture, *this,
+		FSlateFontTextureRHI* This = this;
+		ENQUEUE_RENDER_COMMAND(SlateUpdateFontTextureCommand)(
+			[This](FRHICommandListImmediate& RHICmdList)
 			{
-				FontTexture.UpdateTextureFromSource(FontTexture.PendingSourceData->SourceWidth, FontTexture.PendingSourceData->SourceHeight, FontTexture.PendingSourceData->SourceData);
-				FontTexture.PendingSourceData.Reset();
+				This->UpdateTextureFromSource(This->PendingSourceData->SourceWidth, This->PendingSourceData->SourceHeight, This->PendingSourceData->SourceData);
+				This->PendingSourceData.Reset();
 			});
 	}
 }
diff --git a/Engine/Source/Runtime/SlateRHIRenderer/Private/SlateRHIRenderer.cpp b/Engine/Source/Runtime/SlateRHIRenderer/Private/SlateRHIRenderer.cpp
--- a/Engine/Source/Runtime/SlateRHIRenderer/Private/SlateRHIRenderer.cpp
+++ b/Engine/Source/Runtime/SlateRHIRenderer/Private/SlateRHIRenderer.cpp
@@ -1550,8 +1550,9 @@ void FSlateRHIRenderer::AddWidgetRendererUpdate(const struct FRenderThreadUpdate
 	else
 	{
 		// Enqueue a command to unlock the draw buffer after all windows have been drawn
-		ENQUEUE_UNIQUE_RENDER_COMMAND_ONEPARAMETER(DrawWidgetRendererImmediate,
-			FRenderThreadUpdateContext, InContext, Context,
+		FRenderThreadUpdateContext InContext = Context;
+		ENQUEUE_RENDER_COMMAND(DrawWidgetRendererImmediate)(
+			[InContext](FRHICommandListImmediate& RHICmdList)
 			{
 				static_cast<ISlate3DRenderer*>(InContext.Renderer)->DrawWindowToTarget_RenderThread(RHICmdList, InContext);
 			});
diff --git a/Engine/Source/Runtime/SlateRHIRenderer/Private/SlateRHIResourceManager.cpp b/Engine/Source/Runtime/SlateRHIRenderer/Private/SlateRHIResourceManager.cpp
--- a/Engine/Source/Runtime/SlateRHIRenderer/Private/SlateRHIResourceManager.cpp
+++ b/Engine/Source/Runtime/SlateRHIRenderer/Private/SlateRHIResourceManager.cpp
@@ -298,13 +298,12 @@ void FSlateRHIResourceManager::Tick(float DeltaSeconds)
 		{
 			FSlateRHIResourceManager* ResourceManager;
 		};
-		FDeleteCachedRenderDataContext DeleteCachedRenderDataContext =
+		FDeleteCachedRenderDataContext Context =
 		{
 			this,
 		};
-		ENQUEUE_UNIQUE_RENDER_COMMAND_ONEPARAMETER(
-			DeleteCachedRenderData,
-			FDeleteCachedRenderDataContext, Context, DeleteCachedRenderDataContext,
+		ENQUEUE_RENDER_COMMAND(DeleteCachedRenderData)(
+			[Context](FRHICommandListImmediate& RHICmdList)
 			{
 				// Go through the pending delete buffers and see if any of their fences has cleared
 				// the RHI thread, if so, they should be safe to delete now.
@@ -971,15 +970,14 @@ void FSlateRHIResourceManager::BeginReleasingRenderData(const FSlateRenderDataHa
 		const FSlateRenderDataHandle* RenderDataHandle;
 		const ILayoutCache* LayoutCacher;
 	};
-	FReleaseCachedRenderDataContext ReleaseCachedRenderDataContext =
+	FReleaseCachedRenderDataContext Context =
 	{
 		this,
 		RenderHandle,
 		RenderHandle->GetCacher()
 	};
-	ENQUEUE_UNIQUE_RENDER_COMMAND_ONEPARAMETER(
-		ReleaseCachedRenderData,
-		FReleaseCachedRenderDataContext, Context, ReleaseCachedRenderDataContext,
+	ENQUEUE_RENDER_COMMAND(ReleaseCachedRenderData)(
+		[Context](FRHICommandListImmediate& RHICmdList)
 		{
 			Context.ResourceManager->ReleaseCachedRenderData(RHICmdList, Context.RenderDataHandle, Context.LayoutCacher);
 		});
diff --git a/Engine/Source/Runtime/StreamingPauseRendering/Private/StreamingPauseRendering.cpp b/Engine/Source/Runtime/StreamingPauseRendering/Private/StreamingPauseRendering.cpp
--- a/Engine/Source/Runtime/StreamingPauseRendering/Private/StreamingPauseRendering.cpp
+++ b/Engine/Source/Runtime/StreamingPauseRendering/Private/StreamingPauseRendering.cpp
@@ -136,10 +136,9 @@ void FStreamingPauseRenderingModule::BeginStreamingPause( FViewport* GameViewpor
 			}
 			Canvas.Flush_GameThread();
 
-
-			ENQUEUE_UNIQUE_RENDER_COMMAND_ONEPARAMETER(
-				EndDrawingCommand,
-				FViewport*, Viewport, SceneViewport.Get(),
+			FViewport* Viewport = SceneViewport.Get();
+			ENQUEUE_RENDER_COMMAND(EndDrawingCommand)(
+				[Viewport](FRHICommandListImmediate& RHICmdList)
 				{
 					Viewport->EndRenderFrame(RHICmdList, false, false);
 				});
