Fix bad merge resulting in misnamed and duplicated FSendingRepState::BuildConditionMapFromRepFlags function.

diff --git a/Engine/Source/Runtime/Engine/Private/DataReplication.cpp b/Engine/Source/Runtime/Engine/Private/DataReplication.cpp
--- a/Engine/Source/Runtime/Engine/Private/DataReplication.cpp
+++ b/Engine/Source/Runtime/Engine/Private/DataReplication.cpp
@@ -1126,7 +1126,7 @@ void FObjectReplicator::ReplicateCustomDeltaProperties( FNetBitWriter & Bunch, F
 	UNetConnection * OwningChannelConnection = OwningChannel->Connection;
 
 	// Initialize a map of which conditions are valid
-	const TStaticBitArray<COND_Max> ConditionMap = FSendingRepState::BuildConditionMap(RepFlags);
+	const TStaticBitArray<COND_Max> ConditionMap = FSendingRepState::BuildConditionMapFromRepFlags(RepFlags);
 
 	// Make sure net field export group is registered
 	FNetFieldExportGroup* NetFieldExportGroup = OwningChannel->GetOrCreateNetFieldExportGroupForClassNetCache( Object );
diff --git a/Engine/Source/Runtime/Engine/Private/RepLayout.cpp b/Engine/Source/Runtime/Engine/Private/RepLayout.cpp
--- a/Engine/Source/Runtime/Engine/Private/RepLayout.cpp
+++ b/Engine/Source/Runtime/Engine/Private/RepLayout.cpp
@@ -4959,7 +4959,7 @@ void FRepLayout::BuildHandleToCmdIndexTable_r(
 	}
 }
 
-TStaticBitArray<COND_Max> FSendingRepState::BuildConditionMap(const FReplicationFlags& RepFlags)
+TStaticBitArray<COND_Max> FSendingRepState::BuildConditionMapFromRepFlags(const FReplicationFlags RepFlags)
 {
 	TStaticBitArray<COND_Max> ConditionMap;
 
@@ -5000,7 +5000,7 @@ void FRepLayout::RebuildConditionalProperties(
 {
 	SCOPE_CYCLE_COUNTER(STAT_NetRebuildConditionalTime);
 	
-	TStaticBitArray<COND_Max> ConditionMap = FSendingRepState::BuildConditionMap(RepFlags);
+	TStaticBitArray<COND_Max> ConditionMap = FSendingRepState::BuildConditionMapFromRepFlags(RepFlags);
 	for (auto It = TBitArray<>::FIterator(RepState->InactiveParents); It; ++It)
 	{
 		It.GetValue() = !ConditionMap[Parents[It.GetIndex()].Condition];
@@ -5231,41 +5231,6 @@ void FSendingRepState::CountBytes(FArchive& Ar) const
 	GRANULAR_NETWORK_MEMORY_TRACKING_TRACK("InactiveParents", InactiveParents.CountBytes(Ar));
 }
 
-TStaticBitArray<COND_Max> FSendingRepState::BuildConditionMapFromRepFlags(FReplicationFlags RepFlags)
-{
-	TStaticBitArray<COND_Max> ConditionMap;
-
-	// Setup condition map
-	const bool bIsInitial = RepFlags.bNetInitial ? true : false;
-	const bool bIsOwner = RepFlags.bNetOwner ? true : false;
-	const bool bIsSimulated = RepFlags.bNetSimulated ? true : false;
-	const bool bIsPhysics = RepFlags.bRepPhysics ? true : false;
-	const bool bIsReplay = RepFlags.bReplay ? true : false;
-
-	ConditionMap[COND_None] = true;
-	ConditionMap[COND_InitialOnly] = bIsInitial;
-
-	ConditionMap[COND_OwnerOnly] = bIsOwner;
-	ConditionMap[COND_SkipOwner] = !bIsOwner;
-
-	ConditionMap[COND_SimulatedOnly] = bIsSimulated;
-	ConditionMap[COND_SimulatedOnlyNoReplay] = bIsSimulated && !bIsReplay;
-	ConditionMap[COND_AutonomousOnly] = !bIsSimulated;
-
-	ConditionMap[COND_SimulatedOrPhysics] = bIsSimulated || bIsPhysics;
-	ConditionMap[COND_SimulatedOrPhysicsNoReplay] = (bIsSimulated || bIsPhysics) && !bIsReplay;
-
-	ConditionMap[COND_InitialOrOwner] = bIsInitial || bIsOwner;
-	ConditionMap[COND_ReplayOrOwner] = bIsReplay || bIsOwner;
-	ConditionMap[COND_ReplayOnly] = bIsReplay;
-	ConditionMap[COND_SkipReplay] = !bIsReplay;
-
-	ConditionMap[COND_Custom] = true;
-	ConditionMap[COND_Never] = false;
-
-	return ConditionMap;
-}
-
 void FRepState::CountBytes(FArchive& Ar) const
 {
 	GRANULAR_NETWORK_MEMORY_TRACKING_INIT(Ar, "FRepState::CountBytes");
diff --git a/Engine/Source/Runtime/Engine/Public/Net/RepLayout.h b/Engine/Source/Runtime/Engine/Public/Net/RepLayout.h
--- a/Engine/Source/Runtime/Engine/Public/Net/RepLayout.h
+++ b/Engine/Source/Runtime/Engine/Public/Net/RepLayout.h
@@ -625,7 +625,7 @@ public:
 	 *
 	 * TODO: This doesn't have to be part of FRepState.
 	 */
-	static TStaticBitArray<COND_Max> BuildConditionMap(const FReplicationFlags& InFlags);
+	static TStaticBitArray<COND_Max> BuildConditionMapFromRepFlags(const FReplicationFlags InFlags);
 
 	/** Whether or not FRepLayout::OpenAcked has been called with this FRepState. */
 	bool bOpenAckedCalled;
