Strip Additional IndexBuffers from resources when not in use, which should save about 1MB when used on mobile.

diff --git a/Engine/Plugins/Runtime/EditableMesh/Source/EditableMesh/EditableMesh.Build.cs b/Engine/Plugins/Runtime/EditableMesh/Source/EditableMesh/EditableMesh.Build.cs
--- a/Engine/Plugins/Runtime/EditableMesh/Source/EditableMesh/EditableMesh.Build.cs
+++ b/Engine/Plugins/Runtime/EditableMesh/Source/EditableMesh/EditableMesh.Build.cs
@@ -8,6 +8,7 @@ namespace UnrealBuildTool.Rules
 		{
 			PrivateDependencyModuleNames.AddRange(
                 new string[] {
+					"RHI",
 				    "Core",
 				    "CoreUObject",
 				    "Engine",
diff --git a/Engine/Plugins/Runtime/EditableMesh/Source/EditableMesh/EditableStaticMeshAdapter.cpp b/Engine/Plugins/Runtime/EditableMesh/Source/EditableMesh/EditableStaticMeshAdapter.cpp
--- a/Engine/Plugins/Runtime/EditableMesh/Source/EditableMesh/EditableStaticMeshAdapter.cpp
+++ b/Engine/Plugins/Runtime/EditableMesh/Source/EditableMesh/EditableStaticMeshAdapter.cpp
@@ -727,13 +727,18 @@ void UEditableStaticMeshAdapter::OnRebuildRenderMesh( const UEditableMesh* Edita
 	}
 
 	StaticMeshLOD.IndexBuffer.SetIndices( IndexBuffer, IndexBufferStride );
+	StaticMeshLOD.DepthOnlyIndexBuffer.SetIndices(TArray< uint32 >(), IndexBufferStride);
+
+	if (StaticMeshLOD.AdditionalIndexBuffers == nullptr)
+	{
+		StaticMeshLOD.AdditionalIndexBuffers = new FAdditionalStaticMeshIndexBuffers();
+	}
 
 	// @todo mesheditor: support the other index buffer types
-	StaticMeshLOD.ReversedIndexBuffer.SetIndices( TArray< uint32 >(), IndexBufferStride );
-	StaticMeshLOD.DepthOnlyIndexBuffer.SetIndices( TArray< uint32 >(), IndexBufferStride );
-	StaticMeshLOD.ReversedDepthOnlyIndexBuffer.SetIndices( TArray< uint32 >(), IndexBufferStride );
-	StaticMeshLOD.WireframeIndexBuffer.SetIndices( TArray< uint32 >(), IndexBufferStride );
-	StaticMeshLOD.AdjacencyIndexBuffer.SetIndices( TArray< uint32 >(), IndexBufferStride );
+	StaticMeshLOD.AdditionalIndexBuffers->ReversedIndexBuffer.SetIndices( TArray< uint32 >(), IndexBufferStride );
+	StaticMeshLOD.AdditionalIndexBuffers->ReversedDepthOnlyIndexBuffer.SetIndices( TArray< uint32 >(), IndexBufferStride );
+	StaticMeshLOD.AdditionalIndexBuffers->WireframeIndexBuffer.SetIndices( TArray< uint32 >(), IndexBufferStride );
+	StaticMeshLOD.AdditionalIndexBuffers->AdjacencyIndexBuffer.SetIndices( TArray< uint32 >(), IndexBufferStride );
 
 	StaticMeshLOD.bHasAdjacencyInfo = false;
 	StaticMeshLOD.bHasDepthOnlyIndices = false;
diff --git a/Engine/Source/Developer/MeshBuilder/MeshBuilder.Build.cs b/Engine/Source/Developer/MeshBuilder/MeshBuilder.Build.cs
--- a/Engine/Source/Developer/MeshBuilder/MeshBuilder.Build.cs
+++ b/Engine/Source/Developer/MeshBuilder/MeshBuilder.Build.cs
@@ -10,6 +10,7 @@ namespace UnrealBuildTool.Rules
             PrivateDependencyModuleNames.AddRange(
 				new string[]
 				{
+					"RHI",
 					"Core",
 					"CoreUObject",
                     "Engine",
diff --git a/Engine/Source/Developer/MeshBuilder/Private/StaticMeshBuilder.cpp b/Engine/Source/Developer/MeshBuilder/Private/StaticMeshBuilder.cpp
--- a/Engine/Source/Developer/MeshBuilder/Private/StaticMeshBuilder.cpp
+++ b/Engine/Source/Developer/MeshBuilder/Private/StaticMeshBuilder.cpp
@@ -523,6 +523,11 @@ void BuildVertexBuffer(
 
 void BuildAllBufferOptimizations(FStaticMeshLODResources& StaticMeshLOD, const FMeshBuildSettings& LODBuildSettings, TArray< uint32 >& IndexBuffer, bool bNeeds32BitIndices, TArray< FStaticMeshBuildVertex >& StaticMeshBuildVertices)
 {
+	if (StaticMeshLOD.AdditionalIndexBuffers == nullptr)
+	{
+		StaticMeshLOD.AdditionalIndexBuffers = new FAdditionalStaticMeshIndexBuffers();
+	}
+
 	const EIndexBufferStride::Type IndexBufferStride = bNeeds32BitIndices ? EIndexBufferStride::Force32Bit : EIndexBufferStride::Force16Bit;
 
 	// Build the reversed index buffer.
@@ -542,7 +547,7 @@ void BuildAllBufferOptimizations(FStaticMeshLODResources& StaticMeshLOD, const F
 				InversedIndices[SectionInfo.FirstIndex + i] = IndexBuffer[SectionInfo.FirstIndex + SectionIndexCount - 1 - i];
 			}
 		}
-		StaticMeshLOD.ReversedIndexBuffer.SetIndices(InversedIndices, IndexBufferStride);
+		StaticMeshLOD.AdditionalIndexBuffers->ReversedIndexBuffer.SetIndices(InversedIndices, IndexBufferStride);
 	}
 
 	// Build the depth-only index buffer.
@@ -573,7 +578,7 @@ void BuildAllBufferOptimizations(FStaticMeshLODResources& StaticMeshLOD, const F
 		{
 			ReversedDepthOnlyIndices[i] = DepthOnlyIndices[IndexCount - 1 - i];
 		}
-		StaticMeshLOD.ReversedDepthOnlyIndexBuffer.SetIndices(ReversedDepthOnlyIndices, IndexBufferStride);
+		StaticMeshLOD.AdditionalIndexBuffers->ReversedDepthOnlyIndexBuffer.SetIndices(ReversedDepthOnlyIndices, IndexBufferStride);
 	}
 
 	// Build a list of wireframe edges in the static mesh.
@@ -589,7 +594,7 @@ void BuildAllBufferOptimizations(FStaticMeshLODResources& StaticMeshLOD, const F
 			WireframeIndices.Add(Edge.Vertices[0]);
 			WireframeIndices.Add(Edge.Vertices[1]);
 		}
-		StaticMeshLOD.WireframeIndexBuffer.SetIndices(WireframeIndices, IndexBufferStride);
+		StaticMeshLOD.AdditionalIndexBuffers->WireframeIndexBuffer.SetIndices(WireframeIndices, IndexBufferStride);
 	}
 
 	// Build the adjacency index buffer used for tessellation.
@@ -603,6 +608,6 @@ void BuildAllBufferOptimizations(FStaticMeshLODResources& StaticMeshLOD, const F
 			IndexBuffer,
 			AdjacencyIndices
 		);
-		StaticMeshLOD.AdjacencyIndexBuffer.SetIndices(AdjacencyIndices, IndexBufferStride);
+		StaticMeshLOD.AdditionalIndexBuffers->AdjacencyIndexBuffer.SetIndices(AdjacencyIndices, IndexBufferStride);
 	}
 }
diff --git a/Engine/Source/Developer/MeshUtilities/Private/MeshUtilities.cpp b/Engine/Source/Developer/MeshUtilities/Private/MeshUtilities.cpp
--- a/Engine/Source/Developer/MeshUtilities/Private/MeshUtilities.cpp
+++ b/Engine/Source/Developer/MeshUtilities/Private/MeshUtilities.cpp
@@ -2791,7 +2791,7 @@ public:
 						InversedIndices[SectionInfo.FirstIndex + i] = CombinedIndices[SectionInfo.FirstIndex + SectionIndexCount - 1 - i];
 					}
 				}
-				LODModel.ReversedIndexBuffer.SetIndices(InversedIndices, bNeeds32BitIndices ? EIndexBufferStride::Force32Bit : EIndexBufferStride::Force16Bit);
+				LODModel.AdditionalIndexBuffers->ReversedIndexBuffer.SetIndices(InversedIndices, bNeeds32BitIndices ? EIndexBufferStride::Force32Bit : EIndexBufferStride::Force16Bit);
 			}
 
 			// Build the depth-only index buffer.
@@ -2822,7 +2822,7 @@ public:
 				{
 					ReversedDepthOnlyIndices[i] = DepthOnlyIndices[IndexCount - 1 - i];
 				}
-				LODModel.ReversedDepthOnlyIndexBuffer.SetIndices(ReversedDepthOnlyIndices, bNeeds32BitIndices ? EIndexBufferStride::Force32Bit : EIndexBufferStride::Force16Bit);
+				LODModel.AdditionalIndexBuffers->ReversedDepthOnlyIndexBuffer.SetIndices(ReversedDepthOnlyIndices, bNeeds32BitIndices ? EIndexBufferStride::Force32Bit : EIndexBufferStride::Force16Bit);
 			}
 
 			// Build a list of wireframe edges in the static mesh.
@@ -2838,7 +2838,7 @@ public:
 					WireframeIndices.Add(Edge.Vertices[0]);
 					WireframeIndices.Add(Edge.Vertices[1]);
 				}
-				LODModel.WireframeIndexBuffer.SetIndices(WireframeIndices, bNeeds32BitIndices ? EIndexBufferStride::Force32Bit : EIndexBufferStride::Force16Bit);
+				LODModel.AdditionalIndexBuffers->WireframeIndexBuffer.SetIndices(WireframeIndices, bNeeds32BitIndices ? EIndexBufferStride::Force32Bit : EIndexBufferStride::Force16Bit);
 			}
 
 			// Build the adjacency index buffer used for tessellation.
@@ -2852,7 +2852,7 @@ public:
 					CombinedIndices,
 					AdjacencyIndices
 					);
-				LODModel.AdjacencyIndexBuffer.SetIndices(AdjacencyIndices, bNeeds32BitIndices ? EIndexBufferStride::Force32Bit : EIndexBufferStride::Force16Bit);
+				LODModel.AdditionalIndexBuffers->AdjacencyIndexBuffer.SetIndices(AdjacencyIndices, bNeeds32BitIndices ? EIndexBufferStride::Force32Bit : EIndexBufferStride::Force16Bit);
 			}
 		}
 
diff --git a/Engine/Source/Runtime/Engine/Private/Particles/ParticleSystemRender.cpp b/Engine/Source/Runtime/Engine/Private/Particles/ParticleSystemRender.cpp
--- a/Engine/Source/Runtime/Engine/Private/Particles/ParticleSystemRender.cpp
+++ b/Engine/Source/Runtime/Engine/Private/Particles/ParticleSystemRender.cpp
@@ -1863,14 +1863,14 @@ void FDynamicMeshEmitterData::GetDynamicMeshElementsEmitter(const FParticleSyste
 
 					if (bIsWireframe)
 					{
-						if (LODModel.WireframeIndexBuffer.IsInitialized()
+						if (LODModel.AdditionalIndexBuffers->WireframeIndexBuffer.IsInitialized()
 							&& !(RHISupportsTessellation(ShaderPlatform) && Mesh.VertexFactory->GetType()->SupportsTessellationShaders()))
 						{
 							Mesh.Type = PT_LineList;
 							Mesh.MaterialRenderProxy = Proxy->GetDeselectedWireframeMatInst();
 							BatchElement.FirstIndex = 0;
-							BatchElement.IndexBuffer = &LODModel.WireframeIndexBuffer;
-							BatchElement.NumPrimitives = LODModel.WireframeIndexBuffer.GetNumIndices() / 2;
+							BatchElement.IndexBuffer = &LODModel.AdditionalIndexBuffers->WireframeIndexBuffer;
+							BatchElement.NumPrimitives = LODModel.AdditionalIndexBuffers->WireframeIndexBuffer.GetNumIndices() / 2;
 
 						}
 						else
diff --git a/Engine/Source/Runtime/Engine/Private/StaticMesh.cpp b/Engine/Source/Runtime/Engine/Private/StaticMesh.cpp
--- a/Engine/Source/Runtime/Engine/Private/StaticMesh.cpp
+++ b/Engine/Source/Runtime/Engine/Private/StaticMesh.cpp
@@ -349,15 +349,34 @@ void FStaticMeshLODResources::SerializeBuffers(FArchive& Ar, UStaticMesh* OwnerS
 	IndexBuffer.Serialize(Ar, bNeedsCPUAccess);
 	AccumIndexBufferSize(IndexBuffer, OutBuffersSize.SerializedBuffersSize);
 
-	const bool bSerailizeReversedIndexBuffer = !StripFlags.IsClassDataStripped(CDSF_ReversedIndexBuffer);
-	if (bSerailizeReversedIndexBuffer)
+	const bool bSerializeReversedIndexBuffer = !StripFlags.IsClassDataStripped(CDSF_ReversedIndexBuffer);
+	const bool bSerializeAdjacencyDataIndexBuffer = !StripFlags.IsClassDataStripped(CDSF_AdjacencyData);
+	const bool bSerializeWireframeIndexBuffer = !StripFlags.IsEditorDataStripped();
+
+	FAdditionalStaticMeshIndexBuffers DummyBuffers;
+	FAdditionalStaticMeshIndexBuffers* SerializedAdditionalIndexBuffers = &DummyBuffers;
+	if ((bEnableDepthOnlyIndexBuffer || bEnableReversedIndexBuffer) && (bSerializeReversedIndexBuffer || bSerializeAdjacencyDataIndexBuffer || bSerializeWireframeIndexBuffer || bEnableDepthOnlyIndexBuffer))
 	{
-		ReversedIndexBuffer.Serialize(Ar, bNeedsCPUAccess);
-		AccumIndexBufferSize(ReversedIndexBuffer, OutBuffersSize.ReversedIBsSize);
-		AccumIndexBufferSize(ReversedIndexBuffer, OutBuffersSize.SerializedBuffersSize);
+		if (AdditionalIndexBuffers == nullptr)
+		{
+			AdditionalIndexBuffers = new FAdditionalStaticMeshIndexBuffers();
+		}
+		SerializedAdditionalIndexBuffers = AdditionalIndexBuffers;
+	}
+	
+	SerializedAdditionalIndexBuffers->ReversedIndexBuffer.SetIsStreamed(!bBuffersInlined);
+	SerializedAdditionalIndexBuffers->ReversedDepthOnlyIndexBuffer.SetIsStreamed(!bBuffersInlined);
+	SerializedAdditionalIndexBuffers->WireframeIndexBuffer.SetIsStreamed(!bBuffersInlined);
+	SerializedAdditionalIndexBuffers->AdjacencyIndexBuffer.SetIsStreamed(!bBuffersInlined);
+
+	if (bSerializeReversedIndexBuffer)
+	{
+		SerializedAdditionalIndexBuffers->ReversedIndexBuffer.Serialize(Ar, bNeedsCPUAccess);
+		AccumIndexBufferSize(SerializedAdditionalIndexBuffers->ReversedIndexBuffer, OutBuffersSize.ReversedIBsSize);
+		AccumIndexBufferSize(SerializedAdditionalIndexBuffers->ReversedIndexBuffer, OutBuffersSize.SerializedBuffersSize);
 		if (!bEnableReversedIndexBuffer)
 		{
-			ReversedIndexBuffer.Discard();
+			SerializedAdditionalIndexBuffers->ReversedIndexBuffer.Discard();
 		}
 	}
 
@@ -369,35 +388,35 @@ void FStaticMeshLODResources::SerializeBuffers(FArchive& Ar, UStaticMesh* OwnerS
 		DepthOnlyIndexBuffer.Discard();
 	}
 
-	if (bSerailizeReversedIndexBuffer)
+	if (bSerializeReversedIndexBuffer)
 	{
-		ReversedDepthOnlyIndexBuffer.Serialize(Ar, bNeedsCPUAccess);
-		AccumIndexBufferSize(ReversedDepthOnlyIndexBuffer, OutBuffersSize.ReversedIBsSize);
-		AccumIndexBufferSize(ReversedDepthOnlyIndexBuffer, OutBuffersSize.SerializedBuffersSize);
+		SerializedAdditionalIndexBuffers->ReversedDepthOnlyIndexBuffer.Serialize(Ar, bNeedsCPUAccess);
+		AccumIndexBufferSize(SerializedAdditionalIndexBuffers->ReversedDepthOnlyIndexBuffer, OutBuffersSize.ReversedIBsSize);
+		AccumIndexBufferSize(SerializedAdditionalIndexBuffers->ReversedDepthOnlyIndexBuffer, OutBuffersSize.SerializedBuffersSize);
 		if (!bEnableReversedIndexBuffer)
 		{
-			ReversedDepthOnlyIndexBuffer.Discard();
+			SerializedAdditionalIndexBuffers->ReversedDepthOnlyIndexBuffer.Discard();
 		}
 	}
 
-	if (!StripFlags.IsEditorDataStripped())
+	if (bSerializeWireframeIndexBuffer)
 	{
-		WireframeIndexBuffer.Serialize(Ar, bNeedsCPUAccess);
-		AccumIndexBufferSize(WireframeIndexBuffer, OutBuffersSize.SerializedBuffersSize);
-		bHasWireframeIndices = WireframeIndexBuffer.GetNumIndices() != 0;
+		SerializedAdditionalIndexBuffers->WireframeIndexBuffer.Serialize(Ar, bNeedsCPUAccess);
+		AccumIndexBufferSize(SerializedAdditionalIndexBuffers->WireframeIndexBuffer, OutBuffersSize.SerializedBuffersSize);
+		bHasWireframeIndices = AdditionalIndexBuffers && SerializedAdditionalIndexBuffers->WireframeIndexBuffer.GetNumIndices() != 0;
 	}
 
-	if (!StripFlags.IsClassDataStripped(CDSF_AdjacencyData))
+	if (bSerializeAdjacencyDataIndexBuffer)
 	{
-		AdjacencyIndexBuffer.Serialize(Ar, bNeedsCPUAccess);
-		AccumIndexBufferSize(AdjacencyIndexBuffer, OutBuffersSize.SerializedBuffersSize);
-		bHasAdjacencyInfo = AdjacencyIndexBuffer.GetNumIndices() != 0;
+		SerializedAdditionalIndexBuffers->AdjacencyIndexBuffer.Serialize(Ar, bNeedsCPUAccess);
+		AccumIndexBufferSize(SerializedAdditionalIndexBuffers->AdjacencyIndexBuffer, OutBuffersSize.SerializedBuffersSize);
+		bHasAdjacencyInfo = AdditionalIndexBuffers && SerializedAdditionalIndexBuffers->AdjacencyIndexBuffer.GetNumIndices() != 0;
 	}
 
 	// Needs to be done now because on cooked platform, indices are discarded after RHIInit.
 	bHasDepthOnlyIndices = DepthOnlyIndexBuffer.GetNumIndices() != 0;
-	bHasReversedIndices = bSerailizeReversedIndexBuffer && ReversedIndexBuffer.GetNumIndices() != 0;
-	bHasReversedDepthOnlyIndices = bSerailizeReversedIndexBuffer && ReversedDepthOnlyIndexBuffer.GetNumIndices() != 0;
+	bHasReversedIndices = AdditionalIndexBuffers && bSerializeReversedIndexBuffer && SerializedAdditionalIndexBuffers->ReversedIndexBuffer.GetNumIndices() != 0;
+	bHasReversedDepthOnlyIndices = AdditionalIndexBuffers && bSerializeReversedIndexBuffer && SerializedAdditionalIndexBuffers->ReversedDepthOnlyIndexBuffer.GetNumIndices() != 0;
 	bHasColorVertexData = VertexBuffers.ColorVertexBuffer.GetNumVertices() > 0;
 	DepthOnlyNumTriangles = DepthOnlyIndexBuffer.GetNumIndices() / 3;
 
@@ -411,6 +430,9 @@ void FStaticMeshLODResources::SerializeBuffers(FArchive& Ar, UStaticMesh* OwnerS
 
 void FStaticMeshLODResources::SerializeAvailabilityInfo(FArchive& Ar)
 {
+	const bool bEnableDepthOnlyIndexBuffer = !!CVarSupportDepthOnlyIndexBuffers.GetValueOnAnyThread();
+	const bool bEnableReversedIndexBuffer = !!CVarSupportReversedIndexBuffers.GetValueOnAnyThread();
+
 	Ar << DepthOnlyNumTriangles;
 	uint32 Packed;
 #if WITH_EDITOR
@@ -427,8 +449,6 @@ void FStaticMeshLODResources::SerializeAvailabilityInfo(FArchive& Ar)
 	else
 #endif
 	{
-		const bool bEnableDepthOnlyIndexBuffer = !!CVarSupportDepthOnlyIndexBuffers.GetValueOnAnyThread();
-		const bool bEnableReversedIndexBuffer = !!CVarSupportReversedIndexBuffers.GetValueOnAnyThread();
 		Ar << Packed;
 		DepthOnlyNumTriangles *= static_cast<uint32>(bEnableDepthOnlyIndexBuffer);
 		bHasAdjacencyInfo = Packed & 1u;
@@ -443,32 +463,50 @@ void FStaticMeshLODResources::SerializeAvailabilityInfo(FArchive& Ar)
 	VertexBuffers.PositionVertexBuffer.SerializeMetaData(Ar);
 	VertexBuffers.ColorVertexBuffer.SerializeMetaData(Ar);
 	IndexBuffer.SerializeMetaData(Ar);
-	ReversedIndexBuffer.SerializeMetaData(Ar);
+
+	FAdditionalStaticMeshIndexBuffers DummyBuffers;
+	FAdditionalStaticMeshIndexBuffers* SerializedAdditionalIndexBuffers = &DummyBuffers;
+	if ((bEnableDepthOnlyIndexBuffer || bEnableReversedIndexBuffer) && (bHasReversedIndices || bHasAdjacencyInfo || bHasWireframeIndices || bHasDepthOnlyIndices))
+	{
+		if (AdditionalIndexBuffers == nullptr)
+		{
+			AdditionalIndexBuffers = new FAdditionalStaticMeshIndexBuffers();
+		}
+		SerializedAdditionalIndexBuffers = AdditionalIndexBuffers;
+	}
+
+	SerializedAdditionalIndexBuffers->ReversedIndexBuffer.SetIsStreamed(!bBuffersInlined);
+	SerializedAdditionalIndexBuffers->ReversedDepthOnlyIndexBuffer.SetIsStreamed(!bBuffersInlined);
+	SerializedAdditionalIndexBuffers->WireframeIndexBuffer.SetIsStreamed(!bBuffersInlined);
+	SerializedAdditionalIndexBuffers->AdjacencyIndexBuffer.SetIsStreamed(!bBuffersInlined);
+	
+
+	SerializedAdditionalIndexBuffers->ReversedIndexBuffer.SerializeMetaData(Ar);
 	if (!bHasReversedIndices)
 	{
 		// Reversed indices are either stripped during cook or will be stripped on load.
 		// In either case, clear CachedNumIndices to show that the buffer will be empty after actual loading
-		ReversedIndexBuffer.Discard();
+		SerializedAdditionalIndexBuffers->ReversedIndexBuffer.Discard();
 	}
 	DepthOnlyIndexBuffer.SerializeMetaData(Ar);
 	if (!bHasDepthOnlyIndices)
 	{
 		DepthOnlyIndexBuffer.Discard();
 	}
-	ReversedDepthOnlyIndexBuffer.SerializeMetaData(Ar);
+	SerializedAdditionalIndexBuffers->ReversedDepthOnlyIndexBuffer.SerializeMetaData(Ar);
 	if (!bHasReversedDepthOnlyIndices)
 	{
-		ReversedDepthOnlyIndexBuffer.Discard();
+		SerializedAdditionalIndexBuffers->ReversedDepthOnlyIndexBuffer.Discard();
 	}
-	WireframeIndexBuffer.SerializeMetaData(Ar);
+	SerializedAdditionalIndexBuffers->WireframeIndexBuffer.SerializeMetaData(Ar);
 	if (!bHasWireframeIndices)
 	{
-		WireframeIndexBuffer.Discard();
+		SerializedAdditionalIndexBuffers->WireframeIndexBuffer.Discard();
 	}
-	AdjacencyIndexBuffer.SerializeMetaData(Ar);
+	SerializedAdditionalIndexBuffers->AdjacencyIndexBuffer.SerializeMetaData(Ar);
 	if (!bHasAdjacencyInfo)
 	{
-		AdjacencyIndexBuffer.Discard();
+		SerializedAdditionalIndexBuffers->AdjacencyIndexBuffer.Discard();
 	}
 }
 
@@ -496,11 +534,7 @@ void FStaticMeshLODResources::Serialize(FArchive& Ar, UObject* Owner, int32 Inde
 	VertexBuffers.PositionVertexBuffer.SetIsStreamed(!bBuffersInlined);
 	VertexBuffers.ColorVertexBuffer.SetIsStreamed(!bBuffersInlined);
 	IndexBuffer.SetIsStreamed(!bBuffersInlined);
-	ReversedIndexBuffer.SetIsStreamed(!bBuffersInlined);
 	DepthOnlyIndexBuffer.SetIsStreamed(!bBuffersInlined);
-	ReversedDepthOnlyIndexBuffer.SetIsStreamed(!bBuffersInlined);
-	WireframeIndexBuffer.SetIsStreamed(!bBuffersInlined);
-	AdjacencyIndexBuffer.SetIsStreamed(!bBuffersInlined);
 
 	if (!StripFlags.IsDataStrippedForServer() && !bIsLODCookedOut)
 	{
@@ -878,7 +912,8 @@ void FStaticMeshVertexBuffers::InitFromDynamicVertex(FLocalVertexFactory* Vertex
 };
 
 FStaticMeshLODResources::FStaticMeshLODResources()
-	: DistanceFieldData(NULL)
+	: AdditionalIndexBuffers(nullptr)
+	, DistanceFieldData(nullptr)
 	, MaxDeviation(0.0f)
 	, bHasAdjacencyInfo(false)
 	, bHasDepthOnlyIndices(false)
@@ -899,6 +934,7 @@ FStaticMeshLODResources::FStaticMeshLODResources()
 FStaticMeshLODResources::~FStaticMeshLODResources()
 {
 	delete DistanceFieldData;
+	delete AdditionalIndexBuffers;
 }
 
 void FStaticMeshLODResources::ConditionalForce16BitIndexBuffer(EShaderPlatform MaxShaderPlatform, UStaticMesh* Parent)
@@ -924,13 +960,17 @@ void FStaticMeshLODResources::UpdateIndexMemoryStats()
 #if STATS
 	if (bIncrement)
 	{
-		uint32 iMem = IndexBuffer.GetAllocatedSize();
-		uint32 wiMem = WireframeIndexBuffer.GetAllocatedSize();
-		uint32 riMem = ReversedIndexBuffer.GetAllocatedSize();
-		uint32 doiMem = DepthOnlyIndexBuffer.GetAllocatedSize();
-		uint32 rdoiMem = ReversedDepthOnlyIndexBuffer.GetAllocatedSize();
-		uint32 aiMem = AdjacencyIndexBuffer.GetAllocatedSize();
-		StaticMeshIndexMemory = iMem + wiMem + riMem + doiMem + rdoiMem + aiMem;
+		StaticMeshIndexMemory += IndexBuffer.GetAllocatedSize();
+		StaticMeshIndexMemory += DepthOnlyIndexBuffer.GetAllocatedSize();
+
+		if (AdditionalIndexBuffers)
+		{
+			StaticMeshIndexMemory += AdditionalIndexBuffers->WireframeIndexBuffer.GetAllocatedSize();
+			StaticMeshIndexMemory += AdditionalIndexBuffers->ReversedIndexBuffer.GetAllocatedSize();
+			StaticMeshIndexMemory += AdditionalIndexBuffers->ReversedDepthOnlyIndexBuffer.GetAllocatedSize();
+			StaticMeshIndexMemory += AdditionalIndexBuffers->AdjacencyIndexBuffer.GetAllocatedSize();
+		}
+
 		INC_DWORD_STAT_BY(STAT_StaticMeshIndexMemory, StaticMeshIndexMemory);
 	}
 	else
@@ -970,7 +1010,7 @@ void FStaticMeshLODResources::InitResources(UStaticMesh* Parent)
 	BeginInitResource(&IndexBuffer);
 	if(bHasWireframeIndices)
 	{
-		BeginInitResource(&WireframeIndexBuffer);
+		BeginInitResource(&AdditionalIndexBuffers->WireframeIndexBuffer);
 	}
 	BeginInitResource(&VertexBuffers.StaticMeshVertexBuffer);
 	BeginInitResource(&VertexBuffers.PositionVertexBuffer);
@@ -981,7 +1021,7 @@ void FStaticMeshLODResources::InitResources(UStaticMesh* Parent)
 
 	if (bHasReversedIndices)
 	{
-		BeginInitResource(&ReversedIndexBuffer);
+		BeginInitResource(&AdditionalIndexBuffers->ReversedIndexBuffer);
 	}
 
 	if (bHasDepthOnlyIndices)
@@ -991,12 +1031,12 @@ void FStaticMeshLODResources::InitResources(UStaticMesh* Parent)
 
 	if (bHasReversedDepthOnlyIndices)
 	{
-		BeginInitResource(&ReversedDepthOnlyIndexBuffer);
+		BeginInitResource(&AdditionalIndexBuffers->ReversedDepthOnlyIndexBuffer);
 	}
 
 	if (bHasAdjacencyInfo && RHISupportsTessellation(GMaxRHIShaderPlatform))
 	{
-		BeginInitResource(&AdjacencyIndexBuffer);
+		BeginInitResource(&AdditionalIndexBuffers->AdjacencyIndexBuffer);
 	}
 
 #if RHI_RAYTRACING
@@ -1057,19 +1097,25 @@ void FStaticMeshLODResources::ReleaseResources()
 
 	// Release the vertex and index buffers.
 	
-	// AdjacencyIndexBuffer may not be initialized at this time, but it is safe to release it anyway.
-	// The bInitialized flag will be safely checked in the render thread.
-	// This avoids a race condition regarding releasing this resource.
-	BeginReleaseResource(&AdjacencyIndexBuffer);
+
 
 	BeginReleaseResource(&IndexBuffer);
-	BeginReleaseResource(&WireframeIndexBuffer);
+	
 	BeginReleaseResource(&VertexBuffers.StaticMeshVertexBuffer);
 	BeginReleaseResource(&VertexBuffers.PositionVertexBuffer);
 	BeginReleaseResource(&VertexBuffers.ColorVertexBuffer);
-	BeginReleaseResource(&ReversedIndexBuffer);
 	BeginReleaseResource(&DepthOnlyIndexBuffer);
-	BeginReleaseResource(&ReversedDepthOnlyIndexBuffer);
+
+	if (AdditionalIndexBuffers)
+	{
+		// AdjacencyIndexBuffer may not be initialized at this time, but it is safe to release it anyway.
+		// The bInitialized flag will be safely checked in the render thread.
+		// This avoids a race condition regarding releasing this resource.
+		BeginReleaseResource(&AdditionalIndexBuffers->AdjacencyIndexBuffer);
+		BeginReleaseResource(&AdditionalIndexBuffers->ReversedIndexBuffer);
+		BeginReleaseResource(&AdditionalIndexBuffers->WireframeIndexBuffer);
+		BeginReleaseResource(&AdditionalIndexBuffers->ReversedDepthOnlyIndexBuffer);
+	}
 #if RHI_RAYTRACING
 	BeginReleaseResource(&RayTracingGeometry);
 #endif // RHI_RAYTRACING
@@ -1099,11 +1145,15 @@ void FStaticMeshLODResources::DiscardCPUData()
 	VertexBuffers.PositionVertexBuffer.CleanUp();
 	VertexBuffers.ColorVertexBuffer.CleanUp();
 	IndexBuffer.Discard();
-	ReversedIndexBuffer.Discard();
 	DepthOnlyIndexBuffer.Discard();
-	ReversedDepthOnlyIndexBuffer.Discard();
-	WireframeIndexBuffer.Discard();
-	AdjacencyIndexBuffer.Discard();
+
+	if (AdditionalIndexBuffers)
+	{
+		AdditionalIndexBuffers->ReversedIndexBuffer.Discard();
+		AdditionalIndexBuffers->ReversedDepthOnlyIndexBuffer.Discard();
+		AdditionalIndexBuffers->WireframeIndexBuffer.Discard();
+		AdditionalIndexBuffers->AdjacencyIndexBuffer.Discard();
+	}
 }
 
 /*------------------------------------------------------------------------------
@@ -2392,9 +2442,13 @@ void FStaticMeshRenderData::GetResourceSizeEx(FResourceSizeEx& CumulativeResourc
 		const int32 VBSize = LODRenderData.VertexBuffers.StaticMeshVertexBuffer.GetResourceSize() +
 			LODRenderData.VertexBuffers.PositionVertexBuffer.GetStride()			* LODRenderData.VertexBuffers.PositionVertexBuffer.GetNumVertices() +
 			LODRenderData.VertexBuffers.ColorVertexBuffer.GetStride()				* LODRenderData.VertexBuffers.ColorVertexBuffer.GetNumVertices();
-		const int32 IBSize = LODRenderData.IndexBuffer.GetAllocatedSize()
-			+ LODRenderData.WireframeIndexBuffer.GetAllocatedSize()
-			+ (RHISupportsTessellation(GShaderPlatformForFeatureLevel[GMaxRHIFeatureLevel]) ? LODRenderData.AdjacencyIndexBuffer.GetAllocatedSize() : 0);
+		int32 IBSize = LODRenderData.IndexBuffer.GetAllocatedSize();
+
+		if (LODRenderData.AdditionalIndexBuffers)
+		{
+			IBSize += LODRenderData.AdditionalIndexBuffers->WireframeIndexBuffer.GetAllocatedSize()
+			+ (RHISupportsTessellation(GShaderPlatformForFeatureLevel[GMaxRHIFeatureLevel]) ? LODRenderData.AdditionalIndexBuffers->AdjacencyIndexBuffer.GetAllocatedSize() : 0);
+		}
 
 		CumulativeResourceSize.AddUnknownMemoryBytes(VBSize + IBSize);
 		CumulativeResourceSize.AddUnknownMemoryBytes(LODRenderData.Sections.GetAllocatedSize());
diff --git a/Engine/Source/Runtime/Engine/Private/StaticMeshRender.cpp b/Engine/Source/Runtime/Engine/Private/StaticMeshRender.cpp
--- a/Engine/Source/Runtime/Engine/Private/StaticMeshRender.cpp
+++ b/Engine/Source/Runtime/Engine/Private/StaticMeshRender.cpp
@@ -372,7 +372,7 @@ bool FStaticMeshSceneProxy::GetShadowMeshElement(int32 LODIndex, int32 BatchInde
 	OutBatchElement.VertexFactoryUserData = ProxyLODInfo.OverrideColorVertexBuffer ? ProxyLODInfo.OverrideColorVFUniformBuffer.GetReference() : VFs.VertexFactory.GetUniformBuffer();
 	OutMeshBatch.MaterialRenderProxy = UMaterial::GetDefaultMaterial(MD_Surface)->GetRenderProxy();
 	OutMeshBatch.VertexFactory = ProxyLODInfo.OverrideColorVertexBuffer ? &VFs.VertexFactoryOverrideColorVertexBuffer : &VFs.VertexFactory;
-	OutBatchElement.IndexBuffer = bUseReversedIndices ? &LOD.ReversedDepthOnlyIndexBuffer : &LOD.DepthOnlyIndexBuffer;
+	OutBatchElement.IndexBuffer = LOD.AdditionalIndexBuffers && bUseReversedIndices ? &LOD.AdditionalIndexBuffers->ReversedDepthOnlyIndexBuffer : &LOD.DepthOnlyIndexBuffer;
 	OutMeshBatch.Type = PT_TriangleList;
 	OutBatchElement.FirstIndex = 0;
 	OutBatchElement.NumPrimitives = LOD.DepthOnlyNumTriangles;
@@ -655,14 +655,14 @@ void FStaticMeshSceneProxy::SetIndexSource(int32 LODIndex, int32 SectionIndex, F
 	const FStaticMeshLODResources& LODModel = RenderData->LODResources[LODIndex];
 	if (bWireframe)
 	{
-		if( LODModel.WireframeIndexBuffer.IsInitialized()
+		if( LODModel.AdditionalIndexBuffers->WireframeIndexBuffer.IsInitialized()
 			&& !(RHISupportsTessellation(GetScene().GetShaderPlatform()) && OutMeshElement.VertexFactory->GetType()->SupportsTessellationShaders())
 			)
 		{
 			OutMeshElement.Type = PT_LineList;
 			OutElement.FirstIndex = 0;
-			OutElement.IndexBuffer = &LODModel.WireframeIndexBuffer;
-			OutElement.NumPrimitives = LODModel.WireframeIndexBuffer.GetNumIndices() / 2;
+			OutElement.IndexBuffer = &LODModel.AdditionalIndexBuffers->WireframeIndexBuffer;
+			OutElement.NumPrimitives = LODModel.AdditionalIndexBuffers->WireframeIndexBuffer.GetNumIndices() / 2;
 		}
 		else
 		{
@@ -702,7 +702,7 @@ void FStaticMeshSceneProxy::SetIndexSource(int32 LODIndex, int32 SectionIndex, F
 		}
 		else
 		{
-			OutElement.IndexBuffer = bUseReversedIndices ? &LODModel.ReversedIndexBuffer : &LODModel.IndexBuffer;
+			OutElement.IndexBuffer = bUseReversedIndices ? &LODModel.AdditionalIndexBuffers->ReversedIndexBuffer : &LODModel.IndexBuffer;
 			OutElement.FirstIndex = Section.FirstIndex;
 			OutElement.NumPrimitives = Section.NumTriangles;
 		}
@@ -711,7 +711,7 @@ void FStaticMeshSceneProxy::SetIndexSource(int32 LODIndex, int32 SectionIndex, F
 	if ( bRequiresAdjacencyInformation )
 	{
 		check( LODModel.bHasAdjacencyInfo );
-		OutElement.IndexBuffer = &LODModel.AdjacencyIndexBuffer;
+		OutElement.IndexBuffer = &LODModel.AdditionalIndexBuffers->AdjacencyIndexBuffer;
 		OutMeshElement.Type = PT_12_ControlPointPatchList;
 		OutElement.FirstIndex *= 4;
 	}
diff --git a/Engine/Source/Runtime/Engine/Private/Streaming/StaticMeshUpdate.cpp b/Engine/Source/Runtime/Engine/Private/Streaming/StaticMeshUpdate.cpp
--- a/Engine/Source/Runtime/Engine/Private/Streaming/StaticMeshUpdate.cpp
+++ b/Engine/Source/Runtime/Engine/Private/Streaming/StaticMeshUpdate.cpp
@@ -54,11 +54,16 @@ void FStaticMeshStreamIn::FIntermediateBuffers::CreateFromCPUData_RenderThread(U
 	PositionVertexBuffer = VBs.PositionVertexBuffer.CreateRHIBuffer_RenderThread();
 	ColorVertexBuffer = VBs.ColorVertexBuffer.CreateRHIBuffer_RenderThread();
 	IndexBuffer = LODResource.IndexBuffer.CreateRHIBuffer_RenderThread();
-	ReversedIndexBuffer = LODResource.ReversedIndexBuffer.CreateRHIBuffer_RenderThread();
 	DepthOnlyIndexBuffer = LODResource.DepthOnlyIndexBuffer.CreateRHIBuffer_RenderThread();
-	ReversedDepthOnlyIndexBuffer = LODResource.ReversedDepthOnlyIndexBuffer.CreateRHIBuffer_RenderThread();
-	WireframeIndexBuffer = LODResource.WireframeIndexBuffer.CreateRHIBuffer_RenderThread();
-	AdjacencyIndexBuffer = LODResource.AdjacencyIndexBuffer.CreateRHIBuffer_RenderThread();
+
+
+	if (LODResource.AdditionalIndexBuffers)
+	{
+		ReversedIndexBuffer = LODResource.AdditionalIndexBuffers->ReversedIndexBuffer.CreateRHIBuffer_RenderThread();
+		ReversedDepthOnlyIndexBuffer = LODResource.AdditionalIndexBuffers->ReversedDepthOnlyIndexBuffer.CreateRHIBuffer_RenderThread();
+		WireframeIndexBuffer = LODResource.AdditionalIndexBuffers->WireframeIndexBuffer.CreateRHIBuffer_RenderThread();
+		AdjacencyIndexBuffer = LODResource.AdditionalIndexBuffers->AdjacencyIndexBuffer.CreateRHIBuffer_RenderThread();
+	}
 }
 
 void FStaticMeshStreamIn::FIntermediateBuffers::CreateFromCPUData_Async(UStaticMesh* Mesh, FStaticMeshLODResources& LODResource)
@@ -70,11 +75,15 @@ void FStaticMeshStreamIn::FIntermediateBuffers::CreateFromCPUData_Async(UStaticM
 	PositionVertexBuffer = VBs.PositionVertexBuffer.CreateRHIBuffer_Async();
 	ColorVertexBuffer = VBs.ColorVertexBuffer.CreateRHIBuffer_Async();
 	IndexBuffer = LODResource.IndexBuffer.CreateRHIBuffer_Async();
-	ReversedIndexBuffer = LODResource.ReversedIndexBuffer.CreateRHIBuffer_Async();
 	DepthOnlyIndexBuffer = LODResource.DepthOnlyIndexBuffer.CreateRHIBuffer_Async();
-	ReversedDepthOnlyIndexBuffer = LODResource.ReversedDepthOnlyIndexBuffer.CreateRHIBuffer_Async();
-	WireframeIndexBuffer = LODResource.WireframeIndexBuffer.CreateRHIBuffer_Async();
-	AdjacencyIndexBuffer = LODResource.AdjacencyIndexBuffer.CreateRHIBuffer_Async();
+
+	if (LODResource.AdditionalIndexBuffers)
+	{
+		ReversedIndexBuffer = LODResource.AdditionalIndexBuffers->ReversedIndexBuffer.CreateRHIBuffer_Async();
+		ReversedDepthOnlyIndexBuffer = LODResource.AdditionalIndexBuffers->ReversedDepthOnlyIndexBuffer.CreateRHIBuffer_Async();
+		WireframeIndexBuffer = LODResource.AdditionalIndexBuffers->WireframeIndexBuffer.CreateRHIBuffer_Async();
+		AdjacencyIndexBuffer = LODResource.AdditionalIndexBuffers->AdjacencyIndexBuffer.CreateRHIBuffer_Async();
+	}
 }
 
 void FStaticMeshStreamIn::FIntermediateBuffers::SafeRelease()
@@ -99,11 +108,15 @@ void FStaticMeshStreamIn::FIntermediateBuffers::TransferBuffers(FStaticMeshLODRe
 	VBs.PositionVertexBuffer.InitRHIForStreaming(PositionVertexBuffer, Batcher);
 	VBs.ColorVertexBuffer.InitRHIForStreaming(ColorVertexBuffer, Batcher);
 	LODResource.IndexBuffer.InitRHIForStreaming(IndexBuffer);
-	LODResource.ReversedIndexBuffer.InitRHIForStreaming(ReversedIndexBuffer);
 	LODResource.DepthOnlyIndexBuffer.InitRHIForStreaming(DepthOnlyIndexBuffer);
-	LODResource.ReversedDepthOnlyIndexBuffer.InitRHIForStreaming(ReversedDepthOnlyIndexBuffer);
-	LODResource.WireframeIndexBuffer.InitRHIForStreaming(WireframeIndexBuffer);
-	LODResource.AdjacencyIndexBuffer.InitRHIForStreaming(AdjacencyIndexBuffer);
+
+	if (LODResource.AdditionalIndexBuffers)
+	{
+		LODResource.AdditionalIndexBuffers->ReversedIndexBuffer.InitRHIForStreaming(ReversedIndexBuffer);
+		LODResource.AdditionalIndexBuffers->ReversedDepthOnlyIndexBuffer.InitRHIForStreaming(ReversedDepthOnlyIndexBuffer);
+		LODResource.AdditionalIndexBuffers->WireframeIndexBuffer.InitRHIForStreaming(WireframeIndexBuffer);
+		LODResource.AdditionalIndexBuffers->AdjacencyIndexBuffer.InitRHIForStreaming(AdjacencyIndexBuffer);
+	}
 	SafeRelease();
 }
 
@@ -264,11 +277,15 @@ void FStaticMeshStreamOut::DoReleaseBuffers(const FContext& Context)
 			VBs.ColorVertexBuffer.ReleaseRHIForStreaming(Batcher);
 			// Index buffers don't need to update SRV so we can reuse ReleaseRHI
 			LODResource.IndexBuffer.ReleaseRHI();
-			LODResource.ReversedIndexBuffer.ReleaseRHI();
 			LODResource.DepthOnlyIndexBuffer.ReleaseRHI();
-			LODResource.ReversedDepthOnlyIndexBuffer.ReleaseRHI();
-			LODResource.WireframeIndexBuffer.ReleaseRHI();
-			LODResource.AdjacencyIndexBuffer.ReleaseRHI();
+
+			if (LODResource.AdditionalIndexBuffers)
+			{
+				LODResource.AdditionalIndexBuffers->ReversedIndexBuffer.ReleaseRHI();
+				LODResource.AdditionalIndexBuffers->ReversedDepthOnlyIndexBuffer.ReleaseRHI();
+				LODResource.AdditionalIndexBuffers->WireframeIndexBuffer.ReleaseRHI();
+				LODResource.AdditionalIndexBuffers->AdjacencyIndexBuffer.ReleaseRHI();
+			}
 		}
 	}
 }
diff --git a/Engine/Source/Runtime/Engine/Public/RawIndexBuffer.h b/Engine/Source/Runtime/Engine/Public/RawIndexBuffer.h
--- a/Engine/Source/Runtime/Engine/Public/RawIndexBuffer.h
+++ b/Engine/Source/Runtime/Engine/Public/RawIndexBuffer.h
@@ -125,7 +125,7 @@ public:
 	 * Initialization constructor.
 	 * @param InNeedsCPUAccess	True if resource array data should be accessible by the CPU.
 	 */
-	FRawStaticIndexBuffer(bool InNeedsCPUAccess=false);
+	ENGINE_API FRawStaticIndexBuffer(bool InNeedsCPUAccess=false);
 
 	/**
 	 * Sets a single index value.  Consider using SetIndices() instead if you're setting a lot of indices.
diff --git a/Engine/Source/Runtime/Engine/Public/StaticMeshResources.h b/Engine/Source/Runtime/Engine/Public/StaticMeshResources.h
--- a/Engine/Source/Runtime/Engine/Public/StaticMeshResources.h
+++ b/Engine/Source/Runtime/Engine/Public/StaticMeshResources.h
@@ -267,23 +267,30 @@ struct FStaticMeshVertexBuffers
 	void ENGINE_API InitModelVF(FLocalVertexFactory* VertexFactory);
 };
 
-/** Rendering resources needed to render an individual static mesh LOD. */
-struct FStaticMeshLODResources
+struct FAdditionalStaticMeshIndexBuffers
 {
-	FStaticMeshVertexBuffers VertexBuffers;
-
-	/** Index buffer resource for rendering. */
-	FRawStaticIndexBuffer IndexBuffer;
 	/** Reversed index buffer, used to prevent changing culling state between drawcalls. */
 	FRawStaticIndexBuffer ReversedIndexBuffer;
-	/** Index buffer resource for rendering in depth only passes. */
-	FRawStaticIndexBuffer DepthOnlyIndexBuffer;
 	/** Reversed depth only index buffer, used to prevent changing culling state between drawcalls. */
 	FRawStaticIndexBuffer ReversedDepthOnlyIndexBuffer;
 	/** Index buffer resource for rendering wireframe mode. */
 	FRawStaticIndexBuffer WireframeIndexBuffer;
 	/** Index buffer containing adjacency information required by tessellation. */
 	FRawStaticIndexBuffer AdjacencyIndexBuffer;
+};
+
+/** Rendering resources needed to render an individual static mesh LOD. */
+struct FStaticMeshLODResources
+{
+	FStaticMeshVertexBuffers VertexBuffers;
+
+	/** Index buffer resource for rendering. */
+	FRawStaticIndexBuffer IndexBuffer;
+
+	/** Index buffer resource for rendering in depth only passes. */
+	FRawStaticIndexBuffer DepthOnlyIndexBuffer;
+
+	FAdditionalStaticMeshIndexBuffers* AdditionalIndexBuffers;
 
 #if RHI_RAYTRACING
 	/** Geometry for ray tracing. */
