Added ETextureMipLoadOptions to control texture mip residency.
Removed FirstMipNearest and FirstMipBilinear.

diff --git a/Engine/Source/Runtime/Engine/Classes/Engine/Texture.h b/Engine/Source/Runtime/Engine/Classes/Engine/Texture.h
--- a/Engine/Source/Runtime/Engine/Classes/Engine/Texture.h
+++ b/Engine/Source/Runtime/Engine/Classes/Engine/Texture.h
@@ -47,8 +47,6 @@ enum TextureFilter
 	TF_Nearest UMETA(DisplayName="Nearest"),
 	TF_Bilinear UMETA(DisplayName="Bi-linear"),
 	TF_Trilinear UMETA(DisplayName="Tri-linear"),
-	TF_FirstMipNearest  UMETA(DisplayName = "First mip nearest"),
-	TF_FirstMipBilinear  UMETA(DisplayName = "First mip bi-linear"),
 	/** Use setting from the Texture Group. */
 	TF_Default UMETA(DisplayName="Default (from Texture Group)"),
 	TF_MAX,
@@ -531,6 +529,10 @@ public:
 	UPROPERTY(EditAnywhere, BlueprintReadWrite, Category=Texture, AssetRegistrySearchable, AdvancedDisplay)
 	TEnumAsByte<enum TextureFilter> Filter;
 
+	/** The texture mip load options. */
+	UPROPERTY(EditAnywhere, BlueprintReadWrite, Category = Texture, AssetRegistrySearchable, AdvancedDisplay)
+	ETextureMipLoadOptions MipLoadOptions;
+
 	/** Texture group this texture belongs to */
 	UPROPERTY(EditAnywhere, BlueprintReadWrite, Category=LevelOfDetail, meta=(DisplayName="Texture Group"), AssetRegistrySearchable)
 	TEnumAsByte<enum TextureGroup> LODGroup;
diff --git a/Engine/Source/Runtime/Engine/Classes/Engine/TextureDefines.h b/Engine/Source/Runtime/Engine/Classes/Engine/TextureDefines.h
--- a/Engine/Source/Runtime/Engine/Classes/Engine/TextureDefines.h
+++ b/Engine/Source/Runtime/Engine/Classes/Engine/TextureDefines.h
@@ -149,3 +149,14 @@ enum class ETextureSamplerFilter : uint8
 	AnisotropicPoint,
 	AnisotropicLinear,
 };
+
+UENUM()
+enum class ETextureMipLoadOptions : uint8
+{
+	// Fallback to the LODGroup settings
+	Default,
+	// Load all mips.
+	AllMips,
+	// Load only the first mip.
+	OnlyFirstMip,
+};
diff --git a/Engine/Source/Runtime/Engine/Classes/Engine/TextureLODSettings.h b/Engine/Source/Runtime/Engine/Classes/Engine/TextureLODSettings.h
--- a/Engine/Source/Runtime/Engine/Classes/Engine/TextureLODSettings.h
+++ b/Engine/Source/Runtime/Engine/Classes/Engine/TextureLODSettings.h
@@ -34,6 +34,7 @@ struct ENGINE_API FTextureLODGroup
 		, OptionalMaxLODMipCount(12)
 		, MinMagFilter(NAME_Aniso)
 		, MipFilter(NAME_Point)
+		, MipLoadOptions(ETextureMipLoadOptions::AllMips)
 	{
 		SetupGroup();
 	}
@@ -97,6 +98,9 @@ struct ENGINE_API FTextureLODGroup
 	UPROPERTY()
 	FName MipFilter;
 
+	UPROPERTY()
+	ETextureMipLoadOptions MipLoadOptions;
+
 	void SetupGroup();
 };
 
@@ -158,6 +162,14 @@ public:
 
 	ETextureSamplerFilter GetSamplerFilter(int32 InLODGroup) const;
 
+	/**
+	 * Returns the mip load options of a texture.
+	 *
+	 * @param	Texture		Texture to retrieve the mip load option, must not be 0
+	 * @return	The mip load option
+	 */
+	ETextureMipLoadOptions GetMipLoadOptions(const UTexture* Texture) const;
+
 	/**
 	 * Returns the LODGroup mip gen settings
 	 *
diff --git a/Engine/Source/Runtime/Engine/Private/Texture.cpp b/Engine/Source/Runtime/Engine/Private/Texture.cpp
--- a/Engine/Source/Runtime/Engine/Private/Texture.cpp
+++ b/Engine/Source/Runtime/Engine/Private/Texture.cpp
@@ -59,6 +59,7 @@ UTexture::UTexture(const FObjectInitializer& ObjectInitializer)
 {
 	SRGB = true;
 	Filter = TF_Default;
+	MipLoadOptions = ETextureMipLoadOptions::Default;
 #if WITH_EDITORONLY_DATA
 	AdjustBrightness = 1.0f;
 	AdjustBrightnessCurve = 1.0f;
diff --git a/Engine/Source/Runtime/Engine/Private/Texture2D.cpp b/Engine/Source/Runtime/Engine/Private/Texture2D.cpp
--- a/Engine/Source/Runtime/Engine/Private/Texture2D.cpp
+++ b/Engine/Source/Runtime/Engine/Private/Texture2D.cpp
@@ -506,7 +506,8 @@ void UTexture2D::UpdateResource()
 #else
 	// Note that using TF_FirstMip disables texture streaming, because the mip data becomes lost.
 	// Also, the cleanup of the platform data must go between UpdateCachedLODBias() and UpdateResource().
-	if ((Filter == TF_FirstMipNearest || Filter == TF_FirstMipBilinear) && PlatformData && FPlatformProperties::RequiresCookedData())
+	const bool bLoadOnlyFirstMip = UDeviceProfileManager::Get().GetActiveProfile()->GetTextureLODSettings()->GetMipLoadOptions(this) == ETextureMipLoadOptions::OnlyFirstMip;
+	if (bLoadOnlyFirstMip && PlatformData && FPlatformProperties::RequiresCookedData())
 	{
 		const int32 FirstMip = FMath::Clamp(0, GetCachedLODBias(), PlatformData->Mips.Num() - 1);
 		// Remove any mips after the first mip.
@@ -841,7 +842,6 @@ FTextureResource* UTexture2D::CreateResource()
 					!NeverStream && 
 					(NumMips > 1) && 
 					(LODGroup != TEXTUREGROUP_UI) && 
-					(Filter != TF_FirstMipNearest && Filter != TF_FirstMipBilinear) &&
 					!bTemporarilyDisableStreaming;
 	
 	if (bIsStreamable && NumMips > 0)
@@ -1484,8 +1484,7 @@ uint32 FTexture2DResource::GetSizeY() const
 int32 FTexture2DResource::GetDefaultMipMapBias() const
 {
 	check(Owner);
-	if ((Owner->LODGroup == TEXTUREGROUP_UI && CVarForceHighestMipOnUITexturesEnabled.GetValueOnAnyThread() > 0) || 
-		(!GIsEditor && (Owner->Filter == TF_FirstMipNearest || Owner->Filter == TF_FirstMipBilinear))) // Don't use in editor otherwise it prevents visualizing mips.
+	if (Owner->LODGroup == TEXTUREGROUP_UI && CVarForceHighestMipOnUITexturesEnabled.GetValueOnAnyThread() > 0)
 	{
 		const TIndirectArray<FTexture2DMipMap>& OwnerMips = Owner->GetPlatformMips();
 		return -OwnerMips.Num();
diff --git a/Engine/Source/Runtime/Engine/Private/TextureLODSettings.cpp b/Engine/Source/Runtime/Engine/Private/TextureLODSettings.cpp
--- a/Engine/Source/Runtime/Engine/Private/TextureLODSettings.cpp
+++ b/Engine/Source/Runtime/Engine/Private/TextureLODSettings.cpp
@@ -293,11 +293,9 @@ ETextureSamplerFilter UTextureLODSettings::GetSamplerFilter(const UTexture* Text
 
 	switch(Texture->Filter)
 	{
-		case TF_FirstMipNearest:
 		case TF_Nearest: 
 			Filter = ETextureSamplerFilter::Point; 
 			break;
-		case TF_FirstMipBilinear:
 		case TF_Bilinear: 
 			Filter = ETextureSamplerFilter::Bilinear; 
 			break;
@@ -318,3 +316,17 @@ ETextureSamplerFilter UTextureLODSettings::GetSamplerFilter(int32 InLODGroup) co
 {
 	return TextureLODGroups[InLODGroup].Filter;
 }
+
+
+ETextureMipLoadOptions UTextureLODSettings::GetMipLoadOptions(const UTexture* Texture) const
+{
+	check(Texture);
+	if (Texture->MipLoadOptions != ETextureMipLoadOptions::Default)
+	{
+		return Texture->MipLoadOptions;
+	}
+	else
+	{
+		return TextureLODGroups[Texture->LODGroup].MipLoadOptions;
+	}
+}
diff --git a/Engine/Source/Runtime/SlateRHIRenderer/Private/SlateRHIRenderingPolicy.cpp b/Engine/Source/Runtime/SlateRHIRenderer/Private/SlateRHIRenderingPolicy.cpp
--- a/Engine/Source/Runtime/SlateRHIRenderer/Private/SlateRHIRenderingPolicy.cpp
+++ b/Engine/Source/Runtime/SlateRHIRenderer/Private/SlateRHIRenderingPolicy.cpp
@@ -1335,11 +1335,9 @@ ETextureSamplerFilter FSlateRHIRenderingPolicy::GetSamplerFilter(const UTexture*
 
 	switch (Texture->Filter)
 	{
-	case TF_FirstMipNearest:
 	case TF_Nearest: 
 		Filter = ETextureSamplerFilter::Point; 
 		break;
-	case TF_FirstMipBilinear:
 	case TF_Bilinear:
 		Filter = ETextureSamplerFilter::Bilinear; 
 		break;
