Unshelved from pending changelist '5065149':

- Fixed the NDA platforms DeviceProfiles from leaking out into DefaultDeviceProfiles.ini
- Replaced New Device Profile editor with a note on how to do it in the .ini files, which is the safest way to make a new DP without Epic being able to break NDA
- #jira UE-70268
- #rb [REDACTED_PERSON]

diff --git a/Engine/Source/Editor/DeviceProfileEditor/Private/SDeviceProfileEditor.cpp b/Engine/Source/Editor/DeviceProfileEditor/Private/SDeviceProfileEditor.cpp
--- a/Engine/Source/Editor/DeviceProfileEditor/Private/SDeviceProfileEditor.cpp
+++ b/Engine/Source/Editor/DeviceProfileEditor/Private/SDeviceProfileEditor.cpp
@@ -363,7 +363,10 @@ TSharedPtr< SWidget > SDeviceProfileEditor::CreateMainDeviceProfilePanel()
 			SNew( SBorder )
 			.BorderImage( FEditorStyle::GetBrush( "Docking.Tab.ContentAreaBrush" ) )
 			[
-				SNew( SDeviceProfileCreateProfilePanel, DeviceProfileManager )
+				SNew( STextBlock )
+				.AutoWrapText(true)
+				.Text(LOCTEXT("CreateDeviceProfile", "To create a device profile, edit your project's DefaultDeviceProfiles.ini file. "
+					"Make sure add to or create a section called [DeviceProfileNameAndTypes]. See Engine/Config/BaseDeviceProfiles.ini for examples."))
 			]
 		];
 
diff --git a/Engine/Source/Runtime/Engine/Private/DeviceProfiles/DeviceProfileManager.cpp b/Engine/Source/Runtime/Engine/Private/DeviceProfiles/DeviceProfileManager.cpp
--- a/Engine/Source/Runtime/Engine/Private/DeviceProfiles/DeviceProfileManager.cpp
+++ b/Engine/Source/Runtime/Engine/Private/DeviceProfiles/DeviceProfileManager.cpp
@@ -512,19 +512,13 @@ void UDeviceProfileManager::SaveProfiles(bool bSaveToDefaults)
 		}
 		else
 		{
-			TArray< FString > DeviceProfileMapArray;
-
 			for (int32 DeviceProfileIndex = 0; DeviceProfileIndex < Profiles.Num(); ++DeviceProfileIndex)
 			{
 				UDeviceProfile* CurrentProfile = CastChecked<UDeviceProfile>(Profiles[DeviceProfileIndex]);
 				FString DeviceProfileTypeNameCombo = FString::Printf(TEXT("%s,%s"), *CurrentProfile->GetName(), *CurrentProfile->DeviceType);
-				DeviceProfileMapArray.Add(DeviceProfileTypeNameCombo);
 
 				CurrentProfile->SaveConfig(CPF_Config, *DeviceProfileFileName);
 			}
-
-			GConfig->SetArray(TEXT("DeviceProfiles"), TEXT("DeviceProfileNameAndTypes"), DeviceProfileMapArray, DeviceProfileFileName);
-			GConfig->Flush(false, DeviceProfileFileName);
 		}
 
 		ManagerUpdatedDelegate.Broadcast();
