Fixing culling sounds before they async load in.

diff --git a/Engine/Source/Runtime/Engine/Private/AudioDevice.cpp b/Engine/Source/Runtime/Engine/Private/AudioDevice.cpp
--- a/Engine/Source/Runtime/Engine/Private/AudioDevice.cpp
+++ b/Engine/Source/Runtime/Engine/Private/AudioDevice.cpp
@@ -3287,7 +3287,7 @@ int32 FAudioDevice::GetSortedActiveWaveInstances(TArray<FWaveInstance*>& WaveIns
 								bStopped = true;
 							}
 						}
-						else if (!ActiveSound->bIsPlayingAudio)
+						else if (!ActiveSound->bIsPlayingAudio && ActiveSound->bFinished)
 						{
 							bStopped = true;
 						}
