Fixed issue where MaterialParameterCollections would sometimes revert to default values in PIE.

When invoking PIE, the current UWorld is duplicated, along with all actor instances. Material Parameter Collections have a unique instance for each world. The current material parameter data associated with the source world was not properly serialized during duplication, and the construction script for the actor was not run, therefore the parameter values would revert to their defaults.

This change makes the instance data serialize out during a duplication, but only for PIE mode.

diff --git a/Engine/Source/Editor/UnrealEd/Private/Factories/EditorFactories.cpp b/Engine/Source/Editor/UnrealEd/Private/Factories/EditorFactories.cpp
--- a/Engine/Source/Editor/UnrealEd/Private/Factories/EditorFactories.cpp
+++ b/Engine/Source/Editor/UnrealEd/Private/Factories/EditorFactories.cpp
@@ -650,7 +650,7 @@ UObject* UMaterialParameterCollectionFactoryNew::FactoryCreateNew(UClass* Class,
 		for (TObjectIterator<UWorld> It; It; ++It)
 		{
 			UWorld* CurrentWorld = *It;
-			CurrentWorld->AddParameterCollectionInstance(MaterialParameterCollection, true);
+			CurrentWorld->AddParameterCollectionInstance(MaterialParameterCollection, true /* bUpdateScene */);
 		}
 	}
 
diff --git a/Engine/Source/Runtime/Engine/Classes/Engine/World.h b/Engine/Source/Runtime/Engine/Classes/Engine/World.h
--- a/Engine/Source/Runtime/Engine/Classes/Engine/World.h
+++ b/Engine/Source/Runtime/Engine/Classes/Engine/World.h
@@ -1138,7 +1138,7 @@ private:
 	class UGameInstance*						OwningGameInstance;
 
 	/** Parameter collection instances that hold parameter overrides for this world. */
-	UPROPERTY(Transient)
+	UPROPERTY(NonPIEDuplicateTransient)
 	TArray<class UMaterialParameterCollectionInstance*> ParameterCollectionInstances;
 
 	/** 
@@ -2416,7 +2416,7 @@ public:
 	void SetupParameterCollectionInstances();
 
 	/** Adds a new instance of the given collection, or overwrites an existing instance if there is one. */
-	void AddParameterCollectionInstance(class UMaterialParameterCollection* Collection, bool bUpdateScene);
+	void AddParameterCollectionInstance(class UMaterialParameterCollection* Collection, bool bUpdateScene, bool bOverwriteIfExists = true);
 
 	/** Gets this world's instance for a given collection. */
 	UMaterialParameterCollectionInstance* GetParameterCollectionInstance(const UMaterialParameterCollection* Collection);
diff --git a/Engine/Source/Runtime/Engine/Classes/Materials/MaterialParameterCollectionInstance.h b/Engine/Source/Runtime/Engine/Classes/Materials/MaterialParameterCollectionInstance.h
--- a/Engine/Source/Runtime/Engine/Classes/Materials/MaterialParameterCollectionInstance.h
+++ b/Engine/Source/Runtime/Engine/Classes/Materials/MaterialParameterCollectionInstance.h
@@ -31,6 +31,9 @@ class ENGINE_API UMaterialParameterCollectionInstance : public UObject
 	/** Initializes the instance with the collection it is based off of and the world it is owned by. */
 	void SetCollection(UMaterialParameterCollection* InCollection, UWorld* InWorld);
 
+	/** Assigns new parent world instance. */
+	void SetWorld(UWorld* InWorld);
+
 	/** Sets parameter value overrides, returns false if the parameter was not found. */
 	bool SetScalarParameterValue(FName ParameterName, float ParameterValue);
 	bool SetVectorParameterValue(FName ParameterName, const FLinearColor& ParameterValue);
@@ -66,9 +69,11 @@ protected:
 	TWeakObjectPtr<UWorld> World;
 
 	/** Overrides for scalar parameter values. */
+	UPROPERTY()
 	TMap<FName, float> ScalarParameterValues;
 
 	/** Overrides for vector parameter values. */
+	UPROPERTY()
 	TMap<FName, FLinearColor> VectorParameterValues;
 
 	/** Instance resource which stores the rendering thread representation of this instance. */
diff --git a/Engine/Source/Runtime/Engine/Private/Materials/ParameterCollection.cpp b/Engine/Source/Runtime/Engine/Private/Materials/ParameterCollection.cpp
--- a/Engine/Source/Runtime/Engine/Private/Materials/ParameterCollection.cpp
+++ b/Engine/Source/Runtime/Engine/Private/Materials/ParameterCollection.cpp
@@ -53,7 +53,7 @@ void UMaterialParameterCollection::PostLoad()
 	for (TObjectIterator<UWorld> It; It; ++It)
 	{
 		UWorld* CurrentWorld = *It;
-		CurrentWorld->AddParameterCollectionInstance(this, true);
+		CurrentWorld->AddParameterCollectionInstance(this, true /* bUpdateScene */);
 	}
 
 	UpdateDefaultResource();
@@ -186,7 +186,7 @@ void UMaterialParameterCollection::PostEditChangeProperty(FPropertyChangedEvent&
 		for (TObjectIterator<UWorld> It; It; ++It)
 		{
 			UWorld* CurrentWorld = *It;
-			CurrentWorld->AddParameterCollectionInstance(this, false);
+			CurrentWorld->AddParameterCollectionInstance(this, false /* bUpdateScene */);
 		}
 
 		// Create a material update context so we can safely update materials using this parameter collection.
@@ -466,6 +466,11 @@ void UMaterialParameterCollectionInstance::PostInitProperties()
 	}
 }
 
+void UMaterialParameterCollectionInstance::SetWorld(UWorld* InWorld)
+{
+	World = InWorld;
+}
+
 void UMaterialParameterCollectionInstance::SetCollection(UMaterialParameterCollection* InCollection, UWorld* InWorld)
 {
 	Collection = InCollection;
diff --git a/Engine/Source/Runtime/Engine/Private/World.cpp b/Engine/Source/Runtime/Engine/Private/World.cpp
--- a/Engine/Source/Runtime/Engine/Private/World.cpp
+++ b/Engine/Source/Runtime/Engine/Private/World.cpp
@@ -1033,16 +1033,22 @@ void UWorld::SetupParameterCollectionInstances()
 {
 	QUICK_SCOPE_CYCLE_COUNTER(Stat_World_SetupParameterCollectionInstances);
 
-	// Create an instance for each parameter collection in memory
+	// Newly serialized collection instances need to be parented to this world instance.
+	for (int32 InstanceIndex = 0; InstanceIndex < ParameterCollectionInstances.Num(); InstanceIndex++)
+	{
+		ParameterCollectionInstances[InstanceIndex]->SetWorld(this);
+	}
+
+	// Create an instance for each parameter collection in memory that didn't already exist.
 	for (UMaterialParameterCollection* CurrentCollection : TObjectRange<UMaterialParameterCollection>())
 	{
-		AddParameterCollectionInstance(CurrentCollection, false);
+		AddParameterCollectionInstance(CurrentCollection, false /* bUpdateScene */, false /* bOverwriteIfExists */);
 	}
 
 	UpdateParameterCollectionInstances(false);
 }
 
-void UWorld::AddParameterCollectionInstance(UMaterialParameterCollection* Collection, bool bUpdateScene)
+void UWorld::AddParameterCollectionInstance(UMaterialParameterCollection* Collection, bool bUpdateScene, bool bOverwriteIfExists)
 {
 	int32 ExistingIndex = INDEX_NONE;
 
@@ -1055,23 +1061,30 @@ void UWorld::AddParameterCollectionInstance(UMaterialParameterCollection* Collec
 		}
 	}
 
-	UMaterialParameterCollectionInstance* NewInstance = NewObject<UMaterialParameterCollectionInstance>();
-	NewInstance->SetCollection(Collection, this);
+	const bool bFoundExisting = ExistingIndex != INDEX_NONE;
+	const bool bMakeNewInstance = !bFoundExisting || (bFoundExisting && bOverwriteIfExists);
 
-	if (ExistingIndex != INDEX_NONE)
+	if (bMakeNewInstance)
 	{
-		// Overwrite an existing instance
-		ParameterCollectionInstances[ExistingIndex] = NewInstance;
-	}
-	else
-	{
-		// Add a new instance
-		ParameterCollectionInstances.Add(NewInstance);
-	}
+		UMaterialParameterCollectionInstance* NewInstance = NewObject<UMaterialParameterCollectionInstance>();
+		NewInstance->SetCollection(Collection, this);
 
-	// Ensure the new instance creates initial render thread resources
-	// This needs to happen right away, so they can be picked up by any cached shader bindings
-	NewInstance->DeferredUpdateRenderState(false);
+		if (bFoundExisting)
+		{
+			ParameterCollectionInstances[ExistingIndex] = NewInstance;
+		}
+		else
+		{
+			ParameterCollectionInstances.Add(NewInstance);
+		}
+
+		// Ensure the new instance creates initial render thread resources
+		// This needs to happen right away, so they can be picked up by any cached shader bindings
+		if (NewInstance)
+		{
+			NewInstance->DeferredUpdateRenderState(false);
+		}
+	}
 
 	if (bUpdateScene)
 	{
