Restoring backwards compat for older replays

diff --git a/Engine/Source/Runtime/Engine/Classes/Engine/PackageMapClient.h b/Engine/Source/Runtime/Engine/Classes/Engine/PackageMapClient.h
--- a/Engine/Source/Runtime/Engine/Classes/Engine/PackageMapClient.h
+++ b/Engine/Source/Runtime/Engine/Classes/Engine/PackageMapClient.h
@@ -464,6 +464,7 @@ protected:
 	TSet< uint64 >						NetFieldExports;
 
 private:
+	void ReceiveNetFieldExportsCompat(FInBunch& InBunch);
 
 	bool bIgnoreReceivedExportGUIDs;
 };
diff --git a/Engine/Source/Runtime/Engine/Private/PackageMapClient.cpp b/Engine/Source/Runtime/Engine/Private/PackageMapClient.cpp
--- a/Engine/Source/Runtime/Engine/Private/PackageMapClient.cpp
+++ b/Engine/Source/Runtime/Engine/Private/PackageMapClient.cpp
@@ -1194,9 +1194,7 @@ void UPackageMapClient::ReceiveNetGUIDBunch( FInBunch &InBunch )
 	if ( bHasRepLayoutExport )
 	{
 		// We need to keep this around to ensure we don't break backwards compatability.
-		PRAGMA_DISABLE_DEPRECATION_WARNINGS
-		ReceiveNetFieldExports( InBunch );
-		PRAGMA_ENABLE_DEPRECATION_WARNINGS
+		ReceiveNetFieldExportsCompat( InBunch );
 		return;
 	}
 
@@ -1444,6 +1442,94 @@ void UPackageMapClient::AppendNetFieldExportsInternal(FArchive& Archive, const T
 	}
 }
 
+void UPackageMapClient::ReceiveNetFieldExportsCompat(FInBunch &InBunch)
+{
+	check(Connection->InternalAck);
+
+	// Read number of net field exports
+	uint32 NumLayoutCmdExports = 0;
+	InBunch << NumLayoutCmdExports;
+
+	for (int32 i = 0; i < (int32)NumLayoutCmdExports; i++)
+	{
+		// Read the index that represents the name in the NetFieldExportGroupIndexToPath map
+		uint32 PathNameIndex;
+		InBunch.SerializeIntPacked(PathNameIndex);
+
+		if (InBunch.IsError())
+		{
+			break;
+		}
+
+		int32 MaxExports = 0;
+
+		FNetFieldExportGroup* NetFieldExportGroup = nullptr;
+
+		// See if the path name was exported (we'll expect it if we haven't seen this index before)
+		if (InBunch.ReadBit() == 1)
+		{
+			FString PathName;
+
+			InBunch << PathName;
+			InBunch << MaxExports;
+
+			if (InBunch.IsError())
+			{
+				break;
+			}
+
+			GEngine->NetworkRemapPath(Connection->Driver, PathName, true);
+
+			NetFieldExportGroup = GuidCache->NetFieldExportGroupMap.FindRef(PathName).Get();
+			if (!NetFieldExportGroup)
+			{
+				TSharedPtr<FNetFieldExportGroup> NewNetFieldExportGroup(new FNetFieldExportGroup());
+				NetFieldExportGroup = NewNetFieldExportGroup.Get();
+
+				NetFieldExportGroup->PathName = PathName;
+				NetFieldExportGroup->PathNameIndex = PathNameIndex;
+
+				NetFieldExportGroup->NetFieldExports.SetNum(MaxExports);
+
+				GuidCache->NetFieldExportGroupMap.Add(PathName, NewNetFieldExportGroup);
+			}
+
+			GuidCache->NetFieldExportGroupPathToIndex.Add(PathName, PathNameIndex);
+			GuidCache->NetFieldExportGroupIndexToGroup.Add(PathNameIndex, NetFieldExportGroup);
+		}
+		else
+		{
+			NetFieldExportGroup = GuidCache->NetFieldExportGroupIndexToGroup.FindChecked(PathNameIndex);
+		}
+
+
+		FNetFieldExport NetFieldExport;
+
+		// Read the cmd
+		InBunch << NetFieldExport;
+
+		if (InBunch.IsError())
+		{
+			break;
+		}
+
+		TArray<FNetFieldExport>& NetFieldExportsRef = NetFieldExportGroup->NetFieldExports;
+
+		if ((int32)NetFieldExport.Handle < NetFieldExportsRef.Num())
+		{
+			// Assign it to the correct slot (NetFieldExport.Handle is just the index into the array)
+			NetFieldExportGroup->NetFieldExports[NetFieldExport.Handle] = NetFieldExport;
+		}
+		else
+		{
+			UE_LOG(LogNetPackageMap, Error, TEXT("ReceiveNetFieldExports: Invalid NetFieldExport Handle '%i', Max '%i'."),
+				NetFieldExport.Handle, NetFieldExportsRef.Num());
+
+			InBunch.SetError();
+		}
+	}
+}
+
 void UPackageMapClient::ReceiveNetFieldExports(FArchive& Archive)
 {
 	DECLARE_SCOPE_CYCLE_COUNTER(TEXT("ReceiveNetFieldExports time"), STAT_ReceiveNetFieldExportsTime, STATGROUP_Net);
