Fix respawning startup actors in replay when not using streaming level fixes

diff --git a/Engine/Source/Runtime/Engine/Private/DemoNetDriver.cpp b/Engine/Source/Runtime/Engine/Private/DemoNetDriver.cpp
--- a/Engine/Source/Runtime/Engine/Private/DemoNetDriver.cpp
+++ b/Engine/Source/Runtime/Engine/Private/DemoNetDriver.cpp
@@ -4086,18 +4086,21 @@ void UDemoNetDriver::RespawnNecessaryNetStartupActors(TArray<AActor*>& SpawnedAc
 			continue;
 		}
 
-		const FString LevelPackageName = GetLevelPackageName(*RollbackActor.Level);
-
-		// skip rollback actors in streamed out levels (pending gc)
-		if (!LevelStatusesByName.Contains(LevelPackageName))
+		if (HasLevelStreamingFixes())
 		{
-			continue;
-		}
+			const FString LevelPackageName = GetLevelPackageName(*RollbackActor.Level);
 
-		FLevelStatus& LevelStatus = GetLevelStatus(LevelPackageName);
-		if (!LevelStatus.bIsReady)
-		{
-			continue;
+			// skip rollback actors in streamed out levels (pending gc)
+			if (!LevelStatusesByName.Contains(LevelPackageName))
+			{
+				continue;
+			}
+
+			FLevelStatus& LevelStatus = GetLevelStatus(LevelPackageName);
+			if (!LevelStatus.bIsReady)
+			{
+				continue;
+			}
 		}
 
 		AActor* ExistingActor = FindObjectFast<AActor>(RollbackActor.Level, RollbackActor.Name);
