The issue was caused by UserDefinedStruct overriding the semi-deprecated serialize functions taking an FArchive. These functions need to be made non-virtual as the same problem exists with SerializeDefaultObject

diff --git a/Engine/Source/Runtime/Engine/Classes/Engine/UserDefinedStruct.h b/Engine/Source/Runtime/Engine/Classes/Engine/UserDefinedStruct.h
--- a/Engine/Source/Runtime/Engine/Classes/Engine/UserDefinedStruct.h
+++ b/Engine/Source/Runtime/Engine/Classes/Engine/UserDefinedStruct.h
@@ -100,7 +100,7 @@ public:
 
 	// UObject interface.
 	virtual void Serialize(FStructuredArchive::FRecord Record) override;
-	virtual void SerializeTaggedProperties(FArchive& Ar, uint8* Data, UStruct* DefaultsStruct, uint8* Defaults, const UObject* BreakRecursionIfFullyLoad=NULL) const override;
+	virtual void SerializeTaggedProperties(FStructuredArchive::FSlot Slot, uint8* Data, UStruct* DefaultsStruct, uint8* Defaults, const UObject* BreakRecursionIfFullyLoad = nullptr) const override;
 	virtual FString PropertyNameToDisplayName(FName InName) const override;
 	// End of UObject interface.
 
diff --git a/Engine/Source/Runtime/Engine/Private/UserDefinedStruct.cpp b/Engine/Source/Runtime/Engine/Private/UserDefinedStruct.cpp
--- a/Engine/Source/Runtime/Engine/Private/UserDefinedStruct.cpp
+++ b/Engine/Source/Runtime/Engine/Private/UserDefinedStruct.cpp
@@ -279,9 +279,10 @@ void UUserDefinedStruct::InitializeStruct(void* Dest, int32 ArrayDim) const
 	}
 }
 
-void UUserDefinedStruct::SerializeTaggedProperties(FArchive& Ar, uint8* Data, UStruct* DefaultsStruct, uint8* Defaults, const UObject* BreakRecursionIfFullyLoad) const
+void UUserDefinedStruct::SerializeTaggedProperties(FStructuredArchive::FSlot Slot, uint8* Data, UStruct* DefaultsStruct, uint8* Defaults, const UObject* BreakRecursionIfFullyLoad) const 
 {
 	bool bTemporarilyEnableDelta = false;
+	FArchive& Ar = Slot.GetUnderlyingArchive();
 
 #if WITH_EDITOR
 	// In the editor the default structure may change while the editor is running, so we need to always delta serialize
@@ -311,7 +312,7 @@ void UUserDefinedStruct::SerializeTaggedProperties(FArchive& Ar, uint8* Data, US
 	}
 #endif // WITH_EDITOR
 
-	Super::SerializeTaggedProperties(Ar, Data, DefaultsStruct, Defaults);
+	Super::SerializeTaggedProperties(Slot, Data, DefaultsStruct, Defaults, BreakRecursionIfFullyLoad);
 
 	if (bTemporarilyEnableDelta)
 	{
