Improvements to LLM accounting for skelmeshes and instanced meshes

diff --git a/Engine/Source/Runtime/Core/Private/HAL/LowLevelMemTracker.cpp b/Engine/Source/Runtime/Core/Private/HAL/LowLevelMemTracker.cpp
--- a/Engine/Source/Runtime/Core/Private/HAL/LowLevelMemTracker.cpp
+++ b/Engine/Source/Runtime/Core/Private/HAL/LowLevelMemTracker.cpp
@@ -106,6 +106,7 @@ DECLARE_LLM_MEMORY_STAT(TEXT("VideoRecording"), STAT_VideoRecordingLLM, STATGROU
 DECLARE_LLM_MEMORY_STAT(TEXT("CsvProfiler"), STAT_CsvProfilerLLM, STATGROUP_LLMFULL);
 DECLARE_LLM_MEMORY_STAT(TEXT("MaterialInstance"), STAT_MaterialInstanceLLM, STATGROUP_LLMFULL);
 DECLARE_LLM_MEMORY_STAT(TEXT("SkeletalMesh"), STAT_SkeletalMeshLLM, STATGROUP_LLMFULL);
+DECLARE_LLM_MEMORY_STAT(TEXT("InstancedMesh"), STAT_InstancedMeshLLM, STATGROUP_LLMFULL);
 DECLARE_LLM_MEMORY_STAT(TEXT("Landscape"), STAT_LandscapeLLM, STATGROUP_LLMFULL);
 
 /*
diff --git a/Engine/Source/Runtime/Core/Public/HAL/LowLevelMemTracker.h b/Engine/Source/Runtime/Core/Public/HAL/LowLevelMemTracker.h
--- a/Engine/Source/Runtime/Core/Public/HAL/LowLevelMemTracker.h
+++ b/Engine/Source/Runtime/Core/Public/HAL/LowLevelMemTracker.h
@@ -231,6 +231,7 @@ enum class ELLMTagSet : uint8
 	macro(VideoRecording,						"VideoRecording",				GET_STATFNAME(STAT_VideoRecordingLLM),						GET_STATFNAME(STAT_EngineSummaryLLM),			-1)\
 	macro(MaterialInstance,						"MaterialInstance",				GET_STATFNAME(STAT_MaterialInstanceLLM),					GET_STATFNAME(STAT_EngineSummaryLLM),			-1)\
 	macro(SkeletalMesh,							"SkeletalMesh",					GET_STATFNAME(STAT_SkeletalMeshLLM),						GET_STATFNAME(STAT_EngineSummaryLLM),			-1)\
+	macro(InstancedMesh,						"InstancedMesh",				GET_STATFNAME(STAT_InstancedMeshLLM),						GET_STATFNAME(STAT_EngineSummaryLLM),			-1)\
 	macro(Landscape,							"Landscape",					GET_STATFNAME(STAT_LandscapeLLM),							GET_STATFNAME(STAT_EngineSummaryLLM),			-1)\
 	macro(CsvProfiler,							"CsvProfiler",					GET_STATFNAME(STAT_CsvProfilerLLM),							GET_STATFNAME(STAT_EngineSummaryLLM),			-1)
 
diff --git a/Engine/Source/Runtime/Engine/Private/Components/SkinnedMeshComponent.cpp b/Engine/Source/Runtime/Engine/Private/Components/SkinnedMeshComponent.cpp
--- a/Engine/Source/Runtime/Engine/Private/Components/SkinnedMeshComponent.cpp
+++ b/Engine/Source/Runtime/Engine/Private/Components/SkinnedMeshComponent.cpp
@@ -413,6 +413,7 @@ void USkinnedMeshComponent::GetResourceSizeEx(FResourceSizeEx& CumulativeResourc
 
 FPrimitiveSceneProxy* USkinnedMeshComponent::CreateSceneProxy()
 {
+	LLM_SCOPE(ELLMTag::SkeletalMesh);
 	ERHIFeatureLevel::Type SceneFeatureLevel = GetWorld()->FeatureLevel;
 	FSkeletalMeshSceneProxy* Result = nullptr;
 	FSkeletalMeshRenderData* SkelMeshRenderData = GetSkeletalMeshRenderData();
diff --git a/Engine/Source/Runtime/Engine/Private/HierarchicalInstancedStaticMesh.cpp b/Engine/Source/Runtime/Engine/Private/HierarchicalInstancedStaticMesh.cpp
--- a/Engine/Source/Runtime/Engine/Private/HierarchicalInstancedStaticMesh.cpp
+++ b/Engine/Source/Runtime/Engine/Private/HierarchicalInstancedStaticMesh.cpp
@@ -1969,7 +1969,7 @@ void UHierarchicalInstancedStaticMeshComponent::PreSave(const class ITargetPlatf
 
 void UHierarchicalInstancedStaticMeshComponent::Serialize(FArchive& Ar)
 {
-	LLM_SCOPE(ELLMTag::StaticMesh);	
+	LLM_SCOPE(ELLMTag::InstancedMesh);
 
 	Ar.UsingCustomVersion(FReleaseObjectVersion::GUID);
 
@@ -2864,6 +2864,7 @@ void UHierarchicalInstancedStaticMeshComponent::SetPerInstanceLightMapAndEditorD
 
 FPrimitiveSceneProxy* UHierarchicalInstancedStaticMeshComponent::CreateSceneProxy()
 {
+	LLM_SCOPE(ELLMTag::InstancedMesh);
 	QUICK_SCOPE_CYCLE_COUNTER(STAT_HierarchicalInstancedStaticMeshComponent_CreateSceneProxy);
 	SCOPE_CYCLE_COUNTER(STAT_FoliageCreateProxy);
 
diff --git a/Engine/Source/Runtime/Engine/Private/InstancedStaticMesh.cpp b/Engine/Source/Runtime/Engine/Private/InstancedStaticMesh.cpp
--- a/Engine/Source/Runtime/Engine/Private/InstancedStaticMesh.cpp
+++ b/Engine/Source/Runtime/Engine/Private/InstancedStaticMesh.cpp
@@ -1026,6 +1026,7 @@ void UInstancedStaticMeshComponent::ApplyComponentInstanceData(FInstancedStaticM
 
 FPrimitiveSceneProxy* UInstancedStaticMeshComponent::CreateSceneProxy()
 {
+	LLM_SCOPE(ELLMTag::InstancedMesh);
 	ProxySize = 0;
 
 	// Verify that the mesh is valid before using it.
@@ -1085,6 +1086,7 @@ void UInstancedStaticMeshComponent::CreateHitProxyData(TArray<TRefCountPtr<HHitP
 
 void UInstancedStaticMeshComponent::BuildRenderData(FStaticMeshInstanceData& OutData, TArray<TRefCountPtr<HHitProxy>>& OutHitProxies)
 {
+	LLM_SCOPE(ELLMTag::InstancedMesh);
 	QUICK_SCOPE_CYCLE_COUNTER(STAT_UInstancedStaticMeshComponent_BuildRenderData);
 
 	CreateHitProxyData(OutHitProxies);
@@ -1674,6 +1676,7 @@ void UInstancedStaticMeshComponent::SerializeRenderData(FArchive& Ar)
 
 void UInstancedStaticMeshComponent::Serialize(FArchive& Ar)
 {
+	LLM_SCOPE(ELLMTag::InstancedMesh);
 	Super::Serialize(Ar);
 
 	Ar.UsingCustomVersion(FMobileObjectVersion::GUID);
@@ -2104,6 +2107,8 @@ void UInstancedStaticMeshComponent::InitPerInstanceRenderData(bool InitializeFro
 		return;
 	}
 
+	LLM_SCOPE(ELLMTag::InstancedMesh);
+
 	// If we don't have a random seed for this instanced static mesh component yet, then go ahead and
 	// generate one now.  This will be saved with the static mesh component and used for future generation
 	// of random numbers for this component's instances. (Used by the PerInstanceRandom material expression)
diff --git a/Engine/Source/Runtime/Engine/Private/SkeletalMesh.cpp b/Engine/Source/Runtime/Engine/Private/SkeletalMesh.cpp
--- a/Engine/Source/Runtime/Engine/Private/SkeletalMesh.cpp
+++ b/Engine/Source/Runtime/Engine/Private/SkeletalMesh.cpp
@@ -1538,6 +1538,7 @@ void USkeletalMesh::RemoveLegacyClothingSections()
 
 void USkeletalMesh::PostLoad()
 {
+	LLM_SCOPE(ELLMTag::SkeletalMesh);
 	Super::PostLoad();
 
 #if WITH_EDITOR
diff --git a/Engine/Source/Runtime/Engine/Private/StaticMesh.cpp b/Engine/Source/Runtime/Engine/Private/StaticMesh.cpp
--- a/Engine/Source/Runtime/Engine/Private/StaticMesh.cpp
+++ b/Engine/Source/Runtime/Engine/Private/StaticMesh.cpp
@@ -4315,6 +4315,7 @@ bool UStaticMesh::IsPostLoadThreadSafe() const
 //
 void UStaticMesh::PostLoad()
 {
+	LLM_SCOPE(ELLMTag::StaticMesh);
 	Super::PostLoad();
 
 #if WITH_EDITOR
diff --git a/Engine/Source/Runtime/Engine/Private/StaticMeshRender.cpp b/Engine/Source/Runtime/Engine/Private/StaticMeshRender.cpp
--- a/Engine/Source/Runtime/Engine/Private/StaticMeshRender.cpp
+++ b/Engine/Source/Runtime/Engine/Private/StaticMeshRender.cpp
@@ -1842,6 +1842,7 @@ FPrimitiveSceneProxy* UStaticMeshComponent::CreateSceneProxy()
 	{
 		return nullptr;
 	}
+	LLM_SCOPE(ELLMTag::StaticMesh);
 
 	FPrimitiveSceneProxy* Proxy = ::new FStaticMeshSceneProxy(this, false);
 #if STATICMESH_ENABLE_DEBUG_RENDERING
