Make sure the reorder value is valid, it won't be valid if we're a ISMC vs a HISMC.

diff --git a/Engine/Source/Runtime/Engine/Private/LightMap.cpp b/Engine/Source/Runtime/Engine/Private/LightMap.cpp
--- a/Engine/Source/Runtime/Engine/Private/LightMap.cpp
+++ b/Engine/Source/Runtime/Engine/Private/LightMap.cpp
@@ -283,7 +283,7 @@ struct FLightMapAllocation
 				// Need to create per-LOD instance data to fix that
 				MeshBuildData->PerInstanceLightmapData[InstanceIndex].LightmapUVBias = LightMap->GetCoordinateBias();
 
-				int32 RenderIndex = Component->InstanceReorderTable.IsValidIndex(InstanceIndex) ? Component->InstanceReorderTable[InstanceIndex] : InstanceIndex;
+				int32 RenderIndex = (Component->InstanceReorderTable.IsValidIndex(InstanceIndex) && Component->InstanceReorderTable[InstanceIndex] != INDEX_NONE) ? Component->InstanceReorderTable[InstanceIndex] : InstanceIndex;
 				Component->InstanceUpdateCmdBuffer.SetLightMapData(RenderIndex, MeshBuildData->PerInstanceLightmapData[InstanceIndex].LightmapUVBias);
 				Component->MarkRenderStateDirty();
 			}
diff --git a/Engine/Source/Runtime/Engine/Private/ShadowMap.cpp b/Engine/Source/Runtime/Engine/Private/ShadowMap.cpp
--- a/Engine/Source/Runtime/Engine/Private/ShadowMap.cpp
+++ b/Engine/Source/Runtime/Engine/Private/ShadowMap.cpp
@@ -109,7 +109,8 @@ struct FShadowMapAllocation
 				// TODO: We currently only support one LOD of static lighting in foliage
 				// Need to create per-LOD instance data to fix that
 				MeshBuildData->PerInstanceLightmapData[InstanceIndex].ShadowmapUVBias = ShadowMap->GetCoordinateBias();
-				int32 RenderIndex = Component->InstanceReorderTable.IsValidIndex(InstanceIndex) ? Component->InstanceReorderTable[InstanceIndex] : InstanceIndex;
+				int32 RenderIndex = (Component->InstanceReorderTable.IsValidIndex(InstanceIndex) && Component->InstanceReorderTable[InstanceIndex] != INDEX_NONE) ? Component->InstanceReorderTable[InstanceIndex] : InstanceIndex;
+
 				Component->InstanceUpdateCmdBuffer.SetShadowMapData(RenderIndex, MeshBuildData->PerInstanceLightmapData[InstanceIndex].ShadowmapUVBias);
 				Component->MarkRenderStateDirty();
 			}
