Fixed memory estimation of StaticMeshComponent when Archive is counting memory and filters out editor only data

diff --git a/Engine/Source/Runtime/Engine/Classes/Components/StaticMeshComponent.h b/Engine/Source/Runtime/Engine/Classes/Components/StaticMeshComponent.h
--- a/Engine/Source/Runtime/Engine/Classes/Components/StaticMeshComponent.h
+++ b/Engine/Source/Runtime/Engine/Classes/Components/StaticMeshComponent.h
@@ -101,7 +101,6 @@ struct FStaticMeshComponentLODInfo
 	TUniquePtr<FMeshMapBuildData> OverrideMapBuildData;
 
 	/** Vertex data cached at the time this LOD was painted, if any */
-	UPROPERTY()
 	TArray<struct FPaintedVertex> PaintedVertices;
 
 	/** Vertex colors to use for this mesh LOD */
diff --git a/Engine/Source/Runtime/Engine/Private/Components/StaticMeshComponent.cpp b/Engine/Source/Runtime/Engine/Private/Components/StaticMeshComponent.cpp
--- a/Engine/Source/Runtime/Engine/Private/Components/StaticMeshComponent.cpp
+++ b/Engine/Source/Runtime/Engine/Private/Components/StaticMeshComponent.cpp
@@ -2749,7 +2749,7 @@ FArchive& operator<<(FArchive& Ar,FStaticMeshComponentLODInfo& I)
 	}
 
 	// Serialize out cached vertex information if necessary.
-	if ( !StripFlags.IsEditorDataStripped() )
+	if (!StripFlags.IsEditorDataStripped() && !(Ar.IsFilterEditorOnly() && Ar.IsCountingMemory()))
 	{
 		Ar << I.PaintedVertices;
 	}
