Only show RHI thread timing on 'Stat Unit' and 'Stat UnitGraph" debug display if actually enabled.

diff --git a/Engine/Source/Runtime/Engine/Private/UnrealClient.cpp b/Engine/Source/Runtime/Engine/Private/UnrealClient.cpp
--- a/Engine/Source/Runtime/Engine/Private/UnrealClient.cpp
+++ b/Engine/Source/Runtime/Engine/Private/UnrealClient.cpp
@@ -578,6 +578,7 @@ int32 FStatUnitData::DrawStat(FViewport* InViewport, FCanvas* InCanvas, int32 In
 			}
 			InY += RowHeight;
 		}
+		if (IsRunningRHIInSeparateThread())
 		{
 			const FColor RenderThreadAverageColor = GEngine->GetFrameTimeDisplayColor(RHITTime);
 			InCanvas->DrawShadowedString(X1, InY, TEXT("RHIT:"), Font, bShowUnitTimeGraph ? FColor(255, 100, 255) : FColor::White);
@@ -878,7 +879,8 @@ int32 FStatUnitData::DrawStat(FViewport* InViewport, FCanvas* InCanvas, int32 In
 
 			// If we don't have GPU data to display, then skip this line
 			if ((StatIndex == EGS_GPU && !bHaveGPUData)
-				|| (StatIndex == EGS_Frame && bShowFrameTimeInUnitGraph == false && bHaveGPUData))
+				|| (StatIndex == EGS_Frame && bShowFrameTimeInUnitGraph == false && bHaveGPUData)
+				|| (StatIndex == EGS_RHIT && !IsRunningRHIInSeparateThread()))
 			{
 				continue;
 			}
