exposed some of camera accessors for projects
#rnx

diff --git a/Engine/Source/Runtime/Engine/Classes/Camera/PlayerCameraManager.h b/Engine/Source/Runtime/Engine/Classes/Camera/PlayerCameraManager.h
--- a/Engine/Source/Runtime/Engine/Classes/Camera/PlayerCameraManager.h
+++ b/Engine/Source/Runtime/Engine/Classes/Camera/PlayerCameraManager.h
@@ -655,15 +655,15 @@ public:
 	 * @param	OutCamLoc	Returned camera location
 	 * @param	OutCamRot	Returned camera rotation
 	 */
-	void GetCameraViewPoint(FVector& OutCamLoc, FRotator& OutCamRot) const;
+	virtual void GetCameraViewPoint(FVector& OutCamLoc, FRotator& OutCamRot) const;
 	
 	/** Returns camera's current rotation. */
 	UFUNCTION(BlueprintCallable, Category = "Camera", meta=(Keywords="View Direction"))
-	FRotator GetCameraRotation() const;
+	virtual FRotator GetCameraRotation() const;
 
 	/** Returns camera's current location. */
 	UFUNCTION(BlueprintCallable, Category = "Camera", meta=(Keywords="View Position"))
-	FVector GetCameraLocation() const;
+	virtual FVector GetCameraLocation() const;
 	
 	/** 
 	 * Sets the new desired color scale, enables color scaling, and enables color scale interpolation. 
