Fix crash when serializing data from the DDC

diff --git a/Engine/Source/Runtime/Engine/Private/Animation/AnimCompressionDerivedData.h b/Engine/Source/Runtime/Engine/Private/Animation/AnimCompressionDerivedData.h
--- a/Engine/Source/Runtime/Engine/Private/Animation/AnimCompressionDerivedData.h
+++ b/Engine/Source/Runtime/Engine/Private/Animation/AnimCompressionDerivedData.h
@@ -50,7 +50,7 @@ public:
 		// This is a version string that mimics the old versioning scheme. If you
 		// want to bump this version, generate a new guid using VS->Tools->Create GUID and
 		// return it here. Ex.
-		return TEXT("A74732B14AD94950856C6A09F708445C");
+		return TEXT("691013A9702847FE9A47302771888030");
 	}
 
 	virtual FString GetPluginSpecificCacheKeySuffix() const override;
diff --git a/Engine/Source/Runtime/Engine/Private/Animation/AnimSequence.cpp b/Engine/Source/Runtime/Engine/Private/Animation/AnimSequence.cpp
--- a/Engine/Source/Runtime/Engine/Private/Animation/AnimSequence.cpp
+++ b/Engine/Source/Runtime/Engine/Private/Animation/AnimSequence.cpp
@@ -2712,6 +2712,31 @@ void UAnimSequence::SerializeCompressedData(FArchive& Ar, bool bDDCData)
 				((AnimEncoding*)RotationCodec)->ByteSwapIn(*this, MemoryReader);
 			}
 		}
+		
+#if WITH_EDITOR
+		if (bDDCData)
+		{
+			FString CurveCodecPath;
+			Ar << CurveCodecPath;
+
+			CurveCompressionCodec = CurveCompressionSettings->GetCodec(CurveCodecPath);
+		}
+		else
+#else
+		check(!bDDCData);
+#endif
+		{
+			UAnimCurveCompressionCodec* CurveCodec = nullptr;
+			Ar << CurveCodec;
+			CurveCompressionCodec = CurveCodec;
+		}
+
+		int32 NumCurveBytes;
+		Ar << NumCurveBytes;
+
+		CompressedCurveByteStream.Empty(NumCurveBytes);
+		CompressedCurveByteStream.AddUninitialized(NumCurveBytes);
+		Ar.Serialize(CompressedCurveByteStream.GetData(), NumCurveBytes);
 	}
 	else if (Ar.IsSaving() || Ar.IsCountingMemory())
 	{
