Fixed crash with out of bound access to sample data

diff --git a/Engine/Source/Runtime/Engine/Private/Animation/BlendSpaceBase.cpp b/Engine/Source/Runtime/Engine/Private/Animation/BlendSpaceBase.cpp
--- a/Engine/Source/Runtime/Engine/Private/Animation/BlendSpaceBase.cpp
+++ b/Engine/Source/Runtime/Engine/Private/Animation/BlendSpaceBase.cpp
@@ -183,6 +183,17 @@ void UBlendSpaceBase::TickAssetPlayer(FAnimTickRecord& Instance, struct FAnimNot
 
 		OldSampleDataList.Append(*Instance.BlendSpace.BlendSampleDataCache);
 
+		// @fixme: temporary code to clear the invalid sample data 
+		// related jira: UE-71107
+		for (int32 Index = 0; Index < OldSampleDataList.Num(); ++Index)
+		{^
+			if (!SampleData.IsValidIndex(OldSampleDataList[Index].SampleDataIndex))
+			{
+				OldSampleDataList.RemoveAt(Index);
+				--Index;
+			}
+		}
+		
 		// get sample data based on new input
 		// consolidate all samples and sort them, so that we can handle from biggest weight to smallest
 		Instance.BlendSpace.BlendSampleDataCache->Reset();
