Fix compile error.

diff --git a/Engine/Source/Runtime/Engine/Private/Animation/BlendSpaceBase.cpp b/Engine/Source/Runtime/Engine/Private/Animation/BlendSpaceBase.cpp
--- a/Engine/Source/Runtime/Engine/Private/Animation/BlendSpaceBase.cpp
+++ b/Engine/Source/Runtime/Engine/Private/Animation/BlendSpaceBase.cpp
@@ -186,7 +186,7 @@ void UBlendSpaceBase::TickAssetPlayer(FAnimTickRecord& Instance, struct FAnimNot
 		// @fixme: temporary code to clear the invalid sample data 
 		// related jira: UE-71107
 		for (int32 Index = 0; Index < OldSampleDataList.Num(); ++Index)
-		{^
+		{
 			if (!SampleData.IsValidIndex(OldSampleDataList[Index].SampleDataIndex))
 			{
 				OldSampleDataList.RemoveAt(Index);
