Proper fix for changing static permutation inside materials (replaces temporary fix from CL 4753585). Added RecacheUniformExpressions(true) to UpdateStaticPermutation in order to force uniform expression UB recreation if we change a static permutation. This is required as changing static permutation may also change uniform expression UB layout.
#rnx

diff --git a/Engine/Source/Runtime/Engine/Private/Materials/MaterialInstance.cpp b/Engine/Source/Runtime/Engine/Private/Materials/MaterialInstance.cpp
--- a/Engine/Source/Runtime/Engine/Private/Materials/MaterialInstance.cpp
+++ b/Engine/Source/Runtime/Engine/Private/Materials/MaterialInstance.cpp
@@ -45,11 +45,11 @@ DECLARE_CYCLE_STAT(TEXT("MaterialInstance CopyUniformParamsInternal"), STAT_Mate
  * Cache uniform expressions for the given material.
  * @param MaterialInstance - The material instance for which to cache uniform expressions.
  */
-void CacheMaterialInstanceUniformExpressions(const UMaterialInstance* MaterialInstance)
+void CacheMaterialInstanceUniformExpressions(const UMaterialInstance* MaterialInstance, bool bRecreateUniformBuffer)
 {
 	if (MaterialInstance->Resource)
 	{
-		MaterialInstance->Resource->CacheUniformExpressions_GameThread(false);
+		MaterialInstance->Resource->CacheUniformExpressions_GameThread(bRecreateUniformBuffer);
 	}
 }
 
@@ -3545,6 +3545,7 @@ void UMaterialInstance::UpdateStaticPermutation(const FStaticParameterSet& NewPa
 		StaticParameters = CompareParameters;
 
 		CacheResourceShadersForRendering();
+		RecacheUniformExpressions(true);
 
 		if (MaterialUpdateContext != nullptr)
 		{
@@ -3584,7 +3585,7 @@ void UMaterialInstance::UpdateParameterNames()
 
 void UMaterialInstance::RecacheUniformExpressions(bool bRecreateUniformBuffer) const
 {	
-	CacheMaterialInstanceUniformExpressions(this);
+	CacheMaterialInstanceUniformExpressions(this, bRecreateUniformBuffer);
 }
 
 #if WITH_EDITOR
diff --git a/Engine/Source/Runtime/Engine/Private/Materials/MaterialInstanceSupport.h b/Engine/Source/Runtime/Engine/Private/Materials/MaterialInstanceSupport.h
--- a/Engine/Source/Runtime/Engine/Private/Materials/MaterialInstanceSupport.h
+++ b/Engine/Source/Runtime/Engine/Private/Materials/MaterialInstanceSupport.h
@@ -17,7 +17,7 @@ class UTexture;
  * Cache uniform expressions for the given material instance.
  * @param MaterialInstance - The material instance for which to cache uniform expressions.
  */
-void CacheMaterialInstanceUniformExpressions(const UMaterialInstance* MaterialInstance);
+void CacheMaterialInstanceUniformExpressions(const UMaterialInstance* MaterialInstance, bool bRecreateUniformBuffer = false);
 
 /**
  * Recaches uniform expressions for all material instances with a given parent.
diff --git a/Engine/Source/Runtime/Engine/Private/Materials/MaterialShared.cpp b/Engine/Source/Runtime/Engine/Private/Materials/MaterialShared.cpp
--- a/Engine/Source/Runtime/Engine/Private/Materials/MaterialShared.cpp
+++ b/Engine/Source/Runtime/Engine/Private/Materials/MaterialShared.cpp
@@ -2049,9 +2049,9 @@ void FMaterialRenderProxy::EvaluateUniformExpressions(FUniformExpressionCache& O
 	}
 	else
 	{
-		if (IsValidRef(OutUniformExpressionCache.UniformBuffer)
-			&& OutUniformExpressionCache.UniformBuffer->GetLayout() == UniformBufferStruct.GetLayout())
+		if (IsValidRef(OutUniformExpressionCache.UniformBuffer))
 		{
+			check(OutUniformExpressionCache.UniformBuffer->GetLayout() == UniformBufferStruct.GetLayout());
 			RHIUpdateUniformBuffer(OutUniformExpressionCache.UniformBuffer, TempBuffer);
 		}
 		else
