RL - Fix clang warning

diff --git a/Engine/Source/Runtime/Engine/Public/SceneManagement.h b/Engine/Source/Runtime/Engine/Public/SceneManagement.h
--- a/Engine/Source/Runtime/Engine/Public/SceneManagement.h
+++ b/Engine/Source/Runtime/Engine/Public/SceneManagement.h
@@ -2154,7 +2154,7 @@ class FRayTracingMeshResourceCollector : public FMeshElementCollector
 public:
 	// No MeshBatch should be allocated from an FRayTracingMeshResourceCollector.
 	inline FMeshBatch& AllocateMesh() = delete;
-	ENGINE_API void AddMesh(int32 ViewIndex, FMeshBatch& MeshBatch) = delete;
+	void AddMesh(int32 ViewIndex, FMeshBatch& MeshBatch) = delete;
 	void RegisterOneFrameMaterialProxy(FMaterialRenderProxy* Proxy) = delete;
 
 	FRayTracingMeshResourceCollector(
