RL - Fix Copents without proxy to call trough a null-skincache

diff --git a/Engine/Source/Runtime/Engine/Private/SkeletalRenderGPUSkin.cpp b/Engine/Source/Runtime/Engine/Private/SkeletalRenderGPUSkin.cpp
--- a/Engine/Source/Runtime/Engine/Private/SkeletalRenderGPUSkin.cpp
+++ b/Engine/Source/Runtime/Engine/Private/SkeletalRenderGPUSkin.cpp
@@ -432,7 +432,7 @@ void FSkeletalMeshObjectGPUSkin::ProcessUpdatedDynamicData(FGPUSkinCache* GPUSki
 
 	bool bDataPresent = false;
 
-	bool bGPUSkinCacheEnabled = GEnableGPUSkinCache && (FeatureLevel >= ERHIFeatureLevel::SM5);
+	bool bGPUSkinCacheEnabled = GPUSkinCache && GEnableGPUSkinCache && (FeatureLevel >= ERHIFeatureLevel::SM5);
 
 	if (LOD.MorphVertexBuffer.bNeedsInitialClear && !(bMorph && bMorphNeedsUpdate))
 	{
