Ray traced reflection and translucency default roughness threshold now match equivalent SSR threshold

diff --git a/Engine/Source/Runtime/Engine/Private/Scene.cpp b/Engine/Source/Runtime/Engine/Private/Scene.cpp
--- a/Engine/Source/Runtime/Engine/Private/Scene.cpp
+++ b/Engine/Source/Runtime/Engine/Private/Scene.cpp
@@ -553,13 +553,13 @@ FPostProcessSettings::FPostProcessSettings()
 	ScreenSpaceReflectionIntensity = 100.0f;
 	ScreenSpaceReflectionQuality = 50.0f;
 	ScreenSpaceReflectionMaxRoughness = 0.6f;
-	RayTracingReflectionsMaxRoughness = 1;
+	RayTracingReflectionsMaxRoughness = 0.6f;
 	RayTracingReflectionsMaxBounces = 1;
 	RayTracingReflectionsSamplesPerPixel = 1;
 	RayTracingReflectionsShadows = EReflectedAndRefractedRayTracedShadows::Hard_shadows;
 	
 	TranslucencyType = ETranslucencyType::RayTracing;
-	RayTracingTranslucencyMaxRoughness = 1;
+	RayTracingTranslucencyMaxRoughness = 0.6f;
 	RayTracingTranslucencyRefractionRays = 3; // 3 to: first hit surface, second hit back inner surface and a third to fetch the background.
 	RayTracingTranslucencySamplesPerPixel = 1;
 	RayTracingTranslucencyShadows = EReflectedAndRefractedRayTracedShadows::Hard_shadows;
