Fixed SplineMeshComponent invalid BoxExtent when setting huge Spline's EndTangent
#rnx

diff --git a/Engine/Source/Runtime/Engine/Private/Components/SplineMeshComponent.cpp b/Engine/Source/Runtime/Engine/Private/Components/SplineMeshComponent.cpp
--- a/Engine/Source/Runtime/Engine/Private/Components/SplineMeshComponent.cpp
+++ b/Engine/Source/Runtime/Engine/Private/Components/SplineMeshComponent.cpp
@@ -1199,19 +1199,21 @@ float USplineMeshComponent::GetTextureStreamingTransformScale() const
 #if WITH_EDITOR
 void USplineMeshComponent::PostEditChangeProperty(FPropertyChangedEvent& PropertyChangedEvent)
 {
-	UStaticMeshComponent::PostEditChangeProperty(PropertyChangedEvent);
 	UProperty* MemberPropertyThatChanged = PropertyChangedEvent.MemberProperty;
-	if (MemberPropertyThatChanged)
+	bool bIsSplineParamsChange = MemberPropertyThatChanged && MemberPropertyThatChanged->GetNameCPP() == TEXT("SplineParams");
+	if (bIsSplineParamsChange)
 	{
-		// If the spline params were changed the actual geometry is, so flag the owning HLOD cluster as dirty
-		if (MemberPropertyThatChanged->GetNameCPP() == TEXT("SplineParams"))
-		{
-			SetEndTangent(SplineParams.EndTangent, false);
-
-			IHierarchicalLODUtilitiesModule& Module = FModuleManager::LoadModuleChecked<IHierarchicalLODUtilitiesModule>("HierarchicalLODUtilities");
-			IHierarchicalLODUtilities* Utilities = Module.GetUtilities();
-			Utilities->HandleActorModified(GetOwner());
-		}
+		SetEndTangent(SplineParams.EndTangent, false);
+	}
+	
+	UStaticMeshComponent::PostEditChangeProperty(PropertyChangedEvent);
+	
+	// If the spline params were changed the actual geometry is, so flag the owning HLOD cluster as dirty
+	if (bIsSplineParamsChange)
+	{
+		IHierarchicalLODUtilitiesModule& Module = FModuleManager::LoadModuleChecked<IHierarchicalLODUtilitiesModule>("HierarchicalLODUtilities");
+		IHierarchicalLODUtilities* Utilities = Module.GetUtilities();
+		Utilities->HandleActorModified(GetOwner());
 	}
 }
 #endif
