Mark the end of frame update arrays to NonTransactional to avoid them being restored including components that have internal data that disagrees with them being in the array. Resulted in said componets being added to the list and in turn failing a check when processing.

diff --git a/Engine/Source/Runtime/Engine/Classes/Engine/World.h b/Engine/Source/Runtime/Engine/Classes/Engine/World.h
--- a/Engine/Source/Runtime/Engine/Classes/Engine/World.h
+++ b/Engine/Source/Runtime/Engine/Classes/Engine/World.h
@@ -1230,11 +1230,11 @@ private:
 	FPhysScene*									PhysicsScene;
 
 	/** Array of components that need updates at the end of the frame */
-	UPROPERTY(Transient)
+	UPROPERTY(Transient, NonTransactional)
 	TArray<UActorComponent*> ComponentsThatNeedEndOfFrameUpdate;
 
 	/** Array of components that need game thread updates at the end of the frame */
-	UPROPERTY(Transient)
+	UPROPERTY(Transient, NonTransactional)
 	TArray<UActorComponent*> ComponentsThatNeedEndOfFrameUpdate_OnGameThread;
 
 	/** The state of async tracing - abstracted into its own object for easier reference */
