Default translucency is set to Raster instead of Ray Tracing in the post processing volume

diff --git a/Engine/Source/Runtime/Engine/Private/Scene.cpp b/Engine/Source/Runtime/Engine/Private/Scene.cpp
--- a/Engine/Source/Runtime/Engine/Private/Scene.cpp
+++ b/Engine/Source/Runtime/Engine/Private/Scene.cpp
@@ -558,7 +558,7 @@ FPostProcessSettings::FPostProcessSettings()
 	RayTracingReflectionsSamplesPerPixel = 1;
 	RayTracingReflectionsShadows = EReflectedAndRefractedRayTracedShadows::Hard_shadows;
 	
-	TranslucencyType = ETranslucencyType::RayTracing;
+	TranslucencyType = ETranslucencyType::Raster;
 	RayTracingTranslucencyMaxRoughness = 0.6f;
 	RayTracingTranslucencyRefractionRays = 3; // 3 to: first hit surface, second hit back inner surface and a third to fetch the background.
 	RayTracingTranslucencySamplesPerPixel = 1;
