Fix crash when trying to calculate required bones with no render data

diff --git a/Engine/Source/Runtime/Engine/Private/Components/SkeletalMeshComponent.cpp b/Engine/Source/Runtime/Engine/Private/Components/SkeletalMeshComponent.cpp
--- a/Engine/Source/Runtime/Engine/Private/Components/SkeletalMeshComponent.cpp
+++ b/Engine/Source/Runtime/Engine/Private/Components/SkeletalMeshComponent.cpp
@@ -1490,7 +1490,13 @@ void USkeletalMeshComponent::ComputeRequiredBones(TArray<FBoneIndexType>& OutReq
 	}
 
 	FSkeletalMeshRenderData* SkelMeshRenderData = GetSkeletalMeshRenderData();
-	check(SkelMeshRenderData);
+	if (!SkelMeshRenderData)
+	{
+		//No Render Data?
+		// Jira UE-64409
+		UE_LOG(LogAnimation, Warning, TEXT("Skeletal Mesh asset '%s' has no render data"), *SkeletalMesh->GetName());
+		return;
+	}
 
 	// Make sure we access a valid LOD
 	// @fixme jira UE-30028 Avoid crash when called with partially loaded asset
