Add more ensure to narrow down with bone indices mismatch of animation inside of blendspace

diff --git a/Engine/Source/Runtime/Engine/Private/Animation/BlendSpaceBase.cpp b/Engine/Source/Runtime/Engine/Private/Animation/BlendSpaceBase.cpp
--- a/Engine/Source/Runtime/Engine/Private/Animation/BlendSpaceBase.cpp
+++ b/Engine/Source/Runtime/Engine/Private/Animation/BlendSpaceBase.cpp
@@ -571,7 +571,12 @@ void UBlendSpaceBase::GetAnimationPose(TArray<FBlendSampleData>& BlendSampleData
 			const FBlendSample& Sample = SampleData[BlendSampleDataCache[I].SampleDataIndex];
 			ChildrenWeights[I] = BlendSampleDataCache[I].GetWeight();
 
-			if(Sample.Animation)
+			if(Sample.Animation 
+#if WITH_EDITOR
+				// verify if Sample.Animation->GetSkeleton matches
+				&& ensure(Sample.Animation->GetSkeleton() == GetSkeleton())
+#endif // WITH_EDITOR
+			)
 			{
 				const float Time = FMath::Clamp<float>(BlendSampleDataCache[I].Time, 0.f, Sample.Animation->SequenceLength);
 
@@ -651,7 +656,11 @@ bool UBlendSpaceBase::GetSamplesFromBlendInput(const FVector &BlendInput, TArray
 		for(int32 Ind = 0; Ind < GridElement.MAX_VERTICES; ++Ind)
 		{
 			const int32 SampleDataIndex = GridElement.Indices[Ind];		
-			if(SampleData.IsValidIndex(SampleDataIndex))
+			if( SampleData.IsValidIndex(SampleDataIndex) 
+#if WITH_EDITOR
+				&& SampleData[SampleDataIndex].Animation->GetSkeleton() == GetSkeleton()
+#endif // WITH_EDITOR
+				)
 			{
 				int32 Index = OutSampleDataList.AddUnique(SampleDataIndex);
 				FBlendSampleData& NewSampleData = OutSampleDataList[Index];
