If loading a level containing a ISM after a lighting scenario was loaded, the lighting scenario would not get applied correctly. To fix this, we simply rebuild the render data if there is a lighting scenario, as we cannot assume that the current data was build with the current lighting scenario.

diff --git a/Engine/Source/Runtime/Engine/Private/InstancedStaticMesh.cpp b/Engine/Source/Runtime/Engine/Private/InstancedStaticMesh.cpp
--- a/Engine/Source/Runtime/Engine/Private/InstancedStaticMesh.cpp
+++ b/Engine/Source/Runtime/Engine/Private/InstancedStaticMesh.cpp
@@ -2238,6 +2238,19 @@ void UInstancedStaticMeshComponent::OnPostLoadPerInstanceData()
 
 	// release InstanceDataBuffers
 	InstanceDataBuffers.Reset();
+
+	if (PerInstanceRenderData.IsValid())
+	{
+		AActor* Owner = GetOwner();
+		ULevel* OwnerLevel = Owner->GetLevel();
+		UWorld* OwnerWorld = OwnerLevel ? OwnerLevel->OwningWorld : nullptr;
+
+		if (OwnerWorld && OwnerWorld->GetActiveLightingScenario() != nullptr && OwnerWorld->GetActiveLightingScenario() != OwnerLevel)
+		{
+			//update the instance data if the lighting scenario isn't the owner level
+			InstanceUpdateCmdBuffer.Edit();
+		}
+	}
 }
 
 void UInstancedStaticMeshComponent::PartialNavigationUpdate(int32 InstanceIdx)
