[LLM] InstancedMesh vertex data now included in the tag.

diff --git a/Engine/Source/Runtime/Engine/Private/InstancedStaticMesh.cpp b/Engine/Source/Runtime/Engine/Private/InstancedStaticMesh.cpp
--- a/Engine/Source/Runtime/Engine/Private/InstancedStaticMesh.cpp
+++ b/Engine/Source/Runtime/Engine/Private/InstancedStaticMesh.cpp
@@ -278,6 +278,7 @@ void FStaticMeshInstanceBuffer::InitRHI()
 	if (InstanceData->GetNumInstances() > 0)
 	{
 		QUICK_SCOPE_CYCLE_COUNTER(STAT_FStaticMeshInstanceBuffer_InitRHI);
+		LLM_SCOPE(ELLMTag::InstancedMesh);
 		auto AccessFlags = BUF_Static;
 		CreateVertexBuffer(InstanceData->GetOriginResourceArray(), AccessFlags | BUF_ShaderResource, 16, PF_A32B32G32R32F, InstanceOriginBuffer.VertexBufferRHI, InstanceOriginSRV);
 		CreateVertexBuffer(InstanceData->GetTransformResourceArray(), AccessFlags | BUF_ShaderResource, InstanceData->GetTranslationUsesHalfs() ? 8 : 16, InstanceData->GetTranslationUsesHalfs() ? PF_FloatRGBA : PF_A32B32G32R32F, InstanceTransformBuffer.VertexBufferRHI, InstanceTransformSRV);
