Fix crash in Fast Array Delta Serialization caused by using an invalid array.

diff --git a/Engine/Source/Runtime/Engine/Private/RepLayout.cpp b/Engine/Source/Runtime/Engine/Private/RepLayout.cpp
--- a/Engine/Source/Runtime/Engine/Private/RepLayout.cpp
+++ b/Engine/Source/Runtime/Engine/Private/RepLayout.cpp
@@ -6036,7 +6036,7 @@ bool FRepLayout::DeltaSerializeFastArrayProperty(FFastArrayDeltaSerializeParams&
 										}
 										else
 										{
-											FChangelistIterator ElementChangelistIterator(ElementChangelist, 0);
+											FChangelistIterator ElementChangelistIterator(Temp, 0);
 											FRepHandleIterator ElementHandleIterator(ElementChangelistIterator, Cmds, ArrayHandleToCmdIndex, ElementSize, 1, ItemLayoutStart, ItemLayoutEnd);
 
 											MergeChangeList_r(FoundHandleIterator, ElementHandleIterator, ElementData, ElementChangelist);
