UTextures loaded via normal asset loading are GPU only and compact bulk data doesn't support reload from disk. If it is desired to manipulate the texel data of texture assets, render commands need to be used

diff --git a/Engine/Source/Runtime/Engine/Private/Texture2D.cpp b/Engine/Source/Runtime/Engine/Private/Texture2D.cpp
--- a/Engine/Source/Runtime/Engine/Private/Texture2D.cpp
+++ b/Engine/Source/Runtime/Engine/Private/Texture2D.cpp
@@ -270,6 +270,23 @@ void FTexture2DMipMap::FCompactByteBulkData::Serialize(FArchive& Ar, UObject* Ow
 	}
 }
 
+const void* FTexture2DMipMap::FCompactByteBulkData::LockReadOnly() const
+{
+	return const_cast<FCompactByteBulkData*>(this)->Lock(LOCK_READ_ONLY);
+}
+
+void* FTexture2DMipMap::FCompactByteBulkData::Lock(uint32 LockFlags)
+{
+	ensureMsgf(LockFlags != LOCK_READ_ONLY || TexelData,
+		TEXT("Locking bulk data for read only but no data is available. ")
+		TEXT("A possible cause is that GetCopy has been called with bDiscardInternal copy set. ")
+		TEXT("Note that UTexture's loaded via normal asset loading (e.g. FLinkerLoad) is GPU only ")
+		TEXT("and their CPU copies of texel data are discarded after resource creation. If you ")
+		TEXT("want to manipulate their data, you need to use render commands. For example, to ")
+		TEXT("copy a texture into another, use FRHICommandList::CopyTexture"));
+	return TexelData;
+}
+
 void FTexture2DMipMap::FCompactByteBulkData::Unlock() const
 {
 	if (!!(BulkDataFlags & BULKDATA_SingleUse))
diff --git a/Engine/Source/Runtime/Engine/Public/TextureResource.h b/Engine/Source/Runtime/Engine/Public/TextureResource.h
--- a/Engine/Source/Runtime/Engine/Public/TextureResource.h
+++ b/Engine/Source/Runtime/Engine/Public/TextureResource.h
@@ -63,8 +63,8 @@ struct FTexture2DMipMap
 		bool IsAvailableForUse() const { return !(BulkDataFlags & BULKDATA_Unused); }
 		bool IsBulkDataLoaded() const { return IsInlined(); }
 		bool IsStoredCompressedOnDisk() const { return !!(BulkDataFlags & BULKDATA_SerializeCompressed); }
-		const void* LockReadOnly() const { return TexelData; }
-		void* Lock(uint32 LockFlags) { return TexelData; }
+		const void* LockReadOnly() const;
+		void* Lock(uint32 LockFlags);
 		void Unlock() const;
 		void* Realloc(int32 NumBytes);
 		void GetCopy(void** Dest, bool bDiscardInternalCopy = true);
