#ue4 - Force a substep of falling movement when the motion would pass the jump apex point. We compute the time to apex and refund the remaining time for the next iteration. This also more accurately triggers the NotifyApex() functionality because it's triggered after the movement that reaches the exact apex point.

Confirmed reaching the same exact jump apex position under framerates varying 8 to 120 fps, when previously it could vary in my test up to a few centimeter difference (also improved by the fix in 5290785 to jump consistency when ending the jump force if applicable).

"p.ForceJumpPeakSubstep 0" can disable if not desired for some reason.

diff --git a/Engine/Source/Runtime/Engine/Private/Components/CharacterMovementComponent.cpp b/Engine/Source/Runtime/Engine/Private/Components/CharacterMovementComponent.cpp
--- a/Engine/Source/Runtime/Engine/Private/Components/CharacterMovementComponent.cpp
+++ b/Engine/Source/Runtime/Engine/Private/Components/CharacterMovementComponent.cpp
@@ -196,6 +196,13 @@ namespace CharacterMovementCVars
 		TEXT( "If 1, remove invalid replay samples that can occur due to oversampling (sampling at higher rate than physics is being ticked)" ),
 		ECVF_Default);
 
+	static int32 ForceJumpPeakSubstep = 1;
+	FAutoConsoleVariableRef CVarForceJumpPeakSubstep(
+		TEXT("p.ForceJumpPeakSubstep"),
+		ForceJumpPeakSubstep,
+		TEXT("If 1, force a jump substep to always reach the peak position of a jump, which can often be cut off as framerate lowers."),
+		ECVF_Default);
+
 #if !UE_BUILD_SHIPPING
 
 	int32 NetShowCorrections = 0;
@@ -4108,12 +4115,13 @@ void UCharacterMovementComponent::PhysFalling(float deltaTime, int32 Iterations)
 	FVector FallAcceleration = GetFallingLateralAcceleration(deltaTime);
 	FallAcceleration.Z = 0.f;
 	const bool bHasAirControl = (FallAcceleration.SizeSquared2D() > 0.f);
+	int32 NumApexAttempts = 0;
 
 	float remainingTime = deltaTime;
 	while( (remainingTime >= MIN_TICK_TIME) && (Iterations < MaxSimulationIterations) )
 	{
 		Iterations++;
-		const float timeTick = GetSimulationTimeStep(remainingTime, Iterations);
+		float timeTick = GetSimulationTimeStep(remainingTime, Iterations);
 		remainingTime -= timeTick;
 		
 		const FVector OldLocation = UpdatedComponent->GetComponentLocation();
@@ -4156,7 +4164,7 @@ void UCharacterMovementComponent::PhysFalling(float deltaTime, int32 Iterations)
 			}
 		}
 
-		// Apply gravity
+		// Compute current gravity
 		const FVector Gravity(0.f, 0.f, GetGravityZ());
 		float GravityTime = timeTick;
 
@@ -4177,20 +4185,46 @@ void UCharacterMovementComponent::PhysFalling(float deltaTime, int32 Iterations)
 			}
 		}
 
+		// Apply gravity
 		Velocity = NewFallVelocity(Velocity, Gravity, GravityTime);
 		VelocityNoAirControl = bHasAirControl ? NewFallVelocity(VelocityNoAirControl, Gravity, GravityTime) : Velocity;
 		const FVector AirControlAccel = (Velocity - VelocityNoAirControl) / timeTick;
 
+		// See if we need to sub-step to exactly reach the apex. This is important for avoiding "cutting off the top" of the trajectory as framerate varies.
+		if (CharacterMovementCVars::ForceJumpPeakSubstep && OldVelocity.Z > 0.f && Velocity.Z <= 0.f && NumApexAttempts < 2)
+		{
+			const FVector DerivedAccel = (Velocity - OldVelocity) / timeTick;
+			if (!FMath::IsNearlyZero(DerivedAccel.Z))
+			{
+				const float TimeToApex = -OldVelocity.Z / DerivedAccel.Z;
+				
+				// The time-to-apex calculation should be precise, and we want to avoid adding a substep when we are basically already at the apex from the previous iteration's work.
+				const float ApexTimeMinimum = 0.0001f;
+				if (TimeToApex >= ApexTimeMinimum && TimeToApex < timeTick)
+				{
+					const FVector ApexVelocity = OldVelocity + DerivedAccel * TimeToApex;
+					Velocity = ApexVelocity;
+					Velocity.Z = 0.f; // Should be nearly zero anyway, but this makes apex notifications consistent.
+
+					// We only want to move the amount of time it takes to reach the apex, and refund the unused time for next iteration.
+					remainingTime += (timeTick - TimeToApex);
+					timeTick = TimeToApex;
+					Iterations--;
+					NumApexAttempts++;
+				}
+			}
+		}
+
+		//UE_LOG(LogCharacterMovement, Log, TEXT("dt=(%.6f) OldLocation=(%s) OldVelocity=(%s) NewVelocity=(%s)"), timeTick, *(UpdatedComponent->GetComponentLocation()).ToString(), *OldVelocity.ToString(), *Velocity.ToString());
 		ApplyRootMotionToVelocity(timeTick);
 
-		if( bNotifyApex && (Velocity.Z <= 0.f) )
+		if (bNotifyApex && (Velocity.Z < 0.f))
 		{
 			// Just passed jump apex since now going down
 			bNotifyApex = false;
 			NotifyJumpApex();
 		}
 
-
 		// Compute change in position (using midpoint integration method).
 		FVector Adjusted = 0.5f*(OldVelocity + Velocity) * timeTick;
 		
