Network profiler now correctly reports the size of RPCs with parameters larger than 8KB.
Make network profiler version history explicit.

diff --git a/Engine/Source/Programs/NetworkProfiler/NetworkProfiler/Tokens.cs b/Engine/Source/Programs/NetworkProfiler/NetworkProfiler/Tokens.cs
--- a/Engine/Source/Programs/NetworkProfiler/NetworkProfiler/Tokens.cs
+++ b/Engine/Source/Programs/NetworkProfiler/NetworkProfiler/Tokens.cs
@@ -85,7 +85,7 @@ namespace NetworkProfiler
 					SerializedToken = new TokenSendBunch( BinaryStream, InNetworkStream.GetVersion() );
 					break;
 				case ETokenTypes.SendRPC:
-					SerializedToken = new TokenSendRPC( BinaryStream );
+					SerializedToken = new TokenSendRPC( BinaryStream, InNetworkStream.GetVersion() );
 					break;
 				case ETokenTypes.ReplicateActor:
 					SerializedToken = new TokenReplicateActor( BinaryStream );
@@ -332,20 +332,30 @@ namespace NetworkProfiler
 		/** Name table index of function name. */
 		public int FunctionNameIndex;
 		/** Number of bits serialized/sent for the header. */
-		public UInt16 NumHeaderBits;
+		public int NumHeaderBits;
 		/** Number of bits serialized/sent for the parameters. */
-		public UInt16 NumParameterBits;
+		public int NumParameterBits;
 		/** Number of bits serialized/sent for the footer. */
-		public UInt16 NumFooterBits;
+		public int NumFooterBits;
 
 		/** Constructor, serializing members from passed in stream. */
-		public TokenSendRPC( BinaryReader BinaryStream )
+		public TokenSendRPC(BinaryReader BinaryStream, UInt32 Version)
 		{
-			ActorNameIndex = TokenHelper.LoadPackedInt( BinaryStream );
-			FunctionNameIndex = TokenHelper.LoadPackedInt( BinaryStream );
-			NumHeaderBits = BinaryStream.ReadUInt16(); 					
-			NumParameterBits = BinaryStream.ReadUInt16();
-			NumFooterBits = BinaryStream.ReadUInt16();
+			ActorNameIndex = TokenHelper.LoadPackedInt(BinaryStream);
+			FunctionNameIndex = TokenHelper.LoadPackedInt(BinaryStream);
+
+			if (Version < 13)
+			{
+				NumHeaderBits = BinaryStream.ReadUInt16();
+				NumParameterBits = BinaryStream.ReadUInt16();
+				NumFooterBits = BinaryStream.ReadUInt16();
+			}
+			else
+			{
+				NumHeaderBits = TokenHelper.LoadPackedInt(BinaryStream);
+				NumParameterBits = TokenHelper.LoadPackedInt(BinaryStream);
+				NumFooterBits = TokenHelper.LoadPackedInt(BinaryStream);
+			}
 		}
 
 		/**
diff --git a/Engine/Source/Runtime/Engine/Private/NetworkProfiler.cpp b/Engine/Source/Runtime/Engine/Private/NetworkProfiler.cpp
--- a/Engine/Source/Runtime/Engine/Private/NetworkProfiler.cpp
+++ b/Engine/Source/Runtime/Engine/Private/NetworkProfiler.cpp
@@ -30,8 +30,20 @@ FNetworkProfiler GNetworkProfiler;
 
 /** Magic value, determining that file is a network profiler file.				*/
 #define NETWORK_PROFILER_MAGIC						0x1DBF348C
-/** Version of memory profiler. Incremented on serialization changes.			*/
-#define NETWORK_PROFILER_VERSION					12
+
+/** Version history of network profiler. Entries added on serialization changes. */
+enum class ENetworkProfilerVersionHistory : uint32
+{
+	Initial = 10,
+	ChannelTypesAsStrings = 11,
+	AddressesAsStrings = 12,
+	LargeRPCFix = 13,				// Changed RPC size tracking to packed 32 bits instead of plain uint16s, to handle very large RPCs.
+
+	// New history items go above here.
+
+	VersionPlusOne,
+	Latest = VersionPlusOne - 1
+};
 
 static const FString UnknownName("UnknownName");
 
@@ -66,7 +78,7 @@ enum ENetworkProfilingPayloadType
 
 FNetworkProfilerHeader::FNetworkProfilerHeader()
 	: Magic(NETWORK_PROFILER_MAGIC)
-	, Version(NETWORK_PROFILER_VERSION)
+	, Version(ENetworkProfilerVersionHistory::Latest)
 {
 }
 
@@ -241,37 +253,37 @@ void FNetworkProfiler::SetCurrentConnection( UNetConnection* Connection )
  * @param	Function	Function being called
  * @param	NumBits		Number of bits serialized into bunch for this RPC
  */
-void FNetworkProfiler::TrackSendRPC( const AActor* Actor, const UFunction* Function, uint16 NumHeaderBits, uint16 NumParameterBits, uint16 NumFooterBits, UNetConnection* Connection )
+void FNetworkProfiler::TrackSendRPC(const AActor* Actor, const UFunction* Function, uint32 NumHeaderBits, uint32 NumParameterBits, uint32 NumFooterBits, UNetConnection* Connection)
 {
-	if( bIsTrackingEnabled )
+	if (bIsTrackingEnabled)
 	{
-		SCOPE_LOCK_REF( CriticalSection );
+		SCOPE_LOCK_REF(CriticalSection);
 
-		SetCurrentConnection( Connection );
+		SetCurrentConnection(Connection);
 
-		uint32 ActorNameTableIndex = GetNameTableIndex( Actor->GetName() );
-		uint32 FunctionNameTableIndex = GetNameTableIndex( Function->GetName() );
+		uint32 ActorNameTableIndex = GetNameTableIndex(Actor->GetName());
+		uint32 FunctionNameTableIndex = GetNameTableIndex(Function->GetName());
 
 		uint8 Type = NPTYPE_SendRPC;
 		(*FileWriter) << Type;
-		( *FileWriter ).SerializeIntPacked( ActorNameTableIndex );
-		( *FileWriter ).SerializeIntPacked( FunctionNameTableIndex );
-		( *FileWriter ) << NumHeaderBits;
-		(*FileWriter) << NumParameterBits;
-		(*FileWriter) << NumFooterBits;
+		(*FileWriter).SerializeIntPacked(ActorNameTableIndex);
+		(*FileWriter).SerializeIntPacked(FunctionNameTableIndex);
+		(*FileWriter).SerializeIntPacked(NumHeaderBits);
+		(*FileWriter).SerializeIntPacked(NumParameterBits);
+		(*FileWriter).SerializeIntPacked(NumFooterBits);
 	}
 }
 
-void FNetworkProfiler::TrackQueuedRPC( UNetConnection* Connection, UObject* TargetObject, const AActor* Actor, const UFunction* Function, uint16 NumHeaderBits, uint16 NumParameterBits, uint16 NumFooterBits  )
+void FNetworkProfiler::TrackQueuedRPC(UNetConnection* Connection, UObject* TargetObject, const AActor* Actor, const UFunction* Function, uint32 NumHeaderBits, uint32 NumParameterBits, uint32 NumFooterBits)
 {
-	if( bIsTrackingEnabled )
+	if (bIsTrackingEnabled)
 	{
 		SCOPE_LOCK_REF(CriticalSection);
 		
 		FQueuedRPCInfo Info;
 
-		Info.ActorNameIndex = GetNameTableIndex( Actor->GetName() );
-		Info.FunctionNameIndex = GetNameTableIndex( Function->GetName() );
+		Info.ActorNameIndex = GetNameTableIndex(Actor->GetName());
+		Info.FunctionNameIndex = GetNameTableIndex(Function->GetName());
 
 		Info.Connection = Connection;
 		Info.TargetObject = TargetObject;
@@ -283,25 +295,25 @@ void FNetworkProfiler::TrackQueuedRPC( UNetConnection* Connection, UObject* Targ
 	}
 }
 
-void FNetworkProfiler::FlushQueuedRPCs( UNetConnection* Connection, UObject* TargetObject )
+void FNetworkProfiler::FlushQueuedRPCs(UNetConnection* Connection, UObject* TargetObject)
 {
-	if( bIsTrackingEnabled )
+	if (bIsTrackingEnabled)
 	{
 		SCOPE_LOCK_REF(CriticalSection);
 
-		SetCurrentConnection( Connection );
+		SetCurrentConnection(Connection);
 
-		for ( int i = QueuedRPCs.Num() - 1; i >= 0; --i )
+		for (int i = QueuedRPCs.Num() - 1; i >= 0; --i)
 		{
 			if (QueuedRPCs[i].Connection == Connection && QueuedRPCs[i].TargetObject == TargetObject)
 			{
 				uint8 Type = NPTYPE_SendRPC;
 				(*FileWriter) << Type;
-				(*FileWriter).SerializeIntPacked( QueuedRPCs[i].ActorNameIndex );
-				(*FileWriter).SerializeIntPacked( QueuedRPCs[i].FunctionNameIndex );
-				(*FileWriter) << QueuedRPCs[i].NumHeaderBits;
-				(*FileWriter) << QueuedRPCs[i].NumParameterBits;
-				(*FileWriter) << QueuedRPCs[i].NumFooterBits;
+				(*FileWriter).SerializeIntPacked(QueuedRPCs[i].ActorNameIndex);
+				(*FileWriter).SerializeIntPacked(QueuedRPCs[i].FunctionNameIndex);
+				(*FileWriter).SerializeIntPacked(QueuedRPCs[i].NumHeaderBits);
+				(*FileWriter).SerializeIntPacked(QueuedRPCs[i].NumParameterBits);
+				(*FileWriter).SerializeIntPacked(QueuedRPCs[i].NumFooterBits);
 
 				QueuedRPCs.RemoveAtSwap(i);
 			}
diff --git a/Engine/Source/Runtime/Engine/Public/Net/NetworkProfiler.h b/Engine/Source/Runtime/Engine/Public/Net/NetworkProfiler.h
--- a/Engine/Source/Runtime/Engine/Public/Net/NetworkProfiler.h
+++ b/Engine/Source/Runtime/Engine/Public/Net/NetworkProfiler.h
@@ -21,6 +21,8 @@ struct FURL;
 
 DECLARE_MULTICAST_DELEGATE_OneParam(FOnNetworkProfileFinished, const FString& /*Filename */);
 
+enum class ENetworkProfilerVersionHistory : uint32;
+
 /*=============================================================================
 	Network profiler header.
 =============================================================================*/
@@ -31,7 +33,7 @@ private:
 	/** Magic to ensure we're opening the right file.	*/
 	uint32	Magic;
 	/** Version number to detect version mismatches.	*/
-	uint32	Version;
+	ENetworkProfilerVersionHistory	Version;
 
 	/** Tag, set via -networkprofiler=TAG				*/
 	FString Tag;
@@ -133,9 +135,9 @@ private:
 		UObject* TargetObject;
 		uint32 ActorNameIndex;
 		uint32 FunctionNameIndex;
-		uint16 NumHeaderBits;
-		uint16 NumParameterBits;
-		uint16 NumFooterBits;
+		uint32 NumHeaderBits;
+		uint32 NumParameterBits;
+		uint32 NumFooterBits;
 
 		FQueuedRPCInfo()
 			: Connection(nullptr)
@@ -201,7 +203,7 @@ public:
 	 * @param	NumParameterBits	Number of bits serialized into parameters of this RPC
 	 * @param	NumFooterBits		Number of bits serialized into the footer of this RPC (EndContentBlock)
 	 */
-	void TrackSendRPC( const AActor* Actor, const UFunction* Function, uint16 NumHeaderBits, uint16 NumParameterBits, uint16 NumFooterBits, UNetConnection* Connection );
+	void TrackSendRPC(const AActor* Actor, const UFunction* Function, uint32 NumHeaderBits, uint32 NumParameterBits, uint32 NumFooterBits, UNetConnection* Connection);
 	
 	/**
 	 * Tracks queued RPCs (unreliable multicast) being sent.
@@ -214,7 +216,7 @@ public:
 	 * @param	NumParameterBits	Number of bits serialized into parameters of this RPC
 	 * @param	NumFooterBits		Number of bits serialized into the footer of this RPC (EndContentBlock)
 	 */
-	void TrackQueuedRPC( UNetConnection* Connection, UObject* TargetObject, const AActor* Actor, const UFunction* Function, uint16 NumHeaderBits, uint16 NumParameterBits, uint16 NumFooterBits );
+	void TrackQueuedRPC(UNetConnection* Connection, UObject* TargetObject, const AActor* Actor, const UFunction* Function, uint32 NumHeaderBits, uint32 NumParameterBits, uint32 NumFooterBits);
 
 	/**
 	 * Writes all queued RPCs for the connection to the profiler stream
