Added out of order packet sequence correction, to fix issue with some players experiencing greater packetloss, when this is not necessary.

diff --git a/Engine/Source/Runtime/Engine/Classes/Engine/NetConnection.h b/Engine/Source/Runtime/Engine/Classes/Engine/NetConnection.h
--- a/Engine/Source/Runtime/Engine/Classes/Engine/NetConnection.h
+++ b/Engine/Source/Runtime/Engine/Classes/Engine/NetConnection.h
@@ -19,6 +19,7 @@
 #include "Engine/Channel.h"
 #include "ProfilingDebugging/Histogram.h"
 #include "Containers/ArrayView.h"
+#include "Containers/CircularBuffer.h"
 #include "ReplicationDriver.h"
 #include "Analytics/EngineNetAnalytics.h"
 #include "PacketTraits.h"
@@ -1062,6 +1063,13 @@ public:
 	/** Removes stale entries from DormantReplicatorMap. */
 	void CleanupStaleDormantReplicators();
 
+	/**
+	 * Flush the cache of sequenced packets waiting for a missing packet. Will flush only up to the next missing packet, unless bFlushWholeCache is set.
+	 *
+	 * @param bFlushWholeCache	Whether or not the whole cache should be flushed, or only flush up to the next missing packet
+	 */
+	void FlushPacketOrderCache(bool bFlushWholeCache=false);
+
 protected:
 
 	ENGINE_API void SetPendingCloseDueToSocketSendFailure();
@@ -1146,6 +1154,24 @@ private:
 	bool bConnectionPendingCloseDueToSocketSendFailure;
 
 	FNetworkGUID GetActorGUIDFromOpenBunch(FInBunch& Bunch);
+
+
+	/** Out of order packet tracking/correction */
+
+	/** Stat tracking for the total number of out of order packets, for this connection */
+	int32 TotalOutOfOrderPackets;
+
+	/** Buffer of partially read (post-PacketHandler) sequenced packets, which are waiting for a missing packet/sequence */
+	TOptional<TCircularBuffer<TUniquePtr<FBitReader>>> PacketOrderCache;
+
+	/** The current start index for PacketOrderCache */
+	int32 PacketOrderCacheStartIdx;
+
+	/** The current number of valid packets in PacketOrderCache */
+	int32 PacketOrderCacheCount;
+
+	/** Whether or not PacketOrderCache is presently being flushed */
+	bool bFlushingPacketOrderCache;
 };
 
 
diff --git a/Engine/Source/Runtime/Engine/Private/NetConnection.cpp b/Engine/Source/Runtime/Engine/Private/NetConnection.cpp
--- a/Engine/Source/Runtime/Engine/Private/NetConnection.cpp
+++ b/Engine/Source/Runtime/Engine/Private/NetConnection.cpp
@@ -51,6 +51,14 @@ static TAutoConsoleVariable<int32> CVarMaxChannelSize(TEXT("net.MaxChannelSize")
 static TAutoConsoleVariable<int32> CVarForceNetFlush(TEXT("net.ForceNetFlush"), 0, TEXT("Immediately flush send buffer when written to (helps trace packet writes - WARNING: May be unstable)."));
 #endif
 
+static TAutoConsoleVariable<int32> CVarNetDoPacketOrderCorrection(TEXT("net.DoPacketOrderCorrection"), 0, TEXT("Whether or not to try to fix 'out of order' packet sequences, by caching packets and waiting for the missing sequence."));
+
+static TAutoConsoleVariable<int32> CVarNetPacketOrderCorrectionEnableThreshold(TEXT("net.PacketOrderCorrectionEnableThreshold"), 1, TEXT("The number of 'out of order' packet sequences that need to occur, before correction is enabled."));
+
+static TAutoConsoleVariable<int32> CVarNetPacketOrderMaxMissingPackets(TEXT("net.PacketOrderMaxMissingPackets"), 3, TEXT("The maximum number of missed packet sequences that is allowed, before treating missing packets as lost."));
+
+static TAutoConsoleVariable<int32> CVarNetPacketOrderMaxCachedPackets(TEXT("net.PacketOrderMaxCachedPackets"), 32, TEXT("(NOTE: Must be power of 2!) The maximum number of packets to cache while waiting for missing packet sequences, before treating missing packets as lost."));
+
 extern int32 GNetDormancyValidate;
 
 DECLARE_CYCLE_STAT(TEXT("NetConnection SendAcks"), Stat_NetConnectionSendAck, STATGROUP_Net);
@@ -164,6 +172,11 @@ PRAGMA_ENABLE_DEPRECATION_WARNINGS
 ,	HasDirtyAcks(0u)
 ,	bHasWarnedAboutChannelLimit(false)
 ,	bConnectionPendingCloseDueToSocketSendFailure(false)
+,	TotalOutOfOrderPackets(0)
+,	PacketOrderCache()
+,	PacketOrderCacheStartIdx(0)
+,	PacketOrderCacheCount(0)
+,	bFlushingPacketOrderCache(false)
 {
 	// This isn't ideal, because it won't capture memory derived classes are creating dynamically.
 	// The allocations could *probably* be moved somewhere else (like InitBase), but that
@@ -647,7 +660,7 @@ void UNetConnection::CleanUp()
 
 	Close();
 
-	if (Driver != NULL)
+	if (Driver != nullptr)
 	{
 		// Remove from driver.
 		if (Driver->ServerConnection)
@@ -1130,6 +1143,9 @@ void UNetConnection::ReceivedRawPacket( void* InData, int32 Count )
 			if (Reader.GetBitsLeft() > 0)
 			{
 				ReceivedPacket(Reader);
+
+				// Check if the out of order packet cache needs flushing
+				FlushPacketOrderCache();
 			}
 		}
 		// MalformedPacket - Received a packet with 0's in the last byte
@@ -1145,6 +1161,48 @@ void UNetConnection::ReceivedRawPacket( void* InData, int32 Count )
 	}
 }
 
+void UNetConnection::FlushPacketOrderCache(bool bFlushWholeCache/*=false*/)
+{
+	if (PacketOrderCache.IsSet() && PacketOrderCacheCount > 0)
+	{
+		TCircularBuffer<TUniquePtr<FBitReader>>& Cache = PacketOrderCache.GetValue();
+		int32 CacheEndIdx = PacketOrderCache->GetPreviousIndex(PacketOrderCacheStartIdx);
+		bool bEndOfCacheSet = Cache[CacheEndIdx].IsValid();
+
+		bFlushingPacketOrderCache = true;
+
+		// If the end of the cache has had its value set, this forces the flushing of the whole cache, no matter how many missing sequences there are.
+		// The reason for this (other than making space in the cache), is that when we receive a sequence that is out of range of the cache,
+		// it is stored at the end, and so the cache index no longer lines up with the sequence number - which it needs to.
+		bFlushWholeCache = bFlushWholeCache || bEndOfCacheSet;
+
+		while (PacketOrderCacheCount > 0)
+		{
+			TUniquePtr<FBitReader>& CurCachePacket = Cache[PacketOrderCacheStartIdx];
+
+			if (CurCachePacket.IsValid())
+			{
+				UE_LOG(LogNet, Verbose, TEXT("'Out of Order' Packet Cache, replaying packet with cache index: %i (bFlushWholeCache: %i)"), PacketOrderCacheStartIdx, (int32)bFlushWholeCache);
+
+				ReceivedPacket(*CurCachePacket.Get());
+
+				CurCachePacket.Reset();
+
+				PacketOrderCacheCount--;
+			}
+			// Advance the cache only up to the first missing packet, unless flushing the whole cache
+			else if (!bFlushWholeCache)
+			{
+				break;
+			}
+
+			PacketOrderCacheStartIdx = PacketOrderCache->GetNextIndex(PacketOrderCacheStartIdx);
+		}
+
+		bFlushingPacketOrderCache = false;
+	}
+}
+
 uint32 GNetOutBytes = 0;
 
 void UNetConnection::FlushNet(bool bIgnoreSimulation)
@@ -1583,6 +1641,9 @@ void UNetConnection::ReceivedPacket( FBitReader& Reader )
 		return;
 	}
 
+
+	FBitReaderMark ResetReaderMark(Reader);
+
 	ValidateSendBuffer();
 
 	//Record the packet time to the histogram
@@ -1607,45 +1668,65 @@ void UNetConnection::ReceivedPacket( FBitReader& Reader )
 			return;
 		}
 
-		// Lambda to dispatch delivery notifications, 
-		auto HandlePacketNotification = [&Header, this](FNetPacketNotify::SequenceNumberT AckedSequence, bool bDelivered)
+
+		bool bPacketOrderCacheActive = !bFlushingPacketOrderCache && PacketOrderCache.IsSet();
+		bool bCheckForMissingSequence = bPacketOrderCacheActive && PacketOrderCacheCount == 0;
+		bool bFillingPacketOrderCache = bPacketOrderCacheActive && PacketOrderCacheCount > 0;
+		int32 MaxMissingPackets = (bCheckForMissingSequence ? CVarNetPacketOrderMaxMissingPackets.GetValueOnAnyThread() : 0);
+		int32 PacketSequenceDelta = PacketNotify.GetSequenceDelta(Header);
+
+		if (PacketSequenceDelta > 0)
 		{
-			// Increase LastNotifiedPacketId, this is a full packet Id
-			++LastNotifiedPacketId;
+			const int32 MissingPacketCount = PacketSequenceDelta - 1;
 
-			// Sanity check
-			if (FNetPacketNotify::SequenceNumberT(LastNotifiedPacketId) != AckedSequence)
+			// Cache the packet if we are already caching, and begin caching if we just encountered a missing sequence, within range
+			if (bFillingPacketOrderCache || (bCheckForMissingSequence && MissingPacketCount > 0 && MissingPacketCount <= MaxMissingPackets))
 			{
-				CLOSE_CONNECTION_DUE_TO_SECURITY_VIOLATION(this, ESecurityEvent::Malformed_Packet, TEXT("LastNotifiedPacketId != AckedSequence"));
+				int32 LinearCacheIdx = PacketSequenceDelta - 1;
+				int32 CacheCapacity = PacketOrderCache->Capacity();
+				bool bLastCacheEntry = LinearCacheIdx >= (CacheCapacity - 1);
+
+				// The last cache entry is only set, when we've reached capacity or when we receive a sequence which is out of bounds of the cache
+				LinearCacheIdx = bLastCacheEntry ? (CacheCapacity - 1) : LinearCacheIdx;
+
+				int32 CiruclarCacheIdx = PacketOrderCacheStartIdx;
+
+				for (int32 LinearDec=LinearCacheIdx; LinearDec > 0; LinearDec--)
+				{
+					CiruclarCacheIdx = PacketOrderCache->GetNextIndex(CiruclarCacheIdx);
+				}
+
+				TUniquePtr<FBitReader>& CurCachePacket = PacketOrderCache.GetValue()[CiruclarCacheIdx];
+
+				// Reset the reader to its initial position, and cache the packet
+				if (!CurCachePacket.IsValid())
+				{
+					UE_LOG(LogNet, VeryVerbose, TEXT("'Out of Order' Packet Cache, caching sequence order '%i' (capacity: %i)"), LinearCacheIdx, CacheCapacity);
+
+					CurCachePacket = MakeUnique<FBitReader>(Reader);
+					PacketOrderCacheCount++;
+
+					ResetReaderMark.Pop(*CurCachePacket);
+				}
+				else
+				{
+					TotalOutOfOrderPackets++;
+					Driver->InOutOfOrderPackets++;
+				}
+
 				return;
 			}
 
-			if (bDelivered)
-			{
-				ReceivedAck(LastNotifiedPacketId);
-			}
-			else
-			{
-				ReceivedNak(LastNotifiedPacketId);
-			};
-		};
 
-		// Update incoming sequence data and deliver packet notifications
-		// Packet is only accepted if both the incoming sequence number and incoming ack data are valid
-		int PacketSequenceDelta = PacketNotify.Update(Header, HandlePacketNotification);
-		if (PacketSequenceDelta > 0)
-		{
-			const int32 PacketsLost = PacketSequenceDelta - 1;
-		
-			if ( PacketsLost > 10 )
+			if (MissingPacketCount > 10)
 			{
-				UE_LOG( LogNetTraffic, Verbose, TEXT( "High single frame packet loss. PacketsLost: %i %s" ), PacketsLost, *Describe() );
+				UE_LOG(LogNetTraffic, Verbose, TEXT("High single frame packet loss. PacketsLost: %i %s" ), MissingPacketCount, *Describe());
 			}
 
-			InPacketsLost += PacketsLost;
-			InTotalPacketsLost += PacketsLost;
-			Driver->InPacketsLost += PacketsLost;
-			Driver->InTotalPacketsLost += PacketsLost;
+			InPacketsLost += MissingPacketCount;
+			InTotalPacketsLost += MissingPacketCount;
+			Driver->InPacketsLost += MissingPacketCount;
+			Driver->InTotalPacketsLost += MissingPacketCount;
 			InPacketId += PacketSequenceDelta;
 
 			// Extra information associated with the header
@@ -1657,7 +1738,23 @@ void UNetConnection::ReceivedPacket( FBitReader& Reader )
 		}
 		else
 		{
+			TotalOutOfOrderPackets++;
 			Driver->InOutOfOrderPackets++;
+
+			if (!PacketOrderCache.IsSet() && CVarNetDoPacketOrderCorrection.GetValueOnAnyThread() != 0)
+			{
+				int32 EnableThreshold = CVarNetPacketOrderCorrectionEnableThreshold.GetValueOnAnyThread();
+
+				if (TotalOutOfOrderPackets >= EnableThreshold)
+				{
+					UE_LOG(LogNet, Verbose, TEXT("Hit threshold of %i 'out of order' packet sequences. Enabling out of order packet correction."), EnableThreshold);
+
+					int32 CacheSize = FMath::RoundUpToPowerOfTwo(CVarNetPacketOrderMaxCachedPackets.GetValueOnAnyThread());
+
+					PacketOrderCache.Emplace(CacheSize);
+				}
+			}
+
 			// Protect against replay attacks
 			// We already protect against this for reliable bunches, and unreliable properties
 			// The only bunch we would process would be unreliable RPC's, which could allow for replay attacks
@@ -1665,6 +1762,34 @@ void UNetConnection::ReceivedPacket( FBitReader& Reader )
 			// which protects everything in one fell swoop
 			return;
 		}
+
+
+		// Lambda to dispatch delivery notifications, 
+		auto HandlePacketNotification = [&Header, this](FNetPacketNotify::SequenceNumberT AckedSequence, bool bDelivered)
+		{
+			// Increase LastNotifiedPacketId, this is a full packet Id
+			++LastNotifiedPacketId;
+
+			// Sanity check
+			if (FNetPacketNotify::SequenceNumberT(LastNotifiedPacketId) != AckedSequence)
+			{
+				CLOSE_CONNECTION_DUE_TO_SECURITY_VIOLATION(this, ESecurityEvent::Malformed_Packet, TEXT("LastNotifiedPacketId != AckedSequence"));
+				return;
+			}
+
+			if (bDelivered)
+			{
+				ReceivedAck(LastNotifiedPacketId);
+			}
+			else
+			{
+				ReceivedNak(LastNotifiedPacketId);
+			};
+		};
+
+		// Update incoming sequence data and deliver packet notifications
+		// Packet is only accepted if both the incoming sequence number and incoming ack data are valid
+		PacketNotify.Update(Header, HandlePacketNotification);
 	}
 
 	const bool bIgnoreRPCs = Driver->ShouldIgnoreRPCs();
diff --git a/Engine/Source/Runtime/Engine/Private/NetworkDriver.cpp b/Engine/Source/Runtime/Engine/Private/NetworkDriver.cpp
--- a/Engine/Source/Runtime/Engine/Private/NetworkDriver.cpp
+++ b/Engine/Source/Runtime/Engine/Private/NetworkDriver.cpp
@@ -1693,6 +1693,17 @@ void UNetDriver::TickDispatch( float DeltaTime )
 
 void UNetDriver::PostTickDispatch()
 {
+	// Flush out of order packet caches for connections that did not receive the missing packets during TickDispatch
+	if (ServerConnection != nullptr)
+	{
+		ServerConnection->FlushPacketOrderCache(true);
+	}
+
+	for (UNetConnection* CurConn : ClientConnections)
+	{
+		CurConn->FlushPacketOrderCache(true);
+	}
+
 	if (ReplicationDriver)
 	{
 		ReplicationDriver->PostTickDispatch();
diff --git a/Engine/Source/Runtime/Engine/Public/Net/NetPacketNotify.h b/Engine/Source/Runtime/Engine/Public/Net/NetPacketNotify.h
--- a/Engine/Source/Runtime/Engine/Public/Net/NetPacketNotify.h
+++ b/Engine/Source/Runtime/Engine/Public/Net/NetPacketNotify.h
@@ -71,15 +71,30 @@ public:
 	/** Read header from stream */
 	bool ReadHeader(FNotificationHeader& Data, FBitReader& Reader) const;
 
-	/** Update state of PacketNotification based on received header and invoke packet notifications for received acks.
-		InFunc is a function in the format void)(FNetPacketNotify::SequenceNumberT AckedSequence, bool delivered) to handle packetNotifications.
+	/**
+	 * Gets the delta between the present sequence, and the sequence inside the specified header - if the delta is positive
+	 */
+	SequenceNumberT::DifferenceT GetSequenceDelta(const FNotificationHeader& NotificationData)
+	{
+		if (NotificationData.Seq > InSeq && NotificationData.AckedSeq >= OutAckSeq && OutSeq > NotificationData.AckedSeq)
+		{
+			return SequenceNumberT::Diff(NotificationData.Seq, InSeq);
+		}
+		else
+		{
+			return 0;
+		}
+	}
 
-		Returns the positive delta of the incoming seq number if it is within half the sequence number space
-		Returns 0 if the received sequence number is out of the current window or if the acknowledged seq received by the remote is invalid.
+	/**
+	 * Update state of PacketNotification based on received header and invoke packet notifications for received acks.
+	 *
+	 * @param NotificationData			The header to update from
+	 * @param InFunc					A function in the format (void)(FNetPacketNotify::SequenceNumberT AckedSequence, bool bDelivered) to handle packet notifications.
+	 * @return							The > 0 delta of the incoming seq if within half the seq number space. 0 if the received seq is outside current window ,or the ack seq received is invalid.
 	*/
 	template<class Functor>
 	SequenceNumberT::DifferenceT Update(const FNotificationHeader& NotificationData, Functor&& InFunc);
-
 	/** Get the current SequenceHistory */
 	const SequenceHistoryT& GetInSeqHistory() const { return InSeqHistory; }
 
@@ -138,12 +153,12 @@ private:
 template<class Functor>
 FNetPacketNotify::SequenceNumberT::DifferenceT FNetPacketNotify::Update(const FNotificationHeader& NotificationData, Functor&& InFunc)
 {
-	if (NotificationData.Seq > InSeq && NotificationData.AckedSeq >= OutAckSeq && OutSeq > NotificationData.AckedSeq)
+	const SequenceNumberT::DifferenceT InSeqDelta = GetSequenceDelta(NotificationData);
+
+	if (InSeqDelta > 0)
 	{
 		UE_LOG_PACKET_NOTIFY(TEXT("FNetPacketNotify::Update - Seq %u, InSeq %u"), NotificationData.Seq.Get(), InSeq.Get());
-
-		const SequenceNumberT::DifferenceT InSeqDelta = SequenceNumberT::Diff(NotificationData.Seq, InSeq);
-
+	
 		ProcessReceivedAcks(NotificationData, InFunc);
 
 		// accept sequence
