[NO NON-PII COMMIT MESSAGE]

diff --git a/Engine/Source/Runtime/Engine/Private/AudioThread.cpp b/Engine/Source/Runtime/Engine/Private/AudioThread.cpp
--- a/Engine/Source/Runtime/Engine/Private/AudioThread.cpp
+++ b/Engine/Source/Runtime/Engine/Private/AudioThread.cpp
@@ -433,6 +433,23 @@ void FAudioThread::StopAudioThread()
 	AudioThreadRunnable = nullptr;
 }
 
+FAudioCommandFence::FAudioCommandFence()
+	: FenceDoneEvent(nullptr)
+{
+}
+
+FAudioCommandFence::~FAudioCommandFence()
+{
+	if (FenceDoneEvent)
+	{
+
+		FenceDoneEvent->Wait();
+
+		FPlatformProcess::ReturnSynchEventToPool(FenceDoneEvent);
+		FenceDoneEvent = nullptr;
+	}
+}
+
 void FAudioCommandFence::BeginFence()
 {
 	if (FAudioThread::IsAudioThreadRunning())
@@ -443,6 +460,20 @@ void FAudioCommandFence::BeginFence()
 
 		CompletionEvent = TGraphTask<FNullGraphTask>::CreateTask(GAudioAsyncBatcher.GetAsyncPrereq(), ENamedThreads::GameThread).ConstructAndDispatchWhenReady(
 			GET_STATID(STAT_FNullGraphTask_FenceAudioCommand), ENamedThreads::AudioThread);
+
+		if (FenceDoneEvent)
+		{
+
+			FenceDoneEvent->Wait();
+
+			FPlatformProcess::ReturnSynchEventToPool(FenceDoneEvent);
+			FenceDoneEvent = nullptr;
+		}
+
+		FenceDoneEvent = FPlatformProcess::GetSynchEventFromPool(true);
+
+		FTaskGraphInterface::Get().TriggerEventWhenTaskCompletes(FenceDoneEvent, CompletionEvent, ENamedThreads::GameThread, ENamedThreads::AudioThread);
+
 		FAudioThread::ProcessAllCommands();
 	}
 	else
@@ -460,8 +491,10 @@ bool FAudioCommandFence::IsFenceComplete() const
 		CompletionEvent = nullptr; // this frees the handle for other uses, the NULL state is considered completed
 		return true;
 	}
+
 	check(FAudioThread::IsAudioThreadRunning());
-	return false;
+
+	return FenceDoneEvent->Wait(0);
 }
 
 /**
@@ -470,19 +503,33 @@ bool FAudioCommandFence::IsFenceComplete() const
 void FAudioCommandFence::Wait(bool bProcessGameThreadTasks) const
 {
 	FAudioThread::ProcessAllCommands();
+
+
 	if (!IsFenceComplete()) // this checks the current thread
 	{
-		QUICK_SCOPE_CYCLE_COUNTER(STAT_FAudioCommandFence_Wait);
 		const double StartTime = FPlatformTime::Seconds();
-		FEvent* Event = FPlatformProcess::GetSynchEventFromPool();
-		FTaskGraphInterface::Get().TriggerEventWhenTaskCompletes(Event, CompletionEvent, ENamedThreads::GameThread);
+		QUICK_SCOPE_CYCLE_COUNTER(STAT_FAudioCommandFence_Wait);
 
-		bool bDone;
+		bool bDone = false;
 		const uint32 WaitTimeMs = 35;
 		do
 		{
-			bDone = Event->Wait(WaitTimeMs);
+			if (FenceDoneEvent)
+			{
+				bDone = FenceDoneEvent->Wait(WaitTimeMs);
+			}
+			else
+			{
+				bDone = true;
+			}
+			
+			if(bDone && FenceDoneEvent)
+			{
+				FPlatformProcess::ReturnSynchEventToPool(FenceDoneEvent);
+				FenceDoneEvent = nullptr;
+			}
 
+			// Log how long we've been waiting for the audio thread:
 			float ThisTime = FPlatformTime::Seconds() - StartTime;
  			if (ThisTime > static_cast<float>(WaitTimeMs) / 1000.0f + SMALL_NUMBER)
 			{
@@ -505,9 +552,8 @@ void FAudioCommandFence::Wait(bool bProcessGameThreadTasks) const
 
 		FAudioThread::ResetAudioThreadTimers();
 
-		// Return the event to the pool and decrement the recursion counter.
-		FPlatformProcess::ReturnSynchEventToPool(Event);
-		Event = nullptr;
+
+
 	}
 	
 }
diff --git a/Engine/Source/Runtime/Engine/Public/AudioThread.h b/Engine/Source/Runtime/Engine/Public/AudioThread.h
--- a/Engine/Source/Runtime/Engine/Public/AudioThread.h
+++ b/Engine/Source/Runtime/Engine/Public/AudioThread.h
@@ -125,6 +125,8 @@ struct FAudioThreadSuspendContext
 class ENGINE_API FAudioCommandFence
 {
 public:
+	FAudioCommandFence();
+	~FAudioCommandFence();
 
 	/**
 	* Adds a fence command to the audio command queue.
@@ -145,5 +147,7 @@ public:
 private:
 	/** Graph event that represents completion of this fence **/
 	mutable FGraphEventRef CompletionEvent;
+	/** Event that fires when CompletionEvent is done. **/
+	mutable FEvent* FenceDoneEvent;
 };
 
