Fixing reverb

diff --git a/Engine/Source/Runtime/Engine/Private/SoundWave.cpp b/Engine/Source/Runtime/Engine/Private/SoundWave.cpp
--- a/Engine/Source/Runtime/Engine/Private/SoundWave.cpp
+++ b/Engine/Source/Runtime/Engine/Private/SoundWave.cpp
@@ -1509,9 +1509,6 @@ void USoundWave::Parse( FAudioDevice* AudioDevice, const UPTRINT NodeWaveInstanc
 
 		bool bAlwaysPlay = false;
 
-		// Ensure that a Sound Class's default reverb level is used if we enabled reverb through a sound class and not from the active sound.
-		bool bUseSoundClassDefaultReverb = false;
-
 		// Properties from the sound class
 		WaveInstance->SoundClass = ParseParams.SoundClass;
 		if (ParseParams.SoundClass)
@@ -1534,15 +1531,6 @@ void USoundWave::Parse( FAudioDevice* AudioDevice, const UPTRINT NodeWaveInstanc
 			WaveInstance->bCenterChannelOnly = ActiveSound.bCenterChannelOnly || SoundClassProperties->bCenterChannelOnly;
 			WaveInstance->bEQFilterApplied = ActiveSound.bEQFilterApplied || SoundClassProperties->bApplyEffects;
 			WaveInstance->bReverb = ActiveSound.bReverb || SoundClassProperties->bReverb;
-
-			bUseSoundClassDefaultReverb = SoundClassProperties->bReverb && !ActiveSound.bReverb;
-
-			if (bUseSoundClassDefaultReverb)
-			{
-				WaveInstance->ReverbSendMethod = EReverbSendMethod::Manual;
-				WaveInstance->ManualReverbSendLevel = SoundClassProperties->Default2DReverbSendAmount;
-			}
-
 			WaveInstance->OutputTarget = SoundClassProperties->OutputTarget;
 
 			if (SoundClassProperties->bApplyAmbientVolumes)
@@ -1599,11 +1587,8 @@ void USoundWave::Parse( FAudioDevice* AudioDevice, const UPTRINT NodeWaveInstanc
 		}
 
 		// Copy reverb send settings
-		if (!bUseSoundClassDefaultReverb)
-		{
-			WaveInstance->ReverbSendMethod = ParseParams.ReverbSendMethod;
-			WaveInstance->ManualReverbSendLevel = ParseParams.ManualReverbSendLevel;
-		}
+		WaveInstance->ReverbSendMethod = ParseParams.ReverbSendMethod;
+		WaveInstance->ManualReverbSendLevel = ParseParams.ManualReverbSendLevel;
 		WaveInstance->CustomRevebSendCurve = ParseParams.CustomReverbSendCurve;
 		WaveInstance->ReverbSendLevelRange = ParseParams.ReverbSendLevelRange;
 		WaveInstance->ReverbSendLevelDistanceRange = ParseParams.ReverbSendLevelDistanceRange;
