#UE4 Improving the log output from the ListTimers command and giving TimerManager its own CSV stat to differentiate Tickable objects from timers in csv profiles.

diff --git a/Engine/Source/Runtime/Engine/Private/TimerManager.cpp b/Engine/Source/Runtime/Engine/Private/TimerManager.cpp
--- a/Engine/Source/Runtime/Engine/Private/TimerManager.cpp
+++ b/Engine/Source/Runtime/Engine/Private/TimerManager.cpp
@@ -28,13 +28,14 @@ namespace
 	void DescribeFTimerDataSafely(FOutputDevice& Ar, const FTimerData& Data)
 	{
 		Ar.Logf(
-			TEXT("TimerData %p : bLoop=%s, bRequiresDelegate=%s, Status=%d, Rate=%f, ExpireTime=%f"),
+			TEXT("TimerData %p : bLoop=%s, bRequiresDelegate=%s, Status=%d, Rate=%f, ExpireTime=%f, Delegate=%s"),
 			&Data,
 			Data.bLoop ? TEXT("true") : TEXT("false"),
 			Data.bRequiresDelegate ? TEXT("true") : TEXT("false"),
 			static_cast<int32>(Data.Status),
 			Data.Rate,
-			Data.ExpireTime
+			Data.ExpireTime,
+			*Data.TimerDelegate.ToString()
 		);
 	}
 
@@ -518,7 +519,7 @@ DECLARE_DWORD_COUNTER_STAT(TEXT("TimerManager Heap Size"),STAT_NumHeapEntries,ST
 
 void FTimerManager::Tick(float DeltaTime)
 {
-	CSV_SCOPED_TIMING_STAT_EXCLUSIVE(Tickables);
+	CSV_SCOPED_TIMING_STAT_EXCLUSIVE(TimerManager);
 
 #if DO_TIMEGUARD && 0
 	TArray<FTimerUnifiedDelegate> RunTimerDelegates;
@@ -674,24 +675,22 @@ void FTimerManager::ListTimers() const
 	UE_LOG(LogEngine, Log, TEXT("------- %d Active Timers -------"), ValidActiveTimers.Num());
 	for (const FTimerData* Data : ValidActiveTimers)
 	{
-		FString TimerString = Data->TimerDelegate.ToString();
-		UE_LOG(LogEngine, Log, TEXT("%s"), *TimerString);
+		check(Data);
+		DescribeFTimerDataSafely(*GLog, *Data);
 	}
 
 	UE_LOG(LogEngine, Log, TEXT("------- %d Paused Timers -------"), PausedTimerSet.Num());
 	for (FTimerHandle Handle : PausedTimerSet)
 	{
 		const FTimerData& Data = GetTimer(Handle);
-		FString TimerString = Data.TimerDelegate.ToString();
-		UE_LOG(LogEngine, Log, TEXT("%s"), *TimerString);
+		DescribeFTimerDataSafely(*GLog, Data);
 	}
 
 	UE_LOG(LogEngine, Log, TEXT("------- %d Pending Timers -------"), PendingTimerSet.Num());
 	for (FTimerHandle Handle : PendingTimerSet)
 	{
 		const FTimerData& Data = GetTimer(Handle);
-		FString TimerString = Data.TimerDelegate.ToString();
-		UE_LOG(LogEngine, Log, TEXT("%s"), *TimerString);
+		DescribeFTimerDataSafely(*GLog, Data);
 	}
 
 	UE_LOG(LogEngine, Log, TEXT("------- %d Total Timers -------"), PendingTimerSet.Num() + PausedTimerSet.Num() + ValidActiveTimers.Num());
