Proper DLL export of EnableTracking

diff --git a/Engine/Source/Runtime/Engine/Public/Net/NetworkProfiler.h b/Engine/Source/Runtime/Engine/Public/Net/NetworkProfiler.h
--- a/Engine/Source/Runtime/Engine/Public/Net/NetworkProfiler.h
+++ b/Engine/Source/Runtime/Engine/Public/Net/NetworkProfiler.h
@@ -181,7 +181,7 @@ public:
 	 *
 	 * @param	bShouldEnableTracking	Whether tracking should be enabled or not
 	 */
-	void EnableTracking( bool bShouldEnableTracking );
+	ENGINE_API void EnableTracking( bool bShouldEnableTracking );
 
 	/**
 	 * Marks the beginning of a frame.
