Concurrency Set Updates

diff --git a/Engine/Source/Runtime/Engine/Classes/Sound/SoundBase.h b/Engine/Source/Runtime/Engine/Classes/Sound/SoundBase.h
--- a/Engine/Source/Runtime/Engine/Classes/Sound/SoundBase.h
+++ b/Engine/Source/Runtime/Engine/Classes/Sound/SoundBase.h
@@ -148,6 +148,7 @@ public:
 	//~ Begin UObject Interface.
 	virtual void PostInitProperties() override;
 	virtual void PostLoad() override;
+	virtual void Serialize(FArchive& Ar) override;
 	virtual bool CanBeClusterRoot() const override;
 	virtual bool CanBeInCluster() const override;
 	virtual void Serialize(FArchive& Ar) override;
@@ -221,10 +222,6 @@ public:
 
 	/** Returns the priority to use when evaluating concurrency. */
 	float GetPriority() const;
-
-	/** Returns the sound concurrency object ID if it exists. If it doesn't exist, returns 0. */
-	uint32 GetSoundConcurrencyObjectID() const;
-
 	/** Returns whether the sound has cooked analysis data (e.g. FFT or envelope following data) and returns sound waves which have cooked data. */
 	virtual bool GetSoundWavesWithCookedAnalysisData(TArray<USoundWave*>& OutSoundWaves);
 
diff --git a/Engine/Source/Runtime/Engine/Private/SoundConcurrency.cpp b/Engine/Source/Runtime/Engine/Private/SoundConcurrency.cpp
--- a/Engine/Source/Runtime/Engine/Private/SoundConcurrency.cpp
+++ b/Engine/Source/Runtime/Engine/Private/SoundConcurrency.cpp
@@ -577,9 +577,9 @@ void FSoundConcurrencyManager::StopActiveSound(FActiveSound* ActiveSound)
 		FConcurrencyGroup* ConcurrencyGroup = ConcurrencyGroups.Find(ConcurrencyGroupID);
 		if (!ConcurrencyGroup)
 		{
-			UE_LOG(LogAudio, Fatal, TEXT("Attempting to remove stopped sound '%s' from inactive concurrency group."),
+			UE_LOG(LogAudio, Error, TEXT("Attempting to remove stopped sound '%s' from inactive concurrency group."),
 				ActiveSound->GetSound() ? *ActiveSound->GetSound()->GetName(): TEXT("Unset"));
-			return;
+			continue;
 		}
 
 		check(!ConcurrencyGroup->IsEmpty());
