Fix bad merge in main from Concurrency resolve

diff --git a/Engine/Source/Runtime/Engine/Classes/Sound/SoundBase.h b/Engine/Source/Runtime/Engine/Classes/Sound/SoundBase.h
--- a/Engine/Source/Runtime/Engine/Classes/Sound/SoundBase.h
+++ b/Engine/Source/Runtime/Engine/Classes/Sound/SoundBase.h
@@ -148,7 +148,6 @@ public:
 	//~ Begin UObject Interface.
 	virtual void PostInitProperties() override;
 	virtual void PostLoad() override;
-	virtual void Serialize(FArchive& Ar) override;
 	virtual bool CanBeClusterRoot() const override;
 	virtual bool CanBeInCluster() const override;
 	virtual void Serialize(FArchive& Ar) override;
