Make sure we call PostNetReceive on FastArrayItems when they have newly mapped guids

Reshelved 5351212 for jon.nabozny

diff --git a/Engine/Source/Runtime/Engine/Private/RepLayout.cpp b/Engine/Source/Runtime/Engine/Private/RepLayout.cpp
--- a/Engine/Source/Runtime/Engine/Private/RepLayout.cpp
+++ b/Engine/Source/Runtime/Engine/Private/RepLayout.cpp
@@ -3139,6 +3139,7 @@ void FRepLayout::UpdateUnmappedObjects_r(
 	FRepShadowDataBuffer ShadowData, 
 	FRepObjectDataBuffer Data, 
 	const int32 MaxAbsOffset,
+	bool& bCalledPreNetReceive,
 	bool& bOutSomeObjectsWereMapped,
 	bool& bOutHasMoreUnmapped) const
 {
@@ -3172,7 +3173,7 @@ void FRepLayout::UpdateUnmappedObjects_r(
 
 				const int32 NewMaxOffset = FMath::Min(ShadowArray->Num() * Cmd.ElementSize, Array->Num() * Cmd.ElementSize);
 
-				UpdateUnmappedObjects_r(RepState, GuidReferences.Array, OriginalObject, PackageMap, ShadowArrayData, ArrayData, NewMaxOffset, bOutSomeObjectsWereMapped, bOutHasMoreUnmapped);
+				UpdateUnmappedObjects_r(RepState, GuidReferences.Array, OriginalObject, PackageMap, ShadowArrayData, ArrayData, NewMaxOffset, bCalledPreNetReceive, bOutSomeObjectsWereMapped, bOutHasMoreUnmapped);
 			}
 			else
 			{
@@ -3180,7 +3181,7 @@ void FRepLayout::UpdateUnmappedObjects_r(
 				FRepObjectDataBuffer ArrayData(Array->GetData());
 				const int32 NewMaxOffset = Array->Num() * Cmd.ElementSize;
 
-				UpdateUnmappedObjects_r(RepState, GuidReferences.Array, OriginalObject, PackageMap, nullptr, ArrayData, NewMaxOffset, bOutSomeObjectsWereMapped, bOutHasMoreUnmapped);
+				UpdateUnmappedObjects_r(RepState, GuidReferences.Array, OriginalObject, PackageMap, nullptr, ArrayData, NewMaxOffset, bCalledPreNetReceive, bOutSomeObjectsWereMapped, bOutHasMoreUnmapped);
 			}
 			continue;
 		}
@@ -3219,13 +3220,14 @@ void FRepLayout::UpdateUnmappedObjects_r(
 		// If we resolved some guids, re-deserialize the data which will hook up the object pointer with the property
 		if (bMappedSomeGUIDs)
 		{
-			if (!bOutSomeObjectsWereMapped)
+			if (!bCalledPreNetReceive)
 			{
 				// Call PreNetReceive if we are going to change a value (some game code will need to think this is an actual replicated value)
 				OriginalObject->PreNetReceive();
-				bOutSomeObjectsWereMapped = true;
+				bCalledPreNetReceive = true;
 			}
 
+			bOutSomeObjectsWereMapped = true;
 			const bool bUpdateShadowState = (ShadowData && INDEX_NONE != Parent.RepNotifyNumParams);
 
 			// Copy current value over so we can check to see if it changed
@@ -3282,10 +3284,11 @@ void FRepLayout::UpdateUnmappedObjects(
 {
 	bOutSomeObjectsWereMapped = false;
 	bOutHasMoreUnmapped = false;
+	bool bCalledPreNetReceive = false;
 
 	if (LayoutState == ERepLayoutState::Normal)
 	{
-		UpdateUnmappedObjects_r(RepState, &RepState->GuidReferencesMap, OriginalObject, PackageMap, (uint8*)RepState->StaticBuffer.GetData(), (uint8*)OriginalObject, Owner->GetPropertiesSize(), bOutSomeObjectsWereMapped, bOutHasMoreUnmapped);
+		UpdateUnmappedObjects_r(RepState, &RepState->GuidReferencesMap, OriginalObject, PackageMap, (uint8*)RepState->StaticBuffer.GetData(), (uint8*)OriginalObject, Owner->GetPropertiesSize(), bCalledPreNetReceive, bOutSomeObjectsWereMapped, bOutHasMoreUnmapped);
 	}
 }
 
@@ -6330,16 +6333,22 @@ void FRepLayout::UpdateUnmappedGuidsForFastArray(FFastArrayDeltaSerializeParams&
 
 	FFastArraySerializer& ArraySerializer = Params.ArraySerializer;
 
-	bool bOutSomeObjectsWereMapped = false;
-	bool bOutHasMoreUnmapped = false;
-
+	bool bCalledPreNetReceive = Params.DeltaSerializeInfo.bOutSomeObjectsWereMapped;
 	for (auto& GuidReferencesPair : ArraySerializer.GuidReferencesMap_StructDelta)
 	{
+		bool bOutSomeObjectsWereMapped = false;
+		bool bOutHasMoreUnmapped = false;
+
 		const int32 ItemIndex = ArraySerializer.ItemMap[GuidReferencesPair.Key];
 		const int32 ArrayElementOffset = ItemIndex * ElementSize;
 		FRepObjectDataBuffer ElementData(ArrayData + ArrayElementOffset);
 
-		UpdateUnmappedObjects_r(nullptr, &GuidReferencesPair.Value, Object, PackageMap, nullptr, ElementData, ElementSize, bOutSomeObjectsWereMapped, bOutHasMoreUnmapped);
+		UpdateUnmappedObjects_r(nullptr, &GuidReferencesPair.Value, Object, PackageMap, nullptr, ElementData, ElementSize, bCalledPreNetReceive, bOutSomeObjectsWereMapped, bOutHasMoreUnmapped);
+
+		if (bOutSomeObjectsWereMapped)
+		{
+			((FFastArraySerializerItem*)(ElementData).Data)->PostReplicatedChange(ArraySerializer);
+		}
 
 		DeltaSerializeInfo.bOutHasMoreUnmapped |= bOutHasMoreUnmapped;
 		DeltaSerializeInfo.bOutSomeObjectsWereMapped |= bOutSomeObjectsWereMapped;
diff --git a/Engine/Source/Runtime/Engine/Public/Net/RepLayout.h b/Engine/Source/Runtime/Engine/Public/Net/RepLayout.h
--- a/Engine/Source/Runtime/Engine/Public/Net/RepLayout.h
+++ b/Engine/Source/Runtime/Engine/Public/Net/RepLayout.h
@@ -1713,6 +1713,7 @@ private:
 		FRepShadowDataBuffer ShadowData, 
 		FRepObjectDataBuffer Data, 
 		const int32 MaxAbsOffset,
+		bool& bCalledPreNetReceive,
 		bool& bOutSomeObjectsWereMapped,
 		bool& bOutHasMoreUnmapped) const;
 
