Fix for Fast Array Delta Serialize not updating properly in some cases.

Reshelved 5347611 for jon.nabozny

diff --git a/Engine/Source/Runtime/Engine/Classes/Engine/NetSerialization.h b/Engine/Source/Runtime/Engine/Classes/Engine/NetSerialization.h
--- a/Engine/Source/Runtime/Engine/Classes/Engine/NetSerialization.h
+++ b/Engine/Source/Runtime/Engine/Classes/Engine/NetSerialization.h
@@ -271,7 +271,7 @@ public:
 
 	FNetFastTArrayBaseState():
 		ArrayReplicationKey(INDEX_NONE),
-		ChangelistHistory(INDEX_NONE)
+		ChangelistHistory(0)
 	{}
 
 	virtual bool IsStateEqual(INetDeltaBaseState* OtherState)
@@ -293,7 +293,7 @@ public:
 
 	int32 ArrayReplicationKey;
 
-	int32 ChangelistHistory;
+	uint32 ChangelistHistory;
 };
 
 
@@ -1421,6 +1421,9 @@ struct FFastArrayDeltaSerializeParams
 {
 	FNetDeltaSerializeInfo& DeltaSerializeInfo;
 	FFastArraySerializer& ArraySerializer;
+	const PTRINT PatchedItemOffset;
+	const PTRINT PatchedArrayOffset;
+
 	TArray<FFastArraySerializer_FastArrayDeltaSerialize_FIdxIDPair, TInlineAllocator<8>>* WriteChangedElements = nullptr;
 	FNetFastTArrayBaseState* WriteBaseState = nullptr;
 	TArray<int32, TInlineAllocator<8>>* ReadChangedElements = nullptr;
@@ -1441,11 +1444,17 @@ bool FFastArraySerializer::FastArrayDeltaSerialize_DeltaSerializeStructs(TArray<
 		Parms
 	};
 
+#define CALC_PATCH_OFFSET(Base, Derived) ((PTRINT)((Base*)((Derived*)1)) - 1)
+
 	FFastArrayDeltaSerializeParams DeltaSerializeParams{
 		Parms,
-		ArraySerializer
+		ArraySerializer,
+		/*PatchedItemOffset=*/ CALC_PATCH_OFFSET(FFastArraySerializerItem, Type),
+		/*PatchedArrayOffset*/ CALC_PATCH_OFFSET(FFastArraySerializer, SerializerType)
 	};
 
+#undef CALC_PATCH_OFFSET
+
 	//---------------
 	// Build ItemMap if necessary. This maps ReplicationID to our local index into the Items array.
 	//---------------
diff --git a/Engine/Source/Runtime/Engine/Private/RepLayout.cpp b/Engine/Source/Runtime/Engine/Private/RepLayout.cpp
--- a/Engine/Source/Runtime/Engine/Private/RepLayout.cpp
+++ b/Engine/Source/Runtime/Engine/Private/RepLayout.cpp
@@ -21,6 +21,7 @@
 #include "Net/NetworkGranularMemoryLogging.h"
 #include "Serialization/ArchiveCountMem.h"
 #include "Containers/SortedMap.h"
+#include "Templates/AndOrNot.h"
 
 DECLARE_CYCLE_STAT(TEXT("RepLayout AddPropertyCmd"), STAT_RepLayout_AddPropertyCmd, STATGROUP_Game);
 DECLARE_CYCLE_STAT(TEXT("RepLayout InitFromObjectClass"), STAT_RepLayout_InitFromObjectClass, STATGROUP_Game);
@@ -121,6 +122,48 @@ FConsoleVariableSinkHandle CreateMaxArrayMemoryCVarAndRegisterSink()
 FConsoleVariableSinkHandle MaxRepArraySizeHandle = CreateMaxArraySizeCVarAndRegisterSink();
 FConsoleVariableSinkHandle MaxRepArrayMemorySink = CreateMaxArrayMemoryCVarAndRegisterSink();
 
+namespace UE4_RepLayout_Private
+{
+	template<typename OutputType, typename CommandType, typename BufferType>
+	static typename TTranslateConstVolatile<BufferType, OutputType>::Type*
+	GetTypedProperty(const BufferType& Buffer, const CommandType& Cmd)
+	{
+		using ConstOrNotOutputType = typename TTranslateConstVolatile<BufferType, OutputType>::Type;
+
+		using BaseBufferType = typename TRemovePointer<typename TDecay<BufferType>::Type>::Type;
+
+		static_assert(!TIsPointer<OutputType>::Value, "GetTypedProperty invalid OutputType!  Don't specify output as a pointer.");
+		static_assert(TOr<TAreTypesEqual<uint8, BaseBufferType>, TAreTypesEqual<void, BaseBufferType>>::Value, "GetTypedProperty invalid BufferType! Only TRepDataBufferBase, void*, and uint8* are supported!");
+
+		// TODO: Conditionally compilable runtime type validation.
+		return reinterpret_cast<ConstOrNotOutputType *>(Buffer);
+	}
+
+	template<typename OutputType, typename CommandType, enum ERepDataBufferType BufferDataType, typename BufferUnderlyingType>
+	static typename TTranslateConstVolatile<typename TRepDataBufferBase<BufferDataType, BufferUnderlyingType>::ConstOrNotVoid, OutputType>::Type*
+	GetTypedProperty(const TRepDataBufferBase<BufferDataType, BufferUnderlyingType>& Buffer, const CommandType& Cmd)
+	{
+		using ConstOrNotOutputType = typename TTranslateConstVolatile<typename TRepDataBufferBase<BufferDataType, BufferUnderlyingType>::ConstOrNotVoid, OutputType>::Type;
+
+		static_assert(!TIsPointer<OutputType>::Value, "GetTypedProperty invalid OutputType! Don't specify output as a pointer.");
+
+		// TODO: Conditionally compilable runtime type validation.
+		return reinterpret_cast<ConstOrNotOutputType *>((Buffer + Cmd).Data);
+	}
+
+	template<typename CommandType, typename BufferType>
+	FORCEINLINE static FFastArraySerializerItem* GetFastArrayItem(FFastArrayDeltaSerializeParams& Params, const BufferType& Buffer, const CommandType& Cmd)
+	{
+		return (FFastArraySerializerItem*)(Params.PatchedItemOffset + (PTRINT)(GetTypedProperty<FFastArraySerializerItem>(Buffer, Cmd)));
+	}
+
+	template<typename CommandType, typename BufferType>
+	FORCEINLINE static FFastArraySerializer* GetFastArray(FFastArrayDeltaSerializeParams& Params, const BufferType& Buffer, const CommandType& Cmd)
+	{
+		return (FFastArraySerializer*)(Params.PatchedArrayOffset + (PTRINT)(GetTypedProperty<FFastArraySerializer>(Buffer, Cmd)));
+	}
+}
+
 struct FLifetimeCustomDeltaProperty
 {
 	/** If this is a Fast Array Serializer property, this will be the command index for the Fast Array Item array. */
@@ -140,7 +183,8 @@ struct FLifetimeCustomDeltaProperty
 struct FLifetimeCustomDeltaState
 {
 	TSortedMap<uint16, FLifetimeCustomDeltaProperty> LifetimeCustomDeltaProperties;
-
+	
+	/** The number of valid FFastArraySerializer properties we found. */
 	int32 NumFastArrayItems = 0;
 
 	/**
@@ -162,35 +206,92 @@ struct FLifetimeCustomDeltaState
 //~ instead of just relying on the standard serialization stuff. That would be a bigger backwards compatibility
 //~ risk, however.
 
-struct FCustomDeltaSinglePropertyHistory
+struct FDeltaArrayHistoryItem
 {
+	/** The set of changelists by element ID.*/
 	TMap<int32, TArray<uint16>> ChangelistByID;
-};
 
-struct FCustomDeltaChangedHistory
-{
-	// We could just use a Sparse array here, but in this case it's more convenient
-	// to work with a bitfield and array directly.
+	bool bWasUpdated = false;
+
+	void Reset()
+	{
+		ChangelistByID.Empty();
+		bWasUpdated = false;
+	}
+
+	void CountBytes(FArchive& Ar) const
+	{
+		ChangelistByID.CountBytes(Ar);
 
-	TArray<FCustomDeltaSinglePropertyHistory> FastArrayChangelists;
-	TBitArray<> bWasUpdated;
+		for (const auto& KVP : ChangelistByID)
+		{
+			KVP.Value.CountBytes(Ar);
+		}
+	}
 };
 
-struct FCustomDeltaChangelistState
+struct FDeltaArrayHistoryState
 {
 	/** The maximum number of individual changelists allowed.*/
-	static const int32 MAX_CHANGE_HISTORY = FRepChangelistState::MAX_CHANGE_HISTORY;
+	static const uint32 MAX_CHANGE_HISTORY = FRepChangelistState::MAX_CHANGE_HISTORY;
+
+	//~ TODO: Investigate either making this a Dynamically sized container,
+	//~			or potentially changing ArrayStates to unique pointers / 
+	//~			an indirect array.
+	//~
+	//~			Right now, if a Delta Struct Serialization is disabled for
+	//~			a given Fast Array, this will still use up memory.
 
 	/** Circular buffer of changelists. */
-	FCustomDeltaChangedHistory ChangeHistory[MAX_CHANGE_HISTORY];
+	FDeltaArrayHistoryItem ChangeHistory[MAX_CHANGE_HISTORY];
+
+	/** The latest ArrayReplicationKey sent to any connection. */
+	int32 ArrayReplicationKey = INDEX_NONE;
+
+	/** Copy of the IDToIndexMap from the array when we last sent it. */
+	TMap<int32, int32> IDToIndexMap;
 
 	/** Index in the buffer where changelist history starts (i.e., the Oldest changelist). */
-	int32 HistoryStart = 0;
+	uint32 HistoryStart = 0;
 
 	/** Index in the buffer where changelist history ends (i.e., the Newest changelist). */
-	int32 HistoryEnd = 0;
+	uint32 HistoryEnd = 0;
+
+	void CountBytes(FArchive& Ar) const
+	{
+		IDToIndexMap.CountBytes(Ar);
+
+		for (const FDeltaArrayHistoryItem& HistoryItem : ChangeHistory)
+		{
+			HistoryItem.CountBytes(Ar);
+		}
+	}
+};
 
+struct FCustomDeltaChangelistState
+{
+	FCustomDeltaChangelistState(const int32 NumArrays)
+	{
+		ArrayStates.SetNum(NumArrays);
+	}
+
+	/** Index used to determine whether or not we've compared already on a given frame. */
 	uint32 CompareIndex = 0;
+
+	/**
+	 * An array tracking the last compared history of Arrays.
+	 * Indices should match FLifetimeCustomDeltaProperty::FastArrayNumber.
+	 */
+	TArray<FDeltaArrayHistoryState> ArrayStates;
+
+	void CountBytes(FArchive& Ar) const
+	{
+		ArrayStates.CountBytes(Ar);
+		for (const FDeltaArrayHistoryState& ArrayState : ArrayStates)
+		{
+			ArrayState.CountBytes(Ar);
+		}
+	}
 };
 
 #define ENABLE_PROPERTY_CHECKSUMS
@@ -549,6 +650,12 @@ void FRepChangelistState::CountBytes(FArchive& Ar) const
 {
 	StaticBuffer.CountBytes(Ar);
 	SharedSerialization.CountBytes(Ar);
+
+	if (CustomDeltaChangelistState)
+	{
+		Ar.CountBytes(sizeof(FCustomDeltaChangelistState), sizeof(FCustomDeltaChangelistState));
+		CustomDeltaChangelistState->CountBytes(Ar);
+	}
 }
 
 FReplicationChangelistMgr::FReplicationChangelistMgr(
@@ -5543,7 +5650,7 @@ TSharedPtr<FReplicationChangelistMgr> FRepLayout::CreateReplicationChangelistMgr
 	FCustomDeltaChangelistState* DeltaChangelistState = nullptr;
 	if (LifetimeCustomPropertyState && !!LifetimeCustomPropertyState->NumFastArrayItems)
 	{
-		DeltaChangelistState = new FCustomDeltaChangelistState();
+		DeltaChangelistState = new FCustomDeltaChangelistState(LifetimeCustomPropertyState->NumFastArrayItems);
 	}
 
 	return MakeShareable(new FReplicationChangelistMgr(AsShared(), (const uint8*)InObject->GetArchetype(), DeltaChangelistState));
@@ -5749,6 +5856,8 @@ void FRepLayout::PreSendCustomDeltaProperties(
 	FReplicationChangelistMgr& ChangelistMgr,
 	TMap<int32, TSharedPtr<INetDeltaBaseState>>& CustomDeltaStates) const
 {
+	using namespace UE4_RepLayout_Private;
+
 	if (!Connection->InternalAck)
 	{
 		const FLifetimeCustomDeltaState& LocalLifetimeCustomPropertyState = *LifetimeCustomPropertyState;
@@ -5758,63 +5867,67 @@ void FRepLayout::PreSendCustomDeltaProperties(
 			FRepChangelistState& ChangelistState = *ChangelistMgr.GetRepChangelistState();
 			FCustomDeltaChangelistState& CustomDeltaChangelistState = *ChangelistState.CustomDeltaChangelistState;
 
+			// Check to see whether or not we need to do comparisons this frame.
+			// If we do, then run through our fast array states and generate new history items if needed.
 			if (CustomDeltaChangelistState.CompareIndex != static_cast<uint32>(GFrameCounter))
 			{
 				const bool bIsInitial = (CustomDeltaChangelistState.CompareIndex == 0);
 				CustomDeltaChangelistState.CompareIndex = GFrameCounter;
 
-				// Setup new buffers.
-				++CustomDeltaChangelistState.HistoryEnd;
+				const FConstRepObjectDataBuffer ObjectData(Object);
 
-				if ((CustomDeltaChangelistState.HistoryEnd - CustomDeltaChangelistState.HistoryStart) == FCustomDeltaChangelistState::MAX_CHANGE_HISTORY)
+				for (const auto& KVP : LocalLifetimeCustomPropertyState.LifetimeCustomDeltaProperties)
 				{
-					++CustomDeltaChangelistState.HistoryStart;
-				}
-
-				const int32 NewHistoryIndex = CustomDeltaChangelistState.HistoryEnd % FCustomDeltaChangelistState::MAX_CHANGE_HISTORY;
-				const int32 StartHistoryIndex = CustomDeltaChangelistState.HistoryStart % FCustomDeltaChangelistState::MAX_CHANGE_HISTORY;
+					const uint16 RepIndex = KVP.Key;
+					const FLifetimeCustomDeltaProperty& CustomDeltaProperty = KVP.Value;
 
-				FCustomDeltaChangedHistory& NewHistoryItem = CustomDeltaChangelistState.ChangeHistory[NewHistoryIndex];
-				FCustomDeltaChangedHistory& StartHistoryItem = CustomDeltaChangelistState.ChangeHistory[StartHistoryIndex];
-
-				if (!bIsInitial)
-				{
-					bool bUpdatedAnything = false;
-					const FConstRepObjectDataBuffer ObjectData(Object);
-
-					for (const auto& KVP : LocalLifetimeCustomPropertyState.LifetimeCustomDeltaProperties)
+					// If our Fast Array Items Command is invalid, we can't do anything.
+					// This should have been logged on RepLayout creation.
+					if (INDEX_NONE != CustomDeltaProperty.FastArrayItemsCommand)
 					{
-						const uint16 RepIndex = KVP.Key;
-						const FLifetimeCustomDeltaProperty& CustomDeltaProperty = KVP.Value;
-
-						if (INDEX_NONE != CustomDeltaProperty.FastArrayItemsCommand)
+						const int32 FastArrayNumber = CustomDeltaProperty.FastArrayNumber;
+						const FRepParentCmd& FastArrayCmd = Parents[RepIndex];
+						const FFastArraySerializer& FastArray = *GetTypedProperty<FFastArraySerializer>(ObjectData, FastArrayCmd);
+
+						// Note, we can't rely on FFastArraySerializer::HasBeenSerialized here.
+						// It's possible we're calling PreSendCustomDeltaProperties **before** the first time the struct
+						// was ever serialized, and in that case it would still be false, and prevent us from creating
+						// a history the first time.
+						//
+						// This does mean that Fast Arrays *not* using delta serialization will still have their history
+						// incremented the first time, but that can be fixed in other ways and the history memory
+						// will have already been allocated elsewhere.
+						if (FastArray.IsUsingStructDeltaSerialization())
 						{
-							const FRepParentCmd& FastArrayCmd = Parents[RepIndex];
-							const FFastArraySerializer& FastArray = *reinterpret_cast<const FFastArraySerializer*>((ObjectData + FastArrayCmd).Data);
+							FDeltaArrayHistoryState& FastArrayHistoryState = CustomDeltaChangelistState.ArrayStates[FastArrayNumber];
 
-							if (FastArray.HasBeenSerialized() && FastArray.IsUsingStructDeltaSerialization())
+							if (FastArrayHistoryState.ArrayReplicationKey != FastArray.ArrayReplicationKey)
 							{
-								const TSharedPtr<INetDeltaBaseState>* OldState = CustomDeltaStates.Find(RepIndex);
-								const FNetFastTArrayBaseState* FastArrayBaseState = static_cast<const FNetFastTArrayBaseState*>(OldState ? OldState->Get() : nullptr);
-
-								if (!FastArrayBaseState || FastArray.ArrayReplicationKey != FastArrayBaseState->ArrayReplicationKey)
+								const uint32 HistoryDelta = FastArrayHistoryState.HistoryEnd - FastArrayHistoryState.HistoryStart;
+								const uint32 CurrentHistoryIndex = FastArrayHistoryState.HistoryEnd % FDeltaArrayHistoryState::MAX_CHANGE_HISTORY;
+								const FDeltaArrayHistoryItem& CurrentHistory = FastArrayHistoryState.ChangeHistory[CurrentHistoryIndex];
+
+								// If we don't have any history items, go ahead and create one.
+								// Otherwise, check to see if our current history was actually updated.
+								// If it wasn't updated, that means that no one tried to replicate it last frame (which can be possible due
+								// to rep conditions), and there's no sense in creating a new one.
+								if (HistoryDelta == 0 || CurrentHistory.bWasUpdated)
 								{
-									bUpdatedAnything = true;
+									// If we've reached our buffer size, then move our start history marker up.
+									// In that case the old start history will become our new history.
+									if (HistoryDelta >= FDeltaArrayHistoryState::MAX_CHANGE_HISTORY)
+									{
+										++FastArrayHistoryState.HistoryStart;
+									}
+
+									++FastArrayHistoryState.HistoryEnd;
+									const uint32 NewHistory = FastArrayHistoryState.HistoryEnd % FDeltaArrayHistoryState::MAX_CHANGE_HISTORY;
+									FastArrayHistoryState.ChangeHistory[NewHistory].Reset();
 								}
 							}
 						}
 					}
-
-					// If none of our structs were dirty, then revert our history, unless we haven't created one yet.
-					// This is because we won't know what items are going to use delta until at least the first replication.
-					if (!bUpdatedAnything)
-					{
-						--CustomDeltaChangelistState.HistoryEnd;
-					}
 				}
-
-				NewHistoryItem.FastArrayChangelists.SetNum(LocalLifetimeCustomPropertyState.NumFastArrayItems);
-				NewHistoryItem.bWasUpdated.Init(false, LocalLifetimeCustomPropertyState.NumFastArrayItems);
 			}
 		}
 	}
@@ -5830,6 +5943,8 @@ void FRepLayout::PostSendCustomDeltaProperties(
 
 bool FRepLayout::DeltaSerializeFastArrayProperty(FFastArrayDeltaSerializeParams& Params, FReplicationChangelistMgr* ChangelistMgr) const
 {
+	using namespace UE4_RepLayout_Private;
+
 	SCOPE_CYCLE_COUNTER(STAT_RepLayout_DeltaSerializeFastArray);
 
 	// A portion of this work could be shared across all Fast Array Properties for a given object,
@@ -5858,15 +5973,19 @@ bool FRepLayout::DeltaSerializeFastArrayProperty(FFastArrayDeltaSerializeParams&
 	const int32 ItemLayoutEnd = FastArrayItemCmd.EndCmd - 1;
 
 	UObject* Object = DeltaSerializeInfo.Object;
-	UPackageMapClient* PackageMap = (UPackageMapClient*)DeltaSerializeInfo.Map;
+	UPackageMapClient* PackageMap = static_cast<UPackageMapClient*>(DeltaSerializeInfo.Map);
 	UNetConnection* Connection = DeltaSerializeInfo.Connection;
 	const bool bIsWriting = !!DeltaSerializeInfo.Writer;
 	const bool bInternalAck = !!Connection->InternalAck;
 
 	FRepObjectDataBuffer ObjectData(Object);
-	FScriptArray* ObjectArray = (FScriptArray*)(ObjectData + FastArrayItemCmd).Data;
+	FScriptArray* ObjectArray = GetTypedProperty<FScriptArray>(ObjectData, FastArrayItemCmd);
 	FRepObjectDataBuffer ObjectArrayData(ObjectArray->GetData());
 
+	FFastArraySerializer& ArraySerializer = *GetFastArray(Params, ObjectData, Parent);
+
+	check(&ArraySerializer == &Params.ArraySerializer);
+
 	const int32 ObjectArrayNum = ObjectArray->Num();
 	FNetFieldExportGroup* NetFieldExportGroup = nullptr;
 
@@ -5874,6 +5993,8 @@ bool FRepLayout::DeltaSerializeFastArrayProperty(FFastArrayDeltaSerializeParams&
 	{
 		// Note, PackageMap should hold onto the strong reference for us, so we use raw pointers where
 		// we can.
+
+		// TODO: This feels like something we could cache in PreSend, but we'd need to add plumbing to hold onto it.
 		const FString OwnerPathName = Owner->GetPathName();
 		TSharedPtr<FNetFieldExportGroup> LocalNetFieldExportGroup = PackageMap->GetNetFieldExportGroup(OwnerPathName);
 
@@ -5894,11 +6015,9 @@ bool FRepLayout::DeltaSerializeFastArrayProperty(FFastArrayDeltaSerializeParams&
 		
 		// This is a list of changelists to send, corresponding to items in ChangedElements.
 		TArray<TArray<uint16>> Changelists;
-		
 
 		const FHandleToCmdIndex& ArrayHandleToCmd = BaseHandleToCmdIndex[FastArrayItemCmd.RelativeHandle - 1];
 		const TArray<FHandleToCmdIndex>& ArrayHandleToCmdIndex = *ArrayHandleToCmd.HandleToCmdIndex;
-		const int32 HandlePerElement = ArrayHandleToCmdIndex.Num();
 
 		// We only need to create changelists if we're not in a demo.
 		// Note, we may change this in the future to also delta structs for replays.
@@ -5910,69 +6029,175 @@ bool FRepLayout::DeltaSerializeFastArrayProperty(FFastArrayDeltaSerializeParams&
 				FRepChangelistState& RepChangelistState = *ChangelistMgr->GetRepChangelistState();
 				FCustomDeltaChangelistState& DeltaChangelistState = *RepChangelistState.CustomDeltaChangelistState;
 
-				int32& NewChangelistHistory = Params.WriteBaseState->ChangelistHistory;
-				const int32 LastSentHistory = NewChangelistHistory;
-				NewChangelistHistory = DeltaChangelistState.HistoryEnd;
+				const int32 FastArrayNumber = CustomDeltaProperty.FastArrayNumber;
+				FDeltaArrayHistoryState& FastArrayState = DeltaChangelistState.ArrayStates[FastArrayNumber];
 
-				FRepShadowDataBuffer ShadowData(RepChangelistState.StaticBuffer.GetData());
-				FScriptArray* ShadowArray = (FScriptArray*)(ShadowData + FastArrayItemCmd).Data;
-				FRepShadowDataBuffer ShadowArrayData(ShadowArray->GetData());
 
-				const int32 FastArrayNumber = CustomDeltaProperty.FastArrayNumber;
+				// Params.WriteBaseState should be valid, and have the most up to date IDToChangelist map for the Fast Array.
+				// However, it's ChangelistHistory will be to the last History Number sent to the Fast TArray on the specific
+				// connection we're replicating from.
+				TSharedRef<FNetFastTArrayBaseState> NewArrayDeltaState = StaticCastSharedRef<FNetFastTArrayBaseState>(Params.WriteBaseState->AsShared());
 
-				// Check to see whether or not we need to update the global changelist shared between connections.
+				// Cache off the newest history, the last history sent to this connection, and then update the state
+				// to notify that we're going to send it the newest history.
+				const uint32 NewChangelistHistory = FastArrayState.HistoryEnd;
+				const uint32 LastSentHistory = NewArrayDeltaState->ChangelistHistory;
+				const uint32 LastSentChangelistDelta = NewChangelistHistory - LastSentHistory;
+
+				NewArrayDeltaState->ChangelistHistory = NewChangelistHistory;
+
+				// Cache off the shadow array buffers.
+				FRepShadowDataBuffer ShadowData(RepChangelistState.StaticBuffer.GetData());
+				FScriptArray* ShadowArray = GetTypedProperty<FScriptArray>(ShadowData, FastArrayItemCmd);
+				FRepShadowDataBuffer ShadowArrayData(ShadowArray->GetData());
 
 				// Note, we explicitly pass in a 0 handles everywhere below.
 				// This is because each item will be received individually, the receiving side indices won't
 				// necessarily match our indices, and we already track the changelists separately.
+
+				// Check to see whether or not we need to update the global changelist shared between connections.
 				{
-					const int32 RelativeNewHistory = NewChangelistHistory % FCustomDeltaChangelistState::MAX_CHANGE_HISTORY;
-					FCustomDeltaChangedHistory& HistoryItem = DeltaChangelistState.ChangeHistory[RelativeNewHistory];
+					const uint32 RelativeNewHistory = NewChangelistHistory % FDeltaArrayHistoryState::MAX_CHANGE_HISTORY;
+					const uint32 CompareChangelistDelta = NewChangelistHistory - FastArrayState.HistoryStart;
+					FDeltaArrayHistoryItem& HistoryItem = FastArrayState.ChangeHistory[RelativeNewHistory];
 
-					FBitReference HistoryItemWasUpdated = HistoryItem.bWasUpdated[FastArrayNumber];
-					if (!HistoryItemWasUpdated)
+					if (!HistoryItem.bWasUpdated)
 					{
-						HistoryItemWasUpdated = true;
+						HistoryItem.bWasUpdated = true;
 
 						// Update our shadow array, and reset our pointer in case we reallocated.
-						const int32 ShadowArrayNum = ShadowArray->Num();
-						FScriptArrayHelper Helper((UArrayProperty*)FastArrayItemCmd.Property, ShadowArray);
-						Helper.Resize(ObjectArrayNum);
-						ShadowArrayData = ShadowArray->GetData();
+						FScriptArrayHelper ShadowArrayHelper((UArrayProperty*)FastArrayItemCmd.Property, ShadowArray);
 
-						TMap<int32, TArray<uint16>>& HistoryChangelists = HistoryItem.FastArrayChangelists[CustomDeltaProperty.FastArrayNumber].ChangelistByID;
+						TBitArray<> ShadowArrayItemIsNew(false, ObjectArrayNum);
+						const bool bIsInitial = (FastArrayState.HistoryEnd - FastArrayState.HistoryStart) == 1;
 
-						// Right now, we rely on the Changed Elements passed in to create our new Changelist.
-						// Unfortunately, it's possible that this particular sender was a couple of History Items behind us.
-						// This isn't particularly ideal, and means that we could end up with multiple histories with redundant
-						// changelists.
+						// It's possible that elements have been deleted or otherwise reordered, and our shadow state is out of date.
+						// In order to prevent issues, we'll shuffle our shadow state back to the correct order.
+						// Note, we can't just do a lookup in both maps directly below, because we might end up stomping the same
+						// shadow state multiple times.
 						//
-						// That means more CPU spent merging them, and more bandwidth sending those redundant items.
-						// To get around that, we could keep a FNetFastTArrayBaseState in our changelist history.
-						// That would mean copying the logic in FFastArraySerializer::TFastArraySerializeHelper::BuildChangedAndDeletedBuffers
-						// and making it generic to FFastArraySerializer and FFastArraySerializerItem as opposed to templated.
+						//	Conceptually, you can imagine this process as having two lines of elements: the Object Line and the Shadow Line.
+						//	The Object Line is always considered authoritative, and we just need to make sure the Shadow Line matches that.
 						//
-						// Note, we may need to make that change anyway, because there might be a pathological case with this.
+						//	If the Shadow Line is empty, the only thing we need to do is add a matching number of elements
+						//	and patch their IDs to match those of the Object Line.
 						//
-						// What happens if we reorder items in the array, or remove an entry in the middle ?
-						// In that case, right now, we have no way to tell this, and so we may end comparing shadow states
-						// from mismatched items, generating full changelists.
+						//	If the Object Line is empty, the only thing we need to do is remove all elements from the shadow line.
 						//
-						// That would basically devolve us back to the old behavior of sending the full reordered items
-						// at least once. After that, we'd be back to proper diffing.
+						//	If both lines are non empty, it becomes a fairly straightforward fixup, following these rules.
+						//	It's important to note, that at each step all previously seen elements are guaranteed to have been
+						//	validated, and are generally considered "out of play" or no longer in the line, and because of
+						//	that we should never touch them again.
+						//
+						//	The process stops when either we run out of Shadow Line or Object Line elements.
+						//
+						//		1. If Elements at the front of the line have matching IDs, there's nothing that needs to be done
+						//			and we can move onto the next Element in both lines.
+						//
+						//		2. If Elements at the front of the line have mismatched IDs, then it's either because the Object Element
+						//			was reordered, a previous Object Element was deleted, or a new Object Element was added.
+						//
+						//			a. If we find the Shadow Element in the Shadow Line, it must have been reordered.
+						//				Go ahead and swap the current front of the Shadow Line with the found Shadow Element.
+						//				Now, the elements at the front of the line have mathcing IDs, and we can move onto
+						//				the next Element in both lines.
+						//
+						//			b. If we don't find the Shadow Element in the Shadow Line, the Object Element must be new.
+						//				Go ahead and insert a new item into the Shadow array, fix up its ID, and mark it as new.
+						//				Now, the elements at the front of the line have matching IDs, and we can move onto
+						//				the next Element in both lines.
+						//
+						//		3. If there are elements remaining in both Lines, go back to step 1. Otherwise, continue to step 4.
+						//
+						//		4. At this point, there should be 3 possible outcomes:
+						//
+						//			a. The Object Line and the Shadow Line have the same number of Elements, and all are matching. We're done.
+						//
+						//			b. The Object Line has more elements than the Shadow Line. All missing elements from the Shadow Line
+						//				must be new elements.
+						//
+						//			c. The Shadow Line has more elements than the Object Line. All missing elements from the Object Line
+						//				must have been removed.
 
-						const bool bIsInitial = DeltaChangelistState.HistoryStart == (DeltaChangelistState.HistoryEnd - 1);
-						for (auto& IDIndexPair : ChangedElements)
+						// TODO: Optimize this. Luckily, it only happens once per frame, and only if the array is dirty.
+						//			Maybe this could be merged into the sending code below, the only concern is that
+						//			doesn't tracked deleted elements.
+
+						// Note, this code serves a very similar purpose to FFastArraySerializer::TFastArraySerializeHelper<Type, SerializerType>::BuildChangedAndDeletedBuffers.
+						// The main issue is that we can't rely on that method, because it will be comparing the last state that was replicated to a given particular connection,
+						// and we want to compare the last state that was replicated to *any* connection.
+
+						if (ObjectArrayNum != 0 && ShadowArrayHelper.Num() != 0 && !bIsInitial)
+						{
+							FScriptArrayHelper ObjectArrayHelper((UArrayProperty*)FastArrayItemCmd.Property, ObjectArray);
+
+							TMap<int32, int32>& ShadowIDToIndex = FastArrayState.IDToIndexMap;
+							const TMap<int32, int32>& ObjectIDToIndex = ArraySerializer.ItemMap;
+
+							for (int32 Index = 0; Index < ObjectArrayNum; ++Index)
+							{
+								// Anything else must be a newly appended element, so we are done.
+								if (Index >= ShadowArrayHelper.Num())
+								{
+									break;
+								}
+
+								FFastArraySerializerItem* ObjectItem = GetFastArrayItem(Params, ObjectArrayHelper.GetRawPtr(Index), FastArrayItemCmd);
+								FFastArraySerializerItem* ShadowItem = GetFastArrayItem(Params, ShadowArrayHelper.GetRawPtr(Index), FastArrayItemCmd);
+
+								if (ObjectItem->ReplicationID != ShadowItem->ReplicationID)
+								{
+									// The IDs didn't match, so this is an insert, delete, or swap.
+									if (int32* FoundShadowIndex = ShadowIDToIndex.Find(ShadowItem->ReplicationID))
+									{
+										// We found the element in the shadow array, so there must have been a swap.
+										// Sanity check that the invalid element can only possibly be later in our lines.
+										check(*FoundShadowIndex > Index);
+
+										ShadowArrayHelper.SwapValues(Index, *FoundShadowIndex);
+									}
+									else
+									{
+										// This item must have been inserted into the array (or appended and then shuffled in).
+										// So, insert it into our shadow array and update its ID.
+
+										ShadowArrayItemIsNew[Index] = true;
+										ShadowArrayHelper.InsertValues(Index);
+
+										ShadowItem = GetFastArrayItem(Params, ShadowArrayHelper.GetRawPtr(Index), FastArrayItemCmd);
+										ShadowItem->ReplicationID = ObjectItem->ReplicationID;
+									}
+								}
+							}
+						}
+
+						// Now we can go ahead and resize the array, to make any other changes we need.
+						ShadowArrayHelper.Resize(ObjectArrayNum);
+						ShadowArrayData = ShadowArray->GetData();
+
+						// Go ahead and fix up IDs for any elements that may have just been appended.
+						// Note, we need to this for all elements on the initial pass.
+						// Deleted elements will have been chopped off by the resize.
+						if (bIsInitial || (ShadowArrayHelper.Num() < ObjectArrayNum))
 						{
-							// If this is a new item, just send everything.
-							const bool bForceFail = ShadowArrayNum <= IDIndexPair.Idx;
+							FScriptArrayHelper ObjectArrayHelper((UArrayProperty*)FastArrayItemCmd.Property, ObjectArray);
+							for (int32 ObjectIndex = ShadowArrayHelper.Num(); ObjectIndex < ObjectArrayNum; ++ObjectIndex)
+							{
+								FFastArraySerializerItem* ObjectItem = GetFastArrayItem(Params, ObjectArrayHelper.GetRawPtr(ObjectIndex), FastArrayItemCmd);
+								FFastArraySerializerItem* ShadowItem = GetFastArrayItem(Params, ShadowArrayHelper.GetRawPtr(ObjectIndex), FastArrayItemCmd);
+								ShadowItem->ReplicationID = ObjectItem->ReplicationID;
+								ShadowArrayItemIsNew[ObjectIndex] = true;
+							}
+						}
 
+						for (auto& IDIndexPair : ChangedElements)
+						{
 							// Go ahead and do a property compare here, regardless of what we'll actually use below.
 							// This is to prevent issues where someone with an initial / outdated connection doesn't properly
 							// update the changelists in our history, but does update the shadow state inadvertently.
 							FComparePropertiesSharedParams SharedParams{
 								/*bIsInitial=*/ bIsInitial,
-								/*bForceFail=*/ bForceFail,
+								/*bForceFail=*/ bIsInitial || ShadowArrayItemIsNew[IDIndexPair.Idx],
 								RoleIndex,
 								RemoteRoleIndex,
 								Parents,
@@ -5983,8 +6208,7 @@ bool FRepLayout::DeltaSerializeFastArrayProperty(FFastArrayDeltaSerializeParams&
 
 							FConstRepObjectDataBuffer ElementData(ObjectArrayData + ArrayElementOffset);
 							FRepShadowDataBuffer ElementShadowData(ShadowArrayData + ArrayElementOffset);
-
-							TArray<uint16>& HistoryChangelist = HistoryChangelists.Emplace(IDIndexPair.ID);
+							TArray<uint16>& HistoryChangelist = HistoryItem.ChangelistByID.Emplace(IDIndexPair.ID);
 							CompareProperties_r(SharedParams, ItemLayoutStart, ItemLayoutEnd, ElementShadowData, ElementData, HistoryChangelist, 0);
 
 							if (HistoryChangelist.Num())
@@ -5992,61 +6216,64 @@ bool FRepLayout::DeltaSerializeFastArrayProperty(FFastArrayDeltaSerializeParams&
 								HistoryChangelist.Add(0);
 							}
 						}
+
+						// Copy our replication key, and the IDToIndexMap so we can use them later.
+						FastArrayState.ArrayReplicationKey = NewArrayDeltaState->ArrayReplicationKey;
+						FastArrayState.IDToIndexMap = ArraySerializer.ItemMap;
 					}
 				}
 
 				// Now, merge all of the changelists we need together.
 				// If we're sufficiently far back, or if this is our first transmission, then we'll just force fail
-				// and send all changes. Note, this won't be all changes since the beginning, but just all changes
-				// for the currently dirty items.
+				// and send all changes (happens in the block below, Changelists.Num() == 0).
+
+				// Note, this won't be all changes since the beginning, but just all changes for the currently dirty items.
 
-				if (INDEX_NONE != LastSentHistory && ((NewChangelistHistory - LastSentHistory) < FCustomDeltaChangelistState::MAX_CHANGE_HISTORY))
+				if (LastSentChangelistDelta != 0 && LastSentChangelistDelta < FDeltaArrayHistoryState::MAX_CHANGE_HISTORY)
 				{
+					const FConstRepObjectDataBuffer ConstObjectData(ObjectData);
 					Changelists.SetNum(ChangedElements.Num());
-					for (int32 ChangelistHistory = LastSentHistory + 1; ChangelistHistory <= NewChangelistHistory; ++ChangelistHistory)
-					{
-						const int32 RelativeHistory = ChangelistHistory % FCustomDeltaChangelistState::MAX_CHANGE_HISTORY;
-						FCustomDeltaChangedHistory& HistoryItem = DeltaChangelistState.ChangeHistory[RelativeHistory];
 
-						const FConstRepObjectDataBuffer ConstObjectData(ObjectData);
+					// Note, we iterate from LastSentHistory + 1, because we don't want to send something if
+					// we think its already been sent.
+					// Similarly, we do <= NewChangelistHistory because we need to send the newest history.
+					for (uint32 ChangelistHistory = LastSentHistory + 1; ChangelistHistory <= NewChangelistHistory; ++ChangelistHistory)
+					{
+						const uint32 RelativeHistory = ChangelistHistory % FDeltaArrayHistoryState::MAX_CHANGE_HISTORY;
+						FDeltaArrayHistoryItem& HistoryItem = FastArrayState.ChangeHistory[RelativeHistory];
 
-						if (HistoryItem.bWasUpdated[FastArrayNumber])
+						for (int32 i = 0; i < ChangedElements.Num(); ++i)
 						{
-							for (int32 i = 0; i < ChangedElements.Num(); ++i)
+							const auto& IDIndexPair = ChangedElements[i];
+							if (const TArray<uint16>* FoundChangelist = HistoryItem.ChangelistByID.Find(IDIndexPair.ID))
 							{
-								const TMap<int32, TArray<uint16>>& HistoryChangelists = HistoryItem.FastArrayChangelists[FastArrayNumber].ChangelistByID;
-								const auto& IDIndexPair = ChangedElements[i];
-
-								if (const TArray<uint16>* FoundChangelist = HistoryChangelists.Find(IDIndexPair.ID))
+								if (FoundChangelist->Num() > 1)
 								{
-									if (FoundChangelist->Num() > 1)
-									{
-										// This is basically the DynamicArray case from MergeChangelists, but specialized.
-										// We could probably just make that more generic.
-										// Might also be worth creating a helper Lambda / Struct to create ChangelistIterators / HandleIterators.
-
-										TArray<uint16>& ElementChangelist = Changelists[i];
-										TArray<uint16> Temp = MoveTemp(ElementChangelist);
+									// This is basically the DynamicArray case from MergeChangelists, but specialized.
+									// We could probably just make that more generic.
+									// Might also be worth creating a helper Lambda / Struct to create ChangelistIterators / HandleIterators.
 
-										const FConstRepObjectDataBuffer ElementData(ObjectArrayData + (IDIndexPair.Idx * ElementSize));
+									TArray<uint16>& ElementChangelist = Changelists[i];
+									TArray<uint16> Temp = MoveTemp(ElementChangelist);
 
-										FChangelistIterator FoundChangelistIterator(*FoundChangelist, 0);
-										FRepHandleIterator FoundHandleIterator(FoundChangelistIterator, Cmds, ArrayHandleToCmdIndex, ElementSize, 1, ItemLayoutStart, ItemLayoutEnd);
+									const FConstRepObjectDataBuffer ElementData(ObjectArrayData + (IDIndexPair.Idx * ElementSize));
 
-										if (Temp.Num() == 0)
-										{
-											PruneChangeList_r(FoundHandleIterator, ElementData, ElementChangelist);
-										}
-										else
-										{
-											FChangelistIterator ElementChangelistIterator(Temp, 0);
-											FRepHandleIterator ElementHandleIterator(ElementChangelistIterator, Cmds, ArrayHandleToCmdIndex, ElementSize, 1, ItemLayoutStart, ItemLayoutEnd);
+									FChangelistIterator FoundChangelistIterator(*FoundChangelist, 0);
+									FRepHandleIterator FoundHandleIterator(FoundChangelistIterator, Cmds, ArrayHandleToCmdIndex, ElementSize, 1, ItemLayoutStart, ItemLayoutEnd);
 
-											MergeChangeList_r(FoundHandleIterator, ElementHandleIterator, ElementData, ElementChangelist);
-										}
+									if (Temp.Num() == 0)
+									{
+										PruneChangeList_r(FoundHandleIterator, ElementData, ElementChangelist);
+									}
+									else
+									{
+										FChangelistIterator ElementChangelistIterator(Temp, 0);
+										FRepHandleIterator ElementHandleIterator(ElementChangelistIterator, Cmds, ArrayHandleToCmdIndex, ElementSize, 1, ItemLayoutStart, ItemLayoutEnd);
 
-										ElementChangelist.Add(0);
+										MergeChangeList_r(FoundHandleIterator, ElementHandleIterator, ElementData, ElementChangelist);
 									}
+
+									ElementChangelist.Add(0);
 								}
 							}
 						}
@@ -6148,7 +6375,6 @@ bool FRepLayout::DeltaSerializeFastArrayProperty(FFastArrayDeltaSerializeParams&
 		TArray<int32, TInlineAllocator<8>>& ChangedElements = *Params.ReadChangedElements;
 		TArray<int32, TInlineAllocator<8>>& AddedElements = *Params.ReadAddedElements;
 
-		FFastArraySerializer& ArraySerializer = Params.ArraySerializer;
 		FScriptArrayHelper FastArrayHelper((UArrayProperty*)FastArrayItemCmd.Property, ObjectArray);
 
 		bool bOutGuidsChanged = false;
@@ -6174,13 +6400,13 @@ bool FRepLayout::DeltaSerializeFastArrayProperty(FFastArrayDeltaSerializeParams&
 				ArraySerializer.ItemMap.Add(ID, ElementIndex);
 				AddedElements.Add(ElementIndex);
 
-				ThisElement = (FFastArraySerializerItem*)(FastArrayHelper.GetRawPtr(ElementIndex));
+				ThisElement = GetFastArrayItem(Params, FastArrayHelper.GetRawPtr(ElementIndex), FastArrayItemCmd);
 				ThisElement->ReplicationID = ID;
 			}
 			else
 			{
 				ElementIndex = *ElementIndexPtr;
-				ThisElement = (FFastArraySerializerItem*)(FastArrayHelper.GetRawPtr(ElementIndex));
+				ThisElement = GetFastArrayItem(Params, FastArrayHelper.GetRawPtr(ElementIndex), FastArrayItemCmd);
 				ChangedElements.Add(ElementIndex);
 
 				UE_LOG(LogNetFastTArray, Log, TEXT("   Changed. ID: %d -> Idx: %d"), ID, ElementIndex);
@@ -6275,7 +6501,12 @@ bool FRepLayout::DeltaSerializeFastArrayProperty(FFastArrayDeltaSerializeParams&
 
 void FRepLayout::GatherGuidReferencesForFastArray(FFastArrayDeltaSerializeParams& Params) const
 {
-	const FFastArraySerializer& ArraySerializer = Params.ArraySerializer;
+	using namespace UE4_RepLayout_Private;
+
+	const FConstRepObjectDataBuffer ObjectData(Params.DeltaSerializeInfo.Object);
+	const FRepParentCmd& Parent = Parents[Params.DeltaSerializeInfo.PropertyRepIndex];
+
+	const FFastArraySerializer& ArraySerializer = *GetFastArray(Params, ObjectData, Parent);
 	TSet<FNetworkGUID>& GatherGuids = *Params.DeltaSerializeInfo.GatherGuidReferences;
 	int32 TrackedGuidMemory = 0;
 
@@ -6292,7 +6523,12 @@ void FRepLayout::GatherGuidReferencesForFastArray(FFastArrayDeltaSerializeParams
 
 bool FRepLayout::MoveMappedObjectToUnmappedForFastArray(FFastArrayDeltaSerializeParams& Params) const
 {
-	FFastArraySerializer& ArraySerializer = Params.ArraySerializer;
+	using namespace UE4_RepLayout_Private;
+
+	const FRepObjectDataBuffer ObjectData(Params.DeltaSerializeInfo.Object);
+	const FRepParentCmd& Parent = Parents[Params.DeltaSerializeInfo.PropertyRepIndex];
+
+	FFastArraySerializer& ArraySerializer = *GetFastArray(Params, ObjectData, Parent);
 	const FNetworkGUID& MoveToUnmapped = *Params.DeltaSerializeInfo.MoveGuidToUnmapped;
 
 	bool bFound = false;
@@ -6305,6 +6541,8 @@ bool FRepLayout::MoveMappedObjectToUnmappedForFastArray(FFastArrayDeltaSerialize
 
 void FRepLayout::UpdateUnmappedGuidsForFastArray(FFastArrayDeltaSerializeParams& Params) const
  {
+	using namespace UE4_RepLayout_Private;
+
 	check(LifetimeCustomPropertyState);
 
 	FNetDeltaSerializeInfo& DeltaSerializeInfo = Params.DeltaSerializeInfo;
@@ -6328,10 +6566,10 @@ void FRepLayout::UpdateUnmappedGuidsForFastArray(FFastArrayDeltaSerializeParams&
 	UPackageMap* PackageMap = DeltaSerializeInfo.Map;
 
 	FRepObjectDataBuffer ObjectData(Object);
-	FScriptArray* ScriptArray = (FScriptArray*)(ObjectData + FastArrayItemCmd).Data;
+	FScriptArray* ScriptArray = GetTypedProperty<FScriptArray>(ObjectData, FastArrayItemCmd);
 	FRepObjectDataBuffer ArrayData(ScriptArray->GetData());
 
-	FFastArraySerializer& ArraySerializer = Params.ArraySerializer;
+	FFastArraySerializer& ArraySerializer = *GetFastArray(Params, ObjectData, Parent);
 
 	bool bCalledPreNetReceive = Params.DeltaSerializeInfo.bOutSomeObjectsWereMapped;
 	for (auto& GuidReferencesPair : ArraySerializer.GuidReferencesMap_StructDelta)
diff --git a/Engine/Source/Runtime/Engine/Public/Net/RepLayout.h b/Engine/Source/Runtime/Engine/Public/Net/RepLayout.h
--- a/Engine/Source/Runtime/Engine/Public/Net/RepLayout.h
+++ b/Engine/Source/Runtime/Engine/Public/Net/RepLayout.h
@@ -55,6 +55,12 @@ enum class ERepDataBufferType
 
 namespace UE4_RepLayout_Private
 {
+	// The purpose of this class is to allow us to translate or preserve const volatileness between disparate classes.
+	template<typename T, typename U> struct TTranslateConstVolatile { typedef typename TRemoveCV<U>::Type Type; };
+	template<typename T, typename U> struct TTranslateConstVolatile<const T, U> { typedef const typename TRemoveCV<U>::Type Type; };
+	template<typename T, typename U> struct TTranslateConstVolatile<volatile T, U> { typedef volatile typename TRemoveCV<U>::Type Type; };
+	template<typename T, typename U> struct TTranslateConstVolatile<const volatile T, U> { typedef const volatile typename TRemoveCV<U>::Type Type;  };
+
 	/**
 	 * TRepDataBuffer and TConstRepDataBuffer act as wrapper around internal data
 	 * buffers that FRepLayout may use. This allows FRepLayout to properly interact
@@ -64,17 +70,7 @@ namespace UE4_RepLayout_Private
 	struct TRepDataBufferBase
 	{
 		static constexpr ERepDataBufferType Type = DataType;
-
-	private:
-
-		// To be consistent, we need to match the constness of the base type
-		// to the constness of the void* we use as parameters or as returns.
-		// This bit of code does that. Working inside out:
-		//		1. We remove the const / volatile qualifications from the ConstOrNotType.
-		//		2. We see if the CV-stripped type is the same as the ConstOrNotType.
-		//		3. If they are the same, then we know the input is not const, and so we'll use void*.
-		//			If they aren't the same, assume that ConstOrNotType is const, and use const void*.
-		typedef typename TChooseClass<TAreTypesEqual<ConstOrNotType, typename TRemoveCV<ConstOrNotType>::Type>::Value, void, void const>::Result ConstOrNotVoid;
+		using ConstOrNotVoid = typename TTranslateConstVolatile<ConstOrNotType, void>::Type;
 
 	public:
 
