Reduce chances of using dirty FNetSerializeCB.

Reshelved 5348198 for jon.nabozny

diff --git a/Engine/Source/Runtime/Engine/Private/RepLayout.cpp b/Engine/Source/Runtime/Engine/Private/RepLayout.cpp
--- a/Engine/Source/Runtime/Engine/Private/RepLayout.cpp
+++ b/Engine/Source/Runtime/Engine/Private/RepLayout.cpp
@@ -3156,11 +3156,12 @@ void FRepLayout::GatherGuidReferencesForCustomDeltaProperties(FNetDeltaSerialize
 			UStructProperty* StructProperty = static_cast<UStructProperty*>(Parent.Property);
 			UScriptStruct::ICppStructOps* CppStructOps = StructProperty->Struct->GetCppStructOps();
 
-			Params.Struct = StructProperty->Struct;
-			Params.PropertyRepIndex = KVP.Key;
-			Params.Data = ObjectData + Parent;
+			FNetDeltaSerializeInfo TempParams = Params;
+			TempParams.Struct = StructProperty->Struct;
+			TempParams.PropertyRepIndex = KVP.Key;
+			TempParams.Data = ObjectData + Parent;
 
-			CppStructOps->NetDeltaSerialize(Params, Params.Data);
+			CppStructOps->NetDeltaSerialize(TempParams, TempParams.Data);
 		}
 	}
 }
@@ -3222,15 +3223,20 @@ bool FRepLayout::MoveMappedObjectToUnmappedForCustomDeltaProperties(FNetDeltaSer
 			UStructProperty* StructProperty = static_cast<UStructProperty*>(Parent.Property);
 			UScriptStruct::ICppStructOps* CppStructOps = StructProperty->Struct->GetCppStructOps();
 
-			Params.Struct = StructProperty->Struct;
-			Params.Data = ObjectData + Parent;
-			Params.PropertyRepIndex = KVP.Key;
+			FNetDeltaSerializeInfo TempParams = Params;
 
-			if (CppStructOps->NetDeltaSerialize(Params, Params.Data))
+			TempParams.Struct = StructProperty->Struct;
+			TempParams.Data = ObjectData + Parent;
+			TempParams.PropertyRepIndex = KVP.Key;
+
+			if (CppStructOps->NetDeltaSerialize(TempParams, TempParams.Data))
 			{
 				UnmappedCustomProperties.Add(Parent.Offset, StructProperty);
 				bFound = true;
 			}
+
+			Params.bOutHasMoreUnmapped |= TempParams.bOutHasMoreUnmapped;
+			Params.bOutSomeObjectsWereMapped |= TempParams.bOutSomeObjectsWereMapped;
 		}
 	}
 
@@ -3409,9 +3415,6 @@ PRAGMA_ENABLE_DEPRECATION_WARNINGS
 
 void FRepLayout::UpdateUnmappedObjectsForCustomDeltaProperties(FNetDeltaSerializeInfo& Params, TArray<TPair<int32, UStructProperty*>>& CompletelyUpdated, TArray<TPair<int32, UStructProperty*>>& Updated) const
 {
-	bool bOutSomeObjectsWereMapped = false;
-	bool bOutHasMoreUnmapped = false;
-
 	if (LifetimeCustomPropertyState)
 	{
 		FRepObjectDataBuffer ObjectData(Params.Object);
@@ -3423,33 +3426,32 @@ void FRepLayout::UpdateUnmappedObjectsForCustomDeltaProperties(FNetDeltaSerializ
 			UStructProperty* StructProperty = static_cast<UStructProperty*>(Parent.Property);
 			UScriptStruct::ICppStructOps* CppStructOps = StructProperty->Struct->GetCppStructOps();
 
-			Params.DebugName = Parent.CachedPropertyName.ToString();
-			Params.Struct = StructProperty->Struct;
-			Params.bOutSomeObjectsWereMapped = false;
-			Params.bOutHasMoreUnmapped = false;
-			Params.PropertyRepIndex = KVP.Key;
-			Params.Data = ObjectData + Parent;
+			FNetDeltaSerializeInfo TempParams = Params;
+
+			TempParams.DebugName = Parent.CachedPropertyName.ToString();
+			TempParams.Struct = StructProperty->Struct;
+			TempParams.bOutSomeObjectsWereMapped = false;
+			TempParams.bOutHasMoreUnmapped = false;
+			TempParams.PropertyRepIndex = KVP.Key;
+			TempParams.Data = ObjectData + Parent;
 
 			// Call the custom delta serialize function to handle it
-			CppStructOps->NetDeltaSerialize(Params, Params.Data);
+			CppStructOps->NetDeltaSerialize(TempParams, TempParams.Data);
 
-			if (Params.bOutSomeObjectsWereMapped)
+			if (TempParams.bOutSomeObjectsWereMapped)
 			{
 				Updated.Emplace(Parent.Offset, StructProperty);
 			}
 
-			if (!Params.bOutHasMoreUnmapped)
+			if (!TempParams.bOutHasMoreUnmapped)
 			{
 				CompletelyUpdated.Emplace(Parent.Offset, StructProperty);
 			}
 
-			bOutSomeObjectsWereMapped |= Params.bOutSomeObjectsWereMapped;
-			bOutHasMoreUnmapped |= Params.bOutHasMoreUnmapped;
+			Params.bOutSomeObjectsWereMapped |= Params.bOutSomeObjectsWereMapped;
+			Params.bOutHasMoreUnmapped |= Params.bOutHasMoreUnmapped;
 		}
 	}
-
-	Params.bOutSomeObjectsWereMapped = bOutSomeObjectsWereMapped;
-	Params.bOutHasMoreUnmapped = bOutHasMoreUnmapped;
 }
 
 bool FRepLayout::SendCustomDeltaProperty(
@@ -3474,14 +3476,15 @@ bool FRepLayout::SendCustomDeltaProperty(FNetDeltaSerializeInfo& Params, uint16
 
 	check(CppStructOps); // else should not have STRUCT_NetSerializeNative
 
-	Params.DebugName = Parent.CachedPropertyName.ToString();
-	Params.Struct = StructProperty->Struct;
-	Params.PropertyRepIndex = CustomDeltaRepIndex;
-	Params.Data = FRepObjectDataBuffer(Params.Object) + Parent;
+	FNetDeltaSerializeInfo TempParams = Params;
+	TempParams.DebugName = Parent.CachedPropertyName.ToString();
+	TempParams.Struct = StructProperty->Struct;
+	TempParams.PropertyRepIndex = CustomDeltaRepIndex;
+	TempParams.Data = FRepObjectDataBuffer(Params.Object) + Parent;
 
-	bool bSupportsFastArrayDelta = Params.bSupportsFastArrayDeltaStructSerialization;
+	bool bSupportsFastArrayDelta = TempParams.bSupportsFastArrayDeltaStructSerialization;
 
-	if (Params.bSupportsFastArrayDeltaStructSerialization &&
+	if (TempParams.bSupportsFastArrayDeltaStructSerialization &&
 		EnumHasAnyFlags(Parent.Flags, ERepParentFlags::IsFastArray) &&
 		!!LifetimeCustomPropertyState->NumFastArrayItems)
 	{
@@ -3489,15 +3492,15 @@ bool FRepLayout::SendCustomDeltaProperty(FNetDeltaSerializeInfo& Params, uint16
 		bSupportsFastArrayDelta = CustomDeltaProperty.FastArrayNumber != INDEX_NONE;
 	}
 
-	Params.Writer->WriteBit(bSupportsFastArrayDelta);
+	TempParams.Writer->WriteBit(bSupportsFastArrayDelta);
 
 	if (Parent.Property->ArrayDim != 1)
 	{
 		uint32 StaticArrayIndex = Parent.ArrayIndex;
-		Params.Writer->SerializeIntPacked(StaticArrayIndex);
+		TempParams.Writer->SerializeIntPacked(StaticArrayIndex);
 	}
 
-	TGuardValue<bool> SupportFastArrayDeltaGuard(Params.bSupportsFastArrayDeltaStructSerialization, bSupportsFastArrayDelta);
+	TempParams.bSupportsFastArrayDeltaStructSerialization = bSupportsFastArrayDelta;
 	return CppStructOps->NetDeltaSerialize(Params, Params.Data);
 }
 
@@ -3512,13 +3515,14 @@ PRAGMA_ENABLE_DEPRECATION_WARNINGS
 
 bool FRepLayout::ReceiveCustomDeltaProperty(FNetDeltaSerializeInfo& Params, UStructProperty* Property, uint32& StaticArrayIndex, int32& Offset) const
 {
-	if (Params.Connection->EngineNetworkProtocolVersion >= EEngineNetworkVersionHistory::HISTORY_FAST_ARRAY_DELTA_STRUCT)
+	FNetDeltaSerializeInfo TempParams = Params;
+	if (TempParams.Connection->EngineNetworkProtocolVersion >= EEngineNetworkVersionHistory::HISTORY_FAST_ARRAY_DELTA_STRUCT)
 	{
-		Params.bSupportsFastArrayDeltaStructSerialization = !!Params.Reader->ReadBit();
+		TempParams.bSupportsFastArrayDeltaStructSerialization = !!TempParams.Reader->ReadBit();
 	}
 	else
 	{
-		Params.bSupportsFastArrayDeltaStructSerialization = false;
+		TempParams.bSupportsFastArrayDeltaStructSerialization = false;
 	}
 
 	// Receive array index (static sized array, i.e. MemberVariable[4])
@@ -3526,7 +3530,7 @@ bool FRepLayout::ReceiveCustomDeltaProperty(FNetDeltaSerializeInfo& Params, UStr
 	{
 		check(Property->ArrayDim >= 2);
 
-		Params.Reader->SerializeIntPacked(StaticArrayIndex);
+		TempParams.Reader->SerializeIntPacked(StaticArrayIndex);
 
 		if (StaticArrayIndex >= (uint32)Property->ArrayDim)
 		{
@@ -3551,12 +3555,17 @@ bool FRepLayout::ReceiveCustomDeltaProperty(FNetDeltaSerializeInfo& Params, UStr
 
 	check(CppStructOps);
 
-	Params.DebugName = Parent.CachedPropertyName.ToString();
-	Params.Struct = InnerStruct;
-	Params.PropertyRepIndex = Property->RepIndex + StaticArrayIndex;
-	Params.Data = FRepObjectDataBuffer(Params.Object) + Parent;
+	TempParams.DebugName = Parent.CachedPropertyName.ToString();
+	TempParams.Struct = InnerStruct;
+	TempParams.PropertyRepIndex = Property->RepIndex + StaticArrayIndex;
+	TempParams.Data = FRepObjectDataBuffer(Params.Object) + Parent;
 
-	return CppStructOps->NetDeltaSerialize(Params, Params.Data);
+	const bool bSuccess = CppStructOps->NetDeltaSerialize(Params, Params.Data);
+
+	Params.bOutHasMoreUnmapped |= TempParams.bOutHasMoreUnmapped;
+	Params.bOutSomeObjectsWereMapped |= TempParams.bOutSomeObjectsWereMapped;
+
+	return bSuccess;
 }
 
 void FRepLayout::CallRepNotifies(FReceivingRepState* RepState, UObject* Object) const
