Revert CL 5354291

diff --git a/Engine/Source/Runtime/Engine/Private/RepLayout.cpp b/Engine/Source/Runtime/Engine/Private/RepLayout.cpp
--- a/Engine/Source/Runtime/Engine/Private/RepLayout.cpp
+++ b/Engine/Source/Runtime/Engine/Private/RepLayout.cpp
@@ -3156,12 +3156,11 @@ void FRepLayout::GatherGuidReferencesForCustomDeltaProperties(FNetDeltaSerialize
 			UStructProperty* StructProperty = static_cast<UStructProperty*>(Parent.Property);
 			UScriptStruct::ICppStructOps* CppStructOps = StructProperty->Struct->GetCppStructOps();
 
-			FNetDeltaSerializeInfo TempParams = Params;
-			TempParams.Struct = StructProperty->Struct;
-			TempParams.PropertyRepIndex = KVP.Key;
-			TempParams.Data = ObjectData + Parent;
+			Params.Struct = StructProperty->Struct;
+			Params.PropertyRepIndex = KVP.Key;
+			Params.Data = ObjectData + Parent;
 
-			CppStructOps->NetDeltaSerialize(TempParams, TempParams.Data);
+			CppStructOps->NetDeltaSerialize(Params, Params.Data);
 		}
 	}
 }
@@ -3223,20 +3222,15 @@ bool FRepLayout::MoveMappedObjectToUnmappedForCustomDeltaProperties(FNetDeltaSer
 			UStructProperty* StructProperty = static_cast<UStructProperty*>(Parent.Property);
 			UScriptStruct::ICppStructOps* CppStructOps = StructProperty->Struct->GetCppStructOps();
 
-			FNetDeltaSerializeInfo TempParams = Params;
+			Params.Struct = StructProperty->Struct;
+			Params.Data = ObjectData + Parent;
+			Params.PropertyRepIndex = KVP.Key;
 
-			TempParams.Struct = StructProperty->Struct;
-			TempParams.Data = ObjectData + Parent;
-			TempParams.PropertyRepIndex = KVP.Key;
-
-			if (CppStructOps->NetDeltaSerialize(TempParams, TempParams.Data))
+			if (CppStructOps->NetDeltaSerialize(Params, Params.Data))
 			{
 				UnmappedCustomProperties.Add(Parent.Offset, StructProperty);
 				bFound = true;
 			}
-
-			Params.bOutHasMoreUnmapped |= TempParams.bOutHasMoreUnmapped;
-			Params.bOutSomeObjectsWereMapped |= TempParams.bOutSomeObjectsWereMapped;
 		}
 	}
 
@@ -3415,6 +3409,9 @@ PRAGMA_ENABLE_DEPRECATION_WARNINGS
 
 void FRepLayout::UpdateUnmappedObjectsForCustomDeltaProperties(FNetDeltaSerializeInfo& Params, TArray<TPair<int32, UStructProperty*>>& CompletelyUpdated, TArray<TPair<int32, UStructProperty*>>& Updated) const
 {
+	bool bOutSomeObjectsWereMapped = false;
+	bool bOutHasMoreUnmapped = false;
+
 	if (LifetimeCustomPropertyState)
 	{
 		FRepObjectDataBuffer ObjectData(Params.Object);
@@ -3426,32 +3423,33 @@ void FRepLayout::UpdateUnmappedObjectsForCustomDeltaProperties(FNetDeltaSerializ
 			UStructProperty* StructProperty = static_cast<UStructProperty*>(Parent.Property);
 			UScriptStruct::ICppStructOps* CppStructOps = StructProperty->Struct->GetCppStructOps();
 
-			FNetDeltaSerializeInfo TempParams = Params;
-
-			TempParams.DebugName = Parent.CachedPropertyName.ToString();
-			TempParams.Struct = StructProperty->Struct;
-			TempParams.bOutSomeObjectsWereMapped = false;
-			TempParams.bOutHasMoreUnmapped = false;
-			TempParams.PropertyRepIndex = KVP.Key;
-			TempParams.Data = ObjectData + Parent;
+			Params.DebugName = Parent.CachedPropertyName.ToString();
+			Params.Struct = StructProperty->Struct;
+			Params.bOutSomeObjectsWereMapped = false;
+			Params.bOutHasMoreUnmapped = false;
+			Params.PropertyRepIndex = KVP.Key;
+			Params.Data = ObjectData + Parent;
 
 			// Call the custom delta serialize function to handle it
-			CppStructOps->NetDeltaSerialize(TempParams, TempParams.Data);
+			CppStructOps->NetDeltaSerialize(Params, Params.Data);
 
-			if (TempParams.bOutSomeObjectsWereMapped)
+			if (Params.bOutSomeObjectsWereMapped)
 			{
 				Updated.Emplace(Parent.Offset, StructProperty);
 			}
 
-			if (!TempParams.bOutHasMoreUnmapped)
+			if (!Params.bOutHasMoreUnmapped)
 			{
 				CompletelyUpdated.Emplace(Parent.Offset, StructProperty);
 			}
 
-			Params.bOutSomeObjectsWereMapped |= Params.bOutSomeObjectsWereMapped;
-			Params.bOutHasMoreUnmapped |= Params.bOutHasMoreUnmapped;
+			bOutSomeObjectsWereMapped |= Params.bOutSomeObjectsWereMapped;
+			bOutHasMoreUnmapped |= Params.bOutHasMoreUnmapped;
 		}
 	}
+
+	Params.bOutSomeObjectsWereMapped = bOutSomeObjectsWereMapped;
+	Params.bOutHasMoreUnmapped = bOutHasMoreUnmapped;
 }
 
 bool FRepLayout::SendCustomDeltaProperty(
@@ -3476,15 +3474,14 @@ bool FRepLayout::SendCustomDeltaProperty(FNetDeltaSerializeInfo& Params, uint16
 
 	check(CppStructOps); // else should not have STRUCT_NetSerializeNative
 
-	FNetDeltaSerializeInfo TempParams = Params;
-	TempParams.DebugName = Parent.CachedPropertyName.ToString();
-	TempParams.Struct = StructProperty->Struct;
-	TempParams.PropertyRepIndex = CustomDeltaRepIndex;
-	TempParams.Data = FRepObjectDataBuffer(Params.Object) + Parent;
+	Params.DebugName = Parent.CachedPropertyName.ToString();
+	Params.Struct = StructProperty->Struct;
+	Params.PropertyRepIndex = CustomDeltaRepIndex;
+	Params.Data = FRepObjectDataBuffer(Params.Object) + Parent;
 
-	bool bSupportsFastArrayDelta = TempParams.bSupportsFastArrayDeltaStructSerialization;
+	bool bSupportsFastArrayDelta = Params.bSupportsFastArrayDeltaStructSerialization;
 
-	if (TempParams.bSupportsFastArrayDeltaStructSerialization &&
+	if (Params.bSupportsFastArrayDeltaStructSerialization &&
 		EnumHasAnyFlags(Parent.Flags, ERepParentFlags::IsFastArray) &&
 		!!LifetimeCustomPropertyState->NumFastArrayItems)
 	{
@@ -3492,15 +3489,15 @@ bool FRepLayout::SendCustomDeltaProperty(FNetDeltaSerializeInfo& Params, uint16
 		bSupportsFastArrayDelta = CustomDeltaProperty.FastArrayNumber != INDEX_NONE;
 	}
 
-	TempParams.Writer->WriteBit(bSupportsFastArrayDelta);
+	Params.Writer->WriteBit(bSupportsFastArrayDelta);
 
 	if (Parent.Property->ArrayDim != 1)
 	{
 		uint32 StaticArrayIndex = Parent.ArrayIndex;
-		TempParams.Writer->SerializeIntPacked(StaticArrayIndex);
+		Params.Writer->SerializeIntPacked(StaticArrayIndex);
 	}
 
-	TempParams.bSupportsFastArrayDeltaStructSerialization = bSupportsFastArrayDelta;
+	TGuardValue<bool> SupportFastArrayDeltaGuard(Params.bSupportsFastArrayDeltaStructSerialization, bSupportsFastArrayDelta);
 	return CppStructOps->NetDeltaSerialize(Params, Params.Data);
 }
 
@@ -3515,14 +3512,13 @@ PRAGMA_ENABLE_DEPRECATION_WARNINGS
 
 bool FRepLayout::ReceiveCustomDeltaProperty(FNetDeltaSerializeInfo& Params, UStructProperty* Property, uint32& StaticArrayIndex, int32& Offset) const
 {
-	FNetDeltaSerializeInfo TempParams = Params;
-	if (TempParams.Connection->EngineNetworkProtocolVersion >= EEngineNetworkVersionHistory::HISTORY_FAST_ARRAY_DELTA_STRUCT)
+	if (Params.Connection->EngineNetworkProtocolVersion >= EEngineNetworkVersionHistory::HISTORY_FAST_ARRAY_DELTA_STRUCT)
 	{
-		TempParams.bSupportsFastArrayDeltaStructSerialization = !!TempParams.Reader->ReadBit();
+		Params.bSupportsFastArrayDeltaStructSerialization = !!Params.Reader->ReadBit();
 	}
 	else
 	{
-		TempParams.bSupportsFastArrayDeltaStructSerialization = false;
+		Params.bSupportsFastArrayDeltaStructSerialization = false;
 	}
 
 	// Receive array index (static sized array, i.e. MemberVariable[4])
@@ -3530,7 +3526,7 @@ bool FRepLayout::ReceiveCustomDeltaProperty(FNetDeltaSerializeInfo& Params, UStr
 	{
 		check(Property->ArrayDim >= 2);
 
-		TempParams.Reader->SerializeIntPacked(StaticArrayIndex);
+		Params.Reader->SerializeIntPacked(StaticArrayIndex);
 
 		if (StaticArrayIndex >= (uint32)Property->ArrayDim)
 		{
@@ -3555,17 +3551,12 @@ bool FRepLayout::ReceiveCustomDeltaProperty(FNetDeltaSerializeInfo& Params, UStr
 
 	check(CppStructOps);
 
-	TempParams.DebugName = Parent.CachedPropertyName.ToString();
-	TempParams.Struct = InnerStruct;
-	TempParams.PropertyRepIndex = Property->RepIndex + StaticArrayIndex;
-	TempParams.Data = FRepObjectDataBuffer(Params.Object) + Parent;
-
-	const bool bSuccess = CppStructOps->NetDeltaSerialize(Params, Params.Data);
+	Params.DebugName = Parent.CachedPropertyName.ToString();
+	Params.Struct = InnerStruct;
+	Params.PropertyRepIndex = Property->RepIndex + StaticArrayIndex;
+	Params.Data = FRepObjectDataBuffer(Params.Object) + Parent;
 
-	Params.bOutHasMoreUnmapped |= TempParams.bOutHasMoreUnmapped;
-	Params.bOutSomeObjectsWereMapped |= TempParams.bOutSomeObjectsWereMapped;
-
-	return bSuccess;
+	return CppStructOps->NetDeltaSerialize(Params, Params.Data);
 }
 
 void FRepLayout::CallRepNotifies(FReceivingRepState* RepState, UObject* Object) const
