Merging CL 5358524.

diff --git a/Engine/Source/Runtime/Engine/Private/Materials/MaterialShared.cpp b/Engine/Source/Runtime/Engine/Private/Materials/MaterialShared.cpp
--- a/Engine/Source/Runtime/Engine/Private/Materials/MaterialShared.cpp
+++ b/Engine/Source/Runtime/Engine/Private/Materials/MaterialShared.cpp
@@ -450,10 +450,16 @@ void FMaterialCompilationOutput::Serialize(FArchive& Ar)
 	UniformExpressionSet.Serialize(Ar);
 
 	Ar << UsedSceneTextures;
+#if WITH_EDITOR
 	Ar << NumUsedUVScalars;
 	Ar << NumUsedCustomInterpolatorScalars;
 	Ar << EstimatedNumTextureSamplesVS;
 	Ar << EstimatedNumTextureSamplesPS;
+#else
+	uint8 Tmp8; uint16 Tmp16;
+	Ar << Tmp8 << Tmp8;
+	Ar << Tmp16 << Tmp16;
+#endif
 
 	Ar << bRequiresSceneColorCopy;
 	Ar << bNeedsSceneTextures;
@@ -2300,6 +2306,7 @@ int32 FMaterialResource::GetSamplerUsage() const
 }
 #endif // WITH_EDITOR
 
+#if WITH_EDITOR
 void FMaterialResource::GetUserInterpolatorUsage(uint32& NumUsedUVScalars, uint32& NumUsedCustomInterpolatorScalars) const
 {
 	NumUsedUVScalars = NumUsedCustomInterpolatorScalars = 0;
@@ -2319,6 +2326,7 @@ void FMaterialResource::GetEstimatedNumTextureSamples(uint32& VSSamples, uint32&
 		ShaderMap->GetEstimatedNumTextureSamples(VSSamples, PSSamples);
 	}
 }
+#endif // WITH_EDITOR
 
 FString FMaterialResource::GetMaterialUsageDescription() const
 {
diff --git a/Engine/Source/Runtime/Engine/Public/MaterialShared.h b/Engine/Source/Runtime/Engine/Public/MaterialShared.h
--- a/Engine/Source/Runtime/Engine/Public/MaterialShared.h
+++ b/Engine/Source/Runtime/Engine/Public/MaterialShared.h
@@ -1,4 +1,4 @@
-// Copyright 1998-2019 Epic Games, Inc. All Rights Reserved.
+﻿// Copyright 1998-2019 Epic Games, Inc. All Rights Reserved.
 
 /*=============================================================================
 	MaterialShared.h: Shared material definitions.
@@ -418,10 +418,12 @@ class FMaterialCompilationOutput
 public:
 	FMaterialCompilationOutput() :
 		UsedSceneTextures(0),
+#if WITH_EDITOR
 		NumUsedUVScalars(0),
 		NumUsedCustomInterpolatorScalars(0),
 		EstimatedNumTextureSamplesVS(0),
 		EstimatedNumTextureSamplesPS(0),
+#endif
 		bRequiresSceneColorCopy(false),
 		bNeedsSceneTextures(false),
 		bUsesEyeAdaptation(false),
@@ -442,6 +444,7 @@ public:
 	/** Bitfield of the ESceneTextures used */
 	uint64 UsedSceneTextures;
 
+#if WITH_EDITOR
 	/** Number of used custom UV scalars. */
 	uint8 NumUsedUVScalars;
 
@@ -451,6 +454,7 @@ public:
 	/** Number of times SampleTexture is called, excludes custom nodes. */
 	uint16 EstimatedNumTextureSamplesVS;
 	uint16 EstimatedNumTextureSamplesPS;
+#endif // WITH_EDITOR
 
 	/** Indicates whether the material uses scene color. */
 	bool bRequiresSceneColorCopy;
@@ -897,9 +901,11 @@ public:
 	bool UsesSceneDepthLookup() const { return MaterialCompilationOutput.bUsesSceneDepthLookup; }
 	bool UsesVelocitySceneTexture() const { return MaterialCompilationOutput.bUsesVelocitySceneTexture; }
 	bool UsesDistanceCullFade() const { return MaterialCompilationOutput.bUsesDistanceCullFade; }
+#if WITH_EDITOR
 	uint32 GetNumUsedUVScalars() const { return MaterialCompilationOutput.NumUsedUVScalars; }
 	uint32 GetNumUsedCustomInterpolatorScalars() const { return MaterialCompilationOutput.NumUsedCustomInterpolatorScalars; }
 	void GetEstimatedNumTextureSamples(uint32& VSSamples, uint32& PSSamples) const { VSSamples = MaterialCompilationOutput.EstimatedNumTextureSamplesVS; PSSamples = MaterialCompilationOutput.EstimatedNumTextureSamplesPS; }
+#endif
 	bool UsesSceneTexture(uint32 TexId) const { return MaterialCompilationOutput.UsedSceneTextures & (1ull << TexId); }
 
 	bool IsValidForRendering(bool bFailOnInvalid = false) const
@@ -1878,8 +1884,10 @@ public:
 	ENGINE_API int32 GetSamplerUsage() const;
 #endif
 
+#if WITH_EDITOR
 	ENGINE_API void GetUserInterpolatorUsage(uint32& NumUsedUVScalars, uint32& NumUsedCustomInterpolatorScalars) const;
 	ENGINE_API void GetEstimatedNumTextureSamples(uint32& VSSamples, uint32& PSSamples) const;
+#endif
 
 	ENGINE_API virtual FString GetMaterialUsageDescription() const override;
 
