Hash material shader map ID at cook time, saves around 4MB on Android.
Note that FeatureLevel and QualityLevel are still being duplicated and could likely be removed.

diff --git a/Engine/Source/Runtime/Engine/Private/Materials/Material.cpp b/Engine/Source/Runtime/Engine/Private/Materials/Material.cpp
--- a/Engine/Source/Runtime/Engine/Private/Materials/Material.cpp
+++ b/Engine/Source/Runtime/Engine/Private/Materials/Material.cpp
@@ -233,6 +233,7 @@ void FMaterialResource::GatherExpressionsForCustomInterpolators(TArray<UMaterial
 void FMaterialResource::GetShaderMapId(EShaderPlatform Platform, FMaterialShaderMapId& OutId) const
 {
 	FMaterial::GetShaderMapId(Platform, OutId);
+#if WITH_EDITOR
 	Material->AppendReferencedFunctionIdsTo(OutId.ReferencedFunctions);
 	Material->AppendReferencedParameterCollectionIdsTo(OutId.ReferencedParameterCollections);
 
@@ -242,19 +243,11 @@ void FMaterialResource::GetShaderMapId(EShaderPlatform Platform, FMaterialShader
 	{
 		MaterialInstance->GetBasePropertyOverridesHash(OutId.BasePropertyOverridesHash);
 
-#if !WITH_EDITOR
-		if (FMaterial::GetLoadedCookedShaderMapId() && MaterialInstance->bHasStaticPermutationResource)
-		{
-			OutId.UpdateParameterSet(MaterialInstance->GetStaticParameters());
-		}
-		else
-#endif
-		{
-			FStaticParameterSet CompositedStaticParameters;
-			MaterialInstance->GetStaticParameterValues(CompositedStaticParameters);
-			OutId.UpdateParameterSet(CompositedStaticParameters);		
-		}
+		FStaticParameterSet CompositedStaticParameters;
+		MaterialInstance->GetStaticParameterValues(CompositedStaticParameters);
+		OutId.UpdateParameterSet(CompositedStaticParameters);		
 	}
+#endif // WITH_EDITOR
 }
 
 /**
diff --git a/Engine/Source/Runtime/Engine/Private/Materials/MaterialInstance.cpp b/Engine/Source/Runtime/Engine/Private/Materials/MaterialInstance.cpp
--- a/Engine/Source/Runtime/Engine/Private/Materials/MaterialInstance.cpp
+++ b/Engine/Source/Runtime/Engine/Private/Materials/MaterialInstance.cpp
@@ -3010,6 +3010,7 @@ void UMaterialInstance::Serialize(FArchive& Ar)
 			);
 #endif
 		}
+#if WITH_EDITOR
 		else
 		{
 			FMaterialResource LegacyResource;
@@ -3027,6 +3028,7 @@ void UMaterialInstance::Serialize(FArchive& Ar)
 			TrimToOverriddenOnly(StaticParameters.StaticComponentMaskParameters);
 			TrimToOverriddenOnly(StaticParameters.TerrainLayerWeightParameters);
 		}
+#endif // WITH_EDITOR
 	}
 
 	if (Ar.UE4Ver() >= VER_UE4_MATERIAL_INSTANCE_BASE_PROPERTY_OVERRIDES )
diff --git a/Engine/Source/Runtime/Engine/Private/Materials/MaterialShader.cpp b/Engine/Source/Runtime/Engine/Private/Materials/MaterialShader.cpp
--- a/Engine/Source/Runtime/Engine/Private/Materials/MaterialShader.cpp
+++ b/Engine/Source/Runtime/Engine/Private/Materials/MaterialShader.cpp
@@ -93,7 +93,7 @@ FString GetBlendModeString(EBlendMode BlendMode)
 	return BlendModeName;
 }
 
-
+#if WITH_EDITOR
 /** Creates a string key for the derived data cache given a shader map id. */
 static FString GetMaterialShaderMapKeyString(const FMaterialShaderMapId& ShaderMapId, EShaderPlatform Platform)
 {
@@ -104,8 +104,7 @@ static FString GetMaterialShaderMapKeyString(const FMaterialShaderMapId& ShaderM
 	FMaterialAttributeDefinitionMap::AppendDDCKeyString(ShaderMapKeyString);
 	return FDerivedDataCacheInterface::BuildCacheKey(TEXT("MATSM"), MATERIALSHADERMAP_DERIVEDDATA_VER, *ShaderMapKeyString);
 }
-
-
+#endif // WITH_EDITOR
 
 /** Called for every material shader to update the appropriate stats. */
 void UpdateMaterialShaderCompilingStats(const FMaterial* Material)
@@ -489,11 +488,18 @@ void FMaterialShaderMapId::Serialize(FArchive& Ar)
 	Ar.UsingCustomVersion(FEditorObjectVersion::GUID);
 	Ar.UsingCustomVersion(FReleaseObjectVersion::GUID);
 
-	uint32 UsageInt = Usage;
-	Ar << UsageInt;
-	Usage = (EMaterialShaderMapUsage::Type)UsageInt;
+#if WITH_EDITOR
+	const bool bIsSavingCooked = Ar.IsSaving() && Ar.IsCooking();
+
+	if (!bIsSavingCooked)
+	{
+		uint32 UsageInt = Usage;
+		Ar << UsageInt;
+		Usage = (EMaterialShaderMapUsage::Type)UsageInt;
 
-	Ar << BaseMaterialId;
+		Ar << BaseMaterialId;
+	}
+#endif
 
 	if (Ar.UE4Ver() >= VER_UE4_PURGED_FMATERIAL_COMPILE_OUTPUTS)
 	{
@@ -506,45 +512,64 @@ void FMaterialShaderMapId::Serialize(FArchive& Ar)
 		Ar << LegacyQualityLevel;
 	}
 
-	ParameterSet.Serialize(Ar);
+#if WITH_EDITOR
+	if (!bIsSavingCooked)
+	{
+		ParameterSet.Serialize(Ar);
 
-	Ar << ReferencedFunctions;
+		Ar << ReferencedFunctions;
 
-	if (Ar.UE4Ver() >= VER_UE4_COLLECTIONS_IN_SHADERMAPID)
-	{
-		Ar << ReferencedParameterCollections;
-	}
+		if (Ar.UE4Ver() >= VER_UE4_COLLECTIONS_IN_SHADERMAPID)
+		{
+			Ar << ReferencedParameterCollections;
+		}
 
-	if (Ar.CustomVer(FEditorObjectVersion::GUID) >= FEditorObjectVersion::AddedMaterialSharedInputs &&
-		Ar.CustomVer(FReleaseObjectVersion::GUID) < FReleaseObjectVersion::RemovedMaterialSharedInputCollection)
-	{
-		TArray<FGuid> Deprecated;
-		Ar << Deprecated;
-	}
+		if (Ar.CustomVer(FEditorObjectVersion::GUID) >= FEditorObjectVersion::AddedMaterialSharedInputs &&
+			Ar.CustomVer(FReleaseObjectVersion::GUID) < FReleaseObjectVersion::RemovedMaterialSharedInputCollection)
+		{
+			TArray<FGuid> Deprecated;
+			Ar << Deprecated;
+		}
 
-	Ar << ShaderTypeDependencies;
-	if (Ar.UE4Ver() >= VER_UE4_PURGED_FMATERIAL_COMPILE_OUTPUTS)
-	{
-		Ar << ShaderPipelineTypeDependencies;
-	}
-	Ar << VertexFactoryTypeDependencies;
+		Ar << ShaderTypeDependencies;
+		if (Ar.UE4Ver() >= VER_UE4_PURGED_FMATERIAL_COMPILE_OUTPUTS)
+		{
+			Ar << ShaderPipelineTypeDependencies;
+		}
+		Ar << VertexFactoryTypeDependencies;
 
-	if (Ar.UE4Ver() >= VER_UE4_PURGED_FMATERIAL_COMPILE_OUTPUTS)
-	{
-		Ar << TextureReferencesHash;
+		if (Ar.UE4Ver() >= VER_UE4_PURGED_FMATERIAL_COMPILE_OUTPUTS)
+		{
+			Ar << TextureReferencesHash;
+		}
+		else
+		{
+			FSHAHash LegacyHash;
+			Ar << LegacyHash;
+		}
+
+		if (Ar.UE4Ver() >= VER_UE4_MATERIAL_INSTANCE_BASE_PROPERTY_OVERRIDES)
+		{
+			Ar << BasePropertyOverridesHash;
+		}
 	}
 	else
 	{
-		FSHAHash LegacyHash;
-		Ar << LegacyHash;
-	}
+		// Saving cooked data, this should be valid
+		FSHAHash CookedShaderMapIdHash;
+		GetMaterialHash(CookedShaderMapIdHash);
 
-	if( Ar.UE4Ver() >= VER_UE4_MATERIAL_INSTANCE_BASE_PROPERTY_OVERRIDES)
-	{
-		Ar << BasePropertyOverridesHash;
+		checkf(CookedShaderMapIdHash != FSHAHash(), TEXT("Tried to save an invalid shadermap id hash during cook"));
+		Ar << CookedShaderMapIdHash;
 	}
+#else
+	// Cooked so can assume this is valid
+	Ar << CookedShaderMapIdHash;
+	checkf(CookedShaderMapIdHash != FSHAHash(), TEXT("Loaded an invalid cooked shadermap id hash"));
+#endif // WITH_EDITOR
 }
 
+#if WITH_EDITOR
 /** Hashes the material-specific part of this shader map Id. */
 void FMaterialShaderMapId::GetMaterialHash(FSHAHash& OutHash) const
 {
@@ -580,6 +605,7 @@ void FMaterialShaderMapId::GetMaterialHash(FSHAHash& OutHash) const
 	HashState.Final();
 	HashState.GetHash(&OutHash.Hash[0]);
 }
+#endif // WITH_EDITOR
 
 /** 
 * Tests this set against another for equality, disregarding override settings.
@@ -589,18 +615,26 @@ void FMaterialShaderMapId::GetMaterialHash(FSHAHash& OutHash) const
 */
 bool FMaterialShaderMapId::operator==(const FMaterialShaderMapId& ReferenceSet) const
 {
-	if (Usage != ReferenceSet.Usage)
+#if WITH_EDITOR
+	if (Usage != ReferenceSet.Usage
+		|| BaseMaterialId != ReferenceSet.BaseMaterialId)
+	{
+		return false;
+	}
+#else
+	if (CookedShaderMapIdHash != ReferenceSet.CookedShaderMapIdHash)
 	{
 		return false;
 	}
+#endif
 
-	if (BaseMaterialId != ReferenceSet.BaseMaterialId 
-		|| QualityLevel != ReferenceSet.QualityLevel
+	if (QualityLevel != ReferenceSet.QualityLevel
 		|| FeatureLevel != ReferenceSet.FeatureLevel)
 	{
 		return false;
 	}
 
+#if WITH_EDITOR
 	if (ParameterSet != ReferenceSet.ParameterSet
 		|| ReferencedFunctions.Num() != ReferenceSet.ReferencedFunctions.Num()
 		|| ReferencedParameterCollections.Num() != ReferenceSet.ReferencedParameterCollections.Num()
@@ -670,10 +704,12 @@ bool FMaterialShaderMapId::operator==(const FMaterialShaderMapId& ReferenceSet)
 	{
 		return false;
 	}
+#endif // WITH_EDITOR
 
 	return true;
 }
 
+#if WITH_EDITOR
 void FMaterialShaderMapId::UpdateParameterSet(const FStaticParameterSet& StaticParameters)
 {
 	ParameterSet = StaticParameters;
@@ -704,10 +740,11 @@ void FMaterialShaderMapId::UpdateParameterSet(const FStaticParameterSet& StaticP
 		LayersParam.AppendKeyString(ParameterSetLayerParametersKey);
 	}
 }	
+#endif // WITH_EDITOR
 
+#if WITH_EDITOR
 void FMaterialShaderMapId::AppendKeyString(FString& KeyString) const
 {
-#if WITH_EDITOR
 	KeyString += BaseMaterialId.ToString();
 	KeyString += TEXT("_");
 
@@ -815,12 +852,12 @@ void FMaterialShaderMapId::AppendKeyString(FString& KeyString) const
 	KeyString += BytesToHex(&TextureReferencesHash.Hash[0], sizeof(TextureReferencesHash.Hash));
 
 	KeyString += BytesToHex(&BasePropertyOverridesHash.Hash[0], sizeof(BasePropertyOverridesHash.Hash));
-#endif // WITH_EDITOR
 }
+#endif // WITH_EDITOR
 
+#if WITH_EDITOR
 void FMaterialShaderMapId::SetShaderDependencies(const TArray<FShaderType*>& ShaderTypes, const TArray<const FShaderPipelineType*>& ShaderPipelineTypes, const TArray<FVertexFactoryType*>& VFTypes, EShaderPlatform ShaderPlatform)
 {
-#if WITH_EDITOR
 	if (!FPlatformProperties::RequiresCookedData())
 	{
 		for (int32 ShaderTypeIndex = 0; ShaderTypeIndex < ShaderTypes.Num(); ShaderTypeIndex++)
@@ -848,8 +885,8 @@ void FMaterialShaderMapId::SetShaderDependencies(const TArray<FShaderType*>& Sha
 			ShaderPipelineTypeDependencies.Add(Dependency);
 		}
 	}
-#endif // WITH_EDITOR
 }
+#endif // WITH_EDITOR
 
 /**
  * Enqueues a compilation for a new shader of this type.
@@ -975,7 +1012,7 @@ FShader* FMaterialShaderType::FinishCompileShader(
 TRefCountPtr<FMaterialShaderMap> FMaterialShaderMap::FindId(const FMaterialShaderMapId& ShaderMapId, EShaderPlatform InPlatform)
 {
 	FScopeLock ScopeLock(&GIdToMaterialShaderMapCS);
-	check(ShaderMapId.BaseMaterialId != FGuid());
+	check(ShaderMapId.IsValid());
 	TRefCountPtr<FMaterialShaderMap> Result = GIdToMaterialShaderMap[InPlatform].FindRef(ShaderMapId);
 	check(Result == nullptr || !Result->bDeletedThroughDeferredCleanup);
 	return Result;
@@ -1007,8 +1044,9 @@ void FMaterialShaderMap::FlushShaderTypes(TArray<FShaderType*>& ShaderTypesToFlu
 
 void FMaterialShaderMap::FixupShaderTypes(EShaderPlatform Platform, const TMap<FShaderType*, FString>& ShaderTypeNames, const TMap<const FShaderPipelineType*, FString>& ShaderPipelineTypeNames, const TMap<FVertexFactoryType*, FString>& VertexFactoryTypeNames)
 {
+#if WITH_EDITOR
 	FScopeLock ScopeLock(&GIdToMaterialShaderMapCS);
-
+	
 	TArray<FMaterialShaderMapId> Keys;
 	FMaterialShaderMap::GIdToMaterialShaderMap[Platform].GenerateKeyArray(Keys);
 
@@ -1045,8 +1083,10 @@ void FMaterialShaderMap::FixupShaderTypes(EShaderPlatform Platform, const TMap<F
 
 		FMaterialShaderMap::GIdToMaterialShaderMap[Platform].Add(Key, Values[PairIndex]);
 	}
+#endif // WITH_EDITOR
 }
 
+#if WITH_EDITOR
 void FMaterialShaderMap::LoadFromDerivedDataCache(const FMaterial* Material, const FMaterialShaderMapId& ShaderMapId, EShaderPlatform InPlatform, TRefCountPtr<FMaterialShaderMap>& InOutShaderMap)
 {
 	if (InOutShaderMap != NULL)
@@ -1103,6 +1143,7 @@ void FMaterialShaderMap::SaveToDerivedDataCache()
 	GetDerivedDataCacheRef().Put(*GetMaterialShaderMapKeyString(ShaderMapId, GetShaderPlatform()), SaveData);
 	COOK_STAT(Timer.AddMiss(SaveData.Num()));
 }
+#endif // WITH_EDITOR
 
 TArray<uint8>* FMaterialShaderMap::BackupShadersToMemory()
 {
@@ -1405,7 +1446,7 @@ void FMaterialShaderMap::Compile(
 			Platform = InPlatform;
 			bIsPersistent = Material->IsPersistent();
 
-#if ALLOW_SHADERMAP_DEBUG_DATA
+#if ALLOW_SHADERMAP_DEBUG_DATA && WITH_EDITOR
 			// Store the material name for debugging purposes.
 			// Note: Material instances with static parameters will have the same FriendlyName for their shader maps!
 			FriendlyName = Material->GetFriendlyName();
@@ -1462,10 +1503,13 @@ void FMaterialShaderMap::Compile(
 			IdParameterSetHash.Final();
 			uint32* Hash = (uint32*)&IdParameterSetHash.m_digest[0];
 			FString DebugExtension = FString::Printf( TEXT("_%08x%08x"), Hash[0], Hash[1]);
+#elif ALLOW_SHADERMAP_DEBUG_DATA && !WITH_EDITOR
+			FString DebugExtension = "";
+			DebugDescription = "";
 #else
 			FString DebugExtension = "";
 			FString DebugDescription = "";
-#endif
+#endif // ALLOW_SHADERMAP_DEBUG_DATA && WITH_EDITOR
 
 			uint32 NumShaders = 0;
 			uint32 NumVertexFactories = 0;
@@ -1527,9 +1571,10 @@ void FMaterialShaderMap::Compile(
 				FMaterialShaderType* ShaderType = ShaderTypeIt->GetMaterialShaderType();
 				if (ShaderType &&  ShouldCacheMaterialShader(ShaderType, InPlatform, Material))
 				{
+#if WITH_EDITOR
 					// Verify that the shader map Id contains inputs for any shaders that will be put into this shader map
 					check(InShaderMapId.ContainsShaderType(ShaderType));
-  
+#endif
 					// Compile this material shader for this material.
 					TArray<FString> ShaderErrors;
   
@@ -1565,8 +1610,10 @@ void FMaterialShaderMap::Compile(
 						FMaterialShaderType* ShaderType = (FMaterialShaderType*)(StageTypes[Index]->GetMaterialShaderType());
 						if (ShaderType && ShouldCacheMaterialShader(ShaderType, InPlatform, Material))
 						{
+#if WITH_EDITOR
 							// Verify that the shader map Id contains inputs for any shaders that will be put into this shader map
 							check(InShaderMapId.ContainsShaderType(ShaderType));
+#endif
 							ShaderStagesToCompile.Add(ShaderType);
 						}
 						else
@@ -1578,7 +1625,9 @@ void FMaterialShaderMap::Compile(
 					if (ShaderStagesToCompile.Num() == StageTypes.Num())
 					{
 						// Verify that the shader map Id contains inputs for any shaders that will be put into this shader map
+#if WITH_EDITOR
 						check(InShaderMapId.ContainsShaderPipelineType(Pipeline));
+#endif
 						if (Pipeline->ShouldOptimizeUnusedOutputs(InPlatform))
 						{
 							NumShaders += ShaderStagesToCompile.Num();
@@ -1703,6 +1752,7 @@ FShader* FMaterialShaderMap::ProcessCompilationResultsForSingleJob(FShaderCompil
 	return Shader;
 }
 
+#if WITH_EDITOR
 bool FMaterialShaderMap::ProcessCompilationResults(const TArray<FShaderCommonCompileJob*>& InCompilationResults, int32& InOutJobIndex, float& TimeBudget, TMap<const FVertexFactoryType*, TArray<const FShaderPipelineType*> >& SharedPipelines)
 {
 	check(InOutJobIndex < InCompilationResults.Num());
@@ -1871,6 +1921,7 @@ bool FMaterialShaderMap::ProcessCompilationResults(const TArray<FShaderCommonCom
 
 	return false;
 }
+#endif // WITH_EDITOR
 
 bool FMaterialShaderMap::TryToAddToExistingCompilationTask(FMaterial* Material)
 {
@@ -1996,6 +2047,7 @@ bool FMaterialShaderMap::IsComplete(const FMaterial* Material, bool bSilent)
 	return true;
 }
 
+#if WITH_EDITOR
 void FMaterialShaderMap::LoadMissingShadersFromMemory(const FMaterial* Material)
 {
 	// Make sure we are operating on a referenced shader map or the below Find will cause this shader map to be deleted,
@@ -2090,6 +2142,7 @@ void FMaterialShaderMap::LoadMissingShadersFromMemory(const FMaterial* Material)
 		}
 	}
 }
+#endif // WITH_EDITOR
 
 #if WITH_EDITOR
 const FString* FMaterialShaderMap::GetShaderSource(const FName ShaderTypeName) const
diff --git a/Engine/Source/Runtime/Engine/Private/Materials/MaterialShared.cpp b/Engine/Source/Runtime/Engine/Private/Materials/MaterialShared.cpp
--- a/Engine/Source/Runtime/Engine/Private/Materials/MaterialShared.cpp
+++ b/Engine/Source/Runtime/Engine/Private/Materials/MaterialShared.cpp
@@ -493,6 +493,7 @@ void FMaterial::GetShaderMapId(EShaderPlatform Platform, FMaterialShaderMapId& O
 	}
 	else
 	{
+#if WITH_EDITOR
 		TArray<FShaderType*> ShaderTypes;
 		TArray<FVertexFactoryType*> VFTypes;
 		TArray<const FShaderPipelineType*> ShaderPipelineTypes;
@@ -505,6 +506,12 @@ void FMaterial::GetShaderMapId(EShaderPlatform Platform, FMaterialShaderMapId& O
 		OutId.FeatureLevel = GetFeatureLevel();
 		OutId.SetShaderDependencies(ShaderTypes, ShaderPipelineTypes, VFTypes, Platform);
 		GetReferencedTexturesHash(Platform, OutId.TextureReferencesHash);
+#else
+		OutId.QualityLevel = GetQualityLevelForShaderMapId();
+		OutId.FeatureLevel = GetFeatureLevel();
+
+		UE_LOG(LogMaterial, Warning, TEXT("Tried to access an uncooked shader map ID in a cooked application"));
+#endif
 	}
 }
 
@@ -1629,8 +1636,7 @@ bool FMaterial::CacheShaders(EShaderPlatform Platform, bool bApplyCompletedShade
 bool FMaterial::CacheShaders(const FMaterialShaderMapId& ShaderMapId, EShaderPlatform Platform, bool bApplyCompletedShaderMapForRendering)
 {
 	bool bSucceeded = false;
-
-	check(ShaderMapId.BaseMaterialId.IsValid());
+	check(ShaderMapId.IsValid());
 
 	// If we loaded this material with inline shaders, use what was loaded (GameThreadShaderMap) instead of looking in the DDC
 	if (bContainsInlineShaders)
@@ -1657,6 +1663,7 @@ bool FMaterial::CacheShaders(const FMaterialShaderMapId& ShaderMapId, EShaderPla
 	}
 	else
 	{
+#if WITH_EDITOR
 		// Find the material's cached shader map.
 		GameThreadShaderMap = FMaterialShaderMap::FindId(ShaderMapId, Platform);
 
@@ -1665,6 +1672,7 @@ bool FMaterial::CacheShaders(const FMaterialShaderMapId& ShaderMapId, EShaderPla
 		{
 			FMaterialShaderMap::LoadFromDerivedDataCache(this, ShaderMapId, Platform, GameThreadShaderMap);
 		}
+#endif // WITH_EDITOR
 	}
 
 	UMaterialInterface* MaterialInterface = GetMaterialInterface();
@@ -2993,7 +3001,7 @@ bool FMaterialInstanceBasePropertyOverrides::operator!=(const FMaterialInstanceB
 }
 
 //////////////////////////////////////////////////////////////////////////
-
+#if WITH_EDITOR
 bool FMaterialShaderMapId::ContainsShaderType(const FShaderType* ShaderType) const
 {
 	for (int32 TypeIndex = 0; TypeIndex < ShaderTypeDependencies.Num(); TypeIndex++)
@@ -3031,7 +3039,9 @@ bool FMaterialShaderMapId::ContainsVertexFactoryType(const FVertexFactoryType* V
 	}
 
 	return false;
-}//////////////////////////////////////////////////////////////////////////
+}
+#endif // WITH_EDITOR
+//////////////////////////////////////////////////////////////////////////
 
 FMaterialAttributeDefintion::FMaterialAttributeDefintion(
 		const FGuid& InAttributeID, const FString& InDisplayName, EMaterialProperty InProperty,
diff --git a/Engine/Source/Runtime/Engine/Private/Materials/MeshMaterialShader.cpp b/Engine/Source/Runtime/Engine/Private/Materials/MeshMaterialShader.cpp
--- a/Engine/Source/Runtime/Engine/Private/Materials/MeshMaterialShader.cpp
+++ b/Engine/Source/Runtime/Engine/Private/Materials/MeshMaterialShader.cpp
@@ -213,8 +213,10 @@ uint32 FMeshMaterialShaderMap::BeginCompile(
 		if (ShaderType && ShouldCacheMeshShader(ShaderType, InPlatform, Material, VertexFactoryType))
 		{
 			// Verify that the shader map Id contains inputs for any shaders that will be put into this shader map
+#if WITH_EDITOR
 			check(InShaderMapId.ContainsVertexFactoryType(VertexFactoryType));
 			check(InShaderMapId.ContainsShaderType(ShaderType));
+#endif
 
 			NumShadersPerVF++;
 			// only compile the shader if we don't already have it
@@ -263,8 +265,9 @@ uint32 FMeshMaterialShaderMap::BeginCompile(
 			if (NumShaderStagesToCompile == StageTypes.Num())
 			{
 				// Verify that the shader map Id contains inputs for any shaders that will be put into this shader map
+#if WITH_EDITOR
 				check(InShaderMapId.ContainsShaderPipelineType(Pipeline));
-
+#endif
 				if (Pipeline->ShouldOptimizeUnusedOutputs(InPlatform))
 				{
 					NumShadersPerVF += NumShaderStagesToCompile;
@@ -274,8 +277,10 @@ uint32 FMeshMaterialShaderMap::BeginCompile(
 						const FMeshMaterialShaderType* ShaderType = Shader->GetMeshMaterialShaderType();
 
 						// Verify that the shader map Id contains inputs for any shaders that will be put into this shader map
+#if WITH_EDITOR
 						check(InShaderMapId.ContainsVertexFactoryType(VertexFactoryType));
 						check(InShaderMapId.ContainsShaderType(ShaderType));
+#endif
 						ShaderStagesToCompile.Add((FMeshMaterialShaderType*)ShaderType);
 					}
 
@@ -394,6 +399,7 @@ bool FMeshMaterialShaderMap::IsComplete(
 	return true;
 }
 
+#if WITH_EDITOR
 void FMeshMaterialShaderMap::LoadMissingShadersFromMemory(
 	const FSHAHash& MaterialShaderMapHash, 
 	const FMaterial* Material, 
@@ -463,6 +469,7 @@ void FMeshMaterialShaderMap::LoadMissingShadersFromMemory(
 		}
 	}
 }
+#endif // WITH_EDITOR
 
 /**
  * Removes all entries in the cache with exceptions based on a shader type
diff --git a/Engine/Source/Runtime/Engine/Private/ShaderCompiler/ShaderCompiler.cpp b/Engine/Source/Runtime/Engine/Private/ShaderCompiler/ShaderCompiler.cpp
--- a/Engine/Source/Runtime/Engine/Private/ShaderCompiler/ShaderCompiler.cpp
+++ b/Engine/Source/Runtime/Engine/Private/ShaderCompiler/ShaderCompiler.cpp
@@ -1928,6 +1928,7 @@ void FShaderCompilingManager::ProcessCompiledShaderMaps(
 	TMap<int32, FShaderMapFinalizeResults>& CompiledShaderMaps, 
 	float TimeBudget)
 {
+#if WITH_EDITOR
 	// Keeps shader maps alive as they are passed from the shader compiler and applied to the owning FMaterial
 	TArray<TRefCountPtr<FMaterialShaderMap> > LocalShaderMapReferences;
 	TMap<FMaterial*, FMaterialShaderMap*> MaterialsToUpdate;
@@ -2010,9 +2011,7 @@ void FShaderCompilingManager::ProcessCompiledShaderMaps(
 
 					UE_LOG(LogTemp, Display, TEXT("Marking material as finished 0x%08X%08X"), (int)((int64)(Material) >> 32), (int)((int64)(Material)));
 #endif
-#if WITH_EDITOR
 					Material->RemoveOutstandingCompileId(ShaderMap->CompilingId);
-#endif // WITH_EDITOR
 
 					// Only process results that still match the ID which requested a compile
 					// This avoids applying shadermaps which are out of date and a newer one is in the async compiling pipeline
@@ -2020,10 +2019,8 @@ void FShaderCompilingManager::ProcessCompiledShaderMaps(
 					{
 						if (!bSuccess)
 						{
-#if WITH_EDITOR
 							// Propagate error messages
 							Material->CompileErrors = Errors;
-#endif // WITH_EDITOR
 
 							MaterialsToUpdate.Add( Material, NULL );
 
@@ -2128,21 +2125,18 @@ void FShaderCompilingManager::ProcessCompiledShaderMaps(
 
 		SetShaderMapsOnMaterialResources(MaterialsToApplyToScene);
 
-#if WITH_EDITOR
 		for (TMap<FMaterial*, FMaterialShaderMap*>::TIterator It(MaterialsToUpdate); It; ++It)
 		{
 			FMaterial* Material = It.Key();
 
 			Material->NotifyCompilationFinished();
 		}
-#endif // WITH_EDITOR
 
 		PropagateMaterialChangesToPrimitives(MaterialsToUpdate);
 
-#if WITH_EDITOR
 		FEditorSupportDelegates::RedrawAllViewports.Broadcast();
-#endif // WITH_EDITOR
 	}
+#endif // WITH_EDITOR
 }
 
 
diff --git a/Engine/Source/Runtime/Engine/Public/MaterialShared.h b/Engine/Source/Runtime/Engine/Public/MaterialShared.h
--- a/Engine/Source/Runtime/Engine/Public/MaterialShared.h
+++ b/Engine/Source/Runtime/Engine/Public/MaterialShared.h
@@ -523,13 +523,18 @@ namespace EMaterialShaderMapUsage
 class FMaterialShaderMapId
 {
 public:
+#if !WITH_EDITOR
+	FSHAHash CookedShaderMapIdHash;
+#endif
 
+#if WITH_EDITOR
 	/** 
 	 * The base material's StateId.  
 	 * This guid represents all the state of a UMaterial that is not covered by the other members of FMaterialShaderMapId.
 	 * Any change to the UMaterial that modifies that state (for example, adding an expression) must modify this guid.
 	 */
 	FGuid BaseMaterialId;
+#endif
 
 	/** 
 	 * Quality level that this shader map is going to be compiled at.  
@@ -540,6 +545,7 @@ public:
 	/** Feature level that the shader map is going to be compiled for. */
 	ERHIFeatureLevel::Type FeatureLevel;
 
+#if WITH_EDITOR
 	/** 
 	 * Indicates what use case this shader map will be for.
 	 * This allows the same UMaterial / UMaterialInstance to be compiled with multiple FMaterial derived classes,
@@ -577,39 +583,61 @@ public:
 	
 	/** A hash of the base property overrides for this material instance. */
 	FSHAHash BasePropertyOverridesHash;
+#endif // WITH_EDITOR
 	
-
 	FMaterialShaderMapId()
-		: BaseMaterialId(0, 0, 0, 0)
-		, QualityLevel(EMaterialQualityLevel::High)
+		: QualityLevel(EMaterialQualityLevel::High)
 		, FeatureLevel(ERHIFeatureLevel::SM4)
+#if WITH_EDITOR
 		, Usage(EMaterialShaderMapUsage::Default)
+#endif
 	{ }
 
 	~FMaterialShaderMapId()
 	{ }
 
+#if WITH_EDITOR
 	ENGINE_API void SetShaderDependencies(const TArray<FShaderType*>& ShaderTypes, const TArray<const FShaderPipelineType*>& ShaderPipelineTypes, const TArray<FVertexFactoryType*>& VFTypes, EShaderPlatform ShaderPlatform);
+#endif
 
 	void Serialize(FArchive& Ar);
 
+	bool IsValid() const
+	{
+#if WITH_EDITOR
+		return BaseMaterialId.IsValid();
+#else
+		return (CookedShaderMapIdHash != FSHAHash());
+#endif
+	}
+
 	friend uint32 GetTypeHash(const FMaterialShaderMapId& Ref)
 	{
+#if WITH_EDITOR
 		return Ref.BaseMaterialId.A;
+#else
+		// Using the hash value directly instead of FSHAHash CRC as fairly uniform distribution
+		return *(uint32*)&Ref.CookedShaderMapIdHash.Hash[0];
+#endif
 	}
 
 	SIZE_T GetSizeBytes() const
 	{
 		return sizeof(*this)
+#if WITH_EDITOR
 			+ ReferencedFunctions.GetAllocatedSize()
 			+ ReferencedParameterCollections.GetAllocatedSize()
 			+ ShaderTypeDependencies.GetAllocatedSize()
 			+ ShaderPipelineTypeDependencies.GetAllocatedSize()
-			+ VertexFactoryTypeDependencies.GetAllocatedSize();
+			+ VertexFactoryTypeDependencies.GetAllocatedSize()
+#endif
+			;
 	}
 
+#if WITH_EDITOR
 	/** Hashes the material-specific part of this shader map Id. */
 	void GetMaterialHash(FSHAHash& OutHash) const;
+#endif
 
 	/** 
 	* Tests this set against another for equality, disregarding override settings.
@@ -624,6 +652,7 @@ public:
 		return !(*this == ReferenceSet);
 	}
 
+#if WITH_EDITOR
 	/** Updates the Id's static parameter set */	
 	void UpdateParameterSet(const FStaticParameterSet& StaticParameters);
 	
@@ -648,6 +677,7 @@ public:
 
 	/** Returns true if the requested vertex factory type is a dependency of this shader map Id. */
 	bool ContainsVertexFactoryType(const FVertexFactoryType* VFType) const;
+#endif // WITH_EDITOR
 };
 
 /**
@@ -694,10 +724,12 @@ public:
 		bool bSilent
 		);
 
+#if WITH_EDITOR
 	void LoadMissingShadersFromMemory(
 		const FSHAHash& MaterialShaderMapHash, 
 		const FMaterial* Material, 
 		EShaderPlatform Platform);
+#endif
 
 	/**
 	 * Removes all entries in the cache with exceptions based on a shader type
@@ -746,11 +778,13 @@ public:
 		const TMap<const FShaderPipelineType*, FString>& ShaderPipelineTypeNames,
 		const TMap<FVertexFactoryType*, FString>& VertexFactoryTypeNames);
 
+#if WITH_EDITOR
 	/** 
 	 * Attempts to load the shader map for the given material from the Derived Data Cache.
 	 * If InOutShaderMap is valid, attempts to load the individual missing shaders instead.
 	 */
 	static void LoadFromDerivedDataCache(const FMaterial* Material, const FMaterialShaderMapId& ShaderMapId, EShaderPlatform Platform, TRefCountPtr<FMaterialShaderMap>& InOutShaderMap);
+#endif
 
 	inline FMaterialShaderMap() : FMaterialShaderMap(EShaderPlatform::SP_NumPlatforms) {}
 	FMaterialShaderMap(EShaderPlatform InPlatform);
@@ -774,8 +808,10 @@ public:
 		bool bApplyCompletedShaderMapForRendering
 		);
 
+#if WITH_EDITOR
 	/** Sorts the incoming compiled jobs into the appropriate mesh shader maps, and finalizes this shader map so that it can be used for rendering. */
 	bool ProcessCompilationResults(const TArray<FShaderCommonCompileJob*>& InCompilationResults, int32& ResultIndex, float& TimeBudget, TMap<const FVertexFactoryType*, TArray<const FShaderPipelineType*> >& SharedPipelines);
+#endif
 
 	/**
 	 * Checks whether the material shader map is missing any shader types necessary for the given material.
@@ -784,8 +820,10 @@ public:
 	 */
 	bool IsComplete(const FMaterial* Material, bool bSilent);
 
+#if WITH_EDITOR
 	/** Attempts to load missing shaders from memory. */
 	void LoadMissingShadersFromMemory(const FMaterial* Material);
+#endif
 
 	/**
 	 * Checks to see if the shader map is already being compiled for another material, and if so
@@ -842,8 +880,10 @@ public:
 	/** Serializes the shader map. */
 	void Serialize(FArchive& Ar, bool bInlineShaderResources=true, bool bLoadedByCookedMaterial=false);
 
+#if WITH_EDITOR
 	/** Saves this shader map to the derived data cache. */
 	void SaveToDerivedDataCache();
+#endif
 
 	/** Registers all shaders that have been loaded in Serialize */
 	virtual void RegisterSerializedShaders(bool bLoadedByCookedMaterial) override;
diff --git a/Engine/Source/Runtime/Landscape/Private/LandscapeRender.cpp b/Engine/Source/Runtime/Landscape/Private/LandscapeRender.cpp
--- a/Engine/Source/Runtime/Landscape/Private/LandscapeRender.cpp
+++ b/Engine/Source/Runtime/Landscape/Private/LandscapeRender.cpp
@@ -3711,6 +3711,7 @@ public:
 	{
 		FMaterialResource::GetShaderMapId(Platform, OutId);
 
+#if WITH_EDITOR
 		if (bIsLayerThumbnail || bDisableTessellation)
 		{
 			FSHA1 Hash;
@@ -3722,6 +3723,7 @@ public:
 			Hash.Final();
 			Hash.GetHash(OutId.BasePropertyOverridesHash.Hash);
 		}
+#endif
 	}
 
 	bool IsUsedWithLandscape() const override
