Force component updates on game thread when there is no rendering thread (copied CL# 5358146)

diff --git a/Engine/Source/Runtime/Engine/Private/LevelTick.cpp b/Engine/Source/Runtime/Engine/Private/LevelTick.cpp
--- a/Engine/Source/Runtime/Engine/Private/LevelTick.cpp
+++ b/Engine/Source/Runtime/Engine/Private/LevelTick.cpp
@@ -903,7 +903,10 @@ void UWorld::MarkActorComponentForNeededEndOfFrameUpdate(UActorComponent* Compon
 		{
 			bool bAllowConcurrentUpdates = FApp::ShouldUseThreadingForPerformance() && 
 				(GIsEditor ? !!CVarAllowAsyncRenderThreadUpdatesEditor.GetValueOnGameThread() : !!CVarAllowAsyncRenderThreadUpdates.GetValueOnGameThread());
-			bForceGameThread = !bAllowConcurrentUpdates;
+			bForceGameThread = !bAllowConcurrentUpdates 
+								// When there is no rendering thread force all updates on game thread,
+								// to avoid modifying scene structures from multiple task threads
+								|| !GRenderingThread;
 		}
 
 		if (bForceGameThread)
