Fixed socket transforms on slave meshes where sockets use a bone not present in the master mesh

We now cache the relative ref pose transform from the closest common ancestor and apply that. This means that GetBoneTransform/GetSocketTransform results now match the rendered output (i.e. the refpose bones that are supplied in UpdateRefToLocalMatricesInner).

diff --git a/Engine/Source/Runtime/Engine/Classes/Components/SkinnedMeshComponent.h b/Engine/Source/Runtime/Engine/Classes/Components/SkinnedMeshComponent.h
--- a/Engine/Source/Runtime/Engine/Classes/Components/SkinnedMeshComponent.h
+++ b/Engine/Source/Runtime/Engine/Classes/Components/SkinnedMeshComponent.h
@@ -273,6 +273,35 @@ protected:
 	 */
 	TArray<int32> MasterBoneMap;
 
+	/** Cached relative transform for slave bones that are missing in the master */
+	struct FMissingMasterBoneCacheEntry
+	{
+		FMissingMasterBoneCacheEntry()
+			: RelativeTransform(FTransform::Identity)
+			, CommonAncestorBoneIndex(INDEX_NONE)
+		{}
+
+		FMissingMasterBoneCacheEntry(const FTransform& InRelativeTransform, int32 InCommonAncestorBoneIndex)
+			: RelativeTransform(InRelativeTransform)
+			, CommonAncestorBoneIndex(InCommonAncestorBoneIndex)
+		{}
+
+		/** 
+		 * Relative transform of the missing bone's ref pose, based on the earliest common ancestor 
+		 * this will be equivalent to the component space transform of the bone had it existed in the master. 
+		 */
+		FTransform RelativeTransform;
+
+		/** The index of the earliest common ancestor of the master mesh. Index is the bone index in *this* mesh. */
+		int32 CommonAncestorBoneIndex;
+	};
+
+	/**  
+	 * Map of missing bone indices->transforms so that calls to GetBoneTransform() succeed when bones are not
+	 * present in a master mesh when using master-pose. Index key is the bone index of *this* mesh.
+	 */
+	TMap<int32, FMissingMasterBoneCacheEntry> MissingMasterBoneMap;
+
 	/**
 	*	Mapping for socket overrides, key is the Source socket name and the value is the override socket name
 	*/
@@ -1341,6 +1370,12 @@ private:
 	*/
 	virtual void RefreshMorphTargets() {};
 
+	/**  
+	 * When bones are not resent in a master mesh when using master-pose, we call this to evaluate 
+	 * relative transforms.
+	 */
+	bool GetMissingMasterBoneRelativeTransform(int32 InBoneIndex, FMissingMasterBoneCacheEntry& OutInfo) const;
+
 	// Animation update rate control.
 public:
 	/** Delegate when AnimUpdateRateParams is created, to override its default settings. */
diff --git a/Engine/Source/Runtime/Engine/Private/Components/SkinnedMeshComponent.cpp b/Engine/Source/Runtime/Engine/Private/Components/SkinnedMeshComponent.cpp
--- a/Engine/Source/Runtime/Engine/Private/Components/SkinnedMeshComponent.cpp
+++ b/Engine/Source/Runtime/Engine/Private/Components/SkinnedMeshComponent.cpp
@@ -1132,16 +1132,29 @@ FTransform USkinnedMeshComponent::GetBoneTransform(int32 BoneIdx, const FTransfo
 		}
 		if(BoneIdx < MasterBoneMap.Num())
 		{
-			int32 ParentBoneIndex = MasterBoneMap[BoneIdx];
+			int32 MasterBoneIndex = MasterBoneMap[BoneIdx];
 
 			// If ParentBoneIndex is valid, grab matrix from MasterPoseComponent.
-			if(	ParentBoneIndex != INDEX_NONE && 
-				ParentBoneIndex < MasterPoseComponentInst->GetNumComponentSpaceTransforms())
+			if(	MasterBoneIndex != INDEX_NONE && 
+				MasterBoneIndex < MasterPoseComponentInst->GetNumComponentSpaceTransforms())
 			{
-				return MasterPoseComponentInst->GetComponentSpaceTransforms()[ParentBoneIndex] * LocalToWorld;
+				return MasterPoseComponentInst->GetComponentSpaceTransforms()[MasterBoneIndex] * LocalToWorld;
 			}
 			else
 			{
+				// Is this a missing bone we have cached?
+				FMissingMasterBoneCacheEntry MissingBoneInfo;
+				const FMissingMasterBoneCacheEntry* MissingBoneInfoPtr = MissingMasterBoneMap.Find(BoneIdx);
+				if(MissingBoneInfoPtr != nullptr)
+				{
+					return MissingBoneInfoPtr->RelativeTransform * MasterPoseComponentInst->GetComponentSpaceTransforms()[MissingBoneInfoPtr->CommonAncestorBoneIndex] * LocalToWorld;
+				}
+				// Otherwise we might be able to generate the missing transform on the fly (although this is expensive)
+				else if(GetMissingMasterBoneRelativeTransform(BoneIdx, MissingBoneInfo))
+				{
+					return MissingBoneInfo.RelativeTransform * MasterPoseComponentInst->GetComponentSpaceTransforms()[MissingBoneInfo.CommonAncestorBoneIndex] * LocalToWorld;
+				}
+
 				UE_LOG(LogSkinnedMeshComp, Verbose, TEXT("GetBoneTransform : ParentBoneIndex(%d) out of range of MasterPoseComponent->SpaceBases for %s"), BoneIdx, *this->GetFName().ToString() );
 				return FTransform::Identity;
 			}
@@ -1675,32 +1688,85 @@ UMorphTarget* USkinnedMeshComponent::FindMorphTarget( FName MorphTargetName ) co
 	return NULL;
 }
 
+bool USkinnedMeshComponent::GetMissingMasterBoneRelativeTransform(int32 InBoneIndex, FMissingMasterBoneCacheEntry& OutInfo) const
+{
+	const FReferenceSkeleton& SlaveRefSkeleton = SkeletalMesh->RefSkeleton;
+	check(SlaveRefSkeleton.IsValidIndex(InBoneIndex));
+	const TArray<FTransform>& BoneSpaceRefPoseTransforms = SlaveRefSkeleton.GetRefBonePose();
+
+	OutInfo.CommonAncestorBoneIndex = INDEX_NONE;
+	OutInfo.RelativeTransform = FTransform::Identity;
+
+	FTransform RelativeTransform = BoneSpaceRefPoseTransforms[InBoneIndex];
+
+	// we need to find a common base component-space transform in this skeletal mesh as it
+	// isnt present in the master, so run up the hierarchy
+	int32 CommonAncestorBoneIndex = InBoneIndex;
+	while(CommonAncestorBoneIndex != INDEX_NONE)
+	{
+		CommonAncestorBoneIndex = SlaveRefSkeleton.GetParentIndex(CommonAncestorBoneIndex);
+		if(CommonAncestorBoneIndex != INDEX_NONE)
+		{
+			OutInfo.CommonAncestorBoneIndex = MasterBoneMap[CommonAncestorBoneIndex];
+			if(OutInfo.CommonAncestorBoneIndex != INDEX_NONE)
+			{
+				OutInfo.RelativeTransform = RelativeTransform;
+				return true;
+			}
+
+			RelativeTransform = RelativeTransform * BoneSpaceRefPoseTransforms[CommonAncestorBoneIndex];
+		}
+	}
+
+	return false;
+}
 
 void USkinnedMeshComponent::UpdateMasterBoneMap()
 {
 	MasterBoneMap.Reset();
+	MissingMasterBoneMap.Reset();
 
 	if (SkeletalMesh)
 	{
 		if (USkinnedMeshComponent* MasterPoseComponentPtr = MasterPoseComponent.Get())
 		{
-			if (USkeletalMesh* ParentMesh = MasterPoseComponentPtr->SkeletalMesh)
+			if (USkeletalMesh* MasterMesh = MasterPoseComponentPtr->SkeletalMesh)
 			{
-				MasterBoneMap.AddUninitialized(SkeletalMesh->RefSkeleton.GetNum());
-				if (SkeletalMesh == ParentMesh)
+				const FReferenceSkeleton& SlaveRefSkeleton = SkeletalMesh->RefSkeleton;
+				const FReferenceSkeleton& MasterRefSkeleton = MasterMesh->RefSkeleton;
+
+				MasterBoneMap.AddUninitialized(SlaveRefSkeleton.GetNum());
+				if (SkeletalMesh == MasterMesh)
 				{
 					// if the meshes are the same, the indices must match exactly so we don't need to look them up
-					for (int32 i = 0; i < MasterBoneMap.Num(); i++)
+					for (int32 BoneIndex = 0; BoneIndex < MasterBoneMap.Num(); BoneIndex++)
 					{
-						MasterBoneMap[i] = i;
+						MasterBoneMap[BoneIndex] = BoneIndex;
 					}
 				}
 				else
 				{
-					for (int32 i = 0; i < MasterBoneMap.Num(); i++)
+					for (int32 BoneIndex = 0; BoneIndex < MasterBoneMap.Num(); BoneIndex++)
 					{
-						FName BoneName = SkeletalMesh->RefSkeleton.GetBoneName(i);
-						MasterBoneMap[i] = ParentMesh->RefSkeleton.FindBoneIndex(BoneName);
+						const FName BoneName = SlaveRefSkeleton.GetBoneName(BoneIndex);
+						MasterBoneMap[BoneIndex] = MasterRefSkeleton.FindBoneIndex(BoneName);
+					}
+
+					// Cache bones for any SOCKET bones that are missing in the master.
+					// We assume that sockets will be potentially called more often, so we
+					// leave out missing BONE transforms here to try to balance memory & performance.
+					for(USkeletalMeshSocket* Socket : SkeletalMesh->GetActiveSocketList())
+					{
+						int32 BoneIndex = SlaveRefSkeleton.FindBoneIndex(Socket->BoneName);
+						int32 MasterBoneIndex = MasterRefSkeleton.FindBoneIndex(Socket->BoneName);
+						if(BoneIndex != INDEX_NONE && MasterBoneIndex == INDEX_NONE)
+						{
+							FMissingMasterBoneCacheEntry MissingBoneInfo;
+							if(GetMissingMasterBoneRelativeTransform(BoneIndex, MissingBoneInfo))
+							{
+								MissingMasterBoneMap.Add(BoneIndex, MissingBoneInfo);
+							}
+						}
 					}
 				}
 			}
