#UE4 Reduce verbosity of warning about conflicting primary asset types.

diff --git a/Engine/Source/Runtime/Engine/Private/AssetManager.cpp b/Engine/Source/Runtime/Engine/Private/AssetManager.cpp
--- a/Engine/Source/Runtime/Engine/Private/AssetManager.cpp
+++ b/Engine/Source/Runtime/Engine/Private/AssetManager.cpp
@@ -527,7 +527,7 @@ int32 UAssetManager::ScanPathsForPrimaryAssets(FPrimaryAssetType PrimaryAssetTyp
 
 				if (!IssuedWarnings.Contains(ConflictMsg))
 				{
-					UE_LOG(LogAssetManager, Warning, TEXT("%s"), *ConflictMsg);
+					UE_LOG(LogAssetManager, Display, TEXT("%s"), *ConflictMsg);
 					IssuedWarnings.Add(ConflictMsg);
 				}
 			}
