Add logging of the network profile filename when a network profile is started

diff --git a/Engine/Source/Runtime/Engine/Private/NetworkProfiler.cpp b/Engine/Source/Runtime/Engine/Private/NetworkProfiler.cpp
--- a/Engine/Source/Runtime/Engine/Private/NetworkProfiler.cpp
+++ b/Engine/Source/Runtime/Engine/Private/NetworkProfiler.cpp
@@ -628,6 +628,8 @@ void FNetworkProfiler::TrackSessionChange( bool bShouldContinueTracking, const F
 			FileWriter = IFileManager::Get().CreateFileWriter( *FinalFileName, FILEWRITE_EvenIfReadOnly );
 			check( FileWriter );
 			
+			UE_LOG(LogNet, Log, TEXT("Network Profiler: Creating session file at location %s"), *FinalFileName);
+
 			// Reset the arrays and maps so that they will match up for the new profile.
 			NameToNameTableIndexMap.Reset();
 			NameArray.Reset();
