Fix for "Editing material after changing Decal Blend Mode may cause editor to crash".  Flush rendering commands on UMaterial::PreEditChange() as realtime rendering editor viewport render has been enqueued and the decals FMeshDecalMeshProcessor::AddMeshBatch() logic tries to render FMeshDecalsEmissivePS before the shader map has been updated in UMaterial::PostEditChangeProperty().

diff --git a/Engine/Source/Runtime/Engine/Private/Materials/Material.cpp b/Engine/Source/Runtime/Engine/Private/Materials/Material.cpp
--- a/Engine/Source/Runtime/Engine/Private/Materials/Material.cpp
+++ b/Engine/Source/Runtime/Engine/Private/Materials/Material.cpp
@@ -4280,6 +4280,9 @@ bool UMaterial::CanEditChange(const UProperty* InProperty) const
 
 void UMaterial::PreEditChange(UProperty* PropertyThatChanged)
 {
+	// In Editor realtime rendering changing a material connection (or via undo/redo) changes what is rendered on the RT (specifically in FMeshDecalMeshProcessor::AddMeshBatch() Emmissive connection test)
+	// before the Shader Map has been updated in PostEditChangeProperty() below - Viewport render command enqueue has already happenned so we need to flush that before the material change.
+	FlushRenderingCommands();
 	Super::PreEditChange(PropertyThatChanged);
 }
 
