Fix material update problems with particle component

diff --git a/Engine/Source/Runtime/Engine/Private/Particles/ParticleComponents.cpp b/Engine/Source/Runtime/Engine/Private/Particles/ParticleComponents.cpp
--- a/Engine/Source/Runtime/Engine/Private/Particles/ParticleComponents.cpp
+++ b/Engine/Source/Runtime/Engine/Private/Particles/ParticleComponents.cpp
@@ -4358,6 +4358,7 @@ void UParticleSystemComponent::SetMaterial(int32 ElementIndex, UMaterialInterfac
 		}
 	}
 	MarkRenderDynamicDataDirty();
+	MarkRenderStateDirty();
 }
 
 void UParticleSystemComponent::ClearDynamicData()
@@ -7150,7 +7151,7 @@ void UParticleSystemComponent::SetMaterialParameter(FName Name, UMaterialInterfa
 		FParticleSysParam& P = InstanceParameters[i];
 		if (P.Name == Name && P.ParamType == PSPT_Material)
 		{
-			bIsViewRelevanceDirty = (P.Material != Param) ? true : false;
+			bIsViewRelevanceDirty = bIsViewRelevanceDirty || (P.Material != Param);
 			P.Material = Param;
 			return;
 		}
@@ -7160,7 +7161,7 @@ void UParticleSystemComponent::SetMaterialParameter(FName Name, UMaterialInterfa
 	int32 NewParamIndex = InstanceParameters.AddZeroed();
 	InstanceParameters[NewParamIndex].Name = Name;
 	InstanceParameters[NewParamIndex].ParamType = PSPT_Material;
-	bIsViewRelevanceDirty = (InstanceParameters[NewParamIndex].Material != Param) ? true : false;
+	bIsViewRelevanceDirty = bIsViewRelevanceDirty || (InstanceParameters[NewParamIndex].Material != Param);
 	InstanceParameters[NewParamIndex].Material = Param;
 }
 
