UE4.22 - Fix perf warning
#rnx

diff --git a/Engine/Source/Runtime/Engine/Private/ShaderCompiler/ShaderCompiler.cpp b/Engine/Source/Runtime/Engine/Private/ShaderCompiler/ShaderCompiler.cpp
--- a/Engine/Source/Runtime/Engine/Private/ShaderCompiler/ShaderCompiler.cpp
+++ b/Engine/Source/Runtime/Engine/Private/ShaderCompiler/ShaderCompiler.cpp
@@ -3102,7 +3102,7 @@ void GlobalBeginCompileShader(
 
 	if (IsVulkanPlatform((EShaderPlatform)Target.Platform))
 	{
-		const auto* CVar = IConsoleManager::Get().FindConsoleVariable(TEXT("r.Vulkan.EnableTessellation"));
+		static const auto* CVar = IConsoleManager::Get().FindConsoleVariable(TEXT("r.Vulkan.EnableTessellation"));
 		if (CVar && CVar->GetInt() != 0)
 		{
 			Input.Environment.SetDefine(TEXT("VULKAN_ENABLE_TESSELLATION"), 1);
