Update blendspaces when animation markers change

diff --git a/Engine/Source/Runtime/Engine/Classes/Animation/BlendSpaceBase.h b/Engine/Source/Runtime/Engine/Classes/Animation/BlendSpaceBase.h
--- a/Engine/Source/Runtime/Engine/Classes/Animation/BlendSpaceBase.h
+++ b/Engine/Source/Runtime/Engine/Classes/Animation/BlendSpaceBase.h
@@ -101,6 +101,10 @@ struct FBlendSample
 
 	UPROPERTY(transient)
 	uint8 bIsValid : 1;
+
+	// Cache the samples marker data counter so that we can track if it changes and revalidate the blendspace
+	int32 CachedMarkerDataUpdateCounter;
+
 #endif // WITH_EDITORONLY_DATA
 
 	FBlendSample()
@@ -110,6 +114,7 @@ struct FBlendSample
 #if WITH_EDITORONLY_DATA
 		, bSnapToGrid(true)
 		, bIsValid(false)
+		, CachedMarkerDataUpdateCounter(INDEX_NONE)
 #endif // WITH_EDITORONLY_DATA
 	{		
 	}
@@ -121,6 +126,7 @@ struct FBlendSample
 #if WITH_EDITORONLY_DATA
 		, bSnapToGrid(bInIsSnapped)
 		, bIsValid(bInIsValid)
+		, CachedMarkerDataUpdateCounter(INDEX_NONE)
 #endif // WITH_EDITORONLY_DATA
 	{		
 	}
@@ -245,6 +251,7 @@ class UBlendSpaceBase : public UAnimationAsset, public IInterpolationIndexProvid
 	virtual bool GetAllAnimationSequencesReferred(TArray<UAnimationAsset*>& AnimationAssets, bool bRecursive = true) override;
 	virtual void ReplaceReferredAnimations(const TMap<UAnimationAsset*, UAnimationAsset*>& ReplacementMap) override;
 	virtual int32 GetMarkerUpdateCounter() const;
+	void    RuntimeValidateMarkerData();
 #endif
 	//~ End UAnimationAsset Interface
 	
diff --git a/Engine/Source/Runtime/Engine/Private/Animation/AnimSequence.cpp b/Engine/Source/Runtime/Engine/Private/Animation/AnimSequence.cpp
--- a/Engine/Source/Runtime/Engine/Private/Animation/AnimSequence.cpp
+++ b/Engine/Source/Runtime/Engine/Private/Animation/AnimSequence.cpp
@@ -21,6 +21,7 @@
 #include "Animation/AnimCompress.h"
 #include "Components/SkeletalMeshComponent.h"
 #include "Animation/AnimNotifies/AnimNotify.h"
+#include "Animation/BlendSpaceBase.h"
 #include "Animation/Rig.h"
 #include "Animation/AnimationSettings.h"
 #include "Animation/AnimCurveCompressionCodec.h"
@@ -5369,6 +5370,19 @@ void UAnimSequence::RefreshSyncMarkerDataFromAuthored()
 	{
 		UniqueMarkerNames.Empty();
 	}
+
+#if WITH_EDITOR
+	check(IsInGameThread());
+
+	// Update blend spaces that may be referencing us
+	for(TObjectIterator<UBlendSpaceBase> It; It; ++It)
+	{
+		if(!It->HasAnyFlags(RF_NeedLoad | RF_NeedPostLoad))
+		{
+			It->RuntimeValidateMarkerData();
+		}
+	}
+#endif
 }
 
 bool IsMarkerValid(const FAnimSyncMarker* Marker, bool bLooping, const TArray<FName>& ValidMarkerNames)
diff --git a/Engine/Source/Runtime/Engine/Private/Animation/BlendSpaceBase.cpp b/Engine/Source/Runtime/Engine/Private/Animation/BlendSpaceBase.cpp
--- a/Engine/Source/Runtime/Engine/Private/Animation/BlendSpaceBase.cpp
+++ b/Engine/Source/Runtime/Engine/Private/Animation/BlendSpaceBase.cpp
@@ -510,6 +510,22 @@ int32 UBlendSpaceBase::GetMarkerUpdateCounter() const
 { 
 	return MarkerDataUpdateCounter; 
 }
+
+void UBlendSpaceBase::RuntimeValidateMarkerData()
+{
+	check(IsInGameThread());
+
+	for (FBlendSample& Sample : SampleData)
+	{
+		if (Sample.Animation && Sample.CachedMarkerDataUpdateCounter != Sample.Animation->GetMarkerUpdateCounter())
+		{
+			// Revalidate data
+			ValidateSampleData();
+			return;
+		}
+	}
+}
+
 #endif // WITH_EDITOR
 
 // @todo fixme: slow approach. If the perbone gets popular, we should change this to array of weight 
@@ -779,6 +795,8 @@ void UBlendSpaceBase::ValidateSampleData()
 				AnimLength = Sample.Animation->SequenceLength;
 			}
 
+			Sample.CachedMarkerDataUpdateCounter = Sample.Animation->GetMarkerUpdateCounter();
+
 			if (Sample.Animation->AuthoredSyncMarkers.Num() > 0)
 			{
 				static const TFunction<void(TArray<FName>&, TArray<FAnimSyncMarker>&)> PopulateMarkerNameArray = [](TArray<FName>& Pattern, TArray<struct FAnimSyncMarker>& AuthoredSyncMarkers)
