When updating the instance transforms, UHierarchicalInstancedStaticMeshComponent would reallocate the instance buffers, that the vertex factor, FInstancedStaticMeshVertexFactory, has a pointer to.
Added MarkRenderStateDirty so its updated

diff --git a/Engine/Source/Runtime/Engine/Private/HierarchicalInstancedStaticMesh.cpp b/Engine/Source/Runtime/Engine/Private/HierarchicalInstancedStaticMesh.cpp
--- a/Engine/Source/Runtime/Engine/Private/HierarchicalInstancedStaticMesh.cpp
+++ b/Engine/Source/Runtime/Engine/Private/HierarchicalInstancedStaticMesh.cpp
@@ -2412,6 +2412,7 @@ void UHierarchicalInstancedStaticMeshComponent::BuildTree()
 	{
 		// this is allowed only in editor, at runtime upload will happen when buffer is built from component data
 		PerInstanceRenderData->UpdateFromCommandBuffer(InstanceUpdateCmdBuffer);
+		MarkRenderStateDirty();
 	}
 
 	// all pending edits will be updated
@@ -2716,6 +2717,7 @@ void UHierarchicalInstancedStaticMeshComponent::BuildTreeAsync()
 	{
 		// this is allowed only in editor, at runtime upload will happen when buffer is built from component data
 		PerInstanceRenderData->UpdateFromCommandBuffer(InstanceUpdateCmdBuffer);
+		MarkRenderStateDirty();
 	}
 
 	// Verify that the mesh is valid before using it.
