add missing #undef

diff --git a/Engine/Source/Runtime/Engine/Private/RendererSettings.cpp b/Engine/Source/Runtime/Engine/Private/RendererSettings.cpp
--- a/Engine/Source/Runtime/Engine/Private/RendererSettings.cpp
+++ b/Engine/Source/Runtime/Engine/Private/RendererSettings.cpp
@@ -197,3 +197,5 @@ void URendererOverrideSettings::PostEditChangeProperty(FPropertyChangedEvent& Pr
 	}
 }
 #endif // #if WITH_EDITOR
+
+#undef LOCTEXT_NAMESPACE
