Stop reset to default button clearing curve compression codec

diff --git a/Engine/Source/Runtime/Engine/Classes/Animation/AnimCurveCompressionSettings.h b/Engine/Source/Runtime/Engine/Classes/Animation/AnimCurveCompressionSettings.h
--- a/Engine/Source/Runtime/Engine/Classes/Animation/AnimCurveCompressionSettings.h
+++ b/Engine/Source/Runtime/Engine/Classes/Animation/AnimCurveCompressionSettings.h
@@ -26,9 +26,6 @@ class ENGINE_API UAnimCurveCompressionSettings : public UObject
 	//////////////////////////////////////////////////////////////////////////
 
 #if WITH_EDITORONLY_DATA
-	// UObject overrides
-	virtual void PostInitProperties() override;
-
 	/** Returns whether or not we can use these settings to compress. */
 	bool AreSettingsValid() const;
 
diff --git a/Engine/Source/Runtime/Engine/Private/Animation/AnimCurveCompressionSettings.cpp b/Engine/Source/Runtime/Engine/Private/Animation/AnimCurveCompressionSettings.cpp
--- a/Engine/Source/Runtime/Engine/Private/Animation/AnimCurveCompressionSettings.cpp
+++ b/Engine/Source/Runtime/Engine/Private/Animation/AnimCurveCompressionSettings.cpp
@@ -8,20 +8,11 @@
 UAnimCurveCompressionSettings::UAnimCurveCompressionSettings(const FObjectInitializer& ObjectInitializer)
 	: Super(ObjectInitializer)
 {
+	Codec = CreateDefaultSubobject<UAnimCurveCompressionCodec_CompressedRichCurve>(TEXT("CurveCompressionCodec"));
+	Codec->SetFlags(RF_Transactional);
 }
 
 #if WITH_EDITORONLY_DATA
-void UAnimCurveCompressionSettings::PostInitProperties()
-{
-	Super::PostInitProperties();
-
-	if (!IsTemplate())
-	{
-		// Ensure we are never null
-		Codec = NewObject<UAnimCurveCompressionCodec_CompressedRichCurve>(this, NAME_None, RF_Public);
-	}
-}
-
 bool UAnimCurveCompressionSettings::AreSettingsValid() const
 {
 	return Codec != nullptr && Codec->IsCodecValid();
