Stop users from being able to pick invalid scale compression formats

diff --git a/Engine/Source/Runtime/Engine/Private/Animation/AnimCompress_PerTrackCompression.cpp b/Engine/Source/Runtime/Engine/Private/Animation/AnimCompress_PerTrackCompression.cpp
--- a/Engine/Source/Runtime/Engine/Private/Animation/AnimCompress_PerTrackCompression.cpp
+++ b/Engine/Source/Runtime/Engine/Private/Animation/AnimCompress_PerTrackCompression.cpp
@@ -2201,6 +2201,17 @@ void UAnimCompress_PerTrackCompression::PostEditChangeProperty(struct FPropertyC
 			// and bUseAdaptiveError2 is not true
 			bUseAdaptiveError = (!bUseAdaptiveError2) && bUseAdaptiveError;
 		}
+
+		if (PropertyName == GET_MEMBER_NAME_CHECKED(UAnimCompress_PerTrackCompression, AllowedScaleFormats))
+		{
+			for (TEnumAsByte<enum AnimationCompressionFormat>& ScaleFormat : AllowedScaleFormats)
+			{
+				if (ScaleFormat == ACF_Fixed32NoW || ScaleFormat == ACF_Float32NoW)
+				{
+					ScaleFormat = ACF_None;
+				}
+			}
+		}
 	}
 }
 
