Fix CIS

diff --git a/Engine/Source/Runtime/Engine/Private/Animation/AnimCurveCompressionSettings.cpp b/Engine/Source/Runtime/Engine/Private/Animation/AnimCurveCompressionSettings.cpp
--- a/Engine/Source/Runtime/Engine/Private/Animation/AnimCurveCompressionSettings.cpp
+++ b/Engine/Source/Runtime/Engine/Private/Animation/AnimCurveCompressionSettings.cpp
@@ -8,8 +8,10 @@
 UAnimCurveCompressionSettings::UAnimCurveCompressionSettings(const FObjectInitializer& ObjectInitializer)
 	: Super(ObjectInitializer)
 {
+#if WITH_EDITORONLY_DATA
 	Codec = CreateDefaultSubobject<UAnimCurveCompressionCodec_CompressedRichCurve>(TEXT("CurveCompressionCodec"));
 	Codec->SetFlags(RF_Transactional);
+#endif
 }
 
 #if WITH_EDITORONLY_DATA
