Fix uniform sampled curve decompression

diff --git a/Engine/Source/Runtime/Engine/Private/Animation/AnimCompressionDerivedData.h b/Engine/Source/Runtime/Engine/Private/Animation/AnimCompressionDerivedData.h
--- a/Engine/Source/Runtime/Engine/Private/Animation/AnimCompressionDerivedData.h
+++ b/Engine/Source/Runtime/Engine/Private/Animation/AnimCompressionDerivedData.h
@@ -50,7 +50,7 @@ public:
 		// This is a version string that mimics the old versioning scheme. If you
 		// want to bump this version, generate a new guid using VS->Tools->Create GUID and
 		// return it here. Ex.
-		return TEXT("691013A9702847FE9A47302771888030");
+		return TEXT("AA0E2605005B4B2EBAACD66B5578C97E");
 	}
 
 	virtual FString GetPluginSpecificCacheKeySuffix() const override;
diff --git a/Engine/Source/Runtime/Engine/Private/Animation/AnimCurveCompressionCodec_UniformlySampled.cpp b/Engine/Source/Runtime/Engine/Private/Animation/AnimCurveCompressionCodec_UniformlySampled.cpp
--- a/Engine/Source/Runtime/Engine/Private/Animation/AnimCurveCompressionCodec_UniformlySampled.cpp
+++ b/Engine/Source/Runtime/Engine/Private/Animation/AnimCurveCompressionCodec_UniformlySampled.cpp
@@ -176,6 +176,8 @@ void UAnimCurveCompressionCodec_UniformlySampled::DecompressCurves(const UAnimSe
 	const int32 NumConstantCurves = *(const int32*)&Buffer[BufferOffset];
 	BufferOffset += sizeof(int32);
 
+	const int32 NumAnimatedCurves = NumCurves - NumConstantCurves;
+
 	const int32 NumSamples = *(const int32*)&Buffer[BufferOffset];
 	BufferOffset += sizeof(int32);
 
@@ -188,21 +190,19 @@ void UAnimCurveCompressionCodec_UniformlySampled::DecompressCurves(const UAnimSe
 	BufferOffset += sizeof(float);
 
 	const uint32* ConstantCurvesBitsetPtr = (const uint32*)&Buffer[BufferOffset];
-	BufferOffset += sizeof(uint32) * NumCurves;
-
 	const int32 ConstantBitsetSize = sizeof(uint32) * ((NumCurves + 31) / 32);
-	const float* ConstantSamplesPtr = (const float*)&Buffer[BufferOffset];
 	BufferOffset += ConstantBitsetSize;
 
+	const float* ConstantSamplesPtr = (const float*)&Buffer[BufferOffset];
+	BufferOffset += sizeof(float) * NumConstantCurves;
+
 	const float* AnimatedSamplesPtr = (const float*)&Buffer[BufferOffset];
-	//const int32 NumAnimatedCurves = NumCurves - NumConstantCurves;
-	//BufferOffset += sizeof(float) * NumAnimatedCurves * NumSamples;
+	BufferOffset += sizeof(float) * NumAnimatedCurves * NumSamples;
 
 	const float SamplePoint = CurrentTime * SampleRate_;
 	const int32 SampleIndex0 = FMath::Clamp(FMath::FloorToInt(SamplePoint), 0, NumSamples - 1);
 	const int32 SampleIndex1 = FMath::Min(SampleIndex0 + 1, NumSamples - 1);
 	const float InterpolationAlpha = SamplePoint - float(SampleIndex0);
-	const int32 NumAnimatedCurves = NumCurves - NumConstantCurves;
 
 	const float* AnimatedSamples0 = AnimatedSamplesPtr + (SampleIndex0 * NumAnimatedCurves);
 	const float* AnimatedSamples1 = AnimatedSamplesPtr + (SampleIndex1 * NumAnimatedCurves);
@@ -255,15 +255,15 @@ float UAnimCurveCompressionCodec_UniformlySampled::DecompressCurve(const UAnimSe
 	BufferOffset += sizeof(float);
 
 	const uint32* ConstantCurvesBitsetPtr = (const uint32*)&Buffer[BufferOffset];
-	BufferOffset += sizeof(uint32) * NumCurves;
-
 	const int32 ConstantBitsetSize = sizeof(uint32) * ((NumCurves + 31) / 32);
-	const float* ConstantSamplesPtr = (const float*)&Buffer[BufferOffset];
 	BufferOffset += ConstantBitsetSize;
+	
+	const float* ConstantSamplesPtr = (const float*)&Buffer[BufferOffset];
+	BufferOffset += sizeof(float) * NumConstantCurves;
 
 	const float* AnimatedSamplesPtr = (const float*)&Buffer[BufferOffset];
-	//const int32 NumAnimatedCurves = NumCurves - NumConstantCurves;
-	//BufferOffset += sizeof(float) * NumAnimatedCurves * NumSamples;
+	const int32 NumAnimatedCurves = NumCurves - NumConstantCurves;
+	BufferOffset += sizeof(float) * NumAnimatedCurves * NumSamples;
 
 	for (int32 CurveIndex = 0, ConstantCurveIndex = 0, AnimatedCurveIndex = 0; CurveIndex < NumCurves; ++CurveIndex)
 	{
@@ -282,7 +282,6 @@ float UAnimCurveCompressionCodec_UniformlySampled::DecompressCurve(const UAnimSe
 				const int32 SampleIndex0 = FMath::Clamp(FMath::FloorToInt(SamplePoint), 0, NumSamples - 1);
 				const int32 SampleIndex1 = FMath::Min(SampleIndex0 + 1, NumSamples - 1);
 				const float InterpolationAlpha = SamplePoint - float(SampleIndex0);
-				const int32 NumAnimatedCurves = NumCurves - NumConstantCurves;
 
 				const float Sample0 = AnimatedSamplesPtr[(SampleIndex0 * NumAnimatedCurves) + CurveIndex];
 				const float Sample1 = AnimatedSamplesPtr[(SampleIndex1 * NumAnimatedCurves) + CurveIndex];
