Stop crashing when recompressing reimported animations

diff --git a/Engine/Source/Runtime/Engine/Private/Animation/AnimSequence.cpp b/Engine/Source/Runtime/Engine/Private/Animation/AnimSequence.cpp
--- a/Engine/Source/Runtime/Engine/Private/Animation/AnimSequence.cpp
+++ b/Engine/Source/Runtime/Engine/Private/Animation/AnimSequence.cpp
@@ -2538,6 +2538,12 @@ void UAnimSequence::RequestAnimCompression(FRequestAnimCompressionParams Params)
 		return;
 	}
 
+	if (GetOutermost() == GetTransientPackage())
+	{
+		bUseRawDataOnly = true;
+		return; // Skip transient animations, they are most likely the leftovers of previous compression attempts.
+	}
+
 	if (FPlatformProperties::RequiresCookedData())
 	{
 		return;
