Raytracing GI enable/disable control exposed in the GUI

diff --git a/Engine/Source/Runtime/Engine/Classes/Engine/Scene.h b/Engine/Source/Runtime/Engine/Classes/Engine/Scene.h
--- a/Engine/Source/Runtime/Engine/Classes/Engine/Scene.h
+++ b/Engine/Source/Runtime/Engine/Classes/Engine/Scene.h
@@ -1064,7 +1064,6 @@ struct FPostProcessSettings
 
 	// Ray Tracing
 	
-	// Reflections
 	UPROPERTY(BlueprintReadWrite, Category = Overrides, meta = (PinHiddenByDefault, InlineEditConditionToggle))
 	uint32 bOverride_ReflectionsType : 1;
 
@@ -1080,7 +1079,6 @@ struct FPostProcessSettings
 	UPROPERTY(BlueprintReadWrite, Category = Overrides, meta = (PinHiddenByDefault, InlineEditConditionToggle))
 	uint32 bOverride_RayTracingReflectionsShadows : 1;
 
-	// Translucency
 	UPROPERTY(BlueprintReadWrite, Category = Overrides, meta = (PinHiddenByDefault, InlineEditConditionToggle))
 	uint32 bOverride_TranslucencyType : 1;
 
@@ -1096,6 +1094,8 @@ struct FPostProcessSettings
 	UPROPERTY(BlueprintReadWrite, Category = Overrides, meta = (PinHiddenByDefault, InlineEditConditionToggle))
 	uint32 bOverride_RayTracingTranslucencyShadows : 1;
 
+	UPROPERTY(BlueprintReadWrite, Category = Overrides, meta = (PinHiddenByDefault, InlineEditConditionToggle))
+	uint32 bOverride_RayTracingGI : 1;
 
 	UPROPERTY(BlueprintReadWrite, Category = Overrides, meta = (PinHiddenByDefault, InlineEditConditionToggle))
 	uint32 bOverride_RayTracingGIMaxBounces : 1;
@@ -1573,7 +1573,11 @@ struct FPostProcessSettings
 	UPROPERTY(interp, BlueprintReadWrite, Category="Rendering Features|Global Illumination", meta=(ClampMin = "0", UIMax = "4.0", editcondition = "bOverride_IndirectLightingIntensity", DisplayName = "Indirect Lighting Intensity"))
 	float IndirectLightingIntensity;
 
-	/** Sets the global illumination maximum bounces. */
+	/** Enables ray tracing global illumination. */
+	UPROPERTY(EditAnywhere, BlueprintReadWrite, Category = "Rendering Features|Ray Tracing Global Illumination", meta = (editcondition = "bOverride_RayTracingGI", DisplayName = "Enabled"))
+	uint32 RayTracingGI : 1;
+
+	/** Sets the ray tracing global illumination maximum bounces. */
 	UPROPERTY(EditAnywhere, BlueprintReadWrite, Category = "Rendering Features|Ray Tracing Global Illumination", meta = (ClampMin = "0", ClampMax = "50", editcondition = "bOverride_RayTracingGIMaxBounces", DisplayName = "Max. Bounces"))
 	int32 RayTracingGIMaxBounces;
 
diff --git a/Engine/Source/Runtime/Engine/Private/Scene.cpp b/Engine/Source/Runtime/Engine/Private/Scene.cpp
--- a/Engine/Source/Runtime/Engine/Private/Scene.cpp
+++ b/Engine/Source/Runtime/Engine/Private/Scene.cpp
@@ -510,6 +510,7 @@ FPostProcessSettings::FPostProcessSettings()
 	IndirectLightingColor = FLinearColor(1.0f, 1.0f, 1.0f);
 	IndirectLightingIntensity = 1.0f;
 	ColorGradingIntensity = 1.0f;
+	RayTracingGI = 0;
 	RayTracingGIMaxBounces = 1;
 	RayTracingGISamplesPerPixel = 1;
 
@@ -753,6 +754,7 @@ FPostProcessSettings::FPostProcessSettings(const FPostProcessSettings& Settings)
 	, bOverride_RayTracingTranslucencyRefractionRays(Settings.bOverride_RayTracingTranslucencyRefractionRays)
 	, bOverride_RayTracingTranslucencySamplesPerPixel(Settings.bOverride_RayTracingTranslucencySamplesPerPixel)
 	, bOverride_RayTracingTranslucencyShadows(Settings.bOverride_RayTracingTranslucencyShadows)
+	, bOverride_RayTracingGI(Settings.bOverride_RayTracingGI)
 	, bOverride_RayTracingGIMaxBounces(Settings.bOverride_RayTracingGIMaxBounces)
 	, bOverride_RayTracingGISamplesPerPixel(Settings.bOverride_RayTracingGISamplesPerPixel)
 	, bOverride_PathTracingMaxBounces(Settings.bOverride_PathTracingMaxBounces)
@@ -879,6 +881,7 @@ FPostProcessSettings::FPostProcessSettings(const FPostProcessSettings& Settings)
 	, RayTracingAOSamplesPerPixel(Settings.RayTracingAOSamplesPerPixel)
 	, IndirectLightingColor(Settings.IndirectLightingColor)
 	, IndirectLightingIntensity(Settings.IndirectLightingIntensity)
+	, RayTracingGI(Settings.RayTracingGI)
 	, RayTracingGIMaxBounces(Settings.RayTracingGIMaxBounces)
 	, RayTracingGISamplesPerPixel(Settings.RayTracingGISamplesPerPixel)
 	, ColorGradingIntensity(Settings.ColorGradingIntensity)
diff --git a/Engine/Source/Runtime/Engine/Private/SceneView.cpp b/Engine/Source/Runtime/Engine/Private/SceneView.cpp
--- a/Engine/Source/Runtime/Engine/Private/SceneView.cpp
+++ b/Engine/Source/Runtime/Engine/Private/SceneView.cpp
@@ -1413,6 +1413,11 @@ void FSceneView::OverridePostProcessSettings(const FPostProcessSettings& Src, fl
 			Dest.RayTracingTranslucencyShadows = Src.RayTracingTranslucencyShadows;
 		}
 
+		if (Src.bOverride_RayTracingGI)
+		{
+			Dest.RayTracingGI = Src.RayTracingGI;
+		}
+
 		if (Src.bOverride_RayTracingGIMaxBounces)
 		{
 			Dest.RayTracingGIMaxBounces = Src.RayTracingGIMaxBounces;
diff --git a/Engine/Source/Runtime/Renderer/Private/DeferredShadingRenderer.cpp b/Engine/Source/Runtime/Renderer/Private/DeferredShadingRenderer.cpp
--- a/Engine/Source/Runtime/Renderer/Private/DeferredShadingRenderer.cpp
+++ b/Engine/Source/Runtime/Renderer/Private/DeferredShadingRenderer.cpp
@@ -1622,7 +1622,8 @@ void FDeferredShadingSceneRenderer::Render(FRHICommandListImmediate& RHICmdList)
 				RenderRayTracingSkyLight(RHICmdList, SkyLightRT, HitDistanceRT);
 			}
 
-			if (ShouldRenderRayTracingGlobalIllumination())
+
+			if (ShouldRenderRayTracingGlobalIllumination(Views))
 			{
 				for (int32 ViewIndex = 0; ViewIndex < Views.Num(); ViewIndex++)
 				{
diff --git a/Engine/Source/Runtime/Renderer/Private/RayTracing/RayTracingGlobalIllumination.cpp b/Engine/Source/Runtime/Renderer/Private/RayTracing/RayTracingGlobalIllumination.cpp
--- a/Engine/Source/Runtime/Renderer/Private/RayTracing/RayTracingGlobalIllumination.cpp
+++ b/Engine/Source/Runtime/Renderer/Private/RayTracing/RayTracingGlobalIllumination.cpp
@@ -20,11 +20,13 @@
 #include "PostProcess/SceneFilterRendering.h"
 #include "Raytracing/RaytracingOptions.h"
 
-static int32 GRayTracingGlobalIllumination = 0;
+static int32 GRayTracingGlobalIllumination = -1;
 static FAutoConsoleVariableRef CVarRayTracingGlobalIllumination(
 	TEXT("r.RayTracing.GlobalIllumination"),
 	GRayTracingGlobalIllumination,
-	TEXT("Enabled ray tracing global illumination (default = 0)")
+	TEXT("-1: Value driven by postprocess volume (default) \n")
+	TEXT(" 0: ray tracing ray tracing global illumination off \n")
+	TEXT(" 1: ray tracing global illumination enabled")
 );
 
 static int32 GRayTracingGlobalIlluminationSamplesPerPixel = -1;
@@ -175,9 +177,26 @@ void SetupLightParameters(
 	}
 }
 
-bool ShouldRenderRayTracingGlobalIllumination()
+bool ShouldRenderRayTracingGlobalIllumination(const TArray<FViewInfo>& Views)
 {
-	return GRayTracingGlobalIllumination == 1;
+	if (GRayTracingGlobalIllumination >= 0)
+	{
+		return (GRayTracingGlobalIllumination > 0);
+	}
+	else
+	{
+		for (int32 ViewIndex = 0, Num = Views.Num(); ViewIndex < Num; ViewIndex++)
+		{
+			const FViewInfo& View = Views[ViewIndex];
+			//#dxr_todo: multiview case
+			if (View.FinalPostProcessSettings.RayTracingGI > 0)
+			{
+				return true;
+			}
+		}
+
+		return false;
+	}
 }
 
 class FGlobalIlluminationRGS : public FGlobalShader
@@ -306,7 +325,11 @@ void FDeferredShadingSceneRenderer::RenderRayTracingGlobalIllumination(
 	TRefCountPtr<IPooledRenderTarget>& AmbientOcclusionRT
 )
 {
-	if (!GRayTracingGlobalIllumination) return;
+	if (GRayTracingGlobalIllumination == 0 || (GRayTracingGlobalIllumination == -1 && View.FinalPostProcessSettings.RayTracingGI == 0)) 
+	{
+		return;
+	}
+
 	SCOPED_GPU_STAT(RHICmdList, RayTracingGlobalIllumination);
 
 	FSceneRenderTargets& SceneContext = FSceneRenderTargets::Get(RHICmdList);
diff --git a/Engine/Source/Runtime/Renderer/Private/RayTracing/RaytracingOptions.h b/Engine/Source/Runtime/Renderer/Private/RayTracing/RaytracingOptions.h
--- a/Engine/Source/Runtime/Renderer/Private/RayTracing/RaytracingOptions.h
+++ b/Engine/Source/Runtime/Renderer/Private/RayTracing/RaytracingOptions.h
@@ -13,7 +13,7 @@ extern bool ShouldRenderRayTracingSkyLight(const FSkyLightSceneProxy* SkyLightSc
 
 #if RHI_RAYTRACING
 extern bool ShouldRenderRayTracingAmbientOcclusion();
-extern bool ShouldRenderRayTracingGlobalIllumination();
+extern bool ShouldRenderRayTracingGlobalIllumination(const TArray<FViewInfo>& Views);
 extern bool ShouldRenderRayTracingStochasticRectLight(const FLightSceneInfo& LightSceneInfo);
 
 extern float GetRaytracingOcclusionMaxNormalBias();
@@ -25,7 +25,7 @@ FORCEINLINE bool ShouldRenderRayTracingAmbientOcclusion()
 	return false;
 }
 
-FORCEINLINE bool ShouldRenderRayTracingGlobalIllumination()
+FORCEINLINE bool ShouldRenderRayTracingGlobalIllumination(const TArray<FViewInfo>& Views)
 {
 	return false;
 }
