Fix world settings actors being spawned as transactional even when the world is created as not transactional.  This can cause the world settings to disappear during undo/redo.

diff --git a/Engine/Source/Runtime/Engine/Private/World.cpp b/Engine/Source/Runtime/Engine/Private/World.cpp
--- a/Engine/Source/Runtime/Engine/Private/World.cpp
+++ b/Engine/Source/Runtime/Engine/Private/World.cpp
@@ -1469,14 +1469,19 @@ void UWorld::InitializeNewWorld(const InitializationValues IVS)
 	PersistentLevel->Model->Initialize(nullptr, 1);
 	PersistentLevel->OwningWorld = this;
 
+	// Create the WorldInfo actor.
+	FActorSpawnParameters SpawnInfo; 
+
 	// Mark objects are transactional for undo/ redo.
 	if (IVS.bTransactional)
 	{
+		SpawnInfo.ObjectFlags |= RF_Transactional;
 		PersistentLevel->SetFlags( RF_Transactional );
 		PersistentLevel->Model->SetFlags( RF_Transactional );
 	}
 	else
 	{
+		SpawnInfo.ObjectFlags &= ~RF_Transactional;
 		PersistentLevel->ClearFlags( RF_Transactional );
 		PersistentLevel->Model->ClearFlags( RF_Transactional );
 	}
@@ -1486,8 +1491,6 @@ void UWorld::InitializeNewWorld(const InitializationValues IVS)
 	CurrentLevel = PersistentLevel;
 #endif
 
-	// Create the WorldInfo actor.
-	FActorSpawnParameters SpawnInfo;
 	SpawnInfo.SpawnCollisionHandlingOverride = ESpawnActorCollisionHandlingMethod::AlwaysSpawn;
 	// Set constant name for WorldSettings to make a network replication work between new worlds on host and client
 	SpawnInfo.Name = GEngine->WorldSettingsClass->GetFName();
