Fixed warm-up with managed particle components.

diff --git a/Engine/Source/Runtime/Engine/Private/Particles/ParticleComponents.cpp b/Engine/Source/Runtime/Engine/Private/Particles/ParticleComponents.cpp
--- a/Engine/Source/Runtime/Engine/Private/Particles/ParticleComponents.cpp
+++ b/Engine/Source/Runtime/Engine/Private/Particles/ParticleComponents.cpp
@@ -5147,7 +5147,7 @@ void UParticleSystemComponent::TickComponent(float DeltaTime, enum ELevelTick Ti
 	TotalActiveParticles = 0;
 	bNeedsFinalize = true;
 	
-	if (!IsTickManaged())
+	if (!IsTickManaged() || bWarmingUp)
 	{
 		if (!ThisTickFunction || !ThisTickFunction->IsCompletionHandleValid() || !CanTickInAnyThread() || FXConsoleVariables::bFreezeParticleSimulation || !FXConsoleVariables::bAllowAsyncTick || !FApp::ShouldUseThreadingForPerformance() ||
 			GDistributionType == 0) // this may not be absolutely required, however if you are using distributions it will be glacial anyway. If you want to get rid of this, note that some modules use this indirectly as their criteria for CanTickInAnyThread
