Add null check as it is valid for there to be nulls in the controller list.

diff --git a/Engine/Source/Runtime/Engine/Private/World.cpp b/Engine/Source/Runtime/Engine/Private/World.cpp
--- a/Engine/Source/Runtime/Engine/Private/World.cpp
+++ b/Engine/Source/Runtime/Engine/Private/World.cpp
@@ -5870,10 +5870,12 @@ UWorld* FSeamlessTravelHandler::Tick()
 			{
 				for( FConstControllerIterator Iterator = CurrentWorld->GetControllerIterator(); Iterator; ++Iterator )
 				{
-					AController* Player = Iterator->Get();
-					if (Player->PlayerState || Cast<APlayerController>(Player) != nullptr)
+					if (AController* Player = Iterator->Get())
 					{
-						KeepAnnotation.Set(Player);
+						if (Player->PlayerState || Cast<APlayerController>(Player) != nullptr)
+						{
+							KeepAnnotation.Set(Player);
+						}
 					}
 				}
 			}
