Added r.RayTracing.Translucency.Refraction on post process volume (defaults to 1)
Fixed the fact that it was not taken into account in the shader (CL was 5172901, was robomerged to stagin to maybe a wrong integration).

diff --git a/Engine/Shaders/Private/RayTracing/RayTracingTranslucency.usf b/Engine/Shaders/Private/RayTracing/RayTracingTranslucency.usf
--- a/Engine/Shaders/Private/RayTracing/RayTracingTranslucency.usf
+++ b/Engine/Shaders/Private/RayTracing/RayTracingTranslucency.usf
@@ -175,8 +175,6 @@ void RayTracingTranslucencyRGS()
 			break;
 		}
 		float3 HitPoint = Ray.Origin + Ray.Direction * GetHitT(Payload);
-		bHasScattered = true;
-
 
 		//
 		// Handle surface lighting
@@ -273,8 +271,8 @@ void RayTracingTranslucencyRGS()
 
 	if (!bHasScattered)
 	{
-		// Use the scene opaque radiance for pixel that has not intersected any transparent.
-		PathRadiance = Texture2DSampleLevel(SceneTexturesStruct.SceneColorTexture, SceneTexturesStruct.SceneColorTextureSampler, UV, 0).xyz;
+		// Use the scene radiance for ray that has not been scattered/refracted (no surface or IORin=IORout). Still apply the throughtput in case we have traversed surfaces with opacity>0.
+		PathRadiance += PathThroughput * Texture2DSampleLevel(SceneTexturesStruct.SceneColorTexture, SceneTexturesStruct.SceneColorTextureSampler, UV, 0).xyz;
 	}
 
 #if USE_PREEXPOSURE
diff --git a/Engine/Source/Runtime/Engine/Classes/Engine/Scene.h b/Engine/Source/Runtime/Engine/Classes/Engine/Scene.h
--- a/Engine/Source/Runtime/Engine/Classes/Engine/Scene.h
+++ b/Engine/Source/Runtime/Engine/Classes/Engine/Scene.h
@@ -1095,6 +1095,8 @@ struct FPostProcessSettings
 	uint32 bOverride_RayTracingTranslucencyShadows : 1;
 
 	UPROPERTY(BlueprintReadWrite, Category = Overrides, meta = (PinHiddenByDefault, InlineEditConditionToggle))
+	uint32 bOverride_RayTracingTranslucencyRefraction : 1;
+
 	uint32 bOverride_RayTracingGI : 1;
 
 	UPROPERTY(BlueprintReadWrite, Category = Overrides, meta = (PinHiddenByDefault, InlineEditConditionToggle))
@@ -1790,6 +1792,10 @@ struct FPostProcessSettings
 	UPROPERTY(EditAnywhere, BlueprintReadWrite, Category = "Rendering Features|Ray Tracing Translucency", meta = (editcondition = "bOverride_RayTracingTranslucencyShadows", DisplayName = "Shadows"))
 	EReflectedAndRefractedRayTracedShadows RayTracingTranslucencyShadows;
 
+	/** Sets whether refraction should be enabled or not (if not rays will not scatter and only travel in the same direction as before the intersection event). */
+	UPROPERTY(EditAnywhere, BlueprintReadWrite, Category = "Rendering Features|Ray Tracing Translucency", meta = (editcondition = "bOverride_RayTracingTranslucencyRefraction", DisplayName = "Refraction"))
+	uint8 RayTracingTranslucencyRefraction : 1;
+
 
 	// Path Tracing
 	/** Sets the path tracing maximum bounces */
diff --git a/Engine/Source/Runtime/Engine/Private/Scene.cpp b/Engine/Source/Runtime/Engine/Private/Scene.cpp
--- a/Engine/Source/Runtime/Engine/Private/Scene.cpp
+++ b/Engine/Source/Runtime/Engine/Private/Scene.cpp
@@ -564,6 +564,7 @@ FPostProcessSettings::FPostProcessSettings()
 	RayTracingTranslucencyRefractionRays = 3; // 3 to: first hit surface, second hit back inner surface and a third to fetch the background.
 	RayTracingTranslucencySamplesPerPixel = 1;
 	RayTracingTranslucencyShadows = EReflectedAndRefractedRayTracedShadows::Hard_shadows;
+	RayTracingTranslucencyRefraction = 1;
 
 	PathTracingMaxBounces = 32;
 	PathTracingSamplesPerPixel = 16384;
@@ -754,6 +755,7 @@ FPostProcessSettings::FPostProcessSettings(const FPostProcessSettings& Settings)
 	, bOverride_RayTracingTranslucencyRefractionRays(Settings.bOverride_RayTracingTranslucencyRefractionRays)
 	, bOverride_RayTracingTranslucencySamplesPerPixel(Settings.bOverride_RayTracingTranslucencySamplesPerPixel)
 	, bOverride_RayTracingTranslucencyShadows(Settings.bOverride_RayTracingTranslucencyShadows)
+	, bOverride_RayTracingTranslucencyRefraction(Settings.bOverride_RayTracingTranslucencyRefraction)
 	, bOverride_RayTracingGI(Settings.bOverride_RayTracingGI)
 	, bOverride_RayTracingGIMaxBounces(Settings.bOverride_RayTracingGIMaxBounces)
 	, bOverride_RayTracingGISamplesPerPixel(Settings.bOverride_RayTracingGISamplesPerPixel)
@@ -938,6 +940,7 @@ FPostProcessSettings::FPostProcessSettings(const FPostProcessSettings& Settings)
 	, RayTracingTranslucencyRefractionRays(Settings.RayTracingTranslucencyRefractionRays)
 	, RayTracingTranslucencySamplesPerPixel(Settings.RayTracingTranslucencySamplesPerPixel)
 	, RayTracingTranslucencyShadows(Settings.RayTracingTranslucencyShadows)
+	, RayTracingTranslucencyRefraction(Settings.RayTracingTranslucencyRefraction)
 
 	, PathTracingMaxBounces(Settings.PathTracingMaxBounces)
 	, PathTracingSamplesPerPixel(Settings.PathTracingSamplesPerPixel)
diff --git a/Engine/Source/Runtime/Engine/Private/SceneView.cpp b/Engine/Source/Runtime/Engine/Private/SceneView.cpp
--- a/Engine/Source/Runtime/Engine/Private/SceneView.cpp
+++ b/Engine/Source/Runtime/Engine/Private/SceneView.cpp
@@ -1413,6 +1413,11 @@ void FSceneView::OverridePostProcessSettings(const FPostProcessSettings& Src, fl
 			Dest.RayTracingTranslucencyShadows = Src.RayTracingTranslucencyShadows;
 		}
 
+		if (Src.bOverride_RayTracingTranslucencyRefraction)
+		{
+			Dest.RayTracingTranslucencyRefraction = Src.RayTracingTranslucencyRefraction;
+		}
+
 		if (Src.bOverride_RayTracingGI)
 		{
 			Dest.RayTracingGI = Src.RayTracingGI;
diff --git a/Engine/Source/Runtime/Renderer/Private/RayTracing/RayTracingTranslucency.cpp b/Engine/Source/Runtime/Renderer/Private/RayTracing/RayTracingTranslucency.cpp
--- a/Engine/Source/Runtime/Renderer/Private/RayTracing/RayTracingTranslucency.cpp
+++ b/Engine/Source/Runtime/Renderer/Private/RayTracing/RayTracingTranslucency.cpp
@@ -87,7 +87,7 @@ static FAutoConsoleVariableRef CVarRayTracingTranslucencyHeightFog(
 	GRayTracingTranslucencyHeightFog,
 	TEXT("Enables height fog in ray traced Translucency (default = 1)"));
 
-static int32 GRayTracingTranslucencyRefraction = 1;
+static int32 GRayTracingTranslucencyRefraction = -1;
 static FAutoConsoleVariableRef CVarRayTracingTranslucencyRefraction(
 	TEXT("r.RayTracing.Translucency.Refraction"),
 	GRayTracingTranslucencyRefraction,
@@ -364,7 +364,7 @@ void FDeferredShadingSceneRenderer::RenderRayTracingTranslucencyView(
 	PassParameters->TranslucencyMinRayDistance = FMath::Min(GRayTracingTranslucencyMinRayDistance, GRayTracingTranslucencyMaxRayDistance);
 	PassParameters->TranslucencyMaxRayDistance = GRayTracingTranslucencyMaxRayDistance;
 	PassParameters->TranslucencyMaxRoughness = FMath::Clamp(GRayTracingTranslucencyMaxRoughness >= 0 ? GRayTracingTranslucencyMaxRoughness : View.FinalPostProcessSettings.RayTracingTranslucencyMaxRoughness, 0.01f, 1.0f);
-	PassParameters->TranslucencyRefraction = GRayTracingTranslucencyRefraction;
+	PassParameters->TranslucencyRefraction = GRayTracingTranslucencyRefraction >= 0 ? GRayTracingTranslucencyRefraction : View.FinalPostProcessSettings.RayTracingTranslucencyRefraction;
 	PassParameters->MaxNormalBias = GetRaytracingOcclusionMaxNormalBias();
 
 	PassParameters->TLAS = View.RayTracingScene.RayTracingSceneRHI->GetShaderResourceView();
