Fixed out of order packet correction breaking the ping display (requires client update), and reduced log verbosity/logspam.

#sh-03-12-pingpong

diff --git a/Engine/Source/Runtime/Engine/Private/NetConnection.cpp b/Engine/Source/Runtime/Engine/Private/NetConnection.cpp
--- a/Engine/Source/Runtime/Engine/Private/NetConnection.cpp
+++ b/Engine/Source/Runtime/Engine/Private/NetConnection.cpp
@@ -1182,7 +1182,7 @@ void UNetConnection::FlushPacketOrderCache(bool bFlushWholeCache/*=false*/)
 
 			if (CurCachePacket.IsValid())
 			{
-				UE_LOG(LogNet, Verbose, TEXT("'Out of Order' Packet Cache, replaying packet with cache index: %i (bFlushWholeCache: %i)"), PacketOrderCacheStartIdx, (int32)bFlushWholeCache);
+				UE_LOG(LogNet, VeryVerbose, TEXT("'Out of Order' Packet Cache, replaying packet with cache index: %i (bFlushWholeCache: %i)"), PacketOrderCacheStartIdx, (int32)bFlushWholeCache);
 
 				ReceivedPacket(*CurCachePacket.Get());
 
@@ -1728,13 +1728,6 @@ void UNetConnection::ReceivedPacket( FBitReader& Reader )
 			Driver->InPacketsLost += MissingPacketCount;
 			Driver->InTotalPacketsLost += MissingPacketCount;
 			InPacketId += PacketSequenceDelta;
-
-			// Extra information associated with the header
-			if (!ReadPacketInfo(Reader))
-			{
-				CLOSE_CONNECTION_DUE_TO_SECURITY_VIOLATION(this, ESecurityEvent::Malformed_Packet, TEXT("Failed to read PacketHeader"));
-				return;
-			}
 		}
 		else
 		{
@@ -1790,8 +1783,16 @@ void UNetConnection::ReceivedPacket( FBitReader& Reader )
 		// Update incoming sequence data and deliver packet notifications
 		// Packet is only accepted if both the incoming sequence number and incoming ack data are valid
 		PacketNotify.Update(Header, HandlePacketNotification);
+
+		// Extra information associated with the header (read only after acks have been processed)
+		if (PacketSequenceDelta > 0 && !ReadPacketInfo(Reader))
+		{
+			CLOSE_CONNECTION_DUE_TO_SECURITY_VIOLATION(this, ESecurityEvent::Malformed_Packet, TEXT("Failed to read PacketHeader"));
+			return;
+		}
 	}
 
+
 	const bool bIgnoreRPCs = Driver->ShouldIgnoreRPCs();
 
 	bool bSkipAck = false;
