Store replicated properties in static vars for the DOREPLIFETIME_ACTIVE_OVERRIDE to improve PreReplication() performance

diff --git a/Engine/Source/Runtime/Engine/Public/Net/UnrealNetwork.h b/Engine/Source/Runtime/Engine/Public/Net/UnrealNetwork.h
--- a/Engine/Source/Runtime/Engine/Public/Net/UnrealNetwork.h
+++ b/Engine/Source/Runtime/Engine/Public/Net/UnrealNetwork.h
@@ -114,7 +114,7 @@ static UProperty* GetReplicatedProperty(const UClass* CallingClass, const UClass
 
 #define DOREPLIFETIME_ACTIVE_OVERRIDE(c,v,active)	\
 {													\
-	UProperty* sp##v = GetReplicatedProperty(StaticClass(), c::StaticClass(),GET_MEMBER_NAME_CHECKED(c,v)); \
+	static UProperty* sp##v = GetReplicatedProperty(StaticClass(), c::StaticClass(),GET_MEMBER_NAME_CHECKED(c,v)); \
 	for ( int32 i = 0; i < sp##v->ArrayDim; i++ )											\
 	{																						\
 		ChangedPropertyTracker.SetCustomIsActiveOverride( sp##v->RepIndex + i, active );	\
