#UE4 PreLoad Textures if they are needed during postload and have RF_NeedLoad. Temporary workaround while investigating UE-69421

diff --git a/Engine/Source/Runtime/Engine/Private/Materials/Material.cpp b/Engine/Source/Runtime/Engine/Private/Materials/Material.cpp
--- a/Engine/Source/Runtime/Engine/Private/Materials/Material.cpp
+++ b/Engine/Source/Runtime/Engine/Private/Materials/Material.cpp
@@ -3956,6 +3956,14 @@ void UMaterial::PostLoad()
 			{
 				if (Texture)
 				{
+					if (Texture->HasAnyFlags(RF_NeedLoad))
+					{
+						FLinkerLoad* Loader = GetLinker();
+						if (ensure(Loader))
+						{
+							Loader->Preload(Texture);
+						}
+					}
 					Texture->ConditionalPostLoad();
 				}
 			}
diff --git a/Engine/Source/Runtime/Engine/Private/Particles/ParticleModules.cpp b/Engine/Source/Runtime/Engine/Private/Particles/ParticleModules.cpp
--- a/Engine/Source/Runtime/Engine/Private/Particles/ParticleModules.cpp
+++ b/Engine/Source/Runtime/Engine/Private/Particles/ParticleModules.cpp
@@ -1219,6 +1219,14 @@ void UParticleModuleRequired::PostLoad()
 	{
 		if (CutoutTexture)
 		{
+			if (CutoutTexture->HasAnyFlags(RF_NeedLoad))
+			{
+				FLinkerLoad* Loader = GetLinker();
+				if (ensure(Loader))
+				{
+					Loader->Preload(CutoutTexture);
+				}
+			}
 			CutoutTexture->ConditionalPostLoad();
 			CacheDerivedData();
 		}	
