#UE4 Workaround fix for texture not serialized by the time it is postloaded

diff --git a/Engine/Source/Runtime/Engine/Private/Materials/MaterialInstance.cpp b/Engine/Source/Runtime/Engine/Private/Materials/MaterialInstance.cpp
--- a/Engine/Source/Runtime/Engine/Private/Materials/MaterialInstance.cpp
+++ b/Engine/Source/Runtime/Engine/Private/Materials/MaterialInstance.cpp
@@ -3127,6 +3127,14 @@ void UMaterialInstance::PostLoad()
 		UTexture* Texture = TextureParameterValues[ValueIndex].ParameterValue;
 		if( Texture )
 		{
+			if (Texture->HasAnyFlags(RF_NeedLoad))
+			{
+				FLinkerLoad* Loader = GetLinker();
+				if (ensure(Loader))
+				{
+					Loader->Preload(Texture);
+				}
+			}
 			Texture->ConditionalPostLoad();
 		}
 	}
