#UE4 Fix NaN error when calling FInstancedStaticMeshSceneProxy::SetupProxy due to doing math on MAX_flt.

diff --git a/Engine/Source/Runtime/Engine/Private/InstancedStaticMesh.cpp b/Engine/Source/Runtime/Engine/Private/InstancedStaticMesh.cpp
--- a/Engine/Source/Runtime/Engine/Private/InstancedStaticMesh.cpp
+++ b/Engine/Source/Runtime/Engine/Private/InstancedStaticMesh.cpp
@@ -2152,16 +2152,24 @@ static bool ComponentRequestsCPUAccess(UInstancedStaticMeshComponent* InComponen
 
 void UInstancedStaticMeshComponent::GetInstancesMinMaxScale(FVector& MinScale, FVector& MaxScale) const
 {
-	MinScale = FVector(MAX_flt);
-	MaxScale = FVector(-MAX_flt);
-
-	for (int32 i = 0; i < PerInstanceSMData.Num(); ++i)
+	if (PerInstanceSMData.Num() > 0)
 	{
-		const FInstancedStaticMeshInstanceData& InstanceData = PerInstanceSMData[i];
-		FVector ScaleVector = InstanceData.Transform.GetScaleVector();
+		MinScale = FVector(MAX_flt);
+		MaxScale = FVector(-MAX_flt);
+
+		for (int32 i = 0; i < PerInstanceSMData.Num(); ++i)
+		{
+			const FInstancedStaticMeshInstanceData& InstanceData = PerInstanceSMData[i];
+			FVector ScaleVector = InstanceData.Transform.GetScaleVector();
 
-		MinScale = MinScale.ComponentMin(ScaleVector);
-		MaxScale = MaxScale.ComponentMax(ScaleVector);
+			MinScale = MinScale.ComponentMin(ScaleVector);
+			MaxScale = MaxScale.ComponentMax(ScaleVector);
+		}
+	}
+	else
+	{
+		MinScale = FVector(0);
+		MaxScale = FVector(0);
 	}
 }
 
