#UE4 Adding extra logging to an assert to learn more the next time it happens

diff --git a/Engine/Source/Runtime/Engine/Private/StaticMesh.cpp b/Engine/Source/Runtime/Engine/Private/StaticMesh.cpp
--- a/Engine/Source/Runtime/Engine/Private/StaticMesh.cpp
+++ b/Engine/Source/Runtime/Engine/Private/StaticMesh.cpp
@@ -2198,7 +2198,7 @@ void FStaticMeshRenderData::Cache(UStaticMesh* Owner, const FStaticMeshLODSettin
 			Args.Add(TEXT("StaticMeshName"), FText::FromString( Owner->GetName() ) );
 			FStaticMeshStatusMessageContext StatusContext( FText::Format( NSLOCTEXT("Engine", "BuildingStaticMeshStatus", "Building static mesh {StaticMeshName}..."), Args ) );
 
-			check(Owner->IsMeshDescriptionValid(0));
+			checkf(Owner->IsMeshDescriptionValid(0), TEXT("Bad MeshDescription on %s"), *GetPathNameSafe(Owner));
 
 			IMeshBuilderModule& MeshBuilderModule = FModuleManager::Get().LoadModuleChecked<IMeshBuilderModule>(TEXT("MeshBuilder"));
 			if (!MeshBuilderModule.BuildMesh(*this, Owner, LODGroup))
