- Support TArrayView in FImageUtils::ImageResize so we can use a bigger variety of data structures (avoiding copying)
- Refactored the TArray version to use this internally

diff --git a/Engine/Source/Runtime/Engine/Private/ImageUtils.cpp b/Engine/Source/Runtime/Engine/Private/ImageUtils.cpp
--- a/Engine/Source/Runtime/Engine/Private/ImageUtils.cpp
+++ b/Engine/Source/Runtime/Engine/Private/ImageUtils.cpp
@@ -61,11 +61,30 @@ static bool GetRawData(UTextureRenderTarget2D* TexRT, TArray<uint8>& RawData)
  * @param DstData		Destination image data.
  * @param bLinearSpace	If true, convert colors into linear space before interpolating (slower but more accurate)
  */
-void FImageUtils::ImageResize(int32 SrcWidth, int32 SrcHeight, const TArray<FColor> &SrcData, int32 DstWidth, int32 DstHeight, TArray<FColor> &DstData, bool bLinearSpace )
+void FImageUtils::ImageResize(int32 SrcWidth, int32 SrcHeight, const TArray<FColor> &SrcData, int32 DstWidth, int32 DstHeight, TArray<FColor> &DstData, bool bLinearSpace)
 {
 	DstData.Empty(DstWidth*DstHeight);
 	DstData.AddZeroed(DstWidth*DstHeight);
 
+	ImageResize(SrcWidth, SrcHeight, TArrayView<const FColor>(SrcData), DstWidth, DstHeight, TArrayView<FColor>(DstData), bLinearSpace);
+}
+
+/**
+ * Resizes the given image using a simple average filter and stores it in the destination array.  This version constrains aspect ratio.
+ * Accepts TArrayViews but requires that DstData be pre-sized appropriately
+ *
+ * @param SrcWidth	Source image width.
+ * @param SrcHeight	Source image height.
+ * @param SrcData	Source image data.
+ * @param DstWidth	Destination image width.
+ * @param DstHeight Destination image height.
+ * @param DstData	Destination image data. (must already be sized to DstWidth*DstHeight)
+ */
+void FImageUtils::ImageResize(int32 SrcWidth, int32 SrcHeight, const TArrayView<const FColor> &SrcData, int32 DstWidth, int32 DstHeight, const TArrayView<FColor> &DstData, bool bLinearSpace)
+{
+	check(SrcData.Num() >= SrcWidth * SrcHeight);
+	check(DstData.Num() >= DstWidth * DstHeight);
+
 	float SrcX = 0;
 	float SrcY = 0;
 
diff --git a/Engine/Source/Runtime/Engine/Public/ImageUtils.h b/Engine/Source/Runtime/Engine/Public/ImageUtils.h
--- a/Engine/Source/Runtime/Engine/Public/ImageUtils.h
+++ b/Engine/Source/Runtime/Engine/Public/ImageUtils.h
@@ -62,6 +62,19 @@ public:
 	 */
 	ENGINE_API static void ImageResize(int32 SrcWidth, int32 SrcHeight, const TArray<FColor> &SrcData,  int32 DstWidth, int32 DstHeight, TArray<FColor> &DstData, bool bLinearSpace );
 
+	/**
+	 * Resizes the given image using a simple average filter and stores it in the destination array.  This version constrains aspect ratio.
+	 * Accepts TArrayViews but requires that DstData be pre-sized appropriately
+	 *
+	 * @param SrcWidth	Source image width.
+	 * @param SrcHeight	Source image height.
+	 * @param SrcData	Source image data.
+	 * @param DstWidth	Destination image width.
+	 * @param DstHeight Destination image height.
+	 * @param DstData	Destination image data. (must already be sized to DstWidth*DstHeight)
+	 */
+	ENGINE_API static void ImageResize(int32 SrcWidth, int32 SrcHeight, const TArrayView<const FColor> &SrcData, int32 DstWidth, int32 DstHeight, const TArrayView<FColor> &DstData, bool bLinearSpace);
+
 	/**
 	 * Creates a 2D texture from a array of raw color data.
 	 *
