Tightly pack FMaterialCompilationOutput. Pack bools into bits. Use uint32 instead of uint64 for flags.

diff --git a/Engine/Source/Runtime/Engine/Classes/Materials/MaterialExpression.h b/Engine/Source/Runtime/Engine/Classes/Materials/MaterialExpression.h
--- a/Engine/Source/Runtime/Engine/Classes/Materials/MaterialExpression.h
+++ b/Engine/Source/Runtime/Engine/Classes/Materials/MaterialExpression.h
@@ -6,7 +6,6 @@
 #include "UObject/ObjectMacros.h"
 #include "UObject/Object.h"
 #include "Misc/Guid.h"
-#include "MaterialShared.h"
 #include "MaterialExpressionIO.h"
 
 #include "MaterialExpression.generated.h"
diff --git a/Engine/Source/Runtime/Engine/Private/Materials/HLSLMaterialTranslator.h b/Engine/Source/Runtime/Engine/Private/Materials/HLSLMaterialTranslator.h
--- a/Engine/Source/Runtime/Engine/Private/Materials/HLSLMaterialTranslator.h
+++ b/Engine/Source/Runtime/Engine/Private/Materials/HLSLMaterialTranslator.h
@@ -711,11 +711,14 @@ public:
 				Errorf(TEXT("Only transparent or postprocess materials can read from scene depth."));
 			}
 
-			MaterialCompilationOutput.bUsesSceneDepthLookup = bUsesSceneDepth;
+			if (bUsesSceneDepth)
+			{
+				MaterialCompilationOutput.SetIsSceneTextureUsed(PPI_SceneDepth);
+			}
 
 			MaterialCompilationOutput.bUsesDistanceCullFade = bUsesDistanceCullFade;
 
-			if (MaterialCompilationOutput.bRequiresSceneColorCopy)
+			if (MaterialCompilationOutput.RequiresSceneColorCopy())
 			{
 				if (Domain != MD_Surface)
 				{
@@ -769,7 +772,7 @@ public:
 				Errorf(TEXT("Material using the DeferredDecal domain need to use the BlendModel Translucent (this saves performance)"));
 			}
 
-			if (MaterialCompilationOutput.bNeedsSceneTextures)
+			if (MaterialCompilationOutput.NeedsSceneTextures())
 			{
 				if (Domain != MD_DeferredDecal && Domain != MD_PostProcess)
 				{
@@ -1051,7 +1054,7 @@ public:
 			OutEnvironment.SetDefine(TEXT("NEEDS_PARTICLE_SIZE"), TEXT("1"));
 		}
 
-		if( MaterialCompilationOutput.bNeedsSceneTextures )
+		if( MaterialCompilationOutput.NeedsSceneTextures())
 		{
 			OutEnvironment.SetDefine(TEXT("NEEDS_SCENE_TEXTURES"), TEXT("1"));
 		}
@@ -3826,8 +3829,7 @@ protected:
 	// @param bTextureLookup true: texture, false:no texture lookup, usually to get the size
 	void UseSceneTextureId(ESceneTextureId SceneTextureId, bool bTextureLookup)
 	{
-		MaterialCompilationOutput.bNeedsSceneTextures = true;
-		MaterialCompilationOutput.UsedSceneTextures |= (1ull << SceneTextureId);
+		MaterialCompilationOutput.SetIsSceneTextureUsed(SceneTextureId);
 
 		if(Material->GetMaterialDomain() == MD_DeferredDecal)
 		{
@@ -3890,41 +3892,17 @@ protected:
 			bUsesSceneDepth = true;
 		}
 
-		const bool bNeedsGBuffer = SceneTextureId == PPI_DiffuseColor 
-			|| SceneTextureId == PPI_SpecularColor
-			|| SceneTextureId == PPI_SubsurfaceColor
-			|| SceneTextureId == PPI_BaseColor
-			|| SceneTextureId == PPI_Specular
-			|| SceneTextureId == PPI_Metallic
-			|| SceneTextureId == PPI_WorldNormal
-			|| SceneTextureId == PPI_Opacity
-			|| SceneTextureId == PPI_Roughness
-			|| SceneTextureId == PPI_MaterialAO
-			|| SceneTextureId == PPI_DecalMask
-			|| SceneTextureId == PPI_ShadingModelColor
-			|| SceneTextureId == PPI_ShadingModelID
-			|| SceneTextureId == PPI_StoredBaseColor
-			|| SceneTextureId == PPI_StoredSpecular
-			|| SceneTextureId == PPI_Velocity;
-
+		const bool bNeedsGBuffer = MaterialCompilationOutput.NeedsGBuffer();
 
-		MaterialCompilationOutput.bNeedsGBuffer = MaterialCompilationOutput.bNeedsGBuffer || bNeedsGBuffer;
 
 		if (bNeedsGBuffer && IsForwardShadingEnabled(Platform))
 		{
 			Errorf(TEXT("GBuffer scene textures not available with forward shading."));
 		}
 
-		if (SceneTextureId == PPI_Velocity)
+		if (SceneTextureId == PPI_Velocity && Material->GetMaterialDomain() != MD_PostProcess)
 		{
-			if (Material->GetMaterialDomain() == MD_PostProcess)
-			{
-				MaterialCompilationOutput.bUsesVelocitySceneTexture = true;
-			}
-			else
-			{
-				Errorf(TEXT("Velocity scene textures are only available in post process materials."));
-			}
+			Errorf(TEXT("Velocity scene textures are only available in post process materials."));
 		}
 
 		// not yet tracked:
@@ -3953,7 +3931,7 @@ protected:
 			return INDEX_NONE;
 		}
 
-		MaterialCompilationOutput.bRequiresSceneColorCopy = true;
+		MaterialCompilationOutput.SetIsSceneTextureUsed(PPI_SceneColor);
 		AddEstimatedTextureSample();
 
 		int32 ScreenUVCode = GetScreenAlignedUV(Offset, ViewportUV, bUseOffset);
diff --git a/Engine/Source/Runtime/Engine/Private/Materials/MaterialShared.cpp b/Engine/Source/Runtime/Engine/Private/Materials/MaterialShared.cpp
--- a/Engine/Source/Runtime/Engine/Private/Materials/MaterialShared.cpp
+++ b/Engine/Source/Runtime/Engine/Private/Materials/MaterialShared.cpp
@@ -451,27 +451,34 @@ void FMaterialCompilationOutput::Serialize(FArchive& Ar)
 
 	Ar << UsedSceneTextures;
 #if WITH_EDITOR
-	Ar << NumUsedUVScalars;
-	Ar << NumUsedCustomInterpolatorScalars;
 	Ar << EstimatedNumTextureSamplesVS;
 	Ar << EstimatedNumTextureSamplesPS;
+	Ar << NumUsedUVScalars;
+	Ar << NumUsedCustomInterpolatorScalars;
 #else
 	uint8 Tmp8; uint16 Tmp16;
-	Ar << Tmp8 << Tmp8;
-	Ar << Tmp16 << Tmp16;
+	Ar << Tmp16;
+	Ar << Tmp16;
+	Ar << Tmp8;
+	Ar << Tmp8;
 #endif
 
-	Ar << bRequiresSceneColorCopy;
-	Ar << bNeedsSceneTextures;
-	Ar << bUsesEyeAdaptation;
-	Ar << bModifiesMeshPosition;
-	Ar << bUsesWorldPositionOffset;
-	Ar << bNeedsGBuffer;
-	Ar << bUsesGlobalDistanceField;
-	Ar << bUsesPixelDepthOffset;
-	Ar << bUsesSceneDepthLookup;
-	Ar << bUsesVelocitySceneTexture;
-	Ar << bUsesDistanceCullFade;
+	uint8 PackedFlags = 0;
+	PackedFlags |= (bUsesEyeAdaptation			<< 0);
+	PackedFlags |= (bModifiesMeshPosition		<< 1);
+	PackedFlags |= (bUsesWorldPositionOffset	<< 2);
+	PackedFlags |= (bUsesGlobalDistanceField	<< 3);
+	PackedFlags |= (bUsesPixelDepthOffset		<< 4);
+	PackedFlags |= (bUsesDistanceCullFade		<< 5);
+
+	Ar << PackedFlags;
+
+	bUsesEyeAdaptation			= PackedFlags & (1 << 0);
+	bModifiesMeshPosition		= PackedFlags & (1 << 1);
+	bUsesWorldPositionOffset	= PackedFlags & (1 << 2);
+	bUsesGlobalDistanceField	= PackedFlags & (1 << 3);
+	bUsesPixelDepthOffset		= PackedFlags & (1 << 4);
+	bUsesDistanceCullFade		= PackedFlags & (1 << 5);
 }
 
 void FMaterial::GetShaderMapId(EShaderPlatform Platform, FMaterialShaderMapId& OutId) const
diff --git a/Engine/Source/Runtime/Engine/Private/ShaderDerivedDataVersion.h b/Engine/Source/Runtime/Engine/Private/ShaderDerivedDataVersion.h
--- a/Engine/Source/Runtime/Engine/Private/ShaderDerivedDataVersion.h
+++ b/Engine/Source/Runtime/Engine/Private/ShaderDerivedDataVersion.h
@@ -11,5 +11,5 @@
 // In case of merge conflicts with DDC versions, you MUST generate a new GUID and set this new
 // guid as version
 
-#define GLOBALSHADERMAP_DERIVEDDATA_VER			TEXT("4FFBFDF601234B80893B59E7779F89E7")
-#define MATERIALSHADERMAP_DERIVEDDATA_VER		TEXT("FE2EEAF7F8BC456195582BBB000D17CD")
+#define GLOBALSHADERMAP_DERIVEDDATA_VER			TEXT("DE2EA5D0D20D43508627490DBAE0DD36")
+#define MATERIALSHADERMAP_DERIVEDDATA_VER		TEXT("B0953E44F4E340E7914CE7C2BEDBDEE5")
diff --git a/Engine/Source/Runtime/Engine/Public/MaterialShared.h b/Engine/Source/Runtime/Engine/Public/MaterialShared.h
--- a/Engine/Source/Runtime/Engine/Public/MaterialShared.h
+++ b/Engine/Source/Runtime/Engine/Public/MaterialShared.h
@@ -1,4 +1,4 @@
-﻿// Copyright 1998-2019 Epic Games, Inc. All Rights Reserved.
+// Copyright 1998-2019 Epic Games, Inc. All Rights Reserved.
 
 /*=============================================================================
 	MaterialShared.h: Shared material definitions.
@@ -24,6 +24,7 @@
 #include "Misc/Optional.h"
 #include "Serialization/MemoryWriter.h"
 #include "Serialization/ArchiveProxy.h"
+#include "Materials/MaterialExpressionSceneTexture.h"
 
 class FMaterial;
 class FMaterialCompiler;
@@ -419,21 +420,16 @@ public:
 	FMaterialCompilationOutput() :
 		UsedSceneTextures(0),
 #if WITH_EDITOR
-		NumUsedUVScalars(0),
-		NumUsedCustomInterpolatorScalars(0),
 		EstimatedNumTextureSamplesVS(0),
 		EstimatedNumTextureSamplesPS(0),
+		NumUsedUVScalars(0),
+		NumUsedCustomInterpolatorScalars(0),
 #endif
-		bRequiresSceneColorCopy(false),
-		bNeedsSceneTextures(false),
 		bUsesEyeAdaptation(false),
 		bModifiesMeshPosition(false),
 		bUsesWorldPositionOffset(false),
-		bNeedsGBuffer(false),
 		bUsesGlobalDistanceField(false),
 		bUsesPixelDepthOffset(false),
-		bUsesSceneDepthLookup(false),
-		bUsesVelocitySceneTexture(false),
 		bUsesDistanceCullFade(false)
 	{}
 
@@ -442,51 +438,74 @@ public:
 	FUniformExpressionSet UniformExpressionSet;
 
 	/** Bitfield of the ESceneTextures used */
-	uint64 UsedSceneTextures;
+	uint32 UsedSceneTextures;
 
 #if WITH_EDITOR
+	/** Number of times SampleTexture is called, excludes custom nodes. */
+	uint16 EstimatedNumTextureSamplesVS;
+	uint16 EstimatedNumTextureSamplesPS;
+
 	/** Number of used custom UV scalars. */
 	uint8 NumUsedUVScalars;
 
 	/** Number of used custom vertex interpolation scalars. */
 	uint8 NumUsedCustomInterpolatorScalars;
-
-	/** Number of times SampleTexture is called, excludes custom nodes. */
-	uint16 EstimatedNumTextureSamplesVS;
-	uint16 EstimatedNumTextureSamplesPS;
 #endif // WITH_EDITOR
 
-	/** Indicates whether the material uses scene color. */
-	bool bRequiresSceneColorCopy;
-
-	/** true if the material needs the scenetexture lookups. */
-	bool bNeedsSceneTextures;
-
 	/** true if the material uses the EyeAdaptationLookup */
-	bool bUsesEyeAdaptation;
+	uint8 bUsesEyeAdaptation : 1;
 
 	/** true if the material modifies the the mesh position. */
-	bool bModifiesMeshPosition;
+	uint8 bModifiesMeshPosition : 1;
 
 	/** Whether the material uses world position offset. */
-	bool bUsesWorldPositionOffset;
-
-	/** true if the material uses any GBuffer textures */
-	bool bNeedsGBuffer;
+	uint8 bUsesWorldPositionOffset : 1;
 
 	/** true if material uses the global distance field */
-	bool bUsesGlobalDistanceField;
+	uint8 bUsesGlobalDistanceField : 1;
 
 	/** true if the material writes a pixel depth offset */
-	bool bUsesPixelDepthOffset;
+	uint8 bUsesPixelDepthOffset : 1;
+
+	/** true if the material uses distance cull fade */
+	uint8 bUsesDistanceCullFade : 1;
+
+	bool IsSceneTextureUsed(ESceneTextureId TexId) const { return UsedSceneTextures & (1 << TexId); }
+	void SetIsSceneTextureUsed(ESceneTextureId TexId) { UsedSceneTextures |= (1 << TexId); }
+
+	/** Indicates whether the material uses scene color. */
+	bool RequiresSceneColorCopy() const { return IsSceneTextureUsed(PPI_SceneColor); }
+
+	/** true if the material needs the scenetexture lookups. */
+	bool NeedsSceneTextures() const { return UsedSceneTextures != 0; }
+
+	/** true if the material uses any GBuffer textures */
+	bool NeedsGBuffer() const 
+	{
+		return
+			IsSceneTextureUsed(PPI_DiffuseColor) ||
+			IsSceneTextureUsed(PPI_SpecularColor) ||
+			IsSceneTextureUsed(PPI_SubsurfaceColor) ||
+			IsSceneTextureUsed(PPI_BaseColor) ||
+			IsSceneTextureUsed(PPI_Specular) ||
+			IsSceneTextureUsed(PPI_Metallic) ||
+			IsSceneTextureUsed(PPI_WorldNormal) ||
+			IsSceneTextureUsed(PPI_Opacity) ||
+			IsSceneTextureUsed(PPI_Roughness) ||
+			IsSceneTextureUsed(PPI_MaterialAO) ||
+			IsSceneTextureUsed(PPI_DecalMask) ||
+			IsSceneTextureUsed(PPI_ShadingModelColor) ||
+			IsSceneTextureUsed(PPI_ShadingModelID) ||
+			IsSceneTextureUsed(PPI_StoredBaseColor) ||
+			IsSceneTextureUsed(PPI_StoredSpecular) ||
+			IsSceneTextureUsed(PPI_Velocity);
+	}
 
 	/** true if the material uses the SceneDepth lookup */
-	bool bUsesSceneDepthLookup;
+	bool UsesSceneDepthLookup() const { return IsSceneTextureUsed(PPI_SceneDepth); };	
 
 	/** true if the material uses the Velocity SceneTexture lookup */
-	bool bUsesVelocitySceneTexture;
-	/** true if the material uses distance cull fade */
-	bool bUsesDistanceCullFade;
+	bool UsesVelocitySceneTexture() const { return IsSceneTextureUsed(PPI_Velocity); };
 };
 
 /** 
@@ -930,23 +949,23 @@ public:
 	const FString& GetFriendlyName() const { static FString T; return T; }
 	const FString& GetDebugDescription() const { static FString T; return T; }
 #endif
-	bool RequiresSceneColorCopy() const { return MaterialCompilationOutput.bRequiresSceneColorCopy; }
-	bool NeedsSceneTextures() const { return MaterialCompilationOutput.bNeedsSceneTextures; }
+	bool RequiresSceneColorCopy() const { return MaterialCompilationOutput.RequiresSceneColorCopy(); }
+	bool NeedsSceneTextures() const { return MaterialCompilationOutput.NeedsSceneTextures(); }
 	bool UsesGlobalDistanceField() const { return MaterialCompilationOutput.bUsesGlobalDistanceField; }
 	bool UsesWorldPositionOffset() const { return MaterialCompilationOutput.bUsesWorldPositionOffset; }
-	bool NeedsGBuffer() const { return MaterialCompilationOutput.bNeedsGBuffer; }
+	bool NeedsGBuffer() const { return MaterialCompilationOutput.NeedsGBuffer(); }
 	bool UsesEyeAdaptation() const { return MaterialCompilationOutput.bUsesEyeAdaptation; }
 	bool ModifiesMeshPosition() const { return MaterialCompilationOutput.bModifiesMeshPosition; }
 	bool UsesPixelDepthOffset() const { return MaterialCompilationOutput.bUsesPixelDepthOffset; }
-	bool UsesSceneDepthLookup() const { return MaterialCompilationOutput.bUsesSceneDepthLookup; }
-	bool UsesVelocitySceneTexture() const { return MaterialCompilationOutput.bUsesVelocitySceneTexture; }
+	bool UsesSceneDepthLookup() const { return MaterialCompilationOutput.UsesSceneDepthLookup(); }
+	bool UsesVelocitySceneTexture() const { return MaterialCompilationOutput.UsesVelocitySceneTexture(); }
 	bool UsesDistanceCullFade() const { return MaterialCompilationOutput.bUsesDistanceCullFade; }
 #if WITH_EDITOR
 	uint32 GetNumUsedUVScalars() const { return MaterialCompilationOutput.NumUsedUVScalars; }
 	uint32 GetNumUsedCustomInterpolatorScalars() const { return MaterialCompilationOutput.NumUsedCustomInterpolatorScalars; }
 	void GetEstimatedNumTextureSamples(uint32& VSSamples, uint32& PSSamples) const { VSSamples = MaterialCompilationOutput.EstimatedNumTextureSamplesVS; PSSamples = MaterialCompilationOutput.EstimatedNumTextureSamplesPS; }
 #endif
-	bool UsesSceneTexture(uint32 TexId) const { return MaterialCompilationOutput.UsedSceneTextures & (1ull << TexId); }
+	bool UsesSceneTexture(ESceneTextureId TexId) const { return MaterialCompilationOutput.IsSceneTextureUsed(TexId); }
 
 	bool IsValidForRendering(bool bFailOnInvalid = false) const
 	{
