Fix for sounds not reviving in certain contrived cases where previous passive mix modifiers are not accounted for on audio mixer swap.

diff --git a/Engine/Source/Runtime/Engine/Private/AudioDeviceManager.cpp b/Engine/Source/Runtime/Engine/Private/AudioDeviceManager.cpp
--- a/Engine/Source/Runtime/Engine/Private/AudioDeviceManager.cpp
+++ b/Engine/Source/Runtime/Engine/Private/AudioDeviceManager.cpp
@@ -143,7 +143,9 @@ void FAudioDeviceManager::ToggleAudioMixer()
 
 					// To transfer mix states, we need to re-base the absolute clocks on the mix states
 					// so the target audio device timing won't result in the mixes suddenly stopping.
-					TMap<USoundMix *, FSoundMixState> MixModifiers = AudioDevice->GetSoundMixModifiers();
+					TMap<USoundMix*, FSoundMixState> MixModifiers = AudioDevice->GetSoundMixModifiers();
+					TArray<USoundMix*> PrevPassiveSoundMixModifiers = AudioDevice->GetPrevPassiveSoundMixModifiers();
+					USoundMix* BaseSoundMix = AudioDevice->GetDefaultBaseSoundMixModifier();
 					double AudioClock = AudioDevice->GetAudioClock();
 
 					for (TPair<USoundMix*, FSoundMixState>& SoundMixPair : MixModifiers)
@@ -181,8 +183,7 @@ void FAudioDeviceManager::ToggleAudioMixer()
 					}
 
 					// Transfer the sound mix modifiers to the new audio engine
-					AudioDevice->SetSoundMixModifiers(MixModifiers);
-
+					AudioDevice->SetSoundMixModifiers(MixModifiers, PrevPassiveSoundMixModifiers, BaseSoundMix);
 					// Setup the mute state of the audio device to be the same that it was
 					if (bIsActive)
 					{
diff --git a/Engine/Source/Runtime/Engine/Public/AudioDevice.h b/Engine/Source/Runtime/Engine/Public/AudioDevice.h
--- a/Engine/Source/Runtime/Engine/Public/AudioDevice.h
+++ b/Engine/Source/Runtime/Engine/Public/AudioDevice.h
@@ -1397,14 +1397,26 @@ public:
 		return Effects;
 	}
 
-	TMap<USoundMix *, FSoundMixState> GetSoundMixModifiers()
+	const TMap<USoundMix*, FSoundMixState>& GetSoundMixModifiers() const
 	{
 		return SoundMixModifiers;
 	}
 
-	void SetSoundMixModifiers(TMap<USoundMix *, FSoundMixState>& InSoundMixModifiers)
+	const TArray<USoundMix*>& GetPrevPassiveSoundMixModifiers() const
+	{
+		return PrevPassiveSoundMixModifiers;
+	}
+
+	USoundMix* GetDefaultBaseSoundMixModifier()
+	{
+		return DefaultBaseSoundMix;
+	}
+
+	void SetSoundMixModifiers(const TMap<USoundMix*, FSoundMixState>& InSoundMixModifiers, const TArray<USoundMix*>& InPrevPassiveSoundMixModifiers, USoundMix* InDefaultBaseSoundMix)
 	{
 		SoundMixModifiers = InSoundMixModifiers;
+		PrevPassiveSoundMixModifiers = InPrevPassiveSoundMixModifiers;
+		DefaultBaseSoundMix = InDefaultBaseSoundMix;
 	}
 
 private:
