Fixup case where Fast Array Delta Serializtion wouldn't mark newly seen items as initial.
Add docs covering NAKs.

diff --git a/Engine/Source/Runtime/Engine/Private/DataReplication.cpp b/Engine/Source/Runtime/Engine/Private/DataReplication.cpp
--- a/Engine/Source/Runtime/Engine/Private/DataReplication.cpp
+++ b/Engine/Source/Runtime/Engine/Private/DataReplication.cpp
@@ -685,6 +685,100 @@ void FObjectReplicator::StopReplicating( class UActorChannel * InActorChannel )
 	}
 }
 
+/**
+ * Handling NAKs / Property Retransmission.
+ *
+ * Note, NACK handling only occurs on connections that "replicate" data, which is currently
+ * only Servers. RPC retransmission is handled elsewhere.
+ *
+ * RepLayouts:
+ *
+ *		As we send properties through FRepLayout the is a Changelist Manager that is shared
+ *		between all connections tracks sets of properties that were recently changed (history items),
+ *		as well as one aggregate set of all properties that have ever been sent.
+ *
+ *		Each Sending Rep State, which is connection unique, also tracks the set of changed
+ *		properties. These history items will only be created when replicating the object,
+ *		so there will be fewer of them in general, but they will still contain any properties
+ *		that compared differently (not *just* the properties that were actually replicated).
+ *
+ *		Whenever a NAK is received, we will iterate over the SendingRepState changelist
+ *		and mark any of the properties sent in the NAKed packet for retransmission.
+ *
+ *		The next time Properties are replicated for the Object, we will merge in any changelists
+ *		from NAKed history items.
+ *
+ * Custom Delta Properties:
+ *
+ *		For Custom Delta Properties (CDP), we rely primarily on FPropertyRetirements and INetDeltaBaseState
+ *		for tracking property retransmission.
+ *
+ *		INetDeltaBaseStates are used to tracked internal state specific to a given type of CDP.
+ *		For example, Fast Array Replicators will use FNetFastTArrayBaseState, or some type
+ *		derived from that.
+ *
+ *		When an FObjectReplicator is created, we will create an INetDeltaBaseState for every CDP,
+ *		as well as a dummy FPropertyRetirement. This Property Retirement is used as the head
+ *		of a linked list of Retirements, and is generally never populated with any useful information.
+ *
+ *		Every time we replicate a CDP, we will pass in the most recent Base State, and we will be
+ *		returned a new CDP. If data is actually sent, then we will create a new Property Retirement,
+ *		adding it as the tail of our linked list. The new Property Retirement will also hold a reference
+ *		to the old INetDeltaBaseState (i.e., the state of the CDP before it replicated its properties).
+ *
+ *		Just before replicating, we will go through and free any ACKed FPropertyRetirments (see
+ *		UpdateAckedRetirements).
+ *
+ *		After replicating, we will cache off the returned Base State to be used as the "old" state
+ *		the next time the property is replicated.
+ *
+ *		Whenever a NAK is received, we will run through our Property Retirements. Any retirements
+ *		that predate the NACK will be removed and treated as if they were ACKs. The first
+ *		retirement that is found to be within the NAKed range will have its INetDeltaBaseState
+ *		restored (which should be the state before the NAKed packet was sent), and then
+ *		that retirement as well as all remaining will be removed.
+ *
+ *		The onus is then on the CDP to resend any necessary properties based on its current / live
+ *		state and the restored Net Delta Base State.
+ *
+ * Fast Array Properties:
+ *
+ *		Fast Array Properties are implemented as Custom Delta Properties (CDP). Therefore, they mostly
+ *		follow the flow laid out above.
+
+ *		FNetFastTArrayBaseState is the basis for all Fast Array Serializer INetDeltaBaseStates.
+ *		This struct tracks the Replication Key of the Array, the ID to Replication Key map of individual
+ *		Array Items, and a History Number.
+ *
+ *		As we replicate Fast Array Properties, we use the Array Replication key to see if anything
+ *		is possibly dirty in the Array and the ID to Replication map to see which Array Element
+ *		items actually are dirty. A mismatch between the Net Base State Key and the Key stored on
+ *		the live Fast Array (either the Array Replication Key, or any Item Key) is how we determine
+ *		if the Array or Items is dirty.
+ *
+ *		Whenever a NAK is received, our Old Base State will be reset to the last known ACKed value,
+ *		as described in the CDP section above. This means that our Array Replication Key and ID To
+ *		Item Replication Key should be reset to those states, forcing a mismatch the next time we
+ *		replicate if anything has changed.
+ *
+ *		When net.SupportFastArrayDelta is enabled, we perform an additional step in which we actually
+ *		compare the properties of dirty items. This is very similar to normal Property replication
+ *		using RepLayouts, and leverages most of the same code.
+ *
+ *		This includes tracking history items just like Rep Layout. Instead of tracking histories per
+ *		Sending Rep State / Per Connection, we just manage a single set of Histories on the Rep
+ *		Changelist Mgr. Changelists are stored per Fast Array Item, and are referenced via ID.
+ *
+ *		Whenever we go to replicate a Fast Array Item, we will merge together all changelists since
+ *		we last sent that item, and send those accumulated changes.
+ *
+ *		This means that property retransmission for Fast Array Items is an amalgamation of Rep Layout
+ *		retransmission and CDP retransmission.
+ *
+ *		Whenever a NAK is received, our History Number should be reset to the last known  ACKed value,
+ *		and that should be enough to force us to accumulate any of the NAKed item changelists.
+ */
+
 void FObjectReplicator::ReceivedNak( int32 NakPacketId )
 {
 	const UObject* Object = GetObject();
@@ -1354,14 +1448,11 @@ void FObjectReplicator::ReplicateCustomDeltaProperties( FNetBitWriter & Bunch, F
 	{
 		// Get info.
 		FPropertyRetirement& Retire = Retirement[CustomDeltaProperty];
-		FRepRecord* Rep = &ObjectClass->ClassReps[CustomDeltaProperty];
 
 PRAGMA_DISABLE_DEPRECATION_WARNINGS
 		UProperty* It = LocalRepLayout.GetPropertyForRepIndex(CustomDeltaProperty);
 PRAGMA_ENABLE_DEPRECATION_WARNINGS
 
-		int32 Index = Rep->Index;
-
 		const ELifetimeCondition RepCondition = LocalRepLayout.GetPropertyLifetimeCondition(CustomDeltaProperty);
 
 		check(RepCondition >= 0 && RepCondition < COND_Max);
diff --git a/Engine/Source/Runtime/Engine/Private/RepLayout.cpp b/Engine/Source/Runtime/Engine/Private/RepLayout.cpp
--- a/Engine/Source/Runtime/Engine/Private/RepLayout.cpp
+++ b/Engine/Source/Runtime/Engine/Private/RepLayout.cpp
@@ -6065,11 +6065,13 @@ bool FRepLayout::DeltaSerializeFastArrayProperty(FFastArrayDeltaSerializeParams&
 					{
 						HistoryItem.bWasUpdated = true;
 
+						FastArrayState.ArrayReplicationKey = NewArrayDeltaState->ArrayReplicationKey;
+
 						// Update our shadow array, and reset our pointer in case we reallocated.
 						FScriptArrayHelper ShadowArrayHelper((UArrayProperty*)FastArrayItemCmd.Property, ShadowArray);
 
 						TBitArray<> ShadowArrayItemIsNew(false, ObjectArrayNum);
-						const bool bIsInitial = (FastArrayState.HistoryEnd - FastArrayState.HistoryStart) == 1;
+						const bool bIsInitial = CompareChangelistDelta == 1;
 
 						// It's possible that elements have been deleted or otherwise reordered, and our shadow state is out of date.
 						// In order to prevent issues, we'll shuffle our shadow state back to the correct order.
@@ -6099,7 +6101,7 @@ bool FRepLayout::DeltaSerializeFastArrayProperty(FFastArrayDeltaSerializeParams&
 						//
 						//			a. If we find the Shadow Element in the Shadow Line, it must have been reordered.
 						//				Go ahead and swap the current front of the Shadow Line with the found Shadow Element.
-						//				Now, the elements at the front of the line have mathcing IDs, and we can move onto
+						//				Now, the elements at the front of the line have matching IDs, and we can move onto
 						//				the next Element in both lines.
 						//
 						//			b. If we don't find the Shadow Element in the Shadow Line, the Object Element must be new.
@@ -6122,6 +6124,11 @@ bool FRepLayout::DeltaSerializeFastArrayProperty(FFastArrayDeltaSerializeParams&
 						// TODO: Optimize this. Luckily, it only happens once per frame, and only if the array is dirty.
 						//			Maybe this could be merged into the sending code below, the only concern is that
 						//			doesn't tracked deleted elements.
+						//
+						//			Alternatively, if Custom Delta code was merged into FRepLayout, we might be able to track
+						//			lists of deleted items on a given frame and merge those together just like changelists.
+						//			This would prevent us from needing to call BuildChangedAndDeletedBuffers on Fast TArrays
+						//			for every connection, unless a specific connection was very out of date.
 
 						// Note, this code serves a very similar purpose to FFastArraySerializer::TFastArraySerializeHelper<Type, SerializerType>::BuildChangedAndDeletedBuffers.
 						// The main issue is that we can't rely on that method, because it will be comparing the last state that was replicated to a given particular connection,
@@ -6131,30 +6138,43 @@ bool FRepLayout::DeltaSerializeFastArrayProperty(FFastArrayDeltaSerializeParams&
 						{
 							FScriptArrayHelper ObjectArrayHelper((UArrayProperty*)FastArrayItemCmd.Property, ObjectArray);
 
-							TMap<int32, int32>& ShadowIDToIndex = FastArrayState.IDToIndexMap;
-							const TMap<int32, int32>& ObjectIDToIndex = ArraySerializer.ItemMap;
+							const int32 EndIndex = FMath::Min(ShadowArrayHelper.Num(), ObjectArrayNum);
+							const TMap<int32, int32> OldShadowIDToIndexMap(MoveTemp(FastArrayState.IDToIndexMap));
+							FastArrayState.IDToIndexMap.Reserve(ObjectArrayNum);
 
-							for (int32 Index = 0; Index < ObjectArrayNum; ++Index)
-							{
-								// Anything else must be a newly appended element, so we are done.
-								if (Index >= ShadowArrayHelper.Num())
-								{
-									break;
-								}
+							// We track the Appended Shadow Items, because any index we try and use after such
+							// an append needs to be shifted appropriately.
+							// TODO: We may be able to iterate the list backwards instead, but that may break
+							//			some assumptions laid out in the algorithm above.
+							int32 AppendedShadowItems = 0;
+
+							UE_LOG(LogRep, VeryVerbose, TEXT("DeltaSerializeFastArrayProperty: Fixup Shadow State. Owner=%s, Property=%s, bInitial=%d, ObjecyArrayNum=%d, ShadowArrayNum=%d"),
+								*Owner->GetName(), *Parent.CachedPropertyName.ToString(), !!bIsInitial, ObjectArrayNum, ShadowArrayHelper.Num());
 
+							for (int32 Index = 0; Index < EndIndex; ++Index)
+							{
 								FFastArraySerializerItem* ObjectItem = GetFastArrayItem(Params, ObjectArrayHelper.GetRawPtr(Index), FastArrayItemCmd);
 								FFastArraySerializerItem* ShadowItem = GetFastArrayItem(Params, ShadowArrayHelper.GetRawPtr(Index), FastArrayItemCmd);
 
+								FastArrayState.IDToIndexMap.Emplace(ObjectItem->ReplicationID, Index);
+
+								UE_LOG(LogRep, VeryVerbose, TEXT("DeltaSerializeFastArrayProperty: Handling Item. ID=%d, Index=%d, ShadowID=%d"), ObjectItem->ReplicationID, Index, ShadowItem->ReplicationID);
+
+								// If our IDs match, there's nothing to do.
 								if (ObjectItem->ReplicationID != ShadowItem->ReplicationID)
 								{
 									// The IDs didn't match, so this is an insert, delete, or swap.
-									if (int32* FoundShadowIndex = ShadowIDToIndex.Find(ShadowItem->ReplicationID))
+									if (int32 const * const FoundShadowIndex = OldShadowIDToIndexMap.Find(ObjectItem->ReplicationID))
 									{
 										// We found the element in the shadow array, so there must have been a swap.
-										// Sanity check that the invalid element can only possibly be later in our lines.
-										check(*FoundShadowIndex > Index);
+										// Sanity check that the invalid element can only possibly later in our lines.
+										const int32 FixedShadowIndex = *FoundShadowIndex + AppendedShadowItems;
+
+										UE_LOG(LogRep, VeryVerbose, TEXT("DeltaSerializeFastArrayProperty: Swapped Shadow Item. OldIndex=%d, NewIndex=%d"), Index, FixedShadowIndex);
 
-										ShadowArrayHelper.SwapValues(Index, *FoundShadowIndex);
+										check(FixedShadowIndex > Index);
+
+										ShadowArrayHelper.SwapValues(Index, FixedShadowIndex);
 									}
 									else
 									{
@@ -6166,30 +6186,42 @@ bool FRepLayout::DeltaSerializeFastArrayProperty(FFastArrayDeltaSerializeParams&
 
 										ShadowItem = GetFastArrayItem(Params, ShadowArrayHelper.GetRawPtr(Index), FastArrayItemCmd);
 										ShadowItem->ReplicationID = ObjectItem->ReplicationID;
+
+										++AppendedShadowItems;
+										UE_LOG(LogRep, VeryVerbose, TEXT("DeltaSerializeFastArrayProperty: Added Shadow Item. AppendedShadowItems=%d"), AppendedShadowItems);
 									}
 								}
 							}
 						}
 
 						// Now we can go ahead and resize the array, to make any other changes we need.
-						ShadowArrayHelper.Resize(ObjectArrayNum);
-						ShadowArrayData = ShadowArray->GetData();
-
-						// Go ahead and fix up IDs for any elements that may have just been appended.
-						// Note, we need to this for all elements on the initial pass.
-						// Deleted elements will have been chopped off by the resize.
-						if (bIsInitial || (ShadowArrayHelper.Num() < ObjectArrayNum))
 						{
-							FScriptArrayHelper ObjectArrayHelper((UArrayProperty*)FastArrayItemCmd.Property, ObjectArray);
-							for (int32 ObjectIndex = ShadowArrayHelper.Num(); ObjectIndex < ObjectArrayNum; ++ObjectIndex)
+							const int32 OldShadowArrayNum = ShadowArrayHelper.Num();
+							ShadowArrayHelper.Resize(ObjectArrayNum);
+							ShadowArrayData = ShadowArray->GetData();
+
+							// Go ahead and fix up IDs for any elements that may have just been appended.
+							// Note, we need to this for all elements on the initial pass.
+							// Deleted elements will have been chopped off by the resize.
+							if (bIsInitial || (OldShadowArrayNum < ObjectArrayNum))
 							{
-								FFastArraySerializerItem* ObjectItem = GetFastArrayItem(Params, ObjectArrayHelper.GetRawPtr(ObjectIndex), FastArrayItemCmd);
-								FFastArraySerializerItem* ShadowItem = GetFastArrayItem(Params, ShadowArrayHelper.GetRawPtr(ObjectIndex), FastArrayItemCmd);
-								ShadowItem->ReplicationID = ObjectItem->ReplicationID;
-								ShadowArrayItemIsNew[ObjectIndex] = true;
+								const int32 StartIndex = bIsInitial ? 0 : ShadowArrayHelper.Num();
+								FScriptArrayHelper ObjectArrayHelper((UArrayProperty*)FastArrayItemCmd.Property, ObjectArray);
+								for (int32 Index = StartIndex; Index < ObjectArrayNum; ++Index)
+								{
+									FFastArraySerializerItem* ObjectItem = GetFastArrayItem(Params, ObjectArrayHelper.GetRawPtr(Index), FastArrayItemCmd);
+									FFastArraySerializerItem* ShadowItem = GetFastArrayItem(Params, ShadowArrayHelper.GetRawPtr(Index), FastArrayItemCmd);
+
+									ShadowItem->ReplicationID = ObjectItem->ReplicationID;
+									ShadowArrayItemIsNew[Index] = true;
+
+									UE_LOG(LogRep, VeryVerbose, TEXT("DeltaSerializeFastArrayProperty: Added Shadow Item. Index=%d, ID=%d"), Index, ShadowItem->ReplicationID);
+									FastArrayState.IDToIndexMap.Emplace(ObjectItem->ReplicationID, Index);
+								}
 							}
 						}
 
+						TArray<uint16> NewChangelist;
 						for (auto& IDIndexPair : ChangedElements)
 						{
 							// Go ahead and do a property compare here, regardless of what we'll actually use below.
@@ -6208,18 +6240,16 @@ bool FRepLayout::DeltaSerializeFastArrayProperty(FFastArrayDeltaSerializeParams&
 
 							FConstRepObjectDataBuffer ElementData(ObjectArrayData + ArrayElementOffset);
 							FRepShadowDataBuffer ElementShadowData(ShadowArrayData + ArrayElementOffset);
-							TArray<uint16>& HistoryChangelist = HistoryItem.ChangelistByID.Emplace(IDIndexPair.ID);
-							CompareProperties_r(SharedParams, ItemLayoutStart, ItemLayoutEnd, ElementShadowData, ElementData, HistoryChangelist, 0);
+							NewChangelist.Empty();
+
+							CompareProperties_r(SharedParams, ItemLayoutStart, ItemLayoutEnd, ElementShadowData, ElementData, NewChangelist, 0);
 
-							if (HistoryChangelist.Num())
+							if (NewChangelist.Num())
 							{
-								HistoryChangelist.Add(0);
+								NewChangelist.Add(0);
+								HistoryItem.ChangelistByID.Emplace(IDIndexPair.ID, MoveTemp(NewChangelist));
 							}
 						}
-
-						// Copy our replication key, and the IDToIndexMap so we can use them later.
-						FastArrayState.ArrayReplicationKey = NewArrayDeltaState->ArrayReplicationKey;
-						FastArrayState.IDToIndexMap = ArraySerializer.ItemMap;
 					}
 				}
 
@@ -6229,7 +6259,7 @@ bool FRepLayout::DeltaSerializeFastArrayProperty(FFastArrayDeltaSerializeParams&
 
 				// Note, this won't be all changes since the beginning, but just all changes for the currently dirty items.
 
-				if (LastSentChangelistDelta != 0 && LastSentChangelistDelta < FDeltaArrayHistoryState::MAX_CHANGE_HISTORY)
+				if (LastSentHistory != 0 && LastSentChangelistDelta > 0 && LastSentChangelistDelta < (FDeltaArrayHistoryState::MAX_CHANGE_HISTORY - 1))
 				{
 					const FConstRepObjectDataBuffer ConstObjectData(ObjectData);
 					Changelists.SetNum(ChangedElements.Num());
