Fix Fast Array Delta Struct serialization to use the correct Array Key when receiving.

diff --git a/Engine/Source/Runtime/Engine/Classes/Engine/NetSerialization.h b/Engine/Source/Runtime/Engine/Classes/Engine/NetSerialization.h
--- a/Engine/Source/Runtime/Engine/Classes/Engine/NetSerialization.h
+++ b/Engine/Source/Runtime/Engine/Classes/Engine/NetSerialization.h
@@ -1429,6 +1429,7 @@ struct FFastArrayDeltaSerializeParams
 	TArray<int32, TInlineAllocator<8>>* ReadChangedElements = nullptr;
 	TArray<int32, TInlineAllocator<8>>* ReadAddedElements = nullptr;
 	int32 ReadNumChanged = INDEX_NONE;
+	int32 ReadArrayReplicationKey = INDEX_NONE;
 };
 
 template<typename Type, typename SerializerType>
@@ -1566,6 +1567,7 @@ bool FFastArraySerializer::FastArrayDeltaSerialize_DeltaSerializeStructs(TArray<
 		DeltaSerializeParams.ReadAddedElements = &AddedIndices;
 		DeltaSerializeParams.ReadChangedElements = &ChangedIndices;
 		DeltaSerializeParams.ReadNumChanged = Header.NumChanged;
+		DeltaSerializeParams.ReadArrayReplicationKey = Header.ArrayReplicationKey;
 
 		if (!Parms.NetSerializeCB->NetDeltaSerializeForFastArray(DeltaSerializeParams))
 		{
diff --git a/Engine/Source/Runtime/Engine/Private/RepLayout.cpp b/Engine/Source/Runtime/Engine/Private/RepLayout.cpp
--- a/Engine/Source/Runtime/Engine/Private/RepLayout.cpp
+++ b/Engine/Source/Runtime/Engine/Private/RepLayout.cpp
@@ -6443,7 +6443,7 @@ bool FRepLayout::DeltaSerializeFastArrayProperty(FFastArrayDeltaSerializeParams&
 			}
 
 			// Update this element's most recent array replication key
-			ThisElement->MostRecentArrayReplicationKey = ArraySerializer.ArrayReplicationKey;
+			ThisElement->MostRecentArrayReplicationKey = Params.ReadArrayReplicationKey;
 
 			// Update this element's replication key so that a client can re-serialize the array for client replay recording
 			ThisElement->ReplicationKey++;
