Fixup a crash with Fast Array Struct Serialization where locally removed items may cause crashes for guids.

diff --git a/Engine/Source/Runtime/Engine/Private/RepLayout.cpp b/Engine/Source/Runtime/Engine/Private/RepLayout.cpp
--- a/Engine/Source/Runtime/Engine/Private/RepLayout.cpp
+++ b/Engine/Source/Runtime/Engine/Private/RepLayout.cpp
@@ -6601,25 +6601,32 @@ void FRepLayout::UpdateUnmappedGuidsForFastArray(FFastArrayDeltaSerializeParams&
 
 	FFastArraySerializer& ArraySerializer = *GetFastArray(Params, ObjectData, Parent);
 
-	bool bCalledPreNetReceive = Params.DeltaSerializeInfo.bOutSomeObjectsWereMapped;
-	for (auto& GuidReferencesPair : ArraySerializer.GuidReferencesMap_StructDelta)
+	for (auto It = ArraySerializer.GuidReferencesMap_StructDelta.CreateIterator(); It; ++It)
 	{
-		bool bOutSomeObjectsWereMapped = false;
-		bool bOutHasMoreUnmapped = false;
+		const int32 ElementID = It.Key();
+		if (int32 const * const FoundItemIndex = ArraySerializer.ItemMap.Find(ElementID))
+		{
+			bool bOutSomeObjectsWereMapped = false;
+			bool bOutHasMoreUnmapped = false;
+
+			const int32 ItemIndex = *FoundItemIndex;
+			const int32 ArrayElementOffset = ItemIndex * ElementSize;
+			FRepObjectDataBuffer ElementData(ArrayData + ArrayElementOffset);
 
-		const int32 ItemIndex = ArraySerializer.ItemMap[GuidReferencesPair.Key];
-		const int32 ArrayElementOffset = ItemIndex * ElementSize;
-		FRepObjectDataBuffer ElementData(ArrayData + ArrayElementOffset);
+			UpdateUnmappedObjects_r(nullptr, &It.Value(), Object, PackageMap, nullptr, ElementData, ElementSize, Params.DeltaSerializeInfo.bCalledPreNetReceive, bOutSomeObjectsWereMapped, bOutHasMoreUnmapped);
 
-		UpdateUnmappedObjects_r(nullptr, &GuidReferencesPair.Value, Object, PackageMap, nullptr, ElementData, ElementSize, bCalledPreNetReceive, bOutSomeObjectsWereMapped, bOutHasMoreUnmapped);
+			if (bOutSomeObjectsWereMapped)
+			{
+				GetFastArrayItem(Params, ElementData, FastArrayItemCmd)->PostReplicatedChange(ArraySerializer);
+			}
 
-		if (bOutSomeObjectsWereMapped)
+			DeltaSerializeInfo.bOutHasMoreUnmapped |= bOutHasMoreUnmapped;
+			DeltaSerializeInfo.bOutSomeObjectsWereMapped |= bOutSomeObjectsWereMapped;
+		}
+		else
 		{
-			((FFastArraySerializerItem*)(ElementData).Data)->PostReplicatedChange(ArraySerializer);
+			It.RemoveCurrent();
 		}
-
-		DeltaSerializeInfo.bOutHasMoreUnmapped |= bOutHasMoreUnmapped;
-		DeltaSerializeInfo.bOutSomeObjectsWereMapped |= bOutSomeObjectsWereMapped;
 	}
 }
 
