Prevent cases where NetDeltaSerializeInfo params could be dirtied between calls.

diff --git a/Engine/Source/Runtime/Engine/Private/DataReplication.cpp b/Engine/Source/Runtime/Engine/Private/DataReplication.cpp
--- a/Engine/Source/Runtime/Engine/Private/DataReplication.cpp
+++ b/Engine/Source/Runtime/Engine/Private/DataReplication.cpp
@@ -955,13 +955,6 @@ bool FObjectReplicator::ReceivedBunch(FNetBitReader& Bunch, const FReplicationFl
 
 	FNetSerializeCB NetSerializeCB(ConnectionNetDriver);
 
-	FNetDeltaSerializeInfo Parms;
-	Parms.Map = PackageMap;
-	Parms.Reader = &Reader;
-	Parms.NetSerializeCB = &NetSerializeCB;
-	Parms.Connection = Connection;
-	Parms.Object = Object;
-
 	// Read each property/function blob into Reader (so we've safely jumped over this data in the Bunch/stream at this point)
 	while (OwningChannel->ReadFieldHeaderAndPayload(Object, ClassCache, NetFieldExportGroup, Bunch, &FieldCache, Reader))
 	{
@@ -1011,6 +1004,14 @@ bool FObjectReplicator::ReceivedBunch(FNetBitReader& Bunch, const FReplicationFl
 				}
 			}
 #endif
+
+			FNetDeltaSerializeInfo Parms;
+			Parms.Map = PackageMap;
+			Parms.Reader = &Reader;
+			Parms.NetSerializeCB = &NetSerializeCB;
+			Parms.Connection = Connection;
+			Parms.Object = Object;
+
 			uint32 StaticArrayIndex = 0;
 			int32 Offset = 0;
 			if (!FNetSerializeCB::ReceiveCustomDeltaProperty(LocalRepLayout, Parms, ReplicatedProp, StaticArrayIndex, Offset))
@@ -2030,6 +2031,7 @@ void FObjectReplicator::UpdateUnmappedObjects(bool & bOutHasMoreUnmapped)
 	checkf( bHasQueuedBunches || ReceivingRepState->RepNotifies.Num() == 0, TEXT("Failed RepState RepNotifies check. Num=%d. Object=%s. Channel QueuedBunches=%d"), ReceivingRepState->RepNotifies.Num(), *Object->GetFullName(), OwningChannel ? OwningChannel->QueuedBunches.Num() : 0 );
 	checkf( bHasQueuedBunches || RepNotifies.Num() == 0, TEXT("Failed replicator RepNotifies check. Num=%d. Object=%s. Channel QueuedBunches=%d"), RepNotifies.Num(), *Object->GetFullName(), OwningChannel ? OwningChannel->QueuedBunches.Num() : 0 );
 
+	bool bCalledPreNetReceive = false;
 	bool bSomeObjectsWereMapped = false;
 
 	check(RepLayout);
@@ -2037,7 +2039,7 @@ void FObjectReplicator::UpdateUnmappedObjects(bool & bOutHasMoreUnmapped)
 	const FRepLayout& LocalRepLayout = *RepLayout;
 
 	// Let the rep layout update any unmapped properties
-	LocalRepLayout.UpdateUnmappedObjects(ReceivingRepState, Connection->PackageMap, Object, bSomeObjectsWereMapped, bOutHasMoreUnmapped);
+	LocalRepLayout.UpdateUnmappedObjects(ReceivingRepState, Connection->PackageMap, Object, bCalledPreNetReceive, bSomeObjectsWereMapped, bOutHasMoreUnmapped);
 
 	FNetSerializeCB NetSerializeCB(Connection->Driver);
 
@@ -2048,7 +2050,7 @@ void FObjectReplicator::UpdateUnmappedObjects(bool & bOutHasMoreUnmapped)
 	Parms.NetSerializeCB = &NetSerializeCB;
 
 	Parms.bUpdateUnmappedObjects = true;
-	Parms.bCalledPreNetReceive = bSomeObjectsWereMapped;	// RepLayout used this to flag whether PreNetReceive was called
+	Parms.bCalledPreNetReceive = bCalledPreNetReceive;
 	
 
 	TArray<TPair<int32, UStructProperty*>> CompletelyMappedProperties;
@@ -2057,6 +2059,7 @@ void FObjectReplicator::UpdateUnmappedObjects(bool & bOutHasMoreUnmapped)
 
 	bSomeObjectsWereMapped |= Parms.bOutSomeObjectsWereMapped;
 	bOutHasMoreUnmapped |= Parms.bOutHasMoreUnmapped;
+	bCalledPreNetReceive |= Parms.bCalledPreNetReceive;
 
 	// This should go away when UnmappedCustomProperties goes away, and when RepNotifies
 	// are merged with RepState RepNotifies.
@@ -2075,7 +2078,7 @@ PRAGMA_DISABLE_DEPRECATION_WARNINGS
 	}
 PRAGMA_ENABLE_DEPRECATION_WARNINGS
 
-	if (bSomeObjectsWereMapped)
+	if (bCalledPreNetReceive)
 	{
 		// If we mapped some objects, make sure to call PostNetReceive (some game code will need to think this was actually replicated to work)
 		PostNetReceive();
diff --git a/Engine/Source/Runtime/Engine/Private/RepLayout.cpp b/Engine/Source/Runtime/Engine/Private/RepLayout.cpp
--- a/Engine/Source/Runtime/Engine/Private/RepLayout.cpp
+++ b/Engine/Source/Runtime/Engine/Private/RepLayout.cpp
@@ -3156,11 +3156,12 @@ void FRepLayout::GatherGuidReferencesForCustomDeltaProperties(FNetDeltaSerialize
 			UStructProperty* StructProperty = static_cast<UStructProperty*>(Parent.Property);
 			UScriptStruct::ICppStructOps* CppStructOps = StructProperty->Struct->GetCppStructOps();
 
-			Params.Struct = StructProperty->Struct;
-			Params.PropertyRepIndex = KVP.Key;
-			Params.Data = ObjectData + Parent;
+			FNetDeltaSerializeInfo TempParams = Params;
+			TempParams.Struct = StructProperty->Struct;
+			TempParams.PropertyRepIndex = KVP.Key;
+			TempParams.Data = ObjectData + Parent;
 
-			CppStructOps->NetDeltaSerialize(Params, Params.Data);
+			CppStructOps->NetDeltaSerialize(TempParams, TempParams.Data);
 		}
 	}
 }
@@ -3222,15 +3223,22 @@ bool FRepLayout::MoveMappedObjectToUnmappedForCustomDeltaProperties(FNetDeltaSer
 			UStructProperty* StructProperty = static_cast<UStructProperty*>(Parent.Property);
 			UScriptStruct::ICppStructOps* CppStructOps = StructProperty->Struct->GetCppStructOps();
 
-			Params.Struct = StructProperty->Struct;
-			Params.Data = ObjectData + Parent;
-			Params.PropertyRepIndex = KVP.Key;
+			FNetDeltaSerializeInfo TempParams = Params;
 
-			if (CppStructOps->NetDeltaSerialize(Params, Params.Data))
+			TempParams.Struct = StructProperty->Struct;
+			TempParams.Data = ObjectData + Parent;
+			TempParams.PropertyRepIndex = KVP.Key;
+			TempParams.bOutHasMoreUnmapped = false;
+			TempParams.bOutSomeObjectsWereMapped = false;
+
+			if (CppStructOps->NetDeltaSerialize(TempParams, TempParams.Data))
 			{
 				UnmappedCustomProperties.Add(Parent.Offset, StructProperty);
 				bFound = true;
 			}
+
+			Params.bOutHasMoreUnmapped |= TempParams.bOutHasMoreUnmapped;
+			Params.bOutSomeObjectsWereMapped |= TempParams.bOutSomeObjectsWereMapped;
 		}
 	}
 
@@ -3386,16 +3394,27 @@ void FRepLayout::UpdateUnmappedObjects(
 	FReceivingRepState* RESTRICT RepState,
 	UPackageMap* PackageMap,
 	UObject* OriginalObject,
+	bool& bCalledPreNetReceive,
 	bool& bOutSomeObjectsWereMapped,
 	bool& bOutHasMoreUnmapped) const
 {
 	bOutSomeObjectsWereMapped = false;
 	bOutHasMoreUnmapped = false;
-	bool bCalledPreNetReceive = false;
+	bCalledPreNetReceive = false;
 
 	if (LayoutState == ERepLayoutState::Normal)
 	{
-		UpdateUnmappedObjects_r(RepState, &RepState->GuidReferencesMap, OriginalObject, PackageMap, (uint8*)RepState->StaticBuffer.GetData(), (uint8*)OriginalObject, Owner->GetPropertiesSize(), bCalledPreNetReceive, bOutSomeObjectsWereMapped, bOutHasMoreUnmapped);
+		UpdateUnmappedObjects_r(
+			RepState,
+			&RepState->GuidReferencesMap,
+			OriginalObject,
+			PackageMap,
+			(uint8*)RepState->StaticBuffer.GetData(),
+			(uint8*)OriginalObject,
+			Owner->GetPropertiesSize(),
+			bCalledPreNetReceive,
+			bOutSomeObjectsWereMapped,
+			bOutHasMoreUnmapped);
 	}
 }
 
@@ -3409,9 +3428,6 @@ PRAGMA_ENABLE_DEPRECATION_WARNINGS
 
 void FRepLayout::UpdateUnmappedObjectsForCustomDeltaProperties(FNetDeltaSerializeInfo& Params, TArray<TPair<int32, UStructProperty*>>& CompletelyUpdated, TArray<TPair<int32, UStructProperty*>>& Updated) const
 {
-	bool bOutSomeObjectsWereMapped = false;
-	bool bOutHasMoreUnmapped = false;
-
 	if (LifetimeCustomPropertyState)
 	{
 		FRepObjectDataBuffer ObjectData(Params.Object);
@@ -3423,33 +3439,33 @@ void FRepLayout::UpdateUnmappedObjectsForCustomDeltaProperties(FNetDeltaSerializ
 			UStructProperty* StructProperty = static_cast<UStructProperty*>(Parent.Property);
 			UScriptStruct::ICppStructOps* CppStructOps = StructProperty->Struct->GetCppStructOps();
 
-			Params.DebugName = Parent.CachedPropertyName.ToString();
-			Params.Struct = StructProperty->Struct;
-			Params.bOutSomeObjectsWereMapped = false;
-			Params.bOutHasMoreUnmapped = false;
-			Params.PropertyRepIndex = KVP.Key;
-			Params.Data = ObjectData + Parent;
+			FNetDeltaSerializeInfo TempParams = Params;
+
+			TempParams.DebugName = Parent.CachedPropertyName.ToString();
+			TempParams.Struct = StructProperty->Struct;
+			TempParams.bOutSomeObjectsWereMapped = false;
+			TempParams.bOutHasMoreUnmapped = false;
+			TempParams.PropertyRepIndex = KVP.Key;
+			TempParams.Data = ObjectData + Parent;
 
 			// Call the custom delta serialize function to handle it
-			CppStructOps->NetDeltaSerialize(Params, Params.Data);
+			CppStructOps->NetDeltaSerialize(TempParams, TempParams.Data);
 
-			if (Params.bOutSomeObjectsWereMapped)
+			if (TempParams.bOutSomeObjectsWereMapped)
 			{
 				Updated.Emplace(Parent.Offset, StructProperty);
 			}
 
-			if (!Params.bOutHasMoreUnmapped)
+			if (!TempParams.bOutHasMoreUnmapped)
 			{
 				CompletelyUpdated.Emplace(Parent.Offset, StructProperty);
 			}
 
-			bOutSomeObjectsWereMapped |= Params.bOutSomeObjectsWereMapped;
-			bOutHasMoreUnmapped |= Params.bOutHasMoreUnmapped;
+			Params.bOutSomeObjectsWereMapped |= TempParams.bOutSomeObjectsWereMapped;
+			Params.bOutHasMoreUnmapped |= TempParams.bOutHasMoreUnmapped;
+			Params.bCalledPreNetReceive |= TempParams.bCalledPreNetReceive;
 		}
 	}
-
-	Params.bOutSomeObjectsWereMapped = bOutSomeObjectsWereMapped;
-	Params.bOutHasMoreUnmapped = bOutHasMoreUnmapped;
 }
 
 bool FRepLayout::SendCustomDeltaProperty(
diff --git a/Engine/Source/Runtime/Engine/Public/Net/RepLayout.h b/Engine/Source/Runtime/Engine/Public/Net/RepLayout.h
--- a/Engine/Source/Runtime/Engine/Public/Net/RepLayout.h
+++ b/Engine/Source/Runtime/Engine/Public/Net/RepLayout.h
@@ -1294,6 +1294,7 @@ public:
 		FReceivingRepState* RESTRICT RepState,
 		UPackageMap* PackageMap,
 		UObject* Object,
+		bool& bCalledPreNetReceive,
 		bool& bOutSomeObjectsWereMapped,
 		bool& bOutHasMoreUnmapped) const;
 
@@ -1305,7 +1306,8 @@ public:
 		bool& bOutSomeObjectsWereMapped,
 		bool& bOutHasMoreUnmapped) const
 	{
-		UpdateUnmappedObjects(RepState->GetReceivingRepState(), PackageMap, Object, bOutSomeObjectsWereMapped, bOutHasMoreUnmapped);
+		bool bCalledPreNetReceive = false;
+		UpdateUnmappedObjects(RepState->GetReceivingRepState(), PackageMap, Object, bCalledPreNetReceive, bOutSomeObjectsWereMapped, bOutHasMoreUnmapped);
 	}
 
 	/**
