Recreate uniform expression uniform buffer when material instance resource is initialized. This forced uniform expression UB to be recreated and fixes the check that we are trying to update uniform expression UB which doesn't match current layout.

diff --git a/Engine/Source/Runtime/Engine/Private/Materials/MaterialInstance.cpp b/Engine/Source/Runtime/Engine/Private/Materials/MaterialInstance.cpp
--- a/Engine/Source/Runtime/Engine/Private/Materials/MaterialInstance.cpp
+++ b/Engine/Source/Runtime/Engine/Private/Materials/MaterialInstance.cpp
@@ -484,7 +484,7 @@ void UMaterialInstance::InitResources()
 
 	PropagateDataToMaterialProxy();
 
-	CacheMaterialInstanceUniformExpressions(this);
+	CacheMaterialInstanceUniformExpressions(this, true);
 }
 
 const UMaterial* UMaterialInstance::GetMaterial() const
