Fix curve issues in cooked games

diff --git a/Engine/Source/Runtime/Engine/Classes/Animation/AnimCurveCompressionSettings.h b/Engine/Source/Runtime/Engine/Classes/Animation/AnimCurveCompressionSettings.h
--- a/Engine/Source/Runtime/Engine/Classes/Animation/AnimCurveCompressionSettings.h
+++ b/Engine/Source/Runtime/Engine/Classes/Animation/AnimCurveCompressionSettings.h
@@ -17,14 +17,15 @@ class ENGINE_API UAnimCurveCompressionSettings : public UObject
 {
 	GENERATED_UCLASS_BODY()
 
-#if WITH_EDITORONLY_DATA
 	/** An animation curve compression codec. */
 	UPROPERTY(Category = Compression, Instanced, EditAnywhere, NoClear)
 	UAnimCurveCompressionCodec* Codec;
-#endif
 
 	//////////////////////////////////////////////////////////////////////////
 
+	/** Allow us to convert DDC serialized path back into codec object */
+	UAnimCurveCompressionCodec* GetCodec(const FString& Path);
+
 #if WITH_EDITORONLY_DATA
 	/** Returns whether or not we can use these settings to compress. */
 	bool AreSettingsValid() const;
@@ -36,9 +37,6 @@ class ENGINE_API UAnimCurveCompressionSettings : public UObject
 	 */
 	bool Compress(UAnimSequence& AnimSeq) const;
 
-	/** Allow us to convert DDC serialized path back into codec object */
-	UAnimCurveCompressionCodec* GetCodec(const FString& Path);
-
 	/** Generates a DDC key that takes into account the current settings and selected codec. */
 	FString MakeDDCKey() const;
 #endif
diff --git a/Engine/Source/Runtime/Engine/Classes/Animation/AnimSequence.h b/Engine/Source/Runtime/Engine/Classes/Animation/AnimSequence.h
--- a/Engine/Source/Runtime/Engine/Classes/Animation/AnimSequence.h
+++ b/Engine/Source/Runtime/Engine/Classes/Animation/AnimSequence.h
@@ -688,10 +688,11 @@ public:
 	UPROPERTY(Category = Compression, EditAnywhere)
 	bool bAllowFrameStripping;
 
+#endif
+
 	/** The curve compression settings used to compress curves in this sequence. */
 	UPROPERTY(Category=Compression, EditAnywhere)
 	class UAnimCurveCompressionSettings* CurveCompressionSettings;
-#endif // WITH_EDITORONLY_DATA
 
 	/** The codec used by the compressed data as determined by the compression settings. */
 	class UAnimCurveCompressionCodec* CurveCompressionCodec;
@@ -1063,10 +1064,10 @@ public:
 	 * @return		The approximate size of compressed animation data.
 	 */
 	int32 GetApproxCompressedSize() const;
-#if WITH_EDITORONLY_DATA
+
 	// Initialize curve compression settings, does nothing if scheme already valid
 	void InitCurveCompressionScheme();
-#endif
+
 	/**
 	 *  Utility function for lossless compression of a FRawAnimSequenceTrack 
 	 *  @return true if keys were removed.
diff --git a/Engine/Source/Runtime/Engine/Private/Animation/AnimCurveCompressionSettings.cpp b/Engine/Source/Runtime/Engine/Private/Animation/AnimCurveCompressionSettings.cpp
--- a/Engine/Source/Runtime/Engine/Private/Animation/AnimCurveCompressionSettings.cpp
+++ b/Engine/Source/Runtime/Engine/Private/Animation/AnimCurveCompressionSettings.cpp
@@ -8,10 +8,13 @@
 UAnimCurveCompressionSettings::UAnimCurveCompressionSettings(const FObjectInitializer& ObjectInitializer)
 	: Super(ObjectInitializer)
 {
-#if WITH_EDITORONLY_DATA
 	Codec = CreateDefaultSubobject<UAnimCurveCompressionCodec_CompressedRichCurve>(TEXT("CurveCompressionCodec"));
 	Codec->SetFlags(RF_Transactional);
-#endif
+}
+
+UAnimCurveCompressionCodec* UAnimCurveCompressionSettings::GetCodec(const FString& Path)
+{
+	return Codec->GetCodec(Path);
 }
 
 #if WITH_EDITORONLY_DATA
@@ -38,11 +41,6 @@ bool UAnimCurveCompressionSettings::Compress(UAnimSequence& AnimSeq) const
 	return Success;
 }
 
-UAnimCurveCompressionCodec* UAnimCurveCompressionSettings::GetCodec(const FString& Path)
-{
-	return Codec->GetCodec(Path);
-}
-
 FString UAnimCurveCompressionSettings::MakeDDCKey() const
 {
 	if (Codec == nullptr)
diff --git a/Engine/Source/Runtime/Engine/Private/Animation/AnimSequence.cpp b/Engine/Source/Runtime/Engine/Private/Animation/AnimSequence.cpp
--- a/Engine/Source/Runtime/Engine/Private/Animation/AnimSequence.cpp
+++ b/Engine/Source/Runtime/Engine/Private/Animation/AnimSequence.cpp
@@ -303,8 +303,9 @@ UAnimSequence::UAnimSequence(const FObjectInitializer& ObjectInitializer)
 	ImportResampleFramerate = 0;
 	bAllowFrameStripping = true;
 
-	InitCurveCompressionScheme();
 #endif
+
+	InitCurveCompressionScheme();
 }
 
 void UAnimSequence::PostInitProperties()
@@ -402,12 +403,16 @@ static void LoadOldCompressedTrack(FArchive& Ar, FCompressedTrack& Dst, int32 By
 	Ar << Dst.Ranges[0] << Dst.Ranges[1] << Dst.Ranges[2];
 }
 
-#if WITH_EDITORONLY_DATA
 void UAnimSequence::InitCurveCompressionScheme()
 {
 	// Do this is serialize as if the default animation curve compression asset isn't loaded it will
 // fire a warning if we try and load it in post load
-	if ((CurveCompressionSettings == nullptr || !CurveCompressionSettings->AreSettingsValid())
+
+	bool bCurveCompressionSettingsValid = CurveCompressionSettings != nullptr;
+#if WITH_EDITOR
+	bCurveCompressionSettingsValid = bCurveCompressionSettingsValid && CurveCompressionSettings->AreSettingsValid();
+#endif
+	if (!bCurveCompressionSettingsValid
 #if WITH_HOT_RELOAD
 		&& (GetClass()->HasAnyClassFlags(CLASS_CompiledFromBlueprint) || !HasAnyFlags(RF_ClassDefaultObject) || !GIsHotReload) // Don't do this to native CDOs during Hot-Reload
 #endif
@@ -416,7 +421,6 @@ void UAnimSequence::InitCurveCompressionScheme()
 		CurveCompressionSettings = FAnimationUtils::GetDefaultAnimationCurveCompressionSettings();
 	}
 }
-#endif
 
 void UAnimSequence::Serialize(FArchive& Ar)
 {
@@ -2729,23 +2733,10 @@ void UAnimSequence::SerializeCompressedData(FArchive& Ar, bool bDDCData)
 			}
 		}
 		
-#if WITH_EDITOR
-		if (bDDCData)
-		{
-			FString CurveCodecPath;
-			Ar << CurveCodecPath;
+		FString CurveCodecPath;
+		Ar << CurveCodecPath;
 
-			CurveCompressionCodec = CurveCompressionSettings->GetCodec(CurveCodecPath);
-		}
-		else
-#else
-		check(!bDDCData);
-#endif
-		{
-			UAnimCurveCompressionCodec* CurveCodec = nullptr;
-			Ar << CurveCodec;
-			CurveCompressionCodec = CurveCodec;
-		}
+		CurveCompressionCodec = CurveCompressionSettings->GetCodec(CurveCodecPath);
 
 		int32 NumCurveBytes;
 		Ar << NumCurveBytes;
@@ -2863,19 +2854,8 @@ void UAnimSequence::SerializeCompressedData(FArchive& Ar, bool bDDCData)
 		CompressedTrackOffsets = SavedCompressedTrackOffsets;
 		CompressedScaleOffsets = SavedCompressedScaleOffsets;
 
-#if WITH_EDITOR
-		if (bDDCData)
-		{
-			FString CurveCodecPath = CurveCompressionCodec->GetPathName();
-			Ar << CurveCodecPath;
-		}
-		else
-#else
-		check(!bDDCData);
-#endif
-		{
-			Ar << CurveCompressionCodec;
-		}
+		FString CurveCodecPath = CurveCompressionCodec->GetPathName();
+		Ar << CurveCodecPath;
 
 		int32 NumCurveBytes = CompressedCurveByteStream.Num();
 		Ar << NumCurveBytes;
