RL - Move RectLight MipTree generation from the Proxy to a livetime of the Component and its latest proxy. This avoids the rebuild when other light properties change and can consitently save up to 1.5ms of GPU time.

diff --git a/Engine/Source/Runtime/Engine/Classes/Components/RectLightComponent.h b/Engine/Source/Runtime/Engine/Classes/Components/RectLightComponent.h
--- a/Engine/Source/Runtime/Engine/Classes/Components/RectLightComponent.h
+++ b/Engine/Source/Runtime/Engine/Classes/Components/RectLightComponent.h
@@ -48,9 +48,12 @@ class ENGINE_API URectLightComponent : public ULocalLightComponent
 	float BarnDoorLength;
 
 	/** Texture mapped to the light source rectangle */
-	UPROPERTY(EditAnywhere, BlueprintReadWrite, Category=Light)
+	UPROPERTY(EditAnywhere, BlueprintReadOnly, Category = Light)
 	class UTexture* SourceTexture;
 
+	UFUNCTION(BlueprintCallable, Category = "Rendering|Lighting")
+	void SetSourceTexture(UTexture* bNewValue);
+
 	UFUNCTION(BlueprintCallable, Category="Rendering|Lighting")
 	void SetSourceWidth(float bNewValue);
 
@@ -75,9 +78,15 @@ public:
 	virtual float GetUniformPenumbraSize() const override;
 	virtual FLightSceneProxy* CreateSceneProxy() const override;
 
+	virtual void BeginDestroy() override;
 	//~ Begin UObject Interface
 #if WITH_EDITOR
 	virtual void PostEditChangeProperty(FPropertyChangedEvent& PropertyChangedEvent) override;
 #endif // WITH_EDITOR
 	//~ End UObject Interface
+
+private:
+	void UpdateRayTracingData();
+	friend class FRectLightSceneProxy;
+	struct FRectLightRayTracingData* RayTracingData;
 };
\ No newline at end of file
diff --git a/Engine/Source/Runtime/Engine/Private/Components/RectLightComponent.cpp b/Engine/Source/Runtime/Engine/Private/Components/RectLightComponent.cpp
--- a/Engine/Source/Runtime/Engine/Private/Components/RectLightComponent.cpp
+++ b/Engine/Source/Runtime/Engine/Private/Components/RectLightComponent.cpp
@@ -23,8 +23,29 @@ float GetRectLightBarnDoorMaxAngle()
 	return 88.f;
 }
 
+#if RHI_RAYTRACING
+FRectLightRayTracingData BuildRectLightMipTree(FRHICommandListImmediate& RHICmdList, UTexture* SourceTexture);
+#endif
+
+void URectLightComponent::UpdateRayTracingData()
+{
+#if RHI_RAYTRACING
+	if (IsRayTracingEnabled())
+	{
+		UTexture* NewSourceTexture = SourceTexture;
+		FRectLightRayTracingData* UpdatedRayTracingData = RayTracingData;
+		ENQUEUE_RENDER_COMMAND(BuildRectLightMipTree)(
+			[NewSourceTexture, UpdatedRayTracingData](FRHICommandListImmediate& RHICmdList) mutable
+		{
+			*UpdatedRayTracingData = BuildRectLightMipTree(RHICmdList, NewSourceTexture);
+		});
+	}
+#endif
+}
+
 URectLightComponent::URectLightComponent(const FObjectInitializer& ObjectInitializer)
 	: Super(ObjectInitializer)
+	, RayTracingData(new FRectLightRayTracingData())
 {
 #if WITH_EDITORONLY_DATA
 	if (!IsRunningCommandlet())
@@ -44,6 +65,8 @@ URectLightComponent::URectLightComponent(const FObjectInitializer& ObjectInitial
 	SourceTexture = nullptr;
 	BarnDoorAngle = GetRectLightBarnDoorMaxAngle();
 	BarnDoorLength = 20.0f;
+
+	UpdateRayTracingData();
 }
 
 FLightSceneProxy* URectLightComponent::CreateSceneProxy() const
@@ -51,6 +74,18 @@ FLightSceneProxy* URectLightComponent::CreateSceneProxy() const
 	return new FRectLightSceneProxy(this);
 }
 
+void URectLightComponent::SetSourceTexture(UTexture* NewValue)
+{
+	if (AreDynamicDataChangesAllowed()
+		&& SourceTexture != NewValue)
+	{
+		SourceTexture = NewValue;
+
+		UpdateRayTracingData();
+		MarkRenderStateDirty();
+	}
+}
+
 void URectLightComponent::SetSourceWidth(float NewValue)
 {
 	if (AreDynamicDataChangesAllowed()
@@ -153,6 +188,18 @@ float URectLightComponent::GetUniformPenumbraSize() const
 	}
 }
 
+void URectLightComponent::BeginDestroy()
+{
+	FRectLightRayTracingData* DeletedRenderData = RayTracingData;
+	RayTracingData = nullptr;
+	ENQUEUE_RENDER_COMMAND(DeleteBuildRectLightMipTree)(
+		[DeletedRenderData](FRHICommandListImmediate& RHICmdList)
+	{
+		delete DeletedRenderData;
+	});
+	Super::BeginDestroy();
+}
+
 #if WITH_EDITOR
 /**
  * Called after property has changed via e.g. property window or set command.
@@ -174,18 +221,9 @@ FRectLightSceneProxy::FRectLightSceneProxy(const URectLightComponent* Component)
 	, SourceHeight(Component->SourceHeight)
 	, BarnDoorAngle(FMath::Clamp(Component->BarnDoorAngle, 0.f, GetRectLightBarnDoorMaxAngle()))
 	, BarnDoorLength(FMath::Max(0.1f, Component->BarnDoorLength))
+	, RayTracingData(Component->RayTracingData)
 	, SourceTexture(Component->SourceTexture)
 {
-#if RHI_RAYTRACING
-	if (IsRayTracingEnabled())
-	{
-		ENQUEUE_RENDER_COMMAND(BuildRectLightMipTree)(
-			[this](FRHICommandListImmediate& RHICmdList)
-		{
-			BuildRectLightMipTree(RHICmdList);
-		});
-	}
-#endif
 }
 
 FRectLightSceneProxy::~FRectLightSceneProxy() {}
@@ -337,41 +375,46 @@ private:
 
 IMPLEMENT_SHADER_TYPE(, FBuildRectLightMipTreeCS, TEXT("/Engine/Private/Raytracing/BuildMipTreeCS.usf"), TEXT("BuildRectLightMipTreeCS"), SF_Compute)
 
-void FRectLightSceneProxy::BuildRectLightMipTree(FRHICommandListImmediate& RHICmdList)
+DECLARE_GPU_STAT_NAMED(BuildRectLightMipTreeStat, TEXT("build RectLight MipTree"));
+
+FRectLightRayTracingData BuildRectLightMipTree(FRHICommandListImmediate& RHICmdList, UTexture* SourceTexture)
 {
+	SCOPED_GPU_STAT(RHICmdList, BuildRectLightMipTreeStat);
+
 	check(IsInRenderingThread());
+	FRectLightRayTracingData Data;
+	FTextureRHIRef RhiTexture = SourceTexture ? SourceTexture->Resource->TextureRHI : GWhiteTexture->TextureRHI;
 
 	const auto ShaderMap = GetGlobalShaderMap(ERHIFeatureLevel::SM5);
 	TShaderMapRef<FBuildRectLightMipTreeCS> BuildRectLightMipTreeComputeShader(ShaderMap);
 	RHICmdList.SetComputeShader(BuildRectLightMipTreeComputeShader->GetComputeShader());
 
 	// Allocate MIP tree
-	FLightShaderParameters LightParameters;
-	GetLightShaderParameters(LightParameters);
-	FIntVector TextureSize = LightParameters.SourceTexture->GetSizeXYZ();
+	FIntVector TextureSize = RhiTexture->GetSizeXYZ();
 	uint32 MipLevelCount = FMath::Min(FMath::CeilLogTwo(TextureSize.X), FMath::CeilLogTwo(TextureSize.Y));
-	RectLightMipTreeDimensions = FIntVector(1 << MipLevelCount, 1 << MipLevelCount, 1);
-	uint32 NumElements = RectLightMipTreeDimensions.X * RectLightMipTreeDimensions.Y;
+	Data.RectLightMipTreeDimensions = FIntVector(1 << MipLevelCount, 1 << MipLevelCount, 1);
+	uint32 NumElements = Data.RectLightMipTreeDimensions.X * Data.RectLightMipTreeDimensions.Y;
 	for (uint32 MipLevel = 1; MipLevel <= MipLevelCount; ++MipLevel)
 	{
-		uint32 NumElementsInLevel = (RectLightMipTreeDimensions.X >> MipLevel) * (RectLightMipTreeDimensions.Y >> MipLevel);
+		uint32 NumElementsInLevel = (Data.RectLightMipTreeDimensions.X >> MipLevel) * (Data.RectLightMipTreeDimensions.Y >> MipLevel);
 		NumElements += NumElementsInLevel;
 	}
 
-	RectLightMipTree.Initialize(sizeof(float), NumElements, PF_R32_FLOAT, BUF_UnorderedAccess | BUF_ShaderResource);
+	Data.RectLightMipTree.Initialize(sizeof(float), NumElements, PF_R32_FLOAT, BUF_UnorderedAccess | BUF_ShaderResource);
 
 	// Execute hierarchical build
 	for (uint32 MipLevel = 0; MipLevel <= MipLevelCount; ++MipLevel)
 	{
 		FComputeFenceRHIRef MipLevelFence = RHICmdList.CreateComputeFence(TEXT("RectLightMipTree Build"));
-		BuildRectLightMipTreeComputeShader->SetParameters(RHICmdList, LightParameters.SourceTexture, RectLightMipTreeDimensions, MipLevel, RectLightMipTree);
-		FIntVector MipLevelDimensions = FIntVector(RectLightMipTreeDimensions.X >> MipLevel, RectLightMipTreeDimensions.Y >> MipLevel, 1);
+		BuildRectLightMipTreeComputeShader->SetParameters(RHICmdList, RhiTexture, Data.RectLightMipTreeDimensions, MipLevel, Data.RectLightMipTree);
+		FIntVector MipLevelDimensions = FIntVector(Data.RectLightMipTreeDimensions.X >> MipLevel, Data.RectLightMipTreeDimensions.Y >> MipLevel, 1);
 		FIntVector NumGroups = FIntVector::DivideAndRoundUp(MipLevelDimensions, FBuildRectLightMipTreeCS::GetGroupSize());
 		DispatchComputeShader(RHICmdList, *BuildRectLightMipTreeComputeShader, NumGroups.X, NumGroups.Y, 1);
-		BuildRectLightMipTreeComputeShader->UnsetParameters(RHICmdList, EResourceTransitionAccess::ERWBarrier, EResourceTransitionPipeline::EComputeToCompute, RectLightMipTree, MipLevelFence);
+		BuildRectLightMipTreeComputeShader->UnsetParameters(RHICmdList, EResourceTransitionAccess::ERWBarrier, EResourceTransitionPipeline::EComputeToCompute, Data.RectLightMipTree, MipLevelFence);
 	}
 	FComputeFenceRHIRef TransitionFence = RHICmdList.CreateComputeFence(TEXT("RectLightMipTree Transition"));
-	BuildRectLightMipTreeComputeShader->UnsetParameters(RHICmdList, EResourceTransitionAccess::ERWBarrier, EResourceTransitionPipeline::EComputeToCompute, RectLightMipTree, TransitionFence);
-}
+	BuildRectLightMipTreeComputeShader->UnsetParameters(RHICmdList, EResourceTransitionAccess::ERWBarrier, EResourceTransitionPipeline::EComputeToCompute, Data.RectLightMipTree, TransitionFence);
 
+	return Data;
+}
 #endif // RHI_RAYTRACING
diff --git a/Engine/Source/Runtime/Engine/Public/RectLightSceneProxy.h b/Engine/Source/Runtime/Engine/Public/RectLightSceneProxy.h
--- a/Engine/Source/Runtime/Engine/Public/RectLightSceneProxy.h
+++ b/Engine/Source/Runtime/Engine/Public/RectLightSceneProxy.h
@@ -11,6 +11,14 @@
 #include "LocalLightSceneProxy.h"
 #include "SceneManagement.h"
 
+struct FRectLightRayTracingData
+{
+#if RHI_RAYTRACING
+	FRWBuffer   RectLightMipTree;
+	FIntVector  RectLightMipTreeDimensions;
+#endif
+};
+
 class FRectLightSceneProxy : public FLocalLightSceneProxy
 {
 public:
@@ -18,11 +26,8 @@ public:
 	float		SourceHeight;
 	float		BarnDoorAngle;
 	float		BarnDoorLength;
+	FRectLightRayTracingData* RayTracingData;
 	UTexture*	SourceTexture;
-#if RHI_RAYTRACING
-	FRWBuffer   RectLightMipTree;
-	FIntVector  RectLightMipTreeDimensions;
-#endif
 
 	FRectLightSceneProxy(const URectLightComponent* Component);
 	virtual ~FRectLightSceneProxy();
@@ -39,8 +44,4 @@ public:
 	* @return True if the whole-scene projected shadow should be used.
 	*/
 	virtual bool GetWholeSceneProjectedShadowInitializer(const FSceneViewFamily& ViewFamily, TArray<FWholeSceneProjectedShadowInitializer, TInlineAllocator<6> >& OutInitializers) const;
-
-#if RHI_RAYTRACING
-	void BuildRectLightMipTree(FRHICommandListImmediate& RHICmdList);
-#endif
 };
\ No newline at end of file
diff --git a/Engine/Source/Runtime/Renderer/Private/RayTracing/RayTracingRectLight.cpp b/Engine/Source/Runtime/Renderer/Private/RayTracing/RayTracingRectLight.cpp
--- a/Engine/Source/Runtime/Renderer/Private/RayTracing/RayTracingRectLight.cpp
+++ b/Engine/Source/Runtime/Renderer/Private/RayTracing/RayTracingRectLight.cpp
@@ -359,8 +359,8 @@ void RenderRayTracingRectLightInternal(
 	RectLightData.Height = 2.0f * LightShaderParameters.SourceLength;
 	RectLightData.Texture = LightShaderParameters.SourceTexture;
 	RectLightData.TextureSampler = RHICreateSamplerState(FSamplerStateInitializerRHI(SF_Bilinear, AM_Border, AM_Border, AM_Border));
-	RectLightData.MipTree = RectLightSceneProxy->RectLightMipTree.SRV;
-	RectLightData.MipTreeDimensions = RectLightSceneProxy->RectLightMipTreeDimensions;
+	RectLightData.MipTree = RectLightSceneProxy->RayTracingData->RectLightMipTree.SRV;
+	RectLightData.MipTreeDimensions = RectLightSceneProxy->RayTracingData->RectLightMipTreeDimensions;
 	RectLightData.MaxNormalBias = GetRaytracingMaxNormalBias();
 	RectLightData.BarnCosAngle = FMath::Cos(FMath::DegreesToRadians(RectLightSceneProxy->BarnDoorAngle));
 	RectLightData.BarnLength = RectLightSceneProxy->BarnDoorLength;
