Back out changelist 5371587
Back out changelist 5377090

(No longer required after 5399372)

diff --git a/Engine/Source/Runtime/Engine/Private/Materials/Material.cpp b/Engine/Source/Runtime/Engine/Private/Materials/Material.cpp
--- a/Engine/Source/Runtime/Engine/Private/Materials/Material.cpp
+++ b/Engine/Source/Runtime/Engine/Private/Materials/Material.cpp
@@ -3956,14 +3956,6 @@ void UMaterial::PostLoad()
 			{
 				if (Texture)
 				{
-					if (Texture->HasAnyFlags(RF_NeedLoad))
-					{
-						FLinkerLoad* Loader = GetLinker();
-						if (ensure(Loader))
-						{
-							Loader->Preload(Texture);
-						}
-					}
 					Texture->ConditionalPostLoad();
 				}
 			}
diff --git a/Engine/Source/Runtime/Engine/Private/Materials/MaterialInstance.cpp b/Engine/Source/Runtime/Engine/Private/Materials/MaterialInstance.cpp
--- a/Engine/Source/Runtime/Engine/Private/Materials/MaterialInstance.cpp
+++ b/Engine/Source/Runtime/Engine/Private/Materials/MaterialInstance.cpp
@@ -3127,14 +3127,6 @@ void UMaterialInstance::PostLoad()
 		UTexture* Texture = TextureParameterValues[ValueIndex].ParameterValue;
 		if( Texture )
 		{
-			if (Texture->HasAnyFlags(RF_NeedLoad))
-			{
-				FLinkerLoad* Loader = GetLinker();
-				if (ensure(Loader))
-				{
-					Loader->Preload(Texture);
-				}
-			}
 			Texture->ConditionalPostLoad();
 		}
 	}
diff --git a/Engine/Source/Runtime/Engine/Private/Particles/ParticleModules.cpp b/Engine/Source/Runtime/Engine/Private/Particles/ParticleModules.cpp
--- a/Engine/Source/Runtime/Engine/Private/Particles/ParticleModules.cpp
+++ b/Engine/Source/Runtime/Engine/Private/Particles/ParticleModules.cpp
@@ -1219,14 +1219,6 @@ void UParticleModuleRequired::PostLoad()
 	{
 		if (CutoutTexture)
 		{
-			if (CutoutTexture->HasAnyFlags(RF_NeedLoad))
-			{
-				FLinkerLoad* Loader = GetLinker();
-				if (ensure(Loader))
-				{
-					Loader->Preload(CutoutTexture);
-				}
-			}
 			CutoutTexture->ConditionalPostLoad();
 			CacheDerivedData();
 		}	
