Make the PostDuplicate function for blueprints also call the PostDuplicate function on the CDO of their generated class.

diff --git a/Engine/Source/Runtime/Engine/Private/Blueprint.cpp b/Engine/Source/Runtime/Engine/Private/Blueprint.cpp
--- a/Engine/Source/Runtime/Engine/Private/Blueprint.cpp
+++ b/Engine/Source/Runtime/Engine/Private/Blueprint.cpp
@@ -598,6 +598,11 @@ void UBlueprint::PostDuplicate(bool bDuplicateForPIE)
 	{
 		FBlueprintEditorUtils::PostDuplicateBlueprint(this, bDuplicateForPIE);
 	}
+
+	if (GeneratedClass)
+	{
+		GeneratedClass->GetDefaultObject()->PostDuplicate(bDuplicateForPIE);
+	}
 }
 
 extern COREUOBJECT_API bool GBlueprintUseCompilationManager;
