As a part of prev undo, bump shader version to clear out stale DDC.

diff --git a/Engine/Source/Runtime/Engine/Private/ShaderDerivedDataVersion.h b/Engine/Source/Runtime/Engine/Private/ShaderDerivedDataVersion.h
--- a/Engine/Source/Runtime/Engine/Private/ShaderDerivedDataVersion.h
+++ b/Engine/Source/Runtime/Engine/Private/ShaderDerivedDataVersion.h
@@ -11,5 +11,5 @@
 // In case of merge conflicts with DDC versions, you MUST generate a new GUID and set this new
 // guid as version
 
-#define GLOBALSHADERMAP_DERIVEDDATA_VER			TEXT("4FFBFDF601234B80893B59E7779F89E7")
-#define MATERIALSHADERMAP_DERIVEDDATA_VER		TEXT("FE2EEAF7F8BC456195582BBB000D17CD")
+#define GLOBALSHADERMAP_DERIVEDDATA_VER			TEXT("73D907EC61CB431EB3FEE5F78C520EF6")
+#define MATERIALSHADERMAP_DERIVEDDATA_VER		TEXT("3C50894C66CC4D2690D7BCA3554F7F3B")
