#ue4 - Improved comments for bIgnoreClientMovementErrorChecksAndCorrection. Changed it to not ignore corrections on clients as that is unecessary since server accepts the move. Added UCharacterMovementComponent::ServerShouldUseAuthoritativePosition() to make it more flexible to be able to selectively toggle client authoritative movement. Added bServerAcceptClientAuthoritativePosition to allow component to turn on authoritative movement (not relying on global GameNetworkManager flag).

diff --git a/Engine/Source/Runtime/Engine/Classes/GameFramework/CharacterMovementComponent.h b/Engine/Source/Runtime/Engine/Classes/GameFramework/CharacterMovementComponent.h
--- a/Engine/Source/Runtime/Engine/Classes/GameFramework/CharacterMovementComponent.h
+++ b/Engine/Source/Runtime/Engine/Classes/GameFramework/CharacterMovementComponent.h
@@ -908,14 +908,25 @@ public:
 	uint8 bNetworkMovementModeChanged:1;
 
 	/** 
-	 * True when we should ignore server location difference checks for client error on this movement component 
+	 * If true, we should ignore server location difference checks for client error on this movement component.
 	 * This can be useful when character is moving at extreme speeds for a duration and you need it to look
-	 * smooth on clients. Make sure to disable when done, as this would break this character's server-client
-	 * movement correction.
+	 * smooth on clients without the server correcting the client. Make sure to disable when done, as this would
+	 * break this character's server-client movement correction.
+	 * @see bServerAcceptClientAuthoritativePosition, ServerCheckClientError()
 	 */
 	UPROPERTY(Transient, Category="Character Movement", EditAnywhere, BlueprintReadWrite)
 	uint8 bIgnoreClientMovementErrorChecksAndCorrection:1;
 
+	/**
+	 * If true, and server does not detect client position error, server will copy the client movement location/velocity/etc after simulating the move.
+	 * This can be useful for short bursts of movement that are difficult to sync over the network.
+	 * Note that if bIgnoreClientMovementErrorChecksAndCorrection is used, this means the server will not detect an error.
+	 * Also see GameNetworkManager->ClientAuthorativePosition which permanently enables this behavior.
+	 * @see bIgnoreClientMovementErrorChecksAndCorrection, ServerShouldUseAuthoritativePosition()
+	 */
+	UPROPERTY(Transient, Category="Character Movement", EditAnywhere, BlueprintReadWrite)
+	uint8 bServerAcceptClientAuthoritativePosition : 1;
+
 	/**
 	 * If true, event NotifyJumpApex() to CharacterOwner's controller when at apex of jump. Is cleared when event is triggered.
 	 * By default this is off, and if you want the event to fire you typically set it to true when movement mode changes to "Falling" from another mode (see OnMovementModeChanged).
@@ -2151,6 +2162,11 @@ protected:
 	 */
 	virtual bool ServerCheckClientError(float ClientTimeStamp, float DeltaTime, const FVector& Accel, const FVector& ClientWorldLocation, const FVector& RelativeClientLocation, UPrimitiveComponent* ClientMovementBase, FName ClientBaseBoneName, uint8 ClientMovementMode);
 
+	/**
+	 * If ServerCheckClientError() does not find an error, this determines if the server should also copy the client's movement params rather than keep the server sim result.
+	 */
+	virtual bool ServerShouldUseAuthoritativePosition(float ClientTimeStamp, float DeltaTime, const FVector& Accel, const FVector& ClientWorldLocation, const FVector& RelativeClientLocation, UPrimitiveComponent* ClientMovementBase, FName ClientBaseBoneName, uint8 ClientMovementMode);
+
 	/* Process a move at the given time stamp, given the compressed flags representing various events that occurred (ie jump). */
 	virtual void MoveAutonomous( float ClientTimeStamp, float DeltaTime, uint8 CompressedFlags, const FVector& NewAccel);
 
diff --git a/Engine/Source/Runtime/Engine/Private/Components/CharacterMovementComponent.cpp b/Engine/Source/Runtime/Engine/Private/Components/CharacterMovementComponent.cpp
--- a/Engine/Source/Runtime/Engine/Private/Components/CharacterMovementComponent.cpp
+++ b/Engine/Source/Runtime/Engine/Private/Components/CharacterMovementComponent.cpp
@@ -498,6 +498,7 @@ UCharacterMovementComponent::UCharacterMovementComponent(const FObjectInitialize
 	bImpartBaseVelocityZ = true;
 	bImpartBaseAngularVelocity = true;
 	bIgnoreClientMovementErrorChecksAndCorrection = false;
+	bServerAcceptClientAuthoritativePosition = false;
 	bAlwaysCheckFloor = true;
 
 	// default character can jump, walk, and swim
@@ -7651,17 +7652,6 @@ bool UCharacterMovementComponent::ClientUpdatePositionAfterServerUpdate()
 		return false;
 	}
 
-	if (bIgnoreClientMovementErrorChecksAndCorrection)
-	{
-#if !UE_BUILD_SHIPPING
-		if (CharacterMovementCVars::NetShowCorrections != 0)
-		{
-			UE_LOG(LogNetPlayerMovement, Warning, TEXT("*** Client: %s is set to ignore error checks and corrections with %d saved moves in queue."), *GetNameSafe(CharacterOwner), ClientData->SavedMoves.Num());
-		}
-#endif // !UE_BUILD_SHIPPING
-		return false;
-	}
-
 	ClientData->bUpdatePosition = false;
 
 	// Don't do any network position updates on things running PHYS_RigidBody
@@ -8753,8 +8743,7 @@ void UCharacterMovementComponent::ServerMoveHandleClientError(float ClientTimeSt
 	}
 	else
 	{
-		const AGameNetworkManager* GameNetworkManager = (const AGameNetworkManager*)(AGameNetworkManager::StaticClass()->GetDefaultObject());
-		if (GameNetworkManager->ClientAuthorativePosition)
+		if (ServerShouldUseAuthoritativePosition(ClientTimeStamp, DeltaTime, Accel, ClientLoc, RelativeClientLoc, ClientMovementBase, ClientBaseBoneName, ClientMovementMode))
 		{
 			const FVector LocDiff = UpdatedComponent->GetComponentLocation() - ClientLoc; //-V595
 			if (!LocDiff.IsZero() || ClientMovementMode != PackNetworkMovementMode() || GetMovementBase() != ClientMovementBase || (CharacterOwner && CharacterOwner->GetBasedMovement().BoneName != ClientBaseBoneName))
@@ -8805,12 +8794,21 @@ bool UCharacterMovementComponent::ServerCheckClientError(float ClientTimeStamp,
 			RootMotionSourceDebug::PrintOnScreen(*CharacterOwner, AdjustedDebugString);
 		}
 #endif
+
 		const AGameNetworkManager* GameNetworkManager = (const AGameNetworkManager*)(AGameNetworkManager::StaticClass()->GetDefaultObject());
 		if (GameNetworkManager->ExceedsAllowablePositionError(LocDiff))
 		{
 			bNetworkLargeClientCorrection = (LocDiff.SizeSquared() > FMath::Square(NetworkLargeClientCorrectionDistance));
 			return true;
 		}
+
+		// Check for disagreement in movement mode
+		const uint8 CurrentPackedMovementMode = PackNetworkMovementMode();
+		if (CurrentPackedMovementMode != ClientMovementMode)
+		{
+			return true;
+		}
+
 #if !(UE_BUILD_SHIPPING || UE_BUILD_TEST)
 		if (CharacterMovementCVars::NetForceClientAdjustmentPercent > SMALL_NUMBER)
 		{
@@ -8832,9 +8830,19 @@ bool UCharacterMovementComponent::ServerCheckClientError(float ClientTimeStamp,
 #endif // !UE_BUILD_SHIPPING
 	}
 
-	// Check for disagreement in movement mode
-	const uint8 CurrentPackedMovementMode = PackNetworkMovementMode();
-	if (CurrentPackedMovementMode != ClientMovementMode)
+	return false;
+}
+
+
+bool UCharacterMovementComponent::ServerShouldUseAuthoritativePosition(float ClientTimeStamp, float DeltaTime, const FVector& Accel, const FVector& ClientWorldLocation, const FVector& RelativeClientLocation, UPrimitiveComponent* ClientMovementBase, FName ClientBaseBoneName, uint8 ClientMovementMode)
+{
+	if (bServerAcceptClientAuthoritativePosition)
+	{
+		return true;
+	}
+
+	const AGameNetworkManager* GameNetworkManager = (const AGameNetworkManager*)(AGameNetworkManager::StaticClass()->GetDefaultObject());
+	if (GameNetworkManager->ClientAuthorativePosition)
 	{
 		return true;
 	}
@@ -8842,7 +8850,6 @@ bool UCharacterMovementComponent::ServerCheckClientError(float ClientTimeStamp,
 	return false;
 }
 
-
 bool UCharacterMovementComponent::ServerMove_Validate(float TimeStamp, FVector_NetQuantize10 InAccel, FVector_NetQuantize100 ClientLoc, uint8 MoveFlags, uint8 ClientRoll, uint32 View, UPrimitiveComponent* ClientMovementBase, FName ClientBaseBoneName, uint8 ClientMovementMode)
 {
 	return true;
