Moving instance uniform expression UB recreate from InitResources to FMaterialUpdateContext. Recreating in InitResources may happen outside of the material update context and leave a stale pointer to UB inside cached mesh draw commands.
#rnx

diff --git a/Engine/Source/Runtime/Engine/Private/Materials/MaterialInstance.cpp b/Engine/Source/Runtime/Engine/Private/Materials/MaterialInstance.cpp
--- a/Engine/Source/Runtime/Engine/Private/Materials/MaterialInstance.cpp
+++ b/Engine/Source/Runtime/Engine/Private/Materials/MaterialInstance.cpp
@@ -484,7 +484,7 @@ void UMaterialInstance::InitResources()
 
 	PropagateDataToMaterialProxy();
 
-	CacheMaterialInstanceUniformExpressions(this, true);
+	CacheMaterialInstanceUniformExpressions(this);
 }
 
 const UMaterial* UMaterialInstance::GetMaterial() const
diff --git a/Engine/Source/Runtime/Engine/Private/Materials/MaterialShared.cpp b/Engine/Source/Runtime/Engine/Private/Materials/MaterialShared.cpp
--- a/Engine/Source/Runtime/Engine/Private/Materials/MaterialShared.cpp
+++ b/Engine/Source/Runtime/Engine/Private/Materials/MaterialShared.cpp
@@ -2779,6 +2779,7 @@ FMaterialUpdateContext::~FMaterialUpdateContext()
 			}
 		}
 
+		MI->RecacheUniformExpressions(true);
 		MI->InitStaticPermutation();//bHasStaticPermutation can change.
 		if (MI->bHasStaticPermutationResource)
 		{
