When changing FMaterialInstanceResource parent, we need also to recreate it's uniform expression UB as new parent material can have a different UB layout.

diff --git a/Engine/Source/Runtime/Engine/Private/Materials/MaterialInstance.cpp b/Engine/Source/Runtime/Engine/Private/Materials/MaterialInstance.cpp
--- a/Engine/Source/Runtime/Engine/Private/Materials/MaterialInstance.cpp
+++ b/Engine/Source/Runtime/Engine/Private/Materials/MaterialInstance.cpp
@@ -311,7 +311,7 @@ void FMaterialInstanceResource::GameThread_SetParent(UMaterialInterface* ParentM
 			[Resource, ParentMaterialInterface](FRHICommandListImmediate& RHICmdList)
 			{
 				Resource->Parent = ParentMaterialInterface;
-				Resource->InvalidateUniformExpressionCache(false);
+				Resource->InvalidateUniformExpressionCache(true);
 			});
 
 		if (OldParent)
