Fixed missing shaders when enabling debug view shader modes in PIE and non cooked game.

diff --git a/Engine/Source/Editor/MaterialEditor/Private/MaterialEditorUtilities.cpp b/Engine/Source/Editor/MaterialEditor/Private/MaterialEditorUtilities.cpp
--- a/Engine/Source/Editor/MaterialEditor/Private/MaterialEditorUtilities.cpp
+++ b/Engine/Source/Editor/MaterialEditor/Private/MaterialEditorUtilities.cpp
@@ -705,7 +705,7 @@ void FMaterialEditorUtilities::BuildTextureStreamingData(UMaterialInterface* Upd
 		};
 
 		// Here we need a full rebuild since the shader changed. Although don't wait for previous shaders to fasten the process.
-		if (CompileDebugViewModeShaders(DVSM_OutputMaterialTextureScales, QualityLevel, FeatureLevel, true, false, Materials, SlowTask))
+		if (CompileDebugViewModeShaders(DVSM_OutputMaterialTextureScales, QualityLevel, FeatureLevel, true, false, Materials, &SlowTask))
 		{
 			FMaterialUtilities::FExportErrorManager ExportErrors(FeatureLevel);
 			for (UMaterialInterface* MaterialInterface : Materials)
diff --git a/Engine/Source/Editor/UnrealEd/Private/EditorBuildUtils.cpp b/Engine/Source/Editor/UnrealEd/Private/EditorBuildUtils.cpp
--- a/Engine/Source/Editor/UnrealEd/Private/EditorBuildUtils.cpp
+++ b/Engine/Source/Editor/UnrealEd/Private/EditorBuildUtils.cpp
@@ -1103,14 +1103,14 @@ bool FEditorBuildUtils::EditorBuildTextureStreaming(UWorld* InWorld, EViewModeIn
 	if (bNeedsMaterialData)
 	{
 		TSet<UMaterialInterface*> Materials;
-		if (!GetUsedMaterialsInWorld(InWorld, Materials, BuildTextureStreamingTask))
+		if (!GetUsedMaterialsInWorld(InWorld, Materials, &BuildTextureStreamingTask))
 		{
 			return false;
 		}
 
 		if (Materials.Num())
 		{
-			if (!CompileDebugViewModeShaders(DVSM_OutputMaterialTextureScales, QualityLevel, FeatureLevel, SelectedViewMode == VMI_Unknown, true, Materials, BuildTextureStreamingTask))
+			if (!CompileDebugViewModeShaders(DVSM_OutputMaterialTextureScales, QualityLevel, FeatureLevel, SelectedViewMode == VMI_Unknown, true, Materials, &BuildTextureStreamingTask))
 			{
 				return false;
 			}
@@ -1228,7 +1228,7 @@ bool FEditorBuildUtils::EditorBuildMaterialTextureStreamingData(UPackage* Packag
 	const float OneOverNumMaterials = 1.f / FMath::Max(1.f, (float)Materials.Num());
 
 	bool bAnyPackagesDirtied = false;
-	if (CompileDebugViewModeShaders(DVSM_OutputMaterialTextureScales, QualityLevel, FeatureLevel, true, true, Materials, SlowTask))
+	if (CompileDebugViewModeShaders(DVSM_OutputMaterialTextureScales, QualityLevel, FeatureLevel, true, true, Materials, &SlowTask))
 	{
 		FMaterialUtilities::FExportErrorManager ExportErrors(FeatureLevel);
 		for (UMaterialInterface* MaterialInterface : Materials)
@@ -1304,7 +1304,7 @@ bool FEditorBuildUtils::CompileViewModeShaders(UWorld* InWorld, EViewModeIndex S
 	CollectGarbage(GARBAGE_COLLECTION_KEEPFLAGS);
 
 	TSet<UMaterialInterface*> Materials;
-	if (!GetUsedMaterialsInWorld(InWorld, Materials, CompileShaderTask))
+	if (!GetUsedMaterialsInWorld(InWorld, Materials, &CompileShaderTask))
 	{
 		return false;
 	}
@@ -1321,7 +1321,7 @@ bool FEditorBuildUtils::CompileViewModeShaders(UWorld* InWorld, EViewModeIndex S
 		}
 		else if (SelectedViewMode == VMI_RequiredTextureResolution)*/
 		{
-			if (!CompileDebugViewModeShaders(DebugViewMode, QualityLevel, FeatureLevel, false, true, Materials, CompileShaderTask))
+			if (!CompileDebugViewModeShaders(DebugViewMode, QualityLevel, FeatureLevel, false, true, Materials, &CompileShaderTask))
 			{
 				return false;
 			}
@@ -1360,7 +1360,7 @@ bool FEditorBuildUtils::CompileShadersComplexityViewMode(EMaterialQualityLevel::
 	check(Materials.Num());
 
 	// Finish compiling pending shaders first.
-	if (!WaitForShaderCompilation(LOCTEXT("CompileShaders_Complexity_FinishPendingShadersCompilation", "Waiting For Pending Shaders Compilation"), ProgressTask))
+	if (!WaitForShaderCompilation(LOCTEXT("CompileShaders_Complexity_FinishPendingShadersCompilation", "Waiting For Pending Shaders Compilation"), &ProgressTask))
 	{
 		return false;
 	}
@@ -1391,7 +1391,7 @@ bool FEditorBuildUtils::CompileShadersComplexityViewMode(EMaterialQualityLevel::
 	}
 
 	// wait for compilation to be done and copy the number of instruction from the compiled shaders to the emulated shader set
-	if (WaitForShaderCompilation(LOCTEXT("OfflineShaderCompilation", "Offline Shader Compilation"), ProgressTask))
+	if (WaitForShaderCompilation(LOCTEXT("OfflineShaderCompilation", "Offline Shader Compilation"), &ProgressTask))
 	{
 		FSuspendRenderingThread SuspendObject(false);
 
diff --git a/Engine/Source/Editor/UnrealEd/Private/EditorViewportClient.cpp b/Engine/Source/Editor/UnrealEd/Private/EditorViewportClient.cpp
--- a/Engine/Source/Editor/UnrealEd/Private/EditorViewportClient.cpp
+++ b/Engine/Source/Editor/UnrealEd/Private/EditorViewportClient.cpp
@@ -3634,14 +3634,12 @@ void FEditorViewportClient::Draw(FViewport* InViewport, FCanvas* Canvas)
 
 	ViewFamily.EngineShowFlags = EngineShowFlags;
 
-	if (GIsEditor && World && ViewFamily.GetDebugViewShaderMode() != DVSM_None && HasMissingDebugViewModeShaders(true))
+	if (World && ViewFamily.GetDebugViewShaderMode() != DVSM_None && HasMissingDebugViewModeShaders(true))
 	{
-		FScopedSlowTask CompileShaderTask(3.f, LOCTEXT("CompileMissingViewModeShaders", "Compiling Missing ViewMode Shaders")); // { Get Used Materials, Sync Pending Shader, Wait for Compilation }
-		// CompileShaderTask.MakeDialog(true);
 		TSet<UMaterialInterface*> Materials;
-		if (GetUsedMaterialsInWorld(World, Materials, CompileShaderTask))
+		if (GetUsedMaterialsInWorld(World, Materials, nullptr))
 		{
-			CompileDebugViewModeShaders(ViewFamily.GetDebugViewShaderMode(), GetCachedScalabilityCVars().MaterialQualityLevel, ViewFamily.GetFeatureLevel(), false, false, Materials, CompileShaderTask);
+			CompileDebugViewModeShaders(ViewFamily.GetDebugViewShaderMode(), GetCachedScalabilityCVars().MaterialQualityLevel, ViewFamily.GetFeatureLevel(), false, false, Materials, nullptr);
 		}
 	}
 
diff --git a/Engine/Source/Runtime/Engine/Private/DebugViewModeHelpers.cpp b/Engine/Source/Runtime/Engine/Private/DebugViewModeHelpers.cpp
--- a/Engine/Source/Runtime/Engine/Private/DebugViewModeHelpers.cpp
+++ b/Engine/Source/Runtime/Engine/Private/DebugViewModeHelpers.cpp
@@ -83,7 +83,7 @@ int32 GetNumActorsInWorld(UWorld* InWorld)
 	return ActorCount;
 }
 
-bool WaitForShaderCompilation(const FText& Message, FSlowTask& ProgressTask)
+bool WaitForShaderCompilation(const FText& Message, FSlowTask* ProgressTask)
 {
 	FlushRenderingCommands();
 
@@ -104,19 +104,22 @@ bool WaitForShaderCompilation(const FText& Message, FSlowTask& ProgressTask)
 				const int32 NumberOfShadersCompiledThisFrame = RemainingShaders - RemainingShadersThisFrame;
 
 				const float FrameProgress = (float)NumberOfShadersCompiledThisFrame / (float)NumShadersToBeCompiled;
-				ProgressTask.EnterProgressFrame(FrameProgress);
-				SlowTask.EnterProgressFrame(FrameProgress);
-				if (GWarn->ReceivedUserCancel())
+				if (ProgressTask)
 				{
-					return false;
+					ProgressTask->EnterProgressFrame(FrameProgress);
+					SlowTask.EnterProgressFrame(FrameProgress);
+					if (GWarn->ReceivedUserCancel())
+					{
+						return false;
+					}
 				}
 			}
 			RemainingShaders = RemainingShadersThisFrame;
 		}
 	}
-	else
+	else if (ProgressTask)
 	{
-		ProgressTask.EnterProgressFrame();
+		ProgressTask->EnterProgressFrame();
 		if (GWarn->ReceivedUserCancel())
 		{
 			return false;
@@ -134,7 +137,7 @@ bool WaitForShaderCompilation(const FText& Message, FSlowTask& ProgressTask)
  *
  * @return true if the operation is a success, false if it was canceled.
  */
-bool GetUsedMaterialsInWorld(UWorld* InWorld, OUT TSet<UMaterialInterface*>& OutMaterials, FSlowTask& ProgressTask)
+bool GetUsedMaterialsInWorld(UWorld* InWorld, OUT TSet<UMaterialInterface*>& OutMaterials, FSlowTask* ProgressTask)
 {
 #if WITH_EDITORONLY_DATA
 	if (!InWorld)
@@ -145,7 +148,10 @@ bool GetUsedMaterialsInWorld(UWorld* InWorld, OUT TSet<UMaterialInterface*>& Out
 	const int32 NumActorsInWorld = GetNumActorsInWorld(InWorld);
 	if (!NumActorsInWorld)
 	{
-		ProgressTask.EnterProgressFrame();
+		if (ProgressTask)
+		{
+			ProgressTask->EnterProgressFrame();
+		}
 		return true;
 	}
 
@@ -163,11 +169,14 @@ bool GetUsedMaterialsInWorld(UWorld* InWorld, OUT TSet<UMaterialInterface*>& Out
 
 		for (AActor* Actor : Level->Actors)
 		{
-			ProgressTask.EnterProgressFrame(OneOverNumActorsInWorld);
-			SlowTask.EnterProgressFrame(OneOverNumActorsInWorld);
-			if (GWarn->ReceivedUserCancel())
+			if (ProgressTask)
 			{
-				return false;
+				ProgressTask->EnterProgressFrame(OneOverNumActorsInWorld);
+				SlowTask.EnterProgressFrame(OneOverNumActorsInWorld);
+				if (GWarn->ReceivedUserCancel())
+				{
+					return false;
+				}
 			}
 
 			// Check the actor after incrementing the progress.
@@ -215,7 +224,7 @@ bool GetUsedMaterialsInWorld(UWorld* InWorld, OUT TSet<UMaterialInterface*>& Out
  * @param Materials			The materials to update, the one that failed compilation will be removed (IN OUT).
  * @return true if the operation is a success, false if it was canceled.
  */
-bool CompileDebugViewModeShaders(EDebugViewShaderMode ShaderMode, EMaterialQualityLevel::Type QualityLevel, ERHIFeatureLevel::Type FeatureLevel, bool bFullRebuild, bool bWaitForPreviousShaders, TSet<UMaterialInterface*>& Materials, FSlowTask& ProgressTask)
+bool CompileDebugViewModeShaders(EDebugViewShaderMode ShaderMode, EMaterialQualityLevel::Type QualityLevel, ERHIFeatureLevel::Type FeatureLevel, bool bFullRebuild, bool bWaitForPreviousShaders, TSet<UMaterialInterface*>& Materials, FSlowTask* ProgressTask)
 {
 #if WITH_EDITORONLY_DATA
 	if (!GShaderCompilingManager || !Materials.Num())
diff --git a/Engine/Source/Runtime/Engine/Private/GameViewportClient.cpp b/Engine/Source/Runtime/Engine/Private/GameViewportClient.cpp
--- a/Engine/Source/Runtime/Engine/Private/GameViewportClient.cpp
+++ b/Engine/Source/Runtime/Engine/Private/GameViewportClient.cpp
@@ -1125,6 +1125,14 @@ void UGameViewportClient::Draw(FViewport* InViewport, FCanvas* SceneCanvas)
 		// Force enable view family show flag for HighDPI derived's screen percentage.
 		ViewFamily.EngineShowFlags.ScreenPercentage = true;
 	}
+	if (ViewFamily.GetDebugViewShaderMode() != DVSM_None && HasMissingDebugViewModeShaders(true))
+	{
+		TSet<UMaterialInterface*> Materials;
+		if (GetUsedMaterialsInWorld(MyWorld, Materials, nullptr))
+		{
+			CompileDebugViewModeShaders(ViewFamily.GetDebugViewShaderMode(), GetCachedScalabilityCVars().MaterialQualityLevel, ViewFamily.GetFeatureLevel(), false, false, Materials, nullptr);
+		}
+	}
 #endif
 
 	ViewFamily.ViewExtensions = GEngine->ViewExtensions->GatherActiveExtensions(InViewport);
diff --git a/Engine/Source/Runtime/Engine/Public/DebugViewModeHelpers.h b/Engine/Source/Runtime/Engine/Public/DebugViewModeHelpers.h
--- a/Engine/Source/Runtime/Engine/Public/DebugViewModeHelpers.h
+++ b/Engine/Source/Runtime/Engine/Public/DebugViewModeHelpers.h
@@ -11,6 +11,7 @@
 #include "SceneTypes.h"
 
 class UMaterialInterface;
+struct FSlowTask;
 
 /** 
  * Enumeration for different Quad Overdraw visualization mode.
@@ -41,10 +42,10 @@ FORCEINLINE bool AllowDebugViewShaderMode(EDebugViewShaderMode ShaderMode, EShad
 #endif
 
 ENGINE_API int32 GetNumActorsInWorld(UWorld* InWorld);
-ENGINE_API bool GetUsedMaterialsInWorld(UWorld* InWorld, OUT TSet<UMaterialInterface*>& OutMaterials, struct FSlowTask& Task);
-ENGINE_API bool CompileDebugViewModeShaders(EDebugViewShaderMode Mode, EMaterialQualityLevel::Type QualityLevel, ERHIFeatureLevel::Type FeatureLevel, bool bFullRebuild, bool bWaitForPreviousShaders, TSet<UMaterialInterface*>& Materials, FSlowTask& ProgressTask);
+ENGINE_API bool GetUsedMaterialsInWorld(UWorld* InWorld, OUT TSet<UMaterialInterface*>& OutMaterials, FSlowTask* Task);
+ENGINE_API bool CompileDebugViewModeShaders(EDebugViewShaderMode Mode, EMaterialQualityLevel::Type QualityLevel, ERHIFeatureLevel::Type FeatureLevel, bool bFullRebuild, bool bWaitForPreviousShaders, TSet<UMaterialInterface*>& Materials, FSlowTask* ProgressTask);
 ENGINE_API bool HasMissingDebugViewModeShaders(bool bClearFlag);
 
-ENGINE_API bool WaitForShaderCompilation(const FText& Message, FSlowTask& ProgressTask);
+ENGINE_API bool WaitForShaderCompilation(const FText& Message, FSlowTask* ProgressTask);
 
 
