Eliminate the need for calculating offsets for Fast Array Delta Serialization.
It is compiler specific and error prone.

diff --git a/Engine/Plugins/Runtime/GameplayAbilities/Source/GameplayAbilities/Private/GameplayEffect.cpp b/Engine/Plugins/Runtime/GameplayAbilities/Source/GameplayAbilities/Private/GameplayEffect.cpp
--- a/Engine/Plugins/Runtime/GameplayAbilities/Source/GameplayAbilities/Private/GameplayEffect.cpp
+++ b/Engine/Plugins/Runtime/GameplayAbilities/Source/GameplayAbilities/Private/GameplayEffect.cpp
@@ -1874,7 +1874,7 @@ FActiveGameplayEffectsContainer::FActiveGameplayEffectsContainer()
 	, PendingGameplayEffectHead(nullptr)
 {
 	PendingGameplayEffectNext = &PendingGameplayEffectHead;
-	bUseDeltaStructSerialization = true;
+	SetDeltaSerializationEnabled(true);
 }
 
 FActiveGameplayEffectsContainer::~FActiveGameplayEffectsContainer()
diff --git a/Engine/Source/Runtime/Engine/Classes/Engine/NetSerialization.h b/Engine/Source/Runtime/Engine/Classes/Engine/NetSerialization.h
--- a/Engine/Source/Runtime/Engine/Classes/Engine/NetSerialization.h
+++ b/Engine/Source/Runtime/Engine/Classes/Engine/NetSerialization.h
@@ -395,6 +395,18 @@ struct FFastArraySerializer_FastArrayDeltaSerialize_FIdxIDPair
 	int32	ID;
 };
 
+UENUM()
+enum class EFastArraySerializerDeltaFlags : uint8
+{
+	None,								//! No flags.
+	HasBeenSerialized = 1 << 0,			//! Set when serialization at least once (i.e., this struct has been written or read).
+	HasDeltaBeenRequested = 1 << 1,		//! Set if users requested Delta Serialization for this struct.
+	IsUsingDeltaSerialization = 1 << 2,	//! This will remain unset until we've serialized at least once.
+										//! At that point, this will be set if delta serialization was requested and
+										//! we support it.
+};
+ENUM_CLASS_FLAGS(EFastArraySerializerDeltaFlags);
+
 /** Base struct for wrapping the array used in Fast TArray Replication */
 USTRUCT()
 struct FFastArraySerializer
@@ -406,9 +418,7 @@ struct FFastArraySerializer
 		, ArrayReplicationKey(0)
 		, CachedNumItems(INDEX_NONE)
 		, CachedNumItemsToConsiderForWriting(INDEX_NONE)
-		, bUseDeltaStructSerialization(false)
-		, bCachedUseDeltaStructSerialization(false)
-		, bHasBeenSerialized(false) { }
+		, DeltaFlags(EFastArraySerializerDeltaFlags::None) { }
 
 	~FFastArraySerializer() {}
 
@@ -419,6 +429,7 @@ struct FFastArraySerializer
 	int32 IDCounter;
 
 	/** Counter used to track array replication. */
+	UPROPERTY(NotReplicated)
 	int32 ArrayReplicationKey;
 
 	/** List of items that need to be re-serialized when the referenced objects are mapped */
@@ -507,14 +518,28 @@ struct FFastArraySerializer
 		return !bIsWritingOnClient || Item.ReplicationID != INDEX_NONE;
 	}
 
-	const bool IsUsingStructDeltaSerialization() const
+	void SetDeltaSerializationEnabled(const bool bEnabled)
 	{
-		return bHasBeenSerialized ? bCachedUseDeltaStructSerialization : bUseDeltaStructSerialization;
+		if (!EnumHasAnyFlags(DeltaFlags, EFastArraySerializerDeltaFlags::HasBeenSerialized))
+		{
+			if (bEnabled)
+			{
+				DeltaFlags |= EFastArraySerializerDeltaFlags::HasDeltaBeenRequested;
+			}
+			else
+			{
+				DeltaFlags &= ~EFastArraySerializerDeltaFlags::HasDeltaBeenRequested;
+			}
+		}
+		else
+		{
+			UE_LOG(LogNetFastTArray, Log, TEXT("FFastArraySerializer::SetDeltaSerializationEnabled - Called after array has been serialized. Ignoring"));
+		}
 	}
 
-	const bool HasBeenSerialized() const
+	const EFastArraySerializerDeltaFlags GetDeltaSerializationFlags() const
 	{
-		return bHasBeenSerialized;
+		DeltaFlags;
 	}
 
 private:
@@ -641,15 +666,8 @@ private:
 	int32 CachedNumItems;
 	int32 CachedNumItemsToConsiderForWriting;
 
-protected:
-
-	uint32 bUseDeltaStructSerialization : 1;
-
-private:
-
-	uint32 bCachedUseDeltaStructSerialization : 1;
-
-	uint32 bHasBeenSerialized : 1;
+	UPROPERTY(NotReplicated)
+	EFastArraySerializerDeltaFlags DeltaFlags;
 };
 
 template<typename Type, typename SerializerType>
@@ -1040,7 +1058,7 @@ bool FFastArraySerializer::FastArrayDeltaSerialize(TArray<Type> &Items, FNetDelt
 	// anything from the server (Net Conditions, Static Actors, etc.)
 	// That should be fine though, because none of the GUID Tracking work should actually do anything
 	// until after we've received.
-	if (ArraySerializer.bHasBeenSerialized && ArraySerializer.bCachedUseDeltaStructSerialization)
+	if (EnumHasAllFlags(ArraySerializer.DeltaFlags, EFastArraySerializerDeltaFlags::IsUsingDeltaSerialization))
 	{
 		return FastArrayDeltaSerialize_DeltaSerializeStructs(Items, Parms, ArraySerializer);
 	}
@@ -1202,13 +1220,12 @@ bool FFastArraySerializer::FastArrayDeltaSerialize(TArray<Type> &Items, FNetDelt
 	// If we've made it this far, it means that we're going to be serializing something.
 	// So, it should be safe for us to update our cached state.
 	// Also, make sure that we hit the right path if we need to.
-	if (!ArraySerializer.bHasBeenSerialized)
+	if (!EnumHasAnyFlags(ArraySerializer.DeltaFlags, EFastArraySerializerDeltaFlags::HasBeenSerialized))
 	{
-		ArraySerializer.bHasBeenSerialized = true;
-		ArraySerializer.bCachedUseDeltaStructSerialization = ArraySerializer.bUseDeltaStructSerialization && Parms.bSupportsFastArrayDeltaStructSerialization;
-
-		if (ArraySerializer.bCachedUseDeltaStructSerialization)
+		ArraySerializer.DeltaFlags |= EFastArraySerializerDeltaFlags::HasBeenSerialized;
+		if (Parms.bSupportsFastArrayDeltaStructSerialization && EnumHasAnyFlags(ArraySerializer.DeltaFlags, EFastArraySerializerDeltaFlags::HasDeltaBeenRequested))
 		{
+			ArraySerializer.DeltaFlags |= EFastArraySerializerDeltaFlags::IsUsingDeltaSerialization;
 			return FastArrayDeltaSerialize_DeltaSerializeStructs(Items, Parms, ArraySerializer);
 		}
 	}
@@ -1421,8 +1438,11 @@ struct FFastArrayDeltaSerializeParams
 {
 	FNetDeltaSerializeInfo& DeltaSerializeInfo;
 	FFastArraySerializer& ArraySerializer;
-	const PTRINT PatchedItemOffset;
-	const PTRINT PatchedArrayOffset;
+
+	const TFunction<void(void*, const FFastArrayDeltaSerializeParams&)> PreReplicatedRemove;
+	const TFunction<void(void*, const FFastArrayDeltaSerializeParams&)> PostReplicatedAdd;
+	const TFunction<void(void*, const FFastArrayDeltaSerializeParams&)> PostReplicatedChange;
+	const TFunction<void(void*, const FFastArrayDeltaSerializeParams&, const uint32)> ReceivedItem;
 
 	TArray<FFastArraySerializer_FastArrayDeltaSerialize_FIdxIDPair, TInlineAllocator<8>>* WriteChangedElements = nullptr;
 	FNetFastTArrayBaseState* WriteBaseState = nullptr;
@@ -1435,6 +1455,42 @@ struct FFastArrayDeltaSerializeParams
 template<typename Type, typename SerializerType>
 bool FFastArraySerializer::FastArrayDeltaSerialize_DeltaSerializeStructs(TArray<Type>& Items, FNetDeltaSerializeInfo& Parms, SerializerType& ArraySerializer)
 {
+	/**
+	 * These methods are exposed on FastArraySerializerItems, but they aren't virtual.
+	 * Further, we may not know the exact type when we want to call them, and won't
+	 * safely be able to cast to the type in non-templated code.
+	 *
+	 * Maybe this defeats the purpose of having them not virtual in the first place.
+	 * However, for now ReceivedItem and PostReplicatedChanged are the only ones that will actually
+	 * be called in this way, whereas PostReplicatedAdd and PreReplicatedRemove will still be called
+	 * from templated code.
+	 */
+	struct FFastArrayItemCallbackHelper
+	{
+		static void PreReplicatedRemove(void* FastArrayItem, const struct FFastArrayDeltaSerializeParams& Params)
+		{
+			reinterpret_cast<Type*>(FastArrayItem)->PreReplicatedRemove(static_cast<SerializerType&>(Params.ArraySerializer));
+		}
+
+		static void PostReplicatedAdd(void* FastArrayItem, const struct FFastArrayDeltaSerializeParams& Params)
+		{
+			reinterpret_cast<Type*>(FastArrayItem)->PostReplicatedAdd(static_cast<SerializerType&>(Params.ArraySerializer));
+		}
+
+		static void PostReplicatedChange(void* FastArrayItem, const struct FFastArrayDeltaSerializeParams& Params)
+		{
+			reinterpret_cast<Type*>(FastArrayItem)->PostReplicatedChange(static_cast<SerializerType&>(Params.ArraySerializer));
+		}
+
+		static void ReceivedItem(void* FastArrayItem, const FFastArrayDeltaSerializeParams& Params, const uint32 ReplicationID)
+		{
+			Type* Item = reinterpret_cast<Type*>(FastArrayItem);
+			Item->ReplicationID = ReplicationID;
+			Item->MostRecentArrayReplicationKey = Params.ReadArrayReplicationKey;
+			Item->ReplicationKey++;
+		}
+	};
+
 	SCOPE_CYCLE_COUNTER(STAT_NetSerializeFastArray_DeltaStruct);
 	class UScriptStruct* InnerStruct = Type::StaticStruct();
 
@@ -1445,17 +1501,15 @@ bool FFastArraySerializer::FastArrayDeltaSerialize_DeltaSerializeStructs(TArray<
 		Parms
 	};
 
-#define CALC_PATCH_OFFSET(Base, Derived) ((PTRINT)((Base*)((Derived*)1)) - 1)
-
 	FFastArrayDeltaSerializeParams DeltaSerializeParams{
 		Parms,
 		ArraySerializer,
-		/*PatchedItemOffset=*/ CALC_PATCH_OFFSET(FFastArraySerializerItem, Type),
-		/*PatchedArrayOffset*/ CALC_PATCH_OFFSET(FFastArraySerializer, SerializerType)
+		FFastArrayItemCallbackHelper::PreReplicatedRemove,
+		FFastArrayItemCallbackHelper::PostReplicatedAdd,
+		FFastArrayItemCallbackHelper::PostReplicatedChange,
+		FFastArrayItemCallbackHelper::ReceivedItem,
 	};
 
-#undef CALC_PATCH_OFFSET
-
 	//---------------
 	// Build ItemMap if necessary. This maps ReplicationID to our local index into the Items array.
 	//---------------
diff --git a/Engine/Source/Runtime/Engine/Private/RepLayout.cpp b/Engine/Source/Runtime/Engine/Private/RepLayout.cpp
--- a/Engine/Source/Runtime/Engine/Private/RepLayout.cpp
+++ b/Engine/Source/Runtime/Engine/Private/RepLayout.cpp
@@ -150,20 +150,12 @@ namespace UE4_RepLayout_Private
 		// TODO: Conditionally compilable runtime type validation.
 		return reinterpret_cast<ConstOrNotOutputType *>((Buffer + Cmd).Data);
 	}
-
-	template<typename CommandType, typename BufferType>
-	FORCEINLINE static FFastArraySerializerItem* GetFastArrayItem(FFastArrayDeltaSerializeParams& Params, const BufferType& Buffer, const CommandType& Cmd)
-	{
-		return (FFastArraySerializerItem*)(Params.PatchedItemOffset + (PTRINT)(GetTypedProperty<FFastArraySerializerItem>(Buffer, Cmd)));
-	}
-
-	template<typename CommandType, typename BufferType>
-	FORCEINLINE static FFastArraySerializer* GetFastArray(FFastArrayDeltaSerializeParams& Params, const BufferType& Buffer, const CommandType& Cmd)
-	{
-		return (FFastArraySerializer*)(Params.PatchedArrayOffset + (PTRINT)(GetTypedProperty<FFastArraySerializer>(Buffer, Cmd)));
-	}
 }
 
+//~ TODO: Consider moving the FastArray members into their own sub-struct to save memory for non fast array
+//~			custom delta properties. Almost all Custom Delta properties now **are** Fast Arrays, so this
+//~			probably doesn't matter much at the moment.
+
 struct FLifetimeCustomDeltaProperty
 {
 	/** If this is a Fast Array Serializer property, this will be the command index for the Fast Array Item array. */
@@ -174,6 +166,59 @@ struct FLifetimeCustomDeltaProperty
 	 * This is used to lookup Changelists.
 	 */
 	int32 FastArrayNumber = INDEX_NONE;
+
+	/**
+	 * If this is a Fast Array Serializer property (and it is set up correctly for Delta Serialization), this will be an
+	 * offset from to the property.
+	 */
+	int32 FastArrayDeltaFlagsOffset = INDEX_NONE;
+
+	/**
+	 * If this is a Fast Array Serializer property (and it is set up correctly for Delta Serialization), this will be a pointer
+	 * to the FFastArraySerializer::ArrayReplicationKey property.
+	 */
+	int32 FastArrayArrayReplicationKeyOffset = INDEX_NONE;
+
+	/**
+	 * If this is a Fast Array Serializer property (and it is set up correctly for Delta Serialization), this will be a pointer
+	 * to the FFastArraySerializerItem::ReplicationID property.
+	 */
+	int32 FastArrayItemReplicationIdOffset = INDEX_NONE;
+
+	const EFastArraySerializerDeltaFlags GetFastArrayDeltaFlags(void const * const FastArray) const
+	{
+		return GetRefFromOffsetAndMemory<EFastArraySerializerDeltaFlags>(FastArray, FastArrayDeltaFlagsOffset);
+	}
+
+	const int32 GetFastArrayArrayReplicationKey(void const * const FastArray) const
+	{
+		return GetRefFromOffsetAndMemory<int32>(FastArray, FastArrayArrayReplicationKeyOffset);
+	}
+
+	const int32 GetFastArrayItemReplicationID(void const * const FastArrayItem) const
+	{
+		return GetRefFromOffsetAndMemory<int32>(FastArrayItem, FastArrayItemReplicationIdOffset);
+	}
+
+	int32& GetFastArrayItemReplicationIDMutable(void* const FastArrayItem) const
+	{
+		return GetRefFromOffsetAndMemory<int32>(FastArrayItem, FastArrayItemReplicationIdOffset);
+	}
+
+private:
+
+	template<typename Output>
+	static Output& GetRefFromOffsetAndMemory(void* Memory, const int32 Offset)
+	{
+		checkSlow(Offset != INDEX_NONE);
+		return *reinterpret_cast<Output*>(reinterpret_cast<uint8*>(Memory) + Offset);
+	}
+
+	template<typename Output>
+	static const Output& GetRefFromOffsetAndMemory(void const * const Memory, const int32 Offset)
+	{
+		return GetRefFromOffsetAndMemory<Output>(const_cast<void*>(Memory), Offset);
+	}
 };
 
 /**
@@ -4862,11 +4907,26 @@ void FRepLayout::InitFromClass(UClass* InObjectClass, const UNetConnection* Serv
 					const FRepLayoutCmd& Cmd = Cmds[CmdIndex];
 					if (ERepLayoutCmdType::DynamicArray == Cmd.Type)
 					{
-						if (UStructProperty* MaybeFastArray = Cast<UStructProperty>(static_cast<UArrayProperty*>(Cmd.Property)->Inner))
+						if (UStructProperty* MaybeFastArrayItemsArray = Cast<UStructProperty>(static_cast<UArrayProperty*>(Cmd.Property)->Inner))
 						{
-							if (MaybeFastArray->Struct->IsChildOf(FFastArraySerializerItem::StaticStruct()))
+							UScriptStruct* MaybeFastArrayItem = MaybeFastArrayItemsArray->Struct;
+							if (MaybeFastArrayItem->IsChildOf(FFastArraySerializerItem::StaticStruct()))
 							{
-								FastArrayItemArrayCmd = CmdIndex;
+								// Can't use GET_MEMBER_NAME_CHECKED because this is private.
+								static const FName FastArrayDeltaFlagsName(TEXT("DeltaFlags"));
+								static const FName FastArrayArrayReplicationKeyName(GET_MEMBER_NAME_CHECKED(FFastArraySerializer, ArrayReplicationKey));
+								static const FName FastArrayItemReplicationIDName(GET_MEMBER_NAME_CHECKED(FFastArraySerializerItem, ReplicationID));
+
+								// This better be a script struct, otherwise our flags aren't set up correctly!
+								UScriptStruct* FastArray = CastChecked<UScriptStruct>(MaybeFastArrayItemsArray->GetOwnerStruct());
+
+								CustomDeltaProperty.FastArrayItemsCommand = CmdIndex;
+								CustomDeltaProperty.FastArrayItemReplicationIdOffset = MaybeFastArrayItem->FindPropertyByName(FastArrayItemReplicationIDName)->GetOffset_ForGC();
+								CustomDeltaProperty.FastArrayArrayReplicationKeyOffset = FastArray->FindPropertyByName(FastArrayArrayReplicationKeyName)->GetOffset_ForGC();
+								CustomDeltaProperty.FastArrayDeltaFlagsOffset = FastArray->FindPropertyByName(FastArrayDeltaFlagsName)->GetOffset_ForGC();
+
+								CustomDeltaProperty.FastArrayNumber = LifetimeCustomPropertyState->NumFastArrayItems;
+								++LifetimeCustomPropertyState->NumFastArrayItems;
 								break;
 							}
 						}
@@ -4875,12 +4935,7 @@ void FRepLayout::InitFromClass(UClass* InObjectClass, const UNetConnection* Serv
 					}
 				}
 
-				if (FastArrayItemArrayCmd != INDEX_NONE)
-				{
-					++LifetimeCustomPropertyState->NumFastArrayItems;
-					CustomDeltaProperty.FastArrayItemsCommand = FastArrayItemArrayCmd;
-				}
-				else
+				if (CustomDeltaProperty.FastArrayItemsCommand == INDEX_NONE)
 				{
 					UE_LOG(LogRep, Warning, TEXT("FRepLayout::InitFromClass: Unable to find Fast Array Item array in Fast Array Serializer: %s"), *Parents[ParentIndex].CachedPropertyName.ToString());
 				}
@@ -5903,21 +5958,28 @@ void FRepLayout::PreSendCustomDeltaProperties(
 					{
 						const int32 FastArrayNumber = CustomDeltaProperty.FastArrayNumber;
 						const FRepParentCmd& FastArrayCmd = Parents[RepIndex];
-						const FFastArraySerializer& FastArray = *GetTypedProperty<FFastArraySerializer>(ObjectData, FastArrayCmd);
 
-						// Note, we can't rely on FFastArraySerializer::HasBeenSerialized here.
+						void const * const FastArraySerializer = ObjectData + FastArrayCmd;
+						const EFastArraySerializerDeltaFlags Flags = CustomDeltaProperty.GetFastArrayDeltaFlags(FastArraySerializer);
+
+						// Note, we can't rely on EFastArraySerializerDeltaFlags::HasBeenSerialized here.
 						// It's possible we're calling PreSendCustomDeltaProperties **before** the first time the struct
 						// was ever serialized, and in that case it would still be false, and prevent us from creating
 						// a history the first time.
 						//
-						// This does mean that Fast Arrays *not* using delta serialization will still have their history
-						// incremented the first time, but that can be fixed in other ways and the history memory
-						// will have already been allocated elsewhere.
-						if (FastArray.IsUsingStructDeltaSerialization())
+						// This does mean that Fast Arrays requesting delta serialization will still have their history
+						// incremented the first time, even if the feature is generally disabled.
+						//
+						// TODO: If any fast arrays failed this check, we could probably reset their state,
+						//			because we know we should never try sending them again
+						if (EnumHasAnyFlags(Flags, EFastArraySerializerDeltaFlags::IsUsingDeltaSerialization) ||
+							(!EnumHasAnyFlags(Flags, EFastArraySerializerDeltaFlags::HasBeenSerialized) &&
+								EnumHasAnyFlags(Flags, EFastArraySerializerDeltaFlags::HasDeltaBeenRequested)))
 						{
 							FDeltaArrayHistoryState& FastArrayHistoryState = CustomDeltaChangelistState.ArrayStates[FastArrayNumber];
 
-							if (FastArrayHistoryState.ArrayReplicationKey != FastArray.ArrayReplicationKey)
+							const int32 FastArrayReplicationKey = CustomDeltaProperty.GetFastArrayArrayReplicationKey(FastArraySerializer);
+							if (FastArrayHistoryState.ArrayReplicationKey != FastArrayReplicationKey)
 							{
 								const uint32 HistoryDelta = FastArrayHistoryState.HistoryEnd - FastArrayHistoryState.HistoryStart;
 								const uint32 CurrentHistoryIndex = FastArrayHistoryState.HistoryEnd % FDeltaArrayHistoryState::MAX_CHANGE_HISTORY;
@@ -5998,7 +6060,7 @@ bool FRepLayout::DeltaSerializeFastArrayProperty(FFastArrayDeltaSerializeParams&
 	FScriptArray* ObjectArray = GetTypedProperty<FScriptArray>(ObjectData, FastArrayItemCmd);
 	FRepObjectDataBuffer ObjectArrayData(ObjectArray->GetData());
 
-	FFastArraySerializer& ArraySerializer = *GetFastArray(Params, ObjectData, Parent);
+	FFastArraySerializer& ArraySerializer = Params.ArraySerializer;
 
 	check(&ArraySerializer == &Params.ArraySerializer);
 
@@ -6169,18 +6231,18 @@ bool FRepLayout::DeltaSerializeFastArrayProperty(FFastArrayDeltaSerializeParams&
 
 							for (int32 Index = 0; Index < EndIndex; ++Index)
 							{
-								FFastArraySerializerItem* ObjectItem = GetFastArrayItem(Params, ObjectArrayHelper.GetRawPtr(Index), FastArrayItemCmd);
-								FFastArraySerializerItem* ShadowItem = GetFastArrayItem(Params, ShadowArrayHelper.GetRawPtr(Index), FastArrayItemCmd);
+								const int32 ObjectReplicationID = CustomDeltaProperty.GetFastArrayItemReplicationID(ObjectArrayHelper.GetRawPtr(Index));
+								int32& ShadowReplicationID = CustomDeltaProperty.GetFastArrayItemReplicationIDMutable(ShadowArrayHelper.GetRawPtr(Index));
 
-								FastArrayState.IDToIndexMap.Emplace(ObjectItem->ReplicationID, Index);
+								FastArrayState.IDToIndexMap.Emplace(ObjectReplicationID, Index);
 
-								UE_LOG(LogRep, VeryVerbose, TEXT("DeltaSerializeFastArrayProperty: Handling Item. ID=%d, Index=%d, ShadowID=%d"), ObjectItem->ReplicationID, Index, ShadowItem->ReplicationID);
+								UE_LOG(LogRep, VeryVerbose, TEXT("DeltaSerializeFastArrayProperty: Handling Item. ID=%d, Index=%d, ShadowID=%d"), ObjectReplicationID, Index, ShadowReplicationID);
 
 								// If our IDs match, there's nothing to do.
-								if (ObjectItem->ReplicationID != ShadowItem->ReplicationID)
+								if (ObjectReplicationID != ShadowReplicationID)
 								{
 									// The IDs didn't match, so this is an insert, delete, or swap.
-									if (int32 const * const FoundShadowIndex = OldShadowIDToIndexMap.Find(ObjectItem->ReplicationID))
+									if (int32 const * const FoundShadowIndex = OldShadowIDToIndexMap.Find(ObjectReplicationID))
 									{
 										// We found the element in the shadow array, so there must have been a swap.
 										// Sanity check that the invalid element can only possibly later in our lines.
@@ -6200,8 +6262,8 @@ bool FRepLayout::DeltaSerializeFastArrayProperty(FFastArrayDeltaSerializeParams&
 										ShadowArrayItemIsNew[Index] = true;
 										ShadowArrayHelper.InsertValues(Index);
 
-										ShadowItem = GetFastArrayItem(Params, ShadowArrayHelper.GetRawPtr(Index), FastArrayItemCmd);
-										ShadowItem->ReplicationID = ObjectItem->ReplicationID;
+										int32& NewShadowReplicationID = CustomDeltaProperty.GetFastArrayItemReplicationIDMutable(ShadowArrayHelper.GetRawPtr(Index));
+										NewShadowReplicationID = ObjectReplicationID;
 
 										++AppendedShadowItems;
 										UE_LOG(LogRep, VeryVerbose, TEXT("DeltaSerializeFastArrayProperty: Added Shadow Item. AppendedShadowItems=%d"), AppendedShadowItems);
@@ -6225,14 +6287,14 @@ bool FRepLayout::DeltaSerializeFastArrayProperty(FFastArrayDeltaSerializeParams&
 								FScriptArrayHelper ObjectArrayHelper((UArrayProperty*)FastArrayItemCmd.Property, ObjectArray);
 								for (int32 Index = StartIndex; Index < ObjectArrayNum; ++Index)
 								{
-									FFastArraySerializerItem* ObjectItem = GetFastArrayItem(Params, ObjectArrayHelper.GetRawPtr(Index), FastArrayItemCmd);
-									FFastArraySerializerItem* ShadowItem = GetFastArrayItem(Params, ShadowArrayHelper.GetRawPtr(Index), FastArrayItemCmd);
+									const int32 ObjectReplicationID = CustomDeltaProperty.GetFastArrayItemReplicationID(ObjectArrayHelper.GetRawPtr(Index));
+									int32& ShadowReplicationID = CustomDeltaProperty.GetFastArrayItemReplicationIDMutable(ShadowArrayHelper.GetRawPtr(Index));
 
-									ShadowItem->ReplicationID = ObjectItem->ReplicationID;
+									ShadowReplicationID = ObjectReplicationID;
 									ShadowArrayItemIsNew[Index] = true;
 
-									UE_LOG(LogRep, VeryVerbose, TEXT("DeltaSerializeFastArrayProperty: Added Shadow Item. Index=%d, ID=%d"), Index, ShadowItem->ReplicationID);
-									FastArrayState.IDToIndexMap.Emplace(ObjectItem->ReplicationID, Index);
+									UE_LOG(LogRep, VeryVerbose, TEXT("DeltaSerializeFastArrayProperty: Added Shadow Item. Index=%d, ID=%d"), Index, ShadowReplicationID);
+									FastArrayState.IDToIndexMap.Emplace(ObjectReplicationID, Index);
 								}
 							}
 						}
@@ -6436,7 +6498,7 @@ bool FRepLayout::DeltaSerializeFastArrayProperty(FFastArrayDeltaSerializeParams&
 
 			int32* ElementIndexPtr = ArraySerializer.ItemMap.Find(ID);
 			int32 ElementIndex = 0;
-			FFastArraySerializerItem* ThisElement = nullptr;
+			void* ThisElement = nullptr;
 
 			if (!ElementIndexPtr)
 			{
@@ -6445,24 +6507,18 @@ bool FRepLayout::DeltaSerializeFastArrayProperty(FFastArrayDeltaSerializeParams&
 				ElementIndex = FastArrayHelper.AddValue();
 				ArraySerializer.ItemMap.Add(ID, ElementIndex);
 				AddedElements.Add(ElementIndex);
-
-				ThisElement = GetFastArrayItem(Params, FastArrayHelper.GetRawPtr(ElementIndex), FastArrayItemCmd);
-				ThisElement->ReplicationID = ID;
 			}
 			else
 			{
 				ElementIndex = *ElementIndexPtr;
-				ThisElement = GetFastArrayItem(Params, FastArrayHelper.GetRawPtr(ElementIndex), FastArrayItemCmd);
 				ChangedElements.Add(ElementIndex);
 
 				UE_LOG(LogNetFastTArray, Log, TEXT("   Changed. ID: %d -> Idx: %d"), ID, ElementIndex);
 			}
 
-			// Update this element's most recent array replication key
-			ThisElement->MostRecentArrayReplicationKey = Params.ReadArrayReplicationKey;
+			ThisElement = FastArrayHelper.GetRawPtr(ElementIndex);
 
-			// Update this element's replication key so that a client can re-serialize the array for client replay recording
-			ThisElement->ReplicationKey++;
+			Params.ReceivedItem(ThisElement, Params, ID);
 
 			FGuidReferencesMap& GuidReferences = ArraySerializer.GuidReferencesMap_StructDelta.FindOrAdd(ID);
 
@@ -6552,7 +6608,7 @@ void FRepLayout::GatherGuidReferencesForFastArray(FFastArrayDeltaSerializeParams
 	const FConstRepObjectDataBuffer ObjectData(Params.DeltaSerializeInfo.Object);
 	const FRepParentCmd& Parent = Parents[Params.DeltaSerializeInfo.PropertyRepIndex];
 
-	const FFastArraySerializer& ArraySerializer = *GetFastArray(Params, ObjectData, Parent);
+	const FFastArraySerializer& ArraySerializer = Params.ArraySerializer;
 	TSet<FNetworkGUID>& GatherGuids = *Params.DeltaSerializeInfo.GatherGuidReferences;
 	int32 TrackedGuidMemory = 0;
 
@@ -6574,7 +6630,7 @@ bool FRepLayout::MoveMappedObjectToUnmappedForFastArray(FFastArrayDeltaSerialize
 	const FRepObjectDataBuffer ObjectData(Params.DeltaSerializeInfo.Object);
 	const FRepParentCmd& Parent = Parents[Params.DeltaSerializeInfo.PropertyRepIndex];
 
-	FFastArraySerializer& ArraySerializer = *GetFastArray(Params, ObjectData, Parent);
+	FFastArraySerializer& ArraySerializer = Params.ArraySerializer;
 	const FNetworkGUID& MoveToUnmapped = *Params.DeltaSerializeInfo.MoveGuidToUnmapped;
 
 	bool bFound = false;
@@ -6615,7 +6671,7 @@ void FRepLayout::UpdateUnmappedGuidsForFastArray(FFastArrayDeltaSerializeParams&
 	FScriptArray* ScriptArray = GetTypedProperty<FScriptArray>(ObjectData, FastArrayItemCmd);
 	FRepObjectDataBuffer ArrayData(ScriptArray->GetData());
 
-	FFastArraySerializer& ArraySerializer = *GetFastArray(Params, ObjectData, Parent);
+	FFastArraySerializer& ArraySerializer = Params.ArraySerializer;
 
 	for (auto It = ArraySerializer.GuidReferencesMap_StructDelta.CreateIterator(); It; ++It)
 	{
@@ -6633,7 +6689,7 @@ void FRepLayout::UpdateUnmappedGuidsForFastArray(FFastArrayDeltaSerializeParams&
 
 			if (bOutSomeObjectsWereMapped)
 			{
-				GetFastArrayItem(Params, ElementData, FastArrayItemCmd)->PostReplicatedChange(ArraySerializer);
+				Params.PostReplicatedChange(ElementData, Params);
 			}
 
 			DeltaSerializeInfo.bOutHasMoreUnmapped |= bOutHasMoreUnmapped;
