#OKFORGITHUB public
More robust check for currently active dialogue instances.

diff --git a/Engine/Source/Runtime/Engine/Classes/Sound/DialogueSoundWaveProxy.h b/Engine/Source/Runtime/Engine/Classes/Sound/DialogueSoundWaveProxy.h
--- a/Engine/Source/Runtime/Engine/Classes/Sound/DialogueSoundWaveProxy.h
+++ b/Engine/Source/Runtime/Engine/Classes/Sound/DialogueSoundWaveProxy.h
@@ -20,7 +20,6 @@ class UDialogueSoundWaveProxy : public USoundBase
 
 	USoundWave* SoundWave;
 	TArray<struct FSubtitleCue> Subtitles;
-	FWaveInstance* CurrentWaveInstance;
 
 public:
 	/** Returns whether the sound base is set up in a playable manner */
diff --git a/Engine/Source/Runtime/Engine/Private/DialogueWave.cpp b/Engine/Source/Runtime/Engine/Private/DialogueWave.cpp
--- a/Engine/Source/Runtime/Engine/Private/DialogueWave.cpp
+++ b/Engine/Source/Runtime/Engine/Private/DialogueWave.cpp
@@ -453,6 +453,9 @@ void UDialogueSoundWaveProxy::Parse(class FAudioDevice* AudioDevice, const UPTRI
 	int OldWaveInstanceCount = WaveInstances.Num();
 	bool bHasSubtitles = (Subtitles.Num() > 0);
 
+	// Check if the wave instance exists before we try to add it...
+	const bool bWaveInstanceAlreadyExisted = (ActiveSound.WaveInstances.Contains(NodeWaveInstanceHash));
+
 	ActiveSound.bHasExternalSubtitles = bHasSubtitles; // Need to set this so the sound will virtualize when silent if necessary.
 	SoundWave->Parse(AudioDevice, NodeWaveInstanceHash, ActiveSound, ParseParams, WaveInstances);
 	int NewWaveInstanceCount = WaveInstances.Num();
@@ -464,9 +467,8 @@ void UDialogueSoundWaveProxy::Parse(class FAudioDevice* AudioDevice, const UPTRI
 	}
 
 	// Only Queue subtitles once per each playback of a wave instance
-	if (NewWaveInstance != nullptr && NewWaveInstance != CurrentWaveInstance)
+	if (NewWaveInstance != nullptr && !bWaveInstanceAlreadyExisted)
 	{
-		CurrentWaveInstance = NewWaveInstance;
 		// Add in the subtitle if they exist
 		if (ActiveSound.bHandleSubtitles && bHasSubtitles)
 		{
@@ -474,7 +476,7 @@ void UDialogueSoundWaveProxy::Parse(class FAudioDevice* AudioDevice, const UPTRI
 			{
 				QueueSubtitleParams.AudioComponentID = ActiveSound.GetAudioComponentID();
 				QueueSubtitleParams.WorldPtr = ActiveSound.GetWeakWorld();
-				QueueSubtitleParams.WaveInstance = (PTRINT)CurrentWaveInstance;
+				QueueSubtitleParams.WaveInstance = (PTRINT)NewWaveInstance;
 				QueueSubtitleParams.SubtitlePriority = ActiveSound.SubtitlePriority;
 				QueueSubtitleParams.Duration = GetDuration();
 				QueueSubtitleParams.bManualWordWrap = false;
