Fix missing return statement

diff --git a/Engine/Source/Runtime/Engine/Classes/Engine/NetSerialization.h b/Engine/Source/Runtime/Engine/Classes/Engine/NetSerialization.h
--- a/Engine/Source/Runtime/Engine/Classes/Engine/NetSerialization.h
+++ b/Engine/Source/Runtime/Engine/Classes/Engine/NetSerialization.h
@@ -539,7 +539,7 @@ struct FFastArraySerializer
 
 	const EFastArraySerializerDeltaFlags GetDeltaSerializationFlags() const
 	{
-		DeltaFlags;
+		return DeltaFlags;
 	}
 
 private:
