Make sure we compare original Shadow Array when using Fast Array Delta Serialization.
Previously, this could break during inserting new elements.

diff --git a/Engine/Source/Runtime/Engine/Private/RepLayout.cpp b/Engine/Source/Runtime/Engine/Private/RepLayout.cpp
--- a/Engine/Source/Runtime/Engine/Private/RepLayout.cpp
+++ b/Engine/Source/Runtime/Engine/Private/RepLayout.cpp
@@ -6212,25 +6212,30 @@ bool FRepLayout::DeltaSerializeFastArrayProperty(FFastArrayDeltaSerializeParams&
 						// The main issue is that we can't rely on that method, because it will be comparing the last state that was replicated to a given particular connection,
 						// and we want to compare the last state that was replicated to *any* connection.
 
-						if (ObjectArrayNum != 0 && ShadowArrayHelper.Num() != 0 && !bIsInitial)
 						{
 							FScriptArrayHelper ObjectArrayHelper((UArrayProperty*)FastArrayItemCmd.Property, ObjectArray);
 
-							const int32 EndIndex = FMath::Min(ShadowArrayHelper.Num(), ObjectArrayNum);
-							const TMap<int32, int32> OldShadowIDToIndexMap(MoveTemp(FastArrayState.IDToIndexMap));
-							FastArrayState.IDToIndexMap.Reserve(ObjectArrayNum);
+							// We track this as a non-const, because if we append any items into the middle of the
+							// array, they will be explicitly marked as new, and we still want to compare items
+							// that existed in the array originally.
+							int32 ShadowArrayNum = ShadowArrayHelper.Num();
 
-							// We track the Appended Shadow Items, because any index we try and use after such
-							// an append needs to be shifted appropriately.
-							// TODO: We may be able to iterate the list backwards instead, but that may break
-							//			some assumptions laid out in the algorithm above.
-							int32 AppendedShadowItems = 0;
+							if (ObjectArrayNum != 0 && ShadowArrayNum != 0 && !bIsInitial)
+							{
+								const TMap<int32, int32> OldShadowIDToIndexMap(MoveTemp(FastArrayState.IDToIndexMap));
+								FastArrayState.IDToIndexMap.Reserve(ObjectArrayNum);
 
-							UE_LOG(LogRep, VeryVerbose, TEXT("DeltaSerializeFastArrayProperty: Fixup Shadow State. Owner=%s, Property=%s, bInitial=%d, ObjecyArrayNum=%d, ShadowArrayNum=%d"),
-								*Owner->GetName(), *Parent.CachedPropertyName.ToString(), !!bIsInitial, ObjectArrayNum, ShadowArrayHelper.Num());
+								// We track the Appended Shadow Items, because any index we try and use after such
+								// an append needs to be shifted appropriately.
+								// TODO: We may be able to iterate the list backwards instead, but that may break
+								//			some assumptions laid out in the algorithm above.
+								int32 AppendedShadowItems = 0;
 
-							for (int32 Index = 0; Index < EndIndex; ++Index)
-							{
+								UE_LOG(LogRep, VeryVerbose, TEXT("DeltaSerializeFastArrayProperty: Fixup Shadow State. Owner=%s, Property=%s, bInitial=%d, ObjecyArrayNum=%d, ShadowArrayNum=%d"),
+									*Owner->GetName(), *Parent.CachedPropertyName.ToString(), !!bIsInitial, ObjectArrayNum, ShadowArrayHelper.Num());
+
+								for (int32 Index = 0; Index < ObjectArrayNum && Index < ShadowArrayNum; ++Index)
+								{
 								const int32 ObjectReplicationID = CustomDeltaProperty.GetFastArrayItemReplicationID(ObjectArrayHelper.GetRawPtr(Index));
 								int32& ShadowReplicationID = CustomDeltaProperty.GetFastArrayItemReplicationIDMutable(ShadowArrayHelper.GetRawPtr(Index));
 
@@ -6238,53 +6243,51 @@ bool FRepLayout::DeltaSerializeFastArrayProperty(FFastArrayDeltaSerializeParams&
 
 								UE_LOG(LogRep, VeryVerbose, TEXT("DeltaSerializeFastArrayProperty: Handling Item. ID=%d, Index=%d, ShadowID=%d"), ObjectReplicationID, Index, ShadowReplicationID);
 
-								// If our IDs match, there's nothing to do.
-								if (ObjectReplicationID != ShadowReplicationID)
-								{
-									// The IDs didn't match, so this is an insert, delete, or swap.
-									if (int32 const * const FoundShadowIndex = OldShadowIDToIndexMap.Find(ObjectReplicationID))
-									{
-										// We found the element in the shadow array, so there must have been a swap.
-										// Sanity check that the invalid element can only possibly later in our lines.
-										const int32 FixedShadowIndex = *FoundShadowIndex + AppendedShadowItems;
-
-										UE_LOG(LogRep, VeryVerbose, TEXT("DeltaSerializeFastArrayProperty: Swapped Shadow Item. OldIndex=%d, NewIndex=%d"), Index, FixedShadowIndex);
-
-										check(FixedShadowIndex > Index);
-
-										ShadowArrayHelper.SwapValues(Index, FixedShadowIndex);
-									}
-									else
+									// If our IDs match, there's nothing to do.
+									if (ObjectReplicationID != ShadowReplicationID)
 									{
-										// This item must have been inserted into the array (or appended and then shuffled in).
-										// So, insert it into our shadow array and update its ID.
-
-										ShadowArrayItemIsNew[Index] = true;
-										ShadowArrayHelper.InsertValues(Index);
-
-										int32& NewShadowReplicationID = CustomDeltaProperty.GetFastArrayItemReplicationIDMutable(ShadowArrayHelper.GetRawPtr(Index));
-										NewShadowReplicationID = ObjectReplicationID;
-
-										++AppendedShadowItems;
-										UE_LOG(LogRep, VeryVerbose, TEXT("DeltaSerializeFastArrayProperty: Added Shadow Item. AppendedShadowItems=%d"), AppendedShadowItems);
+										// The IDs didn't match, so this is an insert, delete, or swap.
+										if (int32 const * const FoundShadowIndex = OldShadowIDToIndexMap.Find(ObjectReplicationID))
+										{
+											// We found the element in the shadow array, so there must have been a swap.
+											// Sanity check that the invalid element can only possibly later in our lines.
+											const int32 FixedShadowIndex = *FoundShadowIndex + AppendedShadowItems;
+
+											UE_LOG(LogRepProperties, VeryVerbose, TEXT("DeltaSerializeFastArrayProperty: Swapped Shadow Item. OldIndex=%d, NewIndex=%d"), Index, FixedShadowIndex);
+
+											check(FixedShadowIndex > Index);
+
+											ShadowArrayHelper.SwapValues(Index, FixedShadowIndex);
+										}
+										else
+										{
+											// This item must have been inserted into the array (or appended and then shuffled in).
+											// So, insert it into our shadow array and update its ID.
+
+											ShadowArrayItemIsNew[Index] = true;
+											ShadowArrayHelper.InsertValues(Index);
+
+											int32& NewShadowReplicationID = CustomDeltaProperty.GetFastArrayItemReplicationIDMutable(ShadowArrayHelper.GetRawPtr(Index));
+											NewShadowReplicationID = ObjectReplicationID;
+
+											++AppendedShadowItems;
+											++ShadowArrayNum;
+											UE_LOG(LogRepProperties, VeryVerbose, TEXT("DeltaSerializeFastArrayProperty: Added Shadow Item. AppendedShadowItems=%d"), AppendedShadowItems);
+										}
 									}
 								}
 							}
-						}
 
-						// Now we can go ahead and resize the array, to make any other changes we need.
-						{
-							const int32 OldShadowArrayNum = ShadowArrayHelper.Num();
+							// Now we can go ahead and resize the array, to make any other changes we need.
 							ShadowArrayHelper.Resize(ObjectArrayNum);
 							ShadowArrayData = ShadowArray->GetData();
 
 							// Go ahead and fix up IDs for any elements that may have just been appended.
-							// Note, we need to this for all elements on the initial pass.
+							// Note, we need to do this for all elements on the initial pass.
 							// Deleted elements will have been chopped off by the resize.
-							if (bIsInitial || (OldShadowArrayNum < ObjectArrayNum))
+							if (bIsInitial || (ShadowArrayNum < ObjectArrayNum))
 							{
-								const int32 StartIndex = bIsInitial ? 0 : ShadowArrayHelper.Num();
-								FScriptArrayHelper ObjectArrayHelper((UArrayProperty*)FastArrayItemCmd.Property, ObjectArray);
+								const int32 StartIndex = bIsInitial ? 0 : ShadowArrayNum;
 								for (int32 Index = StartIndex; Index < ObjectArrayNum; ++Index)
 								{
 									const int32 ObjectReplicationID = CustomDeltaProperty.GetFastArrayItemReplicationID(ObjectArrayHelper.GetRawPtr(Index));
