Added fix for theoretical NetConnection::Cleanup, while processing ClientConnections packets.

diff --git a/Engine/Source/Runtime/Engine/Private/NetworkDriver.cpp b/Engine/Source/Runtime/Engine/Private/NetworkDriver.cpp
--- a/Engine/Source/Runtime/Engine/Private/NetworkDriver.cpp
+++ b/Engine/Source/Runtime/Engine/Private/NetworkDriver.cpp
@@ -1699,9 +1699,14 @@ void UNetDriver::PostTickDispatch()
 		ServerConnection->FlushPacketOrderCache(true);
 	}
 
-	for (UNetConnection* CurConn : ClientConnections)
+	TArray<UNetConnection*> ClientConnCopy = ClientConnections;
+
+	for (UNetConnection* CurConn : ClientConnCopy)
 	{
-		CurConn->FlushPacketOrderCache(true);
+		if (!CurConn->IsPendingKill())
+		{
+			CurConn->FlushPacketOrderCache(true);
+		}
 	}
 
 	if (ReplicationDriver)
