Always set RemoteRole in SetReplicates, to prevent issues caused by SetRemoteRoleForBackwardsCompat.

diff --git a/Engine/Source/Runtime/Engine/Private/Actor.cpp b/Engine/Source/Runtime/Engine/Private/Actor.cpp
--- a/Engine/Source/Runtime/Engine/Private/Actor.cpp
+++ b/Engine/Source/Runtime/Engine/Private/Actor.cpp
@@ -3226,24 +3226,25 @@ void AActor::SetReplicates(bool bInReplicates)
 { 
 	if (Role == ROLE_Authority)
 	{
-		// Only call into net driver if we actually changed
-		if (bReplicates != bInReplicates)
-		{
-			// Update our settings before calling into net driver
-			RemoteRole = bInReplicates ? ROLE_SimulatedProxy : ROLE_None;
-			bReplicates = bInReplicates;
-
-			// This actor should already be in the Network Actors List if it was already replicating.
-			if (bReplicates)
+		const bool bNewlyReplicates = (bReplicates == false && bInReplicates == true);
+	
+		// Due to SetRemoteRoleForBackwardsCompat, it's possible that bReplicates is false, but RemoteRole is something
+		// other than ROLE_None.
+		// So, we'll always set RemoteRole here regardless of whether or not bReplicates would change, to fix up that
+		// case.
+		RemoteRole = bInReplicates ? ROLE_SimulatedProxy : ROLE_None;
+		bReplicates = bInReplicates;
+
+		// Only call into net driver if we just started replicating changed
+		// This actor should already be in the Network Actors List if it was already replicating.
+		if (bNewlyReplicates)
+		{
+			// GetWorld will return nullptr on CDO, FYI
+			if (UWorld* MyWorld = GetWorld())		
 			{
-				// GetWorld will return nullptr on CDO, FYI
-				if (UWorld* MyWorld = GetWorld())		
-				{
-					MyWorld->AddNetworkActor(this);
-					ForcePropertyCompare();
-				}
+				MyWorld->AddNetworkActor(this);
+				ForcePropertyCompare();
 			}
-			
 		}
 	}
 	else
