Fix for an async load issue that occasionally led to an ensure() and falling back to use the slow path when optimized dynamic BP component instancing is enabled.

Change summary:
- Modified UBlueprintGeneratedClass::GetDefaultObjectPreloadDependencies() to include non-inherited SCS node objects as a preload dependency. This ensures that the reference they hold to the associated component template object will be serialized prior to the point at which we need to use it to generate the runtime data that's needed to support the fast path.
- Modified UBlueprintGeneratedClass::SerializeDefaultObject() to better handle the state in which we cannot generate runtime data to support the fast path when cooked data is present (i.e. emit a warning and fall back to the slow path). Previously we'd end up in a state where an ensure() was later triggered at instancing time when cooked data is present (which I've been using to catch unexpected cook results). Leaving the ensure() in place for now, but it should no longer be triggered.
- Renamed FBlueprintCookedComponentInstancingData::bIsValid for clarity (per code review).

diff --git a/Engine/Source/Editor/UnrealEd/Private/Kismet2/BlueprintEditorUtils.cpp b/Engine/Source/Editor/UnrealEd/Private/Kismet2/BlueprintEditorUtils.cpp
--- a/Engine/Source/Editor/UnrealEd/Private/Kismet2/BlueprintEditorUtils.cpp
+++ b/Engine/Source/Editor/UnrealEd/Private/Kismet2/BlueprintEditorUtils.cpp
@@ -9174,7 +9174,7 @@ void FBlueprintEditorUtils::BuildComponentInstancingData(UActorComponent* Compon
 		FComponentInstancingDataUtils::RecursivePropertyGather(ComponentTemplateClass, (uint8*)ComponentTemplate, (uint8*)ComponentDefaults, OutData);
 
 		// Flag that cooked data has been built and is now considered to be valid.
-		OutData.bIsValid = true;
+		OutData.bHasValidCookedData = true;
 	}
 }
 
diff --git a/Engine/Source/Runtime/Engine/Classes/Engine/BlueprintGeneratedClass.h b/Engine/Source/Runtime/Engine/Classes/Engine/BlueprintGeneratedClass.h
--- a/Engine/Source/Runtime/Engine/Classes/Engine/BlueprintGeneratedClass.h
+++ b/Engine/Source/Runtime/Engine/Classes/Engine/BlueprintGeneratedClass.h
@@ -541,7 +541,7 @@ struct ENGINE_API FBlueprintCookedComponentInstancingData
 
 	/** Flag indicating whether or not this contains valid cooked data. Note that an empty changed property list can also be a valid template data context. */
 	UPROPERTY()
-	bool bIsValid;
+	bool bHasValidCookedData;
 
 	/** List of property info records with values that differ between the template and the component class CDO. This list will be generated at cook time. */
 	UPROPERTY()
@@ -559,7 +559,7 @@ struct ENGINE_API FBlueprintCookedComponentInstancingData
 	/** Default constructor. */
 	FBlueprintCookedComponentInstancingData()
 	{
-		bIsValid = false;
+		bHasValidCookedData = false;
 		ComponentTemplateClass = nullptr;
 		ComponentTemplateFlags = RF_NoFlags;
 	}
diff --git a/Engine/Source/Runtime/Engine/Private/ActorConstruction.cpp b/Engine/Source/Runtime/Engine/Private/ActorConstruction.cpp
--- a/Engine/Source/Runtime/Engine/Private/ActorConstruction.cpp
+++ b/Engine/Source/Runtime/Engine/Private/ActorConstruction.cpp
@@ -1065,7 +1065,7 @@ UActorComponent* AActor::AddComponent(FName TemplateName, bool bManualAttachment
 				TemplateData = BPGC->CookedComponentInstancingData.Find(TemplateName);
 			}
 			
-			if (!TemplateData || !TemplateData->bIsValid
+			if (!TemplateData || !TemplateData->bHasValidCookedData
 				|| !ensureMsgf(TemplateData->ComponentTemplateClass != nullptr, TEXT("AddComponent fast path (%s.%s): Cooked data is valid, but runtime support data is not initialized. Using the slow path instead."), *BPGC->GetName(), *TemplateName.ToString()))
 			{
 				Template = BPGC->FindComponentTemplateByName(TemplateName);
diff --git a/Engine/Source/Runtime/Engine/Private/Blueprint.cpp b/Engine/Source/Runtime/Engine/Private/Blueprint.cpp
--- a/Engine/Source/Runtime/Engine/Private/Blueprint.cpp
+++ b/Engine/Source/Runtime/Engine/Private/Blueprint.cpp
@@ -1423,7 +1423,7 @@ void UBlueprint::BeginCacheForCookedPlatformData(const ITargetPlatform *TargetPl
 					for (auto RecordIt = TargetInheritableComponentHandler->CreateRecordIterator(); RecordIt; ++RecordIt)
 					{
 						// Only generate cooked data if the target platform supports the template class type. Cooked data may already have been generated if the component was inherited from a nativized parent class.
-						if (!RecordIt->CookedComponentInstancingData.bIsValid && ShouldCookBlueprintComponentTemplate(RecordIt->ComponentTemplate))
+						if (!RecordIt->CookedComponentInstancingData.bHasValidCookedData && ShouldCookBlueprintComponentTemplate(RecordIt->ComponentTemplate))
 						{
 							// Note: This will currently block until finished.
 							// @TODO - Make this an async task so we can potentially cook instancing data for multiple components in parallel.
diff --git a/Engine/Source/Runtime/Engine/Private/BlueprintGeneratedClass.cpp b/Engine/Source/Runtime/Engine/Private/BlueprintGeneratedClass.cpp
--- a/Engine/Source/Runtime/Engine/Private/BlueprintGeneratedClass.cpp
+++ b/Engine/Source/Runtime/Engine/Private/BlueprintGeneratedClass.cpp
@@ -316,16 +316,35 @@ void UBlueprintGeneratedClass::SerializeDefaultObject(UObject* Object, FArchive&
 		// @TODO - Potentially make this serializable (or cooked data) to eliminate the slight load time cost we'll incur below to generate this list in a cooked build. For now, it's not serialized since the raw UProperty references cannot be saved out.
 		UpdateCustomPropertyListForPostConstruction();
 
+		const FString BPGCName = GetName();
+		auto BuildCachedPropertyDataLambda = [BPGCName](FBlueprintCookedComponentInstancingData& CookedData, UActorComponent* SourceTemplate, FString CompVarName)
+		{
+			if (CookedData.bHasValidCookedData)
+			{
+				// This feature requires EDL at cook time, so ensure that the source template is also fully loaded at this point.
+				if (SourceTemplate != nullptr
+					&& ensure(!SourceTemplate->HasAnyFlags(RF_NeedLoad)))
+				{
+					CookedData.BuildCachedPropertyDataFromTemplate(SourceTemplate);
+				}
+				else
+				{
+					// This situation is unexpected; templates that are filtered out by context should not be generating fast path data at cook time. Emit a warning about this.
+					UE_LOG(LogBlueprint, Warning, TEXT("BPComp fast path (%s.%s) : Invalid source template. Will use slow path for dynamic instancing."), *BPGCName, *CompVarName);
+
+					// Invalidate the cooked data so that we fall back to using the slow path when dynamically instancing this node.
+					CookedData.bHasValidCookedData = false;
+				}
+			}
+		};
+
 		// Generate "fast path" instancing data for inherited SCS node templates. This data may also be used to support inherited SCS component default value overrides
 		// in a nativized, cooked build, in which this Blueprint class inherits from a nativized Blueprint parent. See CheckAndApplyComponentTemplateOverrides() below.
 		if (InheritableComponentHandler && (bHasCookedComponentInstancingData || bHasNativizedParent))
 		{
 			for (auto RecordIt = InheritableComponentHandler->CreateRecordIterator(); RecordIt; ++RecordIt)
 			{
-				if (RecordIt->ComponentTemplate && RecordIt->CookedComponentInstancingData.bIsValid)
-				{
-					RecordIt->CookedComponentInstancingData.BuildCachedPropertyDataFromTemplate(RecordIt->ComponentTemplate);
-				}
+				BuildCachedPropertyDataLambda(RecordIt->CookedComponentInstancingData, RecordIt->ComponentTemplate, RecordIt->ComponentKey.GetSCSVariableName().ToString());
 			}
 		}
 
@@ -337,10 +356,7 @@ void UBlueprintGeneratedClass::SerializeDefaultObject(UObject* Object, FArchive&
 				const TArray<USCS_Node*>& AllSCSNodes = SimpleConstructionScript->GetAllNodes();
 				for (USCS_Node* SCSNode : AllSCSNodes)
 				{
-					if (SCSNode->ComponentTemplate && SCSNode->CookedComponentInstancingData.bIsValid)
-					{
-						SCSNode->CookedComponentInstancingData.BuildCachedPropertyDataFromTemplate(SCSNode->ComponentTemplate);
-					}
+					BuildCachedPropertyDataLambda(SCSNode->CookedComponentInstancingData, SCSNode->ComponentTemplate, SCSNode->GetVariableName().ToString());
 				}
 			}
 
@@ -351,12 +367,10 @@ void UBlueprintGeneratedClass::SerializeDefaultObject(UObject* Object, FArchive&
 				{
 					if (ComponentTemplate)
 					{
-						if (FBlueprintCookedComponentInstancingData* ComponentInstancingData = CookedComponentInstancingData.Find(ComponentTemplate->GetFName()))
+						FBlueprintCookedComponentInstancingData* ComponentInstancingData = CookedComponentInstancingData.Find(ComponentTemplate->GetFName());
+						if (ComponentInstancingData != nullptr)
 						{
-							if (ComponentInstancingData->bIsValid)
-							{
-								ComponentInstancingData->BuildCachedPropertyDataFromTemplate(ComponentTemplate);
-							}
+							BuildCachedPropertyDataLambda(*ComponentInstancingData, ComponentTemplate, ComponentTemplate->GetName());
 						}
 					}
 				}
@@ -1186,7 +1200,7 @@ void UBlueprintGeneratedClass::CheckAndApplyComponentTemplateOverrides(UObject*
 						if (ComponentKey.IsValid() && ComponentKey.IsSCSKey())
 						{
 							const FBlueprintCookedComponentInstancingData* OverrideData = ICH->GetOverridenComponentTemplateData(ComponentKey);
-							if (OverrideData != nullptr && OverrideData->bIsValid)
+							if (OverrideData != nullptr && OverrideData->bHasValidCookedData)
 							{
 								// This is the instance of the inherited component subobject that's owned by the given class default object
 								if (UObject* NativizedComponentSubobjectInstance = InClassDefaultObject->GetDefaultSubobjectByName(NativizedComponentSubobjectName))
@@ -1397,6 +1411,13 @@ void UBlueprintGeneratedClass::GetDefaultObjectPreloadDependencies(TArray<UObjec
 			const TArray<USCS_Node*>& AllSCSNodes = CurrentBPClass->SimpleConstructionScript->GetAllNodes();
 			for (USCS_Node* SCSNode : AllSCSNodes)
 			{
+				// An SCS node that's owned by this class must also be considered a preload dependency since we will access its serialized template reference property. Any SCS
+				// nodes that are inherited from a parent class will reference templates through the ICH instead, and that's already a preload dependency on the BP class itself.
+				if (CurrentBPClass == this)
+				{
+					OutDeps.Add(SCSNode);
+				}
+
 				OutDeps.Add(SCSNode->GetActualComponentTemplate(this));
 			}
 		}
@@ -1755,46 +1776,31 @@ void FBlueprintCookedComponentInstancingData::BuildCachedPropertyDataFromTemplat
 		}
 	};
 
-	if (bIsValid)
-	{
-		if (SourceTemplate)
-		{
-			// Make sure the source template has been loaded.
-			if (SourceTemplate->HasAnyFlags(RF_NeedLoad))
-			{
-				if (FLinkerLoad* Linker = SourceTemplate->GetLinker())
-				{
-					Linker->Preload(SourceTemplate);
-				}
-			}
+	checkSlow(bHasValidCookedData);
+	checkSlow(SourceTemplate != nullptr);
+	checkSlow(!SourceTemplate->HasAnyFlags(RF_NeedLoad));
 
-			// Cache source template attributes needed for instancing.
-			ComponentTemplateName = SourceTemplate->GetFName();
-			ComponentTemplateClass = SourceTemplate->GetClass();
-			ComponentTemplateFlags = SourceTemplate->GetFlags();
+	// Cache source template attributes needed for instancing.
+	ComponentTemplateName = SourceTemplate->GetFName();
+	ComponentTemplateClass = SourceTemplate->GetClass();
+	ComponentTemplateFlags = SourceTemplate->GetFlags();
 
-			// This will also load the cached property list, if necessary.
-			const FCustomPropertyListNode* PropertyList = GetCachedPropertyList();
+	// This will also load the cached property list, if necessary.
+	const FCustomPropertyListNode* PropertyList = GetCachedPropertyList();
 
-			// Make sure we don't have any previously-built data.
-			if (!ensure(CachedPropertyData.Num() == 0))
-			{
-				DEC_MEMORY_STAT_BY(STAT_BPCompInstancingFastPathMemory, CachedPropertyData.GetAllocatedSize());
+	// Make sure we don't have any previously-built data.
+	if (!ensure(CachedPropertyData.Num() == 0))
+	{
+		DEC_MEMORY_STAT_BY(STAT_BPCompInstancingFastPathMemory, CachedPropertyData.GetAllocatedSize());
 
-				CachedPropertyData.Empty();
-			}
+		CachedPropertyData.Empty();
+	}
 
-			// Write template data out to the "fast path" buffer. All dependencies will be loaded at this point.
-			FBlueprintComponentInstanceDataWriter InstanceDataWriter(CachedPropertyData, PropertyList);
-			SourceTemplate->Serialize(InstanceDataWriter);
+	// Write template data out to the "fast path" buffer. All dependencies will be loaded at this point.
+	FBlueprintComponentInstanceDataWriter InstanceDataWriter(CachedPropertyData, PropertyList);
+	SourceTemplate->Serialize(InstanceDataWriter);
 
-			INC_MEMORY_STAT_BY(STAT_BPCompInstancingFastPathMemory, CachedPropertyData.GetAllocatedSize());
-		}
-		else
-		{
-			bIsValid = false;
-		}
-	}
+	INC_MEMORY_STAT_BY(STAT_BPCompInstancingFastPathMemory, CachedPropertyData.GetAllocatedSize());
 }
 
 bool UBlueprintGeneratedClass::ArePropertyGuidsAvailable() const
diff --git a/Engine/Source/Runtime/Engine/Private/SCS_Node.cpp b/Engine/Source/Runtime/Engine/Private/SCS_Node.cpp
--- a/Engine/Source/Runtime/Engine/Private/SCS_Node.cpp
+++ b/Engine/Source/Runtime/Engine/Private/SCS_Node.cpp
@@ -88,7 +88,7 @@ UActorComponent* USCS_Node::ExecuteNodeOnActor(AActor* Actor, USceneComponent* P
 	UActorComponent* NewActorComp = nullptr;
 	UBlueprintGeneratedClass* ActualBPGC = CastChecked<UBlueprintGeneratedClass>(Actor->GetClass());
 	const FBlueprintCookedComponentInstancingData* ActualComponentTemplateData = ActualBPGC->UseFastPathComponentInstancing() ? GetActualComponentTemplateData(ActualBPGC) : nullptr;
-	if (ActualComponentTemplateData && ActualComponentTemplateData->bIsValid
+	if (ActualComponentTemplateData && ActualComponentTemplateData->bHasValidCookedData
 		&& ensureMsgf(ActualComponentTemplateData->ComponentTemplateClass != nullptr, TEXT("SCS fast path (%s.%s): Cooked data is valid, but runtime support data is not initialized. Using the slow path instead."), *ActualBPGC->GetName(), *InternalVariableName.ToString()))
 	{
 		// Use cooked instancing data if valid (fast path).
