Fixed functions that were using UPARM(ref) with BlueprintPure rather than BlueprintCallable

diff --git a/Engine/Source/Runtime/Engine/Classes/Kismet/KismetMathLibrary.h b/Engine/Source/Runtime/Engine/Classes/Kismet/KismetMathLibrary.h
--- a/Engine/Source/Runtime/Engine/Classes/Kismet/KismetMathLibrary.h
+++ b/Engine/Source/Runtime/Engine/Classes/Kismet/KismetMathLibrary.h
@@ -901,8 +901,8 @@ class ENGINE_API UKismetMathLibrary : public UBlueprintFunctionLibrary
 	 * @param InX New X coordinate.
 	 * @param InY New Y coordinate.
 	 */
-	UFUNCTION(BlueprintPure, meta = (ScriptMethod = "Set"), Category = "Math|Vector2D")
-	static void Set2D(UPARAM(ref) FVector2D& A, float InX, float InY);
+	UFUNCTION(BlueprintCallable, meta = (ScriptMethod = "Set"), Category = "Math|Vector2D")
+	static void Set2D(UPARAM(ref) FVector2D& A, float X, float Y);
 
 	/**
 	 * Creates a copy of this vector with both axes clamped to the given range.
@@ -1041,7 +1041,7 @@ class ENGINE_API UKismetMathLibrary : public UBlueprintFunctionLibrary
 	 * @param Tolerance Minimum squared length of vector for normalization.
 	 * @see NormalSafe2D()
 	 */
-	UFUNCTION(BlueprintPure, meta = (DisplayName = "Normalize In Place (Vector2D)", Keywords = "Unit Vector", ScriptMethod = "Normalize"), Category = "Math|Vector2D")
+	UFUNCTION(BlueprintCallable, meta = (DisplayName = "Normalize In Place (Vector2D)", Keywords = "Unit Vector", ScriptMethod = "Normalize"), Category = "Math|Vector2D")
 	static void Normalize2D(UPARAM(ref) FVector2D& A, float Tolerance = 1.e-8);
 
 	/** Converts spherical coordinates on the unit sphere into a Cartesian unit length vector. */
@@ -1136,7 +1136,7 @@ class ENGINE_API UKismetMathLibrary : public UBlueprintFunctionLibrary
 	 *
 	 * @param InVector Vector to copy values from.
 	 */
-	UFUNCTION(BlueprintPure, meta = (ScriptMethod = "Assign"), Category = "Math|Vector")
+	UFUNCTION(BlueprintCallable, meta = (ScriptMethod = "Assign"), Category = "Math|Vector")
 	static void Vector_Assign(UPARAM(ref) FVector& A, const FVector& InVector);
 
 	/**
@@ -1146,8 +1146,8 @@ class ENGINE_API UKismetMathLibrary : public UBlueprintFunctionLibrary
 	 * @param InY New Y coordinate.
 	 * @param InZ New Z coordinate.
 	 */
-	UFUNCTION(BlueprintPure, meta = (ScriptMethod = "Set"), Category = "Math|Vector")
-	static void Vector_Set(UPARAM(ref) FVector& A, float InX, float InY, float InZ);
+	UFUNCTION(BlueprintCallable, meta = (ScriptMethod = "Set"), Category = "Math|Vector")
+	static void Vector_Set(UPARAM(ref) FVector& A, float X, float Y, float Z);
 
 	/** Breaks a vector apart into X, Y, Z */
 	UFUNCTION(BlueprintPure, Category="Math|Vector", meta=(NativeBreakFunc))
@@ -1279,7 +1279,7 @@ class ENGINE_API UKismetMathLibrary : public UBlueprintFunctionLibrary
 	static FVector LessLess_VectorRotator(FVector A, FRotator B);
 
 	/** When this vector contains Euler angles (degrees), ensure that angles are between +/-180 */
-	UFUNCTION(BlueprintPure, meta = (ScriptMethod = "UnwindEuler"), Category = "Math|Vector")
+	UFUNCTION(BlueprintCallable, meta = (ScriptMethod = "UnwindEuler"), Category = "Math|Vector")
 	static void Vector_UnwindEuler(UPARAM(ref) FVector& A);
 
 	/** Create a copy of this vector, with its magnitude/size/length clamped between Min and Max. */
@@ -1547,7 +1547,7 @@ class ENGINE_API UKismetMathLibrary : public UBlueprintFunctionLibrary
 	 *
 	 * @param Tolerance Minimum squared length of vector for normalization.
 	 */
-	UFUNCTION(BlueprintPure, meta = (DisplayName = "Normalize In Place (Vector)", ScriptMethod = "Normalize", Keywords = "Unit Vector"), Category = "Math|Vector")
+	UFUNCTION(BlueprintCallable, meta = (DisplayName = "Normalize In Place (Vector)", ScriptMethod = "Normalize", Keywords = "Unit Vector"), Category = "Math|Vector")
 	static void Vector_Normalize(UPARAM(ref) FVector& A, float Tolerance = 1.e-8);
 
 	/** Linearly interpolates between A and B based on Alpha (100% of A when Alpha=0 and 100% of B when Alpha=1) */
@@ -1661,7 +1661,7 @@ class ENGINE_API UKismetMathLibrary : public UBlueprintFunctionLibrary
 	 * @param InAddVect Vector to add.
 	 * @param InRadius Half size of the cube.
 	 */
-	UFUNCTION(BlueprintPure, meta=(ScriptMethod = "AddBounded", Keywords = "Bounding"), Category="Math|Vector")
+	UFUNCTION(BlueprintCallable, meta=(ScriptMethod = "AddBounded", Keywords = "Bounding"), Category="Math|Vector")
 	static void Vector_AddBounded(UPARAM(ref) FVector& A, FVector InAddVect, float InRadius);
 
 	/** Get a copy of this vector, clamped inside of the specified axis aligned cube. */
@@ -1903,7 +1903,7 @@ class ENGINE_API UKismetMathLibrary : public UBlueprintFunctionLibrary
 	 *
 	 * @param InVector Vector to copy values from.
 	 */
-	UFUNCTION(BlueprintPure, meta = (ScriptMethod = "Assign"), Category = "Math|Vector4")
+	UFUNCTION(BlueprintCallable, meta = (ScriptMethod = "Assign"), Category = "Math|Vector4")
 	static void Vector4_Assign(UPARAM(ref) FVector4& A, const FVector4& InVector);
 
 	/**
@@ -1914,8 +1914,8 @@ class ENGINE_API UKismetMathLibrary : public UBlueprintFunctionLibrary
 	 * @param InZ New Z coordinate.
 	 * @param InW New W coordinate.
 	 */
-	UFUNCTION(BlueprintPure, meta = (ScriptMethod = "Set"), Category = "Math|Vector4")
-	static void Vector4_Set(UPARAM(ref) FVector4& A, float InX, float InY, float InZ, float InW);
+	UFUNCTION(BlueprintCallable, meta = (ScriptMethod = "Set"), Category = "Math|Vector4")
+	static void Vector4_Set(UPARAM(ref) FVector4& A, float X, float Y, float Z, float W);
 
 	/** Returns the cross product of two vectors - see  http://mathworld.wolfram.com/CrossProduct.html */
 	UFUNCTION(BlueprintPure, meta = (DisplayName = "Cross Product XYZ (Vector4)", CompactNodeTitle = "cross3", ScriptMethod = "Cross3"), Category = "Math|Vector4")
@@ -2009,7 +2009,7 @@ class ENGINE_API UKismetMathLibrary : public UBlueprintFunctionLibrary
 	 *
 	 * @param Tolerance Minimum squared length of vector for normalization.
 	 */
-	UFUNCTION(BlueprintPure, meta = (DisplayName = "Normalize In Place XYZ (Vector4)", ScriptMethod = "Normalize3", Keywords = "Unit Vector"), Category = "Math|Vector4")
+	UFUNCTION(BlueprintCallable, meta = (DisplayName = "Normalize In Place XYZ (Vector4)", ScriptMethod = "Normalize3", Keywords = "Unit Vector"), Category = "Math|Vector4")
 	static void Vector4_Normalize3(UPARAM(ref) FVector4& A, float Tolerance = 1.e-8);
 
 	/** 
@@ -2269,7 +2269,7 @@ class ENGINE_API UKismetMathLibrary : public UBlueprintFunctionLibrary
 	static float Quat_AngularDistance(const FQuat& A, const FQuat& B);
 
 	/** Modify the quaternion to ensure that the delta between it and B represents the shortest possible rotation angle. */
-	UFUNCTION(BlueprintPure, meta = (DisplayName = "Ensure shortest arc to (Quat)", ScriptMethod = "EnsureShortestArcTo"), Category = "Math|Quat")
+	UFUNCTION(BlueprintCallable, meta = (DisplayName = "Ensure shortest arc to (Quat)", ScriptMethod = "EnsureShortestArcTo"), Category = "Math|Quat")
 	static void Quat_EnforceShortestArcWith(UPARAM(ref) FQuat& A, const FQuat& B);
 
 	/**	Convert a Quaternion into floating-point Euler angles (in degrees). */
@@ -2318,7 +2318,7 @@ class ENGINE_API UKismetMathLibrary : public UBlueprintFunctionLibrary
 	 *
 	 * @param Tolerance Minimum squared length of quaternion for normalization.
 	 */
-	UFUNCTION(BlueprintPure, meta = (DisplayName = "Normalize (Quat)", ScriptMethod = "Normalize"), Category = "Math|Quat")
+	UFUNCTION(BlueprintCallable, meta = (DisplayName = "Normalize (Quat)", ScriptMethod = "Normalize"), Category = "Math|Quat")
 	static void Quat_Normalize(UPARAM(ref) FQuat& Q, float Tolerance = 1.e-4f);
 
 	/**
@@ -2347,18 +2347,27 @@ class ENGINE_API UKismetMathLibrary : public UBlueprintFunctionLibrary
 	static FQuat Quat_Log(const FQuat& Q);
 
 	/** Set X, Y, Z, W components of Quaternion. */
-	UFUNCTION(BlueprintPure, meta = (DisplayName = "Set Components (Quat)", ScriptMethod = "SetComponents"), Category = "Math|Quat")
-	static void Quat_SetComponents(UPARAM(ref) FQuat& Q, float InX, float InY, float InZ, float InW);
+	UFUNCTION(BlueprintCallable, meta = (DisplayName = "Set Components (Quat)", ScriptMethod = "SetComponents"), Category = "Math|Quat")
+	static void Quat_SetComponents(UPARAM(ref) FQuat& Q, float X, float Y, float Z, float W);
 
 	/**
 	 * Convert a vector of floating-point Euler angles (in degrees) into a Quaternion.
 	 * 
+	 * @param Q Quaternion to update
 	 * @param Euler the Euler angles
-	 * @return constructed Quat
 	 */
-	UFUNCTION(BlueprintPure, meta = (DisplayName = "Set from Euler (Quat)", ScriptMethod = "SetFromEuler"), Category = "Math|Quat")
+	UFUNCTION(BlueprintCallable, meta = (DisplayName = "Set from Euler (Quat)", ScriptMethod = "SetFromEuler"), Category = "Math|Quat")
 	static void Quat_SetFromEuler(UPARAM(ref) FQuat& Q, const FVector& Euler);
 
+	/**
+	 * Convert a vector of floating-point Euler angles (in degrees) into a Quaternion.
+	 * 
+	 * @param Euler the Euler angles
+	 * @return constructed Quat
+	 */
+	UFUNCTION(BlueprintPure, meta = (DisplayName = "Make from Euler (Quat)"), Category = "Math|Quat")
+	static FQuat Quat_MakeFromEuler(const FVector& Euler);
+
 	/** Convert to Rotator representation of this Quaternion. */
 	UFUNCTION(BlueprintPure, meta = (DisplayName = "ToRotator (Quat)", CompactNodeTitle = "->", ScriptMethod = "Rotator", Keywords = "cast convert", BlueprintAutocast), Category = "Math|Conversions")
 	static FRotator Quat_Rotator(const FQuat& Q);
@@ -2448,37 +2457,37 @@ class ENGINE_API UKismetMathLibrary : public UBlueprintFunctionLibrary
 	static void BreakColor(FLinearColor InColor, float& R, float& G, float& B, float& A);
 
 	/** Assign contents of InColor */
-	UFUNCTION(BlueprintPure, meta = (ScriptMethod = "Set"), Category = "Math|Color")
+	UFUNCTION(BlueprintCallable, meta = (ScriptMethod = "Set"), Category = "Math|Color")
 	static void LinearColor_Set(UPARAM(ref) FLinearColor& InOutColor, FLinearColor InColor);
 
 	/** Assign individual linear RGBA components. */
-	UFUNCTION(BlueprintPure, meta = (ScriptMethod = "SetRGBA"), Category = "Math|Color")
+	UFUNCTION(BlueprintCallable, meta = (ScriptMethod = "SetRGBA"), Category = "Math|Color")
 	static void LinearColor_SetRGBA(UPARAM(ref) FLinearColor& InOutColor, float R, float G, float B, float A = 1.0f);
 
 	/** Assigns an HSV color to a linear space RGB color */
-	UFUNCTION(BlueprintPure, meta = (ScriptMethod = "SetFromHSV"), Category = "Math|Color")
+	UFUNCTION(BlueprintCallable, meta = (ScriptMethod = "SetFromHSV"), Category = "Math|Color")
 	static void LinearColor_SetFromHSV(UPARAM(ref) FLinearColor& InOutColor, float H, float S, float V, float A = 1.0f);
 
 	/**
 	 * Assigns an FColor coming from an observed sRGB output, into a linear color.
 	 * @param InSRGB The sRGB color that needs to be converted into linear space.
 	 */
-	UFUNCTION(BlueprintPure, meta = (ScriptMethod = "SetFromSRGB"), Category = "Math|Color")
+	UFUNCTION(BlueprintCallable, meta = (ScriptMethod = "SetFromSRGB"), Category = "Math|Color")
 	static void LinearColor_SetFromSRGB(UPARAM(ref) FLinearColor& InOutColor, const FColor& InSRGB);
 
 	/**
 	 * Assigns an FColor coming from an observed Pow(1/2.2) output, into a linear color.
 	 * @param InColor The Pow(1/2.2) color that needs to be converted into linear space.
 	 */
-	UFUNCTION(BlueprintPure, meta = (ScriptMethod = "SetFromPow22"), Category = "Math|Color")
+	UFUNCTION(BlueprintCallable, meta = (ScriptMethod = "SetFromPow22"), Category = "Math|Color")
 	static void LinearColor_SetFromPow22(UPARAM(ref) FLinearColor& InOutColor, const FColor& InColor);
 
 	/** Converts temperature in Kelvins of a black body radiator to RGB chromaticity. */
-	UFUNCTION(BlueprintPure, meta = (ScriptMethod = "SetTemperature"), Category = "Math|Color")
+	UFUNCTION(BlueprintCallable, meta = (ScriptMethod = "SetTemperature"), Category = "Math|Color")
 	static void LinearColor_SetTemperature(UPARAM(ref) FLinearColor& InOutColor, float InTemperature);
 
 	/** Sets to a random color. Choses a quite nice color based on a random hue. */
-	UFUNCTION(BlueprintPure, meta = (ScriptMethod = "SetRandomHue"), Category = "Math|Color")
+	UFUNCTION(BlueprintCallable, meta = (ScriptMethod = "SetRandomHue"), Category = "Math|Color")
 	static void LinearColor_SetRandomHue(UPARAM(ref) FLinearColor& InOutColor);
 
 	/** Convert a float into a LinearColor, where each element is that float */
diff --git a/Engine/Source/Runtime/Engine/Classes/Kismet/KismetMathLibrary.inl b/Engine/Source/Runtime/Engine/Classes/Kismet/KismetMathLibrary.inl
--- a/Engine/Source/Runtime/Engine/Classes/Kismet/KismetMathLibrary.inl
+++ b/Engine/Source/Runtime/Engine/Classes/Kismet/KismetMathLibrary.inl
@@ -1046,9 +1046,9 @@ FVector2D UKismetMathLibrary::Negated2D(const FVector2D& A)
 }
 
 KISMET_MATH_FORCEINLINE
-void UKismetMathLibrary::Set2D(FVector2D& A, float InX, float InY)
+void UKismetMathLibrary::Set2D(FVector2D& A, float X, float Y)
 {
-	A.Set(InX, InY);
+	A.Set(X, Y);
 }
 
 KISMET_MATH_FORCEINLINE
@@ -1254,9 +1254,9 @@ void UKismetMathLibrary::Vector_Assign(FVector& A, const FVector& InVector)
 }
 
 KISMET_MATH_FORCEINLINE
-void UKismetMathLibrary::Vector_Set(FVector& A, float InX, float InY, float InZ)
+void UKismetMathLibrary::Vector_Set(FVector& A, float X, float Y, float Z)
 {
-	A.Set(InX, InY, InZ);
+	A.Set(X, Y, Z);
 }
 
 KISMET_MATH_FORCEINLINE
@@ -1920,9 +1920,9 @@ void UKismetMathLibrary::Vector4_Assign(FVector4& A, const FVector4& InVector)
 }
 
 KISMET_MATH_FORCEINLINE
-void UKismetMathLibrary::Vector4_Set(FVector4& A, float InX, float InY, float InZ, float InW)
+void UKismetMathLibrary::Vector4_Set(FVector4& A, float X, float Y, float Z, float W)
 {
-	A.Set(InX, InY, InZ, InW);
+	A.Set(X, Y, Z, W);
 }
 
 KISMET_MATH_FORCEINLINE
@@ -2791,12 +2791,12 @@ FQuat UKismetMathLibrary::Quat_Log(const FQuat& Q)
 }
 
 KISMET_MATH_FORCEINLINE
-void UKismetMathLibrary::Quat_SetComponents(UPARAM(ref) FQuat& Q, float InX, float InY, float InZ, float InW)
+void UKismetMathLibrary::Quat_SetComponents(UPARAM(ref) FQuat& Q, float X, float Y, float Z, float W)
 {
-	Q.X = InX;
-	Q.Y = InY;
-	Q.Z = InZ;
-	Q.W = InW;
+	Q.X = X;
+	Q.Y = Y;
+	Q.Z = Z;
+	Q.W = W;
 }
 
 KISMET_MATH_FORCEINLINE
@@ -2805,6 +2805,12 @@ void UKismetMathLibrary::Quat_SetFromEuler(UPARAM(ref) FQuat& Q, const FVector&
 	Q = FQuat::MakeFromEuler(Euler);
 }
 
+KISMET_MATH_FORCEINLINE
+FQuat UKismetMathLibrary::Quat_MakeFromEuler(const FVector& Euler)
+{
+	return FQuat::MakeFromEuler(Euler);
+}
+
 KISMET_MATH_FORCEINLINE
 FRotator UKismetMathLibrary::Quat_Rotator(const FQuat& Q)
 {
