Eliminate UWorld::PawnList and reworked FConstPawnIterator to be a wrapper around TActorIterator for backwards compatibility purposes, however, TActorIterator or TActorRange should be preferred for iterating Pawns.
UWorld::AddPawn, RemovePawn, and GetNumPawns have been deprecated

diff --git a/Engine/Source/Runtime/Engine/Classes/Engine/World.h b/Engine/Source/Runtime/Engine/Classes/Engine/World.h
--- a/Engine/Source/Runtime/Engine/Classes/Engine/World.h
+++ b/Engine/Source/Runtime/Engine/Classes/Engine/World.h
@@ -13,7 +13,6 @@
 #include "Engine/EngineBaseTypes.h"
 #include "CollisionQueryParams.h"
 #include "WorldCollision.h"
-#include "GameFramework/Pawn.h"
 #include "EngineDefines.h"
 #include "Engine/Blueprint.h"
 #include "Engine/PendingNetGame.h"
@@ -31,6 +30,7 @@ class AController;
 class AGameModeBase;
 class AGameStateBase;
 class AMatineeActor;
+class APawn;
 class APhysicsVolume;
 class APlayerController;
 class AWorldSettings;
@@ -43,6 +43,7 @@ class UAISystemBase;
 class UCanvas;
 class UDemoNetDriver;
 class UGameViewportClient;
+class ULevel;
 class ULevelStreaming;
 class ULocalPlayer;
 class UMaterialParameterCollection;
@@ -62,13 +63,69 @@ template<typename,typename> class TOctree;
  * Misc. Iterator types
  *
  */
+template <typename ActorType> class TActorIterator;
 typedef TArray<TWeakObjectPtr<AController> >::TConstIterator FConstControllerIterator;
 typedef TArray<TWeakObjectPtr<APlayerController> >::TConstIterator FConstPlayerControllerIterator;
-typedef TArray<TWeakObjectPtr<APawn> >::TConstIterator FConstPawnIterator;	
 typedef TArray<TWeakObjectPtr<ACameraActor> >::TConstIterator FConstCameraActorIterator;
-typedef TArray<class ULevel*>::TConstIterator FConstLevelIterator;
+typedef TArray<ULevel*>::TConstIterator FConstLevelIterator;
 typedef TArray<TWeakObjectPtr<APhysicsVolume> >::TConstIterator FConstPhysicsVolumeIterator;
 
+/** Wrapper object that tries to imitate the TWeakObjectPtr interface for the objects previously in the PawnList and iterated by FConstPawnIterator. */
+struct ENGINE_API FPawnIteratorObject
+{
+	APawn* operator->() const { return Pawn; }
+	APawn& operator*() const { return *Pawn; }
+	APawn* Get() const { return Pawn; }
+
+private:
+	FPawnIteratorObject()
+		: Pawn(nullptr)
+	{
+	}
+
+	FPawnIteratorObject(APawn* InPawn)
+		: Pawn(InPawn)
+	{
+	}
+
+	APawn* Pawn;
+
+	friend class FConstPawnIterator;
+};
+
+template< class T > FORCEINLINE T* Cast(const FPawnIteratorObject& Src) { return Cast<T>(Src.Get()); }
+
+/** 
+ * Imitation iterator class that attempts to provide the basic interface that FConstPawnIterator previously did when a typedef of TArray<TWeakObjectPtr<APawn>>::Iterator.
+ * In general you should prefer not to use this iterator and instead use TActorIterator<APawn> or TActorRange<APawn> (or the desired more derived type).
+ * This iterator will likely be deprecated in a future release.
+ */
+class ENGINE_API FConstPawnIterator
+{
+private:
+	FConstPawnIterator(UWorld* World);
+
+public:
+	~FConstPawnIterator();
+
+	operator bool() const;
+	FPawnIteratorObject operator*() const;
+	TUniquePtr<FPawnIteratorObject> operator->() const;
+
+	FConstPawnIterator& operator++();
+	FConstPawnIterator& operator++(int);
+	UE_DEPRECATED(4.23, "Decrement operator no longer means anything on a pawn iterator")
+	FConstPawnIterator& operator--() { return *this; }
+	UE_DEPRECATED(4.23, "Decrement operator no longer means anything on a pawn iterator")
+	FConstPawnIterator& operator--(int) { return *this; }
+
+private:
+	TActorIterator<APawn>* Iterator;
+
+	friend UWorld;
+};
+
+
 DECLARE_LOG_CATEGORY_EXTERN(LogSpawn, Warning, All);
 
 DECLARE_MULTICAST_DELEGATE_OneParam(FOnActorSpawned, AActor*);
@@ -1220,9 +1277,6 @@ private:
 	/** List of all the player controllers in the world. */
 	TArray<TWeakObjectPtr<class APlayerController> > PlayerControllerList;
 
-	/** List of all the pawns in the world. */
-	TArray<TWeakObjectPtr<class APawn> > PawnList;
-
 	/** List of all the cameras in the world that auto-activate for players. */
 	TArray<TWeakObjectPtr<ACameraActor> > AutoCameraActorList;
 
@@ -2022,6 +2076,7 @@ public:
 	FConstPawnIterator GetPawnIterator() const;
 	
 	/** @return Returns the number of Pawns. */
+	UE_DEPRECATED(4.23, "GetNumPawns is no longer a supported function on UWorld. The version that remains for backwards compatibility is significantly more expensive to call.")
 	int32 GetNumPawns() const;
 
 	/** @return Returns an iterator for the player controller list. */
@@ -2234,19 +2289,16 @@ public:
 	 */
 	void RemoveController( AController* Controller );
 
-	/**
-	 * Inserts the passed in pawn at the front of the linked list of pawns.
-	 *
-	 * @param	Pawn	Pawn to insert, use NULL to clear list
-	 */
-	void AddPawn( APawn* Pawn );
+	UE_DEPRECATED(4.23, "There is no longer a reason to AddPawn to UWorld")
+	void AddPawn( APawn* Pawn ) { }
 	
 	/**
 	 * Removes the passed in pawn from the linked list of pawns.
 	 *
 	 * @param	Pawn	Pawn to remove
 	 */
-	void RemovePawn( APawn* Pawn );
+	UE_DEPRECATED(4.23, "RemovePawn has been deprecated and should no longer need to be called as PawnList is no longer maintained and Unpossess should be handled by EndPlay.")
+	void RemovePawn( APawn* Pawn ) const;
 
 	/**
 	 * Adds the passed in actor to the special network actor list
@@ -2259,7 +2311,7 @@ public:
 	 * Removes the passed in actor to from special network actor list
 	 * @param	Actor	Actor to remove
 	 */
-	void RemoveNetworkActor( AActor* Actor );
+	void RemoveNetworkActor( AActor* Actor ) const;
 
 	/** Add a listener for OnActorSpawned events */
 	FDelegateHandle AddOnActorSpawnedHandler( const FOnActorSpawned::FDelegate& InHandler );
@@ -2275,14 +2327,14 @@ public:
 	 *	
 	 * @return	true if actor is contained by any of the loaded levels, false otherwise
 	 */
-	bool ContainsActor( AActor* Actor );
+	bool ContainsActor( AActor* Actor ) const;
 
 	/**
 	 * Returns whether audio playback is allowed for this scene.
 	 *
 	 * @return true if current world is GWorld, false otherwise
 	 */
-	virtual bool AllowAudioPlayback();
+	bool AllowAudioPlayback() const;
 
 	//~ Begin UObject Interface
 	virtual void Serialize( FArchive& Ar ) override;
@@ -2847,7 +2899,7 @@ public:
 	 * @param	Actor					Actor to remove.
 	 * @param	bShouldModifyLevel		If true, Modify() the level before removing the actor if in the editor.
 	 */
-	void RemoveActor( AActor* Actor, bool bShouldModifyLevel );
+	void RemoveActor( AActor* Actor, bool bShouldModifyLevel ) const;
 
 	/**
 	 * Spawn Actors with given transform and SpawnParameters
diff --git a/Engine/Source/Runtime/Engine/Private/Pawn.cpp b/Engine/Source/Runtime/Engine/Private/Pawn.cpp
--- a/Engine/Source/Runtime/Engine/Private/Pawn.cpp
+++ b/Engine/Source/Runtime/Engine/Private/Pawn.cpp
@@ -106,17 +106,17 @@ void APawn::PostInitializeComponents()
 	
 	if (!IsPendingKill())
 	{
-		GetWorld()->AddPawn( this );
+		UWorld* World = GetWorld();
 
 		// Automatically add Controller to AI Pawns if we are allowed to.
 		if (AutoPossessPlayer == EAutoReceiveInput::Disabled
 			&& AutoPossessAI != EAutoPossessAI::Disabled && Controller == nullptr && GetNetMode() != NM_Client
 #if WITH_EDITOR
-			&& (GIsEditor == false || GetWorld()->IsGameWorld())
+			&& (GIsEditor == false || World->IsGameWorld())
 #endif // WITH_EDITOR
 			)
 		{
-			const bool bPlacedInWorld = (GetWorld()->bStartup);
+			const bool bPlacedInWorld = (World->bStartup);
 			if ((AutoPossessAI == EAutoPossessAI::PlacedInWorldOrSpawned) ||
 				(AutoPossessAI == EAutoPossessAI::PlacedInWorld && bPlacedInWorld) ||
 				(AutoPossessAI == EAutoPossessAI::Spawned && !bPlacedInWorld))
@@ -179,12 +179,8 @@ void APawn::PawnStartFire(uint8 FireModeNum) {}
 
 AActor* APawn::GetMovementBaseActor(const APawn* Pawn)
 {
-	if (Pawn != nullptr && Pawn->GetMovementBase())
-	{
-		return Pawn->GetMovementBase()->GetOwner();
-	}
-
-	return nullptr;
+	UPrimitiveComponent* MovementBase = (Pawn ? Pawn->GetMovementBase() : nullptr);
+	return (MovementBase ? MovementBase->GetOwner() : nullptr);
 }
 
 bool APawn::CanBeBaseForCharacter(class APawn* APawn) const
@@ -393,7 +389,6 @@ void APawn::PawnClientRestart()
 void APawn::Destroyed()
 {
 	DetachFromControllerPendingDestroy();
-	GetWorld()->RemovePawn( this );
 	Super::Destroyed();
 }
 
@@ -403,7 +398,6 @@ void APawn::EndPlay(const EEndPlayReason::Type EndPlayReason)
 	if (EndPlayReason != EEndPlayReason::Destroyed)
 	{
 		DetachFromControllerPendingDestroy();
-		GetWorld()->RemovePawn( this );
 	}
 	
 	Super::EndPlay(EndPlayReason);
diff --git a/Engine/Source/Runtime/Engine/Private/World.cpp b/Engine/Source/Runtime/Engine/Private/World.cpp
--- a/Engine/Source/Runtime/Engine/Private/World.cpp
+++ b/Engine/Source/Runtime/Engine/Private/World.cpp
@@ -133,7 +133,7 @@ static int32 bDisableRemapScriptActors = 0;
 FAutoConsoleVariableRef CVarDisableRemapScriptActors(TEXT("net.DisableRemapScriptActors"), bDisableRemapScriptActors, TEXT("When set, disables name remapping of compiled script actors (for networking)"));
 
 template<class Function>
-static void ForEachNetDriver(UEngine* Engine, UWorld* const World, const Function InFunction)
+static void ForEachNetDriver(UEngine* Engine, const UWorld* const World, const Function InFunction)
 {
 	if (Engine == nullptr || World == nullptr)
 	{
@@ -1621,7 +1621,7 @@ UWorld* UWorld::CreateWorld(const EWorldType::Type InWorldType, bool bInformEngi
 	return NewWorld;
 }
 
-void UWorld::RemoveActor(AActor* Actor, bool bShouldModifyLevel)
+void UWorld::RemoveActor(AActor* Actor, bool bShouldModifyLevel) const
 {
 	ULevel* CheckLevel = Actor->GetLevel();
 	const int32 ActorListIndex = CheckLevel->Actors.Find( Actor );
@@ -1646,12 +1646,12 @@ void UWorld::RemoveActor(AActor* Actor, bool bShouldModifyLevel)
 }
 
 
-bool UWorld::ContainsActor( AActor* Actor )
+bool UWorld::ContainsActor( AActor* Actor ) const
 {
 	return (Actor && Actor->GetWorld() == this);
 }
 
-bool UWorld::AllowAudioPlayback()
+bool UWorld::AllowAudioPlayback() const
 {
 	return bAllowAudioPlayback;
 }
@@ -2476,21 +2476,23 @@ void UWorld::RemoveFromWorld( ULevel* Level, bool bAllowIncrementalRemoval )
 			}
 
 			// Remove any pawns from the pawn list that are about to be streamed out
-			for( FConstPawnIterator Iterator = GetPawnIterator(); Iterator; ++Iterator )
+			for (APawn* Pawn : TActorRange<APawn>(this))
 			{
-				if (APawn* Pawn = Iterator->Get())
+				if (Pawn->IsInLevel(Level))
 				{
-					if (Pawn->IsInLevel(Level))
+					AController* Controller = Pawn->GetController();
+					// This should have happened as part of the RouteEndPlay above, but ensuring to validate this assumption and maintain behavior
+					// with RemovePawn having been deprecated
+					if (!ensure(Controller == nullptr || (Controller->GetPawn() == Pawn)))
 					{
-						RemovePawn(Pawn);
-						--Iterator;
-					}
-					else if (UCharacterMovementComponent* CharacterMovement = Cast<UCharacterMovementComponent>(Pawn->GetMovementComponent()))
-					{
-						// otherwise force floor check in case the floor was streamed out from under it
-						CharacterMovement->bForceNextFloorCheck = true;
+						Controller->UnPossess();
 					}
 				}
+				else if (UCharacterMovementComponent* CharacterMovement = Cast<UCharacterMovementComponent>(Pawn->GetMovementComponent()))
+				{
+					// otherwise force floor check in case the floor was streamed out from under it
+					CharacterMovement->bForceNextFloorCheck = true;
+				}
 			}
 
 			Level->ReleaseRenderingResources();
@@ -4177,35 +4179,69 @@ void UWorld::RemoveController( AController* Controller )
 	}
 }
 
-FConstPawnIterator UWorld::GetPawnIterator() const
+FConstPawnIterator::FConstPawnIterator(UWorld* World)
+	: Iterator(new TActorIterator<APawn>(World))
 {
-	auto Result = PawnList.CreateConstIterator();
-	return (const FConstPawnIterator&)Result;
 }
 
-int32 UWorld::GetNumPawns() const
+FConstPawnIterator::~FConstPawnIterator()
 {
-	return PawnList.Num();
+	delete Iterator;
 }
 
-void UWorld::AddPawn( APawn* Pawn )
+ FConstPawnIterator::operator bool() const
 {
-	check( Pawn );
-	PawnList.AddUnique( Pawn );
+	return (bool)*Iterator;
 }
 
+FConstPawnIterator& FConstPawnIterator::operator++()
+{
+	++(*Iterator);
+	return *this;
+}
 
-void UWorld::RemovePawn( APawn* Pawn )
+FConstPawnIterator& FConstPawnIterator::operator++(int)
 {
-	check( Pawn );
-	
-	AController* Controller = Pawn->GetController();
-	if (Controller && (Controller->GetPawn() == Pawn))
+	++(*Iterator);
+	return *this;
+}
+
+FPawnIteratorObject FConstPawnIterator::operator*() const
+{
+	return FPawnIteratorObject(**Iterator);
+}
+
+TUniquePtr<FPawnIteratorObject> FConstPawnIterator::operator->() const
+{
+	return TUniquePtr<FPawnIteratorObject>(new FPawnIteratorObject(**Iterator));
+}
+
+FConstPawnIterator UWorld::GetPawnIterator() const
+{
+	return FConstPawnIterator(const_cast<UWorld*>(this)); // For backwards compat GetPawnIterator needs to remain const, but TActorIterator can't use a const UWorld.
+}
+
+int32 UWorld::GetNumPawns() const
+{
+	int32 NumPawns = 0;
+	for (APawn* Pawn : TActorRange<APawn>(const_cast<UWorld*>(this))) // For backwards compat GetNumPawns needs to remain const, but TActorRange can't use a const UWorld.
 	{
-		Controller->UnPossess();
+		++NumPawns;
 	}
 
-	PawnList.Remove( Pawn );	
+	return NumPawns;
+}
+
+void UWorld::RemovePawn( APawn* Pawn ) const
+{
+	if (Pawn)
+	{
+		AController* Controller = Pawn->GetController();
+		if (Controller && (Controller->GetPawn() == Pawn))
+		{
+			Controller->UnPossess();
+		}
+	}
 }
 
 
@@ -4249,7 +4285,7 @@ void UWorld::AddNetworkActor( AActor* Actor )
 	});
 }
 
-void UWorld::RemoveNetworkActor( AActor* Actor )
+void UWorld::RemoveNetworkActor( AActor* Actor ) const
 {
 	if (Actor)
 	{
@@ -5977,15 +6013,11 @@ UWorld* FSeamlessTravelHandler::Tick()
 				{
 					KeepAnnotation.Clear(TheActor);
 					TheActor->Rename(nullptr, LoadedWorld->PersistentLevel);
-					// if it's a Controller or a Pawn, add it to the appropriate list in the new world's WorldSettings
+					// if it's a Controller, add it to the appropriate list in the new world's WorldSettings
 					if (TheActor->IsA<AController>())
 					{
 						LoadedWorld->AddController(static_cast<AController*>(TheActor));
 					}
-					else if (TheActor->IsA<APawn>())
-					{
-						LoadedWorld->AddPawn(static_cast<APawn*>(TheActor));
-					}
 					else if (TheActor->IsA<AGameModeBase>())
 					{
 						KeptGameMode = static_cast<AGameModeBase*>(TheActor);
