Fix const error
#rnx

diff --git a/Engine/Source/Runtime/Engine/Classes/Engine/World.h b/Engine/Source/Runtime/Engine/Classes/Engine/World.h
--- a/Engine/Source/Runtime/Engine/Classes/Engine/World.h
+++ b/Engine/Source/Runtime/Engine/Classes/Engine/World.h
@@ -2682,7 +2682,7 @@ public:
 	 * Transacts the specified level -- the correct way to modify a level
 	 * as opposed to calling Level->Modify.
 	 */
-	void ModifyLevel(ULevel* Level);
+	void ModifyLevel(ULevel* Level) const;
 
 	/**
 	 * Ensures that the collision detection tree is fully built. This should be called after the full level reload to make sure
diff --git a/Engine/Source/Runtime/Engine/Private/World.cpp b/Engine/Source/Runtime/Engine/Private/World.cpp
--- a/Engine/Source/Runtime/Engine/Private/World.cpp
+++ b/Engine/Source/Runtime/Engine/Private/World.cpp
@@ -1853,7 +1853,7 @@ bool UWorld::UpdateCullDistanceVolumes(AActor* ActorToUpdate, UPrimitiveComponen
 }
 
 
-void UWorld::ModifyLevel(ULevel* Level)
+void UWorld::ModifyLevel(ULevel* Level) const
 {
 	if( Level && Level->HasAnyFlags(RF_Transactional))
 	{
