Reinclude Pawn.h as it was relied on for Clang compilers
#rnx

diff --git a/Engine/Source/Runtime/Engine/Classes/Engine/World.h b/Engine/Source/Runtime/Engine/Classes/Engine/World.h
--- a/Engine/Source/Runtime/Engine/Classes/Engine/World.h
+++ b/Engine/Source/Runtime/Engine/Classes/Engine/World.h
@@ -13,6 +13,7 @@
 #include "Engine/EngineBaseTypes.h"
 #include "CollisionQueryParams.h"
 #include "WorldCollision.h"
+#include "GameFramework/Pawn.h"
 #include "EngineDefines.h"
 #include "Engine/Blueprint.h"
 #include "Engine/PendingNetGame.h"
@@ -30,7 +31,6 @@ class AController;
 class AGameModeBase;
 class AGameStateBase;
 class AMatineeActor;
-class APawn;
 class APhysicsVolume;
 class APlayerController;
 class AWorldSettings;
