Removing the Custom Primitive Data sampling node in favor of using the existing Parameter nodes

#rnx

diff --git a/Engine/Source/Runtime/Engine/Classes/Materials/MaterialExpressionCustomPrimitiveData.h b/Engine/Source/Runtime/Engine/Classes/Materials/MaterialExpressionCustomPrimitiveData.h
deleted file mode 100644
diff --git a/Engine/Source/Runtime/Engine/Private/Materials/MaterialExpressions.cpp b/Engine/Source/Runtime/Engine/Private/Materials/MaterialExpressions.cpp
--- a/Engine/Source/Runtime/Engine/Private/Materials/MaterialExpressions.cpp
+++ b/Engine/Source/Runtime/Engine/Private/Materials/MaterialExpressions.cpp
@@ -72,7 +72,6 @@
 #include "Materials/MaterialExpressionCosine.h"
 #include "Materials/MaterialExpressionCrossProduct.h"
 #include "Materials/MaterialExpressionCustom.h"
-#include "Materials/MaterialExpressionCustomPrimitiveData.h"
 #include "Materials/MaterialExpressionDDX.h"
 #include "Materials/MaterialExpressionDDY.h"
 #include "Materials/MaterialExpressionDecalDerivative.h"
@@ -15317,95 +15316,6 @@ void UMaterialExpressionCurveAtlasRowParameter::PostEditChangeProperty(FProperty
 }
 #endif
 
-//
-// UMaterialExpressionCustomPrimitiveData
-//
-UMaterialExpressionCustomPrimitiveData::UMaterialExpressionCustomPrimitiveData(const FObjectInitializer& ObjectInitializer)
-	: Super(ObjectInitializer)
-{
-#if WITH_EDITORONLY_DATA
-	// Structure to hold one-time initialization
-	struct FConstructorStatics
-	{
-		FText NAME_Vectors;
-		FConstructorStatics()
-			: NAME_Vectors(LOCTEXT("Custom", "Custom"))
-		{
-		}
-	};
-	static FConstructorStatics ConstructorStatics;
-
-	MenuCategories.Add(ConstructorStatics.NAME_Vectors);
-
-	bShaderInputData = true;
-	bShowOutputNameOnPin = true;
-
-	Outputs.Reset();
-#endif
-}
-
-#if WITH_EDITOR
-int32 UMaterialExpressionCustomPrimitiveData::Compile(class FMaterialCompiler* Compiler, int32 OutputIndex)
-{
-	if (Material && Material->MaterialDomain == MD_DeferredDecal)
-	{
-		return CompilerError(Compiler, TEXT("Expression not available in the deferred decal material domain."));
-	}
-
-	return Compiler->CustomPrimitiveData(CustomIndices[OutputIndex].PrimitiveDataIndex);
-}
-
-void UMaterialExpressionCustomPrimitiveData::GetCaption(TArray<FString>& OutCaptions) const
-{
-	OutCaptions.Add(TEXT("Custom Primitive Data"));
-}
-
-void UMaterialExpressionCustomPrimitiveData::PostEditChangeProperty(FPropertyChangedEvent& PropertyChangedEvent)
-{
-#if WITH_EDITORONLY_DATA
-	
-	if (PropertyChangedEvent.GetPropertyName() == GET_MEMBER_NAME_CHECKED(FPrimitiveDataIndex, CustomDesc) ||
-		PropertyChangedEvent.GetPropertyName() == GET_MEMBER_NAME_CHECKED(UMaterialExpressionCustomPrimitiveData, CustomIndices))
-	{
-		// If more than the supported number of custom floats have been added, remove the overflow
-		if (CustomIndices.Num() > FCustomPrimitiveData::NumCustomPrimitiveDataFloats)
-		{
-			CustomIndices.SetNum(FCustomPrimitiveData::NumCustomPrimitiveDataFloats);
-		}
-
-		// Match the number of pins to the number of descriptions
-		Outputs.SetNumZeroed(CustomIndices.Num());
-
-		// Update the names of the pins
-		for (int i = 0; i < Outputs.Num(); i++)
-		{
-			Outputs[i].OutputName = *(CustomIndices[i].CustomDesc);
-		}
-
-		// Reconstruct the node to get the new names applied to the pins
-		GraphNode->ReconstructNode();
-
-		// No need to update preview if we only change output descriptions
-		bNeedToUpdatePreview = false;
-
-		// Done here (skip calling base class)
-		return;
-	}
-	else if (PropertyChangedEvent.GetPropertyName() == GET_MEMBER_NAME_CHECKED(FPrimitiveDataIndex, PrimitiveDataIndex))
-	{
-		// Clamp indices
-		for (int i = 0; i < CustomIndices.Num(); i++)
-		{
-			const int32 PrimDataIndex = CustomIndices[i].PrimitiveDataIndex;
-			CustomIndices[i].PrimitiveDataIndex = FMath::Clamp(PrimDataIndex, 0, FCustomPrimitiveData::NumCustomPrimitiveDataFloats-1);
-		}
-	}
-#endif
-
-	Super::PostEditChangeProperty(PropertyChangedEvent);
-}
-#endif // WITH_EDITOR
-
 //
 //	UMaterialExpressionShadingModel
 //
diff --git a/Engine/Source/Runtime/Engine/Public/PrimitiveUniformShaderParameters.h b/Engine/Source/Runtime/Engine/Public/PrimitiveUniformShaderParameters.h
--- a/Engine/Source/Runtime/Engine/Public/PrimitiveUniformShaderParameters.h
+++ b/Engine/Source/Runtime/Engine/Public/PrimitiveUniformShaderParameters.h
@@ -41,7 +41,7 @@ BEGIN_GLOBAL_SHADER_PARAMETER_STRUCT(FPrimitiveUniformShaderParameters,ENGINE_AP
 	SHADER_PARAMETER(int32, SingleCaptureIndex)
     SHADER_PARAMETER(uint32, OutputVelocity)
 	// 12 bytes of padding here, feel free to use
-	SHADER_PARAMETER_ARRAY(FVector4, CustomPrimitiveData, [FCustomPrimitiveData::NumCustomPrimitiveDataFloat4s]) // Custom data per primitive that can be accessed through UMaterialExpressionCustomPrimitiveData and modified through UStaticMeshComponent
+	SHADER_PARAMETER_ARRAY(FVector4, CustomPrimitiveData, [FCustomPrimitiveData::NumCustomPrimitiveDataFloat4s]) // Custom data per primitive that can be accessed through material expression parameters and modified through UStaticMeshComponent
 
 END_GLOBAL_SHADER_PARAMETER_STRUCT()
 
