Adding dynamic update of the GPU custom primtive data through new methods on StaticMeshComponent.
Adding a mode for scalar and vector material parameters to choose to fetch data from the Custom Primitive Data.
Disabled scalar/vector material parameters from showing up as overrrides in Material Instance.

#rnx

diff --git a/Engine/Source/Runtime/Engine/Classes/Components/StaticMeshComponent.h b/Engine/Source/Runtime/Engine/Classes/Components/StaticMeshComponent.h
--- a/Engine/Source/Runtime/Engine/Classes/Components/StaticMeshComponent.h
+++ b/Engine/Source/Runtime/Engine/Classes/Components/StaticMeshComponent.h
@@ -387,6 +387,22 @@ public:
 	UFUNCTION(BlueprintCallable, Category="Components|StaticMesh")
 	void GetLocalBounds(FVector& Min, FVector& Max) const;
 
+	/** Set custom primitive data at index DataIndex. */
+	UFUNCTION(BlueprintCallable, Category="Components|StaticMesh")
+	void SetCustomPrimitiveDataFloat(float Value, int32 DataIndex);
+
+	/** Set custom primitive data, two floats at once, from index DataIndex to index DataIndex + 2. */
+	UFUNCTION(BlueprintCallable, Category="Components|StaticMesh")
+	void SetCustomPrimitiveDataVector2(FVector2D Value, int32 DataIndex);
+
+	/** Set custom primitive data, three floats at once, from index DataIndex to index DataIndex + 3. */
+	UFUNCTION(BlueprintCallable, Category="Components|StaticMesh")
+	void SetCustomPrimitiveDataVector3(FVector Value, int32 DataIndex);
+
+	/** Set custom primitive data, four floats at once, from index DataIndex to index DataIndex + 4. */
+	UFUNCTION(BlueprintCallable, Category="Components|StaticMesh")
+	void SetCustomPrimitiveDataVector4(FVector4 Value, int32 DataIndex);
+
 	/** 
 	 * Set forced reverse culling
 	 */
@@ -443,6 +459,7 @@ protected:
 	virtual void CreateRenderState_Concurrent() override;
 	virtual void OnCreatePhysicsState() override;
 	virtual void OnDestroyPhysicsState() override;
+	virtual void SendRenderDynamicData_Concurrent() override;
 public:
 	virtual void InvalidateLightingCacheDetailed(bool bInvalidateBuildEnqueuedLighting, bool bTranslationOnly) override;
 	virtual UObject const* AdditionalStatObject() const override;
@@ -624,6 +641,8 @@ private:
 	/** Update the vertex override colors */
 	void PrivateFixupOverrideColors();
 
+	/** Insert an array of floats into the CustomPrimitiveData, starting at the given index */
+	void SetCustomPrimitiveDataInternal(const TArray<float>& Values, int32 DataIndex);
 protected:
 
 	/** Whether the component type supports static lighting. */
diff --git a/Engine/Source/Runtime/Engine/Classes/Materials/MaterialExpressionParameter.h b/Engine/Source/Runtime/Engine/Classes/Materials/MaterialExpressionParameter.h
--- a/Engine/Source/Runtime/Engine/Classes/Materials/MaterialExpressionParameter.h
+++ b/Engine/Source/Runtime/Engine/Classes/Materials/MaterialExpressionParameter.h
@@ -63,6 +63,6 @@ class UMaterialExpressionParameter : public UMaterialExpression
 	/**
 	 * Get list of parameter names for static parameter sets
 	 */
-	void GetAllParameterInfo(TArray<FMaterialParameterInfo> &OutParameterInfo, TArray<FGuid> &OutParameterIds, const FMaterialParameterInfo& InBaseParameterInfo) const;
+	virtual void GetAllParameterInfo(TArray<FMaterialParameterInfo> &OutParameterInfo, TArray<FGuid> &OutParameterIds, const FMaterialParameterInfo& InBaseParameterInfo) const;
 
 };
diff --git a/Engine/Source/Runtime/Engine/Classes/Materials/MaterialExpressionScalarParameter.h b/Engine/Source/Runtime/Engine/Classes/Materials/MaterialExpressionScalarParameter.h
--- a/Engine/Source/Runtime/Engine/Classes/Materials/MaterialExpressionScalarParameter.h
+++ b/Engine/Source/Runtime/Engine/Classes/Materials/MaterialExpressionScalarParameter.h
@@ -18,6 +18,12 @@ class UMaterialExpressionScalarParameter : public UMaterialExpressionParameter
 	UPROPERTY(EditAnywhere, Category=MaterialExpressionScalarParameter)
 	float DefaultValue;
 
+	UPROPERTY(EditAnywhere, Category=CustomPrimitiveData)
+	bool bUseCustomPrimitiveData = false;
+
+	UPROPERTY(EditAnywhere, Category=CustomPrimitiveData, meta=(ClampMin="0"))
+	uint8 PrimitiveDataIndex = 0;
+
 #if WITH_EDITORONLY_DATA
 	/** 
 	 * Sets the lower bound for the slider on this parameter in the material instance editor. 
@@ -53,6 +59,8 @@ class UMaterialExpressionScalarParameter : public UMaterialExpressionParameter
 #endif
 
 	virtual bool IsUsedAsAtlasPosition() const { return false; }
+
+	virtual void GetAllParameterInfo(TArray<FMaterialParameterInfo> &OutParameterInfo, TArray<FGuid> &OutParameterIds, const FMaterialParameterInfo& InBaseParameterInfo) const override;
 };
 
 
diff --git a/Engine/Source/Runtime/Engine/Classes/Materials/MaterialExpressionVectorParameter.h b/Engine/Source/Runtime/Engine/Classes/Materials/MaterialExpressionVectorParameter.h
--- a/Engine/Source/Runtime/Engine/Classes/Materials/MaterialExpressionVectorParameter.h
+++ b/Engine/Source/Runtime/Engine/Classes/Materials/MaterialExpressionVectorParameter.h
@@ -17,6 +17,12 @@ class UMaterialExpressionVectorParameter : public UMaterialExpressionParameter
 	UPROPERTY(EditAnywhere, Category=MaterialExpressionVectorParameter)
 	FLinearColor DefaultValue;
 
+	UPROPERTY(EditAnywhere, Category=CustomPrimitiveData)
+	bool bUseCustomPrimitiveData = false;
+
+	UPROPERTY(EditAnywhere, Category=CustomPrimitiveData, meta=(ClampMin="0"))
+	uint8 PrimitiveDataIndex = 0;
+
 	//~ Begin UMaterialExpression Interface
 #if WITH_EDITOR
 	virtual int32 Compile(class FMaterialCompiler* Compiler, int32 OutputIndex) override;
@@ -37,5 +43,7 @@ class UMaterialExpressionVectorParameter : public UMaterialExpressionParameter
 #endif
 
 	virtual bool IsUsedAsChannelMask() const {return false;}
+
+	virtual void GetAllParameterInfo(TArray<FMaterialParameterInfo> &OutParameterInfo, TArray<FGuid> &OutParameterIds, const FMaterialParameterInfo& InBaseParameterInfo) const override;
 };
 
diff --git a/Engine/Source/Runtime/Engine/Private/Components/StaticMeshComponent.cpp b/Engine/Source/Runtime/Engine/Private/Components/StaticMeshComponent.cpp
--- a/Engine/Source/Runtime/Engine/Private/Components/StaticMeshComponent.cpp
+++ b/Engine/Source/Runtime/Engine/Private/Components/StaticMeshComponent.cpp
@@ -624,6 +624,13 @@ void UStaticMeshComponent::OnDestroyPhysicsState()
 	bNavigationRelevant = IsNavigationRelevant();
 }
 
+void UStaticMeshComponent::SendRenderDynamicData_Concurrent()
+{
+	GetWorld()->Scene->UpdateCustomPrimitiveData(this);
+
+	Super::SendRenderDynamicData_Concurrent();
+}
+
 #if WITH_EDITORONLY_DATA
 
 /** Return the total number of LOD sections in the LOD resources */
@@ -1874,6 +1881,51 @@ void UStaticMeshComponent::SetDistanceFieldSelfShadowBias(float NewValue)
 	}
 }
 
+void UStaticMeshComponent::SetCustomPrimitiveDataInternal(const TArray<float>& Values, int32 DataIndex)
+{
+	// Can only set data on valid indices and only if there's actually any data to set
+	if (DataIndex < FCustomPrimitiveData::NumCustomPrimitiveDataFloats && Values.Num() > 0)
+	{
+		// Number of floats needed in the custom primitive data array
+		const int32 NeededFloats = FMath::Min(DataIndex + Values.Num(), FCustomPrimitiveData::NumCustomPrimitiveDataFloats);
+
+		// Number of value to copy into the custom primitive data array at index DataIndex. Capped to not overflow
+		const int32 NumValuesToSet = FMath::Min(Values.Num(), FCustomPrimitiveData::NumCustomPrimitiveDataFloats - DataIndex);
+
+		// If trying to set data on an index which doesn't exist yet, allocate up to it
+		// TODO, use a sparse structure instead?
+		if (NeededFloats > CustomPrimitiveData.Data.Num())
+		{
+			CustomPrimitiveData.Data.SetNumZeroed(NeededFloats);
+		}
+	
+		FMemory::Memcpy(&CustomPrimitiveData.Data[DataIndex], Values.GetData(), NumValuesToSet * sizeof(float));
+
+		// Make sure the render dynamic data is marked dirty to ensure an update if the primitive buffer before next draw
+		MarkRenderDynamicDataDirty();
+	}
+}
+
+void UStaticMeshComponent::SetCustomPrimitiveDataFloat(float Value, int32 DataIndex)
+{
+	SetCustomPrimitiveDataInternal({Value}, DataIndex);
+}
+
+void UStaticMeshComponent::SetCustomPrimitiveDataVector2(FVector2D Value, int32 DataIndex)
+{
+	SetCustomPrimitiveDataInternal({Value.X, Value.Y}, DataIndex);
+}
+
+void UStaticMeshComponent::SetCustomPrimitiveDataVector3(FVector Value, int32 DataIndex)
+{
+	SetCustomPrimitiveDataInternal({Value.X, Value.Y, Value.Z}, DataIndex);
+}
+
+void UStaticMeshComponent::SetCustomPrimitiveDataVector4(FVector4 Value, int32 DataIndex)
+{
+	SetCustomPrimitiveDataInternal({Value.X, Value.Y, Value.Z, Value.W}, DataIndex);
+}
+
 void UStaticMeshComponent::SetReverseCulling(bool ReverseCulling)
 {
 	if (ReverseCulling != bReverseCulling)
diff --git a/Engine/Source/Runtime/Engine/Private/Materials/HLSLMaterialTranslator.h b/Engine/Source/Runtime/Engine/Private/Materials/HLSLMaterialTranslator.h
--- a/Engine/Source/Runtime/Engine/Private/Materials/HLSLMaterialTranslator.h
+++ b/Engine/Source/Runtime/Engine/Private/Materials/HLSLMaterialTranslator.h
@@ -5482,13 +5482,50 @@ protected:
 
 	}
 
-	virtual int32 CustomPrimitiveData(int32 OutputIndex) override
+	virtual int32 CustomPrimitiveData(int32 OutputIndex, EMaterialValueType Type) override
 	{
-		const int32 CustomDataIndex = OutputIndex / FCustomPrimitiveData::NumCustomPrimitiveDataFloat4s;
-		const int32 ElementIndex = OutputIndex % 4; // Index x, y, z or w
+		check(OutputIndex < FCustomPrimitiveData::NumCustomPrimitiveDataFloats);
 
-		FString HlslName = FString::Printf(TEXT("CustomPrimitiveData[%d][%d]"), CustomDataIndex, ElementIndex);
-		return GetPrimitiveProperty(MCT_Float, TEXT("CustomPrimitiveData"), *HlslName);
+		const int32 NumComponents = GetNumComponents(Type);
+
+		FString HlslCode;
+			
+		// Only float2, float3 and float4 need this
+		if (NumComponents > 1)
+		{
+			HlslCode.Append(FString::Printf(TEXT("float%d("), NumComponents));
+		}
+
+		for (int i = 0; i < NumComponents; i++)
+		{
+			const int32 CurrentOutputIndex = OutputIndex + i;
+
+			// Check if we are accessing inside the array, otherwise default to 0
+			if (CurrentOutputIndex < FCustomPrimitiveData::NumCustomPrimitiveDataFloats)
+			{
+				const int32 CustomDataIndex = CurrentOutputIndex / FCustomPrimitiveData::NumCustomPrimitiveDataFloat4s;
+				const int32 ElementIndex = CurrentOutputIndex % 4; // Index x, y, z or w
+
+				HlslCode.Append(FString::Printf(TEXT("GetPrimitiveData(Parameters.PrimitiveId).CustomPrimitiveData[%d][%d]"), CustomDataIndex, ElementIndex));
+			}
+			else
+			{
+				HlslCode.Append(TEXT("0.0f"));
+			}
+
+			if (i+1 < NumComponents)
+			{
+				HlslCode.Append(", ");
+			}
+		}
+
+		// This is the matching parenthesis to the first append
+		if (NumComponents > 1)
+		{
+			HlslCode.AppendChar(')');
+		}
+
+		return AddCodeChunk(Type, TEXT("%s"), *HlslCode);
 	}
 
 	virtual int32 ShadingModel(EMaterialShadingModel InSelectedShadingModel) override
diff --git a/Engine/Source/Runtime/Engine/Private/Materials/MaterialExpressions.cpp b/Engine/Source/Runtime/Engine/Private/Materials/MaterialExpressions.cpp
--- a/Engine/Source/Runtime/Engine/Private/Materials/MaterialExpressions.cpp
+++ b/Engine/Source/Runtime/Engine/Private/Materials/MaterialExpressions.cpp
@@ -6352,17 +6352,48 @@ UMaterialExpressionVectorParameter::UMaterialExpressionVectorParameter(const FOb
 #if WITH_EDITOR
 int32 UMaterialExpressionVectorParameter::Compile(class FMaterialCompiler* Compiler, int32 OutputIndex)
 {
-	return Compiler->VectorParameter(ParameterName,DefaultValue);
+	if (bUseCustomPrimitiveData)
+	{
+		return Compiler->CustomPrimitiveData(PrimitiveDataIndex, MCT_Float4);
+	}
+	else
+	{
+		return Compiler->VectorParameter(ParameterName,DefaultValue);
+	}
 }
 
 void UMaterialExpressionVectorParameter::GetCaption(TArray<FString>& OutCaptions) const
 {
-	OutCaptions.Add(FString::Printf(
-		 TEXT("Param (%.3g,%.3g,%.3g,%.3g)"),
-		 DefaultValue.R,
-		 DefaultValue.G,
-		 DefaultValue.B,
-		 DefaultValue.A ));
+	if (bUseCustomPrimitiveData)
+	{
+		FString IndexString = FString::Printf(TEXT("Index %d"), PrimitiveDataIndex);
+
+		// Add info about remaining 3 components
+		for (int i = 1; i < 4; i++)
+		{
+			// Append index if it's valid, otherwise append N/A
+			if(PrimitiveDataIndex+i < FCustomPrimitiveData::NumCustomPrimitiveDataFloats)
+			{
+				IndexString.Append(FString::Printf(TEXT(", %d"), PrimitiveDataIndex+i));
+			}
+			else
+			{
+				IndexString.Append(FString::Printf(TEXT(", N/A")));
+			}
+		}
+
+		OutCaptions.Add(IndexString); 
+		OutCaptions.Add(FString::Printf(TEXT("Custom Primitive Data"))); 
+	}
+	else
+	{
+		OutCaptions.Add(FString::Printf(
+			 TEXT("Param (%.3g,%.3g,%.3g,%.3g)"),
+			 DefaultValue.R,
+			 DefaultValue.G,
+			 DefaultValue.B,
+			 DefaultValue.A ));
+	}
 
 	OutCaptions.Add(FString::Printf(TEXT("'%s'"), *ParameterName.ToString())); 
 }
@@ -6379,6 +6410,14 @@ bool UMaterialExpressionVectorParameter::IsNamedParameter(const FMaterialParamet
 	return false;
 }
 
+void UMaterialExpressionVectorParameter::GetAllParameterInfo(TArray<FMaterialParameterInfo> &OutParameterInfo, TArray<FGuid> &OutParameterIds, const FMaterialParameterInfo& InBaseParameterInfo) const
+{
+	if (!bUseCustomPrimitiveData)
+	{
+		Super::GetAllParameterInfo(OutParameterInfo, OutParameterIds, InBaseParameterInfo);
+	}
+}
+
 #if WITH_EDITOR
 bool UMaterialExpressionVectorParameter::SetParameterValue(FName InParameterName, FLinearColor InValue)
 {
@@ -6400,6 +6439,12 @@ void UMaterialExpressionVectorParameter::PostEditChangeProperty(FPropertyChanged
 		// Callback into the editor
 		FEditorSupportDelegates::VectorParameterDefaultChanged.Broadcast(this, ParameterName, DefaultValue);
 	}
+	else if (PropertyName == GET_MEMBER_NAME_STRING_CHECKED(UMaterialExpressionVectorParameter, PrimitiveDataIndex))
+	{
+		// Clamp value
+		const int32 PrimDataIndex = PrimitiveDataIndex;
+		PrimitiveDataIndex = (uint8)FMath::Clamp(PrimDataIndex, 0, FCustomPrimitiveData::NumCustomPrimitiveDataFloats-1);
+	}
 
 	Super::PostEditChangeProperty(PropertyChangedEvent);
 }
@@ -6619,15 +6664,29 @@ UMaterialExpressionScalarParameter::UMaterialExpressionScalarParameter(const FOb
 #if WITH_EDITOR
 int32 UMaterialExpressionScalarParameter::Compile(class FMaterialCompiler* Compiler, int32 OutputIndex)
 {
-	return Compiler->ScalarParameter(ParameterName,DefaultValue);
+	if (bUseCustomPrimitiveData)
+	{
+		return Compiler->CustomPrimitiveData(PrimitiveDataIndex, MCT_Float);
+	}
+	else
+	{
+		return Compiler->ScalarParameter(ParameterName,DefaultValue);
+	}
 }
 
 void UMaterialExpressionScalarParameter::GetCaption(TArray<FString>& OutCaptions) const
 {
-	 OutCaptions.Add(FString::Printf(
-		 TEXT("Param (%.4g)"),
-		 DefaultValue )); 
-
+	if (bUseCustomPrimitiveData)
+	{
+		OutCaptions.Add(FString::Printf(TEXT("Index %d"), PrimitiveDataIndex)); 
+		OutCaptions.Add(FString::Printf(TEXT("Custom Primitive Data"))); 
+	}
+	else
+	{
+		OutCaptions.Add(FString::Printf(
+			 TEXT("Param (%.4g)"),
+			DefaultValue )); 
+	}
 	 OutCaptions.Add(FString::Printf(TEXT("'%s'"), *ParameterName.ToString())); 
 }
 #endif // WITH_EDITOR
@@ -6643,6 +6702,14 @@ bool UMaterialExpressionScalarParameter::IsNamedParameter(const FMaterialParamet
 	return false;
 }
 
+void UMaterialExpressionScalarParameter::GetAllParameterInfo(TArray<FMaterialParameterInfo> &OutParameterInfo, TArray<FGuid> &OutParameterIds, const FMaterialParameterInfo& InBaseParameterInfo) const
+{
+	if (!bUseCustomPrimitiveData)
+	{
+		Super::GetAllParameterInfo(OutParameterInfo, OutParameterIds, InBaseParameterInfo);
+	}
+}
+
 #if WITH_EDITOR
 bool UMaterialExpressionScalarParameter::SetParameterValue(FName InParameterName, float InValue)
 {
@@ -6664,7 +6731,12 @@ void UMaterialExpressionScalarParameter::PostEditChangeProperty(FPropertyChanged
 		// Callback into the editor
 		FEditorSupportDelegates::ScalarParameterDefaultChanged.Broadcast(this, ParameterName, DefaultValue);
 	}
-
+	else if (PropertyName == GET_MEMBER_NAME_STRING_CHECKED(UMaterialExpressionScalarParameter, PrimitiveDataIndex))
+	{
+		// Clamp value
+		const int32 PrimDataIndex = PrimitiveDataIndex;
+		PrimitiveDataIndex = (uint8)FMath::Clamp(PrimDataIndex, 0, FCustomPrimitiveData::NumCustomPrimitiveDataFloats-1);
+	}
 	Super::PostEditChangeProperty(PropertyChangedEvent);
 }
 
diff --git a/Engine/Source/Runtime/Engine/Public/MaterialCompiler.h b/Engine/Source/Runtime/Engine/Public/MaterialCompiler.h
--- a/Engine/Source/Runtime/Engine/Public/MaterialCompiler.h
+++ b/Engine/Source/Runtime/Engine/Public/MaterialCompiler.h
@@ -293,7 +293,7 @@ public:
 	virtual int32 EyeAdaptation() = 0;
 	virtual int32 AtmosphericLightVector() = 0;
 	virtual int32 AtmosphericLightColor() = 0;
-	virtual int32 CustomPrimitiveData(int32 OutputIndex) = 0;
+	virtual int32 CustomPrimitiveData(int32 OutputIndex, EMaterialValueType Type) = 0;
 	virtual int32 ShadingModel(EMaterialShadingModel InSelectedShadingModel) = 0;
 
 	// The compiler can run in a different state and this affects caching of sub expression, Expressions are different (e.g. View.PrevWorldViewOrigin) when using previous frame's values
@@ -552,9 +552,9 @@ public:
 		return Compiler->AtmosphericLightColor();
 	}
 	
-	virtual int32 CustomPrimitiveData(int32 OutputIndex) override
+	virtual int32 CustomPrimitiveData(int32 OutputIndex, EMaterialValueType Type) override
 	{
-		return Compiler->CustomPrimitiveData(OutputIndex);
+		return Compiler->CustomPrimitiveData(OutputIndex, Type);
 	}
 
 	virtual int32 ShadingModel(EMaterialShadingModel InSelectedShadingModel) override
diff --git a/Engine/Source/Runtime/Engine/Public/PrimitiveSceneProxy.h b/Engine/Source/Runtime/Engine/Public/PrimitiveSceneProxy.h
--- a/Engine/Source/Runtime/Engine/Public/PrimitiveSceneProxy.h
+++ b/Engine/Source/Runtime/Engine/Public/PrimitiveSceneProxy.h
@@ -729,7 +729,7 @@ public:
 	virtual uint8 GetCurrentFirstLODIdx_RenderThread() const { return 0; }
 
 	/** 
-	 * Get the custom primitive data for this scene proxy. Defaults to zero-data if not implemented.
+	 * Get the custom primitive data for this scene proxy.
 	 * @return The payload of custom data that will be set on the primitive and accessible in the material through a material expression.
 	 */
 	ENGINE_API virtual const FCustomPrimitiveData* GetCustomPrimitiveData() const { return nullptr; }
diff --git a/Engine/Source/Runtime/Engine/Public/SceneInterface.h b/Engine/Source/Runtime/Engine/Public/SceneInterface.h
--- a/Engine/Source/Runtime/Engine/Public/SceneInterface.h
+++ b/Engine/Source/Runtime/Engine/Public/SceneInterface.h
@@ -69,6 +69,12 @@ public:
 	virtual void UpdatePrimitiveTransform(UPrimitiveComponent* Primitive) = 0;
 	/** Updates primitive attachment state. */
 	virtual void UpdatePrimitiveAttachment(UPrimitiveComponent* Primitive) = 0;
+	/** 
+	 * Updates the custom primitive data of a static mesh which has already been added to the scene. 
+	 * 
+	 * @param StaticMesh - StaticMesh component to update
+	 */
+	virtual void UpdateCustomPrimitiveData(UStaticMeshComponent* StaticMesh) = 0;
 	/**
 	 * Updates distance field scene data (transforms, uv scale, self-shadow bias, etc.) but doesn't change allocation in the atlas
 	 */
diff --git a/Engine/Source/Runtime/RenderCore/Private/RenderCore.cpp b/Engine/Source/Runtime/RenderCore/Private/RenderCore.cpp
--- a/Engine/Source/Runtime/RenderCore/Private/RenderCore.cpp
+++ b/Engine/Source/Runtime/RenderCore/Private/RenderCore.cpp
@@ -177,6 +177,7 @@ DEFINE_STAT(STAT_UpdatePrimitiveTransformRenderThreadTime);
 DEFINE_STAT(STAT_RemoveScenePrimitiveGT);
 DEFINE_STAT(STAT_AddScenePrimitiveGT);
 DEFINE_STAT(STAT_UpdatePrimitiveTransformGT);
+DEFINE_STAT(STAT_UpdateCustomPrimitiveDataGT);
 
 DEFINE_STAT(STAT_Scene_SetShaderMapsOnMaterialResources_RT);
 DEFINE_STAT(STAT_Scene_UpdateStaticDrawLists_RT);
diff --git a/Engine/Source/Runtime/RenderCore/Public/RenderCore.h b/Engine/Source/Runtime/RenderCore/Public/RenderCore.h
--- a/Engine/Source/Runtime/RenderCore/Public/RenderCore.h
+++ b/Engine/Source/Runtime/RenderCore/Public/RenderCore.h
@@ -128,6 +128,7 @@ DECLARE_CYCLE_STAT_EXTERN(TEXT("Update Influence Weights"),STAT_InfluenceWeights
 DECLARE_CYCLE_STAT_EXTERN(TEXT("RemovePrimitive (GT)"),STAT_RemoveScenePrimitiveGT,STATGROUP_SceneUpdate, RENDERCORE_API);
 DECLARE_CYCLE_STAT_EXTERN(TEXT("AddPrimitive (GT)"),STAT_AddScenePrimitiveGT,STATGROUP_SceneUpdate, RENDERCORE_API);
 DECLARE_CYCLE_STAT_EXTERN(TEXT("UpdatePrimitiveTransform (GT)"),STAT_UpdatePrimitiveTransformGT,STATGROUP_SceneUpdate, RENDERCORE_API);
+DECLARE_CYCLE_STAT_EXTERN(TEXT("UpdateCustomPrimitiveData (GT)"),STAT_UpdateCustomPrimitiveDataGT,STATGROUP_SceneUpdate, RENDERCORE_API);
 
 DECLARE_CYCLE_STAT_EXTERN(TEXT("Set Shader Maps On Material Resources (RT)"),STAT_Scene_SetShaderMapsOnMaterialResources_RT,STATGROUP_SceneUpdate, RENDERCORE_API);
 DECLARE_CYCLE_STAT_EXTERN(TEXT("Update Static Draw Lists (RT)"), STAT_Scene_UpdateStaticDrawLists_RT, STATGROUP_SceneUpdate, RENDERCORE_API);
diff --git a/Engine/Source/Runtime/Renderer/Private/RendererScene.cpp b/Engine/Source/Runtime/Renderer/Private/RendererScene.cpp
--- a/Engine/Source/Runtime/Renderer/Private/RendererScene.cpp
+++ b/Engine/Source/Runtime/Renderer/Private/RendererScene.cpp
@@ -1411,6 +1411,45 @@ void FScene::UpdatePrimitiveAttachment(UPrimitiveComponent* Primitive)
 	}
 }
 
+void FScene::UpdateCustomPrimitiveData(UStaticMeshComponent* StaticMesh)
+{
+	SCOPE_CYCLE_COUNTER(STAT_UpdateCustomPrimitiveDataGT);
+
+	// This path updates the primitive data directly in the GPUScene. 
+	if(StaticMesh->SceneProxy) 
+	{
+		struct FUpdateParams
+		{
+			FScene* Scene;
+			FStaticMeshSceneProxy* PrimitiveSceneProxy;
+			FCustomPrimitiveData CustomPrimitiveData;
+		};
+
+		FUpdateParams UpdateParams;
+		UpdateParams.Scene = this;
+		UpdateParams.PrimitiveSceneProxy = static_cast<FStaticMeshSceneProxy*>(StaticMesh->SceneProxy);
+		UpdateParams.CustomPrimitiveData = StaticMesh->CustomPrimitiveData; 
+
+		ENQUEUE_RENDER_COMMAND(UpdateCustomPrimitiveDataCommand)(
+			[UpdateParams](FRHICommandListImmediate& RHICmdList)
+			{
+				FScopeCycleCounter Context(UpdateParams.PrimitiveSceneProxy->GetStatId());
+				UpdateParams.PrimitiveSceneProxy->CustomPrimitiveData = UpdateParams.CustomPrimitiveData;
+
+				// No need to do any of this if GPUScene isn't used (the custom primitive data will make it to the primitive uniform buffer through FPrimitiveSceneProxy::UpdateUniformBuffer if that's the case)
+				if (UseGPUScene(GMaxRHIShaderPlatform, UpdateParams.Scene->GetFeatureLevel()))
+				{
+					AddPrimitiveToUpdateGPU(*UpdateParams.Scene, UpdateParams.PrimitiveSceneProxy->GetPrimitiveSceneInfo()->PackedIndex);
+				}
+				else
+				{
+					// Make sure the uniform buffer is updated before rendering
+					UpdateParams.PrimitiveSceneProxy->GetPrimitiveSceneInfo()->SetNeedsUniformBufferUpdate(true);
+				}
+			});
+	}
+}
+
 void FScene::UpdatePrimitiveDistanceFieldSceneData_GameThread(UPrimitiveComponent* Primitive)
 {
 	check(IsInGameThread());
@@ -3415,6 +3454,7 @@ public:
 	/** Updates the transform of a primitive which has already been added to the scene. */
 	virtual void UpdatePrimitiveTransform(UPrimitiveComponent* Primitive) override {}
 	virtual void UpdatePrimitiveAttachment(UPrimitiveComponent* Primitive) override {}
+	virtual void UpdateCustomPrimitiveData(UStaticMeshComponent* Primitive) override {}
 
 	virtual void AddLight(ULightComponent* Light) override {}
 	virtual void RemoveLight(ULightComponent* Light) override {}
diff --git a/Engine/Source/Runtime/Renderer/Private/ScenePrivate.h b/Engine/Source/Runtime/Renderer/Private/ScenePrivate.h
--- a/Engine/Source/Runtime/Renderer/Private/ScenePrivate.h
+++ b/Engine/Source/Runtime/Renderer/Private/ScenePrivate.h
@@ -2681,6 +2681,7 @@ public:
 	virtual void ReleasePrimitive(UPrimitiveComponent* Primitive) override;
 	virtual void UpdatePrimitiveTransform(UPrimitiveComponent* Primitive) override;
 	virtual void UpdatePrimitiveAttachment(UPrimitiveComponent* Primitive) override;
+	virtual void UpdateCustomPrimitiveData(UStaticMeshComponent* StaticMesh) override;
 	virtual void UpdatePrimitiveDistanceFieldSceneData_GameThread(UPrimitiveComponent* Primitive) override;
 	virtual FPrimitiveSceneInfo* GetPrimitiveSceneInfo(int32 PrimitiveIndex) override;
 	virtual bool GetPreviousLocalToWorld(const FPrimitiveSceneInfo* PrimitiveSceneInfo, FMatrix& OutPreviousLocalToWorld) const override;
