Add == and != UObject* operators to FPawnIteratorObject to provide additional better backwards compatibility
#rnx

diff --git a/Engine/Source/Runtime/Engine/Classes/Engine/World.h b/Engine/Source/Runtime/Engine/Classes/Engine/World.h
--- a/Engine/Source/Runtime/Engine/Classes/Engine/World.h
+++ b/Engine/Source/Runtime/Engine/Classes/Engine/World.h
@@ -77,6 +77,9 @@ struct ENGINE_API FPawnIteratorObject
 	APawn& operator*() const { return *Pawn; }
 	APawn* Get() const { return Pawn; }
 
+	bool operator==(const UObject* Other) const { return Pawn == Other; }
+	bool operator!=(const UObject* Other) const { return Pawn != Other; }
+
 private:
 	FPawnIteratorObject()
 		: Pawn(nullptr)
