Added gbuffer visualization overview and view mode cycling to DebugCameraController.

diff --git a/Engine/Shaders/Private/PostProcessVisualizeBuffer.usf b/Engine/Shaders/Private/PostProcessVisualizeBuffer.usf
--- a/Engine/Shaders/Private/PostProcessVisualizeBuffer.usf
+++ b/Engine/Shaders/Private/PostProcessVisualizeBuffer.usf
@@ -8,6 +8,8 @@
 #include "PostProcessCommon.ush"
 #include "DeferredShadingCommon.ush"
 
+float4 SelectionColor;
+
 /** Main pixel shader function */
 void MainPS(noperspective float4 UVAndScreenPos : TEXCOORD0, float4 SvPosition : SV_POSITION, out float4 OutColor : SV_Target0)
 {
@@ -29,8 +31,19 @@ void MainPS(noperspective float4 UVAndScreenPos : TEXCOORD0, float4 SvPosition :
 	const int Border1 = 10;
 	const int Border2 = 8;
 	float BorderDistance = ComputeDistanceToRect(QuarterTileLocalPos, int2(0, 0) + Border1.xx, QuarterInsetSize - 2 * Border1.xx);
-	Alpha = 1 - saturate((BorderDistance - Border2.x) * 0.5f);
+	Alpha = (1 - saturate((BorderDistance - Border2.x) * 0.5f));
+
+
 #endif
 
 	OutColor = float4(Texture2DSample(PostprocessInput0, PostprocessInput0Sampler, UVAndScreenPos.xy).rgb, Alpha);
+
+#if DRAWING_TILE
+
+	// border around selected tile
+	OutColor.xyz = lerp(OutColor.xyz, SelectionColor.xyz, saturate(BorderDistance - Border2.x));
+	OutColor.w = lerp(OutColor.w, SelectionColor.w, saturate(BorderDistance - Border2.x));
+	OutColor.w = lerp(OutColor.w, 0.0f, saturate(BorderDistance - Border1.x));
+
+#endif
 }
\ No newline at end of file
diff --git a/Engine/Source/Editor/DetailCustomizations/Private/DebugCameraControllerSettingsCustomization.cpp b/Engine/Source/Editor/DetailCustomizations/Private/DebugCameraControllerSettingsCustomization.cpp
new file mode 100644
--- /dev/null
+++ b/Engine/Source/Editor/DetailCustomizations/Private/DebugCameraControllerSettingsCustomization.cpp
@@ -0,0 +1,67 @@
+// Copyright 1998-2019 Epic Games, Inc. All Rights Reserved.
+
+#include "DebugCameraControllerSettingsCustomization.h"
+#include "Engine/DebugCameraControllerSettings.h"
+#include "Widgets/Text/STextBlock.h"
+#include "PropertyRestriction.h"
+#include "IDetailChildrenBuilder.h"
+#include "DetailLayoutBuilder.h"
+#include "DetailCategoryBuilder.h"
+#include "DetailWidgetRow.h"
+#include "IDetailPropertyRow.h"
+#include "ShowFlags.h"
+
+
+TSharedRef<IPropertyTypeCustomization> FDebugCameraControllerSettingsViewModeIndexCustomization::MakeInstance()
+{
+	return MakeShareable( new FDebugCameraControllerSettingsViewModeIndexCustomization);
+}
+
+FDebugCameraControllerSettingsViewModeIndexCustomization::FDebugCameraControllerSettingsViewModeIndexCustomization()
+{
+}
+
+void FDebugCameraControllerSettingsViewModeIndexCustomization::CustomizeHeader(TSharedRef<class IPropertyHandle> StructPropertyHandle, class FDetailWidgetRow& HeaderRow, IPropertyTypeCustomizationUtils& StructCustomizationUtils)
+{
+	// We'll set up reset to default ourselves
+	const bool bDisplayResetToDefault = false;
+	const FText DisplayNameOverride = FText::GetEmpty();
+	const FText DisplayToolTipOverride = FText::GetEmpty();
+
+	uint32 NumChildren;
+	StructPropertyHandle->GetNumChildren(NumChildren);
+	TSharedPtr<IPropertyHandle> ViewModeIndexHandle;
+
+	for (uint32 ChildIndex = 0; ChildIndex < NumChildren; ++ChildIndex)
+	{
+		const TSharedRef< IPropertyHandle > ChildHandle = StructPropertyHandle->GetChildHandle(ChildIndex).ToSharedRef();
+
+		if (ChildHandle->GetProperty()->GetName() == TEXT("ViewModeIndex"))
+		{
+			ViewModeIndexHandle = ChildHandle;
+		}
+	}
+
+	check(ViewModeIndexHandle.IsValid());
+
+	TSharedPtr<FPropertyRestriction> EnumRestriction = MakeShareable(new FPropertyRestriction(NSLOCTEXT("DebugCycleViewModes", "DebugCycleViewModes", "Cycle view modes for debug camera controller")));
+	const UEnum* const ViewModeIndexEnum = StaticEnum<EViewModeIndex>();
+	EnumRestriction->AddHiddenValue(ViewModeIndexEnum->GetNameStringByValue((uint8)EViewModeIndex::VMI_VisualizeBuffer));
+	ViewModeIndexHandle->AddRestriction(EnumRestriction.ToSharedRef());
+
+	HeaderRow
+		.NameContent()
+		[
+			StructPropertyHandle->CreatePropertyNameWidget(DisplayNameOverride, DisplayToolTipOverride, bDisplayResetToDefault)
+		]
+		.ValueContent()
+		.MinDesiredWidth(500)
+		[
+			ViewModeIndexHandle->CreatePropertyValueWidget()
+		];
+}
+
+void FDebugCameraControllerSettingsViewModeIndexCustomization::CustomizeChildren(TSharedRef<IPropertyHandle> StructPropertyHandle, class IDetailChildrenBuilder& ChildBuilder, IPropertyTypeCustomizationUtils& StructCustomizationUtils)
+{
+}
+
diff --git a/Engine/Source/Editor/DetailCustomizations/Private/DebugCameraControllerSettingsCustomization.h b/Engine/Source/Editor/DetailCustomizations/Private/DebugCameraControllerSettingsCustomization.h
new file mode 100644
--- /dev/null
+++ b/Engine/Source/Editor/DetailCustomizations/Private/DebugCameraControllerSettingsCustomization.h
@@ -0,0 +1,25 @@
+// Copyright 1998-2019 Epic Games, Inc. All Rights Reserved.
+
+#pragma once
+
+#include "CoreMinimal.h"
+#include "Misc/Attribute.h"
+#include "Layout/Visibility.h"
+#include "IPropertyTypeCustomization.h"
+#include "PropertyHandle.h"
+
+
+class FDebugCameraControllerSettingsViewModeIndexCustomization : public IPropertyTypeCustomization
+{
+public:
+	FDebugCameraControllerSettingsViewModeIndexCustomization();
+
+	static TSharedRef<IPropertyTypeCustomization> MakeInstance();
+
+	//~ Begin IPropertyTypeCustomization
+	virtual void CustomizeHeader( TSharedRef<class IPropertyHandle> StructPropertyHandle, class FDetailWidgetRow& HeaderRow, IPropertyTypeCustomizationUtils& StructCustomizationUtils ) override;
+	virtual void CustomizeChildren( TSharedRef<IPropertyHandle> StructPropertyHandle, class IDetailChildrenBuilder& ChildBuilder, IPropertyTypeCustomizationUtils& StructCustomizationUtils ) override;
+	//~ End IPropertyTypeCustomization
+};
+
+
diff --git a/Engine/Source/Editor/DetailCustomizations/Private/DetailCustomizations.cpp b/Engine/Source/Editor/DetailCustomizations/Private/DetailCustomizations.cpp
--- a/Engine/Source/Editor/DetailCustomizations/Private/DetailCustomizations.cpp
+++ b/Engine/Source/Editor/DetailCustomizations/Private/DetailCustomizations.cpp
@@ -137,6 +137,7 @@
 #include "TimecodeDetailsCustomization.h"
 #include "SkeletonDetails.h"
 #include "MaterialShadingModelCustomization.h"
+#include "DebugCameraControllerSettingsCustomization.h"
 IMPLEMENT_MODULE( FDetailCustomizationsModule, DetailCustomizations );
 
 
@@ -270,6 +271,7 @@ void FDetailCustomizationsModule::RegisterPropertyTypeCustomizations()
 	RegisterCustomPropertyTypeLayout("SkeletalMeshOptimizationSettings", FOnGetPropertyTypeCustomizationInstance::CreateStatic(&FSkeletalMeshReductionSettingsDetails::MakeInstance));
 	RegisterCustomPropertyTypeLayout("GrassInput", FOnGetPropertyTypeCustomizationInstance::CreateStatic(&FMaterialExpressionLandscapeGrassInputCustomization::MakeInstance));
 	RegisterCustomPropertyTypeLayout("EMaterialShadingModel", FOnGetPropertyTypeCustomizationInstance::CreateStatic(&FMaterialShadingModelCustomization::MakeInstance));
+	RegisterCustomPropertyTypeLayout("DebugCameraControllerSettingsViewModeIndex", FOnGetPropertyTypeCustomizationInstance::CreateStatic(&FDebugCameraControllerSettingsViewModeIndexCustomization::MakeInstance));
 }
 
 void FDetailCustomizationsModule::RegisterObjectCustomizations()
diff --git a/Engine/Source/Runtime/Engine/Classes/Engine/DebugCameraController.h b/Engine/Source/Runtime/Engine/Classes/Engine/DebugCameraController.h
--- a/Engine/Source/Runtime/Engine/Classes/Engine/DebugCameraController.h
+++ b/Engine/Source/Runtime/Engine/Classes/Engine/DebugCameraController.h
@@ -40,6 +40,14 @@ class ENGINE_API ADebugCameraController
 	UPROPERTY()
 	uint32 bOrbitPivotUseCenter:1;
 
+	/** Whether set view mode to display GBuffer visualization */
+	UPROPERTY()
+	uint32 bEnableBufferVisualization : 1;
+
+	/** Whether set view mode to display GBuffer visualization */
+	UPROPERTY()
+	uint32 bEnableBufferVisualizationFullMode : 1;
+
 	/** Visualizes the frustum of the camera */
 	UPROPERTY()
 	class UDrawFrustumComponent* DrawFrustum;
@@ -103,6 +111,37 @@ class ENGINE_API ADebugCameraController
 	/** Toggles camera orbit hitpoint */
 	void ToggleOrbitHitPoint();
 
+	/** Cycle view mode */
+	void CycleViewMode();
+
+	/** Toggle buffer visualization overview */
+	void ToggleBufferVisualizationOverviewMode();
+
+	/** Buffer overview move up */
+	void BufferVisualizationMoveUp();
+
+	/** Buffer overview move down */
+	void BufferVisualizationMoveDown();
+
+	/** Buffer overview move right */
+	void BufferVisualizationMoveRight();
+
+	/** Buffer overview move left */
+	void BufferVisualizationMoveLeft();
+
+	/** Toggle buffer visualization full display */
+	void ToggleBufferVisualizationFullMode();
+
+	/** Set buffer visualization full mode */
+	void SetBufferVisualizationFullMode(bool bFullMode);
+
+#if !(UE_BUILD_SHIPPING || UE_BUILD_TEST)
+
+	/** Update visualize buffer post processing settings */
+	void UpdateVisualizeBufferPostProcessing(FFinalPostProcessSettings& InOutPostProcessingSettings);
+
+#endif
+
 public:
 
 	/** Currently selected actor, may be invalid */
@@ -217,6 +256,15 @@ protected:
 
 	virtual void SetSpectatorPawn(class ASpectatorPawn* NewSpectatorPawn) override;
 
+	/** Get buffer visualization overview targets based on console var */
+	TArray<FString> GetBufferVisualizationOverviewTargets();
+
+	/** Get next buffer */
+	void GetNextBuffer(const TArray<FString>& OverviewBuffers, int32 Step = 1);
+
+	/** Get next buffer */
+	void GetNextBuffer(int32 Step = 1);
+
 private:
 
 	/** The normalized screen location when a drag starts */
@@ -231,6 +279,15 @@ private:
 	/** Current orbit radius, if orbit is enabled */
 	float OrbitRadius;
 
+	/** Last view mode index before buffer visualization overview was enabled, see enum EViewModeIndex for valid values */
+	int32 LastViewModeIndex;
+
+	/** Last index in settings array for cycle view modes */
+	int32 LastViewModeSettingsIndex;
+	
+	/** Last buffer selected in buffer visualization overview */
+	FString LastSelectedBuffer;
+
 	void OnTouchBegin(ETouchIndex::Type FingerIndex, FVector Location);
 	void OnTouchEnd(ETouchIndex::Type FingerIndex, FVector Location);
 	void OnFingerMove(ETouchIndex::Type FingerIndex, FVector Location);
diff --git a/Engine/Source/Runtime/Engine/Classes/Engine/DebugCameraControllerSettings.h b/Engine/Source/Runtime/Engine/Classes/Engine/DebugCameraControllerSettings.h
new file mode 100644
--- /dev/null
+++ b/Engine/Source/Runtime/Engine/Classes/Engine/DebugCameraControllerSettings.h
@@ -0,0 +1,53 @@
+// Copyright 1998-2019 Epic Games, Inc. All Rights Reserved.
+
+/*=============================================================================
+	DebugCameraControllerSettings.h: Declares the DebugCameraControllerSettings class.
+=============================================================================*/
+
+#pragma once
+
+#include "Engine/DeveloperSettings.h"
+#include "DebugCameraControllerSettings.generated.h"
+
+USTRUCT()
+struct ENGINE_API FDebugCameraControllerSettingsViewModeIndex
+{
+	GENERATED_USTRUCT_BODY()
+
+	FDebugCameraControllerSettingsViewModeIndex() { ViewModeIndex = VMI_BrushWireframe; }
+	FDebugCameraControllerSettingsViewModeIndex(EViewModeIndex InViewModeIndex) : ViewModeIndex(InViewModeIndex) {}
+
+	virtual ~FDebugCameraControllerSettingsViewModeIndex() { }
+
+	UPROPERTY(EditAnywhere, Category = General, meta = (DisplayName = "Cycle View Mode"))
+	TEnumAsByte<EViewModeIndex> ViewModeIndex;
+};
+
+/**
+ * Default debug camera controller settings.
+ */
+UCLASS(config=Engine, defaultconfig, meta=(DisplayName="Debug Camera Controller"))
+class ENGINE_API UDebugCameraControllerSettings : public UDeveloperSettings
+{
+	GENERATED_UCLASS_BODY()
+
+	UPROPERTY(config, EditAnywhere, Category = General)
+	TArray<FDebugCameraControllerSettingsViewModeIndex> CycleViewModes;
+
+public:
+#if WITH_EDITOR
+
+	// UObject interface.
+	virtual void PostLoad() override;
+	virtual void PostEditChangeProperty(FPropertyChangedEvent& PropertyChangedEvent) override;
+
+protected:
+	void RemoveInvalidViewModes();
+
+#endif // WITH_EDITOR
+
+public:
+	TArray<EViewModeIndex> GetCycleViewModes();
+
+	static UDebugCameraControllerSettings * Get() { return CastChecked<UDebugCameraControllerSettings>(UDebugCameraControllerSettings::StaticClass()->GetDefaultObject()); }
+};
diff --git a/Engine/Source/Runtime/Engine/Classes/Engine/EngineBaseTypes.h b/Engine/Source/Runtime/Engine/Classes/Engine/EngineBaseTypes.h
--- a/Engine/Source/Runtime/Engine/Classes/Engine/EngineBaseTypes.h
+++ b/Engine/Source/Runtime/Engine/Classes/Engine/EngineBaseTypes.h
@@ -896,59 +896,59 @@ UENUM()
 enum EViewModeIndex 
 {
 	/** Wireframe w/ brushes. */
-	VMI_BrushWireframe = 0,
+	VMI_BrushWireframe = 0 UMETA(DisplayName = "Brush Wireframe"),
 	/** Wireframe w/ BSP. */
-	VMI_Wireframe = 1,
+	VMI_Wireframe = 1 UMETA(DisplayName = "Wireframe"),
 	/** Unlit. */
-	VMI_Unlit = 2,
+	VMI_Unlit = 2 UMETA(DisplayName = "Unlit"),
 	/** Lit. */
-	VMI_Lit = 3,
-	VMI_Lit_DetailLighting = 4,
+	VMI_Lit = 3 UMETA(DisplayName = "Lit"),
+	VMI_Lit_DetailLighting = 4 UMETA(DisplayName = "Detail Lighting"),
 	/** Lit wo/ materials. */
-	VMI_LightingOnly = 5,
+	VMI_LightingOnly = 5 UMETA(DisplayName = "Lighting Only"),
 	/** Colored according to light count. */
-	VMI_LightComplexity = 6,
+	VMI_LightComplexity = 6 UMETA(DisplayName = "Light Complexity"),
 	/** Colored according to shader complexity. */
-	VMI_ShaderComplexity = 8,
+	VMI_ShaderComplexity = 8 UMETA(DisplayName = "Shader Complexity"),
 	/** Colored according to world-space LightMap texture density. */
-	VMI_LightmapDensity = 9,
+	VMI_LightmapDensity = 9 UMETA(DisplayName = "Lightmap Density"),
 	/** Colored according to light count - showing lightmap texel density on texture mapped objects. */
-	VMI_LitLightmapDensity = 10,
-	VMI_ReflectionOverride = 11,
-	VMI_VisualizeBuffer = 12,
+	VMI_LitLightmapDensity = 10 UMETA(DisplayName = "Lit Lightmap Density"),
+	VMI_ReflectionOverride = 11 UMETA(DisplayName = "Reflections"),
+	VMI_VisualizeBuffer = 12 UMETA(DisplayName = "Buffer Visualization"),
 	//	VMI_VoxelLighting = 13,
 
 	/** Colored according to stationary light overlap. */
-	VMI_StationaryLightOverlap = 14,
+	VMI_StationaryLightOverlap = 14 UMETA(DisplayName = "Stationary Light Overlap"),
 
-	VMI_CollisionPawn = 15, 
-	VMI_CollisionVisibility = 16, 
+	VMI_CollisionPawn = 15 UMETA(DisplayName = "Player Collision"),
+	VMI_CollisionVisibility = 16 UMETA(DisplayName = "Visibility Collision"),
 	//VMI_UNUSED = 17,
 	/** Colored according to the current LOD index. */
-	VMI_LODColoration = 18,
+	VMI_LODColoration = 18 UMETA(DisplayName = "Mesh LOD Coloration"),
 	/** Colored according to the quad coverage. */
-	VMI_QuadOverdraw = 19,
+	VMI_QuadOverdraw = 19 UMETA(DisplayName = "Quad Overdraw"),
 	/** Visualize the accuracy of the primitive distance computed for texture streaming. */
-	VMI_PrimitiveDistanceAccuracy = 20,
+	VMI_PrimitiveDistanceAccuracy = 20 UMETA(DisplayName = "Primitive Distance"),
 	/** Visualize the accuracy of the mesh UV densities computed for texture streaming. */
-	VMI_MeshUVDensityAccuracy = 21,
+	VMI_MeshUVDensityAccuracy = 21  UMETA(DisplayName = "Mesh UV Density"),
 	/** Colored according to shader complexity, including quad overdraw. */
-	VMI_ShaderComplexityWithQuadOverdraw = 22,
+	VMI_ShaderComplexityWithQuadOverdraw = 22 UMETA(DisplayName = "Shader Complexity & Quads"),
 	/** Colored according to the current HLOD index. */
-	VMI_HLODColoration = 23,
+	VMI_HLODColoration = 23  UMETA(DisplayName = "Hierarchical LOD Coloration"),
 	/** Group item for LOD and HLOD coloration*/
-	VMI_GroupLODColoration = 24,
+	VMI_GroupLODColoration = 24  UMETA(DisplayName = "Group LOD Coloration"),
 	/** Visualize the accuracy of the material texture scales used for texture streaming. */
-	VMI_MaterialTextureScaleAccuracy = 25,
+	VMI_MaterialTextureScaleAccuracy = 25 UMETA(DisplayName = "Material Texture Scales"),
 	/** Compare the required texture resolution to the actual resolution. */
-	VMI_RequiredTextureResolution = 26,
+	VMI_RequiredTextureResolution = 26 UMETA(DisplayName = "Required Texture Resolution"),
 
 	// Ray tracing modes
 
 	/** Run path tracing pipeline */
-	VMI_PathTracing = 27,
+	VMI_PathTracing = 27 UMETA(DisplayName = "Path Tracing"),
 	/** Run ray tracing debug pipeline */
-	VMI_RayTracingDebug = 28,
+	VMI_RayTracingDebug = 28 UMETA(DisplayName = "Ray Tracing Debug"),
 
 	VMI_Max UMETA(Hidden),
 
diff --git a/Engine/Source/Runtime/Engine/Private/DebugCameraController.cpp b/Engine/Source/Runtime/Engine/Private/DebugCameraController.cpp
--- a/Engine/Source/Runtime/Engine/Private/DebugCameraController.cpp
+++ b/Engine/Source/Runtime/Engine/Private/DebugCameraController.cpp
@@ -6,6 +6,7 @@
 =============================================================================*/
 
 #include "Engine/DebugCameraController.h"
+#include "Engine/DebugCameraControllerSettings.h"
 #include "EngineGlobals.h"
 #include "CollisionQueryParams.h"
 #include "Engine/World.h"
@@ -21,6 +22,7 @@
 #include "Components/DrawFrustumComponent.h"
 #include "GameFramework/PlayerInput.h"
 #include "GameFramework/GameStateBase.h"
+#include "BufferVisualizationData.h"
 
 static const float SPEED_SCALE_ADJUSTMENT = 0.5f;
 static const float MIN_ORBIT_RADIUS = 30.0f;
@@ -53,6 +55,11 @@ ADebugCameraController::ADebugCameraController(const FObjectInitializer& ObjectI
 	LastOrbitPawnLocation = FVector::ZeroVector;
 	OrbitPivot = FVector::ZeroVector;
 	OrbitRadius = MIN_ORBIT_RADIUS;
+
+	bEnableBufferVisualization = false;
+	bEnableBufferVisualizationFullMode = false;
+	LastViewModeIndex = VMI_Lit;
+	LastViewModeSettingsIndex = 0;
 }
 
 void InitializeDebugCameraInputBindings()
@@ -74,6 +81,13 @@ void InitializeDebugCameraInputBindings()
 		UPlayerInput::AddEngineDefinedActionMapping(FInputActionKeyMapping("DebugCamera_FreezeRendering", EKeys::F));
 		UPlayerInput::AddEngineDefinedActionMapping(FInputActionKeyMapping("DebugCamera_OrbitHitPoint", EKeys::O));
 		UPlayerInput::AddEngineDefinedActionMapping(FInputActionKeyMapping("DebugCamera_OrbitCenter", EKeys::O, true));
+		UPlayerInput::AddEngineDefinedActionMapping(FInputActionKeyMapping("DebugCamera_CycleViewMode", EKeys::V));
+		UPlayerInput::AddEngineDefinedActionMapping(FInputActionKeyMapping("DebugCamera_ToggleBufferVisualizationOverview", EKeys::B));
+		UPlayerInput::AddEngineDefinedActionMapping(FInputActionKeyMapping("DebugCamera_BufferVisualizationUp", EKeys::Up));
+		UPlayerInput::AddEngineDefinedActionMapping(FInputActionKeyMapping("DebugCamera_BufferVisualizationDown", EKeys::Down));
+		UPlayerInput::AddEngineDefinedActionMapping(FInputActionKeyMapping("DebugCamera_BufferVisualizationRight", EKeys::Right));
+		UPlayerInput::AddEngineDefinedActionMapping(FInputActionKeyMapping("DebugCamera_BufferVisualizationLeft", EKeys::Left));
+		UPlayerInput::AddEngineDefinedActionMapping(FInputActionKeyMapping("DebugCamera_ToggleBufferVisualizationFull", EKeys::Enter));
 
 		UPlayerInput::AddEngineDefinedActionMapping(FInputActionKeyMapping("DebugCamera_Select", EKeys::Gamepad_RightTrigger));
 		UPlayerInput::AddEngineDefinedActionMapping(FInputActionKeyMapping("DebugCamera_IncreaseSpeed", EKeys::Gamepad_RightShoulder));
@@ -103,6 +117,13 @@ void ADebugCameraController::SetupInputComponent()
 	InputComponent->BindAction("DebugCamera_FreezeRendering", IE_Pressed, this, &ADebugCameraController::ToggleFreezeRendering);
 	InputComponent->BindAction("DebugCamera_OrbitHitPoint", IE_Pressed, this, &ADebugCameraController::ToggleOrbitHitPoint);
 	InputComponent->BindAction("DebugCamera_OrbitCenter", IE_Pressed, this, &ADebugCameraController::ToggleOrbitCenter);
+	InputComponent->BindAction("DebugCamera_CycleViewMode", IE_Pressed, this, &ADebugCameraController::CycleViewMode);
+	InputComponent->BindAction("DebugCamera_ToggleBufferVisualizationOverview", IE_Pressed, this, &ADebugCameraController::ToggleBufferVisualizationOverviewMode);
+	InputComponent->BindAction("DebugCamera_BufferVisualizationUp", IE_Pressed, this, &ADebugCameraController::BufferVisualizationMoveUp);
+	InputComponent->BindAction("DebugCamera_BufferVisualizationDown", IE_Pressed, this, &ADebugCameraController::BufferVisualizationMoveDown);
+	InputComponent->BindAction("DebugCamera_BufferVisualizationRight", IE_Pressed, this, &ADebugCameraController::BufferVisualizationMoveRight);
+	InputComponent->BindAction("DebugCamera_BufferVisualizationLeft", IE_Pressed, this, &ADebugCameraController::BufferVisualizationMoveLeft);
+	InputComponent->BindAction("DebugCamera_ToggleBufferVisualizationFull", IE_Pressed, this, &ADebugCameraController::ToggleBufferVisualizationFullMode);
 
 	InputComponent->BindTouch(IE_Pressed, this, &ADebugCameraController::OnTouchBegin);
 	InputComponent->BindTouch(IE_Released, this, &ADebugCameraController::OnTouchEnd);
@@ -592,6 +613,247 @@ void ADebugCameraController::ToggleOrbitHitPoint()
 	ToggleOrbit(false);
 }
 
+void ADebugCameraController::CycleViewMode()
+{
+	if (bEnableBufferVisualization)
+	{
+		ToggleBufferVisualizationOverviewMode();
+	}
+
+	if (UGameViewportClient* GameViewportClient = GetWorld()->GetGameViewport())
+	{
+		TArray<EViewModeIndex> DebugViewModes = UDebugCameraControllerSettings::Get()->GetCycleViewModes();
+
+		if (DebugViewModes.Num() == 0)
+		{
+			UE_LOG(LogController, Warning, TEXT("Debug camera controller settings must specify at least one view mode for view mode cycling."));
+			return;
+		}
+
+		int32 CurrViewModeIndex = GameViewportClient->ViewModeIndex;
+		int32 CurrIndex = LastViewModeSettingsIndex;
+
+		int32 NextIndex = (CurrIndex < DebugViewModes.Num() && CurrViewModeIndex == DebugViewModes[CurrIndex]) ? (CurrIndex + 1) % DebugViewModes.Num() : 0;
+		int32 NextViewModeIndex = DebugViewModes[NextIndex];
+
+		if (NextViewModeIndex != CurrViewModeIndex)
+		{
+			FString NextViewModeName = GetViewModeName((EViewModeIndex)NextViewModeIndex);
+
+			if (!NextViewModeName.IsEmpty())
+			{
+				FString Cmd(TEXT("VIEWMODE "));
+				Cmd += NextViewModeName;
+				GameViewportClient->ConsoleCommand(Cmd);
+
+			}
+		}
+
+		LastViewModeSettingsIndex = NextIndex;
+		LastViewModeIndex = NextViewModeIndex;
+	}
+}
+
+TArray<FString> ADebugCameraController::GetBufferVisualizationOverviewTargets()
+{
+	TArray<FString> SelectedBuffers;
+
+	// Get the list of requested buffers from the console
+	static IConsoleVariable* CVar = IConsoleManager::Get().FindConsoleVariable(TEXT("r.BufferVisualizationOverviewTargets"));
+	FString SelectedBufferNames = CVar->GetString();
+
+	// Extract each material name from the comma separated string
+	while (SelectedBufferNames.Len())
+	{
+		FString Left, Right;
+
+		// Detect last entry in the list
+		if (!SelectedBufferNames.Split(TEXT(","), &Left, &Right))
+		{
+			Left = SelectedBufferNames;
+			Right = FString();
+		}
+
+		Left.TrimStartInline();
+		SelectedBuffers.Add(*Left);
+		SelectedBufferNames = Right;
+	}
+
+	return SelectedBuffers;
+}
+
+void ADebugCameraController::ToggleBufferVisualizationOverviewMode()
+{
+	SetBufferVisualizationFullMode(false);
+
+	if (UGameViewportClient* GameViewportClient = GetWorld()->GetGameViewport())
+	{
+		bEnableBufferVisualization = !bEnableBufferVisualization;
+		bEnableBufferVisualizationFullMode = false;
+
+		FString Cmd(TEXT("VIEWMODE "));
+
+		if (bEnableBufferVisualization)
+		{
+			Cmd += GetViewModeName(VMI_VisualizeBuffer);
+			LastViewModeIndex = GameViewportClient->ViewModeIndex;
+
+			TArray<FString> SelectedBuffers = GetBufferVisualizationOverviewTargets();
+
+			if (LastSelectedBuffer.IsEmpty() || !SelectedBuffers.Contains(LastSelectedBuffer))
+			{
+				GetNextBuffer(SelectedBuffers, 1);
+			}
+		}
+		else
+		{
+			Cmd += GetViewModeName((EViewModeIndex)LastViewModeIndex);
+		}
+
+		GameViewportClient->ConsoleCommand(Cmd);
+	}
+}
+
+#if !(UE_BUILD_SHIPPING || UE_BUILD_TEST)
+
+void ADebugCameraController::UpdateVisualizeBufferPostProcessing(FFinalPostProcessSettings& InOutPostProcessingSettings)
+{
+	if (bEnableBufferVisualization)
+	{
+		FName TargetBufferName = *LastSelectedBuffer;
+		if (UMaterial* Material = GetBufferVisualizationData().GetMaterial(TargetBufferName))
+		{
+			InOutPostProcessingSettings.bBufferVisualizationOverviewTargetIsSelected = true;
+			InOutPostProcessingSettings.BufferVisualizationOverviewSelectedTargetMaterialName = Material->GetName();
+			return;
+		}
+	}
+
+	InOutPostProcessingSettings.bBufferVisualizationOverviewTargetIsSelected = false;
+	InOutPostProcessingSettings.BufferVisualizationOverviewSelectedTargetMaterialName.Empty();
+}
+
+#endif
+
+void ADebugCameraController::GetNextBuffer(int32 Step)
+{
+	if (bEnableBufferVisualization)
+	{
+		TArray<FString> OverviewBuffers = GetBufferVisualizationOverviewTargets();
+		GetNextBuffer(OverviewBuffers, Step);
+	}
+}
+
+void ADebugCameraController::GetNextBuffer(const TArray<FString>& OverviewBuffers, int32 Step)
+{
+	if (bEnableBufferVisualization)
+	{
+		int32 BufferIndex = 0;
+
+		if (!LastSelectedBuffer.IsEmpty())
+		{
+			bool bFoundIndex = false;
+
+			for (int32 Index = 0; Index < OverviewBuffers.Num(); Index++)
+			{
+				if (OverviewBuffers[Index] == LastSelectedBuffer)
+				{
+					BufferIndex = Index;
+					bFoundIndex = true;
+					break;
+				}
+			}
+
+			if (!bFoundIndex)
+			{
+				LastSelectedBuffer.Empty();
+			}
+		}
+
+		if (LastSelectedBuffer.IsEmpty())
+		{
+			for (FString Buffer : OverviewBuffers)
+			{
+				if (!Buffer.IsEmpty())
+				{
+					LastSelectedBuffer = Buffer;
+					break;
+				}
+			}
+		}
+		else
+		{
+			int32 Incr = FMath::Abs(Step);
+			int32 Min = Incr == 1 ? (BufferIndex / 4) * 4 : BufferIndex % 4;
+			int32 Max = Min;
+			for (int32 i = 0; i < 3 && Max + Incr < OverviewBuffers.Num(); i++) { Max += Incr; }
+
+			auto Wrap = [&](int32 Index)
+			{
+				if (Index < Min)
+				{
+					Index = Max;
+				}
+				else if (Index > Max)
+				{
+					Index = Min;
+				}
+
+				return Index;
+			};
+
+			int32 NextIndex = Wrap(BufferIndex + Step);
+
+			while (NextIndex != BufferIndex)
+			{
+				if (!OverviewBuffers[NextIndex].IsEmpty())
+				{
+					LastSelectedBuffer = OverviewBuffers[NextIndex];
+					break;
+				}
+				NextIndex = Wrap(NextIndex + Step);
+			}
+		}
+	}
+}
+
+void ADebugCameraController::BufferVisualizationMoveUp()
+{
+	GetNextBuffer(-4);
+}
+
+void ADebugCameraController::BufferVisualizationMoveDown()
+{
+	GetNextBuffer(4);
+}
+
+void ADebugCameraController::BufferVisualizationMoveRight()
+{
+	GetNextBuffer(1);
+}
+
+void ADebugCameraController::BufferVisualizationMoveLeft()
+{
+	GetNextBuffer(-1);
+}
+
+void ADebugCameraController::ToggleBufferVisualizationFullMode()
+{
+	SetBufferVisualizationFullMode(!bEnableBufferVisualizationFullMode);
+}
+
+void ADebugCameraController::SetBufferVisualizationFullMode(bool bFullMode)
+{
+	bEnableBufferVisualizationFullMode = bFullMode;
+
+	static IConsoleVariable* ICVar = IConsoleManager::Get().FindConsoleVariable(FBufferVisualizationData::GetVisualizationTargetConsoleCommandName());
+	if (ICVar)
+	{
+		static const FName EmptyName = NAME_None;
+		ICVar->Set(bFullMode ? *LastSelectedBuffer : *EmptyName.ToString(), ECVF_SetByCode);
+	}
+}
+
 void ADebugCameraController::SelectTargetedObject()
 {
 	FVector CamLoc;
diff --git a/Engine/Source/Runtime/Engine/Private/DebugCameraControllerSettings.cpp b/Engine/Source/Runtime/Engine/Private/DebugCameraControllerSettings.cpp
new file mode 100644
--- /dev/null
+++ b/Engine/Source/Runtime/Engine/Private/DebugCameraControllerSettings.cpp
@@ -0,0 +1,60 @@
+// Copyright 1998-2019 Epic Games, Inc. All Rights Reserved.
+
+#include "Engine/DebugCameraControllerSettings.h"
+
+UDebugCameraControllerSettings::UDebugCameraControllerSettings(const FObjectInitializer& ObjectInitializer)
+	: Super(ObjectInitializer)
+{
+	SectionName = TEXT("Debug Camera Controller");
+
+	CycleViewModes.Add(FDebugCameraControllerSettingsViewModeIndex(VMI_Lit));
+	CycleViewModes.Add(FDebugCameraControllerSettingsViewModeIndex(VMI_Unlit));
+	CycleViewModes.Add(FDebugCameraControllerSettingsViewModeIndex(VMI_Wireframe));
+	CycleViewModes.Add(FDebugCameraControllerSettingsViewModeIndex(VMI_Lit_DetailLighting));
+	CycleViewModes.Add(FDebugCameraControllerSettingsViewModeIndex(VMI_ReflectionOverride));
+	CycleViewModes.Add(FDebugCameraControllerSettingsViewModeIndex(VMI_CollisionPawn));
+	CycleViewModes.Add(FDebugCameraControllerSettingsViewModeIndex(VMI_CollisionVisibility));
+}
+
+TArray<EViewModeIndex> UDebugCameraControllerSettings::GetCycleViewModes()
+{
+	TArray<EViewModeIndex> ViewModes;
+
+	for (FDebugCameraControllerSettingsViewModeIndex Mode : CycleViewModes)
+	{
+		ViewModes.Add(Mode.ViewModeIndex);
+	}
+
+	return ViewModes;
+}
+
+#if WITH_EDITOR
+
+void UDebugCameraControllerSettings::PostLoad()
+{
+	Super::PostLoad();
+	RemoveInvalidViewModes();
+}
+
+void UDebugCameraControllerSettings::PostEditChangeProperty(FPropertyChangedEvent& PropertyChangedEvent)
+{
+	Super::PostEditChangeProperty(PropertyChangedEvent);
+	RemoveInvalidViewModes();
+}
+
+void UDebugCameraControllerSettings::RemoveInvalidViewModes()
+{
+	UEnum* ViewModeIndexEnum = StaticEnum<EViewModeIndex>();
+
+	for (int32 i = 0; i < CycleViewModes.Num(); i++)
+	{
+		FString ViewModeName = ViewModeIndexEnum->GetNameStringByIndex((int32)CycleViewModes[i].ViewModeIndex);
+		if (ViewModeName.IsEmpty())
+		{
+			CycleViewModes.RemoveAt(i);
+			i--;
+		}
+	}
+}
+
+#endif
diff --git a/Engine/Source/Runtime/Engine/Private/DebugCameraHUD.cpp b/Engine/Source/Runtime/Engine/Private/DebugCameraHUD.cpp
--- a/Engine/Source/Runtime/Engine/Private/DebugCameraHUD.cpp
+++ b/Engine/Source/Runtime/Engine/Private/DebugCameraHUD.cpp
@@ -101,6 +101,10 @@ void ADebugCameraHUD::PostRender()
 			Canvas->DrawText(RenderFont, SpeedString, X, yl, 1.f, 1.f, FontRenderInfo);
 			yl += Y;
 
+			FString const ViewModeString = FString::Printf(TEXT("ViewMode=%s"), GetWorld()->GetGameViewport() ? GetViewModeName((EViewModeIndex)GetWorld()->GetGameViewport()->ViewModeIndex) : TEXT("<Not found>"));
+			Canvas->DrawText(RenderFont, ViewModeString, X, yl, 1.f, 1.f, FontRenderInfo);
+			yl += Y;
+
 			const TCHAR* CVarComplexName = TEXT("g.DebugCameraTraceComplex");
 			bool bTraceComplex = true;
 
@@ -219,6 +223,12 @@ void ADebugCameraHUD::PostRender()
 
 			Canvas->DrawText(RenderFont, TEXT("Orbit Selected Actor: Shift-O"), X, yl, 1.f, 1.f, FontRenderInfo);
 			yl += Y;
+
+			Canvas->DrawText(RenderFont, TEXT("Toggle Buffer Visualization: B ('Enter' to toggle full buffer mode)"), X, yl, 1.f, 1.f, FontRenderInfo);
+			yl += Y;
+
+			Canvas->DrawText(RenderFont, TEXT("Cycle View Mode: V"), X, yl, 1.f, 1.f, FontRenderInfo);
+			yl += Y;
 		}
 	}
 #endif
diff --git a/Engine/Source/Runtime/Engine/Private/LocalPlayer.cpp b/Engine/Source/Runtime/Engine/Private/LocalPlayer.cpp
--- a/Engine/Source/Runtime/Engine/Private/LocalPlayer.cpp
+++ b/Engine/Source/Runtime/Engine/Private/LocalPlayer.cpp
@@ -38,6 +38,10 @@
 
 #include "GameDelegates.h"
 
+#if !(UE_BUILD_SHIPPING || UE_BUILD_TEST)
+#include "Engine/DebugCameraController.h"
+#endif
+
 DEFINE_LOG_CATEGORY(LogPlayerManagement);
 
 #if !UE_BUILD_SHIPPING
@@ -865,6 +869,14 @@ FSceneView* ULocalPlayer::CalcSceneView( class FSceneViewFamily* ViewFamily,
 			PlayerController->PlayerCameraManager->UpdatePhotographyPostProcessing(View->FinalPostProcessSettings);
 		}
 
+#if !(UE_BUILD_SHIPPING || UE_BUILD_TEST)
+		ADebugCameraController* DebugCameraController = Cast<ADebugCameraController>(PlayerController);
+		if (DebugCameraController != nullptr)
+		{
+			DebugCameraController->UpdateVisualizeBufferPostProcessing(View->FinalPostProcessSettings);
+		}
+#endif
+
 		View->EndFinalPostprocessSettings(ViewInitOptions);
 	}
 
diff --git a/Engine/Source/Runtime/Engine/Public/FinalPostProcessSettings.h b/Engine/Source/Runtime/Engine/Public/FinalPostProcessSettings.h
--- a/Engine/Source/Runtime/Engine/Public/FinalPostProcessSettings.h
+++ b/Engine/Source/Runtime/Engine/Public/FinalPostProcessSettings.h
@@ -21,6 +21,9 @@ public:
 	, HighResScreenshotMaskMaterial(NULL)
 	, HighResScreenshotCaptureRegionMaterial(NULL)
 	, bBufferVisualizationDumpRequired(false)
+#if !(UE_BUILD_SHIPPING || UE_BUILD_TEST)
+	, bBufferVisualizationOverviewTargetIsSelected(false)
+#endif
 	{
 		// to avoid reallocations we reserve some 
 		ContributingCubemaps.Reserve(8);
@@ -178,6 +181,12 @@ public:
 	bool bBufferVisualizationDumpRequired;
 	FString BufferVisualizationDumpBaseFilename;
 
+#if !(UE_BUILD_SHIPPING || UE_BUILD_TEST)
+	// Optional selected buffer in the buffer visualization overview
+	bool bBufferVisualizationOverviewTargetIsSelected;
+	FString BufferVisualizationOverviewSelectedTargetMaterialName;
+#endif
+
 	/**
 	 * A map of buffer visualization material names to visualization pipes onto which each intermediate render target will be pushed, if set.
 	 * Will always be used if set, regardless of r.BufferVisualizationDumpFrames
diff --git a/Engine/Source/Runtime/Renderer/Private/PostProcess/PostProcessVisualizeBuffer.cpp b/Engine/Source/Runtime/Renderer/Private/PostProcess/PostProcessVisualizeBuffer.cpp
--- a/Engine/Source/Runtime/Renderer/Private/PostProcess/PostProcessVisualizeBuffer.cpp
+++ b/Engine/Source/Runtime/Renderer/Private/PostProcess/PostProcessVisualizeBuffer.cpp
@@ -40,6 +40,7 @@ public:
 	FSceneTextureShaderParameters SceneTextureParameters;
 	FShaderResourceParameter SourceTexture;
 	FShaderResourceParameter SourceTextureSampler;
+	FShaderParameter SelectionColor;
 
 	/** Initialization constructor. */
 	FPostProcessVisualizeBufferPS(const ShaderMetaType::CompiledShaderInitializerType& Initializer)
@@ -47,6 +48,7 @@ public:
 	{
 		PostprocessParameter.Bind(Initializer.ParameterMap);
 		SceneTextureParameters.Bind(Initializer);
+		SelectionColor.Bind(Initializer.ParameterMap, TEXT("SelectionColor"));
 
 		if (bDrawingTile)
 		{
@@ -81,11 +83,18 @@ public:
 		}
 	}
 	
+	void SetSelectionColor(FRHICommandList& RHICmdList, const FVector4& InColor)
+	{
+		SetShaderValue(RHICmdList, GetPixelShader(), SelectionColor, InColor);
+	}
+
 	// FShader interface.
 	virtual bool Serialize(FArchive& Ar) override
 	{
 		bool bShaderHasOutdatedParameters = FGlobalShader::Serialize(Ar);
 		Ar << PostprocessParameter << SceneTextureParameters << SourceTexture << SourceTextureSampler;
+		Ar << SelectionColor;
+
 		return bShaderHasOutdatedParameters;
 	}
 
@@ -100,9 +109,9 @@ public:
 	}
 };
 
-void FRCPassPostProcessVisualizeBuffer::AddVisualizationBuffer(FRenderingCompositeOutputRef InSource, const FString& InName)
+void FRCPassPostProcessVisualizeBuffer::AddVisualizationBuffer(FRenderingCompositeOutputRef InSource, const FString& InName, bool bIsSelected)
 {
-	Tiles.Add(TileData(InSource, InName));
+	Tiles.Add(TileData(InSource, InName, bIsSelected));
 
 	if (InSource.IsValid())
 	{
@@ -225,6 +234,17 @@ void FRCPassPostProcessVisualizeBuffer::Process(FRenderingCompositePassContext&
 
 				PixelShader->SetSourceTexture(Context.RHICmdList, Texture);
 
+				const FLinearColor SelectedColor = FLinearColor::Yellow;
+				const FLinearColor NotSelectedColor = FLinearColor::Transparent;
+				if (It->bIsSelected)
+				{
+					PixelShader->SetSelectionColor(Context.RHICmdList, SelectedColor);
+				}
+				else
+				{
+					PixelShader->SetSelectionColor(Context.RHICmdList, NotSelectedColor);
+				}
+
 				DrawRectangle(
 					Context.RHICmdList,
 					TileX * TileWidth, TileY * TileHeight,
diff --git a/Engine/Source/Runtime/Renderer/Private/PostProcess/PostProcessVisualizeBuffer.h b/Engine/Source/Runtime/Renderer/Private/PostProcess/PostProcessVisualizeBuffer.h
--- a/Engine/Source/Runtime/Renderer/Private/PostProcess/PostProcessVisualizeBuffer.h
+++ b/Engine/Source/Runtime/Renderer/Private/PostProcess/PostProcessVisualizeBuffer.h
@@ -22,17 +22,19 @@ public:
 	{
 		FRenderingCompositeOutputRef Source;
 		FString Name;
+		bool bIsSelected;
 
-		TileData(FRenderingCompositeOutputRef InSource, const FString& InName)
+		TileData(FRenderingCompositeOutputRef InSource, const FString& InName, bool bSelected)
 		: Source (InSource)
 		, Name (InName)
+		, bIsSelected(bSelected)
 		{
 
 		}
 	};
 
 	// constructor
-	void AddVisualizationBuffer(FRenderingCompositeOutputRef InSource, const FString& InName);
+	void AddVisualizationBuffer(FRenderingCompositeOutputRef InSource, const FString& InName, bool bIsSelected = false);
 
 	// interface FRenderingCompositePass ---------
 
diff --git a/Engine/Source/Runtime/Renderer/Private/PostProcess/PostProcessing.cpp b/Engine/Source/Runtime/Renderer/Private/PostProcess/PostProcessing.cpp
--- a/Engine/Source/Runtime/Renderer/Private/PostProcess/PostProcessing.cpp
+++ b/Engine/Source/Runtime/Renderer/Private/PostProcess/PostProcessing.cpp
@@ -1289,8 +1289,16 @@ static void AddGBufferVisualizationOverview(FPostprocessContext& Context, FRende
 						FRenderingCompositePass* QuarterSize = Context.Graph.RegisterPass(new(FMemStack::Get()) FRCPassPostProcessDownsample(PF_Unknown, 0, false, TEXT("MaterialQuarterSize")));
 						QuarterSize->SetInput(ePId_Input0, FRenderingCompositeOutputRef(HalfSize));
 
+						// Set whether current buffer is selected
+						bool bIsSelected = false;
+
+#if !(UE_BUILD_SHIPPING || UE_BUILD_TEST)
+						bIsSelected = (Context.View.FinalPostProcessSettings.bBufferVisualizationOverviewTargetIsSelected &&
+							VisualizationName == Context.View.FinalPostProcessSettings.BufferVisualizationOverviewSelectedTargetMaterialName);
+#endif
+
 						// Mark the quarter size target as the dependency for the composite pass
-						((FRCPassPostProcessVisualizeBuffer*)CompositePass)->AddVisualizationBuffer(FRenderingCompositeOutputRef(QuarterSize), VisualizationName);
+						((FRCPassPostProcessVisualizeBuffer*)CompositePass)->AddVisualizationBuffer(FRenderingCompositeOutputRef(QuarterSize), VisualizationName, bIsSelected);
 					}
 					else
 					{
