Swapped order of key-value in the custom primitive data API on UStaticMeshComponent

#rnx

diff --git a/Engine/Source/Runtime/Engine/Classes/Components/StaticMeshComponent.h b/Engine/Source/Runtime/Engine/Classes/Components/StaticMeshComponent.h
--- a/Engine/Source/Runtime/Engine/Classes/Components/StaticMeshComponent.h
+++ b/Engine/Source/Runtime/Engine/Classes/Components/StaticMeshComponent.h
@@ -389,19 +389,19 @@ public:
 
 	/** Set custom primitive data at index DataIndex. */
 	UFUNCTION(BlueprintCallable, Category="Components|StaticMesh")
-	void SetCustomPrimitiveDataFloat(float Value, int32 DataIndex);
+	void SetCustomPrimitiveDataFloat(int32 DataIndex, float Value);
 
 	/** Set custom primitive data, two floats at once, from index DataIndex to index DataIndex + 2. */
 	UFUNCTION(BlueprintCallable, Category="Components|StaticMesh")
-	void SetCustomPrimitiveDataVector2(FVector2D Value, int32 DataIndex);
+	void SetCustomPrimitiveDataVector2(int32 DataIndex, FVector2D Value);
 
 	/** Set custom primitive data, three floats at once, from index DataIndex to index DataIndex + 3. */
 	UFUNCTION(BlueprintCallable, Category="Components|StaticMesh")
-	void SetCustomPrimitiveDataVector3(FVector Value, int32 DataIndex);
+	void SetCustomPrimitiveDataVector3(int32 DataIndex, FVector Value);
 
 	/** Set custom primitive data, four floats at once, from index DataIndex to index DataIndex + 4. */
 	UFUNCTION(BlueprintCallable, Category="Components|StaticMesh")
-	void SetCustomPrimitiveDataVector4(FVector4 Value, int32 DataIndex);
+	void SetCustomPrimitiveDataVector4(int32 DataIndex, FVector4 Value);
 
 	/** 
 	 * Set forced reverse culling
@@ -642,7 +642,7 @@ private:
 	void PrivateFixupOverrideColors();
 
 	/** Insert an array of floats into the CustomPrimitiveData, starting at the given index */
-	void SetCustomPrimitiveDataInternal(const TArray<float>& Values, int32 DataIndex);
+	void SetCustomPrimitiveDataInternal(int32 DataIndex, const TArray<float>& Values);
 protected:
 
 	/** Whether the component type supports static lighting. */
diff --git a/Engine/Source/Runtime/Engine/Private/Components/StaticMeshComponent.cpp b/Engine/Source/Runtime/Engine/Private/Components/StaticMeshComponent.cpp
--- a/Engine/Source/Runtime/Engine/Private/Components/StaticMeshComponent.cpp
+++ b/Engine/Source/Runtime/Engine/Private/Components/StaticMeshComponent.cpp
@@ -1881,7 +1881,7 @@ void UStaticMeshComponent::SetDistanceFieldSelfShadowBias(float NewValue)
 	}
 }
 
-void UStaticMeshComponent::SetCustomPrimitiveDataInternal(const TArray<float>& Values, int32 DataIndex)
+void UStaticMeshComponent::SetCustomPrimitiveDataInternal(int32 DataIndex, const TArray<float>& Values)
 {
 	// Can only set data on valid indices and only if there's actually any data to set
 	if (DataIndex < FCustomPrimitiveData::NumCustomPrimitiveDataFloats && Values.Num() > 0)
@@ -1906,24 +1906,24 @@ void UStaticMeshComponent::SetCustomPrimitiveDataInternal(const TArray<float>& V
 	}
 }
 
-void UStaticMeshComponent::SetCustomPrimitiveDataFloat(float Value, int32 DataIndex)
+void UStaticMeshComponent::SetCustomPrimitiveDataFloat(int32 DataIndex, float Value)
 {
-	SetCustomPrimitiveDataInternal({Value}, DataIndex);
+	SetCustomPrimitiveDataInternal(DataIndex, {Value});
 }
 
-void UStaticMeshComponent::SetCustomPrimitiveDataVector2(FVector2D Value, int32 DataIndex)
+void UStaticMeshComponent::SetCustomPrimitiveDataVector2(int32 DataIndex, FVector2D Value)
 {
-	SetCustomPrimitiveDataInternal({Value.X, Value.Y}, DataIndex);
+	SetCustomPrimitiveDataInternal(DataIndex, {Value.X, Value.Y});
 }
 
-void UStaticMeshComponent::SetCustomPrimitiveDataVector3(FVector Value, int32 DataIndex)
+void UStaticMeshComponent::SetCustomPrimitiveDataVector3(int32 DataIndex, FVector Value)
 {
-	SetCustomPrimitiveDataInternal({Value.X, Value.Y, Value.Z}, DataIndex);
+	SetCustomPrimitiveDataInternal(DataIndex, {Value.X, Value.Y, Value.Z});
 }
 
-void UStaticMeshComponent::SetCustomPrimitiveDataVector4(FVector4 Value, int32 DataIndex)
+void UStaticMeshComponent::SetCustomPrimitiveDataVector4(int32 DataIndex, FVector4 Value)
 {
-	SetCustomPrimitiveDataInternal({Value.X, Value.Y, Value.Z, Value.W}, DataIndex);
+	SetCustomPrimitiveDataInternal(DataIndex, {Value.X, Value.Y, Value.Z, Value.W});
 }
 
 void UStaticMeshComponent::SetReverseCulling(bool ReverseCulling)
