Added missing headers for non-unity build.

diff --git a/Engine/Source/Runtime/Engine/Classes/Engine/DebugCameraController.h b/Engine/Source/Runtime/Engine/Classes/Engine/DebugCameraController.h
--- a/Engine/Source/Runtime/Engine/Classes/Engine/DebugCameraController.h
+++ b/Engine/Source/Runtime/Engine/Classes/Engine/DebugCameraController.h
@@ -5,6 +5,7 @@
 #include "CoreMinimal.h"
 #include "UObject/ObjectMacros.h"
 #include "Engine/EngineTypes.h"
+#include "FinalPostProcessSettings.h"
 #include "SceneTypes.h"
 #include "GameFramework/PlayerController.h"
 #include "DebugCameraController.generated.h"
diff --git a/Engine/Source/Runtime/Engine/Classes/Engine/DebugCameraControllerSettings.h b/Engine/Source/Runtime/Engine/Classes/Engine/DebugCameraControllerSettings.h
--- a/Engine/Source/Runtime/Engine/Classes/Engine/DebugCameraControllerSettings.h
+++ b/Engine/Source/Runtime/Engine/Classes/Engine/DebugCameraControllerSettings.h
@@ -7,6 +7,8 @@
 #pragma once
 
 #include "Engine/DeveloperSettings.h"
+#include "Engine/EngineTypes.h"
+#include "Engine/EngineBaseTypes.h"
 #include "DebugCameraControllerSettings.generated.h"
 
 USTRUCT()
diff --git a/Engine/Source/Runtime/Engine/Private/DebugCameraController.cpp b/Engine/Source/Runtime/Engine/Private/DebugCameraController.cpp
--- a/Engine/Source/Runtime/Engine/Private/DebugCameraController.cpp
+++ b/Engine/Source/Runtime/Engine/Private/DebugCameraController.cpp
@@ -14,6 +14,7 @@
 #include "Engine/MapBuildDataRegistry.h"
 #include "Engine/Engine.h"
 #include "Engine/Player.h"
+#include "Materials/Material.h"
 #include "EngineUtils.h"
 #include "GameFramework/SpectatorPawn.h"
 #include "GameFramework/SpectatorPawnMovement.h"
