Add more context to FRepHandleIterator ensures. Also, make sure we explicitly trigger allocations for arrays after moves.

diff --git a/Engine/Source/Runtime/Engine/Private/RepLayout.cpp b/Engine/Source/Runtime/Engine/Private/RepLayout.cpp
--- a/Engine/Source/Runtime/Engine/Private/RepLayout.cpp
+++ b/Engine/Source/Runtime/Engine/Private/RepLayout.cpp
@@ -1031,7 +1031,7 @@ bool FRepLayout::CompareProperties(
 	FRepChangedHistory& NewHistoryItem = RepChangelistState->ChangeHistory[HistoryIndex];
 
 	TArray<uint16>& Changed = NewHistoryItem.Changed;
-	Changed.Empty();
+	Changed.Empty(1);
 
 	FComparePropertiesSharedParams SharedParams{
 		/*bIsInitial=*/ !!RepFlags.bNetInitial,
@@ -1521,8 +1521,10 @@ bool FRepHandleIterator::NextHandle()
 	ChangelistIterator.ChangedIndex++;
 
 	if (!ensureMsgf(ChangelistIterator.Changed.IsValidIndex(ChangelistIterator.ChangedIndex),
-			TEXT("Attempted to access invalid iterator index: Handle=%d, ChangedIndex=%d, ChangedNum=%d"),
-			Handle, ChangelistIterator.ChangedIndex, ChangelistIterator.Changed.Num()))
+			TEXT("Attempted to access invalid iterator index: Handle=%d, ChangedIndex=%d, ChangedNum=%d, Owner=%s, LastSuccessfulCmd=%s"),
+			Handle, ChangelistIterator.ChangedIndex, ChangelistIterator.Changed.Num(),
+			*GetPathNameSafe(Owner),
+			*((Cmds.IsValidIndex(LastSuccessfulCmdIndex) && Cmds[LastSuccessfulCmdIndex].Property) ? Cmds[LastSuccessfulCmdIndex].Property->GetPathName() : FString::FromInt(LastSuccessfulCmdIndex))))
 	{
 		return false;
 	}
@@ -1541,17 +1543,21 @@ bool FRepHandleIterator::NextHandle()
 	const int32 RelativeHandle = HandleMinusOne - ArrayIndex * NumHandlesPerElement;
 
 	if (!ensureMsgf(HandleToCmdIndex.IsValidIndex(RelativeHandle),
-			TEXT("Attempted to access invalid RelativeHandle Index: Handle=%d, RelativeHandle=%d, NumHandlesPerElement=%d, ArrayIndex=%d, ArrayElementSize=%d"),
-			Handle, RelativeHandle, NumHandlesPerElement, ArrayIndex, ArrayElementSize))
+			TEXT("Attempted to access invalid RelativeHandle Index: Handle=%d, RelativeHandle=%d, NumHandlesPerElement=%d, ArrayIndex=%d, ArrayElementSize=%d, Owner=%s, LastSuccessfulCmd=%s"),
+			Handle, RelativeHandle, NumHandlesPerElement, ArrayIndex, ArrayElementSize,
+			*GetPathNameSafe(Owner),
+			*((Cmds.IsValidIndex(LastSuccessfulCmdIndex) && Cmds[LastSuccessfulCmdIndex].Property) ? Cmds[LastSuccessfulCmdIndex].Property->GetPathName() : FString::FromInt(LastSuccessfulCmdIndex))))
 	{
 		return false;
 	}
 
 	CmdIndex = HandleToCmdIndex[RelativeHandle].CmdIndex;
 
-	if (!ensureMsgf(CmdIndex >= MinCmdIndex && CmdIndex < MaxCmdIndex,
-			TEXT("Attempted to access Command Index outside of iterator range: Handle=%d, RelativeHandle=%d, CmdIndex=%d, MinCmdIdx=%d, MaxCmdIdx=%d, ArrayIndex=%d"),
-			Handle, RelativeHandle, CmdIndex, MinCmdIndex, MaxCmdIndex, ArrayIndex))
+	if (!ensureMsgf(MinCmdIndex <= CmdIndex && CmdIndex < MaxCmdIndex,
+			TEXT("Attempted to access Command Index outside of iterator range: Handle=%d, RelativeHandle=%d, CmdIndex=%d, MinCmdIdx=%d, MaxCmdIdx=%d, ArrayIndex=%d, Owner=%s, LastSuccessfulCmd=%s"),
+			Handle, RelativeHandle, CmdIndex, MinCmdIndex, MaxCmdIndex, ArrayIndex,
+			*GetPathNameSafe(Owner),
+			*((Cmds.IsValidIndex(LastSuccessfulCmdIndex) && Cmds[LastSuccessfulCmdIndex].Property) ? Cmds[LastSuccessfulCmdIndex].Property->GetPathName() : FString::FromInt(LastSuccessfulCmdIndex))))
 	{
 		return false;
 	}
@@ -1559,19 +1565,25 @@ bool FRepHandleIterator::NextHandle()
 	const FRepLayoutCmd& Cmd = Cmds[CmdIndex];
 
 	if (!ensureMsgf(Cmd.RelativeHandle - 1 == RelativeHandle,
-			TEXT("Command Relative Handle does not match found Relative Handle: Handle=%d, RelativeHandle=%d, CmdIdx=%d, CmdRelativeHandle=%d, ArrayIndex=%d"),
-			Handle, RelativeHandle, CmdIndex, Cmd.RelativeHandle, ArrayIndex))
+			TEXT("Command Relative Handle does not match found Relative Handle: Handle=%d, RelativeHandle=%d, CmdIdx=%d, CmdRelativeHandle=%d, ArrayIndex=%d, Owner=%s, LastSuccessfulCmd=%s"),
+			Handle, RelativeHandle, CmdIndex, Cmd.RelativeHandle, ArrayIndex,
+			*GetPathNameSafe(Owner),
+			*((Cmds.IsValidIndex(LastSuccessfulCmdIndex) && Cmds[LastSuccessfulCmdIndex].Property) ? Cmds[LastSuccessfulCmdIndex].Property->GetPathName() : FString::FromInt(LastSuccessfulCmdIndex))))
 	{
 		return false;
 	}
 
 	if (!ensureMsgf(Cmd.Type != ERepLayoutCmdType::Return,
-			TEXT("Hit unexpected return handle: Handle=%d, RelativeHandle=%d, CmdIdx=%d, ArrayIndex=%d"),
-			Handle, RelativeHandle, CmdIndex, ArrayIndex))
+			TEXT("Hit unexpected return handle: Handle=%d, RelativeHandle=%d, CmdIdx=%d, ArrayIndex=%d, Owner=%s, LastSuccessfulCmd=%s"),
+			Handle, RelativeHandle, CmdIndex, ArrayIndex,
+			*GetPathNameSafe(Owner),
+			*((Cmds.IsValidIndex(LastSuccessfulCmdIndex) && Cmds[LastSuccessfulCmdIndex].Property) ? Cmds[LastSuccessfulCmdIndex].Property->GetPathName() : FString::FromInt(LastSuccessfulCmdIndex))))
 	{
 		return false;
 	}
 
+	LastSuccessfulCmdIndex = CmdIndex;
+
 	return true;
 }
 
@@ -1733,18 +1745,18 @@ void FRepLayout::MergeChangeList_r(
 
 			if (!ActiveIterator1)
 			{
-				FRepHandleIterator ArrayIterator2(ActiveIterator2->ChangelistIterator, Cmds, ArrayHandleToCmdIndex, Cmd.ElementSize, Array->Num(), CmdIndex + 1, Cmd.EndCmd - 1);
+				FRepHandleIterator ArrayIterator2(ActiveIterator2->Owner, ActiveIterator2->ChangelistIterator, Cmds, ArrayHandleToCmdIndex, Cmd.ElementSize, Array->Num(), CmdIndex + 1, Cmd.EndCmd - 1);
 				PruneChangeList_r(ArrayIterator2, ArrayData, OutChanged);
 			}
 			else if (!ActiveIterator2)
 			{
-				FRepHandleIterator ArrayIterator1(ActiveIterator1->ChangelistIterator, Cmds, ArrayHandleToCmdIndex, Cmd.ElementSize, Array->Num(), CmdIndex + 1, Cmd.EndCmd - 1);
+				FRepHandleIterator ArrayIterator1(ActiveIterator1->Owner, ActiveIterator1->ChangelistIterator, Cmds, ArrayHandleToCmdIndex, Cmd.ElementSize, Array->Num(), CmdIndex + 1, Cmd.EndCmd - 1);
 				PruneChangeList_r(ArrayIterator1, ArrayData, OutChanged);
 			}
 			else
 			{
-				FRepHandleIterator ArrayIterator1(ActiveIterator1->ChangelistIterator, Cmds, ArrayHandleToCmdIndex, Cmd.ElementSize, Array->Num(), CmdIndex + 1, Cmd.EndCmd - 1);
-				FRepHandleIterator ArrayIterator2(ActiveIterator2->ChangelistIterator, Cmds, ArrayHandleToCmdIndex, Cmd.ElementSize, Array->Num(), CmdIndex + 1, Cmd.EndCmd - 1);
+				FRepHandleIterator ArrayIterator1(ActiveIterator1->Owner, ActiveIterator1->ChangelistIterator, Cmds, ArrayHandleToCmdIndex, Cmd.ElementSize, Array->Num(), CmdIndex + 1, Cmd.EndCmd - 1);
+				FRepHandleIterator ArrayIterator2(ActiveIterator2->Owner, ActiveIterator2->ChangelistIterator, Cmds, ArrayHandleToCmdIndex, Cmd.ElementSize, Array->Num(), CmdIndex + 1, Cmd.EndCmd - 1);
 
 				MergeChangeList_r(ArrayIterator1, ArrayIterator2, ArrayData, OutChanged);
 			}
@@ -1784,7 +1796,7 @@ void FRepLayout::PruneChangeList_r(
 
 			TArray<FHandleToCmdIndex>& ArrayHandleToCmdIndex = *RepHandleIterator.HandleToCmdIndex[Cmd.RelativeHandle - 1].HandleToCmdIndex;
 
-			FRepHandleIterator ArrayIterator(RepHandleIterator.ChangelistIterator, Cmds, ArrayHandleToCmdIndex, Cmd.ElementSize, Array->Num(), CmdIndex + 1, Cmd.EndCmd - 1);
+			FRepHandleIterator ArrayIterator(RepHandleIterator.Owner, RepHandleIterator.ChangelistIterator, Cmds, ArrayHandleToCmdIndex, Cmd.ElementSize, Array->Num(), CmdIndex + 1, Cmd.EndCmd - 1);
 			PruneChangeList_r(ArrayIterator, ArrayData, OutChanged);
 
 			// Patch in the jump offset
@@ -1803,10 +1815,10 @@ void FRepLayout::FilterChangeList(
 	TArray<uint16>& OutActiveProperties) const
 {
 	FChangelistIterator ChangelistIterator(Changelist, 0);
-	FRepHandleIterator HandleIterator(ChangelistIterator, Cmds, BaseHandleToCmdIndex, 0, 1, 0, Cmds.Num() - 1);
+	FRepHandleIterator HandleIterator(Owner, ChangelistIterator, Cmds, BaseHandleToCmdIndex, 0, 1, 0, Cmds.Num() - 1);
 
-	OutInactiveProperties.Empty();
-	OutActiveProperties.Empty();
+	OutInactiveProperties.Empty(1);
+	OutActiveProperties.Empty(1);
 
 	while (HandleIterator.NextHandle())
 	{
@@ -1843,9 +1855,9 @@ void FRepLayout::FilterChangeListToActive(
 	TArray<uint16>& OutProperties) const
 {
 	FChangelistIterator ChangelistIterator(Changelist, 0);
-	FRepHandleIterator HandleIterator(ChangelistIterator, Cmds, BaseHandleToCmdIndex, 0, 1, 0, Cmds.Num() - 1);
+	FRepHandleIterator HandleIterator(Owner, ChangelistIterator, Cmds, BaseHandleToCmdIndex, 0, 1, 0, Cmds.Num() - 1);
 
-	OutProperties.Empty();
+	OutProperties.Empty(1);
 
 	while (HandleIterator.NextHandle())
 	{
@@ -1971,7 +1983,7 @@ void FRepLayout::SendProperties_r(
 
 			TArray<FHandleToCmdIndex>& ArrayHandleToCmdIndex = *HandleIterator.HandleToCmdIndex[Cmd.RelativeHandle - 1].HandleToCmdIndex;
 
-			FRepHandleIterator ArrayHandleIterator(HandleIterator.ChangelistIterator, Cmds, ArrayHandleToCmdIndex, Cmd.ElementSize, ArrayNum, HandleIterator.CmdIndex + 1, Cmd.EndCmd - 1);
+			FRepHandleIterator ArrayHandleIterator(HandleIterator.Owner, HandleIterator.ChangelistIterator, Cmds, ArrayHandleToCmdIndex, Cmd.ElementSize, ArrayNum, HandleIterator.CmdIndex + 1, Cmd.EndCmd - 1);
 
 			check(ArrayHandleIterator.ArrayElementSize> 0);
 			check(ArrayHandleIterator.NumHandlesPerElement> 0);
@@ -2125,7 +2137,7 @@ void FRepLayout::SendProperties(
 	UE_LOG(LogRepProperties, VeryVerbose, TEXT("SendProperties: Owner=%s, LastChangelistIndex=%d"), *Owner->GetPathName(), RepState->LastChangelistIndex);
 
 	FChangelistIterator ChangelistIterator(Changed, 0);
-	FRepHandleIterator HandleIterator(ChangelistIterator, Cmds, BaseHandleToCmdIndex, 0, 1, 0, Cmds.Num() - 1);
+	FRepHandleIterator HandleIterator(Owner, ChangelistIterator, Cmds, BaseHandleToCmdIndex, 0, 1, 0, Cmds.Num() - 1);
 
 	SendProperties_r(RepState, Writer, bDoChecksum, HandleIterator, Data, 0, &SharedInfo);
 
@@ -2283,7 +2295,7 @@ void FRepLayout::SendProperties_BackwardsCompatible_r(
 
 			TArray<FHandleToCmdIndex>& ArrayHandleToCmdIndex = *HandleIterator.HandleToCmdIndex[Cmd.RelativeHandle - 1].HandleToCmdIndex;
 
-			FRepHandleIterator ArrayHandleIterator(HandleIterator.ChangelistIterator, Cmds, ArrayHandleToCmdIndex, Cmd.ElementSize, ArrayNum, HandleIterator.CmdIndex + 1, Cmd.EndCmd - 1);
+			FRepHandleIterator ArrayHandleIterator(HandleIterator.Owner, HandleIterator.ChangelistIterator, Cmds, ArrayHandleToCmdIndex, Cmd.ElementSize, ArrayNum, HandleIterator.CmdIndex + 1, Cmd.EndCmd - 1);
 
 			check(ArrayHandleIterator.ArrayElementSize> 0);
 			check(ArrayHandleIterator.NumHandlesPerElement> 0);
@@ -2456,7 +2468,7 @@ void FRepLayout::SendProperties_BackwardsCompatible(
 	{
 		UE_LOG(LogRepProperties, VeryVerbose, TEXT("SendProperties_BackwardsCompatible: SendProperties."));
 		FChangelistIterator ChangelistIterator(Changed, 0);
-		FRepHandleIterator HandleIterator(ChangelistIterator, Cmds, BaseHandleToCmdIndex, 0, 1, 0, Cmds.Num() - 1);
+		FRepHandleIterator HandleIterator(Owner, ChangelistIterator, Cmds, BaseHandleToCmdIndex, 0, 1, 0, Cmds.Num() - 1);
 
 		SendProperties_BackwardsCompatible_r(RepState, PackageMapClient, NetFieldExportGroup.Get(), ChangedTracker, Writer, bDoChecksum, HandleIterator, Data);
 	}
@@ -3853,12 +3865,12 @@ void FRepLayout::PruneChangeList(
 {
 	check(Changed.Num() > 0);
 
-	PrunedChanged.Empty();
+	PrunedChanged.Empty(1);
 
 	if (ERepLayoutState::Normal == LayoutState)
 	{
 		FChangelistIterator ChangelistIterator(Changed, 0);
-		FRepHandleIterator HandleIterator(ChangelistIterator, Cmds, BaseHandleToCmdIndex, 0, 1, 0, Cmds.Num() - 1);
+		FRepHandleIterator HandleIterator(Owner, ChangelistIterator, Cmds, BaseHandleToCmdIndex, 0, 1, 0, Cmds.Num() - 1);
 		PruneChangeList_r(HandleIterator, Data, PrunedChanged);
 	}
 
@@ -3872,23 +3884,23 @@ void FRepLayout::MergeChangeList(
 	TArray<uint16>& MergedDirty) const
 {
 	check(Dirty1.Num() > 0);
-	MergedDirty.Empty();
+	MergedDirty.Empty(1);
 
 	if (ERepLayoutState::Normal == LayoutState)
 	{
 		if (Dirty2.Num() == 0)
 		{
 			FChangelistIterator ChangelistIterator(Dirty1, 0);
-			FRepHandleIterator HandleIterator(ChangelistIterator, Cmds, BaseHandleToCmdIndex, 0, 1, 0, Cmds.Num() - 1);
+			FRepHandleIterator HandleIterator(Owner, ChangelistIterator, Cmds, BaseHandleToCmdIndex, 0, 1, 0, Cmds.Num() - 1);
 			PruneChangeList_r(HandleIterator, Data, MergedDirty);
 		}
 		else
 		{
 			FChangelistIterator ChangelistIterator1(Dirty1, 0);
-			FRepHandleIterator HandleIterator1(ChangelistIterator1, Cmds, BaseHandleToCmdIndex, 0, 1, 0, Cmds.Num() - 1);
+			FRepHandleIterator HandleIterator1(Owner, ChangelistIterator1, Cmds, BaseHandleToCmdIndex, 0, 1, 0, Cmds.Num() - 1);
 
 			FChangelistIterator ChangelistIterator2(Dirty2, 0);
-			FRepHandleIterator HandleIterator2(ChangelistIterator2, Cmds, BaseHandleToCmdIndex, 0, 1, 0, Cmds.Num() - 1);
+			FRepHandleIterator HandleIterator2(Owner, ChangelistIterator2, Cmds, BaseHandleToCmdIndex, 0, 1, 0, Cmds.Num() - 1);
 
 			MergeChangeList_r(HandleIterator1, HandleIterator2, Data, MergedDirty);
 		}
@@ -5363,7 +5375,7 @@ void FRepLayout::BuildSharedSerialization(
 #endif
 
 	FChangelistIterator ChangelistIterator(Changed, 0);
-	FRepHandleIterator HandleIterator(ChangelistIterator, Cmds, BaseHandleToCmdIndex, 0, 1, 0, Cmds.Num() - 1);
+	FRepHandleIterator HandleIterator(Owner, ChangelistIterator, Cmds, BaseHandleToCmdIndex, 0, 1, 0, Cmds.Num() - 1);
 
 	BuildSharedSerialization_r(HandleIterator, Data, bWriteHandle, bDoChecksum, 0, SharedInfo);
 
@@ -5403,7 +5415,7 @@ void FRepLayout::BuildSharedSerialization_r(
 
 			TArray<FHandleToCmdIndex>& ArrayHandleToCmdIndex = *HandleIterator.HandleToCmdIndex[Cmd.RelativeHandle - 1].HandleToCmdIndex;
 
-			FRepHandleIterator ArrayIterator(HandleIterator.ChangelistIterator, Cmds, ArrayHandleToCmdIndex, Cmd.ElementSize, Array->Num(), CmdIndex + 1, Cmd.EndCmd - 1);
+			FRepHandleIterator ArrayIterator(HandleIterator.Owner, HandleIterator.ChangelistIterator, Cmds, ArrayHandleToCmdIndex, Cmd.ElementSize, Array->Num(), CmdIndex + 1, Cmd.EndCmd - 1);
 			BuildSharedSerialization_r(ArrayIterator, ArrayData, bWriteHandle, bDoChecksum, ArrayDepth + 1, SharedInfo);
 			continue;
 		}
@@ -6402,7 +6414,7 @@ bool FRepLayout::DeltaSerializeFastArrayProperty(FFastArrayDeltaSerializeParams&
 
 							FConstRepObjectDataBuffer ElementData(ObjectArrayData + ArrayElementOffset);
 							FRepShadowDataBuffer ElementShadowData(ShadowArrayData + ArrayElementOffset);
-							NewChangelist.Empty();
+							NewChangelist.Empty(1);
 
 							CompareProperties_r(SharedParams, ItemLayoutStart, ItemLayoutEnd, ElementShadowData, ElementData, NewChangelist, 0);
 
@@ -6454,11 +6466,12 @@ bool FRepLayout::DeltaSerializeFastArrayProperty(FFastArrayDeltaSerializeParams&
 
 									TArray<uint16>& ElementChangelist = Changelists[i];
 									TArray<uint16> Temp = MoveTemp(ElementChangelist);
+									ElementChangelist.Empty(1);
 
 									const FConstRepObjectDataBuffer ElementData(ObjectArrayData + (IDIndexPair.Idx * ElementSize));
 
 									FChangelistIterator FoundChangelistIterator(*FoundChangelist, 0);
-									FRepHandleIterator FoundHandleIterator(FoundChangelistIterator, Cmds, ArrayHandleToCmdIndex, ElementSize, 1, ItemLayoutStart, ItemLayoutEnd);
+									FRepHandleIterator FoundHandleIterator(Owner, FoundChangelistIterator, Cmds, ArrayHandleToCmdIndex, ElementSize, 1, ItemLayoutStart, ItemLayoutEnd);
 
 									if (Temp.Num() == 0)
 									{
@@ -6467,7 +6480,7 @@ bool FRepLayout::DeltaSerializeFastArrayProperty(FFastArrayDeltaSerializeParams&
 									else
 									{
 										FChangelistIterator ElementChangelistIterator(Temp, 0);
-										FRepHandleIterator ElementHandleIterator(ElementChangelistIterator, Cmds, ArrayHandleToCmdIndex, ElementSize, 1, ItemLayoutStart, ItemLayoutEnd);
+										FRepHandleIterator ElementHandleIterator(Owner, ElementChangelistIterator, Cmds, ArrayHandleToCmdIndex, ElementSize, 1, ItemLayoutStart, ItemLayoutEnd);
 
 										MergeChangeList_r(FoundHandleIterator, ElementHandleIterator, ElementData, ElementChangelist);
 									}
@@ -6543,6 +6556,7 @@ bool FRepLayout::DeltaSerializeFastArrayProperty(FFastArrayDeltaSerializeParams&
 				{
 					FChangelistIterator ChangelistIterator(Changelist, 0);
 					FRepHandleIterator HandleIterator(
+						Owner,
 						ChangelistIterator,
 						Cmds,
 						ArrayHandleToCmdIndex,
diff --git a/Engine/Source/Runtime/Engine/Public/Net/RepLayout.h b/Engine/Source/Runtime/Engine/Public/Net/RepLayout.h
--- a/Engine/Source/Runtime/Engine/Public/Net/RepLayout.h
+++ b/Engine/Source/Runtime/Engine/Public/Net/RepLayout.h
@@ -894,13 +894,14 @@ class FRepHandleIterator
 public:
 
 	FRepHandleIterator(
-		FChangelistIterator&	InChangelistIterator,
-		const TArray<FRepLayoutCmd>&		InCmds,
-		const TArray<FHandleToCmdIndex>&	InHandleToCmdIndex,
-		const int32							InElementSize,
-		const int32							InMaxArrayIndex,
-		const int32							InMinCmdIndex,
-		const int32							InMaxCmdIndex
+		UStruct const * const InOwner,
+		FChangelistIterator& InChangelistIterator,
+		const TArray<FRepLayoutCmd>& InCmds,
+		const TArray<FHandleToCmdIndex>& InHandleToCmdIndex,
+		const int32 InElementSize,
+		const int32 InMaxArrayIndex,
+		const int32 InMinCmdIndex,
+		const int32 InMaxCmdIndex
 	):
 		ChangelistIterator(InChangelistIterator),
 		Cmds(InCmds),
@@ -909,8 +910,12 @@ public:
 		ArrayElementSize(InElementSize),
 		MaxArrayIndex(InMaxArrayIndex),
 		MinCmdIndex(InMinCmdIndex),
-		MaxCmdIndex(InMaxCmdIndex)
-	{}
+		MaxCmdIndex(InMaxCmdIndex),
+		Owner(InOwner),
+		LastSuccessfulCmdIndex(INDEX_NONE)
+	{
+		ensureMsgf(MaxCmdIndex >= MinCmdIndex, TEXT("Invalid Min / Max Command Indices. Owner=%s, MinCmdIndex=%d, MaxCmdIndex=%d"), *GetPathNameSafe(Owner), MinCmdIndex, MaxCmdIndex);
+	}
 
 	/**
 	 * Moves the iterator to the next available handle.
@@ -977,6 +982,12 @@ public:
 
 	/** The Byte offset of Serialized Property data for a dynamic array to the current element. */
 	int32 ArrayOffset;
+
+	UStruct const * const Owner;
+
+private:
+
+	int32 LastSuccessfulCmdIndex;
 };
 
 enum class ECreateRepLayoutFlags
