Fixed static analyzer

diff --git a/Engine/Source/Runtime/Engine/Private/StaticMesh.cpp b/Engine/Source/Runtime/Engine/Private/StaticMesh.cpp
--- a/Engine/Source/Runtime/Engine/Private/StaticMesh.cpp
+++ b/Engine/Source/Runtime/Engine/Private/StaticMesh.cpp
@@ -1295,6 +1295,8 @@ void FStaticMeshRenderData::Serialize(FArchive& Ar, UStaticMesh* Owner, bool bCo
 				if (bValid)
 				{
 #if WITH_EDITOR
+					check(LOD.DistanceFieldData != nullptr);
+
 					bool bDownSampling = Ar.IsCooking() && Ar.IsSaving();
 					
 					if (bDownSampling)
