Changed the raw function pointer parameter to a delegate

#rnx

diff --git a/Engine/Source/Runtime/Engine/Private/Materials/MaterialShader.cpp b/Engine/Source/Runtime/Engine/Private/Materials/MaterialShader.cpp
--- a/Engine/Source/Runtime/Engine/Private/Materials/MaterialShader.cpp
+++ b/Engine/Source/Runtime/Engine/Private/Materials/MaterialShader.cpp
@@ -78,7 +78,7 @@ FString GetShadingModelString(EMaterialShadingModel ShadingModel)
 
 
 /** Converts an FMaterialShadingModelField to a string description containing all the shading models present, delimited by "|" */
-FString GetShadingModelFieldString(FMaterialShadingModelField ShadingModels, FString (*ShadingModelToString)(EMaterialShadingModel), const FString& Delimiter)
+FString GetShadingModelFieldString(FMaterialShadingModelField ShadingModels, const FShadingModelToStringDelegate& Delegate, const FString& Delimiter)
 {
 	FString ShadingModelsName;
 	uint32 TempShadingModels = (uint32)ShadingModels.GetShadingModelField();
@@ -87,12 +87,12 @@ FString GetShadingModelFieldString(FMaterialShadingModelField ShadingModels, FSt
 	{
 		uint32 BitIndex = FMath::CountTrailingZeros(TempShadingModels); // Find index of first set bit
 		TempShadingModels &= ~(1 << BitIndex); // Flip first set bit to 0
-		ShadingModelsName += ShadingModelToString((EMaterialShadingModel)BitIndex); // Add the name of the shading model corresponding to that bit
+		ShadingModelsName += Delegate.Execute((EMaterialShadingModel)BitIndex); // Add the name of the shading model corresponding to that bit
 
 		// If there are more bits left, add a pipe limiter to the string 
 		if (TempShadingModels)
 		{
-			ShadingModelsName.Append(TEXT("|"));
+			ShadingModelsName.Append(Delimiter);
 		}
 	}
 
@@ -102,7 +102,7 @@ FString GetShadingModelFieldString(FMaterialShadingModelField ShadingModels, FSt
 /** Converts an FMaterialShadingModelField to a string description containing all the shading models present, delimited by "|" */
 FString GetShadingModelFieldString(FMaterialShadingModelField ShadingModels)
 {
-	return GetShadingModelFieldString(ShadingModels, &GetShadingModelString, TEXT("|"));
+	return GetShadingModelFieldString(ShadingModels, FShadingModelToStringDelegate::CreateStatic(&GetShadingModelString), TEXT("|"));
 }
 
 /** Converts an EBlendMode to a string description. */
diff --git a/Engine/Source/Runtime/Engine/Public/MaterialShaderType.h b/Engine/Source/Runtime/Engine/Public/MaterialShaderType.h
--- a/Engine/Source/Runtime/Engine/Public/MaterialShaderType.h
+++ b/Engine/Source/Runtime/Engine/Public/MaterialShaderType.h
@@ -27,11 +27,13 @@ class FShaderCompileJob;
 class FUniformExpressionSet;
 class FVertexFactoryType;
 
+DECLARE_DELEGATE_RetVal_OneParam(FString, FShadingModelToStringDelegate, EMaterialShadingModel)
+
 /** Converts an EMaterialShadingModel to a string description. */
 extern ENGINE_API FString GetShadingModelString(EMaterialShadingModel ShadingModel);
 
-/** Converts an FMaterialShadingModelField to a string description, base on the passed in function. */
-extern ENGINE_API FString GetShadingModelFieldString(FMaterialShadingModelField ShadingModels, FString(*ShadingModelToString)(EMaterialShadingModel), const FString& Delimiter = " ");
+/** Converts an FMaterialShadingModelField to a string description, base on the passed in delegate. */
+extern ENGINE_API FString GetShadingModelFieldString(FMaterialShadingModelField ShadingModels, const FShadingModelToStringDelegate& Delegate, const FString& Delimiter = " ");
 
 /** Converts an FMaterialShadingModelField to a string description, base on a default function. */
 extern ENGINE_API FString GetShadingModelFieldString(FMaterialShadingModelField ShadingModels);
