Changed the rebuild of the shading model field to recurse through all material expressions. This is a little naive, since it includes orphaned material expressions too but it does the job.

#rnx

diff --git a/Engine/Source/Runtime/Engine/Classes/Materials/Material.h b/Engine/Source/Runtime/Engine/Classes/Materials/Material.h
--- a/Engine/Source/Runtime/Engine/Classes/Materials/Material.h
+++ b/Engine/Source/Runtime/Engine/Classes/Materials/Material.h
@@ -414,6 +414,12 @@ private:
 	UPROPERTY(AssetRegistrySearchable)
 	FMaterialShadingModelField ShadingModels;
 
+#if WITH_EDITORONLY_DATA
+	/** These are the shading models present in this material. Note that all these shading models might not be used in all feature levels and quality levels. */
+	UPROPERTY(VisibleAnywhere, Transient, Category=Material)
+	FString UsedShadingModels;
+#endif
+
 public:
 
 	/** If BlendMode is BLEND_Masked, the surface is not rendered where OpacityMask < OpacityMaskClipValue. */
@@ -1560,10 +1566,10 @@ public:
 	 * @param	Expression	The expression dynamic parameter to check for duplicates.
 	 */
 	ENGINE_API virtual bool HasDuplicateDynamicParameters(const UMaterialExpression* Expression);
+#endif // WITH_EDITOR
 
-	/** Walk the material expression graph from the Shading Model input pin (or from the Material Attribute) and find all used Shading Models for this material and update the field in this material. */
+	/** Collect all material expressions fomr this material and all its functions and figure out which possible shading models exist in this material */
 	ENGINE_API void RebuildShadingModelField();
-#endif // WITH_EDITOR
 
 	/**
 	 * Iterate through all of the expression nodes and fix up changed properties on
diff --git a/Engine/Source/Runtime/Engine/Private/Materials/Material.cpp b/Engine/Source/Runtime/Engine/Private/Materials/Material.cpp
--- a/Engine/Source/Runtime/Engine/Private/Materials/Material.cpp
+++ b/Engine/Source/Runtime/Engine/Private/Materials/Material.cpp
@@ -52,6 +52,7 @@
 #include "MaterialShaderQualitySettings.h"
 #include "ProfilingDebugging/LoadTimeTracker.h"
 #include "MaterialCompiler.h"
+#include "MaterialShaderType.h"
 #include "Materials/MaterialInstanceSupport.h"
 #include "Interfaces/ITargetPlatformManagerModule.h"
 #include "Interfaces/ITargetPlatform.h"
@@ -3059,11 +3060,8 @@ void UMaterial::UpdateMaterialShaderCacheAndTextureReferences()
 
 void UMaterial::CacheResourceShadersForRendering(bool bRegenerateId)
 {
-	
-#if WITH_EDITOR
 	// Always rebuild the shading model field on recompile
 	RebuildShadingModelField();
-#endif //WITH_EDITOR
 
 	if (bRegenerateId)
 	{
@@ -4626,31 +4624,23 @@ void UMaterial::UpdateExpressionParameterName(UMaterialExpression* Expression)
 		}
 	}
 }
+#endif // WITH_EDITOR
 
 void UMaterial::RebuildShadingModelField()
 {
 	ShadingModels.ClearShadingModels();
-	// If using shading model from material expression, go through the expressions and look for the ShadingModel expression to figure out what shading models need to be supported in this material
+
+	// If using shading model from material expression, go through the expressions and look for the ShadingModel expression to figure out what shading models need to be supported in this material.
+	// This might not be the same as what is actually compiled in to the shader, since there might be feature switches, static switches etc. that skip certain shading models.
 	if (ShadingModel == MSM_FromMaterialExpression)
 	{
-		TArray<UMaterialExpression*> ShadingModelPropertyExpressions;
 
-		// Depending on if we use Material Attributes or not, we have to start at the right property
-		if (bUseMaterialAttributes)
-		{
-			GetExpressionsInPropertyChain(MP_MaterialAttributes, ShadingModelPropertyExpressions, nullptr);
-		}
-		else
-		{
-			GetExpressionsInPropertyChain(MP_ShadingModel, ShadingModelPropertyExpressions, nullptr);
-		}
+		TArray<UMaterialExpressionShadingModel*> ShadingModelExpressions;
+		GetAllExpressionsInMaterialAndFunctionsOfType(ShadingModelExpressions);
 
-		for (UMaterialExpression* MatExpr: ShadingModelPropertyExpressions)
-		{
-			if (UMaterialExpressionShadingModel* ShadingModelExpression = Cast<UMaterialExpressionShadingModel>(MatExpr))
+		for (UMaterialExpressionShadingModel* MatExpr : ShadingModelExpressions)
 			{
-				ShadingModels.AddShadingModel(ShadingModelExpression->ShadingModel);
-			}
+			ShadingModels.AddShadingModel(MatExpr->ShadingModel);
 		}
 
 		// If no expressions have been found, set a default
@@ -4664,8 +4654,17 @@ void UMaterial::RebuildShadingModelField()
 		// If a shading model has been selected directly for the material, set it here
 		ShadingModels.AddShadingModel(ShadingModel);
 	}
+
+#if WITH_EDITORONLY_DATA
+	// Build a string with all the shading models on this material. Used to display the used shading models in this material
+	auto ShadingModelToStringLambda = 
+	[](EMaterialShadingModel InShadingModel) -> FString
+	{ 
+		return StaticEnum<EMaterialShadingModel>()->GetDisplayNameTextByValue(InShadingModel).ToString();
+	};
+	UsedShadingModels = GetShadingModelFieldString(ShadingModels, FShadingModelToStringDelegate::CreateLambda(ShadingModelToStringLambda), " | ");
+#endif
 }
-#endif // WITH_EDITOR
 
 bool UMaterial::GetExpressionParameterName(const UMaterialExpression* Expression, FName& OutName)
 {
