Fixed a case where cached mesh draw commands could end up with deleted sky reflection UB on mobile

diff --git a/Engine/Source/Runtime/Engine/Private/Components/SkyLightComponent.cpp b/Engine/Source/Runtime/Engine/Private/Components/SkyLightComponent.cpp
--- a/Engine/Source/Runtime/Engine/Private/Components/SkyLightComponent.cpp
+++ b/Engine/Source/Runtime/Engine/Private/Components/SkyLightComponent.cpp
@@ -158,33 +158,6 @@ FLinearColor FSkyLightSceneProxy::GetEffectiveLightColor() const
 	return LightColor * GSkylightIntensityMultiplier;
 }
 
-void FSkyLightSceneProxy::UpdateMobileUniformBuffer()
-{
-	float InvBrightness = FMath::Max(FMath::Min(1.0f / AverageBrightness, 65504.f), -65504.f);
-	float SkyMaxMipIndex = 0.f;
-	FTexture* CaptureTexture = GBlackTextureCube;
-	if (ProcessedTexture)
-	{
-		check(ProcessedTexture->IsInitialized());
-		CaptureTexture = ProcessedTexture;
-		SkyMaxMipIndex = FMath::Log2(ProcessedTexture->GetSizeX());
-	}
-		
-	FMobileReflectionCaptureShaderParameters Parameters;
-	Parameters.Params = FVector4(InvBrightness, SkyMaxMipIndex, 0.f, 0.f);
-	Parameters.Texture = CaptureTexture->TextureRHI;
-	Parameters.TextureSampler = CaptureTexture->SamplerStateRHI;
-
-	if (MobileUniformBuffer.GetReference())
-	{
-		MobileUniformBuffer.UpdateUniformBufferImmediate(Parameters);
-	}
-	else
-	{
-		MobileUniformBuffer = TUniformBufferRef<FMobileReflectionCaptureShaderParameters>::CreateUniformBufferImmediate(Parameters, UniformBuffer_MultiFrame);
-	}
-}
-
 FSkyLightSceneProxy::FSkyLightSceneProxy(const USkyLightComponent* InLightComponent)
 	: LightComponent(InLightComponent)
 	, ProcessedTexture(InLightComponent->ProcessedSkyTexture)
diff --git a/Engine/Source/Runtime/Engine/Public/SceneManagement.h b/Engine/Source/Runtime/Engine/Public/SceneManagement.h
--- a/Engine/Source/Runtime/Engine/Public/SceneManagement.h
+++ b/Engine/Source/Runtime/Engine/Public/SceneManagement.h
@@ -1008,9 +1008,6 @@ public:
 	float MinOcclusion;
 	FLinearColor OcclusionTint;
 	int32 SamplesPerPixel;
-	/** Used with mobile renderer */
-	TUniformBufferRef<FMobileReflectionCaptureShaderParameters> MobileUniformBuffer;
-
 
 #if RHI_RAYTRACING
 	bool IsDirtyImportanceSamplingData;
@@ -1043,7 +1040,6 @@ public:
 	}
 	FLinearColor GetEffectiveLightColor() const;
 
-	void UpdateMobileUniformBuffer();
 private:
 	FLinearColor LightColor;
 };
diff --git a/Engine/Source/Runtime/Renderer/Private/MobileBasePass.cpp b/Engine/Source/Runtime/Renderer/Private/MobileBasePass.cpp
--- a/Engine/Source/Runtime/Renderer/Private/MobileBasePass.cpp
+++ b/Engine/Source/Runtime/Renderer/Private/MobileBasePass.cpp
@@ -486,7 +486,7 @@ void TMobileBasePassPSPolicyParamType<FUniformLightMapPolicy>::GetShaderBindings
 			// If no reflection captures are available then attempt to use sky light's texture.
 			if (UseSkyReflectionCapture(Scene) && FeatureLevel > ERHIFeatureLevel::ES2) // not-supported on ES2 at the moment
 			{
-				RelfectionUB = Scene->SkyLight->MobileUniformBuffer;
+				RelfectionUB = Scene->UniformBuffers.MobileSkyReflectionUniformBuffer;
 			}
 			else
 			{
diff --git a/Engine/Source/Runtime/Renderer/Private/MobileBasePassRendering.cpp b/Engine/Source/Runtime/Renderer/Private/MobileBasePassRendering.cpp
--- a/Engine/Source/Runtime/Renderer/Private/MobileBasePassRendering.cpp
+++ b/Engine/Source/Runtime/Renderer/Private/MobileBasePassRendering.cpp
@@ -221,6 +221,25 @@ void SetupMobileDirectionalLightUniformParameters(
 	}
 }
 
+void SetupMobileSkyReflectionUniformParameters(FSkyLightSceneProxy* SkyLight, FMobileReflectionCaptureShaderParameters& Parameters)
+{
+	float InvBrightness = 1.f;
+	float SkyMaxMipIndex = 0.f;
+	FTexture* CaptureTexture = GBlackTextureCube;
+
+	if (SkyLight && SkyLight->ProcessedTexture)
+	{
+		check(SkyLight->ProcessedTexture->IsInitialized());
+		CaptureTexture = SkyLight->ProcessedTexture;
+		SkyMaxMipIndex = FMath::Log2(CaptureTexture->GetSizeX());
+		InvBrightness = FMath::Max(FMath::Min(1.0f / SkyLight->AverageBrightness, 65504.f), -65504.f);
+	}
+		
+	Parameters.Params = FVector4(InvBrightness, SkyMaxMipIndex, 0.f, 0.f);
+	Parameters.Texture = CaptureTexture->TextureRHI;
+	Parameters.TextureSampler = CaptureTexture->SamplerStateRHI;
+}
+
 void FMobileSceneRenderer::RenderMobileBasePass(FRHICommandListImmediate& RHICmdList, const TArrayView<const FViewInfo*> PassViews)
 {
 	CSV_SCOPED_TIMING_STAT_EXCLUSIVE(RenderBasePass);
diff --git a/Engine/Source/Runtime/Renderer/Private/MobileBasePassRendering.h b/Engine/Source/Runtime/Renderer/Private/MobileBasePassRendering.h
--- a/Engine/Source/Runtime/Renderer/Private/MobileBasePassRendering.h
+++ b/Engine/Source/Runtime/Renderer/Private/MobileBasePassRendering.h
@@ -49,6 +49,12 @@ extern void SetupMobileDirectionalLightUniformParameters(
 	bool bDynamicShadows,
 	FMobileDirectionalLightShaderParameters& Parameters);
 
+extern void SetupMobileSkyReflectionUniformParameters(
+	class FSkyLightSceneProxy* SkyLight,
+	FMobileReflectionCaptureShaderParameters& Parameters);
+
+
+
 class FPlanarReflectionSceneProxy;
 class FScene;
 
diff --git a/Engine/Source/Runtime/Renderer/Private/MobileShadingRenderer.cpp b/Engine/Source/Runtime/Renderer/Private/MobileShadingRenderer.cpp
--- a/Engine/Source/Runtime/Renderer/Private/MobileShadingRenderer.cpp
+++ b/Engine/Source/Runtime/Renderer/Private/MobileShadingRenderer.cpp
@@ -282,6 +282,7 @@ void FMobileSceneRenderer::InitViews(FRHICommandListImmediate& RHICmdList)
 		SetupMobileDistortionPassUniformBuffer(RHICmdList, View, DistortionPassParameters);
 		Scene->UniformBuffers.MobileDistortionPassUniformBuffer.UpdateUniformBufferImmediate(DistortionPassParameters);
 	}
+	UpdateSkyReflectionUniformBuffer();
 
 	// Now that the indirect lighting cache is updated, we can update the uniform buffers.
 	UpdatePrimitiveIndirectLightingCacheBuffers();
@@ -875,6 +876,21 @@ void FMobileSceneRenderer::UpdateDirectionalLightUniformBuffers(FRHICommandListI
 	}
 }
 
+void FMobileSceneRenderer::UpdateSkyReflectionUniformBuffer()
+{
+	FSkyLightSceneProxy* SkyLight = nullptr;
+	if (Scene->ReflectionSceneData.RegisteredReflectionCapturePositions.Num() == 0
+		&& Scene->SkyLight
+		&& Scene->SkyLight->ProcessedTexture->TextureRHI)
+	{
+		SkyLight = Scene->SkyLight;
+	}
+
+	FMobileReflectionCaptureShaderParameters Parameters;
+	SetupMobileSkyReflectionUniformParameters(SkyLight, Parameters);
+	Scene->UniformBuffers.MobileSkyReflectionUniformBuffer.UpdateUniformBufferImmediate(Parameters);
+}
+
 void FMobileSceneRenderer::CreateDirectionalLightUniformBuffers(FViewInfo& View)
 {
 	bool bDynamicShadows = ViewFamily.EngineShowFlags.DynamicShadows;
diff --git a/Engine/Source/Runtime/Renderer/Private/RendererScene.cpp b/Engine/Source/Runtime/Renderer/Private/RendererScene.cpp
--- a/Engine/Source/Runtime/Renderer/Private/RendererScene.cpp
+++ b/Engine/Source/Runtime/Renderer/Private/RendererScene.cpp
@@ -981,7 +981,9 @@ void FPersistentUniformBuffers::Initialize()
 	{
 		MobileDirectionalLightUniformBuffers[Index] = TUniformBufferRef<FMobileDirectionalLightShaderParameters>::CreateUniformBufferImmediate(MobileDirectionalLightShaderParameters, UniformBuffer_MultiFrame, EUniformBufferValidation::None);
 	}
-	
+
+	FMobileReflectionCaptureShaderParameters MobileSkyReflectionShaderParameters;
+	MobileSkyReflectionUniformBuffer = TUniformBufferRef<FMobileReflectionCaptureShaderParameters>::CreateUniformBufferImmediate(MobileSkyReflectionShaderParameters, UniformBuffer_MultiFrame, EUniformBufferValidation::None);
 
 #if WITH_EDITOR
 	FSceneTexturesUniformParameters EditorSelectionPassParameters;
@@ -1853,11 +1855,6 @@ void FScene::SetSkyLight(FSkyLightSceneProxy* LightProxy)
 				// The base pass chooses shaders based on whether there's a skylight in the scene, and that is cached in static draw lists
 				Scene->bScenesPrimitivesNeedStaticMeshElementUpdate = true;
 			}
-
-			if (Scene->GetFeatureLevel() <= ERHIFeatureLevel::ES3_1)
-			{
-				Scene->SkyLight->UpdateMobileUniformBuffer();
-			}
 		});
 }
 
diff --git a/Engine/Source/Runtime/Renderer/Private/ScenePrivate.h b/Engine/Source/Runtime/Renderer/Private/ScenePrivate.h
--- a/Engine/Source/Runtime/Renderer/Private/ScenePrivate.h
+++ b/Engine/Source/Runtime/Renderer/Private/ScenePrivate.h
@@ -2384,6 +2384,7 @@ public:
 	  * The first is used for primitives with no lighting channels set.
 	  */
 	TUniformBufferRef<FMobileDirectionalLightShaderParameters> MobileDirectionalLightUniformBuffers[NUM_LIGHTING_CHANNELS+1];
+	TUniformBufferRef<FMobileReflectionCaptureShaderParameters> MobileSkyReflectionUniformBuffer;
 
 #if WITH_EDITOR
 	TUniformBufferRef<FSceneTexturesUniformParameters> EditorSelectionPassUniformBuffer;
diff --git a/Engine/Source/Runtime/Renderer/Private/SceneRendering.h b/Engine/Source/Runtime/Renderer/Private/SceneRendering.h
--- a/Engine/Source/Runtime/Renderer/Private/SceneRendering.h
+++ b/Engine/Source/Runtime/Renderer/Private/SceneRendering.h
@@ -1649,6 +1649,7 @@ protected:
 	void UpdateOpaqueBasePassUniformBuffer(FRHICommandListImmediate& RHICmdList, const FViewInfo& View);
 	void UpdateTranslucentBasePassUniformBuffer(FRHICommandListImmediate& RHICmdList, const FViewInfo& View);
 	void UpdateDirectionalLightUniformBuffers(FRHICommandListImmediate& RHICmdList, const FViewInfo& View);
+	void UpdateSkyReflectionUniformBuffer();
 	
 private:
 	bool bModulatedShadowsInUse;
