Enabling back this change (CL 6444236) and disable assert on gamethread.

===================================
Port CL# 6414161 to Release-4.22

Fix morph target reset for 1 frame when changing material
- removed ActiveMorphtarget clearing on create render state because we're clearing in Destroy Render state
- We update morphtarget info if not disabled and allowed, so that we don't lose the current acitve target data

- jira UE-54099

====================================

The assert caused following jiras

diff --git a/Engine/Source/Runtime/Engine/Private/Animation/AnimInstance.cpp b/Engine/Source/Runtime/Engine/Private/Animation/AnimInstance.cpp
--- a/Engine/Source/Runtime/Engine/Private/Animation/AnimInstance.cpp
+++ b/Engine/Source/Runtime/Engine/Private/Animation/AnimInstance.cpp
@@ -1180,7 +1180,10 @@ void UAnimInstance::UpdateCurvesToComponents(USkeletalMeshComponent* Component /
 	// update curves to component
 	if (Component)
 	{
-		FAnimInstanceProxy& Proxy = GetProxyOnGameThread<FAnimInstanceProxy>();
+		// this is only any thread because EndOfFrameUpdate update can restart render state
+		// and this needs to restart from worker thread
+		// EndOfFrameUpdate is done after all tick is updated, so in theory you shouldn't have 
+		FAnimInstanceProxy& Proxy = GetProxyOnAnyThread<FAnimInstanceProxy>();
 		Component->ApplyAnimationCurvesToComponent(&Proxy.GetAnimationCurves(EAnimCurveType::MaterialCurve), &Proxy.GetAnimationCurves(EAnimCurveType::MorphTargetCurve));
 	}
 }
diff --git a/Engine/Source/Runtime/Engine/Private/Components/SkinnedMeshComponent.cpp b/Engine/Source/Runtime/Engine/Private/Components/SkinnedMeshComponent.cpp
--- a/Engine/Source/Runtime/Engine/Private/Components/SkinnedMeshComponent.cpp
+++ b/Engine/Source/Runtime/Engine/Private/Components/SkinnedMeshComponent.cpp
@@ -586,9 +586,9 @@ void USkinnedMeshComponent::CreateRenderState_Concurrent()
 				const bool bMorphTargetsAllowed = CVarEnableMorphTargets.GetValueOnAnyThread(true) != 0;
 
 				// Are morph targets disabled for this LOD?
-				if (bDisableMorphTarget || !bMorphTargetsAllowed)
+				if (!bDisableMorphTarget && bMorphTargetsAllowed)
 				{
-					ActiveMorphTargets.Empty();
+					RefreshMorphTargets();
 				}
 
 				MeshObject->Update(PredictedLODLevel, this, ActiveMorphTargets, MorphTargetWeights, EPreviousBoneTransformUpdateMode::UpdatePrevious);  // send to rendering thread
