Simplify the work done to transform the rotation when applying level transforms.
#rnx

diff --git a/Engine/Source/Runtime/Engine/Private/LevelUtils.cpp b/Engine/Source/Runtime/Engine/Private/LevelUtils.cpp
--- a/Engine/Source/Runtime/Engine/Private/LevelUtils.cpp
+++ b/Engine/Source/Runtime/Engine/Private/LevelUtils.cpp
@@ -425,14 +425,14 @@ void FLevelUtils::ApplyLevelTransform( ULevel* Level, const FTransform& LevelTra
 		{
 			AActor* Actor = Level->Actors[ActorIndex];
 
-			// Don't want to transform children they should stay relative to there parents.
-			if( Actor && Actor->GetAttachParentActor() == NULL )
+			// Don't want to transform children they should stay relative to their parents.
+			if( Actor && Actor->GetAttachParentActor() == nullptr )
 			{
 				// Has to modify root component directly as GetActorPosition is incorrect this early
 				USceneComponent *RootComponent = Actor->GetRootComponent();
 				if (RootComponent)
 				{
-					RootComponent->SetRelativeLocationAndRotation( LevelTransform.TransformPosition(RootComponent->RelativeLocation), (FTransform(RootComponent->RelativeRotation) * LevelTransform).Rotator());
+					RootComponent->SetRelativeLocationAndRotation( LevelTransform.TransformPosition(RootComponent->RelativeLocation), LevelTransform.TransformRotation(RootComponent->RelativeRotation.Quaternion()) );
 				}			
 			}
 		}
