Fixed shader code gen bug in the Custom Primitive Data material expression. No need to recompile all shaders.

#rnx

diff --git a/Engine/Source/Runtime/Engine/Private/Materials/HLSLMaterialTranslator.h b/Engine/Source/Runtime/Engine/Private/Materials/HLSLMaterialTranslator.h
--- a/Engine/Source/Runtime/Engine/Private/Materials/HLSLMaterialTranslator.h
+++ b/Engine/Source/Runtime/Engine/Private/Materials/HLSLMaterialTranslator.h
@@ -5503,7 +5503,7 @@ protected:
 			// Check if we are accessing inside the array, otherwise default to 0
 			if (CurrentOutputIndex < FCustomPrimitiveData::NumCustomPrimitiveDataFloats)
 			{
-				const int32 CustomDataIndex = CurrentOutputIndex / FCustomPrimitiveData::NumCustomPrimitiveDataFloat4s;
+				const int32 CustomDataIndex = CurrentOutputIndex / 4;
 				const int32 ElementIndex = CurrentOutputIndex % 4; // Index x, y, z or w
 
 				HlslCode.Append(FString::Printf(TEXT("GetPrimitiveData(Parameters.PrimitiveId).CustomPrimitiveData[%d][%d]"), CustomDataIndex, ElementIndex));
