Store the number of loading levels on the world rather than iterating seeking out the information
#rnx

diff --git a/Engine/Source/Runtime/Engine/Classes/Engine/World.h b/Engine/Source/Runtime/Engine/Classes/Engine/World.h
--- a/Engine/Source/Runtime/Engine/Classes/Engine/World.h
+++ b/Engine/Source/Runtime/Engine/Classes/Engine/World.h
@@ -1484,7 +1484,14 @@ public:
 
 	/** The type of travel to perform next when doing a server travel */
 	TEnumAsByte<ETravelType> NextTravelType;
-	
+
+private:
+	/** An internal count of how many streaming levels are currently loading */
+	uint16 NumStreamingLevelsBeingLoaded; // Move this somewhere better
+
+	friend struct FWorldNotifyStreamingLevelLoading;
+
+public:
 	/** The URL to be used for the upcoming server travel */
 	FString NextURL;
 
@@ -3583,6 +3590,25 @@ private:
 	FWorldDelegates() {}
 };
 
+/** Helper struct to allow ULevelStreaming to update its World on how many streaming levels are being loaded */
+struct FWorldNotifyStreamingLevelLoading
+{
+private:
+	static void Started(UWorld* World)
+	{
+		++World->NumStreamingLevelsBeingLoaded;
+	}
+
+	static void Finished(UWorld* World)
+	{
+		if (ensure(World->NumStreamingLevelsBeingLoaded > 0))
+		{
+			--World->NumStreamingLevelsBeingLoaded;
+		}
+	}
+
+	friend ULevelStreaming;
+};
 
 //////////////////////////////////////////////////////////////////////////
 // UWorld inlines:
diff --git a/Engine/Source/Runtime/Engine/Private/LevelStreaming.cpp b/Engine/Source/Runtime/Engine/Private/LevelStreaming.cpp
--- a/Engine/Source/Runtime/Engine/Private/LevelStreaming.cpp
+++ b/Engine/Source/Runtime/Engine/Private/LevelStreaming.cpp
@@ -1001,6 +1001,7 @@ bool ULevelStreaming::RequestLevel(UWorld* PersistentWorld, bool bAllowLevelLoad
 		if (FPackageName::DoesPackageExist(PackageNameToLoadFrom))
 		{
 			CurrentState = ECurrentState::Loading;
+			FWorldNotifyStreamingLevelLoading::Started(PersistentWorld);
 			
 			ULevel::StreamedLevelsOwningWorld.Add(DesiredPackageName, PersistentWorld);
 			UWorld::WorldTypePreLoadMap.FindOrAdd(DesiredPackageName) = PersistentWorld->WorldType;
@@ -1036,6 +1037,13 @@ bool ULevelStreaming::RequestLevel(UWorld* PersistentWorld, bool bAllowLevelLoad
 void ULevelStreaming::AsyncLevelLoadComplete(const FName& InPackageName, UPackage* InLoadedPackage, EAsyncLoadingResult::Type Result)
 {
 	CurrentState = ECurrentState::LoadedNotVisible;
+	if (UWorld* World = GetWorld())
+	{
+		if (World->GetStreamingLevels().Contains(this))
+		{
+			FWorldNotifyStreamingLevelLoading::Finished(World);
+		}
+	}
 
 	if (InLoadedPackage)
 	{
diff --git a/Engine/Source/Runtime/Engine/Private/World.cpp b/Engine/Source/Runtime/Engine/Private/World.cpp
--- a/Engine/Source/Runtime/Engine/Private/World.cpp
+++ b/Engine/Source/Runtime/Engine/Private/World.cpp
@@ -3221,6 +3221,10 @@ bool UWorld::RemoveStreamingLevelAt(const int32 IndexToRemove)
 	if (IndexToRemove >= 0 && IndexToRemove < StreamingLevels.Num())
 	{
 		ULevelStreaming* StreamingLevel = StreamingLevels[IndexToRemove];
+		if (StreamingLevel && StreamingLevel->GetCurrentState() == ULevelStreaming::ECurrentState::Loading)
+		{
+			--NumStreamingLevelsBeingLoaded;
+		}
 		StreamingLevels.RemoveAt(IndexToRemove);
 		StreamingLevelsToConsider.Remove(StreamingLevel);
 		FStreamingLevelPrivateAccessor::OnLevelRemoved(StreamingLevel);
@@ -3254,10 +3258,12 @@ void UWorld::ClearStreamingLevels()
 {
 	StreamingLevels.Reset();
 	StreamingLevelsToConsider.Reset();
+	NumStreamingLevelsBeingLoaded = 0;
 }
 
 void UWorld::PopulateStreamingLevelsToConsider()
 {
+	NumStreamingLevelsBeingLoaded = 0;
 	StreamingLevelsToConsider.Reset();
 	for (ULevelStreaming* StreamingLevel : StreamingLevels)
 	{
@@ -3266,6 +3272,10 @@ void UWorld::PopulateStreamingLevelsToConsider()
 		{
 			FStreamingLevelPrivateAccessor::OnLevelAdded(StreamingLevel);
 		}
+		else if (StreamingLevel->GetCurrentState() == ULevelStreaming::ECurrentState::Loading)
+		{
+			++NumStreamingLevelsBeingLoaded;
+		}
 
 		if (FStreamingLevelPrivateAccessor::DetermineTargetState(StreamingLevel))
 		{
@@ -3457,21 +3467,9 @@ bool UWorld::AllowLevelLoadRequests()
 		}
 
 		// Don't allow requesting new levels if we're playing in game and already busy loading a maximum number of them.
-		if (GLevelStreamingMaxLevelRequestsAtOnceWhileInMatch > 0 && bIsPlaying)
+		if (bIsPlaying && GLevelStreamingMaxLevelRequestsAtOnceWhileInMatch > 0 && NumStreamingLevelsBeingLoaded >= GLevelStreamingMaxLevelRequestsAtOnceWhileInMatch)
 		{
-			int32 NumLoadingLevels = 0;
-			for (ULevelStreaming* LevelStreaming : StreamingLevels)
-			{
-				if (LevelStreaming && LevelStreaming->HasLoadRequestPending())
-				{
-					++NumLoadingLevels;
-				}
-			}
-
-			if (NumLoadingLevels >= GLevelStreamingMaxLevelRequestsAtOnceWhileInMatch)
-			{
-				return false;
-			}
+			return false;
 		}
 	}
 
