Content encryption can now specify an allowed references rule for a group rather than being specific about only allowing soft references.

diff --git a/Engine/Source/Runtime/Engine/Classes/Engine/ContentEncryptionConfig.h b/Engine/Source/Runtime/Engine/Classes/Engine/ContentEncryptionConfig.h
--- a/Engine/Source/Runtime/Engine/Classes/Engine/ContentEncryptionConfig.h
+++ b/Engine/Source/Runtime/Engine/Classes/Engine/ContentEncryptionConfig.h
@@ -6,11 +6,18 @@ class FContentEncryptionConfig
 {
 public:
 
+	enum class EAllowedReferences
+	{
+		None,
+		Soft,
+		All,
+	};
+
 	struct FGroup
 	{
 		TSet<FName> PackageNames;
 		bool bStageTimeOnly = false;
-		bool bIgnoreSoftReferences = false;
+		EAllowedReferences AllowedReferences = EAllowedReferences::None;
 	};
 
 	typedef TMap<FName, FGroup> TGroupMap;
@@ -25,9 +32,9 @@ public:
 		PackageGroups.FindOrAdd(InGroupName).bStageTimeOnly = bInStageTimeOnly;
 	}
 
-	void SetShouldIgnoreSoftReferences(FName InGroupName, bool bInIgnoreSoftReferences)
+	void SetAllowedReferences(FName InGroupName, EAllowedReferences InAllowedReferences)
 	{
-		PackageGroups.FindOrAdd(InGroupName).bIgnoreSoftReferences = bInIgnoreSoftReferences;
+		PackageGroups.FindOrAdd(InGroupName).AllowedReferences = InAllowedReferences;
 	}
 
 	void AddReleasedKey(FGuid InKey)
