minor clean-up
pull request: https://github.com/EpicGames/UnrealEngine/pull/5801

diff --git a/Engine/Source/Runtime/Engine/Private/Components/AudioComponent.cpp b/Engine/Source/Runtime/Engine/Private/Components/AudioComponent.cpp
--- a/Engine/Source/Runtime/Engine/Private/Components/AudioComponent.cpp
+++ b/Engine/Source/Runtime/Engine/Private/Components/AudioComponent.cpp
@@ -320,7 +320,7 @@ bool UAudioComponent::IsInAudibleRange(float* OutMaxDistance) const
 	const FVector Location = GetComponentTransform().GetLocation();
 	const FSoundAttenuationSettings* AttenuationSettingsToApply = bAllowSpatialization ? GetAttenuationSettingsToApply() : nullptr;
 	AudioDevice->GetMaxDistanceAndFocusFactor(Sound, GetWorld(), Location, AttenuationSettingsToApply, MaxDistance, FocusFactor);
-	
+
 	if (OutMaxDistance)
 	{
 		*OutMaxDistance = MaxDistance;
@@ -486,10 +486,13 @@ void UAudioComponent::PlayInternal(const float StartTime, const float FadeInDura
 				NewActiveSound.CurrentAdjustVolumeMultiplier = FadeVolumeLevel;
 			}
 
-			// Bump ActiveCount... this is used to determine if an audio component is still active after "finishing"
+			// Bump ActiveCount... this is used to determine if an audio component is still active after a sound reports back as completed
 			++ActiveCount;
 			AudioDevice->AddNewActiveSound(NewActiveSound);
-			bIsActive = true;
+
+			// In editor, the audio thread is not run separate from the game thread, and can result in calling PlaybackComplete prior
+			// to bIsActive being set. Therefore, we assign to the current state of ActiveCount as opposed to just setting to true.
+			bIsActive = ActiveCount > 0;
 		}
 	}
 }
@@ -655,10 +658,7 @@ void UAudioComponent::PlaybackCompleted(bool bFailedToStart)
 	}
 
 	// Mark inactive before calling destroy to avoid recursion
-	if (bIsActive)
-	{
-		bIsActive = false;
-	}
+	bIsActive = false;
 
 	if (!bFailedToStart && GetWorld() != nullptr && (OnAudioFinished.IsBound() || OnAudioFinishedNative.IsBound()))
 	{
@@ -1324,7 +1324,7 @@ bool UAudioComponent::GetCookedFFTData(const TArray<float>& FrequenciesToGet, TA
 							}
 						}
 					}
-					
+
 					++NumEntriesAdded;
 					bHadData = true;
 				}
