#UE4 Upgrade compile errors to Display verbosity so they will show up in more outputs

diff --git a/Engine/Source/Runtime/Engine/Private/Materials/MaterialInstance.cpp b/Engine/Source/Runtime/Engine/Private/Materials/MaterialInstance.cpp
--- a/Engine/Source/Runtime/Engine/Private/Materials/MaterialInstance.cpp
+++ b/Engine/Source/Runtime/Engine/Private/Materials/MaterialInstance.cpp
@@ -2700,7 +2700,7 @@ void UMaterialInstance::CacheShadersForResources(EShaderPlatform ShaderPlatform,
 			const TArray<FString>& CompileErrors = CurrentResource->GetCompileErrors();
 			for (int32 ErrorIndex = 0; ErrorIndex < CompileErrors.Num(); ErrorIndex++)
 			{
-				UE_LOG(LogMaterial, Log, TEXT("	%s"), *CompileErrors[ErrorIndex]);
+				UE_LOG(LogMaterial, Display, TEXT("	%s"), *CompileErrors[ErrorIndex]);
 			}
 #endif // WITH_EDITOR
 		}
