Animation: Allowed TargetSkeleton to be edited on Animation Blueprints, and updated the comment on it
- This allows sharing one parent anim BP on different skeletons by creating child anim BPs *as long as* some additional step or process ensures the replacement of all referenced anims with new ones that are compatible with the new skeleton (e.g., via a data driven substitution system like an anim set, or thorough use of anim overrides)

diff --git a/Engine/Source/Runtime/Engine/Classes/Animation/AnimBlueprint.h b/Engine/Source/Runtime/Engine/Classes/Animation/AnimBlueprint.h
--- a/Engine/Source/Runtime/Engine/Classes/Animation/AnimBlueprint.h
+++ b/Engine/Source/Runtime/Engine/Classes/Animation/AnimBlueprint.h
@@ -12,6 +12,7 @@
 class SWidget;
 class UAnimationAsset;
 class USkeletalMesh;
+class USkeleton;
 
 USTRUCT()
 struct FAnimGroupInfo
@@ -65,9 +66,13 @@ class ENGINE_API UAnimBlueprint : public UBlueprint, public IInterface_PreviewMe
 {
 	GENERATED_UCLASS_BODY()
 
-	/** The kind of skeleton that animation graphs compiled from the blueprint will animate */
-	UPROPERTY(AssetRegistrySearchable)
-	class USkeleton* TargetSkeleton;
+	/**
+	 * This is the target skeleton asset for anim instances created from this blueprint; all animations
+	 * referenced by the BP should be compatible with this skeleton.  For advanced use only, it is easy
+	 * to cause errors if this is modified without updating or replacing all referenced animations.
+	 */
+	UPROPERTY(AssetRegistrySearchable, EditAnywhere, AdvancedDisplay, Category=ClassOptions)
+	USkeleton* TargetSkeleton;
 
 	// List of animation sync groups
 	UPROPERTY()
