Reworked the Pre-Skinned Local Bounds node. Now outputs Half Extents, Extents, Min or Max of the bounding box. Also added tootips to the pins. This was added to the 12 bytes of padding in the primitive uniform buffer, so no increase in size. This is compatible with previous usage of the node.

diff --git a/Engine/Shaders/Private/SceneData.ush b/Engine/Shaders/Private/SceneData.ush
--- a/Engine/Shaders/Private/SceneData.ush
+++ b/Engine/Shaders/Private/SceneData.ush
@@ -34,10 +34,10 @@ struct FPrimitiveSceneData
 	uint LightingChannelMask;
 	float3 LocalObjectBoundsMax;
 	uint LightmapDataIndex;
-	float3 PreSkinnedLocalBounds;
+	float3 PreSkinnedLocalBoundsMin;
 	int SingleCaptureIndex;
+	float3 PreSkinnedLocalBoundsMax;
     uint OutputVelocity;
-	// <-- 12 byte padding here, feel free to use -->
 	float4 CustomPrimitiveData[NUM_CUSTOM_PRIMITIVE_DATA];
 };
 
@@ -95,10 +95,11 @@ FPrimitiveSceneData GetPrimitiveData(uint PrimitiveId)
 	PrimitiveData.LocalObjectBoundsMax = View.PrimitiveSceneData[PrimitiveBaseOffset + 24].xyz;
 	PrimitiveData.LightmapDataIndex = asuint(View.PrimitiveSceneData[PrimitiveBaseOffset + 24].w);
 
-	PrimitiveData.PreSkinnedLocalBounds = View.PrimitiveSceneData[PrimitiveBaseOffset + 25].xyz;
+	PrimitiveData.PreSkinnedLocalBoundsMin = View.PrimitiveSceneData[PrimitiveBaseOffset + 25].xyz;
 	PrimitiveData.SingleCaptureIndex = View.PrimitiveSceneData[PrimitiveBaseOffset + 25].w;
 
-	PrimitiveData.OutputVelocity = asuint(View.PrimitiveSceneData[PrimitiveBaseOffset + 26].x);
+	PrimitiveData.PreSkinnedLocalBoundsMax = View.PrimitiveSceneData[PrimitiveBaseOffset + 26].xyz;
+	PrimitiveData.OutputVelocity = asuint(View.PrimitiveSceneData[PrimitiveBaseOffset + 26].w);
 
 	UNROLL
 	for (int i = 0; i < NUM_CUSTOM_PRIMITIVE_DATA; i++)
diff --git a/Engine/Source/Developer/MaterialBaking/Private/ExportMaterialProxy.h b/Engine/Source/Developer/MaterialBaking/Private/ExportMaterialProxy.h
--- a/Engine/Source/Developer/MaterialBaking/Private/ExportMaterialProxy.h
+++ b/Engine/Source/Developer/MaterialBaking/Private/ExportMaterialProxy.h
@@ -104,9 +104,9 @@ struct FExportMaterialCompiler : public FProxyMaterialCompiler
 		return Compiler->ObjectBounds();
 	}
 
-	virtual int32 PreSkinnedLocalBounds() override
+	virtual int32 PreSkinnedLocalBounds(int32 OutputIndex) override
 	{
-		return Compiler->PreSkinnedLocalBounds();
+		return Compiler->PreSkinnedLocalBounds(OutputIndex);
 	}
 
 	virtual int32 CameraVector() override
diff --git a/Engine/Source/Developer/MaterialUtilities/Private/MaterialUtilities.cpp b/Engine/Source/Developer/MaterialUtilities/Private/MaterialUtilities.cpp
--- a/Engine/Source/Developer/MaterialUtilities/Private/MaterialUtilities.cpp
+++ b/Engine/Source/Developer/MaterialUtilities/Private/MaterialUtilities.cpp
@@ -282,9 +282,9 @@ struct FExportMaterialCompiler : public FProxyMaterialCompiler
 		return Compiler->ObjectBounds();
 	}
 
-	virtual int32 PreSkinnedLocalBounds() override
+	virtual int32 PreSkinnedLocalBounds(int32 OutputIndex) override
 	{
-		return Compiler->PreSkinnedLocalBounds();
+		return Compiler->PreSkinnedLocalBounds(OutputIndex);
 	}
 
 	virtual int32 CameraVector() override
diff --git a/Engine/Source/Editor/UnrealEd/Private/Lightmass/LightmassRender.cpp b/Engine/Source/Editor/UnrealEd/Private/Lightmass/LightmassRender.cpp
--- a/Engine/Source/Editor/UnrealEd/Private/Lightmass/LightmassRender.cpp
+++ b/Engine/Source/Editor/UnrealEd/Private/Lightmass/LightmassRender.cpp
@@ -116,7 +116,7 @@ struct FLightmassMaterialCompiler : public FProxyMaterialCompiler
 		return Compiler->Constant3(0,0,0);
 	}
 
-	virtual int32 PreSkinnedLocalBounds() override
+	virtual int32 PreSkinnedLocalBounds(int32 OutputIndex) override
 	{
 		return Compiler->Constant3(0, 0, 0);
 	}
diff --git a/Engine/Source/Runtime/Engine/Classes/Materials/MaterialExpressionPreSkinnedLocalBounds.h b/Engine/Source/Runtime/Engine/Classes/Materials/MaterialExpressionPreSkinnedLocalBounds.h
--- a/Engine/Source/Runtime/Engine/Classes/Materials/MaterialExpressionPreSkinnedLocalBounds.h
+++ b/Engine/Source/Runtime/Engine/Classes/Materials/MaterialExpressionPreSkinnedLocalBounds.h
@@ -12,10 +12,16 @@ class UMaterialExpressionPreSkinnedLocalBounds : public UMaterialExpression
 {
 	GENERATED_UCLASS_BODY()
 
-	//~ Begin UMaterialExpression Interface
 #if WITH_EDITOR
+	//~ Begin UMaterialExpression Interface
 	virtual int32 Compile(class FMaterialCompiler* Compiler, int32 OutputIndex) override;
 	virtual void GetCaption(TArray<FString>& OutCaptions) const override;
-#endif
+	virtual void GetConnectorToolTip(int32 InputIndex, int32 OutputIndex, TArray<FString>& OutToolTip) override;
+	virtual void GetExpressionToolTip(TArray<FString>& OutToolTip) override;
 	//~ End UMaterialExpression Interface
+
+#if WITH_EDITORONLY_DATA
+	TArray<FString> OutputToolTips;
+#endif
+#endif
 };
\ No newline at end of file
diff --git a/Engine/Source/Runtime/Engine/Private/Materials/HLSLMaterialTranslator.h b/Engine/Source/Runtime/Engine/Private/Materials/HLSLMaterialTranslator.h
--- a/Engine/Source/Runtime/Engine/Private/Materials/HLSLMaterialTranslator.h
+++ b/Engine/Source/Runtime/Engine/Private/Materials/HLSLMaterialTranslator.h
@@ -3184,9 +3184,23 @@ protected:
 		return GetPrimitiveProperty(MCT_Float3, TEXT("ObjectBounds"), TEXT("ObjectBounds.xyz"));
 	}
 
-	virtual int32 PreSkinnedLocalBounds() override
-	{
-		return GetPrimitiveProperty(MCT_Float3, TEXT("PreSkinnedLocalBounds"), TEXT("PreSkinnedLocalBounds.xyz"));
+	virtual int32 PreSkinnedLocalBounds(int32 OutputIndex) override
+	{
+		switch (OutputIndex)
+		{
+		case 0: // Half extents
+			return AddInlinedCodeChunk(MCT_Float3, TEXT("((GetPrimitiveData(Parameters.PrimitiveId).PreSkinnedLocalBoundsMax - GetPrimitiveData(Parameters.PrimitiveId).PreSkinnedLocalBoundsMin) / 2.0f)"));
+		case 1: // Full extents
+			return AddInlinedCodeChunk(MCT_Float3, TEXT("(GetPrimitiveData(Parameters.PrimitiveId).PreSkinnedLocalBoundsMax - GetPrimitiveData(Parameters.PrimitiveId).PreSkinnedLocalBoundsMin)"));
+		case 2: // Min point
+			return GetPrimitiveProperty(MCT_Float3, TEXT("PreSkinnedLocalBounds"), TEXT("PreSkinnedLocalBoundsMin"));
+		case 3: // Max point
+			return GetPrimitiveProperty(MCT_Float3, TEXT("PreSkinnedLocalBounds"), TEXT("PreSkinnedLocalBoundsMax"));
+		default:
+			check(false);
+		}
+
+		return INDEX_NONE; 
 	}
 
 	virtual int32 DistanceCullFade() override
diff --git a/Engine/Source/Runtime/Engine/Private/Materials/MaterialExpressions.cpp b/Engine/Source/Runtime/Engine/Private/Materials/MaterialExpressions.cpp
--- a/Engine/Source/Runtime/Engine/Private/Materials/MaterialExpressions.cpp
+++ b/Engine/Source/Runtime/Engine/Private/Materials/MaterialExpressions.cpp
@@ -9650,6 +9650,17 @@ UMaterialExpressionPreSkinnedLocalBounds::UMaterialExpressionPreSkinnedLocalBoun
 	MenuCategories.Add(ConstructorStatics.NAME_Vectors);
 
 	bShaderInputData = true;
+	bShowOutputNameOnPin = true;
+
+	Outputs.Reset();
+	Outputs.Add(FExpressionOutput(TEXT("Half Extents"), 1, 1, 1, 1, 0));
+	OutputToolTips.Add("Half the extent (width, depth and height) of the pre-skinned bounding box. In local space.");
+	Outputs.Add(FExpressionOutput(TEXT("Extents"), 1, 1, 1, 1, 0));
+	OutputToolTips.Add("Full extent (width, depth and height) of the pre-skinned bounding box. Same as 2x Half Extents. In local space.");
+	Outputs.Add(FExpressionOutput(TEXT("Min"), 1, 1, 1, 1, 0));
+	OutputToolTips.Add("Minimum 3D point of the pre-skinned bounding box. In local space.");
+	Outputs.Add(FExpressionOutput(TEXT("Max"), 1, 1, 1, 1, 0));
+	OutputToolTips.Add("Maximum 3D point of the pre-skinned bounding box. In local space.");
 #endif
 }
 
@@ -9661,13 +9672,31 @@ int32 UMaterialExpressionPreSkinnedLocalBounds::Compile(class FMaterialCompiler*
 		return CompilerError(Compiler, TEXT("Expression not available in the deferred decal material domain."));
 	}
 
-	return Compiler->PreSkinnedLocalBounds();
+	return Compiler->PreSkinnedLocalBounds(OutputIndex);
 }
 
 void UMaterialExpressionPreSkinnedLocalBounds::GetCaption(TArray<FString>& OutCaptions) const
 {
 	OutCaptions.Add(TEXT("Pre-Skinned Local Bounds"));
 }
+
+void UMaterialExpressionPreSkinnedLocalBounds::GetConnectorToolTip(int32 InputIndex, int32 OutputIndex, TArray<FString>& OutToolTip) 
+{
+#if WITH_EDITORONLY_DATA
+	if (OutputIndex >= 0 && OutputIndex < OutputToolTips.Num())
+	{
+		ConvertToMultilineToolTip(OutputToolTips[OutputIndex], 40, OutToolTip);
+	}
+#endif // WITH_EDITORONLY_DATA
+}
+
+void UMaterialExpressionPreSkinnedLocalBounds::GetExpressionToolTip(TArray<FString>& OutToolTip) 
+{
+	ConvertToMultilineToolTip(TEXT("Returns various info about the pre-skinned local bounding box for skeletal meshes."
+		"Will return the regular local space bounding box for static meshes."
+		"Usable in vertex or pixel shader (no need to pipe this through vertex interpolators)."
+		"Hover the output pins for more information."), 40, OutToolTip);
+}
 #endif // WITH_EDITOR
 
 ///////////////////////////////////////////////////////////////////////////////
diff --git a/Engine/Source/Runtime/Engine/Private/PrimitiveUniformShaderParameters.cpp b/Engine/Source/Runtime/Engine/Private/PrimitiveUniformShaderParameters.cpp
--- a/Engine/Source/Runtime/Engine/Private/PrimitiveUniformShaderParameters.cpp
+++ b/Engine/Source/Runtime/Engine/Private/PrimitiveUniformShaderParameters.cpp
@@ -114,11 +114,11 @@ void FPrimitiveSceneShaderData::Setup(const FPrimitiveUniformShaderParameters& P
 	Data[24] = FVector4(PrimitiveUniformShaderParameters.LocalObjectBoundsMax, 0.0f);
 	Data[24].W = *(const float*)&PrimitiveUniformShaderParameters.LightmapDataIndex;
 
-	Data[25] = FVector4(PrimitiveUniformShaderParameters.PreSkinnedLocalBounds, 0.0f);
+	Data[25] = FVector4(PrimitiveUniformShaderParameters.PreSkinnedLocalBoundsMin, 0.0f);
 	Data[25].W = *(const float*)&PrimitiveUniformShaderParameters.SingleCaptureIndex;
 
-	Data[26] = FVector4(0.0f, 0.0f, 0.0f, 0.0f);
-	Data[26].X = *(const float*)&PrimitiveUniformShaderParameters.OutputVelocity;
+	Data[26] = FVector4(PrimitiveUniformShaderParameters.PreSkinnedLocalBoundsMax, 0.0f);
+	Data[26].W = *(const float*)&PrimitiveUniformShaderParameters.OutputVelocity;
 
 	Data[27] = PrimitiveUniformShaderParameters.CustomPrimitiveData[0];
 	Data[28] = PrimitiveUniformShaderParameters.CustomPrimitiveData[1];
diff --git a/Engine/Source/Runtime/Engine/Public/MaterialCompiler.h b/Engine/Source/Runtime/Engine/Public/MaterialCompiler.h
--- a/Engine/Source/Runtime/Engine/Public/MaterialCompiler.h
+++ b/Engine/Source/Runtime/Engine/Public/MaterialCompiler.h
@@ -140,7 +140,7 @@ public:
 	virtual int32 ObjectWorldPosition() = 0;
 	virtual int32 ObjectRadius() = 0;
 	virtual int32 ObjectBounds() = 0;
-	virtual int32 PreSkinnedLocalBounds() = 0;
+	virtual int32 PreSkinnedLocalBounds(int32 OutputIndex) = 0;
 	virtual int32 DistanceCullFade() = 0;
 	virtual int32 ActorWorldPosition() = 0;
 	virtual int32 ParticleMacroUV() = 0;
@@ -395,7 +395,7 @@ public:
 	virtual int32 ObjectWorldPosition() override { return Compiler->ObjectWorldPosition(); }
 	virtual int32 ObjectRadius() override { return Compiler->ObjectRadius(); }
 	virtual int32 ObjectBounds() override { return Compiler->ObjectBounds(); }
-	virtual int32 PreSkinnedLocalBounds() override { return Compiler->PreSkinnedLocalBounds(); }
+	virtual int32 PreSkinnedLocalBounds(int32 OutputIndex) override { return Compiler->PreSkinnedLocalBounds(OutputIndex); }
 	virtual int32 DistanceCullFade() override { return Compiler->DistanceCullFade(); }
 	virtual int32 ActorWorldPosition() override { return Compiler->ActorWorldPosition(); }
 	virtual int32 ParticleMacroUV() override { return Compiler->ParticleMacroUV(); }
diff --git a/Engine/Source/Runtime/Engine/Public/PrimitiveUniformShaderParameters.h b/Engine/Source/Runtime/Engine/Public/PrimitiveUniformShaderParameters.h
--- a/Engine/Source/Runtime/Engine/Public/PrimitiveUniformShaderParameters.h
+++ b/Engine/Source/Runtime/Engine/Public/PrimitiveUniformShaderParameters.h
@@ -37,10 +37,10 @@ BEGIN_GLOBAL_SHADER_PARAMETER_STRUCT(FPrimitiveUniformShaderParameters,ENGINE_AP
 	SHADER_PARAMETER(uint32,LightingChannelMask)
 	SHADER_PARAMETER(FVector, LocalObjectBoundsMax)		// This is used in a custom material function (ObjectLocalBounds.uasset)
 	SHADER_PARAMETER(uint32,LightmapDataIndex)
-	SHADER_PARAMETER(FVector, PreSkinnedLocalBounds)	// Local space bounds, pre-skinning
+	SHADER_PARAMETER(FVector, PreSkinnedLocalBoundsMin)	// Local space min bounds, pre-skinning
 	SHADER_PARAMETER(int32, SingleCaptureIndex)
+	SHADER_PARAMETER(FVector, PreSkinnedLocalBoundsMax)	// Local space max bounds, pre-skinning
     SHADER_PARAMETER(uint32, OutputVelocity)
-	// 12 bytes of padding here, feel free to use
 	SHADER_PARAMETER_ARRAY(FVector4, CustomPrimitiveData, [FCustomPrimitiveData::NumCustomPrimitiveDataFloat4s]) // Custom data per primitive that can be accessed through material expression parameters and modified through UStaticMeshComponent
 
 END_GLOBAL_SHADER_PARAMETER_STRUCT()
@@ -76,7 +76,8 @@ inline FPrimitiveUniformShaderParameters GetPrimitiveUniformShaderParameters(
 	Result.ObjectBounds = WorldBounds.BoxExtent;
 	Result.LocalObjectBoundsMin = LocalBounds.GetBoxExtrema(0); // 0 == minimum
 	Result.LocalObjectBoundsMax = LocalBounds.GetBoxExtrema(1); // 1 == maximum
-	Result.PreSkinnedLocalBounds = PreSkinnedLocalBounds.BoxExtent;
+	Result.PreSkinnedLocalBoundsMin = PreSkinnedLocalBounds.GetBoxExtrema(0); // 0 == minimum
+	Result.PreSkinnedLocalBoundsMax = PreSkinnedLocalBounds.GetBoxExtrema(1); // 1 == maximum
 	Result.ObjectOrientation = LocalToWorld.GetUnitAxis( EAxis::Z );
 	Result.ActorWorldPosition = ActorPosition;
 	Result.LightingChannelMask = LightingChannelMask;
