Fix control rig viewports not updating

Correctly call CacheBones or CacheBones_WithRoot depending on what root is being evaluated/updated.

#rnx

diff --git a/Engine/Source/Runtime/Engine/Private/Animation/AnimInstanceProxy.cpp b/Engine/Source/Runtime/Engine/Private/Animation/AnimInstanceProxy.cpp
--- a/Engine/Source/Runtime/Engine/Private/Animation/AnimInstanceProxy.cpp
+++ b/Engine/Source/Runtime/Engine/Private/Animation/AnimInstanceProxy.cpp
@@ -1134,7 +1134,15 @@ void FAnimInstanceProxy::UpdateAnimation_WithRoot(FAnimNode_Base* InRootNode, FN
 	ANIM_MT_SCOPE_CYCLE_COUNTER(ProxyUpdateAnimation, !IsInGameThread());
 	FScopeCycleCounterUObject AnimScope(GetAnimInstanceObject());
 
-	CacheBones_WithRoot(InRootNode);
+	if(InRootNode == RootNode)
+	{
+		// Call the correct override point if this is the root node
+		CacheBones();
+	}
+	else
+	{
+		CacheBones_WithRoot(InRootNode);
+	}
 
 	// update native update
 	if(!bUpdatingRoot)
@@ -1186,7 +1194,15 @@ void FAnimInstanceProxy::EvaluateAnimation_WithRoot(FPoseContext& Output, FAnimN
 {
 	ANIM_MT_SCOPE_CYCLE_COUNTER(EvaluateAnimInstance, !IsInGameThread());
 
-	CacheBones_WithRoot(InRootNode);
+	if(InRootNode == RootNode)
+	{
+		// Call the correct override point if this is the root node
+		CacheBones();
+	}
+	else
+	{
+		CacheBones_WithRoot(InRootNode);
+	}
 
 	// Evaluate native code if implemented, otherwise evaluate the node graph
 	if (!Evaluate_WithRoot(Output, InRootNode))
