Fix bug where only 4 float4 values were copied into the primitive unform buffer. It's now based on the count instead of hard coded.

#rnx

diff --git a/Engine/Source/Runtime/Engine/Private/PrimitiveUniformShaderParameters.cpp b/Engine/Source/Runtime/Engine/Private/PrimitiveUniformShaderParameters.cpp
--- a/Engine/Source/Runtime/Engine/Private/PrimitiveUniformShaderParameters.cpp
+++ b/Engine/Source/Runtime/Engine/Private/PrimitiveUniformShaderParameters.cpp
@@ -120,8 +120,10 @@ void FPrimitiveSceneShaderData::Setup(const FPrimitiveUniformShaderParameters& P
 	Data[26] = FVector4(PrimitiveUniformShaderParameters.PreSkinnedLocalBoundsMax, 0.0f);
 	Data[26].W = *(const float*)&PrimitiveUniformShaderParameters.OutputVelocity;
 
-	Data[27] = PrimitiveUniformShaderParameters.CustomPrimitiveData[0];
-	Data[28] = PrimitiveUniformShaderParameters.CustomPrimitiveData[1];
-	Data[29] = PrimitiveUniformShaderParameters.CustomPrimitiveData[2];
-	Data[30] = PrimitiveUniformShaderParameters.CustomPrimitiveData[3];
+	// Set all the custom primitive data float4. This matches the loop in SceneData.ush
+	const int32 CustomPrimitiveDataStartIndex = 27;
+	for (int i = 0; i < FCustomPrimitiveData::NumCustomPrimitiveDataFloat4s; i++)
+	{
+		Data[CustomPrimitiveDataStartIndex + i] = PrimitiveUniformShaderParameters.CustomPrimitiveData[i];
+	}
 }
