Fix duplicate FName symbols in mega-unity builds

#rnx

diff --git a/Engine/Source/Runtime/Engine/Private/Animation/AnimInstance.cpp b/Engine/Source/Runtime/Engine/Private/Animation/AnimInstance.cpp
--- a/Engine/Source/Runtime/Engine/Private/Animation/AnimInstance.cpp
+++ b/Engine/Source/Runtime/Engine/Private/Animation/AnimInstance.cpp
@@ -2396,7 +2396,6 @@ void UAnimInstance::ClearMontageInstanceReferences(FAnimMontageInstance& InMonta
 
 FAnimNode_SubInput* UAnimInstance::GetSubInputNode(FName InSubInputName, FName InGraph)
 {
-	static const FName NAME_AnimGraph("AnimGraph");
 	const FAnimInstanceProxy& Proxy = GetProxyOnAnyThread<FAnimInstanceProxy>();
 	if(InSubInputName == NAME_None && InGraph == NAME_None)
 	{
diff --git a/Engine/Source/Runtime/Engine/Private/Animation/AnimInstanceProxy.cpp b/Engine/Source/Runtime/Engine/Private/Animation/AnimInstanceProxy.cpp
--- a/Engine/Source/Runtime/Engine/Private/Animation/AnimInstanceProxy.cpp
+++ b/Engine/Source/Runtime/Engine/Private/Animation/AnimInstanceProxy.cpp
@@ -27,10 +27,10 @@
 
 #define LOCTEXT_NAMESPACE "AnimInstance"
 
-static const FName NAME_AnimBlueprintLog(TEXT("AnimBlueprintLog"));
-static const FName NAME_Evaluate(TEXT("Evaluate"));
-static const FName NAME_Update(TEXT("Update"));
-static const FName NAME_AnimGraph(TEXT("AnimGraph"));
+const FName NAME_AnimBlueprintLog(TEXT("AnimBlueprintLog"));
+const FName NAME_Evaluate(TEXT("Evaluate"));
+const FName NAME_Update(TEXT("Update"));
+const FName NAME_AnimGraph(TEXT("AnimGraph"));
 
 void FAnimInstanceProxy::UpdateAnimationNode(float DeltaSeconds)
 {
diff --git a/Engine/Source/Runtime/Engine/Private/Animation/AnimNode_SubInstance.cpp b/Engine/Source/Runtime/Engine/Private/Animation/AnimNode_SubInstance.cpp
--- a/Engine/Source/Runtime/Engine/Private/Animation/AnimNode_SubInstance.cpp
+++ b/Engine/Source/Runtime/Engine/Private/Animation/AnimNode_SubInstance.cpp
@@ -191,7 +191,6 @@ void FAnimNode_SubInstance::SetAnimClass(TSubclassOf<UAnimInstance> InClass, con
 
 FName FAnimNode_SubInstance::GetDynamicLinkFunctionName() const
 {
-	static const FName NAME_AnimGraph("AnimGraph");
 	return NAME_AnimGraph;
 }
 
diff --git a/Engine/Source/Runtime/Engine/Public/Animation/AnimInstanceProxy.h b/Engine/Source/Runtime/Engine/Public/Animation/AnimInstanceProxy.h
--- a/Engine/Source/Runtime/Engine/Public/Animation/AnimInstanceProxy.h
+++ b/Engine/Source/Runtime/Engine/Public/Animation/AnimInstanceProxy.h
@@ -35,6 +35,11 @@ struct FPoseContext;
 // Disable node logging for shipping and test builds
 #define ENABLE_ANIM_LOGGING (1 && !NO_LOGGING && !(UE_BUILD_SHIPPING || UE_BUILD_TEST))
 
+extern const FName NAME_AnimBlueprintLog;
+extern const FName NAME_Evaluate;
+extern const FName NAME_Update;
+extern const FName NAME_AnimGraph;
+
 UENUM()
 namespace EDrawDebugItemType
 {
