Strip SkelMesh DuplicatedVerticesBuffer on mobile, as it used only for SkinCache and Editor features which are SM5 only.

diff --git a/Engine/Source/Runtime/Engine/Private/SkeletalMeshLODRenderData.cpp b/Engine/Source/Runtime/Engine/Private/SkeletalMeshLODRenderData.cpp
--- a/Engine/Source/Runtime/Engine/Private/SkeletalMeshLODRenderData.cpp
+++ b/Engine/Source/Runtime/Engine/Private/SkeletalMeshLODRenderData.cpp
@@ -67,10 +67,16 @@ namespace
 // Serialization.
 FArchive& operator<<(FArchive& Ar, FSkelMeshRenderSection& S)
 {
+	const uint8 DuplicatedVertices = 1;
+	
+	// DuplicatedVerticesBuffer is used only for SkinCache and Editor features which is SM5 only
+	uint8 ClassDataStripFlags = 0;
+	ClassDataStripFlags |= (Ar.IsCooking() && Ar.CookingTarget()->SupportsFeature(ETargetPlatformFeatures::MobileRendering)) ? DuplicatedVertices : 0;
+
 	// When data is cooked for server platform some of the
 	// variables are not serialized so that they're always
 	// set to their initial values (for safety)
-	FStripDataFlags StripFlags(Ar);
+	FStripDataFlags StripFlags(Ar, ClassDataStripFlags);
 
 	Ar << S.MaterialIndex;
 	Ar << S.BaseIndex;
@@ -84,7 +90,10 @@ FArchive& operator<<(FArchive& Ar, FSkelMeshRenderSection& S)
 	Ar << S.MaxBoneInfluences;
 	Ar << S.CorrespondClothAssetIndex;
 	Ar << S.ClothingData;
-	Ar << S.DuplicatedVerticesBuffer;
+	if (!StripFlags.IsClassDataStripped(DuplicatedVertices))
+	{
+		Ar << S.DuplicatedVerticesBuffer;
+	}
 	Ar << S.bDisabled;
 
 	return Ar;
@@ -204,7 +213,8 @@ void FSkeletalMeshLODRenderData::InitResources(bool bNeedsVertexColors, int32 LO
 		INC_DWORD_STAT_BY(STAT_SkeletalMeshIndexMemory, AdjacencyMultiSizeIndexContainer.IsIndexBufferValid() ? (AdjacencyMultiSizeIndexContainer.GetIndexBuffer()->Num() * AdjacencyMultiSizeIndexContainer.GetDataTypeSize()) : 0);
 	}
 
-    if (RHISupportsComputeShaders(GMaxRHIShaderPlatform))
+	// DuplicatedVerticesBuffer is used only for SkinCache and Editor features which is SM5 only
+    if (IsFeatureLevelSupported(GMaxRHIShaderPlatform, ERHIFeatureLevel::SM5))
     {
         for (auto& RenderSection : RenderSections)
         {
