Fix skeletal controls (and other nodes) not funcitoning correctly in 'self' layers

CacheBones_AnyThread was not called for other layers as the guard value bBoneCachesInvalidated was getting set back to false. The fix here prevents the guard value from being reset until all recursion paths have been taken.

#rnx

diff --git a/Engine/Source/Runtime/Engine/Private/Animation/AnimInstanceProxy.cpp b/Engine/Source/Runtime/Engine/Private/Animation/AnimInstanceProxy.cpp
--- a/Engine/Source/Runtime/Engine/Private/Animation/AnimInstanceProxy.cpp
+++ b/Engine/Source/Runtime/Engine/Private/Animation/AnimInstanceProxy.cpp
@@ -1216,11 +1216,20 @@ void FAnimInstanceProxy::CacheBones()
 	// If bone caches have been invalidated, have AnimNodes refresh those.
 	if (bBoneCachesInvalidated && RootNode)
 	{
-		bBoneCachesInvalidated = false;
+		CacheBonesRecursionCounter++;
 
 		CachedBonesCounter.Increment();
 		FAnimationCacheBonesContext Context(this);
 		RootNode->CacheBones_AnyThread(Context);
+
+		CacheBonesRecursionCounter--;
+
+		check(CacheBonesRecursionCounter >= 0);
+
+		if(CacheBonesRecursionCounter == 0)
+		{
+			bBoneCachesInvalidated = false;
+		}
 	}
 }
 
@@ -1229,11 +1238,20 @@ void FAnimInstanceProxy::CacheBones_WithRoot(FAnimNode_Base* InRootNode)
 	// If bone caches have been invalidated, have AnimNodes refresh those.
 	if (bBoneCachesInvalidated && InRootNode)
 	{
-		bBoneCachesInvalidated = false;
+		CacheBonesRecursionCounter++;
 
 		CachedBonesCounter.Increment();
 		FAnimationCacheBonesContext Context(this);
 		InRootNode->CacheBones_AnyThread(Context);
+
+		CacheBonesRecursionCounter--;
+
+		check(CacheBonesRecursionCounter >= 0);
+
+		if(CacheBonesRecursionCounter == 0)
+		{
+			bBoneCachesInvalidated = false;
+		}
 	}
 }
 
diff --git a/Engine/Source/Runtime/Engine/Public/Animation/AnimInstanceProxy.h b/Engine/Source/Runtime/Engine/Public/Animation/AnimInstanceProxy.h
--- a/Engine/Source/Runtime/Engine/Public/Animation/AnimInstanceProxy.h
+++ b/Engine/Source/Runtime/Engine/Public/Animation/AnimInstanceProxy.h
@@ -120,6 +120,7 @@ public:
 		, SyncGroupWriteIndex(0)
 		, RootMotionMode(ERootMotionMode::NoRootMotionExtraction)
 		, FrameCounterForUpdate(0)
+		, CacheBonesRecursionCounter(0)
 		, bUpdatingRoot(false)
 		, bBoneCachesInvalidated(false)
 		, bShouldExtractRootMotion(false)
@@ -141,6 +142,7 @@ public:
 		, SyncGroupWriteIndex(0)
 		, RootMotionMode(ERootMotionMode::NoRootMotionExtraction)
 		, FrameCounterForUpdate(0)
+		, CacheBonesRecursionCounter(0)
 		, bUpdatingRoot(false)
 		, bBoneCachesInvalidated(false)
 		, bShouldExtractRootMotion(false)
@@ -864,6 +866,9 @@ private:
 	/** LODLevel used by RequiredBones */
 	int32 LODLevel;
 
+	/** Counter used to control CacheBones recursion behavior - makes sure we cache bones correctly when recursing into different subgraphs */
+	int32 CacheBonesRecursionCounter;
+
 	/** Cached SkeletalMeshComponent LocalToWorld transform. */
 	FTransform SkelMeshCompLocalToWorld;
 
