Fixed check() in DynamicLink() firing incorrectly

In some cirumstances DynamicUnlink() wasnt geting called fo rlayers. This ensures it gets called when we are linked.

#rnx

diff --git a/Engine/Source/Runtime/Engine/Private/Animation/AnimNode_Layer.cpp b/Engine/Source/Runtime/Engine/Private/Animation/AnimNode_Layer.cpp
--- a/Engine/Source/Runtime/Engine/Private/Animation/AnimNode_Layer.cpp
+++ b/Engine/Source/Runtime/Engine/Private/Animation/AnimNode_Layer.cpp
@@ -30,11 +30,14 @@ void FAnimNode_Layer::OnInitializeAnimInstance(const FAnimInstanceProxy* InProxy
 		USkeletalMeshComponent* MeshComp = InAnimInstance->GetSkelMeshComponent();
 		check(MeshComp);
 
-		// Switch from dynamic external to internal, kill old instance
-		if(CurrentTarget && CurrentTarget != InAnimInstance)
+		if(LinkedRoot)
 		{
 			DynamicUnlink(const_cast<UAnimInstance*>(InAnimInstance));
+		}
 
+		// Switch from dynamic external to internal, kill old instance
+		if(CurrentTarget && CurrentTarget != InAnimInstance)
+		{
 			MeshComp->SubInstances.Remove(CurrentTarget);
 			CurrentTarget->MarkPendingKill();
 			CurrentTarget = nullptr;
