Fix nullptr crash in UMaterial::SetVectorParameterValueEditorOnly(), refactored a few methods to avoid code duplication
#rnx

diff --git a/Engine/Source/Runtime/Engine/Classes/Materials/Material.h b/Engine/Source/Runtime/Engine/Classes/Materials/Material.h
--- a/Engine/Source/Runtime/Engine/Classes/Materials/Material.h
+++ b/Engine/Source/Runtime/Engine/Classes/Materials/Material.h
@@ -1824,6 +1824,65 @@ private:
 		return nullptr;
 	}
 
+#if WITH_EDITOR
+
+	template<typename ParameterType, typename ...Arguments>
+	bool SetParameterValueEditorOnly(FName InParameterName, Arguments... Args)
+	{
+		// Warning: in the case of duplicate parameters with different default values, this will find the first in the expression array, not necessarily the one that's used for rendering
+
+		// Lambda which calls SetParameterValue on a given parameter and triggers a PostEditChange event if successful
+		auto TrySetParameterValue = [&](ParameterType* Parameter) -> bool
+		{
+			if (Parameter && Parameter->SetParameterValue(InParameterName, Args...))
+			{
+				if (UProperty* ParamProperty = FindField<UProperty>(ParameterType::StaticClass(), "DefaultValue"))
+				{
+					FPropertyChangedEvent PropertyChangedEvent(ParamProperty);
+					Parameter->PostEditChangeProperty(PropertyChangedEvent);
+					return true;
+				}
+			}
+			return false;
+		};
+
+		for (UMaterialExpression* Expression : Expressions)
+		{
+			if (TrySetParameterValue(Cast<ParameterType>(Expression)))
+			{
+				return true;
+			}
+			else if (UMaterialExpressionMaterialFunctionCall* FunctionCall = Cast<UMaterialExpressionMaterialFunctionCall>(Expression))
+			{
+				if (FunctionCall->MaterialFunction)
+				{
+					TArray<UMaterialFunctionInterface*> Functions;
+					Functions.Add(FunctionCall->MaterialFunction);
+					FunctionCall->MaterialFunction->GetDependentFunctions(Functions);
+
+					for (UMaterialFunctionInterface* Function : Functions)
+					{
+						const TArray<UMaterialExpression*>* ExpressionPtr = Function->GetFunctionExpressions();
+						if (ExpressionPtr)
+						{
+							for (UMaterialExpression* FunctionExpression : *ExpressionPtr)
+							{
+								if (TrySetParameterValue(Cast<ParameterType>(FunctionExpression)))
+								{
+									return true;
+								}
+							}
+						}
+					}
+				}
+			}
+		}
+
+		return false;
+	}
+
+#endif // WITH_EDITOR
+
 };
 
 
diff --git a/Engine/Source/Runtime/Engine/Private/Materials/Material.cpp b/Engine/Source/Runtime/Engine/Private/Materials/Material.cpp
--- a/Engine/Source/Runtime/Engine/Private/Materials/Material.cpp
+++ b/Engine/Source/Runtime/Engine/Private/Materials/Material.cpp
@@ -1270,287 +1270,33 @@ bool UMaterial::IsCompilingOrHadCompileError(ERHIFeatureLevel::Type InFeatureLev
 #if WITH_EDITOR
 bool UMaterial::SetVectorParameterValueEditorOnly(FName ParameterName, FLinearColor InValue)
 {
-	bool bSetNewValue = false;
-	UMaterialExpressionVectorParameter* Parameter = nullptr;
-	for (UMaterialExpression* Expression : Expressions)
-	{
-		Parameter = Cast<UMaterialExpressionVectorParameter>(Expression);
-		if (Parameter && Parameter->SetParameterValue(ParameterName, InValue))
-		{
-			bSetNewValue = true;
-			break;
-			// Warning: in the case of duplicate parameters with different default values, this will find the first in the expression array, not necessarily the one that's used for rendering
-		}
-		else if (UMaterialExpressionMaterialFunctionCall* FunctionCall = Cast<UMaterialExpressionMaterialFunctionCall>(Expression))
-		{
-			if (FunctionCall->MaterialFunction)
-			{
-				TArray<UMaterialFunctionInterface*> Functions;
-				Functions.Add(FunctionCall->MaterialFunction);
-				FunctionCall->MaterialFunction->GetDependentFunctions(Functions);
-
-				for (UMaterialFunctionInterface* Function : Functions)
-				{
-					const TArray<UMaterialExpression*>* FunctionExpressions = Function->GetFunctionExpressions();
-					if(FunctionExpressions)
-					{
-						for (UMaterialExpression* FunctionExpression : *FunctionExpressions)
-						{
-							if (UMaterialExpressionVectorParameter* FunctionExpressionParameter = Cast<UMaterialExpressionVectorParameter>(FunctionExpression))
-							{
-								if (FunctionExpressionParameter->SetParameterValue(ParameterName, InValue))
-								{
-									bSetNewValue = true;
-									break;
-								}
-							}
-						}
-					}
-					if (bSetNewValue)
-					{
-						break;
-					}
-				}
-			}
-		}
-	}
-	if (bSetNewValue)
-	{
-		UProperty* ParamProperty = FindField<UProperty>(UMaterialExpressionVectorParameter::StaticClass(), GET_MEMBER_NAME_STRING_CHECKED(UMaterialExpressionVectorParameter, DefaultValue));
-		FPropertyChangedEvent PropertyChangedEvent(ParamProperty);
-		Parameter->PostEditChangeProperty(PropertyChangedEvent);
-	}
-	return bSetNewValue;
+	return SetParameterValueEditorOnly<UMaterialExpressionVectorParameter>(ParameterName, InValue);
 }
 
 bool UMaterial::SetScalarParameterValueEditorOnly(FName ParameterName, float InValue)
 {
-	bool bSetNewValue = false;
-	UMaterialExpressionScalarParameter* Parameter = nullptr;
-	for (UMaterialExpression* Expression : Expressions)
-	{
-		Parameter = Cast<UMaterialExpressionScalarParameter>(Expression);
-		if (Parameter && Parameter->SetParameterValue(ParameterName, InValue))
-		{
-			bSetNewValue = true;
-			break;
-			// Warning: in the case of duplicate parameters with different default values, this will find the first in the expression array, not necessarily the one that's used for rendering
-		}
-		else if (UMaterialExpressionMaterialFunctionCall* FunctionCall = Cast<UMaterialExpressionMaterialFunctionCall>(Expression))
-		{
-			if (FunctionCall->MaterialFunction)
-			{
-				TArray<UMaterialFunctionInterface*> Functions;
-				Functions.Add(FunctionCall->MaterialFunction);
-				FunctionCall->MaterialFunction->GetDependentFunctions(Functions);
-
-				for (UMaterialFunctionInterface* Function : Functions)
-				{
-					const TArray<UMaterialExpression*>* ExpressionPtr = Function->GetFunctionExpressions();
-					if (ExpressionPtr)
-					{
-						for (UMaterialExpression* FunctionExpression : *ExpressionPtr)
-						{
-							Parameter = Cast<UMaterialExpressionScalarParameter>(FunctionExpression);
-							if (!bSetNewValue && Parameter)
-							{
-								if (Parameter->SetParameterValue(ParameterName, InValue))
-								{
-									bSetNewValue = true;
-									break;
-								}
-							}
-						}
-					}
-					if (bSetNewValue)
-					{
-						break;
-					}
-				}
-			}
-		}
-	}
-	if (bSetNewValue)
-	{
-		UProperty* ParamProperty = FindField<UProperty>(UMaterialExpressionScalarParameter::StaticClass(), GET_MEMBER_NAME_STRING_CHECKED(UMaterialExpressionScalarParameter, DefaultValue));
-		FPropertyChangedEvent PropertyChangedEvent(ParamProperty);
-		Parameter->PostEditChangeProperty(PropertyChangedEvent);
-	}
-
-	return bSetNewValue;
-};
+	return SetParameterValueEditorOnly<UMaterialExpressionScalarParameter>(ParameterName, InValue);
+}
 
 bool UMaterial::SetTextureParameterValueEditorOnly(FName ParameterName, class UTexture* InValue)
 {
-	for (UMaterialExpression* Expression : Expressions)
-	{
-		if (UMaterialExpressionTextureSampleParameter* Parameter = Cast<UMaterialExpressionTextureSampleParameter>(Expression))
-		{
-			if (Parameter->SetParameterValue(ParameterName, InValue))
-			{
-				return true;
-				// Warning: in the case of duplicate parameters with different default values, this will find the first in the expression array, not necessarily the one that's used for rendering
-			}
-		}
-		else if (UMaterialExpressionMaterialFunctionCall* FunctionCall = Cast<UMaterialExpressionMaterialFunctionCall>(Expression))
-		{
-			if (FunctionCall->MaterialFunction)
-			{
-				TArray<UMaterialFunctionInterface*> Functions;
-				Functions.Add(FunctionCall->MaterialFunction);
-				FunctionCall->MaterialFunction->GetDependentFunctions(Functions);
-
-				for (UMaterialFunctionInterface* Function : Functions)
-				{
-					const TArray<UMaterialExpression*>* ExpressionPtr = Function->GetFunctionExpressions();
-					if (ExpressionPtr)
-					{
-						for (UMaterialExpression* FunctionExpression : *ExpressionPtr)
-						{
-							if (UMaterialExpressionTextureSampleParameter* FunctionExpressionParameter = Cast<UMaterialExpressionTextureSampleParameter>(FunctionExpression))
-							{
-								if (FunctionExpressionParameter->SetParameterValue(ParameterName, InValue))
-								{
-									return true;
-								}
-							}
-						}
-					}
-				}
-			}
-		}
-	}
-	return false;
-};
+	return SetParameterValueEditorOnly<UMaterialExpressionTextureSampleParameter>(ParameterName, InValue);
+}
 
 bool UMaterial::SetFontParameterValueEditorOnly(FName ParameterName, class UFont* InFontValue, int32 InFontPage)
 {
-	for (UMaterialExpression* Expression : Expressions)
-	{
-		if (UMaterialExpressionFontSampleParameter* Parameter = Cast<UMaterialExpressionFontSampleParameter>(Expression))
-		{
-			if (Parameter->SetParameterValue(ParameterName, InFontValue, InFontPage))
-			{
-				return true;
-				// Warning: in the case of duplicate parameters with different default values, this will find the first in the expression array, not necessarily the one that's used for rendering
-			}
-		}
-		else if (UMaterialExpressionMaterialFunctionCall* FunctionCall = Cast<UMaterialExpressionMaterialFunctionCall>(Expression))
-		{
-			if (FunctionCall->MaterialFunction)
-			{
-				TArray<UMaterialFunctionInterface*> Functions;
-				Functions.Add(FunctionCall->MaterialFunction);
-				FunctionCall->MaterialFunction->GetDependentFunctions(Functions);
-
-				for (UMaterialFunctionInterface* Function : Functions)
-				{
-					const TArray<UMaterialExpression*>* ExpressionPtr = Function->GetFunctionExpressions();
-					if (ExpressionPtr)
-					{
-						for (UMaterialExpression* FunctionExpression : *ExpressionPtr)
-						{
-							if (UMaterialExpressionFontSampleParameter* FunctionExpressionParameter = Cast<UMaterialExpressionFontSampleParameter>(FunctionExpression))
-							{
-								if (FunctionExpressionParameter->SetParameterValue(ParameterName, InFontValue, InFontPage))
-								{
-									return true;
-								}
-							}
-						}
-					}
-				}
-			}
-		}
-	}
-	return false;
-};
+	return SetParameterValueEditorOnly<UMaterialExpressionFontSampleParameter>(ParameterName, InFontValue, InFontPage);
+}
 
-bool UMaterial::SetStaticSwitchParameterValueEditorOnly(FName ParameterName, bool OutValue, FGuid OutExpressionGuid)
+bool UMaterial::SetStaticSwitchParameterValueEditorOnly(FName ParameterName, bool InValue, FGuid InExpressionGuid)
 {
-	for (UMaterialExpression* Expression : Expressions)
-	{
-		if (UMaterialExpressionStaticSwitchParameter* Parameter = Cast<UMaterialExpressionStaticSwitchParameter>(Expression))
-		{
-			if (Parameter->SetParameterValue(ParameterName, OutValue, OutExpressionGuid))
-			{
-				return true;
-				// Warning: in the case of duplicate parameters with different default values, this will find the first in the expression array, not necessarily the one that's used for rendering
-			}
-		}
-		else if (UMaterialExpressionMaterialFunctionCall* FunctionCall = Cast<UMaterialExpressionMaterialFunctionCall>(Expression))
-		{
-			if (FunctionCall->MaterialFunction)
-			{
-				TArray<UMaterialFunctionInterface*> Functions;
-				Functions.Add(FunctionCall->MaterialFunction);
-				FunctionCall->MaterialFunction->GetDependentFunctions(Functions);
-
-				for (UMaterialFunctionInterface* Function : Functions)
-				{
-					const TArray<UMaterialExpression*>* ExpressionPtr = Function->GetFunctionExpressions();
-					if (ExpressionPtr)
-					{
-						for (UMaterialExpression* FunctionExpression : *ExpressionPtr)
-						{
-							if (UMaterialExpressionStaticSwitchParameter* FunctionExpressionParameter = Cast<UMaterialExpressionStaticSwitchParameter>(FunctionExpression))
-							{
-								if (FunctionExpressionParameter->SetParameterValue(ParameterName, OutValue, OutExpressionGuid))
-								{
-									return true;
-								}
-							}
-						}
-					}
-				}
-			}
-		}
-	}
-	return false;
-};
+	return SetParameterValueEditorOnly<UMaterialExpressionStaticSwitchParameter>(ParameterName, InValue, InExpressionGuid);
+}
 
-bool UMaterial::SetStaticComponentMaskParameterValueEditorOnly(FName ParameterName, bool R, bool G, bool B, bool A, FGuid OutExpressionGuid)
+bool UMaterial::SetStaticComponentMaskParameterValueEditorOnly(FName ParameterName, bool R, bool G, bool B, bool A, FGuid InExpressionGuid)
 {
-	for (UMaterialExpression* Expression : Expressions)
-	{
-		if (UMaterialExpressionStaticComponentMaskParameter* Parameter = Cast<UMaterialExpressionStaticComponentMaskParameter>(Expression))
-		{
-			if (Parameter->SetParameterValue(ParameterName, R, G, B, A, OutExpressionGuid))
-			{
-				return true;
-				// Warning: in the case of duplicate parameters with different default values, this will find the first in the expression array, not necessarily the one that's used for rendering
-			}
-		}
-		else if (UMaterialExpressionMaterialFunctionCall* FunctionCall = Cast<UMaterialExpressionMaterialFunctionCall>(Expression))
-		{
-			if (FunctionCall->MaterialFunction)
-			{
-				TArray<UMaterialFunctionInterface*> Functions;
-				Functions.Add(FunctionCall->MaterialFunction);
-				FunctionCall->MaterialFunction->GetDependentFunctions(Functions);
-
-				for (UMaterialFunctionInterface* Function : Functions)
-				{
-					const TArray<UMaterialExpression*>* ExpressionPtr = Function->GetFunctionExpressions();
-					if (ExpressionPtr)
-					{
-						for (UMaterialExpression* FunctionExpression : *ExpressionPtr)
-						{
-							if (UMaterialExpressionStaticComponentMaskParameter* FunctionExpressionParameter = Cast<UMaterialExpressionStaticComponentMaskParameter>(FunctionExpression))
-							{
-								if (FunctionExpressionParameter->SetParameterValue(ParameterName, R, G, B, A, OutExpressionGuid))
-								{
-									return true;
-								}
-							}
-						}
-					}
-				}
-			}
-		}
-	}
-	return false;
-};
+	return SetParameterValueEditorOnly<UMaterialExpressionStaticComponentMaskParameter>(ParameterName, R, G, B, A, InExpressionGuid);
+}
 #endif
 
 void UMaterial::MarkUsageFlagDirty(EMaterialUsage Usage, bool CurrentValue, bool NewValue)
