Adding cvar to override max concurrent streams.

au.MaxConcurrentStreams

diff --git a/Engine/Source/Runtime/Engine/Private/AudioStreaming.cpp b/Engine/Source/Runtime/Engine/Private/AudioStreaming.cpp
--- a/Engine/Source/Runtime/Engine/Private/AudioStreaming.cpp
+++ b/Engine/Source/Runtime/Engine/Private/AudioStreaming.cpp
@@ -26,6 +26,14 @@ FAutoConsoleVariableRef CVarSpoofFailedStreamChunkLoad(
 	TEXT("0: Not Enabled, 1: Enabled"),
 	ECVF_Default);
 
+static int32 MaxConcurrentStreamsCvar = 0;
+FAutoConsoleVariableRef CVarMaxConcurrentStreams(
+	TEXT("au.MaxConcurrentStreams"),
+	MaxConcurrentStreamsCvar,
+	TEXT("Overrides the max concurrent streams.\n")
+	TEXT("0: Not Overridden, >0 Overridden"),
+	ECVF_Default);
+
 
 /*------------------------------------------------------------------------------
 	Streaming chunks from the derived data cache.
@@ -781,7 +789,11 @@ bool FAudioStreamingManager::CanCreateSoundSource(const FWaveInstance* WaveInsta
 	check(WaveInstance);
 	check(WaveInstance->IsStreaming());
 
-	int32 MaxStreams = GetDefault<UAudioSettings>()->MaximumConcurrentStreams;
+	int32 MaxStreams = MaxConcurrentStreamsCvar;
+	if (!MaxStreams)
+	{
+		MaxStreams = GetDefault<UAudioSettings>()->MaximumConcurrentStreams;
+	}
 
 	FScopeLock Lock(&CriticalSection);
 
@@ -815,7 +827,11 @@ void FAudioStreamingManager::AddStreamingSoundSource(FSoundSource* SoundSource)
 	const FWaveInstance* WaveInstance = SoundSource->GetWaveInstance();
 	if (WaveInstance && WaveInstance->IsStreaming())
 	{
-		int32 MaxStreams = GetDefault<UAudioSettings>()->MaximumConcurrentStreams;
+		int32 MaxStreams = MaxConcurrentStreamsCvar;
+		if (!MaxStreams)
+		{
+			MaxStreams = GetDefault<UAudioSettings>()->MaximumConcurrentStreams;
+		}
 
 		FScopeLock Lock(&CriticalSection);
 
