Adding GetResolutionScaleNormalized to match the setter on GameUserSettings.

diff --git a/Engine/Source/Runtime/Engine/Classes/GameFramework/GameUserSettings.h b/Engine/Source/Runtime/Engine/Classes/GameFramework/GameUserSettings.h
--- a/Engine/Source/Runtime/Engine/Classes/GameFramework/GameUserSettings.h
+++ b/Engine/Source/Runtime/Engine/Classes/GameFramework/GameUserSettings.h
@@ -164,6 +164,10 @@ public:
 	UFUNCTION(BlueprintCallable, Category=Settings, meta=(DisplayName="GetResolutionScaleInformation"))
 	void GetResolutionScaleInformationEx(float& CurrentScaleNormalized, float& CurrentScaleValue, float& MinScaleValue, float& MaxScaleValue) const;
 
+	// Gets the current resolution scale as a normalized 0..1 value between MinScaleValue and MaxScaleValue
+	UFUNCTION(BlueprintCallable, Category = Settings)
+	float GetResolutionScaleNormalized() const;
+
 	// Sets the current resolution scale
 	UE_DEPRECATED(4.12, "Please call SetResolutionScaleValueEx")
 	UFUNCTION(BlueprintCallable, Category=Settings, meta=(DeprecatedFunction, DisplayName="SetResolutionScaleValue_Deprecated"))
diff --git a/Engine/Source/Runtime/Engine/Private/GameUserSettings.cpp b/Engine/Source/Runtime/Engine/Private/GameUserSettings.cpp
--- a/Engine/Source/Runtime/Engine/Private/GameUserSettings.cpp
+++ b/Engine/Source/Runtime/Engine/Private/GameUserSettings.cpp
@@ -743,6 +743,14 @@ void UGameUserSettings::GetResolutionScaleInformationEx(float& CurrentScaleNorma
 	CurrentScaleNormalized = ((float)CurrentScaleValue - (float)MinScaleValue) / (float)(MaxScaleValue - MinScaleValue);
 }
 
+float UGameUserSettings::GetResolutionScaleNormalized() const
+{
+	float CurrentScaleNormalized, CurrentScaleValue, MinScaleValue, MaxScaleValue;
+	GetResolutionScaleInformationEx(CurrentScaleNormalized, CurrentScaleValue, MinScaleValue, MaxScaleValue);
+
+	return CurrentScaleNormalized;
+}
+
 void UGameUserSettings::SetResolutionScaleValue(int32 NewScaleValue)
 {
 	SetResolutionScaleValueEx((float)NewScaleValue);
