Fix assert when DuplicatedVerticesBuffer data is stripped

diff --git a/Engine/Source/Runtime/Engine/Private/SkeletalMeshLODRenderData.cpp b/Engine/Source/Runtime/Engine/Private/SkeletalMeshLODRenderData.cpp
--- a/Engine/Source/Runtime/Engine/Private/SkeletalMeshLODRenderData.cpp
+++ b/Engine/Source/Runtime/Engine/Private/SkeletalMeshLODRenderData.cpp
@@ -482,10 +482,14 @@ void FSkeletalMeshLODRenderData::ReleaseResources()
 	BeginReleaseResource(&SkinWeightVertexBuffer);
 	BeginReleaseResource(&StaticVertexBuffers.ColorVertexBuffer);
 	BeginReleaseResource(&ClothVertexBuffer);
-	for (auto& RenderSection : RenderSections)
+	// DuplicatedVerticesBuffer is used only for SkinCache and Editor features which is SM5 only
+    if (IsFeatureLevelSupported(GMaxRHIShaderPlatform, ERHIFeatureLevel::SM5))
 	{
-		check(RenderSection.DuplicatedVerticesBuffer.DupVertData.Num());
-		BeginReleaseResource(&RenderSection.DuplicatedVerticesBuffer);
+		for (auto& RenderSection : RenderSections)
+		{
+			check(RenderSection.DuplicatedVerticesBuffer.DupVertData.Num());
+			BeginReleaseResource(&RenderSection.DuplicatedVerticesBuffer);
+		}
 	}
 	BeginReleaseResource(&MorphTargetVertexInfoBuffers);
 	
