#UE4 Added the ability to explicitly override PrimaryAssetRules, in case you want to add an override to set a value to a default value.

diff --git a/Engine/Source/Runtime/Engine/Classes/Engine/AssetManager.h b/Engine/Source/Runtime/Engine/Classes/Engine/AssetManager.h
--- a/Engine/Source/Runtime/Engine/Classes/Engine/AssetManager.h
+++ b/Engine/Source/Runtime/Engine/Classes/Engine/AssetManager.h
@@ -319,6 +319,7 @@ public:
 
 	/** Changes the management rules for a specific asset, this overrides the type rules. If passed in Rules is default, delete override */
 	virtual void SetPrimaryAssetRules(FPrimaryAssetId PrimaryAssetId, const FPrimaryAssetRules& Rules);
+	virtual void SetPrimaryAssetRulesExplicitly(FPrimaryAssetId PrimaryAssetId, const FPrimaryAssetRulesExplicitOverride& ExplicitRules);
 
 	/** Gets the management rules for a specific asset, this will merge the type and individual values */
 	virtual FPrimaryAssetRules GetPrimaryAssetRules(FPrimaryAssetId PrimaryAssetId) const;
@@ -600,7 +601,7 @@ protected:
 	TMap<FName, FPrimaryAssetId> AssetPathMap;
 
 	/** Overridden asset management data for specific types */
-	TMap<FPrimaryAssetId, FPrimaryAssetRules> AssetRuleOverrides;
+	TMap<FPrimaryAssetId, FPrimaryAssetRulesExplicitOverride> AssetRuleOverrides;
 
 	/** Map from PrimaryAssetId to list of PrimaryAssetIds that are the parent of this one, for determining chunking/cooking */
 	TMap<FPrimaryAssetId, TArray<FPrimaryAssetId>> ManagementParentMap;
diff --git a/Engine/Source/Runtime/Engine/Classes/Engine/AssetManagerTypes.h b/Engine/Source/Runtime/Engine/Classes/Engine/AssetManagerTypes.h
--- a/Engine/Source/Runtime/Engine/Classes/Engine/AssetManagerTypes.h
+++ b/Engine/Source/Runtime/Engine/Classes/Engine/AssetManagerTypes.h
@@ -74,6 +74,26 @@ struct FPrimaryAssetRules
 	ENGINE_API void PropagateCookRules(const FPrimaryAssetRules& ParentRules);
 };
 
+/** Structure defining overrides to rules */
+struct FPrimaryAssetRulesExplicitOverride
+{
+	FPrimaryAssetRules Rules;
+	uint8 bOverridePriority:1;
+	uint8 bOverrideApplyRecursively:1;
+	uint8 bOverrideChunkId:1;
+	uint8 bOverrideCookRule:1;
+
+	FPrimaryAssetRulesExplicitOverride()
+		: bOverridePriority(false), bOverrideApplyRecursively(false), bOverrideChunkId(false), bOverrideCookRule(false)
+	{
+	}
+
+	bool HasAnyOverride() const { return bOverridePriority || bOverrideApplyRecursively || bOverrideChunkId || bOverrideCookRule; }
+
+	/** Override non-default rules from an override struct. */
+	ENGINE_API void OverrideRulesExplicitly(FPrimaryAssetRules& RulesToOverride) const;
+};
+
 /** Structure with publicly exposed information about an asset type. These can be loaded out of a config file. */
 USTRUCT()
 struct FPrimaryAssetTypeInfo
diff --git a/Engine/Source/Runtime/Engine/Private/AssetManager.cpp b/Engine/Source/Runtime/Engine/Private/AssetManager.cpp
--- a/Engine/Source/Runtime/Engine/Private/AssetManager.cpp
+++ b/Engine/Source/Runtime/Engine/Private/AssetManager.cpp
@@ -811,7 +811,21 @@ void UAssetManager::SetPrimaryAssetTypeRules(FPrimaryAssetType PrimaryAssetType,
 
 void UAssetManager::SetPrimaryAssetRules(FPrimaryAssetId PrimaryAssetId, const FPrimaryAssetRules& Rules)
 {
-	if (Rules.IsDefault())
+	static FPrimaryAssetRules DefaultRules;
+
+	FPrimaryAssetRulesExplicitOverride ExplicitRules;
+	ExplicitRules.Rules = Rules;
+	ExplicitRules.bOverridePriority = (Rules.Priority != DefaultRules.Priority);
+	ExplicitRules.bOverrideApplyRecursively = (Rules.bApplyRecursively != DefaultRules.bApplyRecursively);
+	ExplicitRules.bOverrideChunkId = (Rules.ChunkId != DefaultRules.ChunkId);
+	ExplicitRules.bOverrideCookRule = (Rules.CookRule != DefaultRules.CookRule);
+	
+	SetPrimaryAssetRulesExplicitly(PrimaryAssetId, ExplicitRules);
+}
+
+void UAssetManager::SetPrimaryAssetRulesExplicitly(FPrimaryAssetId PrimaryAssetId, const FPrimaryAssetRulesExplicitOverride& ExplicitRules)
+{
+	if (!ExplicitRules.HasAnyOverride())
 	{
 		AssetRuleOverrides.Remove(PrimaryAssetId);
 	}
@@ -822,7 +836,7 @@ void UAssetManager::SetPrimaryAssetRules(FPrimaryAssetId PrimaryAssetId, const F
 			UE_LOG(LogAssetManager, Error, TEXT("Duplicate Rule overrides found for asset %s!"), *PrimaryAssetId.ToString());
 		}
 
-		AssetRuleOverrides.Add(PrimaryAssetId, Rules);
+		AssetRuleOverrides.Add(PrimaryAssetId, ExplicitRules);
 	}
 
 	bIsManagementDatabaseCurrent = false;
@@ -840,11 +854,11 @@ FPrimaryAssetRules UAssetManager::GetPrimaryAssetRules(FPrimaryAssetId PrimaryAs
 		Result = (*FoundType)->Info.Rules;
 
 		// Selectively override
-		const FPrimaryAssetRules* FoundRules = AssetRuleOverrides.Find(PrimaryAssetId);
+		const FPrimaryAssetRulesExplicitOverride* FoundRulesOverride = AssetRuleOverrides.Find(PrimaryAssetId);
 
-		if (FoundRules)
+		if (FoundRulesOverride)
 		{
-			Result.OverrideRules(*FoundRules);
+			FoundRulesOverride->OverrideRulesExplicitly(Result);
 		}
 
 		if (Result.Priority < 0)
diff --git a/Engine/Source/Runtime/Engine/Private/AssetManagerTypes.cpp b/Engine/Source/Runtime/Engine/Private/AssetManagerTypes.cpp
--- a/Engine/Source/Runtime/Engine/Private/AssetManagerTypes.cpp
+++ b/Engine/Source/Runtime/Engine/Private/AssetManagerTypes.cpp
@@ -100,6 +100,29 @@ void FPrimaryAssetRules::OverrideRules(const FPrimaryAssetRules& OverrideRules)
 	}
 }
 
+void FPrimaryAssetRulesExplicitOverride::OverrideRulesExplicitly(FPrimaryAssetRules& RulesToOverride) const
+{
+	if (bOverridePriority)
+	{
+		RulesToOverride.Priority = Rules.Priority;
+	}
+
+	if (bOverrideApplyRecursively)
+	{
+		RulesToOverride.bApplyRecursively = Rules.bApplyRecursively;
+	}
+
+	if (bOverrideChunkId)
+	{
+		RulesToOverride.ChunkId = Rules.ChunkId;
+	}
+
+	if (bOverrideCookRule)
+	{
+		RulesToOverride.CookRule = Rules.CookRule;
+	}
+}
+
 void FPrimaryAssetRules::PropagateCookRules(const FPrimaryAssetRules& ParentRules)
 {
 	static FPrimaryAssetRules DefaultRules;
