Add back loop check memzero (minor clean-up as move to memzero instead of memset) to zero what was garbage memory in buffer.

diff --git a/Engine/Source/Runtime/Engine/Private/ADPCMAudioInfo.cpp b/Engine/Source/Runtime/Engine/Private/ADPCMAudioInfo.cpp
--- a/Engine/Source/Runtime/Engine/Private/ADPCMAudioInfo.cpp
+++ b/Engine/Source/Runtime/Engine/Private/ADPCMAudioInfo.cpp
@@ -248,6 +248,12 @@ bool FADPCMAudioInfo::ReadCompressedData(uint8* Destination, bool bLooping, uint
 				CurrentUncompressedBlockSampleIndex = UncompressedBlockSize / sizeof(uint16);
 				CurrentCompressedBlockIndex = 0;
 				TotalSamplesStreamed = 0;
+				if(!bLooping)
+				{
+					// Zero remaining buffer
+					FMemory::Memzero(OutData, BufferSize);
+					return true;
+				}
 			}
 		}
 	}
@@ -272,8 +278,8 @@ bool FADPCMAudioInfo::ReadCompressedData(uint8* Destination, bool bLooping, uint
 			TotalSamplesStreamed = 0;
 			if(!bLooping)
 			{
-				// Set the remaining buffer to 0
-				FMemory::Memset(OutData, 0, BufferSize);
+				// Zero remaining buffer
+				FMemory::Memzero(OutData, BufferSize);
 				return true;
 			}
 		}
