[NO NON-PII COMMIT MESSAGE]

diff --git a/Engine/Source/Runtime/Engine/Private/Components/SkinnedMeshComponent.cpp b/Engine/Source/Runtime/Engine/Private/Components/SkinnedMeshComponent.cpp
--- a/Engine/Source/Runtime/Engine/Private/Components/SkinnedMeshComponent.cpp
+++ b/Engine/Source/Runtime/Engine/Private/Components/SkinnedMeshComponent.cpp
@@ -1933,12 +1933,15 @@ void USkinnedMeshComponent::AddSocketOverride(FName SourceSocketName, FName Over
 {
 	if (FName* FoundName = SocketOverrideLookup.Find(SourceSocketName))
 	{
-		if (bWarnHasOverrided)
+		if (*FoundName != OverrideSocketName)
 		{
-			UE_LOG(LogSkinnedMeshComp, Warning, TEXT("AddSocketOverride(%s, %s): Component(%s) Actor(%s) has already defined an override for socket(%s), replacing %s as override"),
-				*SourceSocketName.ToString(), *OverrideSocketName.ToString(), *GetName(), *GetNameSafe(GetOuter()), *SourceSocketName.ToString(), *(FoundName->ToString()));
+			if (bWarnHasOverrided)
+			{
+				UE_LOG(LogSkinnedMeshComp, Warning, TEXT("AddSocketOverride(%s, %s): Component(%s) Actor(%s) has already defined an override for socket(%s), replacing %s as override"),
+					*SourceSocketName.ToString(), *OverrideSocketName.ToString(), *GetName(), *GetNameSafe(GetOuter()), *SourceSocketName.ToString(), *(FoundName->ToString()));
+			}
+			*FoundName = OverrideSocketName;
 		}
-		*FoundName = OverrideSocketName;
 	}
 	else
 	{
