Calculate and emit some texture streaming stats in Test builds when STATS isn't enabled.

diff --git a/Engine/Source/Runtime/Engine/Private/Streaming/AsyncTextureStreaming.cpp b/Engine/Source/Runtime/Engine/Private/Streaming/AsyncTextureStreaming.cpp
--- a/Engine/Source/Runtime/Engine/Private/Streaming/AsyncTextureStreaming.cpp
+++ b/Engine/Source/Runtime/Engine/Private/Streaming/AsyncTextureStreaming.cpp
@@ -755,7 +755,11 @@ void FRenderAssetStreamingMipCalcTask::DoWork()
 
 	StreamingData.OnTaskDone_Async();
 
-	STAT(UpdateStats_Async());
+#if STATS
+	UpdateStats_Async();
+#elif UE_BUILD_TEST
+	UpdateCSVOnlyStats_Async();
+#endif // STATS
 }
 
 void FRenderAssetStreamingMipCalcTask::UpdateStats_Async()
@@ -860,3 +864,42 @@ void FRenderAssetStreamingMipCalcTask::UpdateStats_Async()
 	Stats.Timestamp = FPlatformTime::Seconds();
 #endif
 }
+
+void FRenderAssetStreamingMipCalcTask::UpdateCSVOnlyStats_Async()
+{
+	const TArray<FStreamingRenderAsset>& StreamingRenderAssets = StreamingManager.StreamingRenderAssets;
+
+	FRenderAssetStreamingStats& Stats = StreamingManager.GatheredStats;
+
+	Stats.RenderAssetPool = PoolSize;
+
+	Stats.SafetyPool = MemoryMargin;
+	Stats.TemporaryPool = TempMemoryBudget;
+	Stats.StreamingPool = MemoryBudget;
+	Stats.NonStreamingMips = AllocatedMemory;
+
+	Stats.RequiredPool = 0;
+	Stats.CachedMips = 0;
+	Stats.WantedMips = 0;
+
+	for (const FStreamingRenderAsset& StreamingRenderAsset : StreamingRenderAssets)
+	{
+		if (IsAborted()) break;
+		if (!StreamingRenderAsset.RenderAsset) continue;
+
+		const int64 ResidentSize = StreamingRenderAsset.GetSize(StreamingRenderAsset.ResidentMips);
+		const int64 RequiredSize = StreamingRenderAsset.GetSize(StreamingRenderAsset.GetPerfectWantedMips());
+		const int64 BudgetedSize = StreamingRenderAsset.GetSize(StreamingRenderAsset.BudgetedMips);
+
+		// How much the streamer would use if there was no limit.
+		Stats.RequiredPool += RequiredSize;
+
+		// Remove from the non streaming budget what is actually taken by streaming.
+		Stats.NonStreamingMips -= ResidentSize;
+
+		const int64 UsedSize = FMath::Min3<int64>(RequiredSize, BudgetedSize, ResidentSize);
+
+		Stats.WantedMips += UsedSize;
+		Stats.CachedMips += FMath::Max<int64>(ResidentSize - UsedSize, 0);
+	}
+}
diff --git a/Engine/Source/Runtime/Engine/Private/Streaming/AsyncTextureStreaming.h b/Engine/Source/Runtime/Engine/Private/Streaming/AsyncTextureStreaming.h
--- a/Engine/Source/Runtime/Engine/Private/Streaming/AsyncTextureStreaming.h
+++ b/Engine/Source/Runtime/Engine/Private/Streaming/AsyncTextureStreaming.h
@@ -186,6 +186,8 @@ private:
 
 	void UpdateStats_Async();
 
+	void UpdateCSVOnlyStats_Async();
+
 	FORCEINLINE TStatId GetStatId() const
 	{
 		RETURN_QUICK_DECLARE_CYCLE_STAT(FRenderAssetStreamingMipCalcTask, STATGROUP_ThreadPoolAsyncTasks);
diff --git a/Engine/Source/Runtime/Engine/Private/Streaming/StreamingManagerTexture.cpp b/Engine/Source/Runtime/Engine/Private/Streaming/StreamingManagerTexture.cpp
--- a/Engine/Source/Runtime/Engine/Private/Streaming/StreamingManagerTexture.cpp
+++ b/Engine/Source/Runtime/Engine/Private/Streaming/StreamingManagerTexture.cpp
@@ -26,6 +26,8 @@
 
 CSV_DECLARE_CATEGORY_MODULE_EXTERN(CORE_API, Basic);
 
+CSV_DEFINE_CATEGORY(TextureStreaming, true);
+
 static TAutoConsoleVariable<int32> CVarStreamingOverlapAssetAndLevelTicks(
 	TEXT("r.Streaming.OverlapAssetAndLevelTicks"),
 	!WITH_EDITOR && (PLATFORM_PS4 || PLATFORM_XBOXONE),
@@ -1259,6 +1261,11 @@ void FRenderAssetStreamingManager::UpdateStats()
 	MaxEverRequired = FMath::Max<int64>(MaxEverRequired, DisplayedStats.RequiredPool);
 }
 
+void FRenderAssetStreamingManager::UpdateCSVOnlyStats()
+{
+	DisplayedStats = GatheredStats;
+}
+
 void FRenderAssetStreamingManager::LogViewLocationChange()
 {
 #if STREAMING_LOG_VIEWCHANGES
@@ -1347,6 +1354,12 @@ void FRenderAssetStreamingManager::UpdateResourceStreaming( float DeltaTime, boo
 	LogViewLocationChange();
 	STAT(DisplayedStats.Apply();)
 
+	CSV_CUSTOM_STAT(TextureStreaming, StreamingPool, ((float)(DisplayedStats.RequiredPool + (GPoolSizeVRAMPercentage > 0 ? 0 : DisplayedStats.NonStreamingMips))) / (1024.0f * 1024.0f), ECsvCustomStatOp::Set);
+	CSV_CUSTOM_STAT(TextureStreaming, SafetyPool, ((float)DisplayedStats.SafetyPool) / (1024.0f * 1024.0f), ECsvCustomStatOp::Set);
+	CSV_CUSTOM_STAT(TextureStreaming, TemporaryPool, ((float)DisplayedStats.TemporaryPool) / (1024.0f * 1024.0f), ECsvCustomStatOp::Set);
+	CSV_CUSTOM_STAT(TextureStreaming, CachedMips, ((float)DisplayedStats.CachedMips) / (1024.0f * 1024.0f), ECsvCustomStatOp::Set);
+	CSV_CUSTOM_STAT(TextureStreaming, WantedMips, ((float)DisplayedStats.WantedMips) / (1024.0f * 1024.0f), ECsvCustomStatOp::Set);
+
 	RenderAssetInstanceAsyncWork->EnsureCompletion();
 
 	if (NumRenderAssetProcessingStages <= 0 || bProcessEverything)
@@ -1373,7 +1386,11 @@ void FRenderAssetStreamingManager::UpdateResourceStreaming( float DeltaTime, boo
 		STAT(GatheredStats.UpdateStreamingDataCycles = 0);
 		STAT(GatheredStats.StreamTexturesCycles = 0);
 		STAT(GatheredStats.CallbacksCycles = 0);
-		STAT(UpdateStats();)
+#if STATS
+		UpdateStats();
+#elif UE_BUILD_TEST
+		UpdateCSVOnlyStats();
+#endif // STATS
 	}
 	else if (ProcessingStage == 0)
 	{
@@ -1451,7 +1468,11 @@ void FRenderAssetStreamingManager::UpdateResourceStreaming( float DeltaTime, boo
 		ProcessingStage = 0;
 
 		STAT(GatheredStats.StreamTexturesCycles += FPlatformTime::Cycles();)
-		STAT(UpdateStats();)
+#if STATS
+			UpdateStats();
+#elif UE_BUILD_TEST
+			UpdateCSVOnlyStats();
+#endif // STATS
 	}
 
 	if (!bProcessEverything)
diff --git a/Engine/Source/Runtime/Engine/Private/Streaming/StreamingManagerTexture.h b/Engine/Source/Runtime/Engine/Private/Streaming/StreamingManagerTexture.h
--- a/Engine/Source/Runtime/Engine/Private/Streaming/StreamingManagerTexture.h
+++ b/Engine/Source/Runtime/Engine/Private/Streaming/StreamingManagerTexture.h
@@ -276,6 +276,7 @@ private:
 
 	void	SetLastUpdateTime();
 	void	UpdateStats();
+	void	UpdateCSVOnlyStats();
 	void	LogViewLocationChange();
 
 	void	IncrementalUpdate( float Percentage, bool bUpdateDynamicComponents);
diff --git a/Engine/Source/Runtime/Engine/Private/Streaming/TextureStreamingHelpers.cpp b/Engine/Source/Runtime/Engine/Private/Streaming/TextureStreamingHelpers.cpp
--- a/Engine/Source/Runtime/Engine/Private/Streaming/TextureStreamingHelpers.cpp
+++ b/Engine/Source/Runtime/Engine/Private/Streaming/TextureStreamingHelpers.cpp
@@ -8,7 +8,6 @@ TextureStreamingHelpers.cpp: Definitions of classes used for texture streaming.
 #include "UnrealEngine.h"
 #include "Engine/Texture2D.h"
 #include "GenericPlatform/GenericPlatformMemoryPoolStats.h"
-#include "ProfilingDebugging/CsvProfiler.h"
 
 /** Streaming stats */
 
@@ -47,8 +46,6 @@ DECLARE_MEMORY_STAT_POOL(TEXT("Average Required PoolSize"),	STAT_StreamingOvervi
 
 DEFINE_STAT(STAT_RenderAssetStreaming_GameThreadUpdateTime);
 
-CSV_DEFINE_CATEGORY(TextureStreaming, true);
-
 DEFINE_LOG_CATEGORY(LogContentStreaming);
 
 int32 FRenderAssetStreamingSettings::ExtraIOLatency = 0;
@@ -382,8 +379,6 @@ void FRenderAssetStreamingStats::Apply()
 	GRequiredPoolSizeSum += RequiredPool + (GPoolSizeVRAMPercentage > 0 ? 0 : NonStreamingMips);
 	GRequiredPoolSizeCount += 1;
 	GAverageRequiredPool = (int64)((double)GRequiredPoolSizeSum / (double)FMath::Max<int64>(1, GRequiredPoolSizeCount));
-
-	CSV_CUSTOM_STAT(TextureStreaming, StreamingPoolSize, ((float)(RequiredPool + (GPoolSizeVRAMPercentage > 0 ? 0 : NonStreamingMips))) / (1024.0f * 1024.0f), ECsvCustomStatOp::Set);
 #endif
 
 	SET_MEMORY_STAT(STAT_StreamingOverview01_StreamableTextures, RequiredPool + CachedMips); 
