#rnx

diff --git a/Engine/Source/Runtime/Engine/Private/Slate/SceneViewport.cpp b/Engine/Source/Runtime/Engine/Private/Slate/SceneViewport.cpp
--- a/Engine/Source/Runtime/Engine/Private/Slate/SceneViewport.cpp
+++ b/Engine/Source/Runtime/Engine/Private/Slate/SceneViewport.cpp
@@ -1934,6 +1934,12 @@ void FSceneViewport::InitDynamicRHI()
 
 		static const auto CVarDefaultBackBufferPixelFormat = IConsoleManager::Get().FindTConsoleVariableDataInt(TEXT("r.DefaultBackBufferPixelFormat"));
 		EPixelFormat SceneTargetFormat = EDefaultBackBufferPixelFormat::Convert2PixelFormat(EDefaultBackBufferPixelFormat::FromInt(CVarDefaultBackBufferPixelFormat->GetValueOnRenderThread()));
+#if PLATFORM_HTML5
+		if ( SceneTargetFormat == PF_A2B10G10R10 )
+		{	// UE-71220: this seems to be FDummyViewport - force to valid pixel format type for HTML5
+			SceneTargetFormat = PF_FloatRGBA;
+		}
+#endif
 
 		FRHIResourceCreateInfo CreateInfo;
 		FTexture2DRHIRef BufferedRTRHI;
