Fixing hot reload crash when adding a new C++ class in the content browser

diff --git a/Engine/Source/Runtime/Engine/Private/Animation/AnimSequence.cpp b/Engine/Source/Runtime/Engine/Private/Animation/AnimSequence.cpp
--- a/Engine/Source/Runtime/Engine/Private/Animation/AnimSequence.cpp
+++ b/Engine/Source/Runtime/Engine/Private/Animation/AnimSequence.cpp
@@ -470,7 +470,11 @@ void UAnimSequence::Serialize(FArchive& Ar)
 
 	// Do this is serialize as if the default animation curve compression asset isn't loaded it will
 	// fire a warning if we try and load it in post load
-	if (CurveCompressionSettings == nullptr || !CurveCompressionSettings->AreSettingsValid())
+	if ((CurveCompressionSettings == nullptr || !CurveCompressionSettings->AreSettingsValid()) 
+#if WITH_HOT_RELOAD
+		&& (GetClass()->HasAnyClassFlags(CLASS_CompiledFromBlueprint) || !HasAnyFlags(RF_ClassDefaultObject) || !GIsHotReload) // Don't do this to native CDOs during Hot-Reload
+#endif
+		)
 	{
 		CurveCompressionSettings = FAnimationUtils::GetDefaultAnimationCurveCompressionSettings();
 	}
