Optimization when gathering asset, by reserving space for asset managers.

When building lists of package managers for an asset, make sure to reserve memory in advance. On a full cook of a large dataset, walltime of UAssetManager::GetPackageChunkIds went from 10.5s to 9.2s.

diff --git a/Engine/Source/Runtime/AssetRegistry/Private/AssetRegistryState.cpp b/Engine/Source/Runtime/AssetRegistry/Private/AssetRegistryState.cpp
--- a/Engine/Source/Runtime/AssetRegistry/Private/AssetRegistryState.cpp
+++ b/Engine/Source/Runtime/AssetRegistry/Private/AssetRegistryState.cpp
@@ -635,6 +635,7 @@ bool FAssetRegistryState::GetReferencers(const FAssetIdentifier& AssetIdentifier
 		TArray<FDependsNode*> DependencyNodes;
 		Node->GetReferencers(DependencyNodes, InReferenceType);
 
+		OutReferencers.Reserve(DependencyNodes.Num());
 		for (FDependsNode* DependencyNode : DependencyNodes)
 		{
 			OutReferencers.Add(DependencyNode->GetIdentifier());
diff --git a/Engine/Source/Runtime/Engine/Private/AssetManager.cpp b/Engine/Source/Runtime/Engine/Private/AssetManager.cpp
--- a/Engine/Source/Runtime/Engine/Private/AssetManager.cpp
+++ b/Engine/Source/Runtime/Engine/Private/AssetManager.cpp
@@ -2616,6 +2616,7 @@ bool UAssetManager::GetPackageManagers(FName PackageName, bool bRecurseToParents
 
 	bool bFoundAny = false;
 	TArray<FAssetIdentifier> ReferencingPrimaryAssets;
+	ReferencingPrimaryAssets.Reserve(128);
 
 	AssetRegistry.GetReferencers(PackageName, ReferencingPrimaryAssets, EAssetRegistryDependencyType::Manage);
 
@@ -3216,6 +3217,7 @@ bool UAssetManager::GetPackageChunkIds(FName PackageName, const ITargetPlatform*
 
 	// Add all chunk ids from the asset rules of managers. By default priority will not override other chunks
 	TSet<FPrimaryAssetId> Managers;
+	Managers.Reserve(128);
 
 	GetPackageManagers(PackageName, true, Managers);
 	return GetPrimaryAssetSetChunkIds(Managers, TargetPlatform, ExistingChunkList, OutChunkList);
