Some extra FTimerManager thread checks.
InternalClearTimer now takes the handle by value in case of any aliasing.
Checks added to trap state errors across a timer execution.

diff --git a/Engine/Source/Runtime/Engine/Private/TimerManager.cpp b/Engine/Source/Runtime/Engine/Private/TimerManager.cpp
--- a/Engine/Source/Runtime/Engine/Private/TimerManager.cpp
+++ b/Engine/Source/Runtime/Engine/Private/TimerManager.cpp
@@ -170,6 +170,9 @@ FTimerData& FTimerManager::GetTimer(FTimerHandle const& InHandle)
 
 FTimerData* FTimerManager::FindTimer(FTimerHandle const& InHandle)
 {
+	// not currently threadsafe
+	check(IsInGameThread());
+
 	if (!InHandle.IsValid())
 	{
 		return nullptr;
@@ -297,7 +300,7 @@ FTimerHandle FTimerManager::InternalSetTimerForNextTick(FTimerUnifiedDelegate&&
 	return NewTimerHandle;
 }
 
-void FTimerManager::InternalClearTimer(FTimerHandle const& InHandle)
+void FTimerManager::InternalClearTimer(FTimerHandle InHandle)
 {
 	SCOPE_CYCLE_COUNTER(STAT_ClearTimer);
 
@@ -598,11 +601,12 @@ void FTimerManager::Tick(float DeltaTime)
 				RunTimerDelegates.Add(Top->TimerDelegate);
 #endif
 
-
+				checkf(!WillRemoveTimerAssert(CurrentlyExecutingTimer), TEXT("RemoveTimer(CurrentlyExecutingTimer) - due to fail before Execute()"));
 				Top->TimerDelegate.Execute();
 
 				// Update Top pointer, in case it has been invalidated by the Execute call
 				Top = FindTimer(CurrentlyExecutingTimer);
+				checkf(!Top || !WillRemoveTimerAssert(CurrentlyExecutingTimer), TEXT("RemoveTimer(CurrentlyExecutingTimer) - due to fail after Execute()"));
 				if (!Top || Top->Status != ETimerStatus::Executing)
 				{
 					break;
@@ -661,6 +665,9 @@ TStatId FTimerManager::GetStatId() const
 
 void FTimerManager::ListTimers() const
 {
+	// not currently threadsafe
+	check(IsInGameThread());
+
 	TArray<const FTimerData*> ValidActiveTimers;
 	ValidActiveTimers.Reserve(ActiveTimerHeap.Num());
 	for (FTimerHandle Handle : ActiveTimerHeap)
@@ -697,6 +704,14 @@ void FTimerManager::ListTimers() const
 	UE_LOG(LogEngine, Log, TEXT("------- %d Total Timers -------"), PendingTimerSet.Num() + PausedTimerSet.Num() + ValidActiveTimers.Num());
 }
 
+void FTimerManager::SetGameInstance(UGameInstance* InGameInstance)
+{
+	// not currently threadsafe
+	check(IsInGameThread());
+
+	OwningGameInstance = InGameInstance;
+}
+
 FTimerHandle FTimerManager::AddTimer(FTimerData&& TimerData)
 {
 	const void* TimerIndicesByObjectKey = TimerData.TimerDelegate.GetBoundObject();
@@ -741,6 +756,29 @@ void FTimerManager::RemoveTimer(FTimerHandle Handle)
 	Timers.RemoveAt(Handle.GetIndex());
 }
 
+bool FTimerManager::WillRemoveTimerAssert(FTimerHandle Handle) const
+{
+	const FTimerData& Data = GetTimer(Handle);
+
+	// Remove TimerIndicesByObject entry if necessary
+	if (const void* TimerIndicesByObjectKey = Data.TimerIndicesByObjectKey)
+	{
+		const TSet<FTimerHandle>* TimersForObject = ObjectToTimers.Find(TimerIndicesByObjectKey);
+		if (!TimersForObject)
+		{
+			return true;
+		}
+
+		const FTimerHandle* Found = TimersForObject->Find(Handle);
+		if (!Found)
+		{
+			return true;
+		}
+	}
+
+	return false;
+}
+
 FTimerHandle FTimerManager::GenerateHandle(int32 Index)
 {
 	uint64 NewSerialNumber = ++LastAssignedSerialNumber;
diff --git a/Engine/Source/Runtime/Engine/Public/TimerManager.h b/Engine/Source/Runtime/Engine/Public/TimerManager.h
--- a/Engine/Source/Runtime/Engine/Public/TimerManager.h
+++ b/Engine/Source/Runtime/Engine/Public/TimerManager.h
@@ -387,7 +387,7 @@ public:
 	void ListTimers() const;
 
 	/** Used by the UGameInstance constructor to set this manager's owning game instance. */
-	void SetGameInstance(UGameInstance* InGameInstance) { OwningGameInstance = InGameInstance; }
+	void SetGameInstance(UGameInstance* InGameInstance);
 
 // This should be private, but needs to be public for testing.
 public:
@@ -406,7 +406,7 @@ protected:
 private:
 	void InternalSetTimer( FTimerHandle& InOutHandle, FTimerUnifiedDelegate&& InDelegate, float InRate, bool InbLoop, float InFirstDelay );
 	FTimerHandle InternalSetTimerForNextTick( FTimerUnifiedDelegate&& InDelegate );
-	void InternalClearTimer( FTimerHandle const& InDelegate );
+	void InternalClearTimer( FTimerHandle InDelegate );
 	void InternalClearAllTimers( void const* Object );
 	float InternalGetTimerRate( FTimerData const* const TimerData ) const;
 	float InternalGetTimerElapsed( FTimerData const* const TimerData ) const;
@@ -423,6 +423,7 @@ private:
 	FTimerHandle AddTimer(FTimerData&& TimerData);
 	/** Removes a timer from the Timers list at the given index, also cleaning up the TimerIndicesByObject map */
 	void RemoveTimer(FTimerHandle Handle);
+	bool WillRemoveTimerAssert(FTimerHandle Handle) const;
 
 	/** The array of timers - all other arrays will index into this */
 	TSparseArray<FTimerData> Timers;
