Fix compile error when using FORCE_USE_STATS in shipping. In cl 3005135 the function declaration and definition for HandleListStreamingTexturesCommand was moved inside !UE_BUILD_SHIPPING blocks while the call site was not.

diff --git a/Engine/Source/Runtime/Engine/Private/Streaming/StreamingManagerTexture.cpp b/Engine/Source/Runtime/Engine/Private/Streaming/StreamingManagerTexture.cpp
--- a/Engine/Source/Runtime/Engine/Private/Streaming/StreamingManagerTexture.cpp
+++ b/Engine/Source/Runtime/Engine/Private/Streaming/StreamingManagerTexture.cpp
@@ -1814,13 +1814,11 @@ bool FStreamingManagerTexture::Exec( UWorld* InWorld, const TCHAR* Cmd, FOutputD
 		return HandleDumpTextureStreamingStatsCommand( Cmd, Ar );
 	}
 #endif
-#if STATS
+#if !UE_BUILD_SHIPPING
 	if (FParse::Command(&Cmd,TEXT("ListStreamingTextures")))
 	{
 		return HandleListStreamingTexturesCommand( Cmd, Ar );
 	}
-#endif
-#if !UE_BUILD_SHIPPING
 	if (FParse::Command(&Cmd, TEXT("ResetMaxEverRequiredTextures")))
 	{
 		return HandleResetMaxEverRequiredTexturesCommand(Cmd, Ar);
