Keep compile fix submitted in 4897024
Undo //UE4/Dev-Core/Engine/Source/Runtime/Engine/Private/Streaming/StreamingManagerTexture.h/cpp@4916080

diff --git a/Engine/Source/Runtime/Engine/Private/Streaming/StreamingManagerTexture.cpp b/Engine/Source/Runtime/Engine/Private/Streaming/StreamingManagerTexture.cpp
--- a/Engine/Source/Runtime/Engine/Private/Streaming/StreamingManagerTexture.cpp
+++ b/Engine/Source/Runtime/Engine/Private/Streaming/StreamingManagerTexture.cpp
@@ -1391,7 +1391,8 @@ bool FStreamingManagerTexture::HandleDumpTextureStreamingStatsCommand( const TCH
 }
 #endif // STATS_FAST
 
-#if STATS
+#if !UE_BUILD_SHIPPING
+
 bool FStreamingManagerTexture::HandleListStreamingTexturesCommand( const TCHAR* Cmd, FOutputDevice& Ar )
 {
 	FScopeLock ScopeLock(&CriticalSection);
@@ -1448,9 +1449,6 @@ bool FStreamingManagerTexture::HandleListStreamingTexturesCommand( const TCHAR*
 	}
 	return true;
 }
-#endif // STATS
-
-#if !UE_BUILD_SHIPPING
 
 bool FStreamingManagerTexture::HandleResetMaxEverRequiredTexturesCommand(const TCHAR* Cmd, FOutputDevice& Ar)
 {
diff --git a/Engine/Source/Runtime/Engine/Private/Streaming/StreamingManagerTexture.h b/Engine/Source/Runtime/Engine/Private/Streaming/StreamingManagerTexture.h
--- a/Engine/Source/Runtime/Engine/Private/Streaming/StreamingManagerTexture.h
+++ b/Engine/Source/Runtime/Engine/Private/Streaming/StreamingManagerTexture.h
@@ -107,10 +107,8 @@ struct FStreamingManagerTexture : public ITextureStreamingManager
 #if STATS_FAST
 	bool HandleDumpTextureStreamingStatsCommand( const TCHAR* Cmd, FOutputDevice& Ar );
 #endif // STATS_FAST
-#if STATS
-	bool HandleListStreamingTexturesCommand(const TCHAR* Cmd, FOutputDevice& Ar);
-#endif // STATS
 #if !UE_BUILD_SHIPPING
+	bool HandleListStreamingTexturesCommand( const TCHAR* Cmd, FOutputDevice& Ar );
 	bool HandleResetMaxEverRequiredTexturesCommand(const TCHAR* Cmd, FOutputDevice& Ar);
 	bool HandleLightmapStreamingFactorCommand( const TCHAR* Cmd, FOutputDevice& Ar );
 	bool HandleCancelTextureStreamingCommand( const TCHAR* Cmd, FOutputDevice& Ar );
