Performance fixes to TSortedMap with FName keys - FName comparison is way slower (case insensitive string comparison) compared to index comparison.

diff --git a/Engine/Plugins/Compositing/Composure/Source/Composure/Private/MovieScene/MovieSceneComposureExportSectionTemplate.cpp b/Engine/Plugins/Compositing/Composure/Source/Composure/Private/MovieScene/MovieSceneComposureExportSectionTemplate.cpp
--- a/Engine/Plugins/Compositing/Composure/Source/Composure/Private/MovieScene/MovieSceneComposureExportSectionTemplate.cpp
+++ b/Engine/Plugins/Compositing/Composure/Source/Composure/Private/MovieScene/MovieSceneComposureExportSectionTemplate.cpp
@@ -81,7 +81,7 @@ private:
 struct FMovieSceneComposureExportPasses
 {
 	/** Map of internal transform pass name (or NAME_None for the Output), to the export config options */
-	TSortedMap<FName, FMovieSceneComposureExportPass> PassesToExport;
+	TSortedMap<FName, FMovieSceneComposureExportPass, FDefaultAllocator, FNameSortIndexes> PassesToExport;
 
 	void AddPass(const FMovieSceneComposureExportPass& InPass, ACompositingElement* CompShotElement)
 	{
diff --git a/Engine/Source/Runtime/CoreUObject/Public/Serialization/BulkData.h b/Engine/Source/Runtime/CoreUObject/Public/Serialization/BulkData.h
--- a/Engine/Source/Runtime/CoreUObject/Public/Serialization/BulkData.h
+++ b/Engine/Source/Runtime/CoreUObject/Public/Serialization/BulkData.h
@@ -742,7 +742,7 @@ class FFormatContainer
 {
 	friend class UBodySetup;
 
-	TSortedMap<FName, FByteBulkData*> Formats;
+	TSortedMap<FName, FByteBulkData*, FDefaultAllocator, FNameSortIndexes> Formats;
 	uint32 Alignment;
 public:
 	~FFormatContainer()
diff --git a/Engine/Source/Runtime/Engine/Classes/Components/MeshComponent.h b/Engine/Source/Runtime/Engine/Classes/Components/MeshComponent.h
--- a/Engine/Source/Runtime/Engine/Classes/Components/MeshComponent.h
+++ b/Engine/Source/Runtime/Engine/Classes/Components/MeshComponent.h
@@ -141,7 +141,7 @@ protected:
 		float ScalarParameterDefaultValue = 0.f;
 	};
 
-	TSortedMap<FName, FMaterialParameterCache> MaterialParameterCache;
+	TSortedMap<FName, FMaterialParameterCache, FDefaultAllocator, FNameSortIndexes> MaterialParameterCache;
 
 	UPROPERTY(EditAnywhere, AdvancedDisplay, BlueprintReadOnly, Category = MaterialParameters)
 	bool bEnableMaterialParameterCaching;
diff --git a/Engine/Source/Runtime/Engine/Classes/Components/SkinnedMeshComponent.h b/Engine/Source/Runtime/Engine/Classes/Components/SkinnedMeshComponent.h
--- a/Engine/Source/Runtime/Engine/Classes/Components/SkinnedMeshComponent.h
+++ b/Engine/Source/Runtime/Engine/Classes/Components/SkinnedMeshComponent.h
@@ -276,7 +276,7 @@ protected:
 	/**
 	*	Mapping for socket overrides, key is the Source socket name and the value is the override socket name
 	*/
-	TSortedMap<FName, FName> SocketOverrideLookup;
+	TSortedMap<FName, FName, FDefaultAllocator, FNameSortIndexes> SocketOverrideLookup;
 
 public:
 #if WITH_EDITORONLY_DATA
diff --git a/Engine/Source/Runtime/Engine/Private/LevelTick.cpp b/Engine/Source/Runtime/Engine/Private/LevelTick.cpp
--- a/Engine/Source/Runtime/Engine/Private/LevelTick.cpp
+++ b/Engine/Source/Runtime/Engine/Private/LevelTick.cpp
@@ -1244,7 +1244,7 @@ static void RecordWorldCountsToCSV(UWorld* World)
 
 			bool bDetailed = (CVarDetailedTickContextForCSV.GetValueOnGameThread() != 0);
 
-			TSortedMap<FName, int32> TickContextToCountMap;
+			TSortedMap<FName, int32, FDefaultAllocator, FNameSortIndexes> TickContextToCountMap;
 			int32 EnabledCount;
 			FTickTaskManagerInterface::Get().GetEnabledTickFunctionCounts(World, TickContextToCountMap, EnabledCount, bDetailed);
 
diff --git a/Engine/Source/Runtime/Engine/Private/TickTaskManager.cpp b/Engine/Source/Runtime/Engine/Private/TickTaskManager.cpp
--- a/Engine/Source/Runtime/Engine/Private/TickTaskManager.cpp
+++ b/Engine/Source/Runtime/Engine/Private/TickTaskManager.cpp
@@ -1175,7 +1175,7 @@ public:
 		DisabledCount += AllDisabledTickFunctions.Num();
 	}
 
-	FORCEINLINE void AddTickFunctionToMap(TSortedMap<FName, int32>& ClassNameToCountMap, FTickFunction* Function, bool bDetailed)
+	FORCEINLINE void AddTickFunctionToMap(TSortedMap<FName, int32, FDefaultAllocator, FNameSortIndexes>& ClassNameToCountMap, FTickFunction* Function, bool bDetailed)
 	{
 		FName ContextName = Function->DiagnosticContext(bDetailed);
 		// Find entry for this context (or add it if not present)
@@ -1184,7 +1184,7 @@ public:
 		CurrentCount++; 
 	}
 
-	void AddTickFunctionsToMap(TSortedMap<FName, int32>& ClassNameToCountMap, int32& EnabledCount, bool bDetailed)
+	void AddTickFunctionsToMap(TSortedMap<FName, int32, FDefaultAllocator, FNameSortIndexes>& ClassNameToCountMap, int32& EnabledCount, bool bDetailed)
 	{
 		// Add ticks from AllEnabledTickFunctions
 		for (TSet<FTickFunction*>::TIterator It(AllEnabledTickFunctions); It; ++It)
@@ -1626,7 +1626,7 @@ private:
 
 		if (bGrouped)
 		{
-			TSortedMap<FName, int32> TickContextToCountMap;
+			TSortedMap<FName, int32, FDefaultAllocator, FNameSortIndexes> TickContextToCountMap;
 			GetEnabledTickFunctionCounts(InWorld, TickContextToCountMap, EnabledCount, true);
 
 			// Build sorted array of tick context by count
@@ -1678,7 +1678,7 @@ private:
 	}
 
 
-	virtual void GetEnabledTickFunctionCounts(UWorld* InWorld, TSortedMap<FName, int32>& TickContextToCountMap, int32& EnabledCount, bool bDetailed)
+	virtual void GetEnabledTickFunctionCounts(UWorld* InWorld, TSortedMap<FName, int32, FDefaultAllocator, FNameSortIndexes>& TickContextToCountMap, int32& EnabledCount, bool bDetailed)
 	{
 		check(InWorld);
 		check(InWorld->TickTaskLevel);
diff --git a/Engine/Source/Runtime/Engine/Public/TickTaskManagerInterface.h b/Engine/Source/Runtime/Engine/Public/TickTaskManagerInterface.h
--- a/Engine/Source/Runtime/Engine/Public/TickTaskManagerInterface.h
+++ b/Engine/Source/Runtime/Engine/Public/TickTaskManagerInterface.h
@@ -60,7 +60,7 @@ public:
 	virtual void DumpAllTickFunctions(FOutputDevice& Ar, UWorld* InWorld, bool bEnabled, bool bDisabled, bool bGrouped) = 0;
 
 	/** Returns a map of enabled ticks, grouped by 'diagnostic context' string, along with count of enabled ticks */
-	virtual void GetEnabledTickFunctionCounts(UWorld* InWorld, TSortedMap<FName, int32>& TickContextToCountMap, int32& EnabledCount, bool bDetailed) = 0;
+	virtual void GetEnabledTickFunctionCounts(UWorld* InWorld, TSortedMap<FName, int32, FDefaultAllocator, FNameSortIndexes>& TickContextToCountMap, int32& EnabledCount, bool bDetailed) = 0;
 
 	/**
 	 * Singleton to retrieve the GLOBAL tick task manager
