Fix StreamableHandle notification bug

diff --git a/Engine/Source/Runtime/Engine/Private/StreamableManager.cpp b/Engine/Source/Runtime/Engine/Private/StreamableManager.cpp
--- a/Engine/Source/Runtime/Engine/Private/StreamableManager.cpp
+++ b/Engine/Source/Runtime/Engine/Private/StreamableManager.cpp
@@ -561,7 +561,7 @@ void FStreamableHandle::UpdateCombinedHandle()
 		bAllCompleted = bAllCompleted && !ChildHandle->IsLoadingInProgress();
 		bAllCanceled = bAllCanceled && ChildHandle->WasCanceled();
 
-		if (!bAllCompleted || !bAllCanceled)
+		if (!bAllCompleted && !bAllCanceled)
 		{
 			return;
 		}
