CIS fix.

diff --git a/Engine/Source/Runtime/Engine/Private/PerPlatformProperties.cpp b/Engine/Source/Runtime/Engine/Private/PerPlatformProperties.cpp
--- a/Engine/Source/Runtime/Engine/Private/PerPlatformProperties.cpp
+++ b/Engine/Source/Runtime/Engine/Private/PerPlatformProperties.cpp
@@ -85,7 +85,7 @@ FString FPerPlatformInt::ToString() const
 #if WITH_EDITORONLY_DATA
 	TArray<FName> SortedPlatforms;
 	PerPlatform.GetKeys(/*out*/ SortedPlatforms);
-	SortedPlatforms.Sort();
+	SortedPlatforms.Sort(FNameLexicalLess());
 
 	for (FName Platform : SortedPlatforms)
 	{
