Remove no longer needed deprecation disabling pragma and pack bitfield as uint8 instead of uint16
#rnx

diff --git a/Engine/Source/Runtime/Engine/Classes/Engine/World.h b/Engine/Source/Runtime/Engine/Classes/Engine/World.h
--- a/Engine/Source/Runtime/Engine/Classes/Engine/World.h
+++ b/Engine/Source/Runtime/Engine/Classes/Engine/World.h
@@ -420,7 +420,6 @@ struct TStructOpsTypeTraits<FEndPhysicsTickFunction> : public TStructOpsTypeTrai
 };
 
 /* Struct of optional parameters passed to SpawnActor function(s). */
-PRAGMA_DISABLE_DEPRECATION_WARNINGS // Required for auto-generated functions referencing bNoCollisionFail
 struct ENGINE_API FActorSpawnParameters
 {
 	FActorSpawnParameters();
@@ -448,30 +447,29 @@ private:
 	friend class UPackageMapClient;
 
 	/* Is the actor remotely owned. This should only be set true by the package map when it is creating an actor on a client that was replicated from the server. */
-	uint16	bRemoteOwned:1;
+	uint8	bRemoteOwned:1;
 	
 public:
 
 	bool IsRemoteOwned() const { return bRemoteOwned; }
 
 	/* Determines whether spawning will not fail if certain conditions are not met. If true, spawning will not fail because the class being spawned is `bStatic=true` or because the class of the template Actor is not the same as the class of the Actor being spawned. */
-	uint16	bNoFail:1;
+	uint8	bNoFail:1;
 
 	/* Determines whether the construction script will be run. If true, the construction script will not be run on the spawned Actor. Only applicable if the Actor is being spawned from a Blueprint. */
-	uint16	bDeferConstruction:1;
+	uint8	bDeferConstruction:1;
 	
 	/* Determines whether or not the actor may be spawned when running a construction script. If true spawning will fail if a construction script is being run. */
-	uint16	bAllowDuringConstructionScript:1;
+	uint8	bAllowDuringConstructionScript:1;
 
 #if WITH_EDITOR
 	/** Determines whether the begin play cycle will run on the spawned actor when in the editor. */
-	uint16 bTemporaryEditorActor:1;
+	uint8 bTemporaryEditorActor:1;
 #endif
 	
 	/* Flags used to describe the spawned actor/object instance. */
 	EObjectFlags ObjectFlags;		
 };
-PRAGMA_ENABLE_DEPRECATION_WARNINGS
 
 
 /**
