Don't copy the pins array by default. Allow the user to use it in place or duplicate if actually needed.

diff --git a/Engine/Source/Runtime/Engine/Classes/EdGraph/EdGraphNode.h b/Engine/Source/Runtime/Engine/Classes/EdGraph/EdGraphNode.h
--- a/Engine/Source/Runtime/Engine/Classes/EdGraph/EdGraphNode.h
+++ b/Engine/Source/Runtime/Engine/Classes/EdGraph/EdGraphNode.h
@@ -359,7 +359,7 @@ public:
 	TWeakPtr<SGraphNode> DEPRECATED_NodeWidget;
 
 	/** Get all pins this node owns */
-	TArray<UEdGraphPin*> GetAllPins() { return Pins; }
+	const TArray<UEdGraphPin*>& GetAllPins() const { return Pins; }
 
 	struct FNameParameterHelper
 	{
