Stop iterating the class hierarchy when there couldn't be more input delegate binding objects to find.
#rnx

diff --git a/Engine/Source/Runtime/Engine/Private/Actor.cpp b/Engine/Source/Runtime/Engine/Private/Actor.cpp
--- a/Engine/Source/Runtime/Engine/Private/Actor.cpp
+++ b/Engine/Source/Runtime/Engine/Private/Actor.cpp
@@ -3354,10 +3354,7 @@ void AActor::EnableInput(APlayerController* PlayerController)
 			InputComponent->bBlockInput = bBlockInput;
 			InputComponent->Priority = InputPriority;
 
-			if (UInputDelegateBinding::SupportsInputDelegate(GetClass()))
-			{
-				UInputDelegateBinding::BindInputDelegates(GetClass(), InputComponent);
-			}
+			UInputDelegateBinding::BindInputDelegates(GetClass(), InputComponent);
 		}
 		else
 		{
diff --git a/Engine/Source/Runtime/Engine/Private/InputDelegateBinding.cpp b/Engine/Source/Runtime/Engine/Private/InputDelegateBinding.cpp
--- a/Engine/Source/Runtime/Engine/Private/InputDelegateBinding.cpp
+++ b/Engine/Source/Runtime/Engine/Private/InputDelegateBinding.cpp
@@ -31,7 +31,7 @@ void UInputDelegateBinding::BindInputDelegates(const UClass* InClass, UInputComp
 												UInputVectorAxisDelegateBinding::StaticClass(),
 										   };
 
-	if (InClass)
+	if (SupportsInputDelegate(InClass))
 	{
 		BindInputDelegates(InClass->GetSuperClass(), InputComponent);
 
