Get rid of autos
#rnx

diff --git a/Engine/Source/Runtime/Engine/Private/Brush.cpp b/Engine/Source/Runtime/Engine/Private/Brush.cpp
--- a/Engine/Source/Runtime/Engine/Private/Brush.cpp
+++ b/Engine/Source/Runtime/Engine/Private/Brush.cpp
@@ -157,8 +157,8 @@ void ABrush::SetIsTemporarilyHiddenInEditor( bool bIsHidden )
 	{
 		Super::SetIsTemporarilyHiddenInEditor(bIsHidden);
 		
-		auto* Level = GetLevel();
-		auto* Model = Level ? Level->Model : nullptr;
+		ULevel* Level = GetLevel();
+		UModel* Model = Level ? Level->Model : nullptr;
 
 		if (Level && Model)
 		{
