Move GetKeyFieldName out of FDataTableExporterJSON so that it is included in non-editor builds.

diff --git a/Engine/Source/Runtime/Engine/Private/DataTableJSON.cpp b/Engine/Source/Runtime/Engine/Private/DataTableJSON.cpp
--- a/Engine/Source/Runtime/Engine/Private/DataTableJSON.cpp
+++ b/Engine/Source/Runtime/Engine/Private/DataTableJSON.cpp
@@ -65,6 +65,19 @@ namespace
 
 }
 
+FString DataTableJSONUtils::GetKeyFieldName(const UDataTable& InDataTable)
+{
+	FString ExplicitString = InDataTable.ImportKeyField;
+	if (ExplicitString.IsEmpty())
+	{
+		return TEXT("Name");
+	}
+	else
+	{
+		return ExplicitString;
+	}
+}
+
 
 #if WITH_EDITOR
 
@@ -97,7 +110,7 @@ bool FDataTableExporterJSON::WriteTable(const UDataTable& InDataTable)
 		return false;
 	}
 
-	FString KeyField = GetKeyFieldName(InDataTable);
+	FString KeyField = DataTableJSONUtils::GetKeyFieldName(InDataTable);
 	JsonWriter->WriteArrayStart();
 
 	// Iterate over rows
@@ -194,19 +207,6 @@ bool FDataTableExporterJSON::WriteStruct(const UScriptStruct* InStruct, const vo
 	return true;
 }
 
-FString FDataTableExporterJSON::GetKeyFieldName(const UDataTable& InDataTable)
-{
-	FString ExplicitString = InDataTable.ImportKeyField;
-	if (ExplicitString.IsEmpty())
-	{
-		return TEXT("Name");
-	}
-	else
-	{
-		return ExplicitString;
-	}
-}
-
 bool FDataTableExporterJSON::WriteStructEntry(const void* InRowData, const UProperty* InProperty, const void* InPropertyData)
 {
 	const FString Identifier = DataTableUtils::GetPropertyExportName(InProperty, DTExportFlags);
@@ -443,7 +443,7 @@ bool FDataTableImporterJSON::ReadTable()
 bool FDataTableImporterJSON::ReadRow(const TSharedRef<FJsonObject>& InParsedTableRowObject, const int32 InRowIdx)
 {
 	// Get row name
-	FString RowKey = FDataTableExporterJSON::GetKeyFieldName(*DataTable);
+	FString RowKey = DataTableJSONUtils::GetKeyFieldName(*DataTable);
 	FName RowName = DataTableUtils::MakeValidName(InParsedTableRowObject->GetStringField(RowKey));
 
 	// Check its not 'none'
diff --git a/Engine/Source/Runtime/Engine/Private/DataTableJSON.h b/Engine/Source/Runtime/Engine/Private/DataTableJSON.h
--- a/Engine/Source/Runtime/Engine/Private/DataTableJSON.h
+++ b/Engine/Source/Runtime/Engine/Private/DataTableJSON.h
@@ -18,6 +18,12 @@ struct TPrettyJsonPrintPolicy;
 template <class CharType, class PrintPolicy>
 class TJsonWriter;
 
+namespace DataTableJSONUtils
+{
+	/** Returns what string is used as the key/name field for a data table */
+	FString GetKeyFieldName(const UDataTable& InDataTable);
+}
+
 #if WITH_EDITOR
 
 class FDataTableExporterJSON
@@ -43,9 +49,6 @@ public:
 	/** Writes the contents of a single row */
 	bool WriteStruct(const UScriptStruct* InStruct, const void* InStructData, const FString* FieldToSkip = nullptr);
 
-	/** Returns what string is used as the key/name field for a data table */
-	static FString GetKeyFieldName(const UDataTable& InDataTable);
-
 private:
 	bool WriteStructEntry(const void* InRowData, const UProperty* InProperty, const void* InPropertyData);
 
