Fix ServerTravel function to use the FURL parser to extract the map name, instead of doing it with raw string operations. This fixes an issue found while fixing UE-70487 where trying to ServerTravel to a short package path would fail to resolve the network guid of the destination map

diff --git a/Engine/Source/Runtime/Engine/Private/GameModeBase.cpp b/Engine/Source/Runtime/Engine/Private/GameModeBase.cpp
--- a/Engine/Source/Runtime/Engine/Private/GameModeBase.cpp
+++ b/Engine/Source/Runtime/Engine/Private/GameModeBase.cpp
@@ -444,37 +444,24 @@ void AGameModeBase::ProcessServerTravel(const FString& URL, bool bAbsolute)
 #if WITH_SERVER_CODE
 	StartToLeaveMap();
 
+	UE_LOG(LogGameMode, Log, TEXT("ProcessServerTravel: %s"), *URL);
+	UWorld* World = GetWorld();
+	check(World);
+	FWorldContext& WorldContext = GEngine->GetWorldContextFromWorldChecked(World);
+
 	// Force an old style load screen if the server has been up for a long time so that TimeSeconds doesn't overflow and break everything
 	bool bSeamless = (bUseSeamlessTravel && GetWorld()->TimeSeconds < 172800.0f); // 172800 seconds == 48 hours
 
-	FString NextMap;
-	if (URL.ToUpper().Contains(TEXT("?RESTART")))
-	{
-		NextMap = UWorld::RemovePIEPrefix(GetOutermost()->GetName());
-	}
-	else
-	{
-		int32 OptionStart = URL.Find(TEXT("?"));
-		if (OptionStart == INDEX_NONE)
-		{
-			NextMap = URL;
-		}
-		else
-		{
-			NextMap = URL.Left(OptionStart);
-		}
-	}
+	// Compute the next URL, and pull the map out of it. This handles short->long package name conversion
+	FURL NextURL = FURL(&WorldContext.LastURL, *URL, bAbsolute ? TRAVEL_Absolute : TRAVEL_Relative);
 
-	FGuid NextMapGuid = UEngine::GetPackageGuid(FName(*NextMap), GetWorld()->IsPlayInEditor());
+	FGuid NextMapGuid = UEngine::GetPackageGuid(FName(*NextURL.Map), GetWorld()->IsPlayInEditor());
 
 	// Notify clients we're switching level and give them time to receive.
-	FString URLMod = URL;
+	FString URLMod = NextURL.ToString();
 	APlayerController* LocalPlayer = ProcessClientTravel(URLMod, NextMapGuid, bSeamless, bAbsolute);
 
-	UE_LOG(LogGameMode, Log, TEXT("ProcessServerTravel: %s"), *URL);
-	UWorld* World = GetWorld();
-	check(World);
-	World->NextURL = URL;
+	World->NextURL = URLMod;
 	ENetMode NetMode = GetNetMode();
 
 	if (bSeamless)
