Hide import categories for composite data table

diff --git a/Engine/Source/Runtime/Engine/Classes/Engine/CompositeDataTable.h b/Engine/Source/Runtime/Engine/Classes/Engine/CompositeDataTable.h
--- a/Engine/Source/Runtime/Engine/Classes/Engine/CompositeDataTable.h
+++ b/Engine/Source/Runtime/Engine/Classes/Engine/CompositeDataTable.h
@@ -9,7 +9,7 @@
 /**
  * Data table composed of a stack of other data tables.
  */
-UCLASS(MinimalAPI, BlueprintType)
+UCLASS(MinimalAPI, BlueprintType, hideCategories=(ImportOptions,ImportSource))
 class UCompositeDataTable
 	: public UDataTable
 {
