Use cached version of Level rather than looking it up again
#rnx

diff --git a/Engine/Source/Runtime/Engine/Public/EngineUtils.h b/Engine/Source/Runtime/Engine/Public/EngineUtils.h
--- a/Engine/Source/Runtime/Engine/Public/EngineUtils.h
+++ b/Engine/Source/Runtime/Engine/Public/EngineUtils.h
@@ -311,7 +311,7 @@ public:
 				&& ActorLevel->GetWorld() == LocalCurrentWorld)
 			{
 				// ignore non-persistent world settings
-				if (LocalCurrentActor->GetLevel() == LocalCurrentWorld->PersistentLevel || !LocalCurrentActor->IsA(AWorldSettings::StaticClass()))
+				if (ActorLevel == LocalCurrentWorld->PersistentLevel || !LocalCurrentActor->IsA(AWorldSettings::StaticClass()))
 				{
 					State->CurrentActor = LocalCurrentActor;
 					State->Index = LocalIndex;
