PR #5557: Add GetActorOfClass function (Contributed by [REDACTED_PERSON])
Also clean up some similar functions

diff --git a/Engine/Source/Runtime/Engine/Classes/Kismet/GameplayStatics.h b/Engine/Source/Runtime/Engine/Classes/Kismet/GameplayStatics.h
--- a/Engine/Source/Runtime/Engine/Classes/Kismet/GameplayStatics.h
+++ b/Engine/Source/Runtime/Engine/Classes/Kismet/GameplayStatics.h
@@ -72,6 +72,15 @@ class ENGINE_API UGameplayStatics : public UBlueprintFunctionLibrary
 	/** Bind the bounds of an array of Actors */
 	UFUNCTION(BlueprintCallable, Category="Collision")
 	static void GetActorArrayBounds(const TArray<AActor*>& Actors, bool bOnlyCollidingComponents, FVector& Center, FVector& BoxExtent);
+	
+	/** 
+	 *	Find the first Actor in the world of the specified class. 
+	 *	This is a slow operation, use with caution e.g. do not use every frame.
+	 *	@param	ActorClass	Class of Actor to find. Must be specified or result will be empty.
+	 *	@return				Actor of the specified class.
+	 */
+	UFUNCTION(BlueprintCallable, Category="Utilities", meta=(WorldContext="WorldContextObject", DeterminesOutputType="ActorClass"))
+	static class AActor* GetActorOfClass(const UObject* WorldContextObject, TSubclassOf<AActor> ActorClass);
 
 	/** 
 	 *	Find all Actors in the world of the specified class. 
diff --git a/Engine/Source/Runtime/Engine/Private/GameplayStatics.cpp b/Engine/Source/Runtime/Engine/Private/GameplayStatics.cpp
--- a/Engine/Source/Runtime/Engine/Private/GameplayStatics.cpp
+++ b/Engine/Source/Runtime/Engine/Private/GameplayStatics.cpp
@@ -763,21 +763,39 @@ void UGameplayStatics::GetActorArrayBounds(const TArray<AActor*>& Actors, bool b
 	}
 }
 
+AActor* UGameplayStatics::GetActorOfClass(const UObject* WorldContextObject, TSubclassOf<AActor> ActorClass)
+{
+	QUICK_SCOPE_CYCLE_COUNTER(UGameplayStatics_GetActorOfClass);
+
+	// We do nothing if no is class provided
+	if (ActorClass)
+	{
+		if (UWorld* World = GEngine->GetWorldFromContextObject(WorldContextObject, EGetWorldErrorMode::LogAndReturnNull))
+		{
+			for (TActorIterator<AActor> It(World, ActorClass); It; ++It)
+			{
+				AActor* Actor = *It;
+				return Actor;
+			}
+		}
+	}
+
+	return nullptr;
+}
+
 void UGameplayStatics::GetAllActorsOfClass(const UObject* WorldContextObject, TSubclassOf<AActor> ActorClass, TArray<AActor*>& OutActors)
 {
 	QUICK_SCOPE_CYCLE_COUNTER(UGameplayStatics_GetAllActorsOfClass);
 	OutActors.Reset();
 
-	UWorld* World = GEngine->GetWorldFromContextObject(WorldContextObject, EGetWorldErrorMode::LogAndReturnNull);
-
 	// We do nothing if no is class provided, rather than giving ALL actors!
-	if (ActorClass && World)
+	if (ActorClass)
 	{
-		for(TActorIterator<AActor> It(World, ActorClass); It; ++It)
+		if (UWorld* World = GEngine->GetWorldFromContextObject(WorldContextObject, EGetWorldErrorMode::LogAndReturnNull))
 		{
-			AActor* Actor = *It;
-			if(!Actor->IsPendingKill())
+			for (TActorIterator<AActor> It(World, ActorClass); It; ++It)
 			{
+				AActor* Actor = *It;
 				OutActors.Add(Actor);
 			}
 		}
@@ -787,18 +805,20 @@ void UGameplayStatics::GetAllActorsOfClass(const UObject* WorldContextObject, TS
 void UGameplayStatics::GetAllActorsWithInterface(const UObject* WorldContextObject, TSubclassOf<UInterface> Interface, TArray<AActor*>& OutActors)
 {
 	QUICK_SCOPE_CYCLE_COUNTER(UGameplayStatics_GetAllActorsWithTag);
-	OutActors.Empty();
+	OutActors.Reset();
 
-	UWorld* World = GEngine->GetWorldFromContextObject(WorldContextObject, EGetWorldErrorMode::LogAndReturnNull);
-	// We do nothing if not class provided, rather than giving ALL actors!
-	if (Interface && World)
+	// We do nothing if no interface provided, rather than giving ALL actors!
+	if (Interface)
 	{
-		for(FActorIterator It(World); It; ++It)
+		if (UWorld* World = GEngine->GetWorldFromContextObject(WorldContextObject, EGetWorldErrorMode::LogAndReturnNull))
 		{
-			AActor* Actor = *It;
-			if (Actor && !Actor->IsPendingKill() && Actor->GetClass()->ImplementsInterface(Interface))
+			for (FActorIterator It(World); It; ++It)
 			{
-				OutActors.Add(Actor);
+				AActor* Actor = *It;
+				if (Actor->GetClass()->ImplementsInterface(Interface))
+				{
+					OutActors.Add(Actor);
+				}
 			}
 		}
 	}
@@ -807,19 +827,20 @@ void UGameplayStatics::GetAllActorsWithInterface(const UObject* WorldContextObje
 void UGameplayStatics::GetAllActorsWithTag(const UObject* WorldContextObject, FName Tag, TArray<AActor*>& OutActors)
 {
 	QUICK_SCOPE_CYCLE_COUNTER(UGameplayStatics_GetAllActorsWithTag);
-	OutActors.Empty();
-
-	UWorld* World = GEngine->GetWorldFromContextObject(WorldContextObject, EGetWorldErrorMode::LogAndReturnNull);
+	OutActors.Reset();
 
 	// We do nothing if no tag is provided, rather than giving ALL actors!
-	if (!Tag.IsNone() && World)
+	if (!Tag.IsNone())
 	{
-		for (FActorIterator It(World); It; ++It)
+		if (UWorld* World = GEngine->GetWorldFromContextObject(WorldContextObject, EGetWorldErrorMode::LogAndReturnNull))
 		{
-			AActor* Actor = *It;
-			if (Actor && !Actor->IsPendingKill() && Actor->ActorHasTag(Tag))
+			for (FActorIterator It(World); It; ++It)
 			{
-				OutActors.Add(Actor);
+				AActor* Actor = *It;
+				if (Actor->ActorHasTag(Tag))
+				{
+					OutActors.Add(Actor);
+				}
 			}
 		}
 	}
