PR #5562: Add GetInputChordDisplayName function (Contributed by [REDACTED_PERSON])

diff --git a/Engine/Source/Runtime/Engine/Classes/Kismet/KismetInputLibrary.h b/Engine/Source/Runtime/Engine/Classes/Kismet/KismetInputLibrary.h
--- a/Engine/Source/Runtime/Engine/Classes/Kismet/KismetInputLibrary.h
+++ b/Engine/Source/Runtime/Engine/Classes/Kismet/KismetInputLibrary.h
@@ -174,6 +174,9 @@ class ENGINE_API UKismetInputLibrary : public UBlueprintFunctionLibrary
 	UFUNCTION(BlueprintPure, meta=( DisplayName = "Is Right Command Down" ), Category="Utilities|InputEvent")
 	static bool InputEvent_IsRightCommandDown(const FInputEvent& Input);
 
+	/** @return The display name of the input chord */
+	UFUNCTION(BlueprintPure, meta = (DisplayName = "Get Input Chord Display Name"), Category = "Utilities|Key") 
+	static FText InputChord_GetDisplayName(const FInputChord& Key); 
 
 	/**
 	 * Returns the key for this event.
diff --git a/Engine/Source/Runtime/Engine/Private/KismetInputLibrary.cpp b/Engine/Source/Runtime/Engine/Private/KismetInputLibrary.cpp
--- a/Engine/Source/Runtime/Engine/Private/KismetInputLibrary.cpp
+++ b/Engine/Source/Runtime/Engine/Private/KismetInputLibrary.cpp
@@ -137,6 +137,10 @@ bool UKismetInputLibrary::InputEvent_IsRightCommandDown(const FInputEvent& Input
 	return Input.IsRightCommandDown();
 }
 
+FText UKismetInputLibrary::InputChord_GetDisplayName(const FInputChord& Key) 
+{ 
+	return Key.GetInputText(); 
+}
 
 FKey UKismetInputLibrary::GetKey(const FKeyEvent& Input)
 {
