Clean up autos out of ICH
Pass structs by reference
Other assorted cleanup
#rnx

diff --git a/Engine/Source/Runtime/Engine/Classes/Engine/InheritableComponentHandler.h b/Engine/Source/Runtime/Engine/Classes/Engine/InheritableComponentHandler.h
--- a/Engine/Source/Runtime/Engine/Classes/Engine/InheritableComponentHandler.h
+++ b/Engine/Source/Runtime/Engine/Classes/Engine/InheritableComponentHandler.h
@@ -29,7 +29,7 @@ struct ENGINE_API FComponentKey
 	FComponentKey(UBlueprint* Blueprint, const FUCSComponentId& UCSComponentID);
 #endif 
 
-	bool Match(const FComponentKey OtherKey) const;
+	bool Match(const FComponentKey& OtherKey) const;
 
 	bool IsSCSKey() const
 	{
@@ -105,19 +105,19 @@ private:
 
 public:
 
-	UActorComponent* CreateOverridenComponentTemplate(FComponentKey Key);
-	void RemoveOverridenComponentTemplate(FComponentKey Key);
+	UActorComponent* CreateOverridenComponentTemplate(const FComponentKey& Key);
+	void RemoveOverridenComponentTemplate(const FComponentKey& Key);
 	void UpdateOwnerClass(UBlueprintGeneratedClass* OwnerClass);
 	void ValidateTemplates();
 	bool IsValid() const;
-	UActorComponent* FindBestArchetype(FComponentKey Key) const;
+	UActorComponent* FindBestArchetype(const FComponentKey& Key) const;
 
 	bool IsEmpty() const
 	{
 		return 0 == Records.Num();
 	}
 
-	bool RefreshTemplateName(FComponentKey OldKey);
+	bool RefreshTemplateName(const FComponentKey& OldKey);
 
 	FComponentKey FindKey(UActorComponent* ComponentTemplate) const;
 #endif
@@ -134,8 +134,8 @@ public:
 
 	FComponentKey FindKey(const FName VariableName) const;
 
-	UActorComponent* GetOverridenComponentTemplate(FComponentKey Key) const;
-	const FBlueprintCookedComponentInstancingData* GetOverridenComponentTemplateData(FComponentKey Key) const;
+	UActorComponent* GetOverridenComponentTemplate(const FComponentKey& Key) const;
+	const FBlueprintCookedComponentInstancingData* GetOverridenComponentTemplateData(const FComponentKey& Key) const;
 
 	TArray<FComponentOverrideRecord>::TIterator CreateRecordIterator()
 	{
@@ -144,14 +144,15 @@ public:
 
 	void GetAllTemplates(TArray<UActorComponent*>& OutArray) const
 	{
-		for (auto Record : Records)
+		OutArray.Reserve(OutArray.Num() + Records.Num());
+		for (const FComponentOverrideRecord& Record : Records)
 		{
 			OutArray.Add(Record.ComponentTemplate);
 		}
 	}
 
 private:
-	const FComponentOverrideRecord* FindRecord(const FComponentKey Key) const;
+	const FComponentOverrideRecord* FindRecord(const FComponentKey& Key) const;
 
 	/** Helper method used to assist with fixing up component template names at load time. */
 	void FixComponentTemplateName(UActorComponent* ComponentTemplate, const FString& NewName);
diff --git a/Engine/Source/Runtime/Engine/Private/InheritableComponentHandler.cpp b/Engine/Source/Runtime/Engine/Private/InheritableComponentHandler.cpp
--- a/Engine/Source/Runtime/Engine/Private/InheritableComponentHandler.cpp
+++ b/Engine/Source/Runtime/Engine/Private/InheritableComponentHandler.cpp
@@ -77,7 +77,7 @@ void UInheritableComponentHandler::PostLoad()
 }
 
 #if WITH_EDITOR
-UActorComponent* UInheritableComponentHandler::CreateOverridenComponentTemplate(FComponentKey Key)
+UActorComponent* UInheritableComponentHandler::CreateOverridenComponentTemplate(const FComponentKey& Key)
 {
 	for (int32 Index = 0; Index < Records.Num(); ++Index)
 	{
@@ -93,7 +93,7 @@ UActorComponent* UInheritableComponentHandler::CreateOverridenComponentTemplate(
 		}
 	}
 
-	auto BestArchetype = FindBestArchetype(Key);
+	UActorComponent* BestArchetype = FindBestArchetype(Key);
 	if (!BestArchetype)
 	{
 		UE_LOG(LogBlueprint, Warning, TEXT("CreateOverridenComponentTemplate '%s': cannot find archetype for component '%s' from '%s'"),
@@ -113,7 +113,7 @@ UActorComponent* UInheritableComponentHandler::CreateOverridenComponentTemplate(
 	}
 
 	ensure(Cast<UBlueprintGeneratedClass>(GetOuter()));
-	auto NewComponentTemplate = NewObject<UActorComponent>(
+	UActorComponent* NewComponentTemplate = NewObject<UActorComponent>(
 		GetOuter(), BestArchetype->GetClass(), NewComponentTemplateName, RF_ArchetypeObject | RF_Public | RF_InheritableComponentTemplate, BestArchetype);
 
 	// HACK: NewObject can return a pre-existing object which will not have been initialized to the archetype.  When we remove the old handlers, we mark them pending
@@ -127,30 +127,34 @@ UActorComponent* UInheritableComponentHandler::CreateOverridenComponentTemplate(
 	}
 
 	// Clear transient flag if it was transient before and re copy off archetype
-	if (NewComponentTemplate->HasAnyFlags(RF_Transient) && UnnecessaryComponents.Contains(NewComponentTemplate))
+	if (NewComponentTemplate->HasAnyFlags(RF_Transient))
 	{
-		NewComponentTemplate->ClearFlags(RF_Transient);
-		UnnecessaryComponents.Remove(NewComponentTemplate);
+		const int32 ComponentIndex = UnnecessaryComponents.Find(NewComponentTemplate);
+		if (ComponentIndex != INDEX_NONE)
+		{
+			NewComponentTemplate->ClearFlags(RF_Transient);
+			UnnecessaryComponents.RemoveAtSwap(ComponentIndex);
 
-		UEngine::FCopyPropertiesForUnrelatedObjectsParams CopyParams;
-		CopyParams.bDoDelta = false;
-		UEngine::CopyPropertiesForUnrelatedObjects(BestArchetype, NewComponentTemplate, CopyParams);
+			UEngine::FCopyPropertiesForUnrelatedObjectsParams CopyParams;
+			CopyParams.bDoDelta = false;
+			UEngine::CopyPropertiesForUnrelatedObjects(BestArchetype, NewComponentTemplate, CopyParams);
+		}
 	}
 
 	FComponentOverrideRecord NewRecord;
 	NewRecord.ComponentKey = Key;
 	NewRecord.ComponentClass = NewComponentTemplate->GetClass();
 	NewRecord.ComponentTemplate = NewComponentTemplate;
-	Records.Add(NewRecord);
+	Records.Emplace(MoveTemp(NewRecord));
 
 	return NewComponentTemplate;
 }
 
-void UInheritableComponentHandler::RemoveOverridenComponentTemplate(FComponentKey Key)
+void UInheritableComponentHandler::RemoveOverridenComponentTemplate(const FComponentKey& Key)
 {
 	for (int32 Index = 0; Index < Records.Num(); ++Index)
 	{
-		FComponentOverrideRecord& Record = Records[Index];
+		const FComponentOverrideRecord& Record = Records[Index];
 		if (Record.ComponentKey.Match(Key))
 		{
 			Record.ComponentTemplate->MarkPendingKill(); // hack needed to be able to identify if NewObject returns this back to us in the future
@@ -162,9 +166,9 @@ void UInheritableComponentHandler::RemoveOverridenComponentTemplate(FComponentKe
 
 void UInheritableComponentHandler::UpdateOwnerClass(UBlueprintGeneratedClass* OwnerClass)
 {
-	for (auto& Record : Records)
+	for (FComponentOverrideRecord& Record : Records)
 	{
-		auto OldComponentTemplate = Record.ComponentTemplate;
+		UActorComponent* OldComponentTemplate = Record.ComponentTemplate;
 		if (OldComponentTemplate && (OwnerClass != OldComponentTemplate->GetOuter()))
 		{
 			Record.ComponentTemplate = DuplicateObject(OldComponentTemplate, OwnerClass, OldComponentTemplate->GetFName());
@@ -231,7 +235,7 @@ void UInheritableComponentHandler::ValidateTemplates()
 
 bool UInheritableComponentHandler::IsValid() const
 {
-	for (auto& Record : Records)
+	for (const FComponentOverrideRecord& Record : Records)
 	{
 		if (!IsRecordValid(Record))
 		{
@@ -341,18 +345,18 @@ bool UInheritableComponentHandler::IsRecordNecessary(const FComponentOverrideRec
 			return false;
 		}
 	
-		auto ChildComponentTemplate = Record.ComponentTemplate;
-		auto ParentComponentTemplate = FindBestArchetype(Record.ComponentKey);
+		UActorComponent* ChildComponentTemplate = Record.ComponentTemplate;
+		UActorComponent* ParentComponentTemplate = FindBestArchetype(Record.ComponentKey);
 		check(ChildComponentTemplate && ParentComponentTemplate && (ParentComponentTemplate != ChildComponentTemplate));
 		return !FComponentComparisonHelper::AreIdentical(ChildComponentTemplate, ParentComponentTemplate);
 	}
 }
 
-UActorComponent* UInheritableComponentHandler::FindBestArchetype(FComponentKey Key) const
+UActorComponent* UInheritableComponentHandler::FindBestArchetype(const FComponentKey& Key) const
 {
 	UActorComponent* ClosestArchetype = nullptr;
 
-	auto ActualBPGC = Cast<UBlueprintGeneratedClass>(GetOuter());
+	UBlueprintGeneratedClass* ActualBPGC = Cast<UBlueprintGeneratedClass>(GetOuter());
 	if (ActualBPGC && Key.GetComponentOwner() && (ActualBPGC != Key.GetComponentOwner()))
 	{
 		ActualBPGC = Cast<UBlueprintGeneratedClass>(ActualBPGC->GetSuperClass());
@@ -374,9 +378,9 @@ UActorComponent* UInheritableComponentHandler::FindBestArchetype(FComponentKey K
 	return ClosestArchetype;
 }
 
-bool UInheritableComponentHandler::RefreshTemplateName(FComponentKey OldKey)
+bool UInheritableComponentHandler::RefreshTemplateName(const FComponentKey& OldKey)
 {
-	for (auto& Record : Records)
+	for (FComponentOverrideRecord& Record : Records)
 	{
 		if (Record.ComponentKey.Match(OldKey))
 		{
@@ -389,7 +393,7 @@ bool UInheritableComponentHandler::RefreshTemplateName(FComponentKey OldKey)
 
 FComponentKey UInheritableComponentHandler::FindKey(UActorComponent* ComponentTemplate) const
 {
-	for (auto& Record : Records)
+	for (const FComponentOverrideRecord& Record : Records)
 	{
 		if (Record.ComponentTemplate == ComponentTemplate)
 		{
@@ -403,12 +407,11 @@ FComponentKey UInheritableComponentHandler::FindKey(UActorComponent* ComponentTe
 
 void UInheritableComponentHandler::PreloadAllTempates()
 {
-	for (auto Record : Records)
+	for (const FComponentOverrideRecord& Record : Records)
 	{
 		if (Record.ComponentTemplate && Record.ComponentTemplate->HasAllFlags(RF_NeedLoad))
 		{
-			auto Linker = Record.ComponentTemplate->GetLinker();
-			if (Linker)
+			if (FLinkerLoad* Linker = Record.ComponentTemplate->GetLinker())
 			{
 				Linker->Preload(Record.ComponentTemplate);
 			}
@@ -420,8 +423,7 @@ void UInheritableComponentHandler::PreloadAll()
 {
 	if (HasAllFlags(RF_NeedLoad))
 	{
-		auto Linker = GetLinker();
-		if (Linker)
+		if (FLinkerLoad* Linker = GetLinker())
 		{
 			Linker->Preload(this);
 		}
@@ -441,21 +443,21 @@ FComponentKey UInheritableComponentHandler::FindKey(const FName VariableName) co
 	return FComponentKey();
 }
 
-UActorComponent* UInheritableComponentHandler::GetOverridenComponentTemplate(FComponentKey Key) const
+UActorComponent* UInheritableComponentHandler::GetOverridenComponentTemplate(const FComponentKey& Key) const
 {
-	auto Record = FindRecord(Key);
+	const FComponentOverrideRecord* Record = FindRecord(Key);
 	return Record ? Record->ComponentTemplate : nullptr;
 }
 
-const FBlueprintCookedComponentInstancingData* UInheritableComponentHandler::GetOverridenComponentTemplateData(FComponentKey Key) const
+const FBlueprintCookedComponentInstancingData* UInheritableComponentHandler::GetOverridenComponentTemplateData(const FComponentKey& Key) const
 {
-	auto Record = FindRecord(Key);
+	const FComponentOverrideRecord* Record = FindRecord(Key);
 	return Record ? &Record->CookedComponentInstancingData : nullptr;
 }
 
-const FComponentOverrideRecord* UInheritableComponentHandler::FindRecord(const FComponentKey Key) const
+const FComponentOverrideRecord* UInheritableComponentHandler::FindRecord(const FComponentKey& Key) const
 {
-	for (auto& Record : Records)
+	for (const FComponentOverrideRecord& Record : Records)
 	{
 		if (Record.ComponentKey.Match(Key))
 		{
@@ -469,12 +471,12 @@ void UInheritableComponentHandler::FixComponentTemplateName(UActorComponent* Com
 {
 	// Look for a collision with the template we're trying to rename here. It's possible that names were swapped on the
 	// original component template objects that were inherited from the associated Blueprint's parent class, for example.
-	FComponentOverrideRecord* MatchingRecord = Records.FindByPredicate([ComponentTemplate, NewName](FComponentOverrideRecord& Record)
+	FComponentOverrideRecord* MatchingRecord = Records.FindByPredicate([ComponentTemplate, &NewName](FComponentOverrideRecord& Record)
 	{
 		if (Record.ComponentTemplate && Record.ComponentTemplate != ComponentTemplate && Record.ComponentTemplate->GetName() == NewName)
 		{
 			const UActorComponent* OriginalTemplate = Record.ComponentKey.GetOriginalTemplate();
-			return ensureMsgf(OriginalTemplate && OriginalTemplate->GetName() != Record.ComponentTemplate->GetName(),
+			return ensureMsgf(OriginalTemplate && OriginalTemplate->GetFName() != Record.ComponentTemplate->GetFName(),
 				TEXT("Found a collision with an existing override record, but its associated template object is either invalid or already matches its inherited template's name (%s). This is unexpected."), *NewName);
 		}
 
@@ -514,7 +516,7 @@ FComponentKey::FComponentKey(UBlueprint* Blueprint, const FUCSComponentId& UCSCo
 }
 #endif // WITH_EDITOR
 
-bool FComponentKey::Match(const FComponentKey OtherKey) const
+bool FComponentKey::Match(const FComponentKey& OtherKey) const
 {
 	return (OwnerClass == OtherKey.OwnerClass) && (AssociatedGuid == OtherKey.AssociatedGuid);
 }
