PR #5642: Added comment about possible missing super (Contributed by [REDACTED_PERSON])
#rnx

diff --git a/Engine/Source/Runtime/Engine/Private/Components/ActorComponent.cpp b/Engine/Source/Runtime/Engine/Private/Components/ActorComponent.cpp
--- a/Engine/Source/Runtime/Engine/Private/Components/ActorComponent.cpp
+++ b/Engine/Source/Runtime/Engine/Private/Components/ActorComponent.cpp
@@ -563,6 +563,7 @@ void UActorComponent::PreEditChange(UProperty* PropertyThatWillChange)
 		// Don't do do a full recreate in this situation, and instead simply detach.
 		if( !IsPendingKill() )
 		{
+			// One way this check can fail is that component subclass does not call Super::PostEditChangeProperty
 			check(!EditReregisterContexts.Find(this));
 			EditReregisterContexts.Add(this,new FComponentReregisterContext(this));
 		}
