Add new GetAuthoredName functions to UField and UEnum that expose this via virtuals instead of doing it in data table utils
If PPF_ExternalEditor is passed to export text, it uses the authored names and also turns off delta serialization for containers
Change CustomFindProperty to not be editor only and add CheckAuthoredName to enum name lookup, this allows importing data tables using user structs at runtime for hot fixing
Deprecate UsePrettyProperty/EnumNames for data table export, the comments were wrong before and that behavior is now always enabled
Changed several functions with "DisplayName" to use authored or friendly name instead, as DisplayName should always refer to the DisplayName MetaData and is usually translated

diff --git a/Engine/Source/Editor/Kismet/Private/UserDefinedStructureEditor.cpp b/Engine/Source/Editor/Kismet/Private/UserDefinedStructureEditor.cpp
--- a/Engine/Source/Editor/Kismet/Private/UserDefinedStructureEditor.cpp
+++ b/Engine/Source/Editor/Kismet/Private/UserDefinedStructureEditor.cpp
@@ -588,7 +588,7 @@ public:
 		auto StructureDetailsSP = StructureDetails.Pin();
 		if(StructureDetailsSP.IsValid())
 		{
-			return FText::FromString(FStructureEditorUtils::GetVariableDisplayName(StructureDetailsSP->GetUserDefinedStruct(), FieldGuid));
+			return FText::FromString(FStructureEditorUtils::GetVariableFriendlyName(StructureDetailsSP->GetUserDefinedStruct(), FieldGuid));
 		}
 		return FText::GetEmpty();
 	}
diff --git a/Engine/Source/Editor/UnrealEd/Private/DataTableEditorUtils.cpp b/Engine/Source/Editor/UnrealEd/Private/DataTableEditorUtils.cpp
--- a/Engine/Source/Editor/UnrealEd/Private/DataTableEditorUtils.cpp
+++ b/Engine/Source/Editor/UnrealEd/Private/DataTableEditorUtils.cpp
@@ -693,7 +693,7 @@ void FDataTableEditorUtils::CacheDataTableForEditing(const UDataTable* DataTable
 	for (int32 Index = 0; Index < StructProps.Num(); ++Index)
 	{
 		const UProperty* Prop = StructProps[Index];
-		const FText PropertyDisplayName = FText::FromString(DataTableUtils::GetPropertyDisplayName(Prop, FName::NameToDisplayString(Prop->GetName(), Prop->IsA<UBoolProperty>())));
+		const FText PropertyDisplayName = DataTableUtils::GetPropertyDisplayName(Prop, FName::NameToDisplayString(Prop->GetName(), Prop->IsA<UBoolProperty>()));
 
 		FDataTableEditorColumnHeaderDataPtr CachedColumnData;
 		
diff --git a/Engine/Source/Editor/UnrealEd/Private/Kismet2/StructureEditorUtils.cpp b/Engine/Source/Editor/UnrealEd/Private/Kismet2/StructureEditorUtils.cpp
--- a/Engine/Source/Editor/UnrealEd/Private/Kismet2/StructureEditorUtils.cpp
+++ b/Engine/Source/Editor/UnrealEd/Private/Kismet2/StructureEditorUtils.cpp
@@ -305,7 +305,7 @@ bool FStructureEditorUtils::AddVariable(UUserDefinedStruct* Struct, const FEdGra
 		FString DisplayName;
 		const FName VarName = FMemberVariableNameHelper::Generate(Struct, FString(), Guid, &DisplayName);
 		check(NULL == GetVarDesc(Struct).FindByPredicate(FStructureEditorUtils::FFindByNameHelper<FStructVariableDescription>(VarName)));
-		check(IsUniqueVariableDisplayName(Struct, DisplayName));
+		check(IsUniqueVariableFriendlyName(Struct, DisplayName));
 
 		FStructVariableDescription NewVar;
 		NewVar.VarName = VarName;
@@ -353,7 +353,7 @@ bool FStructureEditorUtils::RenameVariable(UUserDefinedStruct* Struct, FGuid Var
 		FStructVariableDescription* VarDesc = GetVarDescByGuid(Struct, VarGuid);
 		if (VarDesc 
 			&& !NewDisplayNameStr.IsEmpty()
-			&& IsUniqueVariableDisplayName(Struct, NewDisplayNameStr))
+			&& IsUniqueVariableFriendlyName(Struct, NewDisplayNameStr))
 		{
 			const FScopedTransaction Transaction(LOCTEXT("RenameVariable", "Rename Variable"));
 			ModifyStructData(Struct);
@@ -466,7 +466,7 @@ bool FStructureEditorUtils::ChangeVariableDefaultValue(UUserDefinedStruct* Struc
 	return false;
 }
 
-bool FStructureEditorUtils::IsUniqueVariableDisplayName(const UUserDefinedStruct* Struct, const FString& DisplayName)
+bool FStructureEditorUtils::IsUniqueVariableFriendlyName(const UUserDefinedStruct* Struct, const FString& DisplayName)
 {
 	if(Struct)
 	{
@@ -482,13 +482,19 @@ bool FStructureEditorUtils::IsUniqueVariableDisplayName(const UUserDefinedStruct
 	return false;
 }
 
-FString FStructureEditorUtils::GetVariableDisplayName(const UUserDefinedStruct* Struct, FGuid VarGuid)
+FString FStructureEditorUtils::GetVariableFriendlyName(const UUserDefinedStruct* Struct, FGuid VarGuid)
 {
 	const FStructVariableDescription* VarDesc = GetVarDescByGuid(Struct, VarGuid);
 	return VarDesc ? VarDesc->FriendlyName : FString();
 }
 
-UProperty* FStructureEditorUtils::GetPropertyByDisplayName(const UUserDefinedStruct* Struct, FString DisplayName)
+FString FStructureEditorUtils::GetVariableFriendlyNameForProperty(const UUserDefinedStruct* Struct, const UProperty* Property)
+{
+	const FStructVariableDescription* VarDesc = Struct && Property ? GetVarDesc(Struct).FindByPredicate(FFindByNameHelper<FStructVariableDescription>(Property->GetFName())) : nullptr;
+	return VarDesc ? VarDesc->FriendlyName : FString();
+}
+
+UProperty* FStructureEditorUtils::GetPropertyByFriendlyName(const UUserDefinedStruct* Struct, FString DisplayName)
 {
 	if (Struct)
 	{
diff --git a/Engine/Source/Editor/UnrealEd/Public/Kismet2/StructureEditorUtils.h b/Engine/Source/Editor/UnrealEd/Public/Kismet2/StructureEditorUtils.h
--- a/Engine/Source/Editor/UnrealEd/Public/Kismet2/StructureEditorUtils.h
+++ b/Engine/Source/Editor/UnrealEd/Public/Kismet2/StructureEditorUtils.h
@@ -87,11 +87,13 @@ public:
 
 	static bool ChangeVariableDefaultValue(UUserDefinedStruct* Struct, FGuid VarGuid, const FString& NewDefaultValue);
 
-	static bool IsUniqueVariableDisplayName(const UUserDefinedStruct* Struct, const FString& DisplayName);
+	static bool IsUniqueVariableFriendlyName(const UUserDefinedStruct* Struct, const FString& DisplayName);
 
-	static FString GetVariableDisplayName(const UUserDefinedStruct* Struct, FGuid VarGuid);
+	static FString GetVariableFriendlyName(const UUserDefinedStruct* Struct, FGuid VarGuid);
 
-	static UProperty* GetPropertyByDisplayName(const UUserDefinedStruct* Struct, FString DisplayName);
+	static FString GetVariableFriendlyNameForProperty(const UUserDefinedStruct* Struct, const UProperty* Property);
+
+	static UProperty* GetPropertyByFriendlyName(const UUserDefinedStruct* Struct, FString DisplayName);
 
 	static FString GetVariableTooltip(const UUserDefinedStruct* Struct, FGuid VarGuid);
 	
diff --git a/Engine/Source/Runtime/Core/Public/UObject/PropertyPortFlags.h b/Engine/Source/Runtime/Core/Public/UObject/PropertyPortFlags.h
--- a/Engine/Source/Runtime/Core/Public/UObject/PropertyPortFlags.h
+++ b/Engine/Source/Runtime/Core/Public/UObject/PropertyPortFlags.h
@@ -32,7 +32,11 @@ enum EPropertyPortFlags
 	/** Indicates that this is a blueprint pin or something else that is saved to disk as import text */
 	PPF_SerializedAsImportText		= 0x00000040,
 
-	//								= 0x00000080,
+	/** 
+	 * Set if this is being exported/imported for an external editor format like CSV.
+	 * If true, it uses authored names instead of internal names and default values are always written out.
+	 */
+	PPF_ExternalEditor				= 0x00000080,
 
 	/** only include properties which are marked CPF_InstancedReference */
 	PPF_SubobjectsOnly				= 0x00000100,
diff --git a/Engine/Source/Runtime/CoreUObject/Private/UObject/Class.cpp b/Engine/Source/Runtime/CoreUObject/Private/UObject/Class.cpp
--- a/Engine/Source/Runtime/CoreUObject/Private/UObject/Class.cpp
+++ b/Engine/Source/Runtime/CoreUObject/Private/UObject/Class.cpp
@@ -166,6 +166,16 @@ UStruct* UField::GetOwnerStruct() const
 	return nullptr;
 }
 
+FString UField::GetAuthoredName() const
+{
+	UStruct* Struct = GetOwnerStruct();
+	if (Struct)
+	{
+		return Struct->GetAuthoredNameForField(this);
+	}
+	return FString();
+}
+
 void UField::Bind()
 {
 }
@@ -219,10 +229,7 @@ struct FDisplayNameHelper
 
 		if (auto Property = dynamic_cast<const UProperty*>(&Object))
 		{
-			if (auto OwnerStruct = Property->GetOwnerStruct())
-			{
-				return OwnerStruct->PropertyNameToDisplayName(Property->GetFName());
-			}
+			return Property->GetAuthoredName();
 		}
 
 		return Object.GetName();
@@ -1045,12 +1052,11 @@ void UStruct::SerializeTaggedProperties(FStructuredArchive::FSlot Slot, uint8* D
 
 				RemainingArrayDim = Property ? Property->ArrayDim : 0;
 			}
-#if WITH_EDITOR
+
 			if (!Property)
 			{
 				Property = CustomFindProperty(Tag.Name);
 			}
-#endif // WITH_EDITOR
 
 			FName PropID = Property ? Property->GetID() : NAME_None;
 			FName ArrayInnerID = NAME_None;
@@ -1482,6 +1488,21 @@ void UStruct::SetSuperStruct(UStruct* NewSuperStruct)
 #endif
 }
 
+FString UStruct::PropertyNameToDisplayName(FName InName) const
+{
+	const UField* FoundField = FindField<const UField>(this, InName);
+	return GetAuthoredNameForField(FoundField);
+}
+
+FString UStruct::GetAuthoredNameForField(const UField* Field) const
+{
+	if (Field)
+	{
+		return Field->GetName();
+	}
+	return FString();
+}
+
 #if WITH_EDITOR || HACK_HEADER_GENERATOR
 bool UStruct::GetBoolMetaDataHierarchical(const FName& Key) const
 {
@@ -2339,13 +2360,15 @@ void UScriptStruct::ExportText(FString& ValueStr, const void* Value, const void*
 						ValueStr += TEXT(",");
 					}
 
+					const FString PropertyName = (PortFlags & PPF_ExternalEditor) != 0 ? *It->GetAuthoredName() : It->GetName();
+
 					if (It->ArrayDim == 1)
 					{
-						ValueStr += FString::Printf(TEXT("%s="), *It->GetName());
+						ValueStr += FString::Printf(TEXT("%s="), *PropertyName);
 					}
 					else
 					{
-						ValueStr += FString::Printf(TEXT("%s[%i]="), *It->GetName(), Index);
+						ValueStr += FString::Printf(TEXT("%s[%i]="), *PropertyName, Index);
 					}
 					ValueStr += InnerValue;
 				}
diff --git a/Engine/Source/Runtime/CoreUObject/Private/UObject/Enum.cpp b/Engine/Source/Runtime/CoreUObject/Private/UObject/Enum.cpp
--- a/Engine/Source/Runtime/CoreUObject/Private/UObject/Enum.cpp
+++ b/Engine/Source/Runtime/CoreUObject/Private/UObject/Enum.cpp
@@ -414,6 +414,17 @@ FText UEnum::GetDisplayNameTextByValue(int64 Value) const
 	return GetDisplayNameTextByIndex(Index);
 }
 
+FString UEnum::GetAuthoredNameStringByIndex(int32 InIndex) const
+{
+	return GetNameStringByIndex(InIndex);
+}
+
+FString UEnum::GetAuthoredNameStringByValue(int64 Value) const
+{
+	int32 Index = GetIndexByValue(Value);
+	return GetAuthoredNameStringByIndex(Index);
+}
+
 int32 UEnum::GetIndexByNameString(const FString& InSearchString, EGetByNameFlags Flags) const
 {
 	ENameCase         NameComparisonMethod   = !!(Flags & EGetByNameFlags::CaseSensitive) ? ENameCase  ::CaseSensitive : ENameCase  ::IgnoreCase;
@@ -469,6 +480,9 @@ int32 UEnum::GetIndexByNameString(const FString& InSearchString, EGetByNameFlags
 	FName SearchName = FName(*SearchEnumEntryString);
 	FName ModifiedName = FName(*ModifiedEnumEntryString);
 
+	// Check authored name, but only if this is a subclass of Enum that might have implemented it
+	const bool bCheckAuthoredName = !!(Flags & EGetByNameFlags::CheckAuthoredName) && (GetClass() != UEnum::StaticClass());
+
 	const int32 Count = Names.Num();
 	for (int32 Counter = 0; Counter < Count; ++Counter)
 	{
@@ -476,6 +490,16 @@ int32 UEnum::GetIndexByNameString(const FString& InSearchString, EGetByNameFlags
 		{
 			return Counter;
 		}
+
+		if (bCheckAuthoredName)
+		{
+			FString AuthoredName = GetAuthoredNameStringByIndex(Counter);
+
+			if (AuthoredName.Equals(SearchEnumEntryString, StringComparisonMethod) || AuthoredName.Equals(ModifiedEnumEntryString, StringComparisonMethod))
+			{
+				return Counter;
+			}
+		}
 	}
 
 	if (!InSearchString.Equals(SearchEnumEntryString, StringComparisonMethod))
diff --git a/Engine/Source/Runtime/CoreUObject/Private/UObject/EnumProperty.cpp b/Engine/Source/Runtime/CoreUObject/Private/UObject/EnumProperty.cpp
--- a/Engine/Source/Runtime/CoreUObject/Private/UObject/EnumProperty.cpp
+++ b/Engine/Source/Runtime/CoreUObject/Private/UObject/EnumProperty.cpp
@@ -238,6 +238,10 @@ void UEnumProperty::ExportTextItem(FString& ValueStr, const void* PropertyValue,
 			{
 				ValueStr += Enum->GetDisplayNameTextByValue(Value).ToString();
 			}
+			else if (PortFlags & PPF_ExternalEditor)
+			{
+				ValueStr += Enum->GetAuthoredNameStringByValue(Value);
+			}
 			else
 			{
 				ValueStr += Enum->GetNameStringByValue(Value);
@@ -264,7 +268,7 @@ const TCHAR* UEnumProperty::ImportText_Internal(const TCHAR* InBuffer, void* Dat
 		FString Temp;
 		if (const TCHAR* Buffer = UPropertyHelpers::ReadToken(InBuffer, Temp, true))
 		{
-			int32 EnumIndex = Enum->GetIndexByName(*Temp);
+			int32 EnumIndex = Enum->GetIndexByName(*Temp, EGetByNameFlags::CheckAuthoredName);
 			if (EnumIndex == INDEX_NONE && (Temp.IsNumeric() && !Algo::Find(Temp, TEXT('.'))))
 			{
 				int64 EnumValue = INDEX_NONE;
diff --git a/Engine/Source/Runtime/CoreUObject/Private/UObject/Property.cpp b/Engine/Source/Runtime/CoreUObject/Private/UObject/Property.cpp
--- a/Engine/Source/Runtime/CoreUObject/Private/UObject/Property.cpp
+++ b/Engine/Source/Runtime/CoreUObject/Private/UObject/Property.cpp
@@ -1000,12 +1000,11 @@ const TCHAR* UProperty::ImportSingleProperty( const TCHAR* Str, void* DestData,
 			{
 				Property = FindField<UProperty>(ObjectStruct, NewPropertyName);
 			}
-#if WITH_EDITOR
+
 			if (!Property)
 			{
 				Property = ObjectStruct->CustomFindProperty(PropertyName);
 			}
-#endif	// WITH_EDITOR
 		}		
 
 		delete[] Token;
diff --git a/Engine/Source/Runtime/CoreUObject/Private/UObject/PropertyArray.cpp b/Engine/Source/Runtime/CoreUObject/Private/UObject/PropertyArray.cpp
--- a/Engine/Source/Runtime/CoreUObject/Private/UObject/PropertyArray.cpp
+++ b/Engine/Source/Runtime/CoreUObject/Private/UObject/PropertyArray.cpp
@@ -335,18 +335,21 @@ void UArrayProperty::ExportTextInnerItem(FString& ValueStr, UProperty* Inner, co
 
 	uint8* StructDefaults = NULL;
 	UStructProperty* StructProperty = dynamic_cast<UStructProperty*>(Inner);
+
+	const bool bReadableForm = (0 != (PPF_BlueprintDebugView & PortFlags));
+	const bool bExternalEditor = (0 != (PPF_ExternalEditor & PortFlags));
+
 	// ArrayProperties only export a diff because array entries are cleared and recreated upon import. Static arrays are overwritten when importing,
 	// so we export the entire struct to ensure all data is copied over correctly. Behavior is currently inconsistent when copy/pasting between the two types.
 	// In the future, static arrays could export diffs if the property being imported to is reset to default before the import.
-	if ( StructProperty != NULL && Inner->ArrayDim == 1 )
+	// When exporting to an external editor, we want to save defaults so all information is available for editing
+	if ( StructProperty != NULL && Inner->ArrayDim == 1 && !bExternalEditor )
 	{
 		checkSlow(StructProperty->Struct);
 		StructDefaults = (uint8*)FMemory::Malloc(StructProperty->Struct->GetStructureSize() * Inner->ArrayDim);
 		StructProperty->InitializeValue(StructDefaults);
 	}
 
-	const bool bReadableForm = (0 != (PPF_BlueprintDebugView & PortFlags));
-
 	int32 Count = 0;
 	for( int32 i=0; i<PropertySize; i++ )
 	{
@@ -375,7 +378,11 @@ void UArrayProperty::ExportTextInnerItem(FString& ValueStr, UProperty* Inner, co
 
 		// Always use struct defaults if the inner is a struct, for symmetry with the import of array inner struct defaults
 		uint8* PropDefault = nullptr;
-		if (StructProperty)
+		if (bExternalEditor)
+		{
+			PropDefault = PropData;
+		}
+		else if (StructProperty)
 		{
 			PropDefault = StructDefaults;
 		}
diff --git a/Engine/Source/Runtime/CoreUObject/Private/UObject/PropertyByte.cpp b/Engine/Source/Runtime/CoreUObject/Private/UObject/PropertyByte.cpp
--- a/Engine/Source/Runtime/CoreUObject/Private/UObject/PropertyByte.cpp
+++ b/Engine/Source/Runtime/CoreUObject/Private/UObject/PropertyByte.cpp
@@ -338,6 +338,10 @@ void UByteProperty::ExportTextItem( FString& ValueStr, const void* PropertyValue
 			{
 				ValueStr += Enum->GetDisplayNameTextByValue(*(const uint8*)PropertyValue).ToString();
 			}
+			else if (PortFlags & PPF_ExternalEditor)
+			{
+				ValueStr += Enum->GetAuthoredNameStringByValue(*(const uint8*)PropertyValue);
+			}
 			else
 			{
 				ValueStr += Enum->GetNameStringByValue(*(const uint8*)PropertyValue);
@@ -360,7 +364,7 @@ const TCHAR* UByteProperty::ImportText_Internal( const TCHAR* InBuffer, void* Da
 		FString Temp;
 		if (const TCHAR* Buffer = UPropertyHelpers::ReadToken(InBuffer, Temp, true))
 		{
-			int32 EnumIndex = Enum->GetIndexByName(*Temp);
+			int32 EnumIndex = Enum->GetIndexByName(*Temp, EGetByNameFlags::CheckAuthoredName);
 			if (EnumIndex == INDEX_NONE && (Temp.IsNumeric() && !Algo::Find(Temp, TEXT('.'))))
 			{
 				int64 EnumValue = INDEX_NONE;
diff --git a/Engine/Source/Runtime/CoreUObject/Private/UObject/PropertyMap.cpp b/Engine/Source/Runtime/CoreUObject/Private/UObject/PropertyMap.cpp
--- a/Engine/Source/Runtime/CoreUObject/Private/UObject/PropertyMap.cpp
+++ b/Engine/Source/Runtime/CoreUObject/Private/UObject/PropertyMap.cpp
@@ -552,13 +552,19 @@ void UMapProperty::ExportTextItem(FString& ValueStr, const void* PropertyValue,
 		return;
 	}
 
+	const bool bExternalEditor = (0 != (PPF_ExternalEditor & PortFlags));
+
 	uint8* StructDefaults = nullptr;
 	if (UStructProperty* StructValueProp = dynamic_cast<UStructProperty*>(ValueProp))
 	{
 		checkSlow(StructValueProp->Struct);
 
-		StructDefaults = (uint8*)FMemory::Malloc(MapLayout.SetLayout.Size);
-		ValueProp->InitializeValue(StructDefaults + MapLayout.ValueOffset);
+		if (!bExternalEditor)
+		{
+			// For external editor, we always export all fields
+			StructDefaults = (uint8*)FMemory::Malloc(MapLayout.SetLayout.Size);
+			ValueProp->InitializeValue(StructDefaults + MapLayout.ValueOffset);
+		}
 	}
 	ON_SCOPE_EXIT
 	{
@@ -596,6 +602,12 @@ void UMapProperty::ExportTextItem(FString& ValueStr, const void* PropertyValue,
 				// Always use struct defaults if the inner is a struct, for symmetry with the import of array inner struct defaults
 				uint8* PropDefault = StructDefaults ? StructDefaults : DefaultValue ? DefaultMapHelper.FindMapPairPtrWithKey(PropData) : nullptr;
 
+				if (bExternalEditor)
+				{
+					// For external editor, always write
+					PropDefault = PropData;
+				}
+
 				ValueProp->ExportTextItem(ValueStr, PropData + MapLayout.ValueOffset, PropDefault + MapLayout.ValueOffset, Parent, PortFlags | PPF_Delimited, ExportRootScope);
 
 				--Count;
@@ -629,6 +641,12 @@ void UMapProperty::ExportTextItem(FString& ValueStr, const void* PropertyValue,
 				// Always use struct defaults if the inner is a struct, for symmetry with the import of array inner struct defaults
 				uint8* PropDefault = StructDefaults ? StructDefaults : DefaultValue ? DefaultMapHelper.FindMapPairPtrWithKey(PropData) : nullptr;
 
+				if (bExternalEditor)
+				{
+					// For external editor, always write
+					PropDefault = PropData;
+				}
+
 				ValueProp->ExportTextItem(ValueStr, PropData + MapLayout.ValueOffset, PropDefault + MapLayout.ValueOffset, Parent, PortFlags | PPF_Delimited, ExportRootScope);
 
 				ValueStr += TEXT(")");
diff --git a/Engine/Source/Runtime/CoreUObject/Private/UObject/PropertySet.cpp b/Engine/Source/Runtime/CoreUObject/Private/UObject/PropertySet.cpp
--- a/Engine/Source/Runtime/CoreUObject/Private/UObject/PropertySet.cpp
+++ b/Engine/Source/Runtime/CoreUObject/Private/UObject/PropertySet.cpp
@@ -483,13 +483,19 @@ void USetProperty::ExportTextItem(FString& ValueStr, const void* PropertyValue,
 		return;
 	}
 
+	const bool bExternalEditor = (0 != (PPF_ExternalEditor & PortFlags));
+
 	uint8* StructDefaults = nullptr;
 	if (UStructProperty* StructElementProp = Cast<UStructProperty>(ElementProp))
 	{
 		checkSlow(StructElementProp->Struct);
 
-		StructDefaults = (uint8*)FMemory::Malloc(SetLayout.Size);
-		ElementProp->InitializeValue(StructDefaults);
+		if (!bExternalEditor)
+		{
+			// For external editor, we always export all fields
+			StructDefaults = (uint8*)FMemory::Malloc(SetLayout.Size);
+			ElementProp->InitializeValue(StructDefaults);
+		}
 	}
 
 	ON_SCOPE_EXIT
@@ -524,6 +530,12 @@ void USetProperty::ExportTextItem(FString& ValueStr, const void* PropertyValue,
 				// Always use struct defaults if the element is a struct, for symmetry with the import of array inner struct defaults
 				uint8* PropDefault = StructDefaults ? StructDefaults : DefaultValue ? DefaultSetHelper.FindElementPtr(PropData) : nullptr;
 
+				if (bExternalEditor)
+				{
+					// For external editor, always write
+					PropDefault = PropData;
+				}
+
 				ElementProp->ExportTextItem(ValueStr, PropData, PropDefault, Parent, PortFlags | PPF_Delimited, ExportRootScope);
 
 				--Count;
@@ -548,8 +560,15 @@ void USetProperty::ExportTextItem(FString& ValueStr, const void* PropertyValue,
 					ValueStr += TCHAR(',');
 				}
 
-				ElementProp->ExportTextItem(ValueStr, PropData, nullptr, Parent, PortFlags | PPF_Delimited, ExportRootScope);
+				uint8* PropDefault = nullptr;
 
+				if (bExternalEditor)
+				{
+					// For external editor, always write
+					PropDefault = PropData;
+				}
+
+				ElementProp->ExportTextItem(ValueStr, PropData, PropDefault, Parent, PortFlags | PPF_Delimited, ExportRootScope);
 
 				--Count;
 			}
diff --git a/Engine/Source/Runtime/CoreUObject/Public/UObject/Class.h b/Engine/Source/Runtime/CoreUObject/Public/UObject/Class.h
--- a/Engine/Source/Runtime/CoreUObject/Public/UObject/Class.h
+++ b/Engine/Source/Runtime/CoreUObject/Public/UObject/Class.h
@@ -82,6 +82,13 @@ class COREUOBJECT_API UField : public UObject
 	/** Goes up the outer chain to look for a UStruct */
 	UStruct* GetOwnerStruct() const;
 
+	/** 
+	 * Returns a human readable string that was assigned to this field at creation. 
+	 * By default this is the same as GetName() but it can be overridden if that is an internal-only name.
+	 * This name is consistent in editor/cooked builds, is not localized, and is useful for data import/export.
+	 */
+	FString GetAuthoredName() const;
+
 #if WITH_EDITOR || HACK_HEADER_GENERATOR
 	/**
 	 * Finds the localized display name or native display name as a fallback.
@@ -360,10 +367,8 @@ public:
 	virtual void DestroyStruct(void* Dest, int32 ArrayDim = 1) const;
 
 public:
-#if WITH_EDITOR
-	/** Used to search for properties in user defined structs */
+	/** Look up a property by an alternate name if it was not found in the first search, this is overridden for user structs */
 	virtual UProperty* CustomFindProperty(const FName InName) const { return nullptr; };
-#endif // WITH_EDITOR
 
 	/** Serialize an expression to an archive. Returns expression token */
 	virtual EExprToken SerializeExpr(int32& iCode, FArchive& Ar);
@@ -428,11 +433,11 @@ public:
 	 */
 	virtual void SetSuperStruct(UStruct* NewSuperStruct);
 
-	/** Returns the a human readable string for a property, overridden for user defined structs */
-	virtual FString PropertyNameToDisplayName(FName InName) const 
-	{ 
-		return InName.ToString();
-	}
+	UE_DEPRECATED(4.23, "Replace with GetAuthoredNameForField or UField::GetAuthoredName")
+	virtual FString PropertyNameToDisplayName(FName InName) const;
+
+	/** Returns a human readable string for a given field, overridden for user defined structs */
+	virtual FString GetAuthoredNameForField(const UField* Field) const;
 
 #if WITH_EDITOR || HACK_HEADER_GENERATOR
 	/** Try and find boolean metadata with the given key. If not found on this class, work up hierarchy looking for it. */
@@ -1578,13 +1583,19 @@ public:
 
 typedef FText(*FEnumDisplayNameFn)(int32);
 
-// Optional flags for the UEnum::Get*ByName() functions.
+/** Optional flags for the UEnum::Get*ByName() functions. */
 enum class EGetByNameFlags
 {
 	None = 0,
 
-	ErrorIfNotFound = 0x01,
-	CaseSensitive   = 0x02
+	/** Outputs an warning if the enum lookup fails */
+	ErrorIfNotFound		= 0x01,
+
+	/** Does a case sensitive match */
+	CaseSensitive		= 0x02,
+
+	/** Checks the GetAuthoredNameStringByIndex value as well as normal names */
+	CheckAuthoredName	= 0x04,
 };
 
 ENUM_CLASS_FLAGS(EGetByNameFlags)
@@ -1671,6 +1682,20 @@ public:
 	/** Version of GetDisplayNameTextByIndex that takes a value instead */
 	FText GetDisplayNameTextByValue(int64 InValue) const;
 
+	/**
+	 * Returns the unlocalized logical name originally assigned to the enum at creation.
+	 * By default this is the same as the short name but it is overridden in child classes with different internal name storage.
+	 * This name is consistent in cooked and editor builds and is useful for things like external data import/export.
+	 *
+	 * @param InIndex Index of the enum value to get Display Name for
+	 *
+	 * @return The author-specified name, or an empty string if Index is invalid
+	 */
+	virtual FString GetAuthoredNameStringByIndex(int32 InIndex) const;
+
+	/** Version of GetAuthoredNameByIndex that takes a value instead */
+	FString GetAuthoredNameStringByValue(int64 InValue) const;
+
 	/** Gets max value of Enum. Defaults to zero if there are no entries. */
 	int64 GetMaxEnumValue() const;
 
diff --git a/Engine/Source/Runtime/Engine/Classes/Engine/UserDefinedEnum.h b/Engine/Source/Runtime/Engine/Classes/Engine/UserDefinedEnum.h
--- a/Engine/Source/Runtime/Engine/Classes/Engine/UserDefinedEnum.h
+++ b/Engine/Source/Runtime/Engine/Classes/Engine/UserDefinedEnum.h
@@ -62,6 +62,7 @@ public:
 
 	/** Overridden to read DisplayNameMap*/
 	virtual FText GetDisplayNameTextByIndex(int32 InIndex) const override;
+	virtual FString GetAuthoredNameStringByIndex(int32 InIndex) const override;
 
 	virtual bool SetEnums(TArray<TPair<FName, int64>>& InNames, ECppForm InCppForm, bool bAddMaxKeyIfMissing = true) override;
 
diff --git a/Engine/Source/Runtime/Engine/Classes/Engine/UserDefinedStruct.h b/Engine/Source/Runtime/Engine/Classes/Engine/UserDefinedStruct.h
--- a/Engine/Source/Runtime/Engine/Classes/Engine/UserDefinedStruct.h
+++ b/Engine/Source/Runtime/Engine/Classes/Engine/UserDefinedStruct.h
@@ -88,10 +88,6 @@ public:
 	virtual void PostLoad() override;
 	// End of UObject interface.
 
-	// UScriptStruct interface.
-	virtual UProperty* CustomFindProperty(const FName Name) const override;
-	// End of UScriptStruct interface.
-
 	/** Creates a new guid if needed */
 	void ValidateGuid();
 
@@ -101,7 +97,7 @@ public:
 	// UObject interface.
 	virtual void Serialize(FStructuredArchive::FRecord Record) override;
 	virtual void SerializeTaggedProperties(FStructuredArchive::FSlot Slot, uint8* Data, UStruct* DefaultsStruct, uint8* Defaults, const UObject* BreakRecursionIfFullyLoad = nullptr) const override;
-	virtual FString PropertyNameToDisplayName(FName InName) const override;
+	virtual FString GetAuthoredNameForField(const UField* Field) const override;
 	// End of UObject interface.
 
 	// UScriptStruct interface.
@@ -110,6 +106,7 @@ public:
 	virtual void RecursivelyPreload() override;
 	virtual FGuid GetCustomGuid() const override;
 	virtual FString GetStructCPPName() const override;
+	virtual UProperty* CustomFindProperty(const FName Name) const override;
 	// End of  UScriptStruct interface.
 
 	/** Returns the raw memory of the default instance */
diff --git a/Engine/Source/Runtime/Engine/Private/DataTable.cpp b/Engine/Source/Runtime/Engine/Private/DataTable.cpp
--- a/Engine/Source/Runtime/Engine/Private/DataTable.cpp
+++ b/Engine/Source/Runtime/Engine/Private/DataTable.cpp
@@ -398,7 +398,7 @@ UProperty* UDataTable::FindTableProperty(const FName& PropertyName) const
 
 			for (TFieldIterator<UProperty> It(RowStruct); It; ++It)
 			{
-				if (PropertyNameStr == RowStruct->PropertyNameToDisplayName(It->GetFName()))
+				if (PropertyNameStr == RowStruct->GetAuthoredNameForField(*It))
 				{
 					Property = *It;
 					break;
@@ -671,7 +671,7 @@ TArray<UProperty*> UDataTable::GetTablePropertyArray(const TArray<const TCHAR*>&
 			}
 #endif // WITH_EDITOR
 
-			const FString DisplayName = DataTableUtils::GetPropertyDisplayName(ColumnProp, ExpectedPropNames[PropIdx].ToString());
+			const FString DisplayName = DataTableUtils::GetPropertyExportName(ColumnProp);
 			OutProblems.Add(FString::Printf(TEXT("Expected column '%s' not found in input."), *DisplayName));
 		}
 	}
@@ -754,7 +754,7 @@ TArray<FString> UDataTable::GetColumnTitles() const
 		{
 			UProperty* Prop = *It;
 			check(Prop != nullptr);
-			const FString DisplayName = DataTableUtils::GetPropertyDisplayName(Prop, Prop->GetName());
+			const FString DisplayName = DataTableUtils::GetPropertyExportName(Prop);
 			Result.Add(DisplayName);
 		}
 	}
diff --git a/Engine/Source/Runtime/Engine/Private/DataTableCSV.cpp b/Engine/Source/Runtime/Engine/Private/DataTableCSV.cpp
--- a/Engine/Source/Runtime/Engine/Private/DataTableCSV.cpp
+++ b/Engine/Source/Runtime/Engine/Private/DataTableCSV.cpp
@@ -227,7 +227,7 @@ bool FDataTableImporterCSV::ReadTable()
 			if(Error.Len() > 0)
 			{
 				FString ColumnName = (ColumnProp != nullptr) 
-					? DataTableUtils::GetPropertyDisplayName(ColumnProp, ColumnProp->GetName())
+					? DataTableUtils::GetPropertyExportName(ColumnProp)
 					: FString(TEXT("NONE"));
 				ImportProblems.Add(FString::Printf(TEXT("Problem assigning string '%s' to property '%s' on row '%s' : %s"), *CellValue, *ColumnName, *RowName.ToString(), *Error));
 			}
diff --git a/Engine/Source/Runtime/Engine/Private/DataTableJSON.cpp b/Engine/Source/Runtime/Engine/Private/DataTableJSON.cpp
--- a/Engine/Source/Runtime/Engine/Private/DataTableJSON.cpp
+++ b/Engine/Source/Runtime/Engine/Private/DataTableJSON.cpp
@@ -479,7 +479,7 @@ bool FDataTableImporterJSON::ReadStruct(const TSharedRef<FJsonObject>& InParsedO
 		UProperty* BaseProp = *It;
 		check(BaseProp);
 
-		const FString ColumnName = DataTableUtils::GetPropertyDisplayName(BaseProp, BaseProp->GetName());
+		const FString ColumnName = DataTableUtils::GetPropertyExportName(BaseProp);
 
 		TSharedPtr<FJsonValue> ParsedPropertyValue;
 		for (const FString& PropertyName : DataTableUtils::GetPropertyImportNames(BaseProp))
diff --git a/Engine/Source/Runtime/Engine/Private/DataTableUtils.cpp b/Engine/Source/Runtime/Engine/Private/DataTableUtils.cpp
--- a/Engine/Source/Runtime/Engine/Private/DataTableUtils.cpp
+++ b/Engine/Source/Runtime/Engine/Private/DataTableUtils.cpp
@@ -90,27 +90,6 @@ void AssignStringToProperty(const FString& InString, const UProperty* InProp, ui
 
 void GetPropertyValueAsStringDirect(const UProperty* InProp, const uint8* InData, const int32 InPortFlags, const EDataTableExportFlags InDTExportFlags, FString& OutString)
 {
-	if (!!(InDTExportFlags & EDataTableExportFlags::UsePrettyEnumNames))
-	{
-		UEnum* Enum = nullptr;
-		int64 Val = 0;
-		if (const UEnumProperty* EnumProp = Cast<UEnumProperty>(InProp))
-		{
-			Enum = EnumProp->GetEnum();
-			Val = EnumProp->GetUnderlyingProperty()->GetSignedIntPropertyValue(InData);
-		}
-		else if (const UByteProperty* ByteProp = Cast<UByteProperty>(InProp))
-		{
-			Enum = ByteProp->GetIntPropertyEnum();
-			Val = *InData;
-		}
-
-		if (UUserDefinedEnum* UDEnum = Cast<UUserDefinedEnum>(Enum))
-		{
-			OutString.Append(GetSourceString(UDEnum->GetDisplayNameTextByValue(Val)));
-			return;
-		}
-	}
 #if WITH_EDITOR
 	if (InPortFlags & PPF_PropertyWindow)
 	{
@@ -249,7 +228,7 @@ FString DataTableUtils::AssignStringToPropertyDirect(const FString& InString, co
 	FStringOutputDevice ImportError;
 	if(InProp && IsSupportedTableProperty(InProp))
 	{
-		DataTableUtilsImpl::AssignStringToPropertyDirect(InString, InProp, InData, PPF_None, ImportError);
+		DataTableUtilsImpl::AssignStringToPropertyDirect(InString, InProp, InData, PPF_ExternalEditor, ImportError);
 	}
 
 	FString Error = ImportError;
@@ -263,7 +242,7 @@ FString DataTableUtils::AssignStringToProperty(const FString& InString, const UP
 	{
 		if(InProp->ArrayDim == 1)
 		{
-			DataTableUtilsImpl::AssignStringToProperty(InString, InProp, InData, 0, PPF_None, ImportError);
+			DataTableUtilsImpl::AssignStringToProperty(InString, InProp, InData, 0, PPF_ExternalEditor, ImportError);
 		}
 		else
 		{
@@ -305,7 +284,7 @@ FString DataTableUtils::GetPropertyValueAsStringDirect(const UProperty* InProp,
 
 	if(InProp && IsSupportedTableProperty(InProp))
 	{
-		DataTableUtilsImpl::GetPropertyValueAsStringDirect(InProp, InData, PPF_None, InDTExportFlags, Result);
+		DataTableUtilsImpl::GetPropertyValueAsStringDirect(InProp, InData, PPF_ExternalEditor, InDTExportFlags, Result);
 	}
 
 	return Result;
@@ -319,7 +298,7 @@ FString DataTableUtils::GetPropertyValueAsString(const UProperty* InProp, const
 	{
 		if(InProp->ArrayDim == 1)
 		{
-			DataTableUtilsImpl::GetPropertyValueAsString(InProp, InData, 0, PPF_None, InDTExportFlags, Result);
+			DataTableUtilsImpl::GetPropertyValueAsString(InProp, InData, 0, PPF_ExternalEditor, InDTExportFlags, Result);
 		}
 		else
 		{
@@ -327,7 +306,7 @@ FString DataTableUtils::GetPropertyValueAsString(const UProperty* InProp, const
 
 			for(int32 Index = 0; Index < InProp->ArrayDim; ++Index)
 			{
-				DataTableUtilsImpl::GetPropertyValueAsString(InProp, InData, Index, PPF_Delimited, InDTExportFlags, Result);
+				DataTableUtilsImpl::GetPropertyValueAsString(InProp, InData, Index, PPF_Delimited | PPF_ExternalEditor, InDTExportFlags, Result);
 
 				if((Index + 1) < InProp->ArrayDim)
 				{
@@ -349,7 +328,7 @@ FText DataTableUtils::GetPropertyValueAsTextDirect(const UProperty* InProp, cons
 	if(InProp && IsSupportedTableProperty(InProp))
 	{
 		FString ExportedString;
-		DataTableUtilsImpl::GetPropertyValueAsStringDirect(InProp, InData, PPF_PropertyWindow, EDataTableExportFlags::UsePrettyPropertyNames | EDataTableExportFlags::UsePrettyEnumNames | EDataTableExportFlags::UseJsonObjectsForStructs, ExportedString);
+		DataTableUtilsImpl::GetPropertyValueAsStringDirect(InProp, InData, PPF_PropertyWindow, EDataTableExportFlags::UseJsonObjectsForStructs, ExportedString);
 
 		Result = FText::FromString(MoveTemp(ExportedString));
 	}
@@ -367,7 +346,7 @@ FText DataTableUtils::GetPropertyValueAsText(const UProperty* InProp, const uint
 
 		if(InProp->ArrayDim == 1)
 		{
-			DataTableUtilsImpl::GetPropertyValueAsString(InProp, InData, 0, PPF_PropertyWindow, EDataTableExportFlags::UsePrettyPropertyNames | EDataTableExportFlags::UsePrettyEnumNames | EDataTableExportFlags::UseJsonObjectsForStructs, ExportedString);
+			DataTableUtilsImpl::GetPropertyValueAsString(InProp, InData, 0, PPF_PropertyWindow, EDataTableExportFlags::UseJsonObjectsForStructs, ExportedString);
 		}
 		else
 		{
@@ -375,7 +354,7 @@ FText DataTableUtils::GetPropertyValueAsText(const UProperty* InProp, const uint
 
 			for(int32 Index = 0; Index < InProp->ArrayDim; ++Index)
 			{
-				DataTableUtilsImpl::GetPropertyValueAsString(InProp, InData, Index, PPF_PropertyWindow | PPF_Delimited, EDataTableExportFlags::UsePrettyPropertyNames | EDataTableExportFlags::UsePrettyEnumNames | EDataTableExportFlags::UseJsonObjectsForStructs, ExportedString);
+				DataTableUtilsImpl::GetPropertyValueAsString(InProp, InData, Index, PPF_PropertyWindow | PPF_Delimited, EDataTableExportFlags::UseJsonObjectsForStructs, ExportedString);
 
 				if((Index + 1) < InProp->ArrayDim)
 				{
@@ -453,11 +432,7 @@ FString DataTableUtils::GetPropertyExportName(const UProperty* Prop, const EData
 	{
 		return FString();
 	}
-	if (Prop->GetOwnerStruct()->IsA(UUserDefinedStruct::StaticClass()) && !!(InDTExportFlags & EDataTableExportFlags::UsePrettyPropertyNames))
-	{
-		return GetPropertyDisplayName(Prop, Prop->GetName());
-	}
-	return Prop->GetName();
+	return Prop->GetAuthoredName();
 }
 
 TArray<FString> DataTableUtils::GetPropertyImportNames(const UProperty* Prop)
@@ -467,18 +442,19 @@ TArray<FString> DataTableUtils::GetPropertyImportNames(const UProperty* Prop)
 	{
 		Result.AddUnique(Prop->GetName());
 	}
-	Result.AddUnique(GetPropertyExportName(Prop, EDataTableExportFlags::UsePrettyPropertyNames));
+	Result.AddUnique(GetPropertyExportName(Prop));
 	return Result;
 }
 
-FString DataTableUtils::GetPropertyDisplayName(const UProperty* Prop, const FString& DefaultName)
+FText DataTableUtils::GetPropertyDisplayName(const UProperty* Prop, const FString& DefaultName)
 {
 #if WITH_EDITOR
-	static const FName DisplayNameKey(TEXT("DisplayName"));
-	return (Prop && Prop->HasMetaData(DisplayNameKey)) ? Prop->GetMetaData(DisplayNameKey) : DefaultName;
-#else  // WITH_EDITOR
-	return DefaultName;
-#endif // WITH_EDITOR
+	if (Prop)
+	{
+		return Prop->GetDisplayNameText();
+	}
+#endif  // WITH_EDITOR
+	return FText::FromString(DefaultName);
 }
 
 TArray<FString> DataTableUtils::GetColumnDataAsString(const UDataTable* InTable, const FName& PropertyName, const EDataTableExportFlags InDTExportFlags)
diff --git a/Engine/Source/Runtime/Engine/Private/UserDefinedEnum.cpp b/Engine/Source/Runtime/Engine/Private/UserDefinedEnum.cpp
--- a/Engine/Source/Runtime/Engine/Private/UserDefinedEnum.cpp
+++ b/Engine/Source/Runtime/Engine/Private/UserDefinedEnum.cpp
@@ -160,6 +160,21 @@ FText UUserDefinedEnum::GetDisplayNameTextByIndex(int32 InIndex) const
 	return Super::GetDisplayNameTextByIndex(InIndex);
 }
 
+FString UUserDefinedEnum::GetAuthoredNameStringByIndex(int32 InIndex) const
+{
+	const FName EnumEntryName = *GetNameStringByIndex(InIndex);
+
+	if (const FText* EnumEntryDisplayName = DisplayNameMap.Find(EnumEntryName))
+	{
+		if (const FString* SourceString = FTextInspector::GetSourceString(*EnumEntryDisplayName))
+		{
+			return *SourceString;
+		}
+	}
+
+	return Super::GetAuthoredNameStringByIndex(InIndex);
+}
+
 bool UUserDefinedEnum::SetEnums(TArray<TPair<FName, int64>>& InNames, ECppForm InCppForm, bool bAddMaxKeyIfMissing)
 {
 	ensure(bAddMaxKeyIfMissing);
diff --git a/Engine/Source/Runtime/Engine/Private/UserDefinedStruct.cpp b/Engine/Source/Runtime/Engine/Private/UserDefinedStruct.cpp
--- a/Engine/Source/Runtime/Engine/Private/UserDefinedStruct.cpp
+++ b/Engine/Source/Runtime/Engine/Private/UserDefinedStruct.cpp
@@ -212,14 +212,6 @@ void UUserDefinedStruct::GetAssetRegistryTags(TArray<FAssetRegistryTag>& OutTags
 	OutTags.Add(FAssetRegistryTag(TEXT("Tooltip"), FStructureEditorUtils::GetTooltip(this), FAssetRegistryTag::TT_Hidden));
 }
 
-UProperty* UUserDefinedStruct::CustomFindProperty(const FName Name) const
-{
-	const FGuid PropertyGuid = FStructureEditorUtils::GetGuidFromPropertyName(Name);
-	UProperty* Property = PropertyGuid.IsValid() ? FStructureEditorUtils::GetPropertyByGuid(this, PropertyGuid) : FStructureEditorUtils::GetPropertyByDisplayName(this, Name.ToString());
-	ensure(!Property || !PropertyGuid.IsValid() || PropertyGuid == FStructureEditorUtils::GetGuidForProperty(Property));
-	return Property;
-}
-
 void UUserDefinedStruct::ValidateGuid()
 {
 	// Backward compatibility:
@@ -238,22 +230,56 @@ void UUserDefinedStruct::ValidateGuid()
 
 #endif	// WITH_EDITOR
 
-FString UUserDefinedStruct::PropertyNameToDisplayName(FName Name) const
+UProperty* UUserDefinedStruct::CustomFindProperty(const FName Name) const
 {
 #if WITH_EDITOR
-	FGuid PropertyGuid = FStructureEditorUtils::GetGuidFromPropertyName(Name);
-	return FStructureEditorUtils::GetVariableDisplayName(this, PropertyGuid);
-#endif	// WITH_EDITOR
+	// If we have the editor data, check that first as it's more up to date
+	const FGuid PropertyGuid = FStructureEditorUtils::GetGuidFromPropertyName(Name);
+	UProperty* EditorProperty = PropertyGuid.IsValid() ? FStructureEditorUtils::GetPropertyByGuid(this, PropertyGuid) : FStructureEditorUtils::GetPropertyByFriendlyName(this, Name.ToString());
+	ensure(!EditorProperty || !PropertyGuid.IsValid() || PropertyGuid == FStructureEditorUtils::GetGuidForProperty(EditorProperty));
+	if (EditorProperty)
+	{
+		return EditorProperty;
+	}
+#endif // WITH_EDITOR
+
+	// Check the authored names for each field
+	FString NameString = Name.ToString();
+	for (UProperty* CurrentProp : TFieldRange<UProperty>(this))
+	{
+		if (GetAuthoredNameForField(CurrentProp) == NameString)
+		{
+			return CurrentProp;
+		}
+	}
+	return nullptr;
+}
+
+FString UUserDefinedStruct::GetAuthoredNameForField(const UField* Field) const
+{
+	const UProperty* Property = Cast<UProperty>(Field);
+	if (!Property)
+	{
+		return Super::GetAuthoredNameForField(Field);
+	}
+
+#if WITH_EDITOR
+	const FString EditorName = FStructureEditorUtils::GetVariableFriendlyNameForProperty(this, Property);
+	if (!EditorName.IsEmpty())
+	{
+		return EditorName;
+	}
+#endif // WITH_EDITOR
 
 	const int32 GuidStrLen = 32;
 	const int32 MinimalPostfixlen = GuidStrLen + 3;
-	const FString OriginalName = Name.ToString();
+	const FString OriginalName = Property->GetName();
 	if (OriginalName.Len() > MinimalPostfixlen)
 	{
 		FString DisplayName = OriginalName.LeftChop(GuidStrLen + 1);
 		int FirstCharToRemove = -1;
 		const bool bCharFound = DisplayName.FindLastChar(TCHAR('_'), FirstCharToRemove);
-		if(bCharFound && (FirstCharToRemove > 0))
+		if (bCharFound && (FirstCharToRemove > 0))
 		{
 			return DisplayName.Mid(0, FirstCharToRemove);
 		}
diff --git a/Engine/Source/Runtime/Engine/Public/DataTableUtils.h b/Engine/Source/Runtime/Engine/Public/DataTableUtils.h
--- a/Engine/Source/Runtime/Engine/Public/DataTableUtils.h
+++ b/Engine/Source/Runtime/Engine/Public/DataTableUtils.h
@@ -14,14 +14,13 @@ enum class EDataTableExportFlags : uint8
 	/** No specific options. */
 	None = 0,
 
-	/** Export properties using their display name, rather than their internal name. */
-	UsePrettyPropertyNames = 1<<0,
+	/** Export nested structs as JSON objects (JSON exporter only), rather than as exported text. */
+	UseJsonObjectsForStructs = 1 << 0,
 
-	/** Export User Defined Enums using their display name, rather than their internal name. */
-	UsePrettyEnumNames = 1<<1,
+	// DEPRECATED. Native properties/enums are always exported using their internal name, user struct/enums are always exported using the friendly names set in the editor
 
-	/** Export nested structs as JSON objects (JSON exporter only), rather than as exported text. */
-	UseJsonObjectsForStructs = 1<<2,
+	UsePrettyPropertyNames UE_DEPRECATED(4.23, "UsePrettyPropertyNames is deprecated, we now always use the unlocalized but readable authored names") = 1 << 6,
+	UsePrettyEnumNames UE_DEPRECATED(4.23, "UsePrettyEnumNames is deprecated, we now always use the unlocalized but readable authored names") = 1 << 7,
 };
 ENUM_CLASS_FLAGS(EDataTableExportFlags);
 
@@ -81,7 +80,7 @@ namespace DataTableUtils
 	/**
 	 * Util to get the friendly display unlocalized name of a given property for export to files.
 	 */
-	ENGINE_API FString GetPropertyExportName(const UProperty* Prop, const EDataTableExportFlags InDTExportFlags);
+	ENGINE_API FString GetPropertyExportName(const UProperty* Prop, const EDataTableExportFlags InDTExportFlags = EDataTableExportFlags::None);
 
 	/**
 	 * Util to get the all variants for export names for backwards compatibility.
@@ -89,9 +88,11 @@ namespace DataTableUtils
 	ENGINE_API TArray<FString> GetPropertyImportNames(const UProperty* Prop);
 
 	/**
-	 * Util to get the friendly display name of a given property.
+	 * Util to get the localized display name of a given property.
 	 */
-	ENGINE_API FString GetPropertyDisplayName(const UProperty* Prop, const FString& DefaultName);
+	ENGINE_API FText GetPropertyDisplayName(const UProperty* Prop, const FString& DefaultName);
+
+	/**
 
 	/**
 	 * Output each row for a specific column/property in the table (doesn't include the title)
