Another CIS fix

diff --git a/Engine/Source/Runtime/Engine/Public/DataTableUtils.h b/Engine/Source/Runtime/Engine/Public/DataTableUtils.h
--- a/Engine/Source/Runtime/Engine/Public/DataTableUtils.h
+++ b/Engine/Source/Runtime/Engine/Public/DataTableUtils.h
@@ -92,8 +92,6 @@ namespace DataTableUtils
 	 */
 	ENGINE_API FText GetPropertyDisplayName(const UProperty* Prop, const FString& DefaultName);
 
-	/**
-
 	/**
 	 * Output each row for a specific column/property in the table (doesn't include the title)
 	 */
