UGameplayStatics::AsyncSaveGameToSlot is the native version, UAsyncActionHandleSaveGame provides a blueprint accessible wrapper
Clean up the save game function headers in UGameplayStatics

diff --git a/Engine/Source/Runtime/Engine/Classes/GameFramework/AsyncActionHandleSaveGame.h b/Engine/Source/Runtime/Engine/Classes/GameFramework/AsyncActionHandleSaveGame.h
new file mode 100644
--- /dev/null
+++ b/Engine/Source/Runtime/Engine/Classes/GameFramework/AsyncActionHandleSaveGame.h
@@ -0,0 +1,76 @@
+// Copyright 1998-2019 Epic Games, Inc. All Rights Reserved.
+
+#pragma once
+
+#include "CoreMinimal.h"
+#include "UObject/ObjectMacros.h"
+#include "Engine/StreamableManager.h"
+#include "Kismet/BlueprintAsyncActionBase.h"
+#include "Templates/SubclassOf.h"
+
+#include "AsyncActionHandleSaveGame.generated.h"
+
+DECLARE_DYNAMIC_MULTICAST_DELEGATE_TwoParams(FOnAsyncHandleSaveGame, USaveGame*, SaveGame, bool, bSuccess);
+
+/** Async action to handle async load/save of a USaveGame. This can be subclassed by a specific game */
+UCLASS()
+class ENGINE_API UAsyncActionHandleSaveGame : public UBlueprintAsyncActionBase
+{
+	GENERATED_BODY()
+public:
+
+	/**
+	 * Schedule an async save to a specific slot. UGameplayStatics::AsyncSaveGameToSlot is the native version of this.
+	 * When the save has succeeded or failed, the completed pin is activated with success/failure and the save game object.
+	 * Keep in mind that some platforms may not support trying to load and save at the same time.
+	 *
+	 * @param SaveGameObject	Object that contains data about the save game that we want to write out.
+	 * @param SlotName			Name of the save game slot to load from.
+	 * @param UserIndex			For some platforms, master user index to identify the user doing the loading.
+	 */
+	UFUNCTION(BlueprintCallable, meta=(BlueprintInternalUseOnly="true", Category = "SaveGame",  WorldContext = "WorldContextObject"))
+	static UAsyncActionHandleSaveGame* AsyncSaveGameToSlot(UObject* WorldContextObject, USaveGame* SaveGameObject, const FString& SlotName, const int32 UserIndex);
+
+	/**
+	 * Schedule an async load of a specific slot. UGameplayStatics::AsyncLoadGameFromSlot is the native version of this.
+	 * When the load has succeeded or failed, the completed pin is activated with success/failure and the newly loaded save game object if valid.
+	 * Keep in mind that some platforms may not support trying to load and save at the same time.
+	 *
+	 * @param SlotName			Name of the save game slot to load from.
+	 * @param UserIndex			For some platforms, master user index to identify the user doing the loading.
+	 */
+	UFUNCTION(BlueprintCallable, meta = (BlueprintInternalUseOnly = "true", Category = "SaveGame", WorldContext = "WorldContextObject"))
+	static UAsyncActionHandleSaveGame* AsyncLoadGameFromSlot(UObject* WorldContextObject, const FString& SlotName, const int32 UserIndex);
+
+	/** Delegate called when the save/load completes */
+	UPROPERTY(BlueprintAssignable)
+	FOnAsyncHandleSaveGame Completed;
+
+	/** Execute the actual operation */
+	virtual void Activate() override;
+
+protected:
+	enum class ESaveGameOperation : uint8
+	{
+		Save,
+		Load,
+	};
+
+	/** Which operation is being run */
+	ESaveGameOperation Operation;
+	
+	/** Slot/user to use */
+	FString SlotName;
+	int32 UserIndex;
+
+	/** The object that was either saved or loaded */
+	UPROPERTY()
+	USaveGame* SaveGameObject;
+	
+	/** Function callbacks for load/save */
+	virtual void HandleAsyncSave(const FString& SlotName, const int32 UserIndex, bool bSuccess);
+	virtual void HandleAsyncLoad(const FString& SlotName, const int32 UserIndex, USaveGame* LoadedSave);
+	
+	/** Called at completion of save/load to execute delegate */
+	virtual void ExecuteCompleted(bool bSuccess);
+};
diff --git a/Engine/Source/Runtime/Engine/Classes/Kismet/GameplayStatics.h b/Engine/Source/Runtime/Engine/Classes/Kismet/GameplayStatics.h
--- a/Engine/Source/Runtime/Engine/Classes/Kismet/GameplayStatics.h
+++ b/Engine/Source/Runtime/Engine/Classes/Kismet/GameplayStatics.h
@@ -36,6 +36,13 @@ class FMemoryReader;
 struct FDialogueContext;
 
 
+/** Delegate called from AsyncLoadGameFromSlot. First two parameters are passed in SlotName and UserIndex, third parameter is the returned SaveGame, or null if it failed to load */
+DECLARE_DELEGATE_ThreeParams(FAsyncSaveGameToSlotDelegate, const FString&, const int32, bool);
+
+/** Delegate called from AsyncLoadGameFromSlot. First two parameters are passed in SlotName and UserIndex, third parameter is the returned SaveGame, or null if it failed to load */
+DECLARE_DELEGATE_ThreeParams(FAsyncLoadGameFromSlotDelegate, const FString&, const int32, USaveGame*);
+
+/** Static class with useful gameplay utility functions that can be called from both Blueprint and C++ */
 UCLASS()
 class ENGINE_API UGameplayStatics : public UBlueprintFunctionLibrary
 {
@@ -810,87 +817,114 @@ public:
 	// --- Save Game functions ------------------------------
 
 	/** 
-	 *	Create a new, empty SaveGame object to set data on and then pass to SaveGameToSlot.
-	 *	@param	SaveGameClass	Class of SaveGame to create
-	 *	@return					New SaveGame object to write data to
+	 * Create a new, empty SaveGame object to set data on and then pass to SaveGameToSlot.
+	 * @param	SaveGameClass	Class of SaveGame to create
+	 * @return					New SaveGame object to write data to
 	 */
-	UFUNCTION(BlueprintCallable, Category="Game", meta=(DeterminesOutputType="SaveGameClass"))
+	UFUNCTION(BlueprintCallable, Category="SaveGame", meta=(DeterminesOutputType="SaveGameClass"))
 	static USaveGame* CreateSaveGameObject(TSubclassOf<USaveGame> SaveGameClass);
 
 	/** 
-	 *	Create a new, empty SaveGame object to set data on and then pass to SaveGameToSlot.
-	 *	@param	SaveGameBlueprint	Blueprint of SaveGame to create
-	 *	@return						New SaveGame object to write data to
+	 * Serialize our USaveGame object into a given array of bytes
+	 * @param SaveGameObject	Object that contains data about the save game that we want to write out
+	 * @param OutSaveData		Byte array that is written into
+	 * @return					Whether we successfully wrote data
 	 */
-	UFUNCTION(BlueprintCallable, Category="Game", meta=(DeprecatedFunction, DeprecationMessage="Use GameplayStatics.CreateSaveGameObject instead."))
-	static USaveGame* CreateSaveGameObjectFromBlueprint(UBlueprint* SaveGameBlueprint);
+	static bool SaveGameToMemory(USaveGame* SaveGameObject, TArray<uint8>& OutSaveData);
 
 	/** 
-	 *	Serialize our USaveGame object into a given array of bytes
-	 *	@param SaveGameObject	Object that contains data about the save game that we want to write out
-	 *	@return					Whether we successfully wrote data
+	 * Save the contents of the buffer to a platform-specific save slot/file 
+	 * @param InSaveData		Data to save
+	 * @param SlotName			Name of save game slot to save to.
+	 * @param UserIndex			For some platforms, master user index to identify the user doing the saving.
+	 * @return					Whether we successfully saved this information
 	 */
-	static bool SaveGameToMemory(USaveGame * SaveGameObject, TArray<uint8>& OutSaveData);
+	static bool SaveDataToSlot(const TArray<uint8>& InSaveData, const FString& SlotName, const int32 UserIndex);
 
 	/** 
-	 *	Save the contents of the buffer to a slot/file
-	 *	@param InSaveData		Data to save
-	 *	@param SlotName			Name of save game slot to save to.
-	 *  @param UserIndex		For some platforms, master user index to identify the user doing the saving.
+	 * Schedule an async save to a specific slot. UAsyncActionHandleSaveGame::AsyncSaveGameToSlot is the blueprint version of this.
+	 * This will do the serialize on the game thread, the platform-specific write on a worker thread, then call the complete delegate on the game thread.
+	 * The passed in delegate will be copied to a worker thread so make sure any payload is thread safe to copy by value.
+	 *
+	 * @param SaveGameObject	Object that contains data about the save game that we want to write out.
+	 * @param SlotName			Name of the save game slot to load from.
+	 * @param UserIndex			For some platforms, master user index to identify the user doing the loading.
+	 * @param SavedDelegate		Delegate that will be called on game thread when save succeeds or fails.
 	 */
-	static bool SaveDataToSlot(const TArray<uint8> & InSaveData, const FString & SlotName, const int32 UserIndex);
+	static void AsyncSaveGameToSlot(USaveGame* SaveGameObject, const FString& SlotName, const int32 UserIndex, FAsyncSaveGameToSlotDelegate SavedDelegate = FAsyncSaveGameToSlotDelegate());
 
 	/** 
-	 *	Save the contents of the SaveGameObject to a slot.
-	 *	@param SaveGameObject	Object that contains data about the save game that we want to write out
-	 *	@param SlotName			Name of save game slot to save to.
-	 *  @param UserIndex		For some platforms, master user index to identify the user doing the saving.
-	 *	@return					Whether we successfully saved this information
+	 * Save the contents of the SaveGameObject to a platform-specific save slot/file.
+	 * @param SaveGameObject	Object that contains data about the save game that we want to write out
+	 * @param SlotName			Name of save game slot to save to.
+	 * @param UserIndex			For some platforms, master user index to identify the user doing the saving.
+	 * @return					Whether we successfully saved this information
 	 */
-	UFUNCTION(BlueprintCallable, Category="Game")
+	UFUNCTION(BlueprintCallable, Category="SaveGame")
 	static bool SaveGameToSlot(USaveGame* SaveGameObject, const FString& SlotName, const int32 UserIndex);
 
 	/**
-	 *	See if a save game exists with the specified name.
-	 *	@param SlotName			Name of save game slot.
-	 *  @param UserIndex		For some platforms, master user index to identify the user doing the saving.
+	 * See if a save game exists with the specified name.
+	 * @param SlotName			Name of save game slot.
+	 * @param UserIndex			For some platforms, master user index to identify the user doing the saving.
 	 */
-	UFUNCTION(BlueprintCallable, Category="Game")
+	UFUNCTION(BlueprintCallable, Category="SaveGame")
 	static bool DoesSaveGameExist(const FString& SlotName, const int32 UserIndex);
 
 	/**
-	 *	Load the contents from a given array of bytes.
-	 *	@param InSaveData		The array containing the serialized USaveGame data to load
-	 *	@return SaveGameObject	Object containing loaded game state (NULL if load fails)
+	 * Tries to load a SaveGame object from a given array of bytes.
+	 * @param InSaveData		The array containing the serialized USaveGame data to load
+	 * @return					Object containing loaded game state (nullptr if load fails)
 	 */
 	static USaveGame* LoadGameFromMemory(const TArray<uint8>& InSaveData);
 
+	/**
+	 * Load contents from a slot/file into a buffer of save data.
+	 * @param OutSaveData		Data buffer to load into
+	 * @param SlotName			Name of save game slot to save to.
+	 * @param UserIndex			For some platforms, master user index to identify the user doing the saving.
+	 * @return					Whether valid save data was found and loaded.
+	 */
+	static bool LoadDataFromSlot(TArray<uint8>& OutSaveData, const FString& SlotName, const int32 UserIndex);
+
 	/** 
-	 *	Load the contents from a given slot.
-	 *	@param SlotName			Name of the save game slot to load from.
-	 *  @param UserIndex		For some platforms, master user index to identify the user doing the loading.
-	 *	@return SaveGameObject	Object containing loaded game state (NULL if load fails)
+	 * Schedule an async load of a specific slot. UAsyncActionHandleSaveGame::AsyncLoadGameFromSlot is the blueprint version of this.
+	 * This will do the platform-specific read on a worker thread, the serialize and creation on the game thread, and then will call the passed in delegate
+	 * The passed in delegate will be copied to a worker thread so make sure any payload is thread safe to copy by value
+	 *
+	 * @param SlotName			Name of the save game slot to load from.
+	 * @param UserIndex			For some platforms, master user index to identify the user doing the loading.
+	 * @param LoadedDelegate	Delegate that will be called on game thread when load succeeds or fails.
+	 */
+	static void AsyncLoadGameFromSlot(const FString& SlotName, const int32 UserIndex, FAsyncLoadGameFromSlotDelegate LoadedDelegate);
+
+	/** 
+	 * Load the contents from a given slot.
+	 * @param SlotName			Name of the save game slot to load from.
+	 * @param UserIndex			For some platforms, master user index to identify the user doing the loading.
+	 * @return					Object containing loaded game state (nullptr if load fails)
 	 */
-	UFUNCTION(BlueprintCallable, Category="Game")
+	UFUNCTION(BlueprintCallable, Category="SaveGame")
 	static USaveGame* LoadGameFromSlot(const FString& SlotName, const int32 UserIndex);
 
 	/** 
-	 *	Takes the provided buffer and consumes it, parsing past the internal header data, returning a MemoryReader
-	 *  that has: 1) been set up with all the related header information, and 2) offset to where tagged USaveGame object 
-	 *  serialization begins. NOTE: that the returned object has a reference to the supplied data - scope them 
-	 *  accordingly.
-	 *	@param  SaveData	A byte array, presumably produced by one of the SaveGame functions here.
-	 *	@return				A memory reader, wrapping SaveData, offset to the point past the header data.
+	 * Takes the provided buffer and consumes it, parsing past the internal header data, returning a MemoryReader
+	 * that has: 1) been set up with all the related header information, and 2) offset to where tagged USaveGame object 
+	 * serialization begins. NOTE: that the returned object has a reference to the supplied data - scope them 
+	 * accordingly.
+	 *
+	 * @param SaveData			A byte array, presumably produced by one of the SaveGame functions here.
+	 * @return					A memory reader, wrapping SaveData, offset to the point past the header data.
 	 */
 	static FMemoryReader StripSaveGameHeader(const TArray<uint8>& SaveData);
 
 	/**
 	 * Delete a save game in a particular slot.
-	 *	@param SlotName			Name of save game slot to delete.
-	 *  @param UserIndex		For some platforms, master user index to identify the user doing the deletion.
-	 *  @return True if a file was actually able to be deleted. use DoesSaveGameExist to distinguish between delete failures and failure due to file not existing.
+	 * @param SlotName			Name of save game slot to delete.
+	 * @param UserIndex			For some platforms, master user index to identify the user doing the deletion.
+	 * @return					True if a file was actually able to be deleted. use DoesSaveGameExist to distinguish between delete failures and failure due to file not existing.
 	 */
-	UFUNCTION(BlueprintCallable, Category = "Game")
+	UFUNCTION(BlueprintCallable, Category = "SaveGame")
 	static bool DeleteGameInSlot(const FString& SlotName, const int32 UserIndex);
 
 	/** Returns the frame delta time in seconds, adjusted by time dilation. */
@@ -913,6 +947,7 @@ public:
 	UFUNCTION(BlueprintPure, Category="Utilities|Time", meta=(WorldContext="WorldContextObject"))
 	static float GetAudioTimeSeconds(const UObject* WorldContextObject);
 
+	/** Returns time in seconds since the application was started. Unlike the other time functions this is accurate to the exact time this function is called instead of set once per frame. */
 	UFUNCTION(BlueprintPure, Category="Utilities|Time", meta=(WorldContext="WorldContextObject"))
 	static void GetAccurateRealTime(const UObject* WorldContextObject, int32& Seconds, float& PartialSeconds);
 
diff --git a/Engine/Source/Runtime/Engine/Private/GameFramework/AsyncActionHandleSaveGame.cpp b/Engine/Source/Runtime/Engine/Private/GameFramework/AsyncActionHandleSaveGame.cpp
new file mode 100644
--- /dev/null
+++ b/Engine/Source/Runtime/Engine/Private/GameFramework/AsyncActionHandleSaveGame.cpp
@@ -0,0 +1,64 @@
+// Copyright 1998-2019 Epic Games, Inc. All Rights Reserved.
+
+#include "GameFramework/AsyncActionHandleSaveGame.h"
+#include "Kismet/GameplayStatics.h"
+
+UAsyncActionHandleSaveGame* UAsyncActionHandleSaveGame::AsyncSaveGameToSlot(UObject* WorldContextObject, USaveGame* SaveGameObject, const FString& SlotName, const int32 UserIndex)
+{
+	UAsyncActionHandleSaveGame* Action = NewObject<UAsyncActionHandleSaveGame>();
+	Action->Operation = ESaveGameOperation::Save;
+	Action->SaveGameObject = SaveGameObject;
+	Action->SlotName = SlotName;
+	Action->UserIndex = UserIndex;
+	Action->RegisterWithGameInstance(WorldContextObject);
+
+	return Action;
+}
+
+UAsyncActionHandleSaveGame* UAsyncActionHandleSaveGame::AsyncLoadGameFromSlot(UObject* WorldContextObject, const FString& SlotName, const int32 UserIndex)
+{
+	UAsyncActionHandleSaveGame* Action = NewObject<UAsyncActionHandleSaveGame>();
+	Action->Operation = ESaveGameOperation::Load;
+	Action->SaveGameObject = nullptr;
+	Action->SlotName = SlotName;
+	Action->UserIndex = UserIndex;
+	Action->RegisterWithGameInstance(WorldContextObject);
+
+	return Action;
+}
+
+void UAsyncActionHandleSaveGame::Activate()
+{
+	switch (Operation)
+	{
+	case ESaveGameOperation::Save:
+		UGameplayStatics::AsyncSaveGameToSlot(SaveGameObject, SlotName, UserIndex, FAsyncSaveGameToSlotDelegate::CreateUObject(this, &UAsyncActionHandleSaveGame::HandleAsyncSave));
+		return;
+	case ESaveGameOperation::Load:
+		UGameplayStatics::AsyncLoadGameFromSlot(SlotName, UserIndex, FAsyncLoadGameFromSlotDelegate::CreateUObject(this, &UAsyncActionHandleSaveGame::HandleAsyncLoad));
+		return;
+	}
+
+	UE_LOG(LogScript, Error, TEXT("UAsyncActionHandleSaveGame Created with invalid operation!"));
+
+	ExecuteCompleted(false);
+}
+
+void UAsyncActionHandleSaveGame::HandleAsyncSave(const FString& InSlotName, const int32 InUserIndex, bool bSuccess)
+{
+	ExecuteCompleted(bSuccess);
+}
+
+void UAsyncActionHandleSaveGame::HandleAsyncLoad(const FString& InSlotName, const int32 InUserIndex, USaveGame* LoadedSave)
+{
+	SaveGameObject = LoadedSave;
+	ExecuteCompleted(SaveGameObject != nullptr);
+}
+
+void UAsyncActionHandleSaveGame::ExecuteCompleted(bool bSuccess)
+{
+	Completed.Broadcast(SaveGameObject, bSuccess);
+
+	SaveGameObject = nullptr;
+	SetReadyToDestroy();
+}
diff --git a/Engine/Source/Runtime/Engine/Private/GameplayStatics.cpp b/Engine/Source/Runtime/Engine/Private/GameplayStatics.cpp
--- a/Engine/Source/Runtime/Engine/Private/GameplayStatics.cpp
+++ b/Engine/Source/Runtime/Engine/Private/GameplayStatics.cpp
@@ -41,6 +41,7 @@
 #include "PhysicsEngine/BodySetup.h"
 #include "Misc/EngineVersion.h"
 #include "ContentStreaming.h"
+#include "Async/Async.h"
 
 #define LOCTEXT_NAMESPACE "GameplayStatics"
 
@@ -1793,15 +1794,6 @@ USaveGame* UGameplayStatics::CreateSaveGameObject(TSubclassOf<USaveGame> SaveGam
 	return nullptr;
 }
 
-USaveGame* UGameplayStatics::CreateSaveGameObjectFromBlueprint(UBlueprint* SaveGameBlueprint)
-{
-	if (SaveGameBlueprint && SaveGameBlueprint->GeneratedClass && SaveGameBlueprint->GeneratedClass->IsChildOf(USaveGame::StaticClass()))
-	{
-		return NewObject<USaveGame>(GetTransientPackage(), SaveGameBlueprint->GeneratedClass);
-	}
-	return nullptr;
-}
-
 bool UGameplayStatics::SaveGameToMemory(USaveGame* SaveGameObject, TArray<uint8>& OutSaveData )
 {
 	FMemoryWriter MemoryWriter(OutSaveData, true);
@@ -1829,24 +1821,33 @@ bool UGameplayStatics::SaveDataToSlot(const TArray<uint8>& InSaveData, const FSt
 	return false;
 }
 
-bool UGameplayStatics::SaveGameToSlot(USaveGame* SaveGameObject, const FString& SlotName, const int32 UserIndex)
+void UGameplayStatics::AsyncSaveGameToSlot(USaveGame* SaveGameObject, const FString& SlotName, const int32 UserIndex, FAsyncSaveGameToSlotDelegate SavedDelegate)
 {
-	ISaveGameSystem* SaveSystem = IPlatformFeaturesModule::Get().GetSaveGameSystem();
-	// If we have a system and an object to save and a save name...
-	if(SaveSystem && SaveGameObject && (SlotName.Len() > 0))
-	{
-		TArray<uint8> ObjectBytes;
-		FMemoryWriter MemoryWriter(ObjectBytes, true);
+	TArray<uint8> ObjectBytes;
+	SaveGameToMemory(SaveGameObject, ObjectBytes);
 
-		FSaveGameHeader SaveHeader(SaveGameObject->GetClass());
-		SaveHeader.Write(MemoryWriter);
+	AsyncTask(ENamedThreads::AnyHiPriThreadNormalTask, [SlotName, UserIndex, SavedDelegate, ObjectBytes]()
+	{
+		bool bSuccess = SaveDataToSlot(ObjectBytes, SlotName, UserIndex);
 
-		// Then save the object state, replacing object refs and names with strings
-		FObjectAndNameAsStringProxyArchive Ar(MemoryWriter, false);
-		SaveGameObject->Serialize(Ar);
+		// Now schedule the callback on the game thread, but only if it was bound to anything
+		if (SavedDelegate.IsBound())
+		{
+			AsyncTask(ENamedThreads::GameThread, [SlotName, UserIndex, SavedDelegate, bSuccess]()
+			{
+				SavedDelegate.ExecuteIfBound(SlotName, UserIndex, bSuccess);
+			});
+		}
+	});
+}
 
-		// Stuff that data into the save system with the desired file name
-		return SaveSystem->SaveGame(false, *SlotName, UserIndex, ObjectBytes);
+bool UGameplayStatics::SaveGameToSlot(USaveGame* SaveGameObject, const FString& SlotName, const int32 UserIndex)
+{
+	// This is a wrapper around the functions reading to/from a byte array
+	TArray<uint8> ObjectBytes;
+	if (SaveGameToMemory(SaveGameObject, ObjectBytes))
+	{
+		return SaveDataToSlot(ObjectBytes, SlotName, UserIndex);
 	}
 	return false;
 }
@@ -1869,26 +1870,14 @@ bool UGameplayStatics::DeleteGameInSlot(const FString& SlotName, const int32 Use
 	return false;
 }
 
-USaveGame* UGameplayStatics::LoadGameFromSlot(const FString& SlotName, const int32 UserIndex)
+USaveGame* UGameplayStatics::LoadGameFromMemory(const TArray<uint8>& InSaveData)
 {
-	ISaveGameSystem* SaveSystem = IPlatformFeaturesModule::Get().GetSaveGameSystem();
-	// If we have a save system and a valid name..
-	if (SaveSystem && (SlotName.Len() > 0))
+	if (InSaveData.Num() == 0)
 	{
-		// Load raw data from slot
-		TArray<uint8> ObjectBytes;
-		bool bSuccess = SaveSystem->LoadGame(false, *SlotName, UserIndex, ObjectBytes);
-		if (bSuccess)
-		{
-			return LoadGameFromMemory(ObjectBytes);
-		}
+		// Empty buffer, return instead of causing a bad serialize that could crash
+		return nullptr;
 	}
 
-	return nullptr;
-}
-
-USaveGame* UGameplayStatics::LoadGameFromMemory(const TArray<uint8>& InSaveData)
-{
 	USaveGame* OutSaveGameObject = nullptr;
 
 	FMemoryReader MemoryReader(InSaveData, true);
@@ -1915,6 +1904,57 @@ USaveGame* UGameplayStatics::LoadGameFromMemory(const TArray<uint8>& InSaveData)
 	return OutSaveGameObject;
 }
 
+bool UGameplayStatics::LoadDataFromSlot(TArray<uint8>& OutSaveData, const FString& SlotName, const int32 UserIndex)
+{
+	ISaveGameSystem* SaveSystem = IPlatformFeaturesModule::Get().GetSaveGameSystem();
+	// If we have a save system and a valid name..
+	if (SaveSystem && (SlotName.Len() > 0))
+	{
+		if (SaveSystem->LoadGame(false, *SlotName, UserIndex, OutSaveData))
+		{
+			return true;
+		}
+	}
+
+	// Clear buffer on a failed read
+	OutSaveData.Reset();
+	return false;
+}
+
+void UGameplayStatics::AsyncLoadGameFromSlot(const FString& SlotName, const int32 UserIndex, FAsyncLoadGameFromSlotDelegate LoadedDelegate)
+{
+	AsyncTask(ENamedThreads::AnyHiPriThreadNormalTask, [SlotName, UserIndex, LoadedDelegate]()
+	{
+		// Do the actual I/O on the background thread
+		TArray<uint8> ObjectBytes;
+		LoadDataFromSlot(ObjectBytes, SlotName, UserIndex);
+
+		// Now schedule the serialize and callback on the game thread
+		AsyncTask(ENamedThreads::GameThread, [SlotName, UserIndex, LoadedDelegate, ObjectBytes]()
+		{
+			USaveGame* LoadedGame = nullptr;
+			if (ObjectBytes.Num() > 0)
+			{
+				LoadedGame = LoadGameFromMemory(ObjectBytes);
+			}
+
+			LoadedDelegate.ExecuteIfBound(SlotName, UserIndex, LoadedGame);
+		});
+	});
+}
+
+USaveGame* UGameplayStatics::LoadGameFromSlot(const FString& SlotName, const int32 UserIndex)
+{
+	// This is a wrapper around the functions reading to/from a byte array
+	TArray<uint8> ObjectBytes;
+	if (LoadDataFromSlot(ObjectBytes, SlotName, UserIndex))
+	{
+		return LoadGameFromMemory(ObjectBytes);
+	}
+
+	return nullptr;
+}
+
 FMemoryReader UGameplayStatics::StripSaveGameHeader(const TArray<uint8>& SaveData)
 {
 	FMemoryReader MemoryReader(SaveData, /*bIsPersistent =*/true);
