PR #5473: Set Mouse Capture Mode from Blueprint (Contributed by [REDACTED_PERSON])

diff --git a/Engine/Source/Runtime/Engine/Classes/Kismet/GameplayStatics.h b/Engine/Source/Runtime/Engine/Classes/Kismet/GameplayStatics.h
--- a/Engine/Source/Runtime/Engine/Classes/Kismet/GameplayStatics.h
+++ b/Engine/Source/Runtime/Engine/Classes/Kismet/GameplayStatics.h
@@ -267,6 +267,18 @@ class ENGINE_API UGameplayStatics : public UBlueprintFunctionLibrary
 	UFUNCTION(BlueprintCallable, Category = "Rendering", meta = (WorldContext = "WorldContextObject"))
 	static bool GetEnableWorldRendering(const UObject* WorldContextObject);
 
+	/**
+	 * Returns the current viewport mouse capture mode
+	 */
+	UFUNCTION(BlueprintPure, Category = "Utilities", meta = (WorldContext = "WorldContextObject"))
+	static EMouseCaptureMode GetViewportMouseCaptureMode(const UObject* WorldContextObject);
+
+	/**
+	 * Sets the current viewport mouse capture mode
+	 */
+	UFUNCTION(BlueprintCallable, Category = "Utilities", meta = (WorldContext = "WorldContextObject"))
+	static void SetViewportMouseCaptureMode(const UObject* WorldContextObject, const EMouseCaptureMode MouseCaptureMode);
+
 	/** Hurt locally authoritative actors within the radius. Will only hit components that block the Visibility channel.
 	 * @param BaseDamage - The base damage to apply, i.e. the damage at the origin.
 	 * @param Origin - Epicenter of the damage area.
diff --git a/Engine/Source/Runtime/Engine/Private/GameplayStatics.cpp b/Engine/Source/Runtime/Engine/Private/GameplayStatics.cpp
--- a/Engine/Source/Runtime/Engine/Private/GameplayStatics.cpp
+++ b/Engine/Source/Runtime/Engine/Private/GameplayStatics.cpp
@@ -374,6 +374,34 @@ bool UGameplayStatics::GetEnableWorldRendering(const UObject* WorldContextObject
 	return false;
 }
 
+EMouseCaptureMode UGameplayStatics::GetViewportMouseCaptureMode(const UObject* WorldContextObject)
+{
+	UWorld* const World = GEngine->GetWorldFromContextObject(WorldContextObject, EGetWorldErrorMode::LogAndReturnNull);
+	if (World)
+	{
+		UGameViewportClient* const GameViewportClient = World->GetGameViewport();
+		if (GameViewportClient)
+		{
+			return GameViewportClient->CaptureMouseOnClick();
+		}
+	}
+
+	return EMouseCaptureMode::NoCapture;
+}
+
+void UGameplayStatics::SetViewportMouseCaptureMode(const UObject* WorldContextObject, const EMouseCaptureMode MouseCaptureMode)
+{
+	UWorld* const World = GEngine->GetWorldFromContextObject(WorldContextObject, EGetWorldErrorMode::LogAndReturnNull);
+	if (World)
+	{
+		UGameViewportClient* const GameViewportClient = World->GetGameViewport();
+		if (GameViewportClient)
+		{
+			GameViewportClient->SetCaptureMouseOnClick(MouseCaptureMode);
+		}
+	}
+}
+
 /** @RETURN True if weapon trace from Origin hits component VictimComp.  OutHitResult will contain properties of the hit. */
 static bool ComponentIsDamageableFrom(UPrimitiveComponent* VictimComp, FVector const& Origin, AActor const* IgnoredActor, const TArray<AActor*>& IgnoreActors, ECollisionChannel TraceChannel, FHitResult& OutHitResult)
 {
