Instanced components are now correctly attached to blueprint constructed components.

diff --git a/Engine/Source/Runtime/Engine/Private/ActorConstruction.cpp b/Engine/Source/Runtime/Engine/Private/ActorConstruction.cpp
--- a/Engine/Source/Runtime/Engine/Private/ActorConstruction.cpp
+++ b/Engine/Source/Runtime/Engine/Private/ActorConstruction.cpp
@@ -29,6 +29,7 @@
 #include "Engine/SimpleConstructionScript.h"
 #include "Components/ChildActorComponent.h"
 #include "ProfilingDebugging/CsvProfiler.h"
+#include "Algo/Transform.h"
 
 #if WITH_EDITOR
 #include "Editor.h"
@@ -420,6 +421,19 @@ void AActor::RerunConstructionScripts()
 			// In this case we need to "tunnel out" to find the parent component which has been created by the construction script.
 			if (UActorComponent* CSAddedComponent = GetComponentAddedByConstructionScript(Component))
 			{
+				// If we have any instanced components attached to us and we're going to be destroyed we need to explicitly detach them so they don't choose a new
+				// parent component and the attachment data we stored in the component instance data cache won't get applied
+				if (USceneComponent* SceneComp = Cast<USceneComponent>(Component))
+				{
+					TInlineComponentArray<USceneComponent*> InstancedChildren;
+					Algo::TransformIf(SceneComp->GetAttachChildren(), InstancedChildren, [](USceneComponent* SC) { return SC->CreationMethod == EComponentCreationMethod::Instance; }, [](USceneComponent* SC) { return SC; });
+					for (USceneComponent* InstancedChild : InstancedChildren)
+					{
+						InstancedChild->DetachFromComponent(FDetachmentTransformRules::KeepRelativeTransform);
+					}
+				}
+
+
 				// Determine if this component is an inner of a component added by the construction script
 				const bool bIsInnerComponent = (CSAddedComponent != Component);
 
