Reenable ensures about cyclic attachment.
#rnx

diff --git a/Engine/Source/Runtime/Engine/Private/Components/SceneComponent.cpp b/Engine/Source/Runtime/Engine/Private/Components/SceneComponent.cpp
--- a/Engine/Source/Runtime/Engine/Private/Components/SceneComponent.cpp
+++ b/Engine/Source/Runtime/Engine/Private/Components/SceneComponent.cpp
@@ -1711,9 +1711,9 @@ void USceneComponent::SetupAttachment(class USceneComponent* InParent, FName InS
 		if (ensureMsgf(InParent != this, TEXT("Cannot attach a component to itself.")))
 		{
 			// Paragon hack
-			if (/*ensureMsgf*/(InParent == nullptr || !InParent->IsAttachedTo(this))) //, TEXT("Setting up attachment would create a cycle.")))
+			if (ensureMsgf(InParent == nullptr || !InParent->IsAttachedTo(this), TEXT("Setting up attachment would create a cycle.")))
 			{
-				if (/*ensureMsgf*/(AttachParent == nullptr || !AttachParent->AttachChildren.Contains(this)))// , TEXT("SetupAttachment cannot be used once a component has already had AttachTo used to connect it to a parent.")))
+				if (ensureMsgf(AttachParent == nullptr || !AttachParent->AttachChildren.Contains(this), TEXT("SetupAttachment cannot be used once a component has already had AttachTo used to connect it to a parent.")))
 				{
 					AttachParent = InParent;
 					AttachSocketName = InSocketName;
