Clear ChildActorTemplate on an instance if the child actor class is set to a different class.

diff --git a/Engine/Source/Runtime/Engine/Private/Components/ChildActorComponent.cpp b/Engine/Source/Runtime/Engine/Private/Components/ChildActorComponent.cpp
--- a/Engine/Source/Runtime/Engine/Private/Components/ChildActorComponent.cpp
+++ b/Engine/Source/Runtime/Engine/Private/Components/ChildActorComponent.cpp
@@ -466,11 +466,20 @@ void UChildActorComponent::SetChildActorClass(TSubclassOf<AActor> Class)
 			ChildActorTemplate = nullptr;
 		}
 	}
-	else if (IsRegistered())
+	else
 	{
-		ChildActorName = NAME_None;
-		DestroyChildActor();
-		CreateChildActor();
+		// Clear actor template if it no longer matches the set class
+		if (ChildActorTemplate && ChildActorTemplate->GetClass() != ChildActorClass)
+		{
+			ChildActorTemplate = nullptr;
+		}
+
+		if (IsRegistered())
+		{
+			ChildActorName = NAME_None;
+			DestroyChildActor();
+			CreateChildActor();
+		}
 	}
 }
 
