Fix bad auto-resolve
#rnx

diff --git a/Engine/Source/Runtime/Engine/Private/InheritableComponentHandler.cpp b/Engine/Source/Runtime/Engine/Private/InheritableComponentHandler.cpp
--- a/Engine/Source/Runtime/Engine/Private/InheritableComponentHandler.cpp
+++ b/Engine/Source/Runtime/Engine/Private/InheritableComponentHandler.cpp
@@ -145,6 +145,7 @@ UActorComponent* UInheritableComponentHandler::CreateOverridenComponentTemplate(
 
 	// Clear transient flag if it was transient before and re copy off archetype
 	if (NewComponentTemplate->HasAnyFlags(RF_Transient))
+	{
 		const int32 ComponentIndex = UnnecessaryComponents.Find(NewComponentTemplate);
 		if (ComponentIndex != INDEX_NONE)
 		{
@@ -155,6 +156,7 @@ UActorComponent* UInheritableComponentHandler::CreateOverridenComponentTemplate(
 			CopyParams.bDoDelta = false;
 			UEngine::CopyPropertiesForUnrelatedObjects(BestArchetype, NewComponentTemplate, CopyParams);
 		}
+	}
 
 	FComponentOverrideRecord NewRecord;
 	NewRecord.ComponentKey = Key;
