Fix PointLight, RectLight, and SpotLight so that the light component can have properties modified when subclassed in blueprints

diff --git a/Engine/Source/Runtime/Engine/Classes/Engine/PointLight.h b/Engine/Source/Runtime/Engine/Classes/Engine/PointLight.h
--- a/Engine/Source/Runtime/Engine/Classes/Engine/PointLight.h
+++ b/Engine/Source/Runtime/Engine/Classes/Engine/PointLight.h
@@ -12,7 +12,7 @@ class APointLight : public ALight
 {
 	GENERATED_UCLASS_BODY()
 
-	UPROPERTY(BlueprintReadOnly, Category="Light", meta=(ExposeFunctionCategories="PointLight,Rendering|Lighting"))
+	UPROPERTY(BlueprintReadOnly, VisibleAnywhere, Category="Light", meta=(ExposeFunctionCategories="PointLight,Rendering|Lighting"))
 	class UPointLightComponent* PointLightComponent;
 
 	// BEGIN DEPRECATED (use component functions now in level script)
diff --git a/Engine/Source/Runtime/Engine/Classes/Engine/RectLight.h b/Engine/Source/Runtime/Engine/Classes/Engine/RectLight.h
--- a/Engine/Source/Runtime/Engine/Classes/Engine/RectLight.h
+++ b/Engine/Source/Runtime/Engine/Classes/Engine/RectLight.h
@@ -12,7 +12,7 @@ class ARectLight : public ALight
 {
 	GENERATED_UCLASS_BODY()
 
-	UPROPERTY(BlueprintReadOnly, Category="Light", meta=(ExposeFunctionCategories="RectLight,Rendering|Lighting"))
+	UPROPERTY(BlueprintReadOnly, VisibleAnywhere, Category="Light", meta=(ExposeFunctionCategories="RectLight,Rendering|Lighting"))
 	class URectLightComponent* RectLightComponent;
 
 #if WITH_EDITOR
diff --git a/Engine/Source/Runtime/Engine/Classes/Engine/SpotLight.h b/Engine/Source/Runtime/Engine/Classes/Engine/SpotLight.h
--- a/Engine/Source/Runtime/Engine/Classes/Engine/SpotLight.h
+++ b/Engine/Source/Runtime/Engine/Classes/Engine/SpotLight.h
@@ -12,7 +12,7 @@ class ASpotLight : public ALight
 {
 	GENERATED_UCLASS_BODY()
 
-	UPROPERTY(BlueprintReadOnly, Category="Light", meta=(ExposeFunctionCategories="SpotLight,Rendering|Lighting"))
+	UPROPERTY(BlueprintReadOnly, VisibleAnywhere, Category="Light", meta=(ExposeFunctionCategories="SpotLight,Rendering|Lighting"))
 	class USpotLightComponent* SpotLightComponent;
 
 #if WITH_EDITORONLY_DATA
