Fix FText determinism issues - blueprints failing to avoid loc key creation

diff --git a/Engine/Source/Editor/BlueprintGraph/Classes/EdGraphSchema_K2.h b/Engine/Source/Editor/BlueprintGraph/Classes/EdGraphSchema_K2.h
--- a/Engine/Source/Editor/BlueprintGraph/Classes/EdGraphSchema_K2.h
+++ b/Engine/Source/Editor/BlueprintGraph/Classes/EdGraphSchema_K2.h
@@ -572,7 +572,7 @@ public:
 	UK2Node* CreateSplitPinNode(UEdGraphPin* Pin, const FCreateSplitPinNodeParams& Params) const;
 
 	/** Reads in a FString and gets the values of the pin defaults for that type. This can be passed to DefaultValueSimpleValidation to validate. OwningObject can be null */
-	virtual void GetPinDefaultValuesFromString(const FEdGraphPinType& PinType, UObject* OwningObject, const FString& NewValue, FString& UseDefaultValue, UObject*& UseDefaultObject, FText& UseDefaultText) const;
+	virtual void GetPinDefaultValuesFromString(const FEdGraphPinType& PinType, UObject* OwningObject, const FString& NewValue, FString& UseDefaultValue, UObject*& UseDefaultObject, FText& UseDefaultText, bool bPreserveTextIdentity = true) const;
 
 	/** Do validation, that doesn't require a knowledge about actual pin */
 	virtual bool DefaultValueSimpleValidation(const FEdGraphPinType& PinType, const FName PinName, const FString& NewDefaultValue, UObject* NewDefaultObject, const FText& InText, FString* OutMsg = nullptr) const;
diff --git a/Engine/Source/Editor/BlueprintGraph/Private/EdGraphSchema_K2.cpp b/Engine/Source/Editor/BlueprintGraph/Private/EdGraphSchema_K2.cpp
--- a/Engine/Source/Editor/BlueprintGraph/Private/EdGraphSchema_K2.cpp
+++ b/Engine/Source/Editor/BlueprintGraph/Private/EdGraphSchema_K2.cpp
@@ -4600,7 +4600,7 @@ void UEdGraphSchema_K2::HandleGraphBeingDeleted(UEdGraph& GraphBeingRemoved) con
 	}
 }
 
-void UEdGraphSchema_K2::GetPinDefaultValuesFromString(const FEdGraphPinType& PinType, UObject* OwningObject, const FString& NewDefaultValue, FString& UseDefaultValue, UObject*& UseDefaultObject, FText& UseDefaultText) const
+void UEdGraphSchema_K2::GetPinDefaultValuesFromString(const FEdGraphPinType& PinType, UObject* OwningObject, const FString& NewDefaultValue, FString& UseDefaultValue, UObject*& UseDefaultObject, FText& UseDefaultText, bool bPreserveTextIdentity) const
 {
 	if ((PinType.PinCategory == PC_Object)
 		|| (PinType.PinCategory == PC_Class)
@@ -4632,16 +4632,20 @@ void UEdGraphSchema_K2::GetPinDefaultValuesFromString(const FEdGraphPinType& Pin
 	}
 	else if (PinType.PinCategory == PC_Text)
 	{
-		FString PackageNamespace;
-#if USE_STABLE_LOCALIZATION_KEYS
-		if (GIsEditor)
+		if (bPreserveTextIdentity)
 		{
-			PackageNamespace = TextNamespaceUtil::EnsurePackageNamespace(OwningObject);
+			UseDefaultText = FTextStringHelper::CreateFromBuffer(*NewDefaultValue);
 		}
-#endif // USE_STABLE_LOCALIZATION_KEYS
-		if (!FTextStringHelper::ReadFromBuffer(*NewDefaultValue, UseDefaultText, nullptr, *PackageNamespace))
+		else
 		{
-			UseDefaultText = FText::FromString(NewDefaultValue);
+			FString PackageNamespace;
+#if USE_STABLE_LOCALIZATION_KEYS
+			if (GIsEditor)
+			{
+				PackageNamespace = TextNamespaceUtil::EnsurePackageNamespace(OwningObject);
+			}
+#endif // USE_STABLE_LOCALIZATION_KEYS
+			UseDefaultText = FTextStringHelper::CreateFromBuffer(*NewDefaultValue, nullptr, *PackageNamespace);
 		}
 		UseDefaultObject = nullptr;
 		UseDefaultValue.Empty();
@@ -4674,7 +4678,7 @@ void UEdGraphSchema_K2::TrySetDefaultValue(UEdGraphPin& Pin, const FString& NewD
 	UObject* UseDefaultObject = nullptr;
 	FText UseDefaultText;
 
-	GetPinDefaultValuesFromString(Pin.PinType, Pin.GetOwningNodeUnchecked(), NewDefaultValue, UseDefaultValue, UseDefaultObject, UseDefaultText);
+	GetPinDefaultValuesFromString(Pin.PinType, Pin.GetOwningNodeUnchecked(), NewDefaultValue, UseDefaultValue, UseDefaultObject, UseDefaultText, /*bPreserveTextIdentity*/false);
 
 	// Check the default value and make it an error if it's bogus
 	if (IsPinDefaultValid(&Pin, UseDefaultValue, UseDefaultObject, UseDefaultText).IsEmpty())
diff --git a/Engine/Source/Runtime/Engine/Private/UserDefinedStruct.cpp b/Engine/Source/Runtime/Engine/Private/UserDefinedStruct.cpp
--- a/Engine/Source/Runtime/Engine/Private/UserDefinedStruct.cpp
+++ b/Engine/Source/Runtime/Engine/Private/UserDefinedStruct.cpp
@@ -2,7 +2,6 @@
 
 #include "Engine/UserDefinedStruct.h"
 #include "UObject/UObjectHash.h"
-#include "Serialization/PropertyLocalizationDataGathering.h"
 #include "UObject/StructOnScope.h"
 #include "UObject/UnrealType.h"
 #include "UObject/LinkerLoad.h"
@@ -10,7 +9,6 @@
 #include "Misc/SecureHash.h"
 #include "UObject/PropertyPortFlags.h"
 #include "Misc/PackageName.h"
-#include "Serialization/TextReferenceCollector.h"
 #include "Blueprint/BlueprintSupport.h"
 
 #if WITH_EDITOR
@@ -48,56 +46,10 @@ void FUserStructOnScopeIgnoreDefaults::Initialize()
 	}
 }
 
-#if WITH_EDITORONLY_DATA
-namespace
-{
-	void GatherUserDefinedStructForLocalization(const UObject* const Object, FPropertyLocalizationDataGatherer& PropertyLocalizationDataGatherer, const EPropertyLocalizationGathererTextFlags GatherTextFlags)
-	{
-		const UUserDefinedStruct* const UserDefinedStruct = CastChecked<UUserDefinedStruct>(Object);
-
-		PropertyLocalizationDataGatherer.GatherLocalizationDataFromObject(UserDefinedStruct, GatherTextFlags);
-
-		const FString PathToObject = UserDefinedStruct->GetPathName();
-
-		PropertyLocalizationDataGatherer.GatherLocalizationDataFromStructFields(PathToObject, UserDefinedStruct, UserDefinedStruct->GetDefaultInstance(), nullptr, GatherTextFlags);
-	}
-
-	void CollectUserDefinedStructTextReferences(UObject* Object, FArchive& Ar)
-	{
-		UUserDefinedStruct* const UserDefinedStruct = CastChecked<UUserDefinedStruct>(Object);
-
-		// User Defined Structs need some special handling as they store their default data in a way that serialize doesn't pick up
-		UUserDefinedStructEditorData* UDSEditorData = Cast<UUserDefinedStructEditorData>(UserDefinedStruct->EditorData);
-		if (UDSEditorData)
-		{
-			for (const FStructVariableDescription& StructVariableDesc : UDSEditorData->VariablesDescriptions)
-			{
-				static const FName TextCategory = TEXT("text"); // Must match UEdGraphSchema_K2::PC_Text
-				if (StructVariableDesc.Category == TextCategory)
-				{
-					FText StructVariableValue;
-					if (FTextStringHelper::ReadFromBuffer(*StructVariableDesc.DefaultValue, StructVariableValue))
-					{
-						Ar << StructVariableValue;
-					}
-				}
-			}
-		}
-
-		UserDefinedStruct->Serialize(Ar);
-	}
-}
-#endif
-
 UUserDefinedStruct::UUserDefinedStruct(const FObjectInitializer& ObjectInitializer)
 	: Super(ObjectInitializer)
 {
 	DefaultStructInstance.SetPackage(GetOutermost());
-
-#if WITH_EDITORONLY_DATA
-	{ static const FAutoRegisterLocalizationDataGatheringCallback AutomaticRegistrationOfLocalizationGatherer(UUserDefinedStruct::StaticClass(), &GatherUserDefinedStructForLocalization); }
-	{ static const FAutoRegisterTextReferenceCollectorCallback AutomaticRegistrationOfTextReferenceCollector(UUserDefinedStruct::StaticClass(), &CollectUserDefinedStructTextReferences); }
-#endif
 }
 
 void UUserDefinedStruct::Serialize(FStructuredArchive::FRecord Record)
