Timeline components should be constructed regardless of whether they are wired in or not

diff --git a/Engine/Source/Editor/KismetCompiler/Private/KismetCompiler.cpp b/Engine/Source/Editor/KismetCompiler/Private/KismetCompiler.cpp
--- a/Engine/Source/Editor/KismetCompiler/Private/KismetCompiler.cpp
+++ b/Engine/Source/Editor/KismetCompiler/Private/KismetCompiler.cpp
@@ -2629,11 +2629,6 @@ void FKismetCompilerContext::ExpandTimelineNodes(UEdGraph* SourceGraph)
 
 			// Set the timeline template as wired/not wired for component pruning later
 			const bool bWiredIn = bPlayPinConnected || bPlayFromStartPinConnected || bStopPinConnected || bReversePinConnected || bReverseFromEndPinConnected || bSetNewTimePinConnected;
-			const bool bWiredOut = bUpdatePinConnected || bFinishedPinConnected;
-			const bool bPlayWired = Timeline->bAutoPlay;
-			const bool bReferenced = TimelinePlayNodes.Find(TimelineNode->TimelineName) != INDEX_NONE;
-
-			Timeline->bValidatedAsWired = bWiredIn || bReferenced || (bPlayWired && bWiredOut);
 
 			// Only create nodes for play/stop if they are actually connected - otherwise we get a 'unused node being pruned' warning
 			if (bWiredIn)
diff --git a/Engine/Source/Runtime/Engine/Classes/Engine/TimelineTemplate.h b/Engine/Source/Runtime/Engine/Classes/Engine/TimelineTemplate.h
--- a/Engine/Source/Runtime/Engine/Classes/Engine/TimelineTemplate.h
+++ b/Engine/Source/Runtime/Engine/Classes/Engine/TimelineTemplate.h
@@ -166,10 +166,6 @@ class UTimelineTemplate : public UObject
 	UPROPERTY(EditAnywhere, BlueprintReadWrite, Category=TimelineTemplate)
 	uint8 bReplicated:1;
 
-	/** Compiler Validated As Wired up */
-	UPROPERTY()
-	uint8 bValidatedAsWired:1;
-
 	/** If we want the timeline to ignore global time dilation */
 	UPROPERTY(EditAnywhere, Category = TimelineTemplate)
 	uint8 bIgnoreTimeDilation : 1;
diff --git a/Engine/Source/Runtime/Engine/Private/BlueprintGeneratedClass.cpp b/Engine/Source/Runtime/Engine/Private/BlueprintGeneratedClass.cpp
--- a/Engine/Source/Runtime/Engine/Private/BlueprintGeneratedClass.cpp
+++ b/Engine/Source/Runtime/Engine/Private/BlueprintGeneratedClass.cpp
@@ -1004,7 +1004,6 @@ void UBlueprintGeneratedClass::CreateTimelineComponent(AActor* Actor, const UTim
 {
 	if (!Actor
 		|| !TimelineTemplate
-		|| !TimelineTemplate->bValidatedAsWired
 		|| Actor->IsTemplate()
 		|| Actor->IsPendingKill())
 	{
@@ -1132,7 +1131,7 @@ void UBlueprintGeneratedClass::CreateComponentsForActor(const UClass* ThisClass,
 		for (UTimelineTemplate* TimelineTemplate : BPGC->Timelines)
 		{
 			// Not fatal if NULL, but shouldn't happen and ignored if not wired up in graph
-			if (TimelineTemplate && TimelineTemplate->bValidatedAsWired)
+			if (TimelineTemplate)
 			{
 				CreateTimelineComponent(Actor, TimelineTemplate);
 			}
@@ -1144,7 +1143,7 @@ void UBlueprintGeneratedClass::CreateComponentsForActor(const UClass* ThisClass,
 		{
 			const UTimelineTemplate* TimelineTemplate = Cast<const UTimelineTemplate>(MiscObj);
 			// Not fatal if NULL, but shouldn't happen and ignored if not wired up in graph
-			if (TimelineTemplate && TimelineTemplate->bValidatedAsWired)
+			if (TimelineTemplate)
 			{
 				CreateTimelineComponent(Actor, TimelineTemplate);
 			}
diff --git a/Engine/Source/Runtime/Engine/Private/TimelineTemplate.cpp b/Engine/Source/Runtime/Engine/Private/TimelineTemplate.cpp
--- a/Engine/Source/Runtime/Engine/Private/TimelineTemplate.cpp
+++ b/Engine/Source/Runtime/Engine/Private/TimelineTemplate.cpp
@@ -39,7 +39,6 @@ UTimelineTemplate::UTimelineTemplate(const FObjectInitializer& ObjectInitializer
 	TimelineLength = 5.0f;
 	TimelineGuid = FGuid::NewGuid();
 	bReplicated = false;
-	bValidatedAsWired = false;
 
 	UpdateCachedNames();
 }
