PR #5726: GetGameTimeSinceCreation does not modify the Actor (Contributed by [REDACTED_PERSON])

diff --git a/Engine/Source/Runtime/Engine/Classes/GameFramework/Actor.h b/Engine/Source/Runtime/Engine/Classes/GameFramework/Actor.h
--- a/Engine/Source/Runtime/Engine/Classes/GameFramework/Actor.h
+++ b/Engine/Source/Runtime/Engine/Classes/GameFramework/Actor.h
@@ -1280,7 +1280,7 @@ public:
 
 	/** The number of seconds (in game time) since this Actor was created, relative to Get Game Time In Seconds. */
 	UFUNCTION(BlueprintPure, Category=Actor)
-	float GetGameTimeSinceCreation();
+	float GetGameTimeSinceCreation() const;
 
 protected:
 	/** Event when play begins for this actor. */
diff --git a/Engine/Source/Runtime/Engine/Private/Actor.cpp b/Engine/Source/Runtime/Engine/Private/Actor.cpp
--- a/Engine/Source/Runtime/Engine/Private/Actor.cpp
+++ b/Engine/Source/Runtime/Engine/Private/Actor.cpp
@@ -4797,7 +4797,7 @@ void AActor::K2_SetActorRelativeTransform(const FTransform& NewRelativeTransform
 	SetActorRelativeTransform(NewRelativeTransform, bSweep, (bSweep ? &SweepHitResult : nullptr), TeleportFlagToEnum(bTeleport));
 }
 
-float AActor::GetGameTimeSinceCreation()
+float AActor::GetGameTimeSinceCreation() const
 {
 	if (UWorld* MyWorld = GetWorld())
 	{
