Conform bHideFromSceneOutliner to pattern established in 4822786. uint8 packs better.
#rnx

diff --git a/Engine/Source/Runtime/Engine/Classes/Engine/World.h b/Engine/Source/Runtime/Engine/Classes/Engine/World.h
--- a/Engine/Source/Runtime/Engine/Classes/Engine/World.h
+++ b/Engine/Source/Runtime/Engine/Classes/Engine/World.h
@@ -464,11 +464,10 @@ public:
 
 #if WITH_EDITOR
 	/** Determines whether the begin play cycle will run on the spawned actor when in the editor. */
-	uint8 bTemporaryEditorActor:1;
+	uint8	bTemporaryEditorActor:1;
 
 	/* Determines wether or not the actor should be hidden from the Scene Outliner */
-	uint16 bHideFromSceneOutliner : 1;
-	
+	uint8	bHideFromSceneOutliner:1;
 #endif
 	
 	/* Flags used to describe the spawned actor/object instance. */
