This appears to have broken when BPGC::PostLoad stopped getting called in the editor, which is being discussed.

diff --git a/Engine/Source/Runtime/Engine/Private/BlueprintGeneratedClass.cpp b/Engine/Source/Runtime/Engine/Private/BlueprintGeneratedClass.cpp
--- a/Engine/Source/Runtime/Engine/Private/BlueprintGeneratedClass.cpp
+++ b/Engine/Source/Runtime/Engine/Private/BlueprintGeneratedClass.cpp
@@ -116,11 +116,6 @@ void UBlueprintGeneratedClass::PostLoad()
 			}
 		}
 
-		if (Package && Package->HasAnyPackageFlags(PKG_ForDiffing))
-		{
-			ClassFlags |= CLASS_Deprecated;
-		}
-
 #if UE_BLUEPRINT_EVENTGRAPH_FASTCALLS
 		// Patch the fast calls (needed as we can't bump engine version to serialize it directly in UFunction right now)
 		for (const FEventGraphFastCallPair& Pair : FastCallPairs_DEPRECATED)
@@ -1628,6 +1623,13 @@ void UBlueprintGeneratedClass::Serialize(FArchive& Ar)
 	if (Ar.IsLoading() && 0 == (Ar.GetPortFlags() & PPF_Duplicate))
 	{
 		CreatePersistentUberGraphFrame(ClassDefaultObject, true);
+
+		UPackage* Package = GetOutermost();
+		if (Package && Package->HasAnyPackageFlags(PKG_ForDiffing))
+		{
+			// If this is a diff package, set class to deprecated. This happens here to make sure it gets hit in all load cases
+			ClassFlags |= CLASS_Deprecated;
+		}
 	}
 }
 
