Fix importing csv where key field is not the first field

diff --git a/Engine/Source/Runtime/Engine/Classes/Engine/DataTable.h b/Engine/Source/Runtime/Engine/Classes/Engine/DataTable.h
--- a/Engine/Source/Runtime/Engine/Classes/Engine/DataTable.h
+++ b/Engine/Source/Runtime/Engine/Classes/Engine/DataTable.h
@@ -335,7 +335,7 @@ public:
 	ENGINE_API TArray<FString> CreateTableFromJSONString(const FString& InString);
 
 	/** Get array of UProperties that corresponds to columns in the table */
-	TArray<UProperty*> GetTablePropertyArray(const TArray<const TCHAR*>& Cells, UStruct* RowStruct, TArray<FString>& OutProblems, bool bIncludeNameColumn = false);
+	TArray<UProperty*> GetTablePropertyArray(const TArray<const TCHAR*>& Cells, UStruct* RowStruct, TArray<FString>& OutProblems, int32 KeyColumn = 0);
 	
 	/** 
 	 *	Create table from another Data Table
diff --git a/Engine/Source/Runtime/Engine/Private/DataTable.cpp b/Engine/Source/Runtime/Engine/Private/DataTable.cpp
--- a/Engine/Source/Runtime/Engine/Private/DataTable.cpp
+++ b/Engine/Source/Runtime/Engine/Private/DataTable.cpp
@@ -93,7 +93,7 @@ void UDataTable::LoadStructData(FArchive& Ar)
 	UScriptStruct* LoadUsingStruct = RowStruct;
 	if (!LoadUsingStruct)
 	{
-		if (!HasAnyFlags(RF_ClassDefaultObject))
+		if (!HasAnyFlags(RF_ClassDefaultObject) && GetOutermost() != GetTransientPackage())
 		{
 			UE_LOG(LogDataTable, Error, TEXT("Missing RowStruct while loading DataTable '%s'!"), *GetPathName());
 		}
@@ -129,7 +129,7 @@ void UDataTable::SaveStructData(FArchive& Ar)
 	UScriptStruct* SaveUsingStruct = RowStruct;
 	if (!SaveUsingStruct)
 	{
-		if (!HasAnyFlags(RF_ClassDefaultObject))
+		if (!HasAnyFlags(RF_ClassDefaultObject) && GetOutermost() != GetTransientPackage())
 		{
 			UE_LOG(LogDataTable, Error, TEXT("Missing RowStruct while saving DataTable '%s'!"), *GetPathName());
 		}
@@ -318,7 +318,7 @@ UScriptStruct& UDataTable::GetEmptyUsingStruct() const
 	UScriptStruct* EmptyUsingStruct = RowStruct;
 	if (!EmptyUsingStruct)
 	{
-		if (!HasAnyFlags(RF_ClassDefaultObject))
+		if (!HasAnyFlags(RF_ClassDefaultObject) && GetOutermost() != GetTransientPackage())
 		{
 			UE_LOG(LogDataTable, Error, TEXT("Missing RowStruct while emptying DataTable '%s'!"), *GetPathName());
 		}
@@ -589,7 +589,7 @@ bool UDataTable::WriteTableAsJSONObject(const TSharedRef< TJsonWriter<TCHAR, TPr
 }
 #endif
 
-TArray<UProperty*> UDataTable::GetTablePropertyArray(const TArray<const TCHAR*>& Cells, UStruct* InRowStruct, TArray<FString>& OutProblems, bool bIncludeNameColumn)
+TArray<UProperty*> UDataTable::GetTablePropertyArray(const TArray<const TCHAR*>& Cells, UStruct* InRowStruct, TArray<FString>& OutProblems, int32 KeyColumn)
 {
 	TArray<UProperty*> ColumnProps;
 
@@ -602,8 +602,13 @@ TArray<UProperty*> UDataTable::GetTablePropertyArray(const TArray<const TCHAR*>&
 		ColumnProps.AddZeroed( Cells.Num() );
 
 		// Skip first column depending on option
-		for (int32 ColIdx = bIncludeNameColumn ? 0 : 1; ColIdx < Cells.Num(); ++ColIdx)
+		for (int32 ColIdx = 0; ColIdx < Cells.Num(); ++ColIdx)
 		{
+			if (ColIdx == KeyColumn)
+			{
+				continue;
+			}
+
 			const TCHAR* ColumnValue = Cells[ColIdx];
 
 			FName PropName = DataTableUtils::MakeValidName(ColumnValue);
diff --git a/Engine/Source/Runtime/Engine/Private/DataTableCSV.cpp b/Engine/Source/Runtime/Engine/Private/DataTableCSV.cpp
--- a/Engine/Source/Runtime/Engine/Private/DataTableCSV.cpp
+++ b/Engine/Source/Runtime/Engine/Private/DataTableCSV.cpp
@@ -151,7 +151,6 @@ bool FDataTableImporterCSV::ReadTable()
 
 	// Find property for each column
 	int32 KeyColumn = 0;
-	bool bIncludeNameColumn = false;
 	if (!DataTable->ImportKeyField.IsEmpty())
 	{
 		// Search for key column by name
@@ -166,7 +165,7 @@ bool FDataTableImporterCSV::ReadTable()
 		}
 	}
 	
-	TArray<UProperty*> ColumnProps = DataTable->GetTablePropertyArray(Rows[0], DataTable->RowStruct, ImportProblems, bIncludeNameColumn);
+	TArray<UProperty*> ColumnProps = DataTable->GetTablePropertyArray(Rows[0], DataTable->RowStruct, ImportProblems, KeyColumn);
 
 	// Empty existing data
 	DataTable->EmptyTable();
@@ -196,7 +195,15 @@ bool FDataTableImporterCSV::ReadTable()
 		// Check its not 'none'
 		if(RowName == NAME_None)
 		{
-			ImportProblems.Add(FString::Printf(TEXT("Row '%d' missing a name."), RowIdx));
+			if (!DataTable->ImportKeyField.IsEmpty())
+			{
+				ImportProblems.Add(FString::Printf(TEXT("Row '%d' missing key field '%s'."), RowIdx, *DataTable->ImportKeyField));
+			}
+			else
+			{
+				ImportProblems.Add(FString::Printf(TEXT("Row '%d' missing a name."), RowIdx));
+			}
+
 			continue;
 		}
 
@@ -216,8 +223,13 @@ bool FDataTableImporterCSV::ReadTable()
 		DataTable->AddRowInternal(RowName, RowData);
 
 		// Now iterate over cells (skipping first cell unless we had an explicit name)
-		for(int32 CellIdx = bIncludeNameColumn ? 0 : 1; CellIdx<Cells.Num(); CellIdx++)
+		for(int32 CellIdx = 0; CellIdx < Cells.Num(); CellIdx++)
 		{
+			if (CellIdx == KeyColumn)
+			{
+				continue;
+			}
+
 			// Try and assign string to data using the column property
 			UProperty* ColumnProp = ColumnProps[CellIdx];
 			const FString CellValue = Cells[CellIdx];
diff --git a/Engine/Source/Runtime/Engine/Private/DataTableJSON.cpp b/Engine/Source/Runtime/Engine/Private/DataTableJSON.cpp
--- a/Engine/Source/Runtime/Engine/Private/DataTableJSON.cpp
+++ b/Engine/Source/Runtime/Engine/Private/DataTableJSON.cpp
@@ -449,7 +449,7 @@ bool FDataTableImporterJSON::ReadRow(const TSharedRef<FJsonObject>& InParsedTabl
 	// Check its not 'none'
 	if (RowName.IsNone())
 	{
-		ImportProblems.Add(FString::Printf(TEXT("Row '%d' missing a name."), InRowIdx));
+		ImportProblems.Add(FString::Printf(TEXT("Row '%d' missing key field '%s'."), InRowIdx, *RowKey));
 		return false;
 	}
 
