Remove unimplemented function definition
#rnx

diff --git a/Engine/Source/Runtime/Engine/Public/AudioDeviceManager.h b/Engine/Source/Runtime/Engine/Public/AudioDeviceManager.h
--- a/Engine/Source/Runtime/Engine/Public/AudioDeviceManager.h
+++ b/Engine/Source/Runtime/Engine/Public/AudioDeviceManager.h
@@ -243,9 +243,6 @@ private:
 		{}
 	};
 
-	/** Call back for garbage collector, ensures no processing is happening on the thread before collecting resources */
-	void OnPreGarbageCollect();
-
 	/** Returns index of the given handle */
 	uint32 GetIndex(uint32 Handle) const;
 
