Fix regression in type checking in reference replacement logic, and null check when processing attached children in actor construction

diff --git a/Engine/Source/Runtime/CoreUObject/Private/UObject/PropertyObject.cpp b/Engine/Source/Runtime/CoreUObject/Private/UObject/PropertyObject.cpp
--- a/Engine/Source/Runtime/CoreUObject/Private/UObject/PropertyObject.cpp
+++ b/Engine/Source/Runtime/CoreUObject/Private/UObject/PropertyObject.cpp
@@ -93,6 +93,11 @@ void UObjectProperty::SerializeItem( FStructuredArchive::FSlot Slot, void* Value
 		// Serialize in place
 		UObject** ObjectPtr = GetPropertyValuePtr(Value);
 		Slot << (*ObjectPtr);
+
+		if(!UnderlyingArchive.IsSaving())
+		{
+			CheckValidObject(ObjectPtr);
+		}
 	}
 	else
 	{
diff --git a/Engine/Source/Runtime/Engine/Private/ActorConstruction.cpp b/Engine/Source/Runtime/Engine/Private/ActorConstruction.cpp
--- a/Engine/Source/Runtime/Engine/Private/ActorConstruction.cpp
+++ b/Engine/Source/Runtime/Engine/Private/ActorConstruction.cpp
@@ -426,7 +426,7 @@ void AActor::RerunConstructionScripts()
 				if (USceneComponent* SceneComp = Cast<USceneComponent>(Component))
 				{
 					TInlineComponentArray<USceneComponent*> InstancedChildren;
-					Algo::TransformIf(SceneComp->GetAttachChildren(), InstancedChildren, [](USceneComponent* SC) { return SC->CreationMethod == EComponentCreationMethod::Instance; }, [](USceneComponent* SC) { return SC; });
+					Algo::TransformIf(SceneComp->GetAttachChildren(), InstancedChildren, [](USceneComponent* SC) { return SC && SC->CreationMethod == EComponentCreationMethod::Instance; }, [](USceneComponent* SC) { return SC; });
 					for (USceneComponent* InstancedChild : InstancedChildren)
 					{
 						InstancedChild->DetachFromComponent(FDetachmentTransformRules::KeepRelativeTransform);
