PR #5314: Add ENGINE_API to FStreamLevelAction so that other modules can use it (Contributed by [REDACTED_PERSON])

diff --git a/Engine/Source/Runtime/Engine/Classes/Engine/LevelStreaming.h b/Engine/Source/Runtime/Engine/Classes/Engine/LevelStreaming.h
--- a/Engine/Source/Runtime/Engine/Classes/Engine/LevelStreaming.h
+++ b/Engine/Source/Runtime/Engine/Classes/Engine/LevelStreaming.h
@@ -15,7 +15,7 @@ class ULevel;
 class ULevelStreaming;
 
 // Stream Level Action
-class FStreamLevelAction : public FPendingLatentAction
+class ENGINE_API FStreamLevelAction : public FPendingLatentAction
 {
 public:
 	bool			bLoading;
