Comment fixup

diff --git a/Engine/Source/Runtime/Engine/Private/Pawn.cpp b/Engine/Source/Runtime/Engine/Private/Pawn.cpp
--- a/Engine/Source/Runtime/Engine/Private/Pawn.cpp
+++ b/Engine/Source/Runtime/Engine/Private/Pawn.cpp
@@ -42,7 +42,7 @@ APawn::APawn(const FObjectInitializer& ObjectInitializer)
 
 	if (HasAnyFlags(RF_ClassDefaultObject) && GetClass() == APawn::StaticClass())
 	{
-		// WARNING: This line is why the AISupport plugin has to load the AIModule at PreEarlyLoadingScreen time, otherwise this load fails and CDOs are corrupt in the editor
+		// WARNING: This line is why the AISupport plugin has to load the AIModule before UObject initialization, otherwise this load fails and CDOs are corrupt in the editor
 		AIControllerClass = LoadClass<AController>(nullptr, *((UEngine*)(UEngine::StaticClass()->GetDefaultObject()))->AIControllerClassName.ToString(), nullptr, LOAD_None, nullptr);
 	}
 	else
