Switched APawn.AIControllerClass from TSubclassOf to TSoftClassPtr to avoid requiring AIModule early load #UE4

with proper deprecation and back-comp.

diff --git a/Engine/Source/Runtime/Engine/Classes/GameFramework/Pawn.h b/Engine/Source/Runtime/Engine/Classes/GameFramework/Pawn.h
--- a/Engine/Source/Runtime/Engine/Classes/GameFramework/Pawn.h
+++ b/Engine/Source/Runtime/Engine/Classes/GameFramework/Pawn.h
@@ -93,7 +93,7 @@ public:
 
 	/**
 	 * Determines when the Pawn creates and is possessed by an AI Controller (on level start, when spawned, etc).
-	 * Only possible if AIControllerClass is set, and ignored if AutoPossessPlayer is enabled.
+	 * Only possible if AIControllerClassRef is set, and ignored if AutoPossessPlayer is enabled.
 	 * @see AutoPossessPlayer
 	 */
 	UPROPERTY(EditAnywhere, Category=Pawn)
@@ -103,10 +103,15 @@ public:
 	UPROPERTY(replicated)
 	uint8 RemoteViewPitch;
 
-	/** Default class to use when pawn is controlled by AI. */
-	UPROPERTY(EditAnywhere, BlueprintReadWrite, meta=(DisplayName="AI Controller Class"), Category=Pawn)
+	/** This property is deprecated. Use \'AI Controller Class\' instead */
+	UE_DEPRECATED(4.23, "This property is deprecated. Use AIControllerClassRef instead.")
+	UPROPERTY(BlueprintReadWrite, meta=(DisplayName="DEPRECATED AI Controller Class", DeprecatedProperty, DeprecationMessage = "This property is deprecated. Use \'AIController Class Ref\' instead"), Category=Pawn)
 	TSubclassOf<AController> AIControllerClass;
 
+	/** Default controller class to use when pawn is controlled by AI. */
+	UPROPERTY(EditAnywhere, BlueprintReadWrite, meta = (AllowedClasses = "AIController"), Category = Pawn)
+	TSoftClassPtr<AController> AIControllerClassRef;
+
 	/**
 	 * Return our PawnNoiseEmitterComponent, if any. Default implementation returns the first PawnNoiseEmitterComponent found in the components array.
 	 * If one isn't found, then it tries to find one on the Pawn's current Controller.
@@ -234,6 +239,7 @@ public:
 	virtual void EndPlay(const EEndPlayReason::Type EndPlayReason) override;
 	virtual void PreInitializeComponents() override;
 	virtual void PostInitializeComponents() override;
+	virtual void PostInitProperties() override;
 	virtual const AActor* GetNetOwner() const override;
 	virtual UPlayer* GetNetOwningPlayer() override;
 	virtual class UNetConnection* GetNetConnection() const override;
diff --git a/Engine/Source/Runtime/Engine/Private/Pawn.cpp b/Engine/Source/Runtime/Engine/Private/Pawn.cpp
--- a/Engine/Source/Runtime/Engine/Private/Pawn.cpp
+++ b/Engine/Source/Runtime/Engine/Private/Pawn.cpp
@@ -32,6 +32,20 @@
 DEFINE_LOG_CATEGORY(LogDamage);
 DEFINE_LOG_CATEGORY_STATIC(LogPawn, Warning, All);
 
+namespace 
+{
+	FORCEINLINE void PatchAIControllerClass(APawn& Pawn)
+	{
+		PRAGMA_DISABLE_DEPRECATION_WARNINGS
+		if (Pawn.AIControllerClass)
+		{
+			Pawn.AIControllerClassRef = Pawn.AIControllerClass;
+		}
+		Pawn.AIControllerClass = nullptr;
+		PRAGMA_ENABLE_DEPRECATION_WARNINGS
+	}
+}
+
 APawn::APawn(const FObjectInitializer& ObjectInitializer)
 	: Super(ObjectInitializer)
 {
@@ -42,12 +56,11 @@ APawn::APawn(const FObjectInitializer& ObjectInitializer)
 
 	if (HasAnyFlags(RF_ClassDefaultObject) && GetClass() == APawn::StaticClass())
 	{
-		// WARNING: This line is why the AISupport plugin has to load the AIModule before UObject initialization, otherwise this load fails and CDOs are corrupt in the editor
-		AIControllerClass = LoadClass<AController>(nullptr, *((UEngine*)(UEngine::StaticClass()->GetDefaultObject()))->AIControllerClassName.ToString(), nullptr, LOAD_None, nullptr);
+		AIControllerClassRef = FSoftObjectPath(((UEngine*)(UEngine::StaticClass()->GetDefaultObject()))->AIControllerClassName);
 	}
 	else
 	{
-		AIControllerClass = ((APawn*)APawn::StaticClass()->GetDefaultObject())->AIControllerClass;
+		AIControllerClassRef = ((APawn*)APawn::StaticClass()->GetDefaultObject())->AIControllerClassRef;
 	}
 	bCanBeDamaged = true;
 	
@@ -131,12 +144,23 @@ void APawn::PostInitializeComponents()
 	}
 }
 
+void APawn::PostInitProperties()
+{
+	Super::PostInitProperties();
+
+	// 4.23 patching
+	PatchAIControllerClass(*this);
+}
+
 void APawn::PostLoad()
 {
 	Super::PostLoad();
 
 	// A pawn should never have this enabled, so we'll aggressively disable it if it did occur.
 	AutoReceiveInput = EAutoReceiveInput::Disabled;
+
+	// 4.23 patching
+	PatchAIControllerClass(*this);
 }
 
 void APawn::PostRegisterAllComponents()
@@ -312,14 +336,16 @@ void APawn::SpawnDefaultController()
 	{
 		return;
 	}
-	if ( AIControllerClass != nullptr )
+
+	UClass* AIControllerClassPtr = AIControllerClassRef.Get();
+	if (AIControllerClassPtr)
 	{
 		FActorSpawnParameters SpawnInfo;
 		SpawnInfo.Instigator = Instigator;
 		SpawnInfo.SpawnCollisionHandlingOverride = ESpawnActorCollisionHandlingMethod::AlwaysSpawn;
 		SpawnInfo.OverrideLevel = GetLevel();
 		SpawnInfo.ObjectFlags |= RF_Transient;	// We never want to save AI controllers into a map
-		AController* NewController = GetWorld()->SpawnActor<AController>(AIControllerClass, GetActorLocation(), GetActorRotation(), SpawnInfo);
+		AController* NewController = GetWorld()->SpawnActor<AController>(AIControllerClassPtr, GetActorLocation(), GetActorRotation(), SpawnInfo);
 		if (NewController != nullptr)
 		{
 			// if successful will result in setting this->Controller 
