Fix Merge errors
#rnx

diff --git a/Engine/Source/Runtime/Engine/Classes/Engine/StreamableManager.h b/Engine/Source/Runtime/Engine/Classes/Engine/StreamableManager.h
--- a/Engine/Source/Runtime/Engine/Classes/Engine/StreamableManager.h
+++ b/Engine/Source/Runtime/Engine/Classes/Engine/StreamableManager.h
@@ -317,16 +317,13 @@ struct ENGINE_API FStreamableManager : public FGCObject
 	const FString& GetManagerName() const;
 
 	/** Modifies the debug name of this manager, used for debugging GC references */
-	void SetManagerName(const FString& InName);
+	void SetManagerName(FString InName);
 
 	/** Add referenced objects to stop them from GCing */
 	virtual void AddReferencedObjects(FReferenceCollector& Collector) override;
-	virtual FString GetReferencerName() const override { return TEXT("FStreamableManager"); }
+	virtual FString GetReferencerName() const override { return ManagerName; }
 	virtual bool GetReferencerPropertyName(UObject* Object, FString& OutPropertyName) const override;
 
-	/** Returns debug name */
-	virtual FString GetReferencerName() const override;
-
 	FStreamableManager();
 	~FStreamableManager();
 
diff --git a/Engine/Source/Runtime/Engine/Private/StreamableManager.cpp b/Engine/Source/Runtime/Engine/Private/StreamableManager.cpp
--- a/Engine/Source/Runtime/Engine/Private/StreamableManager.cpp
+++ b/Engine/Source/Runtime/Engine/Private/StreamableManager.cpp
@@ -797,19 +797,14 @@ void FStreamableManager::AddReferencedObjects(FReferenceCollector& Collector)
 	}
 }
 
-FString FStreamableManager::GetReferencerName() const
-{
-	return ManagerName;
-}
-
 const FString& FStreamableManager::GetManagerName() const
 {
 	return ManagerName;
 }
 
-void FStreamableManager::SetManagerName(const FString& InName)
+void FStreamableManager::SetManagerName(FString InName)
 {
-	ManagerName = InName;
+	ManagerName = MoveTemp(InName);
 }
 
 bool FStreamableManager::GetReferencerPropertyName(UObject* Object, FString& OutPropertyName) const
