PR #5880: ACharacter::CanCrouch() should be const (Contributed by [REDACTED_PERSON])

diff --git a/Engine/Source/Runtime/Engine/Classes/GameFramework/Character.h b/Engine/Source/Runtime/Engine/Classes/GameFramework/Character.h
--- a/Engine/Source/Runtime/Engine/Classes/GameFramework/Character.h
+++ b/Engine/Source/Runtime/Engine/Classes/GameFramework/Character.h
@@ -732,7 +732,8 @@ public:
 	virtual void UnCrouch(bool bClientSimulation = false);
 
 	/** @return true if this character is currently able to crouch (and is not currently crouched) */
-	virtual bool CanCrouch();
+	UFUNCTION(BlueprintCallable, Category=Character)
+	virtual bool CanCrouch() const;
 
 	/** 
 	 * Called when Character stops crouching. Called on non-owned Characters through bIsCrouched replication.
diff --git a/Engine/Source/Runtime/Engine/Private/Character.cpp b/Engine/Source/Runtime/Engine/Private/Character.cpp
--- a/Engine/Source/Runtime/Engine/Private/Character.cpp
+++ b/Engine/Source/Runtime/Engine/Private/Character.cpp
@@ -376,7 +376,7 @@ void ACharacter::SetReplicateMovement(bool bInReplicateMovement)
 	}
 }
 
-bool ACharacter::CanCrouch()
+bool ACharacter::CanCrouch() const
 {
 	return !bIsCrouched && CharacterMovement && CharacterMovement->CanEverCrouch() && GetRootComponent() && !GetRootComponent()->IsSimulatingPhysics();
 }
