#ue4 - Optimize UPrimitiveComponent::UpdateOverlapsImpl() to manipulate pointers to FOverlapInfo structs rather than the large structs themselves. Reduces cost of this function by 10%.

diff --git a/Engine/Source/Runtime/Engine/Private/Components/PrimitiveComponent.cpp b/Engine/Source/Runtime/Engine/Private/Components/PrimitiveComponent.cpp
--- a/Engine/Source/Runtime/Engine/Private/Components/PrimitiveComponent.cpp
+++ b/Engine/Source/Runtime/Engine/Private/Components/PrimitiveComponent.cpp
@@ -53,6 +53,7 @@ namespace PrimitiveComponentStatics
 }
 
 typedef TArray<FOverlapInfo, TInlineAllocator<3>> TInlineOverlapInfoArray;
+typedef TArray<const FOverlapInfo*, TInlineAllocator<8>> TInlineOverlapPointerArray;
 
 DEFINE_LOG_CATEGORY_STATIC(LogPrimitiveComponent, Log, All);
 
@@ -153,6 +154,12 @@ struct FFastOverlapInfoCompare
 			&& MyBaseInfo.GetBodyIndex() == Info.GetBodyIndex();
 	}
 
+	bool operator() (const FOverlapInfo* Info)
+	{
+		return MyBaseInfo.OverlapInfo.Component.HasSameIndexAndSerialNumber(Info->OverlapInfo.Component)
+			&& MyBaseInfo.GetBodyIndex() == Info->GetBodyIndex();
+	}
+
 private:
 	const FOverlapInfo& MyBaseInfo;
 
@@ -166,6 +173,13 @@ FORCEINLINE_DEBUGGABLE int32 IndexOfOverlapFast(const TArray<FOverlapInfo, Alloc
 	return OverlapArray.IndexOfByPredicate(FFastOverlapInfoCompare(SearchItem));
 }
 
+// Version that works with arrays of pointers and pointers to search items.
+template<class AllocatorType>
+FORCEINLINE_DEBUGGABLE int32 IndexOfOverlapFast(const TArray<const FOverlapInfo*, AllocatorType>& OverlapPtrArray, const FOverlapInfo* SearchItem)
+{
+	return OverlapPtrArray.IndexOfByPredicate(FFastOverlapInfoCompare(*SearchItem));
+}
+
 // Helper for adding an FOverlapInfo uniquely to an Array, using IndexOfOverlapFast and knowing that at least one overlap is valid (non-null).
 template<class AllocatorType>
 FORCEINLINE_DEBUGGABLE void AddUniqueOverlapFast(TArray<FOverlapInfo, AllocatorType>& OverlapArray, FOverlapInfo& NewOverlap)
@@ -195,7 +209,24 @@ FORCEINLINE_DEBUGGABLE static bool CanComponentsGenerateOverlap(const UPrimitive
 		&& MyComponent->GetCollisionResponseToComponent(OtherComp) == ECR_Overlap;
 }
 
-// Predicate to remove components from overlaps array that can no longer overlap
+// Predicate to identify components from overlaps array that can overlap
+struct FPredicateFilterCanOverlap
+{
+	FPredicateFilterCanOverlap(const UPrimitiveComponent& OwningComponent)
+		: MyComponent(OwningComponent)
+	{
+	}
+
+	bool operator() (const FOverlapInfo& Info) const
+	{
+		return CanComponentsGenerateOverlap(&MyComponent, Info.OverlapInfo.GetComponent());
+	}
+
+private:
+	const UPrimitiveComponent& MyComponent;
+};
+
+// Predicate to identify components from overlaps array that can no longer overlap
 struct FPredicateFilterCannotOverlap
 {
 	FPredicateFilterCannotOverlap(const UPrimitiveComponent& OwningComponent)
@@ -212,6 +243,31 @@ private:
 	const UPrimitiveComponent& MyComponent;
 };
 
+// Helper to initialize an array to point to data members in another array.
+template <class ElementType, class AllocatorType1, class AllocatorType2>
+FORCEINLINE_DEBUGGABLE static void GetPointersToArrayData(TArray<const ElementType*, AllocatorType1>& Pointers, const TArray<ElementType, AllocatorType2>& DataArray)
+{
+	const int32 NumItems = DataArray.Num();
+	Pointers.SetNumUninitialized(NumItems);
+	for (int32 i=0; i < NumItems; i++)
+	{
+		Pointers[i] = &(DataArray[i]);
+	}
+}
+
+// Helper to initialize an array to point to data members in another array which satisfy a predicate.
+template <class ElementType, class AllocatorType1, class AllocatorType2, typename PredicateT>
+FORCEINLINE_DEBUGGABLE static void GetPointersToArrayDataByPredicate(TArray<const ElementType*, AllocatorType1>& Pointers, const TArray<ElementType, AllocatorType2>& DataArray, PredicateT Predicate)
+{
+	Pointers.Reserve(Pointers.Num() + DataArray.Num());
+	for (const ElementType& Item : DataArray)
+	{
+		if (Invoke(Predicate, Item))
+		{
+			Pointers.Add(&Item);
+		}
+	}
+}
 
 ///////////////////////////////////////////////////////////////////////////////
 // PRIMITIVE COMPONENT
@@ -1742,7 +1798,7 @@ static FORCEINLINE_DEBUGGABLE bool ShouldCheckOverlapFlagToQueueOverlaps(const U
 }
 
 
-static bool ShouldIgnoreOverlapResult(const UWorld* World, const AActor* ThisActor, const UPrimitiveComponent& ThisComponent, const AActor* OtherActor, const UPrimitiveComponent& OtherComponent, bool bCheckOverlapFlags)
+static FORCEINLINE_DEBUGGABLE bool ShouldIgnoreOverlapResult(const UWorld* World, const AActor* ThisActor, const UPrimitiveComponent& ThisComponent, const AActor* OtherActor, const UPrimitiveComponent& OtherComponent, bool bCheckOverlapFlags)
 {
 	// Don't overlap with self
 	if (&ThisComponent == &OtherComponent)
@@ -2402,15 +2458,13 @@ void UPrimitiveComponent::BeginComponentOverlap(const FOverlapInfo& OtherOverlap
 		return;
 	}
 
-	//	UE_LOG(LogActor, Log, TEXT("BEGIN OVERLAP! Self=%s SelfComp=%s, Other=%s, OtherComp=%s"), *GetNameSafe(this), *GetNameSafe(MyComp), *GetNameSafe(OtherComp->GetOwner()), *GetNameSafe(OtherComp));
-
-	UPrimitiveComponent* OtherComp = OtherOverlap.OverlapInfo.Component.Get();
-
-	if (OtherComp)
+	const bool bComponentsAlreadyTouching = (IndexOfOverlapFast(OverlappingComponents, OtherOverlap) != INDEX_NONE);
+	if (!bComponentsAlreadyTouching)
 	{
-		bool const bComponentsAlreadyTouching = (IndexOfOverlapFast(OverlappingComponents, OtherOverlap) != INDEX_NONE);
-		if (!bComponentsAlreadyTouching && CanComponentsGenerateOverlap(this, OtherComp))
+		UPrimitiveComponent* OtherComp = OtherOverlap.OverlapInfo.Component.Get();
+		if (CanComponentsGenerateOverlap(this, OtherComp))
 		{
+			//UE_LOG(LogActor, Log, TEXT("BEGIN OVERLAP! Self=%s SelfComp=%s, Other=%s, OtherComp=%s"), *GetNameSafe(this), *GetNameSafe(MyComp), *GetNameSafe(OtherComp->GetOwner()), *GetNameSafe(OtherComp));
 			GlobalOverlapEventsCounter++;			
 			AActor* const OtherActor = OtherComp->GetOwner();
 			AActor* const MyActor = GetOwner();
@@ -2476,8 +2530,6 @@ void UPrimitiveComponent::EndComponentOverlap(const FOverlapInfo& OtherOverlap,
 		return;
 	}
 
-	//	UE_LOG(LogActor, Log, TEXT("END OVERLAP! Self=%s SelfComp=%s, Other=%s, OtherComp=%s"), *GetNameSafe(this), *GetNameSafe(MyComp), *GetNameSafe(OtherActor), *GetNameSafe(OtherComp));
-
 	const int32 OtherOverlapIdx = IndexOfOverlapFast(OtherComp->OverlappingComponents, FOverlapInfo(this, INDEX_NONE));
 	if (OtherOverlapIdx != INDEX_NONE)
 	{
@@ -2487,6 +2539,7 @@ void UPrimitiveComponent::EndComponentOverlap(const FOverlapInfo& OtherOverlap,
 	const int32 OverlapIdx = IndexOfOverlapFast(OverlappingComponents, OtherOverlap);
 	if (OverlapIdx != INDEX_NONE)
 	{
+		//UE_LOG(LogActor, Log, TEXT("END OVERLAP! Self=%s SelfComp=%s, Other=%s, OtherComp=%s"), *GetNameSafe(this), *GetNameSafe(MyComp), *GetNameSafe(OtherActor), *GetNameSafe(OtherComp));
 		GlobalOverlapEventsCounter++;
 		OverlappingComponents.RemoveAtSwap(OverlapIdx, 1, false);
 
@@ -2642,7 +2695,9 @@ const TArray<FOverlapInfo>* UPrimitiveComponent::ConvertSweptOverlapsToCurrentOv
 
 				// Check components we hit during the sweep, keep only those still overlapping
 				const FCollisionQueryParams UnusedQueryParams(NAME_None, FCollisionQueryParams::GetUnknownStatId());
-				for (int32 Index = SweptOverlapsIndex; Index < SweptOverlaps.Num(); ++Index)
+				const int32 NumSweptOverlaps = SweptOverlaps.Num();
+				OverlapsAtEndLocation.Reserve(OverlapsAtEndLocation.Num() + NumSweptOverlaps);
+				for (int32 Index = SweptOverlapsIndex; Index < NumSweptOverlaps; ++Index)
 				{
 					const FOverlapInfo& OtherOverlap = SweptOverlaps[Index];
 					UPrimitiveComponent* OtherPrimitive = OtherOverlap.OverlapInfo.GetComponent();
@@ -2700,6 +2755,7 @@ const TArray<FOverlapInfo>* UPrimitiveComponent::ConvertRotationOverlapsToCurren
 			if (bEnableFastOverlapCheck)
 			{
 				// Add all current overlaps that are not children. Children test for their own overlaps after we update our own, and we ignore children in our own update.
+				OverlapsAtEndLocation.Reserve(OverlapsAtEndLocation.Num() + CurrentOverlaps.Num());
 				Algo::CopyIf(CurrentOverlaps, OverlapsAtEndLocation, FPredicateOverlapHasDifferentActor(*Actor));
 				Result = &OverlapsAtEndLocation;
 			}
@@ -2866,36 +2922,39 @@ bool UPrimitiveComponent::UpdateOverlapsImpl(const TArray<FOverlapInfo>* NewPend
 			if (NewPendingOverlaps)
 			{
 				// Note: BeginComponentOverlap() only triggers overlaps where GetGenerateOverlapEvents() is true on both components.
-				for (int32 Idx=0; Idx<NewPendingOverlaps->Num(); ++Idx)
+				const int32 NumNewPendingOverlaps = NewPendingOverlaps->Num();
+				for (int32 Idx=0; Idx < NumNewPendingOverlaps; ++Idx)
 				{
 					BeginComponentOverlap( (*NewPendingOverlaps)[Idx], bDoNotifies );
 				}
 			}
 
-			// now generate full list of new touches, so we can compare to existing list and
-			// determine what changed
-			TInlineOverlapInfoArray NewOverlappingComponents;
+			// now generate full list of new touches, so we can compare to existing list and determine what changed
+			TInlineOverlapInfoArray OverlapMultiResult;
+			TInlineOverlapPointerArray NewOverlappingComponentPtrs;
 
-			// If pending kill, we should not generate any new overlaps
-			if (!IsPendingKill())
+			// If pending kill, we should not generate any new overlaps. Also not if overlaps were just disabled during BeginComponentOverlap.
+			if (!IsPendingKill() && GetGenerateOverlapEvents())
 			{
 				// Might be able to avoid testing for new overlaps at the end location.
-				if (OverlapsAtEndLocation != NULL && bAllowCachedOverlapsCVar && PrevTransform.Equals(GetComponentTransform()))
+				if (OverlapsAtEndLocation != nullptr && bAllowCachedOverlapsCVar && PrevTransform.Equals(GetComponentTransform()))
 				{
 					UE_LOG(LogPrimitiveComponent, VeryVerbose, TEXT("%s->%s Skipping overlap test!"), *GetNameSafe(GetOwner()), *GetName());
-					NewOverlappingComponents = *OverlapsAtEndLocation;
-
-					// BeginComponentOverlap may have disabled what we thought were valid overlaps at the end (collision response or overlap flags could change).
-					// Or we have overlaps from a scoped update that didn't require overlap events, but we need to remove those now.
-					if (NewPendingOverlaps && NewPendingOverlaps->Num() > 0)
+					const bool bCheckForInvalid = (NewPendingOverlaps && NewPendingOverlaps->Num() > 0);
+					if (bCheckForInvalid)
 					{
-						NewOverlappingComponents.RemoveAllSwap(FPredicateFilterCannotOverlap(*this), /*bAllowShrinking*/ false);
+						// BeginComponentOverlap may have disabled what we thought were valid overlaps at the end (collision response or overlap flags could change).
+						GetPointersToArrayDataByPredicate(NewOverlappingComponentPtrs, *OverlapsAtEndLocation, FPredicateFilterCanOverlap(*this));
+					}
+					else
+					{
+						GetPointersToArrayData(NewOverlappingComponentPtrs, *OverlapsAtEndLocation);
 					}
 				}
 				else
 				{
 					UE_LOG(LogPrimitiveComponent, VeryVerbose, TEXT("%s->%s Performing overlaps!"), *GetNameSafe(GetOwner()), *GetName());
-					UWorld* const MyWorld = MyActor->GetWorld();
+					UWorld* const MyWorld = GetWorld();
 					TArray<FOverlapResult> Overlaps;
 					// note this will optionally include overlaps with components in the same actor (depending on bIgnoreChildren). 
 					FComponentQueryParams Params(SCENE_QUERY_STAT(UpdateOverlaps), bIgnoreChildren ? MyActor : nullptr);
@@ -2904,33 +2963,37 @@ bool UPrimitiveComponent::UpdateOverlapsImpl(const TArray<FOverlapInfo>* NewPend
 					InitSweepCollisionParams(Params, ResponseParam);
 					ComponentOverlapMulti(Overlaps, MyWorld, GetComponentLocation(), GetComponentQuat(), GetCollisionObjectType(), Params);
 
-					for( int32 ResultIdx=0; ResultIdx<Overlaps.Num(); ResultIdx++ )
+					for (int32 ResultIdx=0; ResultIdx < Overlaps.Num(); ResultIdx++)
 					{
 						const FOverlapResult& Result = Overlaps[ResultIdx];
 
 						UPrimitiveComponent* const HitComp = Result.Component.Get();
 						if (HitComp && (HitComp != this) && HitComp->GetGenerateOverlapEvents())
 						{
-							if (!ShouldIgnoreOverlapResult(MyWorld, MyActor, *this, Result.GetActor(), *HitComp, /*bCheckOverlapFlags=*/ true))
+							const bool bCheckOverlapFlags = false; // Already checked above
+							if (!ShouldIgnoreOverlapResult(MyWorld, MyActor, *this, Result.GetActor(), *HitComp, bCheckOverlapFlags))
 							{
-								NewOverlappingComponents.Add(FOverlapInfo(HitComp, Result.ItemIndex));		// don't need to add unique unless the overlap check can return dupes
+								OverlapMultiResult.Emplace(HitComp, Result.ItemIndex);		// don't need to add unique unless the overlap check can return dupes
 							}
 						}
 					}
+
+					// Fill pointers to overlap results. We ensure below that OverlapMultiResult stays in scope so these pointers remain valid.
+					GetPointersToArrayData(NewOverlappingComponentPtrs, OverlapMultiResult);
 				}
 			}
 
+			// If we have any overlaps from BeginComponentOverlap() (from now or in the past), see if anything has changed by filtering NewOverlappingComponents
 			if (OverlappingComponents.Num() > 0)
 			{
-				// make a copy of the old that we can manipulate to avoid n^2 searching later
-				TInlineOverlapInfoArray OldOverlappingComponents;
+				TInlineOverlapPointerArray OldOverlappingComponentPtrs;
 				if (bIgnoreChildren)
 				{
-					Algo::CopyIf(OverlappingComponents, OldOverlappingComponents, FPredicateOverlapHasDifferentActor(*MyActor));
+					GetPointersToArrayDataByPredicate(OldOverlappingComponentPtrs, OverlappingComponents, FPredicateOverlapHasDifferentActor(*MyActor));
 				}
 				else
 				{
-					OldOverlappingComponents = OverlappingComponents;
+					GetPointersToArrayData(OldOverlappingComponentPtrs, OverlappingComponents);
 				}
 
 				// Now we want to compare the old and new overlap lists to determine 
@@ -2938,47 +3001,61 @@ bool UPrimitiveComponent::UpdateOverlapsImpl(const TArray<FOverlapInfo>* NewPend
 				// what overlaps are in new and not in old (need begin overlap notifies).
 				// We do this by removing common entries from both lists, since overlapping status has not changed for them.
 				// What is left over will be what has changed.
-				for (int32 CompIdx=0; CompIdx < OldOverlappingComponents.Num() && NewOverlappingComponents.Num() > 0; ++CompIdx)
+				for (int32 CompIdx=0; CompIdx < OldOverlappingComponentPtrs.Num() && NewOverlappingComponentPtrs.Num() > 0; ++CompIdx)
 				{
-					// RemoveSingleSwap is ok, since it is not necessary to maintain order
+					// RemoveAtSwap is ok, since it is not necessary to maintain order
 					const bool bAllowShrinking = false;
 
-					const FOverlapInfo& SearchItem = OldOverlappingComponents[CompIdx];
-					const int32 NewElementIdx = IndexOfOverlapFast(NewOverlappingComponents, SearchItem);
+					const FOverlapInfo* SearchItem = OldOverlappingComponentPtrs[CompIdx];
+					const int32 NewElementIdx = IndexOfOverlapFast(NewOverlappingComponentPtrs, SearchItem);
 					if (NewElementIdx != INDEX_NONE)
 					{
-						NewOverlappingComponents.RemoveAtSwap(NewElementIdx, 1, bAllowShrinking);
-						OldOverlappingComponents.RemoveAtSwap(CompIdx, 1, bAllowShrinking);
+						NewOverlappingComponentPtrs.RemoveAtSwap(NewElementIdx, 1, bAllowShrinking);
+						OldOverlappingComponentPtrs.RemoveAtSwap(CompIdx, 1, bAllowShrinking);
 						--CompIdx;
 					}
 				}
 
-				// OldOverlappingComponents now contains only previous overlaps that are confirmed to no longer be valid.
-				for (auto CompIt = OldOverlappingComponents.CreateIterator(); CompIt; ++CompIt)
+				const int32 NumOldOverlaps = OldOverlappingComponentPtrs.Num();
+				if (NumOldOverlaps > 0)
 				{
-					const FOverlapInfo& OtherOverlap = *CompIt;
-					if (OtherOverlap.OverlapInfo.Component.IsValid())
+					// Now we have to make a copy of the overlaps because we can't keep pointers to them, that list is about to be manipulated in EndComponentOverlap().
+					TInlineOverlapInfoArray OldOverlappingComponents;
+					OldOverlappingComponents.SetNumUninitialized(NumOldOverlaps);
+					for (int32 i=0; i < NumOldOverlaps; i++)
 					{
-						EndComponentOverlap(OtherOverlap, bDoNotifies, false);
+						OldOverlappingComponents[i] = *(OldOverlappingComponentPtrs[i]);
 					}
-					else
+
+					// OldOverlappingComponents now contains only previous overlaps that are confirmed to no longer be valid.
+					for (const FOverlapInfo& OtherOverlap : OldOverlappingComponents)
 					{
-						// Remove stale item. Reclaim memory only if it's getting large, to try to avoid churn but avoid bloating component's memory usage.
-						const bool bAllowShrinking = (OverlappingComponents.Max() >= 24);
-						const int32 StaleElementIndex = IndexOfOverlapFast(OverlappingComponents, OtherOverlap);
-						if (StaleElementIndex != INDEX_NONE)
+						if (OtherOverlap.OverlapInfo.Component.IsValid())
 						{
-							OverlappingComponents.RemoveAtSwap(StaleElementIndex, 1, bAllowShrinking);
+							EndComponentOverlap(OtherOverlap, bDoNotifies, false);
+						}
+						else
+						{
+							// Remove stale item. Reclaim memory only if it's getting large, to try to avoid churn but avoid bloating component's memory usage.
+							const bool bAllowShrinking = (OverlappingComponents.Max() >= 24);
+							const int32 StaleElementIndex = IndexOfOverlapFast(OverlappingComponents, OtherOverlap);
+							if (StaleElementIndex != INDEX_NONE)
+							{
+								OverlappingComponents.RemoveAtSwap(StaleElementIndex, 1, bAllowShrinking);
+							}
 						}
 					}
 				}
 			}
 
+			// Ensure these arrays are still in scope, because we kept pointers to them in NewOverlappingComponentPtrs.
+			static_assert(sizeof(OverlapMultiResult) != 0, "Variable must be in this scope");
+			static_assert(sizeof(OverlapsAtEndLocation) != 0, "Variable must be in this scope");
+
 			// NewOverlappingComponents now contains only new overlaps that didn't exist previously.
-			for (auto CompIt = NewOverlappingComponents.CreateIterator(); CompIt; ++CompIt)
+			for (const FOverlapInfo* NewOverlap : NewOverlappingComponentPtrs)
 			{
-				const FOverlapInfo& OtherOverlap = *CompIt;
-				BeginComponentOverlap(OtherOverlap, bDoNotifies);
+				BeginComponentOverlap(*NewOverlap, bDoNotifies);
 			}
 		}
 	}
