#ue4 - Rework UpdateOverlaps() API to use TArrayView and thereby not require that all input arrays are TArrays with allocations on the heap. Converted MoveComponentImpl() to use an inline array of FOverlapInfos. Prevents many allocations in the common case that objects move within a few overlapping volumes.

diff --git a/Engine/Source/Runtime/Engine/Classes/Components/PrimitiveComponent.h b/Engine/Source/Runtime/Engine/Classes/Components/PrimitiveComponent.h
--- a/Engine/Source/Runtime/Engine/Classes/Components/PrimitiveComponent.h
+++ b/Engine/Source/Runtime/Engine/Classes/Components/PrimitiveComponent.h
@@ -722,10 +722,15 @@ protected:
 
 private:
 	/** Convert a set of overlaps from a sweep to a subset that includes only those at the end location (filling in OverlapsAtEndLocation). */
-	const TArray<FOverlapInfo>* ConvertSweptOverlapsToCurrentOverlaps(TArray<FOverlapInfo>& OverlapsAtEndLocation, const TArray<FOverlapInfo>& SweptOverlaps, int32 SweptOverlapsIndex, const FVector& EndLocation, const FQuat& EndRotationQuat);
+	template<typename AllocatorType>
+	bool ConvertSweptOverlapsToCurrentOverlaps(TArray<FOverlapInfo, AllocatorType>& OutOverlapsAtEndLocation, const TOverlapArrayView& SweptOverlaps, int32 SweptOverlapsIndex, const FVector& EndLocation, const FQuat& EndRotationQuat);
 
 	/** Convert a set of overlaps from a symmetric change in rotation to a subset that includes only those at the end location (filling in OverlapsAtEndLocation). */
-	const TArray<FOverlapInfo>* ConvertRotationOverlapsToCurrentOverlaps(TArray<FOverlapInfo>& OverlapsAtEndLocation, const TArray<FOverlapInfo>& CurrentOverlaps);
+	template<typename AllocatorType>
+	bool ConvertRotationOverlapsToCurrentOverlaps(TArray<FOverlapInfo, AllocatorType>& OutOverlapsAtEndLocation, const TOverlapArrayView& CurrentOverlaps);
+
+	template<typename AllocatorType>
+	bool GetOverlapsWithActor_Template(const AActor* Actor, TArray<FOverlapInfo, AllocatorType>& OutOverlaps) const;
 
 	// FScopedMovementUpdate needs access to the above two functions.
 	friend FScopedMovementUpdate;
@@ -804,7 +809,7 @@ public:
 	 *									Generally this should only be used if this component is the RootComponent of the owning actor and overlaps with other descendant components have been verified.
 	 * @return							True if we can skip calling this in the future (i.e. no useful work is being done.)
 	 */
-	virtual bool UpdateOverlapsImpl(TArray<FOverlapInfo> const* NewPendingOverlaps=nullptr, bool bDoNotifies=true, const TArray<FOverlapInfo>* OverlapsAtEndLocation=nullptr) override;
+	virtual bool UpdateOverlapsImpl(const TOverlapArrayView* NewPendingOverlaps=nullptr, bool bDoNotifies=true, const TOverlapArrayView* OverlapsAtEndLocation=nullptr) override;
 
 	/** Update current physics volume for this component, if bShouldUpdatePhysicsVolume is true. Overridden to use the overlaps to find the physics volume. */
 	virtual void UpdatePhysicsVolume( bool bTriggerNotifiers ) override;
diff --git a/Engine/Source/Runtime/Engine/Classes/Components/SceneComponent.h b/Engine/Source/Runtime/Engine/Classes/Components/SceneComponent.h
--- a/Engine/Source/Runtime/Engine/Classes/Components/SceneComponent.h
+++ b/Engine/Source/Runtime/Engine/Classes/Components/SceneComponent.h
@@ -45,6 +45,9 @@ struct ENGINE_API FOverlapInfo
 // All added members of FOverlapInfo are PODs.
 template<> struct TIsPODType<FOverlapInfo> { enum { Value = TIsPODType<FHitResult>::Value }; };
 
+typedef TArray<FOverlapInfo, TInlineAllocator<3>> TInlineOverlapInfoArray;
+typedef TArrayView<const FOverlapInfo> TOverlapArrayView;
+
 /** Detail mode for scene component rendering, corresponds with the integer value of UWorld::GetDetailMode() */
 UENUM()
 enum EDetailMode
@@ -893,7 +896,7 @@ protected:
 	bool CheckStaticMobilityAndWarn(const FText& ActionText) const;
 
 	/** Internal helper for UpdateOverlaps */
-	virtual bool UpdateOverlapsImpl(TArray<FOverlapInfo> const* PendingOverlaps = nullptr, bool bDoNotifies = true, const TArray<FOverlapInfo>* OverlapsAtEndLocation = nullptr);
+	virtual bool UpdateOverlapsImpl(const TOverlapArrayView* PendingOverlaps = nullptr, bool bDoNotifies = true, const TOverlapArrayView* OverlapsAtEndLocation = nullptr);
 
 private:
 	void PropagateTransformUpdate(bool bTransformChanged, EUpdateTransformFlags UpdateTransformFlags = EUpdateTransformFlags::None, ETeleportType Teleport = ETeleportType::None);
@@ -902,7 +905,7 @@ private:
 public:
 
 	/** Queries world and updates overlap tracking state for this component */
-	bool UpdateOverlaps(TArray<FOverlapInfo> const* PendingOverlaps = nullptr, bool bDoNotifies = true, const TArray<FOverlapInfo>* OverlapsAtEndLocation = nullptr);
+	bool UpdateOverlaps(const TOverlapArrayView* PendingOverlaps = nullptr, bool bDoNotifies = true, const TOverlapArrayView* OverlapsAtEndLocation = nullptr);
 
 	/**
 	 * Tries to move the component by a movement vector (Delta) and sets rotation to NewRotation.
@@ -1437,7 +1440,8 @@ class ENGINE_API FScopedMovementUpdate : private FNoncopyable
 {
 public:
 	
-	typedef TArray<struct FHitResult, TInlineAllocator<2>> TBlockingHitArray;
+	typedef TArray<struct FHitResult, TInlineAllocator<2>> TScopedBlockingHitArray;
+	typedef TArray<struct FOverlapInfo, TInlineAllocator<3>> TScopedOverlapInfoArray;
 
 	FScopedMovementUpdate( USceneComponent* Component, EScopedUpdate::Type ScopeBehavior = EScopedUpdate::DeferredUpdates, bool bRequireOverlapsEventFlagToQueueOverlaps = true );
 	~FScopedMovementUpdate();
@@ -1481,16 +1485,16 @@ public:
 	bool RequiresOverlapsEventFlag() const;
 
 	/** Returns the pending overlaps within this scope. */
-	const TArray<FOverlapInfo>& GetPendingOverlaps() const;
+	const TScopedOverlapInfoArray& GetPendingOverlaps() const;
 
 	/** Returns the list of pending blocking hits, which will be used for notifications once the move is committed. */
-	const TBlockingHitArray& GetPendingBlockingHits() const;
+	const TScopedBlockingHitArray& GetPendingBlockingHits() const;
 
 	//--------------------------------------------------------------------------------------------------------//
 	// These methods are intended only to be used by SceneComponent and derived classes.
 
 	/** Add overlaps to the queued overlaps array. This is intended for use only by SceneComponent and its derived classes whenever movement is performed. */
-	void AppendOverlapsAfterMove(const TArray<FOverlapInfo>& NewPendingOverlaps, bool bSweep, bool bIncludesOverlapsAtEnd);
+	void AppendOverlapsAfterMove(const TOverlapArrayView& NewPendingOverlaps, bool bSweep, bool bIncludesOverlapsAtEnd);
 
 	/** Keep current pending overlaps after a move but make note that there was movement (just a symmetric rotation). */
 	void KeepCurrentOverlapsAfterRotation(bool bSweep);
@@ -1509,7 +1513,8 @@ public:
 
 protected:
 	/** Fills in the list of overlaps at the end location (in EndOverlaps). Returns pointer to the list, or null if it can't be computed. */
-	const TArray<FOverlapInfo>* GetOverlapsAtEnd(class UPrimitiveComponent& PrimComponent, TArray<FOverlapInfo>& EndOverlaps, bool bTransformChanged) const;
+	template<typename AllocatorType>
+	TOptional<TOverlapArrayView> GetOverlapsAtEnd(class UPrimitiveComponent& PrimComponent, TArray<FOverlapInfo, AllocatorType>& OutEndOverlaps, bool bTransformChanged) const;
 
 	bool SetWorldLocationAndRotation(FVector NewLocation, const FQuat& NewQuat, bool bNoPhysics = false, ETeleportType Teleport = ETeleportType::None);
 
@@ -1536,9 +1541,9 @@ protected:
 	FRotator InitialRelativeRotation;
 	FVector InitialRelativeScale;
 
-	int32 FinalOverlapCandidatesIndex;		// If not INDEX_NONE, overlaps at this index and beyond in PendingOverlaps are at the final destination
-	TArray<FOverlapInfo> PendingOverlaps;	// All overlaps encountered during the scope of moves.
-	TBlockingHitArray BlockingHits;			// All blocking hits encountered during the scope of moves.
+	int32 FinalOverlapCandidatesIndex;			// If not INDEX_NONE, overlaps at this index and beyond in PendingOverlaps are at the final destination
+	TScopedOverlapInfoArray PendingOverlaps;	// All overlaps encountered during the scope of moves.
+	TScopedBlockingHitArray BlockingHits;		// All blocking hits encountered during the scope of moves.
 
 	uint8 bDeferUpdates:1;
 	uint8 bHasMoved:1;
@@ -1575,12 +1580,12 @@ FORCEINLINE bool FScopedMovementUpdate::RequiresOverlapsEventFlag() const
 	return bRequireOverlapsEventFlag;
 }
 
-FORCEINLINE const TArray<struct FOverlapInfo>& FScopedMovementUpdate::GetPendingOverlaps() const
+FORCEINLINE const FScopedMovementUpdate::TScopedOverlapInfoArray& FScopedMovementUpdate::GetPendingOverlaps() const
 {
 	return PendingOverlaps;
 }
 
-FORCEINLINE const FScopedMovementUpdate::TBlockingHitArray& FScopedMovementUpdate::GetPendingBlockingHits() const
+FORCEINLINE const FScopedMovementUpdate::TScopedBlockingHitArray& FScopedMovementUpdate::GetPendingBlockingHits() const
 {
 	return BlockingHits;
 }
@@ -1645,7 +1650,7 @@ FORCEINLINE_DEBUGGABLE void USceneComponent::BeginScopedMovementUpdate(class FSc
 	ScopedMovementStack.Push(&ScopedUpdate);
 }
 
-FORCEINLINE_DEBUGGABLE bool USceneComponent::UpdateOverlaps(TArray<FOverlapInfo> const* PendingOverlaps /* = nullptr */, bool bDoNotifies /* = true */, const TArray<FOverlapInfo>* OverlapsAtEndLocation /* = nullptr */)
+FORCEINLINE_DEBUGGABLE bool USceneComponent::UpdateOverlaps(const TOverlapArrayView* PendingOverlaps /* = nullptr */, bool bDoNotifies /* = true */, const TOverlapArrayView* OverlapsAtEndLocation /* = nullptr */)
 {
 	if (IsDeferringMovementUpdates())
 	{
diff --git a/Engine/Source/Runtime/Engine/Classes/Components/SkeletalMeshComponent.h b/Engine/Source/Runtime/Engine/Classes/Components/SkeletalMeshComponent.h
--- a/Engine/Source/Runtime/Engine/Classes/Components/SkeletalMeshComponent.h
+++ b/Engine/Source/Runtime/Engine/Classes/Components/SkeletalMeshComponent.h
@@ -1350,7 +1350,7 @@ public:
 	virtual FBoxSphereBounds CalcBounds(const FTransform& LocalToWorld) const override;
 	virtual bool IsAnySimulatingPhysics() const override;
 	virtual void OnUpdateTransform(EUpdateTransformFlags UpdateTransformFlags, ETeleportType Teleport = ETeleportType::None) override;
-	virtual bool UpdateOverlapsImpl(TArray<FOverlapInfo> const* PendingOverlaps=NULL, bool bDoNotifies=true, const TArray<FOverlapInfo>* OverlapsAtEndLocation=NULL) override;
+	virtual bool UpdateOverlapsImpl(const TOverlapArrayView* PendingOverlaps=NULL, bool bDoNotifies=true, const TOverlapArrayView* OverlapsAtEndLocation=NULL) override;
 	//~ End USceneComponent Interface.
 
 	//~ Begin UPrimitiveComponent Interface.
diff --git a/Engine/Source/Runtime/Engine/Classes/Components/SkinnedMeshComponent.h b/Engine/Source/Runtime/Engine/Classes/Components/SkinnedMeshComponent.h
--- a/Engine/Source/Runtime/Engine/Classes/Components/SkinnedMeshComponent.h
+++ b/Engine/Source/Runtime/Engine/Classes/Components/SkinnedMeshComponent.h
@@ -784,7 +784,7 @@ public:
 	virtual bool DoesSocketExist(FName InSocketName) const override;
 	virtual bool HasAnySockets() const override;
 	virtual void QuerySupportedSockets(TArray<FComponentSocketDescription>& OutSockets) const override;
-	virtual bool UpdateOverlapsImpl(TArray<FOverlapInfo> const* PendingOverlaps=NULL, bool bDoNotifies=true, const TArray<FOverlapInfo>* OverlapsAtEndLocation=NULL) override;
+	virtual bool UpdateOverlapsImpl(const TOverlapArrayView* PendingOverlaps=NULL, bool bDoNotifies=true, const TOverlapArrayView* OverlapsAtEndLocation=NULL) override;
 	//~ End USceneComponent Interface
 
 	//~ Begin UPrimitiveComponent Interface
diff --git a/Engine/Source/Runtime/Engine/Private/Components/PrimitiveComponent.cpp b/Engine/Source/Runtime/Engine/Private/Components/PrimitiveComponent.cpp
--- a/Engine/Source/Runtime/Engine/Private/Components/PrimitiveComponent.cpp
+++ b/Engine/Source/Runtime/Engine/Private/Components/PrimitiveComponent.cpp
@@ -52,7 +52,6 @@ namespace PrimitiveComponentStatics
 	static const FText MobilityWarnText = LOCTEXT("InvalidMove", "move");
 }
 
-typedef TArray<FOverlapInfo, TInlineAllocator<3>> TInlineOverlapInfoArray;
 typedef TArray<const FOverlapInfo*, TInlineAllocator<8>> TInlineOverlapPointerArray;
 
 DEFINE_LOG_CATEGORY_STATIC(LogPrimitiveComponent, Log, All);
@@ -255,6 +254,17 @@ FORCEINLINE_DEBUGGABLE static void GetPointersToArrayData(TArray<const ElementTy
 	}
 }
 
+template <class ElementType, class AllocatorType1>
+FORCEINLINE_DEBUGGABLE static void GetPointersToArrayData(TArray<const ElementType*, AllocatorType1>& Pointers, const TArrayView<const ElementType>& DataArray)
+{
+	const int32 NumItems = DataArray.Num();
+	Pointers.SetNumUninitialized(NumItems);
+	for (int32 i=0; i < NumItems; i++)
+	{
+		Pointers[i] = &(DataArray[i]);
+	}
+}
+
 // Helper to initialize an array to point to data members in another array which satisfy a predicate.
 template <class ElementType, class AllocatorType1, class AllocatorType2, typename PredicateT>
 FORCEINLINE_DEBUGGABLE static void GetPointersToArrayDataByPredicate(TArray<const ElementType*, AllocatorType1>& Pointers, const TArray<ElementType, AllocatorType2>& DataArray, PredicateT Predicate)
@@ -269,6 +279,19 @@ FORCEINLINE_DEBUGGABLE static void GetPointersToArrayDataByPredicate(TArray<cons
 	}
 }
 
+template <class ElementType, class AllocatorType1, typename PredicateT>
+FORCEINLINE_DEBUGGABLE static void GetPointersToArrayDataByPredicate(TArray<const ElementType*, AllocatorType1>& Pointers, const TArrayView<const ElementType>& DataArray, PredicateT Predicate)
+{
+	Pointers.Reserve(Pointers.Num() + DataArray.Num());
+	for (const ElementType& Item : DataArray)
+	{
+		if (Invoke(Predicate, Item))
+		{
+			Pointers.Add(&Item);
+		}
+	}
+}
+
 ///////////////////////////////////////////////////////////////////////////////
 // PRIMITIVE COMPONENT
 ///////////////////////////////////////////////////////////////////////////////
@@ -1919,7 +1942,7 @@ bool UPrimitiveComponent::MoveComponentImpl( const FVector& Delta, const FQuat&
 	bool bMoved = false;
 	bool bIncludesOverlapsAtEnd = false;
 	bool bRotationOnly = false;
-	TArray<FOverlapInfo> PendingOverlaps;
+	TInlineOverlapInfoArray PendingOverlaps;
 	AActor* const Actor = GetOwner();
 
 	if ( !bSweep )
@@ -2125,21 +2148,24 @@ bool UPrimitiveComponent::MoveComponentImpl( const FVector& Delta, const FQuat&
 		{
 			if (bIncludesOverlapsAtEnd)
 			{
-				TArray<FOverlapInfo> OverlapsAtEndLocation;
-				const TArray<FOverlapInfo>* OverlapsAtEndLocationPtr = nullptr; // When non-null, used as optimization to avoid work in UpdateOverlaps.
+				TInlineOverlapInfoArray OverlapsAtEndLocation;
+				bool bHasEndOverlaps = false;
 				if (bRotationOnly)
 				{
-					OverlapsAtEndLocationPtr = ConvertRotationOverlapsToCurrentOverlaps(OverlapsAtEndLocation, GetOverlapInfos());
+					bHasEndOverlaps = ConvertRotationOverlapsToCurrentOverlaps(OverlapsAtEndLocation, OverlappingComponents);
 				}
 				else
 				{
-					OverlapsAtEndLocationPtr = ConvertSweptOverlapsToCurrentOverlaps(OverlapsAtEndLocation, PendingOverlaps, 0, GetComponentLocation(), GetComponentQuat());
+					bHasEndOverlaps = ConvertSweptOverlapsToCurrentOverlaps(OverlapsAtEndLocation, PendingOverlaps, 0, GetComponentLocation(), GetComponentQuat());
 				}
-				UpdateOverlaps(&PendingOverlaps, true, OverlapsAtEndLocationPtr);
+				TOverlapArrayView PendingOverlapsView(PendingOverlaps);
+				TOverlapArrayView OverlapsAtEndView(OverlapsAtEndLocation);
+				UpdateOverlaps(&PendingOverlapsView, true, bHasEndOverlaps ? &OverlapsAtEndView : nullptr);
 			}
 			else
 			{
-				UpdateOverlaps(&PendingOverlaps, true, nullptr);
+				TOverlapArrayView PendingOverlapsView(PendingOverlaps);
+				UpdateOverlaps(&PendingOverlapsView, true, nullptr);
 			}
 		}
 	}
@@ -2369,7 +2395,8 @@ bool UPrimitiveComponent::IsOverlappingActor(const AActor* Other) const
 	return false;
 }
 
-bool UPrimitiveComponent::GetOverlapsWithActor(const AActor* Actor, TArray<FOverlapInfo>& OutOverlaps) const
+template<typename AllocatorType>
+bool UPrimitiveComponent::GetOverlapsWithActor_Template(const AActor* Actor, TArray<FOverlapInfo, AllocatorType>& OutOverlaps) const
 {
 	const int32 InitialCount = OutOverlaps.Num();
 	if (Actor)
@@ -2387,6 +2414,11 @@ bool UPrimitiveComponent::GetOverlapsWithActor(const AActor* Actor, TArray<FOver
 	return InitialCount != OutOverlaps.Num();
 }
 
+bool UPrimitiveComponent::GetOverlapsWithActor(const AActor* Actor, TArray<FOverlapInfo>& OutOverlaps) const
+{
+	return GetOverlapsWithActor_Template(Actor, OutOverlaps);
+}
+
 #if WITH_EDITOR
 bool UPrimitiveComponent::ComponentIsTouchingSelectionBox(const FBox& InSelBBox, const FEngineShowFlags& ShowFlags, const bool bConsiderOnlyBSP, const bool bMustEncompassEntireComponent) const
 {
@@ -2675,13 +2707,14 @@ void UPrimitiveComponent::GetOverlappingComponents(TSet<UPrimitiveComponent*>& O
 	}
 }
 
-const TArray<FOverlapInfo>* UPrimitiveComponent::ConvertSweptOverlapsToCurrentOverlaps(
-	TArray<FOverlapInfo>& OverlapsAtEndLocation, const TArray<FOverlapInfo>& SweptOverlaps, int32 SweptOverlapsIndex,
+template<typename AllocatorType>
+bool UPrimitiveComponent::ConvertSweptOverlapsToCurrentOverlaps(
+	TArray<FOverlapInfo, AllocatorType>& OverlapsAtEndLocation, const TOverlapArrayView& SweptOverlaps, int32 SweptOverlapsIndex,
 	const FVector& EndLocation, const FQuat& EndRotationQuat)
 {
 	checkSlow(SweptOverlapsIndex >= 0);
 
-	const TArray<FOverlapInfo>* Result = nullptr;
+	bool bResult = false;
 	const bool bForceGatherOverlaps = !ShouldCheckOverlapFlagToQueueOverlaps(*this);
 	if ((GetGenerateOverlapEvents() || bForceGatherOverlaps) && bAllowCachedOverlapsCVar)
 	{
@@ -2706,12 +2739,12 @@ const TArray<FOverlapInfo>* UPrimitiveComponent::ConvertSweptOverlapsToCurrentOv
 						if (OtherPrimitive->bMultiBodyOverlap)
 						{
 							// Not handled yet. We could do it by checking every body explicitly and track each body index in the overlap test, but this seems like a rare need.
-							return nullptr;
+							return false;
 						}
 						else if (Cast<USkeletalMeshComponent>(OtherPrimitive) || Cast<USkeletalMeshComponent>(this))
 						{
 							// SkeletalMeshComponent does not support this operation, and would return false in the test when an actual query could return true.
-							return nullptr;
+							return false;
 						}
 						else if (OtherPrimitive->ComponentOverlapComponent(this, EndLocation, EndRotationQuat, UnusedQueryParams))
 						{
@@ -2725,27 +2758,27 @@ const TArray<FOverlapInfo>* UPrimitiveComponent::ConvertSweptOverlapsToCurrentOv
 				checkfSlow(OverlapsAtEndLocation.FindByPredicate(FPredicateOverlapHasSameActor(*Actor)) == nullptr,
 					TEXT("Child overlaps should not be included in the SweptOverlaps() array in UPrimitiveComponent::ConvertSweptOverlapsToCurrentOverlaps()."));
 
-				Result = &OverlapsAtEndLocation;
+				bResult = true;
 			}
 			else
 			{
 				if (SweptOverlaps.Num() == 0 && AreAllCollideableDescendantsRelative())
 				{
 					// Add overlaps with components in this actor.
-					GetOverlapsWithActor(Actor, OverlapsAtEndLocation);
-					Result = &OverlapsAtEndLocation;
+					GetOverlapsWithActor_Template(Actor, OverlapsAtEndLocation);
+					bResult = true;
 				}
 			}
 		}
 	}
 
-	return Result;
+	return bResult;
 }
 
-
-const TArray<FOverlapInfo>* UPrimitiveComponent::ConvertRotationOverlapsToCurrentOverlaps(TArray<FOverlapInfo>& OverlapsAtEndLocation, const TArray<FOverlapInfo>& CurrentOverlaps)
+template<typename AllocatorType>
+bool UPrimitiveComponent::ConvertRotationOverlapsToCurrentOverlaps(TArray<FOverlapInfo, AllocatorType>& OutOverlapsAtEndLocation, const TOverlapArrayView& CurrentOverlaps)
 {
-	const TArray<FOverlapInfo>* Result = nullptr;
+	bool bResult = false;
 	const bool bForceGatherOverlaps = !ShouldCheckOverlapFlagToQueueOverlaps(*this);
 	if ((GetGenerateOverlapEvents() || bForceGatherOverlaps) && bAllowCachedOverlapsCVar)
 	{
@@ -2755,14 +2788,14 @@ const TArray<FOverlapInfo>* UPrimitiveComponent::ConvertRotationOverlapsToCurren
 			if (bEnableFastOverlapCheck)
 			{
 				// Add all current overlaps that are not children. Children test for their own overlaps after we update our own, and we ignore children in our own update.
-				OverlapsAtEndLocation.Reserve(OverlapsAtEndLocation.Num() + CurrentOverlaps.Num());
-				Algo::CopyIf(CurrentOverlaps, OverlapsAtEndLocation, FPredicateOverlapHasDifferentActor(*Actor));
-				Result = &OverlapsAtEndLocation;
+				OutOverlapsAtEndLocation.Reserve(OutOverlapsAtEndLocation.Num() + CurrentOverlaps.Num());
+				Algo::CopyIf(CurrentOverlaps, OutOverlapsAtEndLocation, FPredicateOverlapHasDifferentActor(*Actor));
+				bResult = true;
 			}
 		}
 	}
 
-	return Result;
+	return bResult;
 }
 
 
@@ -2898,7 +2931,7 @@ TArray<UPrimitiveComponent*> UPrimitiveComponent::CopyArrayOfMoveIgnoreComponent
 	return MoveIgnoreComponents;
 }
 
-bool UPrimitiveComponent::UpdateOverlapsImpl(const TArray<FOverlapInfo>* NewPendingOverlaps, bool bDoNotifies, const TArray<FOverlapInfo>* OverlapsAtEndLocation)
+bool UPrimitiveComponent::UpdateOverlapsImpl(const TOverlapArrayView* NewPendingOverlaps, bool bDoNotifies, const TOverlapArrayView* OverlapsAtEndLocation)
 {
 	SCOPE_CYCLE_COUNTER(STAT_UpdateOverlaps); 
 	SCOPE_CYCLE_UOBJECT(ComponentScope, this);
diff --git a/Engine/Source/Runtime/Engine/Private/Components/SceneComponent.cpp b/Engine/Source/Runtime/Engine/Private/Components/SceneComponent.cpp
--- a/Engine/Source/Runtime/Engine/Private/Components/SceneComponent.cpp
+++ b/Engine/Source/Runtime/Engine/Private/Components/SceneComponent.cpp
@@ -835,9 +835,10 @@ void USceneComponent::EndScopedMovementUpdate(class FScopedMovementUpdate& Compl
 				if (PrimitiveThis)
 				{
 					// NOTE: UpdateOverlaps filters events to only consider overlaps where bGenerateOverlapEvents is true for both components, so it's ok if we queued up other overlaps.
-					TArray<FOverlapInfo> EndOverlaps;
-					const TArray<FOverlapInfo>* EndOverlapsPtr = CurrentScopedUpdate->GetOverlapsAtEnd(*PrimitiveThis, EndOverlaps, bTransformChanged);
-					UpdateOverlaps(&CurrentScopedUpdate->GetPendingOverlaps(), true, EndOverlapsPtr);
+					TInlineOverlapInfoArray EndOverlaps;
+					const TOverlapArrayView PendingOverlaps(CurrentScopedUpdate->GetPendingOverlaps());
+					const TOptional<TOverlapArrayView> EndOverlapsOptional = CurrentScopedUpdate->GetOverlapsAtEnd(*PrimitiveThis, EndOverlaps, bTransformChanged);
+					UpdateOverlaps(&PendingOverlaps, true, EndOverlapsOptional.IsSet() ? &(EndOverlapsOptional.GetValue()) : nullptr);
 				}
 				else
 				{
@@ -2759,7 +2760,7 @@ bool USceneComponent::InternalSetWorldLocationAndRotation(FVector NewLocation, c
 	return false;
 }
 
-bool USceneComponent::UpdateOverlapsImpl(TArray<FOverlapInfo> const* PendingOverlaps, bool bDoNotifies, const TArray<FOverlapInfo>* OverlapsAtEndLocation)
+bool USceneComponent::UpdateOverlapsImpl(const TOverlapArrayView* PendingOverlaps, bool bDoNotifies, const TOverlapArrayView* OverlapsAtEndLocation)
 {
 	SCOPE_CYCLE_COUNTER(STAT_UpdateOverlaps); 
 
@@ -3463,7 +3464,7 @@ void FScopedMovementUpdate::RevertMove()
 	TeleportType = ETeleportType::None;
 }
 
-void FScopedMovementUpdate::AppendOverlapsAfterMove(const TArray<FOverlapInfo>& NewPendingOverlaps, bool bSweep, bool bIncludesOverlapsAtEnd)
+void FScopedMovementUpdate::AppendOverlapsAfterMove(const TOverlapArrayView& NewPendingOverlaps, bool bSweep, bool bIncludesOverlapsAtEnd)
 {
 	bHasMoved = true;
 	const bool bWasForcing = (CurrentOverlapState == EOverlapState::eForceUpdate);
@@ -3474,7 +3475,7 @@ void FScopedMovementUpdate::AppendOverlapsAfterMove(const TArray<FOverlapInfo>&
 		if (NewPendingOverlaps.Num())
 		{
 			FinalOverlapCandidatesIndex = PendingOverlaps.Num();
-			PendingOverlaps.Append(NewPendingOverlaps);
+			PendingOverlaps.Append(NewPendingOverlaps.GetData(), NewPendingOverlaps.Num());
 		}
 		else
 		{
@@ -3487,7 +3488,7 @@ void FScopedMovementUpdate::AppendOverlapsAfterMove(const TArray<FOverlapInfo>&
 		// We don't know about the final overlaps in the case of a teleport.
 		CurrentOverlapState = EOverlapState::eUnknown;
 		FinalOverlapCandidatesIndex = INDEX_NONE;
-		PendingOverlaps.Append(NewPendingOverlaps);
+		PendingOverlaps.Append(NewPendingOverlaps.GetData(), NewPendingOverlaps.Num());
 	}
 
 	if (bWasForcing)
@@ -3546,9 +3547,10 @@ void FScopedMovementUpdate::OnInnerScopeComplete(const FScopedMovementUpdate& In
 	}	
 }
 
-const TArray<FOverlapInfo>* FScopedMovementUpdate::GetOverlapsAtEnd(class UPrimitiveComponent& PrimComponent, TArray<FOverlapInfo>& EndOverlaps, bool bTransformChanged) const
+template<typename AllocatorType>
+TOptional<TOverlapArrayView> FScopedMovementUpdate::GetOverlapsAtEnd(class UPrimitiveComponent& PrimComponent, TArray<FOverlapInfo, AllocatorType>& OutEndOverlaps, bool bTransformChanged) const
 {
-	const TArray<FOverlapInfo>* EndOverlapsPtr = nullptr;
+	TOptional<TOverlapArrayView> Result;
 	switch (CurrentOverlapState)
 	{
 		case FScopedMovementUpdate::EOverlapState::eUseParent:
@@ -3556,19 +3558,21 @@ const TArray<FOverlapInfo>* FScopedMovementUpdate::GetOverlapsAtEnd(class UPrimi
 			// Only rotation could have possibly changed
 			if (bTransformChanged && PrimComponent.AreSymmetricRotations(InitialTransform.GetRotation(), PrimComponent.GetComponentQuat(), PrimComponent.GetComponentScale()))
 			{
-				EndOverlapsPtr = PrimComponent.ConvertRotationOverlapsToCurrentOverlaps(EndOverlaps, PrimComponent.GetOverlapInfos());
+				if (PrimComponent.ConvertRotationOverlapsToCurrentOverlaps(OutEndOverlaps, PrimComponent.GetOverlapInfos()))
+				{
+					Result = OutEndOverlaps;
+				}
 			}
 			else
 			{
 				// Use current overlaps (unchanged)
-				EndOverlapsPtr = &PrimComponent.GetOverlapInfos();
+				Result = PrimComponent.GetOverlapInfos();
 			}
 			break;
 		}
 		case FScopedMovementUpdate::EOverlapState::eUnknown:
 		case FScopedMovementUpdate::EOverlapState::eForceUpdate:
 		{
-			EndOverlapsPtr = nullptr;
 			break;
 		}
 		case FScopedMovementUpdate::EOverlapState::eIncludesOverlaps:
@@ -3576,7 +3580,7 @@ const TArray<FOverlapInfo>* FScopedMovementUpdate::GetOverlapsAtEnd(class UPrimi
 			if (FinalOverlapCandidatesIndex == INDEX_NONE)
 			{
 				// Overlapping nothing
-				EndOverlapsPtr = &EndOverlaps;
+				Result = OutEndOverlaps;
 			}
 			else
 			{
@@ -3584,9 +3588,14 @@ const TArray<FOverlapInfo>* FScopedMovementUpdate::GetOverlapsAtEnd(class UPrimi
 				const bool bMatchingScale = FTransform::AreScale3DsEqual(InitialTransform, PrimComponent.GetComponentTransform());
 				if (bMatchingScale)
 				{
-					EndOverlapsPtr = PrimComponent.ConvertSweptOverlapsToCurrentOverlaps(
-						EndOverlaps, GetPendingOverlaps(), FinalOverlapCandidatesIndex,
+					const bool bHasEndOverlaps = PrimComponent.ConvertSweptOverlapsToCurrentOverlaps(
+						OutEndOverlaps, PendingOverlaps, FinalOverlapCandidatesIndex,
 						PrimComponent.GetComponentLocation(), PrimComponent.GetComponentQuat());
+					
+					if (bHasEndOverlaps)
+					{
+						Result = OutEndOverlaps;
+					}
 				}
 			}
 			break;
@@ -3598,7 +3607,7 @@ const TArray<FOverlapInfo>* FScopedMovementUpdate::GetOverlapsAtEnd(class UPrimi
 		}
 	}
 
-	return EndOverlapsPtr;
+	return Result;
 }
 
 
diff --git a/Engine/Source/Runtime/Engine/Private/Components/SkinnedMeshComponent.cpp b/Engine/Source/Runtime/Engine/Private/Components/SkinnedMeshComponent.cpp
--- a/Engine/Source/Runtime/Engine/Private/Components/SkinnedMeshComponent.cpp
+++ b/Engine/Source/Runtime/Engine/Private/Components/SkinnedMeshComponent.cpp
@@ -2137,7 +2137,7 @@ void USkinnedMeshComponent::QuerySupportedSockets(TArray<FComponentSocketDescrip
 	}
 }
 
-bool USkinnedMeshComponent::UpdateOverlapsImpl(TArray<FOverlapInfo> const* PendingOverlaps, bool bDoNotifies, const TArray<FOverlapInfo>* OverlapsAtEndLocation)
+bool USkinnedMeshComponent::UpdateOverlapsImpl(const TOverlapArrayView* PendingOverlaps, bool bDoNotifies, const TOverlapArrayView* OverlapsAtEndLocation)
 {
 	// we don't support overlap test on destructible or physics asset
 	// so use SceneComponent::UpdateOverlaps to handle children
diff --git a/Engine/Source/Runtime/Engine/Private/SkeletalMeshComponentPhysics.cpp b/Engine/Source/Runtime/Engine/Private/SkeletalMeshComponentPhysics.cpp
--- a/Engine/Source/Runtime/Engine/Private/SkeletalMeshComponentPhysics.cpp
+++ b/Engine/Source/Runtime/Engine/Private/SkeletalMeshComponentPhysics.cpp
@@ -1317,7 +1317,7 @@ void USkeletalMeshComponent::OnUpdateTransform(EUpdateTransformFlags UpdateTrans
 	}
 }
 
-bool USkeletalMeshComponent::UpdateOverlapsImpl(TArray<FOverlapInfo> const* PendingOverlaps, bool bDoNotifies, const TArray<FOverlapInfo>* OverlapsAtEndLocation)
+bool USkeletalMeshComponent::UpdateOverlapsImpl(const TOverlapArrayView* PendingOverlaps, bool bDoNotifies, const TOverlapArrayView* OverlapsAtEndLocation)
 {
 	// Parent class (USkinnedMeshComponent) routes only to children, but we really do want to test our own bodies for overlaps.
 	return UPrimitiveComponent::UpdateOverlapsImpl(PendingOverlaps, bDoNotifies, OverlapsAtEndLocation);
