PR #5556: Fix scope name of GetAllActorsWithInterface (Contributed by [REDACTED_PERSON])

diff --git a/Engine/Source/Runtime/Engine/Private/GameplayStatics.cpp b/Engine/Source/Runtime/Engine/Private/GameplayStatics.cpp
--- a/Engine/Source/Runtime/Engine/Private/GameplayStatics.cpp
+++ b/Engine/Source/Runtime/Engine/Private/GameplayStatics.cpp
@@ -852,7 +852,7 @@ void UGameplayStatics::GetAllActorsOfClass(const UObject* WorldContextObject, TS
 
 void UGameplayStatics::GetAllActorsWithInterface(const UObject* WorldContextObject, TSubclassOf<UInterface> Interface, TArray<AActor*>& OutActors)
 {
-	QUICK_SCOPE_CYCLE_COUNTER(UGameplayStatics_GetAllActorsWithTag);
+	QUICK_SCOPE_CYCLE_COUNTER(UGameplayStatics_GetAllActorsWithInterface);
 	OutActors.Reset();
 
 	// We do nothing if no interface provided, rather than giving ALL actors!
