In the Editor, the parameters to atan2 are labeled A and B, but the internal function it calls has parameters labeled Y and X (in that order). In other words...

A -> Y
B -> X

Usually atan2 is presented as it is in the code, making the blueprint version wrong. It should be changed to (Y,X) or at the very least (B,A).

diff --git a/Engine/Config/BaseEngine.ini b/Engine/Config/BaseEngine.ini
--- a/Engine/Config/BaseEngine.ini
+++ b/Engine/Config/BaseEngine.ini
@@ -662,6 +662,10 @@ bUseStreamingPause=false
 +ClassRedirects=(OldName="EditorAutomationObject",NewName="EditorUtilityObject")
 
 +PropertyRedirects=(OldName="StructVariableDescription.bDontEditoOnInstance",NewName="bDontEditOnInstance")
++PropertyRedirects=(OldName="KismetMathLibrary.DegAtan2.A",NewName="Y")
++PropertyRedirects=(OldName="KismetMathLibrary.DegAtan2.B",NewName="X")
++PropertyRedirects=(OldName="KismetMathLibrary.Atan2.A",NewName="Y")
++PropertyRedirects=(OldName="KismetMathLibrary.Atan2.B",NewName="X")
 
 [CoreUObject.Metadata]
 ; Note: UnrealHeaderTool should be rerun after making changes to MetadataRedirects to catch any keys specified in class headers
diff --git a/Engine/Source/Runtime/Engine/Classes/Kismet/KismetMathLibrary.h b/Engine/Source/Runtime/Engine/Classes/Kismet/KismetMathLibrary.h
--- a/Engine/Source/Runtime/Engine/Classes/Kismet/KismetMathLibrary.h
+++ b/Engine/Source/Runtime/Engine/Classes/Kismet/KismetMathLibrary.h
@@ -590,7 +590,7 @@ class ENGINE_API UKismetMathLibrary : public UBlueprintFunctionLibrary
 
 	/** Returns the inverse tan (atan2) of A/B (result is in Radians)*/
 	UFUNCTION(BlueprintPure, meta=(DisplayName = "Atan2 (Radians)"), Category="Math|Trig")
-	static float Atan2(float A, float B);
+	static float Atan2(float Y, float X);
 
 	/** Returns exponential(e) to the power A (e^A)*/
 	UFUNCTION(BlueprintPure, Category="Math|Float", meta=(CompactNodeTitle = "e"))
@@ -662,7 +662,7 @@ class ENGINE_API UKismetMathLibrary : public UBlueprintFunctionLibrary
 
 	/** Returns the inverse tan (atan2) of A/B (result is in Degrees)*/
 	UFUNCTION(BlueprintPure, meta=(DisplayName = "Atan2 (Degrees)"), Category="Math|Trig")
-	static float DegAtan2(float A, float B);
+	static float DegAtan2(float Y, float X);
 
 	/** 
 	 * Clamps an arbitrary angle to be between the given angles.  Will clamp to nearest boundary.
diff --git a/Engine/Source/Runtime/Engine/Classes/Kismet/KismetMathLibrary.inl b/Engine/Source/Runtime/Engine/Classes/Kismet/KismetMathLibrary.inl
--- a/Engine/Source/Runtime/Engine/Classes/Kismet/KismetMathLibrary.inl
+++ b/Engine/Source/Runtime/Engine/Classes/Kismet/KismetMathLibrary.inl
@@ -590,9 +590,9 @@ float UKismetMathLibrary::Atan(float A)
 }
 
 KISMET_MATH_FORCEINLINE
-float UKismetMathLibrary::Atan2(float A, float B)
+float UKismetMathLibrary::Atan2(float Y, float X)
 {
-	return FMath::Atan2(A, B);
+	return FMath::Atan2(Y, X);
 }	
 
 KISMET_MATH_FORCEINLINE
@@ -632,9 +632,9 @@ float UKismetMathLibrary::DegAtan(float A)
 }
 
 KISMET_MATH_FORCEINLINE
-float UKismetMathLibrary::DegAtan2(float A, float B)
+float UKismetMathLibrary::DegAtan2(float Y, float X)
 {
-	return (180.f)/PI * FMath::Atan2(A, B);
+	return (180.f)/PI * FMath::Atan2(Y, X);
 }
 
 KISMET_MATH_FORCEINLINE
