Fix ensure that the streaming level is removed from its original world when adding a sublevel as part of CommitMapChange

diff --git a/Engine/Source/Runtime/Engine/Private/UnrealEngine.cpp b/Engine/Source/Runtime/Engine/Private/UnrealEngine.cpp
--- a/Engine/Source/Runtime/Engine/Private/UnrealEngine.cpp
+++ b/Engine/Source/Runtime/Engine/Private/UnrealEngine.cpp
@@ -13569,14 +13569,18 @@ bool UEngine::CommitMapChange( FWorldContext &Context )
 
 		// Rename the newly loaded streaming levels so that their outer is correctly set to the main context's world,
 		// rather than the fake world.
-		for (ULevelStreaming* const FakeWorldStreamingLevel : FakeWorld->GetStreamingLevels())
+		TArray<ULevelStreaming*> StreamingLevelsToMove;
+		StreamingLevelsToMove.SetNumUninitialized(FakeWorld->GetStreamingLevels().Num());
+		for (int32 Index = FakeWorld->GetStreamingLevels().Num() - 1; Index >= 0; --Index)
 		{
+			ULevelStreaming* const FakeWorldStreamingLevel = FakeWorld->GetStreamingLevels()[Index];
 			FakeWorldStreamingLevel->Rename(nullptr, Context.World(), REN_ForceNoResetLoaders | REN_DontCreateRedirectors);
+			FakeWorld->RemoveStreamingLevelAt(Index);
+			StreamingLevelsToMove[Index] = FakeWorldStreamingLevel;
 		}
 
 		// Move the secondary levels to the world info levels array.
-		Context.World()->AddStreamingLevels(FakeWorld->GetStreamingLevels());
-		FakeWorld->ClearStreamingLevels();
+		Context.World()->AddStreamingLevels(StreamingLevelsToMove);
 
 		// fixup up any kismet streaming objects to force them to be loaded if they were preloaded, this
 		// will keep streaming volumes from immediately unloading the levels that were just loaded
