Unshelved from pending changelist '4243857':

PR #4917: DebugCam Spawn Pawn fix (Contributed by [REDACTED_PERSON])

diff --git a/Engine/Source/Runtime/Engine/Private/DebugCameraController.cpp b/Engine/Source/Runtime/Engine/Private/DebugCameraController.cpp
--- a/Engine/Source/Runtime/Engine/Private/DebugCameraController.cpp
+++ b/Engine/Source/Runtime/Engine/Private/DebugCameraController.cpp
@@ -254,7 +254,8 @@ ASpectatorPawn* ADebugCameraController::SpawnSpectatorPawn()
 			SpawnParams.Owner = this;
 			SpawnParams.SpawnCollisionHandlingOverride = ESpawnActorCollisionHandlingMethod::AlwaysSpawn;
 			SpawnParams.ObjectFlags |= RF_Transient;	// We never want to save spectator pawns into a map
-			SpawnedSpectator = GetWorld()->SpawnActor<ASpectatorPawn>(ASpectatorPawn::StaticClass(), GetSpawnLocation(), GetControlRotation(), SpawnParams);
+
+			SpawnedSpectator = GetWorld()->SpawnActor<ASpectatorPawn>((GameState->SpectatorClass ? GameState->SpectatorClass : ASpectatorPawn::StaticClass()), GetSpawnLocation(), GetControlRotation(), SpawnParams);
 			if (SpawnedSpectator)
 			{
 				SpawnedSpectator->PossessedBy(this);
