Fixing Clang error from CL 6665592

diff --git a/Engine/Source/Runtime/Engine/Private/DebugCameraController.cpp b/Engine/Source/Runtime/Engine/Private/DebugCameraController.cpp
--- a/Engine/Source/Runtime/Engine/Private/DebugCameraController.cpp
+++ b/Engine/Source/Runtime/Engine/Private/DebugCameraController.cpp
@@ -255,7 +255,7 @@ ASpectatorPawn* ADebugCameraController::SpawnSpectatorPawn()
 			SpawnParams.SpawnCollisionHandlingOverride = ESpawnActorCollisionHandlingMethod::AlwaysSpawn;
 			SpawnParams.ObjectFlags |= RF_Transient;	// We never want to save spectator pawns into a map
 
-			SpawnedSpectator = GetWorld()->SpawnActor<ASpectatorPawn>((GameState->SpectatorClass ? GameState->SpectatorClass : ASpectatorPawn::StaticClass()), GetSpawnLocation(), GetControlRotation(), SpawnParams);
+			SpawnedSpectator = GetWorld()->SpawnActor<ASpectatorPawn>((*GameState->SpectatorClass ? *GameState->SpectatorClass : ASpectatorPawn::StaticClass()), GetSpawnLocation(), GetControlRotation(), SpawnParams);
 			if (SpawnedSpectator)
 			{
 				SpawnedSpectator->PossessedBy(this);
