#ue4 - Improve comments for UProjectileMovementComponent regarding interpolation, to make it clear that you must do additional hookups aside from just turning on "Interpolate Movement" (ie point it to the interpolated visual component and plug in network updates). Should resolve some confusion mentioned in some user questions.

diff --git a/Engine/Source/Runtime/Engine/Classes/GameFramework/ProjectileMovementComponent.h b/Engine/Source/Runtime/Engine/Classes/GameFramework/ProjectileMovementComponent.h
--- a/Engine/Source/Runtime/Engine/Classes/GameFramework/ProjectileMovementComponent.h
+++ b/Engine/Source/Runtime/Engine/Classes/GameFramework/ProjectileMovementComponent.h
@@ -98,7 +98,8 @@ class ENGINE_API UProjectileMovementComponent : public UMovementComponent
 
 	/**
 	 * If true and there is an interpolated component set, location (and optionally rotation) interpolation is enabled which allows the interpolated object to smooth uneven updates
-	 * of the UpdatedComponent's location (usually to smooth network updates). 
+	 * of the UpdatedComponent's location (usually to smooth network updates). This requires using SetInterpolatedComponent() to indicate the visual component that lags behind the collision,
+	 * and using MoveInterpolationTarget() when the new target location/rotation is received (usually on a net update).
 	 * @see SetInterpolatedComponent(), MoveInterpolationTarget()
 	 */
 	UPROPERTY(EditAnywhere, BlueprintReadWrite, Category=ProjectileInterpolation)
