#ue4 - SpringArmComponent: Add virtual GetDesiredRotation() to allow custom target rotation before spring contraints are enforced.

PR #5780

diff --git a/Engine/Source/Runtime/Engine/Classes/GameFramework/SpringArmComponent.h b/Engine/Source/Runtime/Engine/Classes/GameFramework/SpringArmComponent.h
--- a/Engine/Source/Runtime/Engine/Classes/GameFramework/SpringArmComponent.h
+++ b/Engine/Source/Runtime/Engine/Classes/GameFramework/SpringArmComponent.h
@@ -152,6 +152,10 @@ class ENGINE_API USpringArmComponent : public USceneComponent
 
 	/** The name of the socket at the end of the spring arm (looking back towards the spring arm origin) */
 	static const FName SocketName;
+
+	/** Returns the desired rotation for the spring arm, before the rotation constraints such as bInheritPitch etc are enforced. */
+	virtual FRotator GetDesiredRotation() const;
+
 protected:
 	/** Cached component-space socket location */
 	FVector RelativeSocketLocation;
diff --git a/Engine/Source/Runtime/Engine/Private/GameFramework/SpringArmComponent.cpp b/Engine/Source/Runtime/Engine/Private/GameFramework/SpringArmComponent.cpp
--- a/Engine/Source/Runtime/Engine/Private/GameFramework/SpringArmComponent.cpp
+++ b/Engine/Source/Runtime/Engine/Private/GameFramework/SpringArmComponent.cpp
@@ -42,9 +42,14 @@ USpringArmComponent::USpringArmComponent(const FObjectInitializer& ObjectInitial
  	UnfixedCameraPosition = FVector::ZeroVector;
 }
 
+FRotator USpringArmComponent::GetDesiredRotation() const
+{
+	return GetComponentRotation();
+}
+
 FRotator USpringArmComponent::GetTargetRotation() const
 {
-	FRotator DesiredRot = GetComponentRotation();
+	FRotator DesiredRot = GetDesiredRotation();
 
 	if (bUsePawnControlRotation)
 	{
