Don't keep popping up warning messages about invalid composite table data when users do unrelated actions on the table.

diff --git a/Engine/Source/Runtime/Engine/Classes/Engine/CompositeDataTable.h b/Engine/Source/Runtime/Engine/Classes/Engine/CompositeDataTable.h
--- a/Engine/Source/Runtime/Engine/Classes/Engine/CompositeDataTable.h
+++ b/Engine/Source/Runtime/Engine/Classes/Engine/CompositeDataTable.h
@@ -67,9 +67,9 @@ protected:
 	// if bClearParentTables is false then the row map will be cleared but the parent table array won't be changed
 	void EmptyCompositeTable(bool bClearParentTables);
 
-	void UpdateCachedRowMap();
+	void UpdateCachedRowMap(bool bWarnOnInvalidChildren = true);
 
-	void OnParentTablesUpdated();
+	void OnParentTablesUpdated(EPropertyChangeType::Type ChangeType = EPropertyChangeType::Unspecified);
 
 	// true if this asset is currently being loaded; false otherwise
 	uint8 bIsLoading : 1;
diff --git a/Engine/Source/Runtime/Engine/Private/CompositeDataTable.cpp b/Engine/Source/Runtime/Engine/Private/CompositeDataTable.cpp
--- a/Engine/Source/Runtime/Engine/Private/CompositeDataTable.cpp
+++ b/Engine/Source/Runtime/Engine/Private/CompositeDataTable.cpp
@@ -52,22 +52,25 @@ UCompositeDataTable::ERowState UCompositeDataTable::GetRowState(FName RowName) c
 }
 #endif
 
-void UCompositeDataTable::UpdateCachedRowMap()
+void UCompositeDataTable::UpdateCachedRowMap(bool bWarnOnInvalidChildren)
 {
 	bool bLeaveEmpty = false;
 	// Throw up an error message and stop if any loops are found
 	if (const UCompositeDataTable* LoopTable = FindLoops(TArray<const UCompositeDataTable*>()))
 	{
-		const FText ErrorMsg = FText::Format(LOCTEXT("FoundLoopError", "Cyclic dependency found. Table {0} depends on itself. Please fix your data"), FText::FromString(LoopTable->GetPathName()));
-#if WITH_EDITOR
-		if (!bIsLoading)
+		if (bWarnOnInvalidChildren)
 		{
-			FMessageDialog::Open(EAppMsgType::Ok, ErrorMsg);
-		}
-		else
+			const FText ErrorMsg = FText::Format(LOCTEXT("FoundLoopError", "Cyclic dependency found. Table {0} depends on itself. Please fix your data"), FText::FromString(LoopTable->GetPathName()));
+#if WITH_EDITOR
+			if (!bIsLoading)
+			{
+				FMessageDialog::Open(EAppMsgType::Ok, ErrorMsg);
+			}
+			else
 #endif
-		{
-			UE_LOG(LogDataTable, Warning, TEXT("%s"), *ErrorMsg.ToString());
+			{
+				UE_LOG(LogDataTable, Warning, TEXT("%s"), *ErrorMsg.ToString());
+			}
 		}
 		bLeaveEmpty = true;
 
@@ -95,10 +98,13 @@ void UCompositeDataTable::UpdateCachedRowMap()
 			}
 			if (ParentTable->RowStruct != RowStruct)
 			{
-				bParentsHaveDifferentRowStruct = true;
-				FString CompositeRowStructName = RowStruct ? RowStruct->GetName() : "Missing row struct";
-				FString ParentRowStructName = ParentTable->RowStruct ? ParentTable->RowStruct->GetName() : "Missing row struct";
-				UE_LOG(LogDataTable, Error, TEXT("Composite tables must have the same row struct as their parent tables. Composite Table: %s, Composite Row Struct: %s, Parent Table: %s, Parent Row Struct: %s."), *GetName(), *CompositeRowStructName, *ParentTable->GetName(), *ParentRowStructName);
+				if (bWarnOnInvalidChildren)
+				{
+					bParentsHaveDifferentRowStruct = true;
+					FString CompositeRowStructName = RowStruct ? RowStruct->GetName() : "Missing row struct";
+					FString ParentRowStructName = ParentTable->RowStruct ? ParentTable->RowStruct->GetName() : "Missing row struct";
+					UE_LOG(LogDataTable, Error, TEXT("Composite tables must have the same row struct as their parent tables. Composite Table: %s, Composite Row Struct: %s, Parent Table: %s, Parent Row Struct: %s."), *GetName(), *CompositeRowStructName, *ParentTable->GetName(), *ParentRowStructName);
+				}
 				continue;
 			}
 
@@ -245,14 +251,14 @@ void UCompositeDataTable::PostEditChangeProperty(FPropertyChangedEvent& Property
 
 	if (PropertyName == Name_ParentTables)
 	{
-		OnParentTablesUpdated();
+		OnParentTablesUpdated(PropertyChangedEvent.ChangeType);
 	}
 
 	Super::PostEditChangeProperty(PropertyChangedEvent);
 }
 #endif // WITH_EDITOR
 
-void UCompositeDataTable::OnParentTablesUpdated()
+void UCompositeDataTable::OnParentTablesUpdated(EPropertyChangeType::Type ChangeType)
 {
 	for (UDataTable* Table : OldParentTables)
 	{
@@ -262,13 +268,13 @@ void UCompositeDataTable::OnParentTablesUpdated()
 		}
 	}
 
-	UpdateCachedRowMap();
+	UpdateCachedRowMap(ChangeType == EPropertyChangeType::ValueSet || ChangeType == EPropertyChangeType::Duplicate);
 
 	for (UDataTable* Table : ParentTables)
 	{
 		if (Table && OldParentTables.Find(Table) == INDEX_NONE)
 		{
-			Table->OnDataTableChanged().AddUObject(this, &UCompositeDataTable::UpdateCachedRowMap);
+			Table->OnDataTableChanged().AddUObject(this, &UCompositeDataTable::UpdateCachedRowMap, true);
 		}
 	}
 
