Added a BP callable method SetGameCameraCutThisFrame to APlayerCameraManager, allowing game code to set the flag as needed (e.g., when teleporting a player and wanting no motion blur the next frame)
PR #5844: add function to set bGameCameraCutThisFrame to true via BP (Contributed by [REDACTED_PERSON])

diff --git a/Engine/Source/Runtime/Engine/Classes/Camera/PlayerCameraManager.h b/Engine/Source/Runtime/Engine/Classes/Camera/PlayerCameraManager.h
--- a/Engine/Source/Runtime/Engine/Classes/Camera/PlayerCameraManager.h
+++ b/Engine/Source/Runtime/Engine/Classes/Camera/PlayerCameraManager.h
@@ -859,6 +859,10 @@ public:
 	/** Returns first existing instance of the specified camera anim, or NULL if none exists. */
 	class UCameraAnimInst* FindInstanceOfCameraAnim(class UCameraAnim const* Anim) const;
 
+	/** Sets the bGameCameraCutThisFrame flag to true (indicating we did a camera cut this frame; useful for game code to call, e.g., when performing a teleport that should be seamless) */
+	UFUNCTION(BlueprintCallable, Category = "Camera")
+	void SetGameCameraCutThisFrame() { bGameCameraCutThisFrame = true; }
+
 protected:
 	/** Gets specified temporary CameraActor ready to update the specified Anim. */
 	void InitTempCameraActor(class ACameraActor* CamActor, class UCameraAnimInst const* AnimInstToInitFor) const;
