Rollback CL# 6440677 as it never loads the class and modifying the existing implemenation would introduce a synchronous load
#rnx

diff --git a/Engine/Source/Runtime/Engine/Classes/GameFramework/Pawn.h b/Engine/Source/Runtime/Engine/Classes/GameFramework/Pawn.h
--- a/Engine/Source/Runtime/Engine/Classes/GameFramework/Pawn.h
+++ b/Engine/Source/Runtime/Engine/Classes/GameFramework/Pawn.h
@@ -103,15 +103,10 @@ public:
 	UPROPERTY(replicated)
 	uint8 RemoteViewPitch;
 
-	/** This property is deprecated. Use \'AI Controller Class\' instead */
-	UE_DEPRECATED(4.23, "This property is deprecated. Use AIControllerClassRef instead.")
-	UPROPERTY(BlueprintReadWrite, meta=(DisplayName="DEPRECATED AI Controller Class", DeprecatedProperty, DeprecationMessage = "This property is deprecated. Use \'AIController Class Ref\' instead"), Category=Pawn)
+	/** Default class to use when pawn is controlled by AI. */
+	UPROPERTY(EditAnywhere, BlueprintReadWrite, meta=(DisplayName="AI Controller Class"), Category=Pawn)
 	TSubclassOf<AController> AIControllerClass;
 
-	/** Default controller class to use when pawn is controlled by AI. */
-	UPROPERTY(EditAnywhere, BlueprintReadWrite, meta = (AllowedClasses = "AIController"), Category = Pawn)
-	TSoftClassPtr<AController> AIControllerClassRef;
-
 	/**
 	 * Return our PawnNoiseEmitterComponent, if any. Default implementation returns the first PawnNoiseEmitterComponent found in the components array.
 	 * If one isn't found, then it tries to find one on the Pawn's current Controller.
@@ -239,7 +234,6 @@ public:
 	virtual void EndPlay(const EEndPlayReason::Type EndPlayReason) override;
 	virtual void PreInitializeComponents() override;
 	virtual void PostInitializeComponents() override;
-	virtual void PostInitProperties() override;
 	virtual const AActor* GetNetOwner() const override;
 	virtual UPlayer* GetNetOwningPlayer() override;
 	virtual class UNetConnection* GetNetConnection() const override;
diff --git a/Engine/Source/Runtime/Engine/Private/Pawn.cpp b/Engine/Source/Runtime/Engine/Private/Pawn.cpp
--- a/Engine/Source/Runtime/Engine/Private/Pawn.cpp
+++ b/Engine/Source/Runtime/Engine/Private/Pawn.cpp
@@ -32,20 +32,6 @@
 DEFINE_LOG_CATEGORY(LogDamage);
 DEFINE_LOG_CATEGORY_STATIC(LogPawn, Warning, All);
 
-namespace 
-{
-	FORCEINLINE void PatchAIControllerClass(APawn& Pawn)
-	{
-		PRAGMA_DISABLE_DEPRECATION_WARNINGS
-		if (Pawn.AIControllerClass)
-		{
-			Pawn.AIControllerClassRef = Pawn.AIControllerClass;
-		}
-		Pawn.AIControllerClass = nullptr;
-		PRAGMA_ENABLE_DEPRECATION_WARNINGS
-	}
-}
-
 APawn::APawn(const FObjectInitializer& ObjectInitializer)
 	: Super(ObjectInitializer)
 {
@@ -56,11 +42,12 @@ APawn::APawn(const FObjectInitializer& ObjectInitializer)
 
 	if (HasAnyFlags(RF_ClassDefaultObject) && GetClass() == APawn::StaticClass())
 	{
-		AIControllerClassRef = FSoftObjectPath(((UEngine*)(UEngine::StaticClass()->GetDefaultObject()))->AIControllerClassName);
+		// WARNING: This line is why the AISupport plugin has to load the AIModule before UObject initialization, otherwise this load fails and CDOs are corrupt in the editor
+		AIControllerClass = LoadClass<AController>(nullptr, *((UEngine*)(UEngine::StaticClass()->GetDefaultObject()))->AIControllerClassName.ToString(), nullptr, LOAD_None, nullptr);
 	}
 	else
 	{
-		AIControllerClassRef = ((APawn*)APawn::StaticClass()->GetDefaultObject())->AIControllerClassRef;
+		AIControllerClass = ((APawn*)APawn::StaticClass()->GetDefaultObject())->AIControllerClass;
 	}
 	bCanBeDamaged = true;
 	
@@ -144,23 +131,12 @@ void APawn::PostInitializeComponents()
 	}
 }
 
-void APawn::PostInitProperties()
-{
-	Super::PostInitProperties();
-
-	// 4.23 patching
-	PatchAIControllerClass(*this);
-}
-
 void APawn::PostLoad()
 {
 	Super::PostLoad();
 
 	// A pawn should never have this enabled, so we'll aggressively disable it if it did occur.
 	AutoReceiveInput = EAutoReceiveInput::Disabled;
-
-	// 4.23 patching
-	PatchAIControllerClass(*this);
 }
 
 void APawn::PostRegisterAllComponents()
@@ -333,15 +309,14 @@ void APawn::SpawnDefaultController()
 		return;
 	}
 
-	UClass* AIControllerClassPtr = AIControllerClassRef.Get();
-	if (AIControllerClassPtr)
+	if (AIControllerClass != nullptr)
 	{
 		FActorSpawnParameters SpawnInfo;
 		SpawnInfo.Instigator = Instigator;
 		SpawnInfo.SpawnCollisionHandlingOverride = ESpawnActorCollisionHandlingMethod::AlwaysSpawn;
 		SpawnInfo.OverrideLevel = GetLevel();
 		SpawnInfo.ObjectFlags |= RF_Transient;	// We never want to save AI controllers into a map
-		AController* NewController = GetWorld()->SpawnActor<AController>(AIControllerClassPtr, GetActorLocation(), GetActorRotation(), SpawnInfo);
+		AController* NewController = GetWorld()->SpawnActor<AController>(AIControllerClass, GetActorLocation(), GetActorRotation(), SpawnInfo);
 		if (NewController != nullptr)
 		{
 			// if successful will result in setting this->Controller 
