Clean up comments for LoadAsset_Blocking and change it to only do a sync load if it's not already loaded
Basedon PR #5898

diff --git a/Engine/Source/Editor/BlueprintGraph/Private/K2Node_LoadAsset.cpp b/Engine/Source/Editor/BlueprintGraph/Private/K2Node_LoadAsset.cpp
--- a/Engine/Source/Editor/BlueprintGraph/Private/K2Node_LoadAsset.cpp
+++ b/Engine/Source/Editor/BlueprintGraph/Private/K2Node_LoadAsset.cpp
@@ -202,7 +202,7 @@ void UK2Node_LoadAsset::ExpandNode(class FKismetCompilerContext& CompilerContext
 
 FText UK2Node_LoadAsset::GetTooltipText() const
 {
-	return FText(LOCTEXT("UK2Node_LoadAssetGetTooltipText", "Async Load Asset"));
+	return FText(LOCTEXT("UK2Node_LoadAssetGetTooltipText", "Asynchronously loads a Soft Object Reference and returns object of the correct type if the load succeeds"));
 }
 
 FText UK2Node_LoadAsset::GetNodeTitle(ENodeTitleType::Type TitleType) const
@@ -289,7 +289,7 @@ FName UK2Node_LoadAssetClass::NativeFunctionName() const
 
 FText UK2Node_LoadAssetClass::GetTooltipText() const
 {
-	return FText(LOCTEXT("UK2Node_LoadAssetClassGetTooltipText", "Async Load Class Asset"));
+	return FText(LOCTEXT("UK2Node_LoadAssetClassGetTooltipText", "Asynchronously loads a Soft Class Reference and returns class of the correct type if the load succeeds"));
 }
 
 FText UK2Node_LoadAssetClass::GetNodeTitle(ENodeTitleType::Type TitleType) const
diff --git a/Engine/Source/Runtime/Engine/Classes/Kismet/KismetSystemLibrary.h b/Engine/Source/Runtime/Engine/Classes/Kismet/KismetSystemLibrary.h
--- a/Engine/Source/Runtime/Engine/Classes/Kismet/KismetSystemLibrary.h
+++ b/Engine/Source/Runtime/Engine/Classes/Kismet/KismetSystemLibrary.h
@@ -208,6 +208,7 @@ class ENGINE_API UKismetSystemLibrary : public UBlueprintFunctionLibrary
 	UFUNCTION(BlueprintPure, Category = "SoftObjectPath", meta = (NativeBreakFunc, BlueprintThreadSafe))
 	static void BreakSoftObjectPath(FSoftObjectPath InSoftObjectPath, FString& PathString);
 
+	/** Converts a Soft Object Path into a base Soft Object Reference, this is not guaranteed to be resolvable */
 	UFUNCTION(BlueprintPure, meta = (DisplayName = "ToSoftObjectReference (SoftObjectPath)", CompactNodeTitle = "->"), Category = "Utilities")
 	static TSoftObjectPtr<UObject> Conv_SoftObjPathToSoftObjRef(const FSoftObjectPath& SoftObjectPath);
 
@@ -219,6 +220,10 @@ class ENGINE_API UKismetSystemLibrary : public UBlueprintFunctionLibrary
 	UFUNCTION(BlueprintPure, Category = "SoftClassPath", meta = (NativeBreakFunc, BlueprintThreadSafe))
 	static void BreakSoftClassPath(FSoftClassPath InSoftClassPath, FString& PathString);
 
+	/** Converts a Soft Class Path into a base Soft Class Reference, this is not guaranteed to be resolvable */
+	UFUNCTION(BlueprintPure, meta = (DisplayName = "ToSoftClassReference (SoftClassPath)", CompactNodeTitle = "->"), Category = "Utilities")
+	static TSoftClassPtr<UObject> Conv_SoftClassPathToSoftClassRef(const FSoftClassPath& SoftClassPath);
+
 	/** Returns true if the Soft Object Reference is not null */
 	UFUNCTION(BlueprintPure, Category = "Utilities", meta = (BlueprintThreadSafe))
 	static bool IsValidSoftObjectReference(const TSoftObjectPtr<UObject>& SoftObjectReference);
@@ -235,6 +240,10 @@ class ENGINE_API UKismetSystemLibrary : public UBlueprintFunctionLibrary
 	UFUNCTION(BlueprintPure, meta = (DisplayName = "NotEqual (SoftObjectReference)", CompactNodeTitle = "!=", BlueprintThreadSafe), Category = "Utilities")
 	static bool NotEqual_SoftObjectReference(const TSoftObjectPtr<UObject>& A, const TSoftObjectPtr<UObject>& B);
 
+	/** Resolves or loads a Soft Object Reference immediately, this will cause hitches and Async Load Asset should be used if possible */
+	UFUNCTION(BlueprintCallable, Category = "Utilities", meta = (DeterminesOutputType = "Asset"))
+	static UObject* LoadAsset_Blocking(TSoftObjectPtr<UObject> Asset);
+
 	/** Returns true if the Soft Class Reference is not null */
 	UFUNCTION(BlueprintPure, Category = "Utilities", meta = (BlueprintThreadSafe))
 	static bool IsValidSoftClassReference(const TSoftClassPtr<UObject>& SoftClassReference);
@@ -251,6 +260,12 @@ class ENGINE_API UKismetSystemLibrary : public UBlueprintFunctionLibrary
 	UFUNCTION(BlueprintPure, meta = (DisplayName = "NotEqual (SoftClassReference)", CompactNodeTitle = "!=", BlueprintThreadSafe), Category = "Utilities")
 	static bool NotEqual_SoftClassReference(const TSoftClassPtr<UObject>& A, const TSoftClassPtr<UObject>& B);
 
+	/** Resolves or loads a Soft Class Reference immediately, this will cause hitches and Async Load Class Asset should be used if possible */
+	UFUNCTION(BlueprintCallable, Category = "Utilities", meta = (DeterminesOutputType = "AssetClass"))
+	static UClass* LoadClassAsset_Blocking(TSoftClassPtr<UObject> AssetClass);
+
+	// Internal functions used by K2Node_LoadAsset and K2Node_ConvertAsset
+
 	UFUNCTION(BlueprintPure, meta = (BlueprintInternalUseOnly = "true"), Category = "Utilities")
 	static UObject* Conv_SoftObjectReferenceToObject(const TSoftObjectPtr<UObject>& SoftObject);
 
@@ -268,9 +283,6 @@ class ENGINE_API UKismetSystemLibrary : public UBlueprintFunctionLibrary
 	UFUNCTION(BlueprintCallable, meta = (Latent, LatentInfo = "LatentInfo", WorldContext = "WorldContextObject", BlueprintInternalUseOnly = "true"), Category = "Utilities")
 	static void LoadAsset(UObject* WorldContextObject, TSoftObjectPtr<UObject> Asset, FOnAssetLoaded OnLoaded, FLatentActionInfo LatentInfo);
 
-	UFUNCTION(BlueprintCallable, Category="Utilities", meta=(DeterminesOutputType = "Asset"))
-	static UObject* LoadAsset_Blocking(TSoftObjectPtr<UObject> Asset);
-
 	DECLARE_DYNAMIC_DELEGATE_OneParam(FOnAssetClassLoaded, TSubclassOf<UObject>, Loaded);
 
 	UFUNCTION(BlueprintCallable, meta = (Latent, LatentInfo = "LatentInfo", WorldContext = "WorldContextObject", BlueprintInternalUseOnly = "true"), Category = "Utilities")
diff --git a/Engine/Source/Runtime/Engine/Private/KismetSystemLibrary.cpp b/Engine/Source/Runtime/Engine/Private/KismetSystemLibrary.cpp
--- a/Engine/Source/Runtime/Engine/Private/KismetSystemLibrary.cpp
+++ b/Engine/Source/Runtime/Engine/Private/KismetSystemLibrary.cpp
@@ -988,6 +988,11 @@ void UKismetSystemLibrary::BreakSoftClassPath(FSoftClassPath InSoftClassPath, FS
 	PathString = InSoftClassPath.ToString();
 }
 
+TSoftClassPtr<UObject> UKismetSystemLibrary::Conv_SoftClassPathToSoftClassRef(const FSoftClassPath& SoftClassPath)
+{
+	return TSoftClassPtr<UObject>(SoftClassPath);
+}
+
 bool UKismetSystemLibrary::IsValidSoftObjectReference(const TSoftObjectPtr<UObject>& SoftObjectReference)
 {
 	return !SoftObjectReference.IsNull();
@@ -1008,6 +1013,11 @@ bool UKismetSystemLibrary::NotEqual_SoftObjectReference(const TSoftObjectPtr<UOb
 	return A != B;
 }
 
+UObject* UKismetSystemLibrary::LoadAsset_Blocking(TSoftObjectPtr<UObject> Asset)
+{
+	return Asset.LoadSynchronous();
+}
+
 bool UKismetSystemLibrary::IsValidSoftClassReference(const TSoftClassPtr<UObject>& SoftClassReference)
 {
 	return !SoftClassReference.IsNull();
@@ -1028,6 +1038,11 @@ bool UKismetSystemLibrary::NotEqual_SoftClassReference(const TSoftClassPtr<UObje
 	return A != B;
 }
 
+UClass* UKismetSystemLibrary::LoadClassAsset_Blocking(TSoftClassPtr<UObject> AssetClass)
+{
+	return AssetClass.LoadSynchronous();
+}
+
 UObject* UKismetSystemLibrary::Conv_SoftObjectReferenceToObject(const TSoftObjectPtr<UObject>& SoftObject)
 {
 	return SoftObject.Get();
@@ -2561,11 +2576,6 @@ void UKismetSystemLibrary::LoadAsset(UObject* WorldContextObject, TSoftObjectPtr
 	}
 }
 
-UObject* UKismetSystemLibrary::LoadAsset_Blocking(TSoftObjectPtr<UObject> Asset)
-{
-	return Asset.ToSoftObjectPath().TryLoad();
-}
-
 void UKismetSystemLibrary::LoadAssetClass(UObject* WorldContextObject, TSoftClassPtr<UObject> AssetClass, UKismetSystemLibrary::FOnAssetClassLoaded OnLoaded, FLatentActionInfo LatentInfo)
 {
 	struct FLoadAssetClassAction : public FLoadAssetActionBase
