PR #5510: UDataTable logs should match function (Contributed by [REDACTED_PERSON])

diff --git a/Engine/Source/Runtime/Engine/Classes/Engine/DataTable.h b/Engine/Source/Runtime/Engine/Classes/Engine/DataTable.h
--- a/Engine/Source/Runtime/Engine/Classes/Engine/DataTable.h
+++ b/Engine/Source/Runtime/Engine/Classes/Engine/DataTable.h
@@ -146,11 +146,11 @@ public:
 	{
 		if (RowStruct == nullptr)
 		{
-			UE_LOG(LogDataTable, Error, TEXT("UDataTable::FindRow : DataTable '%s' has no RowStruct specified (%s)."), *GetPathName(), ContextString);
+			UE_LOG(LogDataTable, Error, TEXT("UDataTable::GetAllRows : DataTable '%s' has no RowStruct specified (%s)."), *GetPathName(), ContextString);
 		}
 		else if (!RowStruct->IsChildOf(T::StaticStruct()))
 		{
-			UE_LOG(LogDataTable, Error, TEXT("UDataTable::FindRow : Incorrect type specified for DataTable '%s' (%s)."), *GetPathName(), ContextString);
+			UE_LOG(LogDataTable, Error, TEXT("UDataTable::GetAllRows : Incorrect type specified for DataTable '%s' (%s)."), *GetPathName(), ContextString);
 		}
 		else
 		{
@@ -218,11 +218,11 @@ public:
 	{
 		if (RowStruct == nullptr)
 		{
-			UE_LOG(LogDataTable, Error, TEXT("UDataTable::FindRow : DataTable '%s' has no RowStruct specified (%s)."), *GetPathName(), ContextString);
+			UE_LOG(LogDataTable, Error, TEXT("UDataTable::ForeachRow : DataTable '%s' has no RowStruct specified (%s)."), *GetPathName(), ContextString);
 		}
 		else if (!RowStruct->IsChildOf(T::StaticStruct()))
 		{
-			UE_LOG(LogDataTable, Error, TEXT("UDataTable::FindRow : Incorrect type specified for DataTable '%s' (%s)."), *GetPathName(), ContextString);
+			UE_LOG(LogDataTable, Error, TEXT("UDataTable::ForeachRow : Incorrect type specified for DataTable '%s' (%s)."), *GetPathName(), ContextString);
 		}
 		else
 		{
@@ -247,13 +247,13 @@ public:
 	{
 		if(RowStruct == nullptr)
 		{
-			//UE_CLOG(MustExist, LogDataTable, Error, TEXT("UDataTable::FindRow : DataTable '%s' has no RowStruct specified (%s)."), *GetPathName(), *ContextString);
+			//UE_CLOG(MustExist, LogDataTable, Error, TEXT("UDataTable::FindRowUnchecked : DataTable '%s' has no RowStruct specified (%s)."), *GetPathName(), *ContextString);
 			return nullptr;
 		}
 
 		if(RowName == NAME_None)
 		{
-			//UE_CLOG(MustExist, LogDataTable, Warning, TEXT("UDataTable::FindRow : NAME_None is invalid row name for DataTable '%s' (%s)."), *GetPathName(), *ContextString);
+			//UE_CLOG(MustExist, LogDataTable, Warning, TEXT("UDataTable::FindRowUnchecked : NAME_None is invalid row name for DataTable '%s' (%s)."), *GetPathName(), *ContextString);
 			return nullptr;
 		}
 
@@ -390,7 +390,7 @@ struct ENGINE_API FDataTableRowHandle
 		{
 			if (RowName != NAME_None)
 			{
-				UE_LOG(LogDataTable, Warning, TEXT("FDataTableRowHandle::FindRow : No DataTable for row %s (%s)."), *RowName.ToString(), ContextString);
+				UE_LOG(LogDataTable, Warning, TEXT("FDataTableRowHandle::GetRow : No DataTable for row %s (%s)."), *RowName.ToString(), ContextString);
 			}
 			return nullptr;
 		}
@@ -461,7 +461,7 @@ struct ENGINE_API FDataTableCategoryHandle
 		{
 			if (RowContents != NAME_None)
 			{
-				UE_LOG(LogDataTable, Warning, TEXT("FDataTableCategoryHandle::FindRow : No DataTable for row %s (%s)."), *RowContents.ToString(), *ContextString);
+				UE_LOG(LogDataTable, Warning, TEXT("FDataTableCategoryHandle::GetRows : No DataTable for row %s (%s)."), *RowContents.ToString(), *ContextString);
 			}
 
 			return;
@@ -471,7 +471,7 @@ struct ENGINE_API FDataTableCategoryHandle
 		{
 			if (RowContents != NAME_None)
 			{
-				UE_LOG(LogDataTable, Warning, TEXT("FDataTableCategoryHandle::FindRow : No Column selected for row %s (%s)."), *RowContents.ToString(), *ContextString);
+				UE_LOG(LogDataTable, Warning, TEXT("FDataTableCategoryHandle::GetRows : No Column selected for row %s (%s)."), *RowContents.ToString(), *ContextString);
 			}
 
 			return;
