Expose Sockets to python.
I opted to take steps towards a more controlled API, the intent is that the Sockets Array not be public but I did not take it all the way in this step. This does mean that the Python API will be a bit more restrictive but I think thats a good thing as we learn how people use it. The next step would be adding index based API to StaticMesh for sockets most likely then making it private.

diff --git a/Engine/Source/Editor/UnrealEd/Private/Fbx/FbxStaticMeshImport.cpp b/Engine/Source/Editor/UnrealEd/Private/Fbx/FbxStaticMeshImport.cpp
--- a/Engine/Source/Editor/UnrealEd/Private/Fbx/FbxStaticMeshImport.cpp
+++ b/Engine/Source/Editor/UnrealEd/Private/Fbx/FbxStaticMeshImport.cpp
@@ -2201,7 +2201,7 @@ void UnFbx::FFbxImporter::ImportStaticMeshLocalSockets(UStaticMesh* StaticMesh,
 			check(Socket);
 
 			Socket->SocketName = SocketNode.SocketName;
-			StaticMesh->Sockets.Add(Socket);
+			StaticMesh->AddSocket(Socket);
 		}
 
 		if (Socket)
@@ -2267,7 +2267,7 @@ void UnFbx::FFbxImporter::ImportStaticMeshGlobalSockets( UStaticMesh* StaticMesh
 			check(Socket);
 
 			Socket->SocketName = SocketNode.SocketName;
-			StaticMesh->Sockets.Add(Socket);
+			StaticMesh->AddSocket(Socket);
 			//Remove the axis conversion for the socket since its attach to a mesh containing this conversion.
 			const FbxAMatrix& SocketMatrix = Scene->GetAnimationEvaluator()->GetNodeGlobalTransform(SocketNode.Node) * FFbxDataConverter::GetAxisConversionMatrixInv();
 			FTransform SocketTransform;
diff --git a/Engine/Source/Editor/UnrealEd/Private/SSocketManager.cpp b/Engine/Source/Editor/UnrealEd/Private/SSocketManager.cpp
--- a/Engine/Source/Editor/UnrealEd/Private/SSocketManager.cpp
+++ b/Engine/Source/Editor/UnrealEd/Private/SSocketManager.cpp
@@ -405,7 +405,7 @@ void SSocketManager::CreateSocket()
 		NewSocket->OnPropertyChanged().AddSP( this, &SSocketManager::OnSocketPropertyChanged );
 
 		CurrentStaticMesh->PreEditChange(NULL);
-		CurrentStaticMesh->Sockets.Add(NewSocket);
+		CurrentStaticMesh->AddSocket(NewSocket);
 		CurrentStaticMesh->PostEditChange();
 		CurrentStaticMesh->MarkPackageDirty();
 
@@ -437,7 +437,7 @@ void SSocketManager::DuplicateSelectedSocket()
 
 		// Add the new socket to the static mesh
 		CurrentStaticMesh->PreEditChange(NULL);
-		CurrentStaticMesh->Sockets.Add(NewSocket);
+		CurrentStaticMesh->AddSocket(NewSocket);
 		CurrentStaticMesh->PostEditChange();
 		CurrentStaticMesh->MarkPackageDirty();
 
diff --git a/Engine/Source/Editor/UnrealEd/Private/StaticMeshEdit.cpp b/Engine/Source/Editor/UnrealEd/Private/StaticMeshEdit.cpp
--- a/Engine/Source/Editor/UnrealEd/Private/StaticMeshEdit.cpp
+++ b/Engine/Source/Editor/UnrealEd/Private/StaticMeshEdit.cpp
@@ -1609,7 +1609,7 @@ void RestoreExistingMeshData(ExistingStaticMeshData* ExistingMeshDataPtr, UStati
 		UStaticMeshSocket* Socket = NewMesh->FindSocket(ExistingSocket->SocketName);
 		if (!Socket && !ExistingSocket->bSocketCreatedAtImport)
 		{
-			NewMesh->Sockets.Add(ExistingSocket);
+			NewMesh->AddSocket(ExistingSocket);
 		}
 	}
 
diff --git a/Engine/Source/Runtime/Engine/Classes/Engine/StaticMesh.h b/Engine/Source/Runtime/Engine/Classes/Engine/StaticMesh.h
--- a/Engine/Source/Runtime/Engine/Classes/Engine/StaticMesh.h
+++ b/Engine/Source/Runtime/Engine/Classes/Engine/StaticMesh.h
@@ -1073,6 +1073,12 @@ public:
 #endif // WITH_EDITORONLY_DATA
 	}
 
+	/**
+	 *	Add a socket object in this StaticMesh.
+	 */
+	UFUNCTION(BlueprintCallable, Category = "StaticMesh")
+	ENGINE_API void AddSocket(UStaticMeshSocket* Socket);
+
 	/**
 	 *	Find a socket object in this StaticMesh by name.
 	 *	Entering NAME_None will return NULL. If there are multiple sockets with the same name, will return the first one.
@@ -1080,6 +1086,12 @@ public:
 	UFUNCTION(BlueprintCallable, Category = "StaticMesh")
 	ENGINE_API class UStaticMeshSocket* FindSocket(FName InSocketName) const;
 
+	/**
+	 *	Remove a socket object in this StaticMesh by providing it's pointer. Use FindSocket() if needed.
+	 */
+	UFUNCTION(BlueprintCallable, Category = "StaticMesh")
+	ENGINE_API void RemoveSocket(UStaticMeshSocket* Socket);
+
 	/**
 	 * Returns vertex color data by position.
 	 * For matching to reimported meshes that may have changed or copying vertex paint data from mesh to mesh.
diff --git a/Engine/Source/Runtime/Engine/Private/StaticMesh.cpp b/Engine/Source/Runtime/Engine/Private/StaticMesh.cpp
--- a/Engine/Source/Runtime/Engine/Private/StaticMesh.cpp
+++ b/Engine/Source/Runtime/Engine/Private/StaticMesh.cpp
@@ -4993,6 +4993,11 @@ void UStaticMesh::GenerateLodsInPackage()
 
 #endif // #if WITH_EDITOR
 
+void UStaticMesh::AddSocket(UStaticMeshSocket* Socket)
+{
+	Sockets.AddUnique(Socket);
+}
+
 UStaticMeshSocket* UStaticMesh::FindSocket(FName InSocketName) const
 {
 	if(InSocketName == NAME_None)
@@ -5011,6 +5016,11 @@ UStaticMeshSocket* UStaticMesh::FindSocket(FName InSocketName) const
 	return NULL;
 }
 
+void UStaticMesh::RemoveSocket(UStaticMeshSocket* Socket)
+{
+	Sockets.Remove(Socket);
+}
+
 /*-----------------------------------------------------------------------------
 UStaticMeshSocket
 -----------------------------------------------------------------------------*/
