Swap a UE_LOG warning to Display to prevent a warning when we intentionnaly set a material usage.

diff --git a/Engine/Source/Runtime/Engine/Private/Materials/Material.cpp b/Engine/Source/Runtime/Engine/Private/Materials/Material.cpp
--- a/Engine/Source/Runtime/Engine/Private/Materials/Material.cpp
+++ b/Engine/Source/Runtime/Engine/Private/Materials/Material.cpp
@@ -1804,7 +1804,7 @@ bool UMaterial::SetMaterialUsage(bool &bNeedsRecompile, EMaterialUsage Usage)
 			//Do not warn the user during automation testing
 			if (!GIsAutomationTesting)
 			{
-				UE_LOG(LogMaterial, Warning, TEXT("Material %s needed to have new flag set %s !"), *GetPathName(), *GetUsageName(Usage));
+				UE_LOG(LogMaterial, Display, TEXT("Material %s needed to have new flag set %s !"), *GetPathName(), *GetUsageName(Usage));
 			}
 
 			// Open a material update context so this material can be modified safely.
