Disable stateless challenge ack retries, once the client sends a non-handshake packet.

Fixes an issue where connections would be initialized into a bad state, if connecting to the same IP:Port, as a prior connection.

diff --git a/Engine/Source/Runtime/Engine/Private/PacketHandlers/StatelessConnectHandlerComponent.cpp b/Engine/Source/Runtime/Engine/Private/PacketHandlers/StatelessConnectHandlerComponent.cpp
--- a/Engine/Source/Runtime/Engine/Private/PacketHandlers/StatelessConnectHandlerComponent.cpp
+++ b/Engine/Source/Runtime/Engine/Private/PacketHandlers/StatelessConnectHandlerComponent.cpp
@@ -710,14 +710,17 @@ void StatelessConnectHandlerComponent::Incoming(FBitReader& Packet)
 			}
 			else if (Handler->Mode == Handler::Mode::Server)
 			{
-				// The server should not be receiving handshake packets at this stage - resend the ack in case it was lost.
-				// In this codepath, this component is linked to a UNetConnection, and the Last* values below, cache the handshake info.
+				if (LastChallengeSuccessAddress.IsValid())
+				{
+					// The server should not be receiving handshake packets at this stage - resend the ack in case it was lost.
+					// In this codepath, this component is linked to a UNetConnection, and the Last* values below, cache the handshake info.
 #if !UE_BUILD_SHIPPING
-				UE_LOG(LogHandshake, Log, TEXT("Received unexpected post-connect handshake packet - resending ack for LastChallengeSuccessAddress %s and LastCookie %s."),
-						*LastChallengeSuccessAddress->ToString(true), *FString::FromBlob(AuthorisedCookie, COOKIE_BYTE_SIZE));
+					UE_LOG(LogHandshake, Log, TEXT("Received unexpected post-connect handshake packet - resending ack for LastChallengeSuccessAddress %s and LastCookie %s."),
+							*LastChallengeSuccessAddress->ToString(true), *FString::FromBlob(AuthorisedCookie, COOKIE_BYTE_SIZE));
 #endif
 
-				SendChallengeAck(LastChallengeSuccessAddress, AuthorisedCookie);
+					SendChallengeAck(LastChallengeSuccessAddress, AuthorisedCookie);
+				}
 			}
 		}
 		else
@@ -735,6 +738,11 @@ void StatelessConnectHandlerComponent::Incoming(FBitReader& Packet)
 		UE_LOG(LogHandshake, Log, TEXT("Incoming: Error reading handshake bit from packet."));
 	}
 #endif
+	// Servers should wipe LastChallengeSuccessAddress when the first non-handshake packet is received by the client, in order to disable challenge ack resending
+	else if (LastChallengeSuccessAddress.IsValid() && Handler->Mode == Handler::Mode::Server)
+	{
+		LastChallengeSuccessAddress.Reset();
+	}
 }
 
 void StatelessConnectHandlerComponent::Outgoing(FBitWriter& Packet, FOutPacketTraits& Traits)
