Fix missing exports for the NetDriver which causes some platforms (HTML5) to not compile.

#rnx

diff --git a/Engine/Source/Runtime/Engine/Classes/Engine/NetDriver.h b/Engine/Source/Runtime/Engine/Classes/Engine/NetDriver.h
--- a/Engine/Source/Runtime/Engine/Classes/Engine/NetDriver.h
+++ b/Engine/Source/Runtime/Engine/Classes/Engine/NetDriver.h
@@ -853,10 +853,10 @@ public:
 	ENGINE_API virtual void LowLevelDestroy();
 
 	/* @return network number */
-	virtual FString LowLevelGetNetworkNumber();
+	ENGINE_API virtual FString LowLevelGetNetworkNumber();
 
 	/* @return local addr of this machine if set */
-	virtual TSharedPtr<const FInternetAddr> GetLocalAddr() { return LocalAddr; }
+	ENGINE_API virtual TSharedPtr<const FInternetAddr> GetLocalAddr() { return LocalAddr; }
 
 	/** Make sure this connection is in a reasonable state. */
 	ENGINE_API virtual void AssertValid();
