More robust fix for SceneComponent Attach replication
Replicate extra information indicating if relativeLocation/Rotation should be snapped or not.

diff --git a/Engine/Source/Runtime/Engine/Classes/Components/SceneComponent.h b/Engine/Source/Runtime/Engine/Classes/Components/SceneComponent.h
--- a/Engine/Source/Runtime/Engine/Classes/Components/SceneComponent.h
+++ b/Engine/Source/Runtime/Engine/Classes/Components/SceneComponent.h
@@ -146,11 +146,11 @@ public:
 	FBoxSphereBounds Bounds;
 
 	/** Location of the component relative to its parent */
-	UPROPERTY(EditAnywhere, BlueprintReadOnly, ReplicatedUsing=OnRep_RelativeLocation, Category = Transform)
+	UPROPERTY(EditAnywhere, BlueprintReadOnly, ReplicatedUsing=OnRep_Transform, Category = Transform)
 	FVector RelativeLocation;
 
 	/** Rotation of the component relative to its parent */
-	UPROPERTY(EditAnywhere, BlueprintReadOnly, ReplicatedUsing=OnRep_RelativeRotation, Category=Transform)
+	UPROPERTY(EditAnywhere, BlueprintReadOnly, ReplicatedUsing=OnRep_Transform, Category=Transform)
 	FRotator RelativeRotation;
 
 	/**
@@ -207,6 +207,12 @@ private:
 	UPROPERTY(Transient, Replicated)
 	uint8 bShouldBeAttached : 1;
 
+	UPROPERTY(Transient, Replicated)
+	uint8 bShouldSnapLocationWhenAttached : 1;
+
+	UPROPERTY(Transient, Replicated)
+	uint8 bShouldSnapRotationWhenAttached : 1;
+
 	/**
 	 * Whether or not the cached PhysicsVolume this component overlaps should be updated when the component is moved.
 	 * @see GetPhysicsVolume()
@@ -252,8 +258,6 @@ protected:
 private:
 	uint8 bNetUpdateTransform : 1;
 	uint8 bNetUpdateAttachment : 1;
-	uint8 bNetHasReceivedRelativeLocation : 1;
-	uint8 bNetHasReceivedRelativeRotation : 1;
 
 public:
 	/** Global flag to enable/disable overlap optimizations, settable with p.SkipUpdateOverlapsOptimEnabled cvar */ 
@@ -313,12 +317,6 @@ private:
 	UFUNCTION()
 	void OnRep_Transform();
 
-	UFUNCTION()
-	void OnRep_RelativeLocation();
-
-	UFUNCTION()
-	void OnRep_RelativeRotation();
-
 	UFUNCTION()
 	void OnRep_AttachParent();
 
diff --git a/Engine/Source/Runtime/Engine/Private/Components/SceneComponent.cpp b/Engine/Source/Runtime/Engine/Private/Components/SceneComponent.cpp
--- a/Engine/Source/Runtime/Engine/Private/Components/SceneComponent.cpp
+++ b/Engine/Source/Runtime/Engine/Private/Components/SceneComponent.cpp
@@ -75,9 +75,6 @@ USceneComponent::USceneComponent(const FObjectInitializer& ObjectInitializer /*=
 	// default behavior is visible
 	bVisible = true;
 	bAutoActivate = false;
-
-	bNetHasReceivedRelativeLocation = false;
-	bNetHasReceivedRelativeRotation = false;
 	bShouldBeAttached = AttachParent != nullptr;
 }
 
@@ -654,6 +651,8 @@ void USceneComponent::OnRegister()
 			AttachParent = nullptr;
 			AttachSocketName = NAME_None;
 			bShouldBeAttached = false;
+			bShouldSnapLocationWhenAttached = false;
+			bShouldSnapRotationWhenAttached = false;
 		}
 	}
 	
@@ -1777,6 +1776,9 @@ bool USceneComponent::AttachToComponent(USceneComponent* Parent, const FAttachme
 		ensureMsgf(!AttachmentRules.bWeldSimulatedBodies, TEXT("AttachToComponent when called from a constructor cannot weld simulated bodies. Consider calling SetupAttachment directly instead."));
 		ensureMsgf(AttachmentRules.LocationRule == EAttachmentRule::KeepRelative && AttachmentRules.RotationRule == EAttachmentRule::KeepRelative && AttachmentRules.ScaleRule == EAttachmentRule::KeepRelative, TEXT("AttachToComponent when called from a constructor is only setting up attachment and will always be treated as KeepRelative. Consider calling SetupAttachment directly instead."));
 		SetupAttachment(Parent, SocketName);
+		bShouldSnapLocationWhenAttached = false;
+		bShouldSnapRotationWhenAttached = false;
+
 		return true;
 	}
 
@@ -1941,6 +1943,9 @@ bool USceneComponent::AttachToComponent(USceneComponent* Parent, const FAttachme
 		AttachSocketName = SocketName;
 		bShouldBeAttached = AttachParent != nullptr;
 
+		bShouldSnapLocationWhenAttached = AttachmentRules.LocationRule == EAttachmentRule::SnapToTarget;
+		bShouldSnapRotationWhenAttached = AttachmentRules.RotationRule == EAttachmentRule::SnapToTarget;
+
 		OnAttachmentChanged();
 
 		// Preserve order of previous attachment if valid (in case we're doing a reattach operation inside a loop that might assume the AttachChildren order won't change)
@@ -2142,6 +2147,9 @@ void USceneComponent::DetachFromComponent(const FDetachmentTransformRules& Detac
 		AttachSocketName = NAME_None;
 		bShouldBeAttached = 0;
 
+		bShouldSnapLocationWhenAttached = false;
+		bShouldSnapRotationWhenAttached = false;
+
 		OnAttachmentChanged();
 
 		// If desired, update RelativeLocation and RelativeRotation to maintain current world position after detachment
@@ -3113,18 +3121,6 @@ void USceneComponent::OnRep_Transform()
 	bNetUpdateTransform = true;
 }
 
-void USceneComponent::OnRep_RelativeLocation()
-{
-	bNetHasReceivedRelativeLocation = true;
-	OnRep_Transform();
-}
-
-void USceneComponent::OnRep_RelativeRotation()
-{
-	bNetHasReceivedRelativeRotation = true;
-	OnRep_Transform();
-}
-
 void USceneComponent::OnRep_AttachParent()
 {
 	bNetUpdateAttachment = true;
@@ -3218,16 +3214,14 @@ void USceneComponent::PostRepNotifies()
 		Exchange(NetOldAttachParent, AttachParent);
 		Exchange(NetOldAttachSocketName, AttachSocketName);
 		
-		// Note: This is a local fix for JIRA UE-43355. If we receive bNetUpdateAttachment without having received a updated transform, assume that we intend to snap and that the relative location/rotation should defaulted
-		if (!bNetHasReceivedRelativeLocation)
+		// Note: This is a local fix for JIRA UE-43355.
+		if (bShouldSnapLocationWhenAttached)
 		{
 			RelativeLocation = FVector::ZeroVector;
-			bNetHasReceivedRelativeLocation = true;
 		}
-		if (!bNetHasReceivedRelativeRotation)
+		if (bShouldSnapRotationWhenAttached)
 		{
 			RelativeRotation = FRotator::ZeroRotator;
-			bNetHasReceivedRelativeRotation = true;
 		}
 
 		// Check if this is a detach
@@ -3239,11 +3233,15 @@ void USceneComponent::PostRepNotifies()
 		else
 		{
 			const uint8 bOldShouldBeAttached = bShouldBeAttached;
+			const uint8 bOldShouldSnapLocationWhenAttached = bShouldSnapLocationWhenAttached;
+			const uint8 bOldShouldSnapRotationWhenAttached = bShouldSnapRotationWhenAttached;
 
 			AttachToComponent(NetOldAttachParent, FAttachmentTransformRules::KeepRelativeTransform, NetOldAttachSocketName);
 
-			// restore bShouldBeAttached to what we have received
+			// restore to what we have received from the server
 			bShouldBeAttached = bOldShouldBeAttached;
+			bShouldSnapLocationWhenAttached = bOldShouldSnapLocationWhenAttached;
+			bShouldSnapRotationWhenAttached = bOldShouldSnapRotationWhenAttached;
 		}
 
 		bNetUpdateAttachment = false;
@@ -3265,6 +3263,8 @@ void USceneComponent::GetLifetimeReplicatedProps(TArray< FLifetimeProperty > & O
 	DOREPLIFETIME(USceneComponent, bAbsoluteScale);
 	DOREPLIFETIME(USceneComponent, bVisible);
 	DOREPLIFETIME(USceneComponent, bShouldBeAttached);
+	DOREPLIFETIME(USceneComponent, bShouldSnapLocationWhenAttached);
+	DOREPLIFETIME(USceneComponent, bShouldSnapRotationWhenAttached);
 	DOREPLIFETIME(USceneComponent, AttachParent);
 	DOREPLIFETIME(USceneComponent, AttachChildren);
 	DOREPLIFETIME(USceneComponent, AttachSocketName);
