ForceRelevantNextUpdate would only force relevancy to the first connection in the list.

diff --git a/Engine/Source/Runtime/Engine/Classes/Engine/NetDriver.h b/Engine/Source/Runtime/Engine/Classes/Engine/NetDriver.h
--- a/Engine/Source/Runtime/Engine/Classes/Engine/NetDriver.h
+++ b/Engine/Source/Runtime/Engine/Classes/Engine/NetDriver.h
@@ -1259,4 +1259,6 @@ private:
 
 	/** NetDriver time to end packet loss burst simulation. */
 	float PacketLossBurstEndTime;
+
+	TArray<FNetworkObjectInfo*> ProcessedForceRelevantActors;
 };
diff --git a/Engine/Source/Runtime/Engine/Private/NetworkDriver.cpp b/Engine/Source/Runtime/Engine/Private/NetworkDriver.cpp
--- a/Engine/Source/Runtime/Engine/Private/NetworkDriver.cpp
+++ b/Engine/Source/Runtime/Engine/Private/NetworkDriver.cpp
@@ -4017,7 +4017,10 @@ int32 UNetDriver::ServerReplicateActors_ProcessPrioritizedActors( UNetConnection
 			// if the actor is now relevant or was recently relevant
 			const bool bIsRecentlyRelevant = bIsRelevant || ( Channel && Time - Channel->RelevantTime < RelevantTimeout ) || ActorInfo->bForceRelevantNextUpdate;
 
-			ActorInfo->bForceRelevantNextUpdate = false;
+			if (ActorInfo->bForceRelevantNextUpdate)
+			{
+				ProcessedForceRelevantActors.AddUnique(ActorInfo);
+			}
 
 			if ( bIsRecentlyRelevant )
 			{
@@ -4460,6 +4463,12 @@ int32 UNetDriver::ServerReplicateActors(float DeltaSeconds)
 		}
 	}
 
+	for (FNetworkObjectInfo* ActorInfo : ProcessedForceRelevantActors)
+	{
+		ActorInfo->bForceRelevantNextUpdate = false;
+	}
+	ProcessedForceRelevantActors.Reset();
+
 	// shuffle the list of connections if not all connections were ticked
 	if (NumClientsToTick < ClientConnections.Num())
 	{
