Added call to ModulesChangedHandle.Reset() in DeinitializeModuleWatcher to fix issue with calling InitializeModuleWatcher again after a call to DeinitializeModuleWatcher
Fixed casts in SusystemCollection::GetSubsystemArray() causing compilation errors if used.

diff --git a/Engine/Source/Runtime/Engine/Private/Subsystems/SubsystemCollection.cpp b/Engine/Source/Runtime/Engine/Private/Subsystems/SubsystemCollection.cpp
--- a/Engine/Source/Runtime/Engine/Private/Subsystems/SubsystemCollection.cpp
+++ b/Engine/Source/Runtime/Engine/Private/Subsystems/SubsystemCollection.cpp
@@ -287,6 +287,7 @@ void FSubsystemModuleWatcher::DeinitializeModuleWatcher()
 	if (ModulesChangedHandle.IsValid())
 	{
 		FModuleManager::Get().OnModulesChanged().Remove(ModulesChangedHandle);
+		ModulesChangedHandle.Reset();
 	}
 }
 
diff --git a/Engine/Source/Runtime/Engine/Public/Subsystems/SubsystemCollection.h b/Engine/Source/Runtime/Engine/Public/Subsystems/SubsystemCollection.h
--- a/Engine/Source/Runtime/Engine/Public/Subsystems/SubsystemCollection.h
+++ b/Engine/Source/Runtime/Engine/Public/Subsystems/SubsystemCollection.h
@@ -91,7 +91,7 @@ public:
 		TSubclassOf<TBaseType> SubsystemBaseClass = SubsystemClass;
 
 		const TArray<USubsystem*>& Array = GetSubsystemArrayInternal(SubsystemBaseClass);
-		TArray<TSubsystemClass*>* SpecificArray = reinterpret_cast<TArray<TSubsystemClass*>*>(&Array);
+		const TArray<TSubsystemClass*>* SpecificArray = reinterpret_cast<const TArray<TSubsystemClass*>*>(&Array);
 		return *SpecificArray;
 	}
 
