Remove brackets as default for IPv6 and only adding them in if a port is present when creating ToString

diff --git a/Engine/Source/Runtime/Engine/Classes/Engine/EngineBaseTypes.h b/Engine/Source/Runtime/Engine/Classes/Engine/EngineBaseTypes.h
--- a/Engine/Source/Runtime/Engine/Classes/Engine/EngineBaseTypes.h
+++ b/Engine/Source/Runtime/Engine/Classes/Engine/EngineBaseTypes.h
@@ -853,6 +853,11 @@ struct ENGINE_API FURL
 	 */
 	FString ToString( bool FullyQualified=0 ) const;
 
+	/**
+	 * Prepares the Host and Port values into a standards compliant string
+	 */
+	FString GetHostPortString() const;
+
 	/**
 	 * Serializes a FURL to or from an archive.
 	 */
diff --git a/Engine/Source/Runtime/Engine/Private/URL.cpp b/Engine/Source/Runtime/Engine/Private/URL.cpp
--- a/Engine/Source/Runtime/Engine/Private/URL.cpp
+++ b/Engine/Source/Runtime/Engine/Private/URL.cpp
@@ -231,7 +231,7 @@ FURL::FURL( FURL* Base, const TCHAR* TextURL, ETravelType Type )
 		// Handle pure filenames & Posix paths.
 		bool FarHost=0;
 		bool FarMap=0;
-		if( FCString::Strlen(URL)>2 && ((URL[0] != '[' && URL[1]==':') || (URL[0]=='/' && !FPackageName::IsValidLongPackageName(URL, true))) )
+		if( FCString::Strlen(URL)>2 && ((URL[0] != '[' && URL[0] != ':' && URL[1]==':') || (URL[0]=='/' && !FPackageName::IsValidLongPackageName(URL, true))) )
 		{
 			// Pure filename.
 			Protocol = UrlConfig.DefaultProtocol;
@@ -245,16 +245,19 @@ FURL::FURL( FURL* Base, const TCHAR* TextURL, ETravelType Type )
 		else
 		{
 			// Determine location of the first opening square bracket.
-			// Square brackets enclose an IPv6 address.
+			// Square brackets enclose an IPv6 address if they have a port.
 			const TCHAR* SquareBracket = FCString::Strchr(URL, '[');
+			const TCHAR* FirstColonLocation = FCString::Strchr(URL, ':');
+			const TCHAR* LastColonLocation = FCString::Strrchr(URL, ':');
+			const bool bHasMultipleColons = FirstColonLocation != nullptr && LastColonLocation != nullptr && FirstColonLocation != LastColonLocation
+				&& FCString::Strchr(FirstColonLocation + 1, ':') != LastColonLocation;
 
 			// Parse protocol. Don't consider colons that occur after the opening square
 			// brace, because they are valid characters in an IPv6 address.
-			if
-			(	(FCString::Strchr(URL,':')!=NULL)
-			&&	(FCString::Strchr(URL,':')>URL+1)
-			&&  (!SquareBracket || (FCString::Strchr(URL,':')<SquareBracket))
-			&&	(FCString::Strchr(URL,'.')==NULL || FCString::Strchr(URL,':')<FCString::Strchr(URL,'.')) )
+			if( (FirstColonLocation !=nullptr)
+			&&	(FirstColonLocation >URL+1)
+			&&  ((!SquareBracket && !bHasMultipleColons) || (FirstColonLocation <SquareBracket) || (bHasMultipleColons && *(FirstColonLocation+1) == '/'))
+			&&	(FCString::Strchr(URL,'.')==nullptr || FirstColonLocation <FCString::Strchr(URL,'.')) )
 			{
 				TCHAR* ss = URL;
 				URL      = FCString::Strchr(URL,':');
@@ -272,6 +275,8 @@ FURL::FURL( FURL* Base, const TCHAR* TextURL, ETravelType Type )
 
 			// Parse optional host name and port.
 			const TCHAR* Dot = FCString::Strchr(URL,'.');
+			// This doesn't use FirstColonLocation as the URL could be modified by the Protocol code above
+			const TCHAR* Colon = FCString::Strchr(URL, ':');
 			const int32 ExtLen = FPackageName::GetMapPackageExtension().Len();
 
 			const bool bIsHostnameWithDot =
@@ -282,9 +287,11 @@ FURL::FURL( FURL* Base, const TCHAR* TextURL, ETravelType Type )
 				&&	(FCString::Strnicmp( Dot+1,TEXT("demo"), 4 ) != 0 || FChar::IsAlnum(Dot[5]));
 
 			// Square bracket indicates an IPv6 address, but IPv6 addresses can contain dots also
-			if ( bIsHostnameWithDot || SquareBracket )
+			// They also typically have multiple colons in them as well.
+			if (bIsHostnameWithDot || SquareBracket || (Colon && Colon != LastColonLocation) || bHasMultipleColons)
 			{
 				TCHAR* ss = URL;
+				// Clear out the URL such that all that should be left is the map
 				URL     = FCString::Strchr(URL,'/');
 				if( URL )
 				{
@@ -299,23 +306,27 @@ FURL::FURL( FURL* Base, const TCHAR* TextURL, ETravelType Type )
 				{
 					PortText = ClosingSquareBracket;
 				}
-				TCHAR* t = FCString::Strchr(PortText,':');
-				if( t )
+
+				// If we don't have a square bracket, check for another instance of a colon. This will tell us if we're dealing with an IPv6 address
+				// if there isn't one, then we're dealing with a port.
+				TCHAR* PortData = FCString::Strchr(PortText, ':');
+				if (PortData && (ClosingSquareBracket || (!ClosingSquareBracket && FCString::Strchr(PortData + 1, ':') == nullptr)))
 				{
 					// Port.
-					*(t++) = 0;
-					Port = FCString::Atoi( t );
+					*(PortData++) = 0;
+					Port = FCString::Atoi(PortData);
 				}
-/*
+
+				// If the input was an IPv6 address with a port, we need to remove the brackets then.
 				if(SquareBracket && ClosingSquareBracket)
 				{
 					// Trim the brackets from the host address
 					*ClosingSquareBracket = 0;
 					Host = ss + 1;
 				}
-				else*/
+				else
 				{
-					// Normal IPv4 address
+					// Otherwise, leave the address as is.
 					Host = ss;
 				}
 
@@ -437,17 +448,10 @@ FString FURL::ToString (bool FullyQualified) const
 		}
 	}
 
-	// Emit host.
+	// Emit host and port
 	if ((Host != UrlConfig.DefaultHost) || (Port != UrlConfig.DefaultPort))
 	{
-		Result += Host;
-
-		if (Port != UrlConfig.DefaultPort)
-		{
-			Result += TEXT(":");
-			Result += FString::Printf(TEXT("%i"), Port);
-		}
-
+		Result += GetHostPortString();
 		Result += TEXT("/");
 	}
 
@@ -474,6 +478,39 @@ FString FURL::ToString (bool FullyQualified) const
 	return Result;
 }
 
+//
+// Convert the host and port values of this URL into a string that's safe for serialization
+//
+FString FURL::GetHostPortString() const
+{
+	FString Result;
+	int32 FirstColonIndex, LastColonIndex;
+	bool bNotUsingDefaultPort = (Port != UrlConfig.DefaultPort);
+
+	// If this is an IPv6 address (determined if there's more than one colon) 
+	// and we're going to be adding the port, we need to put in the brackets
+	// This is done because there's no sane way to serialize the address otherwise
+	if (Host.FindChar(':', FirstColonIndex) && Host.FindLastChar(':', LastColonIndex)
+		&& FirstColonIndex != LastColonIndex && bNotUsingDefaultPort)
+	{
+		Result += FString::Printf(TEXT("[%s]"), *Host);
+	}
+	else
+	{
+		// Otherwise print the IPv6/IPv4 address as is
+		Result += Host;
+	}
+
+
+	if (bNotUsingDefaultPort)
+	{
+		Result += TEXT(":");
+		Result += FString::Printf(TEXT("%i"), Port);
+	}
+
+	return Result;
+}
+
 /*-----------------------------------------------------------------------------
 	Informational.
 -----------------------------------------------------------------------------*/
@@ -607,3 +644,97 @@ bool FURL::operator==( const FURL& Other ) const
 	return 1;
 }
 
+/*-----------------------------------------------------------------------------
+	Testing
+-----------------------------------------------------------------------------*/
+#if WITH_DEV_AUTOMATION_TESTS && !UE_BUILD_SHIPPING
+
+struct FURLTestCase
+{
+	FURLTestCase(const FString& InQueryString, const FString& InHost, const FString InProtocol=TEXT(""), int32 InPort = -1) :
+		QueryString(InQueryString),
+		Protocol(InProtocol),
+		Host(InHost),
+		Port(InPort)
+	{
+
+	}
+	FString QueryString;
+	FString Protocol;
+	FString Host;
+	int32 Port;
+};
+
+static bool URLSerializationTests(UWorld* InWorld, const TCHAR* Cmd, FOutputDevice& Ar)
+{
+	TArray<FURLTestCase> TestCases;
+	if (FParse::Command(&Cmd, TEXT("URLSERIALIZATION")))
+	{
+		TestCases.Push(FURLTestCase(TEXT("2001:0db8:85a3:0000:0000:8a2e:0370:7334"), TEXT("2001:0db8:85a3:0000:0000:8a2e:0370:7334")));
+		TestCases.Push(FURLTestCase(TEXT("[2001:0db8:85a3:0000:0000:8a2e:0370:7334]:7778"), TEXT("2001:0db8:85a3:0000:0000:8a2e:0370:7334"), TEXT(""), 7778));
+		TestCases.Push(FURLTestCase(TEXT("[2001:db8:85a3::8a2e:370:7334]"), TEXT("2001:db8:85a3::8a2e:370:7334")));
+		TestCases.Push(FURLTestCase(TEXT("epic://2001:db8:85a3::8a2e:370:7334"), TEXT("2001:db8:85a3::8a2e:370:7334"), TEXT("epic")));
+		TestCases.Push(FURLTestCase(TEXT("192.168.0.1"), TEXT("192.168.0.1")));
+		TestCases.Push(FURLTestCase(TEXT("test://192.168.0.1"), TEXT("192.168.0.1"), TEXT("test")));
+		TestCases.Push(FURLTestCase(TEXT("::ffff::192.168.0.1"), TEXT("::ffff::192.168.0.1")));
+		TestCases.Push(FURLTestCase(TEXT("[::ffff::192.168.0.1]"), TEXT("::ffff::192.168.0.1")));
+		TestCases.Push(FURLTestCase(TEXT("[::ffff::192.168.0.1]:7778"), TEXT("::ffff::192.168.0.1"), TEXT(""), 7778));
+		TestCases.Push(FURLTestCase(TEXT("test://::ffff::192.168.0.1"), TEXT("::ffff::192.168.0.1"), TEXT("test")));
+		TestCases.Push(FURLTestCase(TEXT("unreal:192.168.0.1:7776"), TEXT("192.168.0.1"), TEXT("unreal"), 7776));
+		TestCases.Push(FURLTestCase(TEXT("192.168.0.1"), TEXT("192.168.0.1")));
+		TestCases.Push(FURLTestCase(TEXT("http://[::ffff::192.168.0.1]:8080"), TEXT("::ffff::192.168.0.1"), TEXT("http"), 8080));
+		TestCases.Push(FURLTestCase(TEXT("https:[2001:db8:85a3::8a2e:370:7334]:443"), TEXT("2001:db8:85a3::8a2e:370:7334"), TEXT("https"), 443));
+		TestCases.Push(FURLTestCase(TEXT("steam.76561197993275299:20/"), TEXT("steam.76561197993275299"), TEXT(""), 20));
+
+		bool bAllCasesPassed = true;
+		for (const auto& TestCase : TestCases)
+		{
+			FURL TestURL(nullptr, *TestCase.QueryString, TRAVEL_Absolute);
+			// This is ugly, but we use it to report which cases failed and why.
+			const bool bHostMatched = (TestURL.Host == TestCase.Host);
+			const bool bPortMatched = (TestCase.Port == -1 || TestCase.Port == TestURL.Port);
+			const bool bProtocolMatched = (TestCase.Protocol.IsEmpty() || TestCase.Protocol == TestURL.Protocol);
+
+			if (bHostMatched && bPortMatched &&	bProtocolMatched)
+			{
+				UE_LOG(LogCore, Log, TEXT("Test %s passed!"), *TestCase.QueryString);
+			}
+			else
+			{
+				bAllCasesPassed = false;
+				UE_LOG(LogCore, Warning, TEXT("Test %s failed! Matching flags: Host[%d] Port[%d] Protocol[%d]"), *TestCase.QueryString, bHostMatched, bPortMatched, bProtocolMatched);
+				if (!bHostMatched)
+				{
+					UE_LOG(LogCore, Warning, TEXT("URL had host %s, expected %s"), *TestURL.Host, *TestCase.Host);
+				}
+
+				if (!bPortMatched)
+				{
+					UE_LOG(LogCore, Warning, TEXT("URL had port %d, expected %d"), TestURL.Port, TestCase.Port);
+				}
+
+				if (!bProtocolMatched)
+				{
+					UE_LOG(LogCore, Warning, TEXT("URL had protocol %s, expected %s"), *TestURL.Protocol, *TestCase.Protocol);
+				}
+			}
+		}
+
+		if (bAllCasesPassed)
+		{
+			UE_LOG(LogCore, Log, TEXT("All URL serialization cases passed."));
+		}
+		else
+		{
+			UE_LOG(LogCore, Warning, TEXT("An URL serialization case failed!"));
+		}
+
+		return true;
+	}
+
+	return false;
+}
+
+FStaticSelfRegisteringExec FURLTests(URLSerializationTests);
+#endif // WITH_DEV_AUTOMATION_TESTS && !UE_BUILD_SHIPPING
+
diff --git a/Engine/Source/Runtime/Engine/Private/World.cpp b/Engine/Source/Runtime/Engine/Private/World.cpp
--- a/Engine/Source/Runtime/Engine/Private/World.cpp
+++ b/Engine/Source/Runtime/Engine/Private/World.cpp
@@ -6436,7 +6436,7 @@ bool UWorld::UsesGameHiddenFlags() const
 
 FString UWorld::GetAddressURL() const
 {
-	return FString::Printf( TEXT("%s:%i"), *URL.Host, URL.Port );
+	return FString::Printf( TEXT("%s"), *URL.GetHostPortString() );
 }
 
 FString UWorld::RemovePIEPrefix(const FString &Source)
diff --git a/Engine/Source/Runtime/Sockets/Private/BSDIPv6Sockets/IPAddressBSDIPv6.cpp b/Engine/Source/Runtime/Sockets/Private/BSDIPv6Sockets/IPAddressBSDIPv6.cpp
--- a/Engine/Source/Runtime/Sockets/Private/BSDIPv6Sockets/IPAddressBSDIPv6.cpp
+++ b/Engine/Source/Runtime/Sockets/Private/BSDIPv6Sockets/IPAddressBSDIPv6.cpp
@@ -73,45 +73,32 @@ void FInternetAddrBSDIPv6::SetIp(const TCHAR* InAddr, bool& bIsValid)
 	bIsValid = false;
 
 	FString AddressString(InAddr);
+	FString Port;
 
-	const bool bHasOpenBracket = AddressString.Contains("[");
-	const int32 CloseBracketIndex = AddressString.Find("]");
-	const bool bHasCloseBracket = CloseBracketIndex != -1;
+	// Find some colons to try to determine the input given to us.
+	const int32 FirstColonIndex = AddressString.Find(":", ESearchCase::IgnoreCase, ESearchDir::FromStart);
+	const int32 LastColonIndex = AddressString.Find(":", ESearchCase::IgnoreCase, ESearchDir::FromEnd);
 
-	// IPv6 may or may not include open and close brackets.
-	// However, only IPv6 address can have them.
-	bool bIsIPv6 = bHasOpenBracket && bHasCloseBracket;
+	// IPv6 will always have at least 2 colons somewhere.
+	bool bIsIPv6 = FirstColonIndex != LastColonIndex;
 
-	// Valid IPv4 should not contain an open or close bracket
-	const bool bIsLikelyIPv4 = !bHasOpenBracket && !bHasCloseBracket;
+	// If we have no colons or one colon (IPv4 + Port).
+	bool bIsLikelyIPv4 = !bIsIPv6 || LastColonIndex != INDEX_NONE;
 
 	if (bIsLikelyIPv4 || bIsIPv6)
 	{
-		const int32 LastColonIndex = AddressString.Find(":", ESearchCase::IgnoreCase, ESearchDir::FromEnd);
-
-		// Double check to ensure this isn't actually an IPv6 address without brackets.
-		if (bIsLikelyIPv4)
-		{
-			// IPv4 addresses can contain at most 1 colon.
-			const int32 FirstColonIndex = AddressString.Find(":");
-			bIsIPv6 = (FirstColonIndex != LastColonIndex);
-		}
-
-		// IPv4 address will only have a port when a colon is present.
-		// IPv6 address will only have a port when surrounded by brackets.
-		const bool bHasPort = (INDEX_NONE != LastColonIndex) && (!bIsIPv6 || (bHasCloseBracket && LastColonIndex > CloseBracketIndex));
-		FString Port;
-		if (bHasPort)
+		// Check to see if we have a port.
+		if (AddressString.Contains("]:") || (!bIsIPv6 && LastColonIndex != INDEX_NONE))
 		{
 			Port = AddressString.RightChop(LastColonIndex + 1);
 			AddressString = AddressString.Left(LastColonIndex);
 		}
 
+		AddressString.RemoveFromStart("[");
+		AddressString.RemoveFromEnd("]");
+
 		if (bIsIPv6)
 		{
-			AddressString.RemoveFromStart("[");
-			AddressString.RemoveFromEnd("]");
-
 			// Check for valid IPv6 address
 			const auto InAddrAnsi = StringCast<ANSICHAR>(*AddressString);
 #if PLATFORM_IOS
@@ -146,7 +133,7 @@ void FInternetAddrBSDIPv6::SetIp(const TCHAR* InAddr, bool& bIsValid)
 				}
 		}
 
-		if (bHasPort)
+		if (!Port.IsEmpty())
 		{
 			SetPort(FCString::Atoi(*Port));
 		}
@@ -269,14 +256,11 @@ FString FInternetAddrBSDIPv6::ToString(bool bAppendPort) const
 
 	inet_ntop(AF_INET6, (void*)&Addr.sin6_addr, IPStr, INET6_ADDRSTRLEN);
 
-	FString Result("[");
-	Result += IPStr;
-	Result += "]";
+	FString Result(IPStr);
 
 	if (bAppendPort)
 	{
-		Result += ":";
-		Result += FString::Printf(TEXT("%d"), GetPort());
+		Result = FString::Printf(TEXT("[%s]:%d"), IPStr, GetPort());
 	}
 
 	return Result;
diff --git a/Engine/Source/Runtime/Sockets/Private/BSDSockets/IPAddressBSD.cpp b/Engine/Source/Runtime/Sockets/Private/BSDSockets/IPAddressBSD.cpp
--- a/Engine/Source/Runtime/Sockets/Private/BSDSockets/IPAddressBSD.cpp
+++ b/Engine/Source/Runtime/Sockets/Private/BSDSockets/IPAddressBSD.cpp
@@ -142,23 +142,19 @@ void FInternetAddrBSD::SetIp(const TCHAR* InAddr, bool& bIsValid)
 	FString AddressString(InAddr);
 	FString Port;
 
-	const bool bHasOpenBracket = AddressString.Contains("[");
-	const int32 CloseBracketIndex = AddressString.Find("]");
-	const bool bHasCloseBracket = CloseBracketIndex != INDEX_NONE;
-	bool bIsIPv6 = bHasOpenBracket && bHasCloseBracket;
-
+	// Find some colons to try to determine the input given to us.
+	const int32 FirstColonIndex = AddressString.Find(":", ESearchCase::IgnoreCase, ESearchDir::FromStart);
 	const int32 LastColonIndex = AddressString.Find(":", ESearchCase::IgnoreCase, ESearchDir::FromEnd);
 
-	// IPv4 address will only have a port when a colon is present.
-	// IPv6 address will only have a port when surrounded by brackets.
-	const bool bHasPort = (INDEX_NONE != LastColonIndex) && (!bIsIPv6 || (bHasCloseBracket && LastColonIndex > CloseBracketIndex));
-
-	if (bHasPort)
+	// Either this is an IPv6 address as we will have a "]:" when ports are involved
+	// or this is an IPv4 address as there's one colon and it exists in the string
+	if (AddressString.Contains("]:") || (FirstColonIndex == LastColonIndex && LastColonIndex != INDEX_NONE))
 	{
 		Port = AddressString.RightChop(LastColonIndex + 1);
 		AddressString = AddressString.Left(LastColonIndex);
 	}
 
+	// Strip these for backwards compatibility.
 	AddressString.RemoveFromStart("[");
 	AddressString.RemoveFromEnd("]");
 	
@@ -169,7 +165,7 @@ void FInternetAddrBSD::SetIp(const TCHAR* InAddr, bool& bIsValid)
 		{
 			bIsValid = true;
 			SetRawIp(NewAddr->GetRawIp());
-			if (bHasPort)
+			if (!Port.IsEmpty())
 			{
 				SetPort(FCString::Atoi(*Port));
 			}
@@ -446,7 +442,9 @@ FString FInternetAddrBSD::ToString(bool bAppendPort) const
 				IPv6Str = IPv6Str.Left(InterfaceMarkerIndex);
 			}
 
-			ReturnVal = FString::Printf(TEXT("[%s]"), *IPv6Str);
+			// Using dynamic formatting strings are deprecated.
+			ReturnVal = bAppendPort ? FString::Printf(TEXT("[%s]"), *IPv6Str) 
+				: FString::Printf(TEXT("%s"), *IPv6Str);
 		}
 		else
 		{
diff --git a/Engine/Source/Runtime/Sockets/Private/BSDSockets/IPAddressBSD.h b/Engine/Source/Runtime/Sockets/Private/BSDSockets/IPAddressBSD.h
--- a/Engine/Source/Runtime/Sockets/Private/BSDSockets/IPAddressBSD.h
+++ b/Engine/Source/Runtime/Sockets/Private/BSDSockets/IPAddressBSD.h
@@ -98,7 +98,7 @@ public:
 	 * Sets the ip address from a string IPv6 or IPv4 address.
 	 * Ports may be included using the form Address:Port, or excluded and set manually.
 	 *
-	 * IPv6 - [1111:2222:3333:4444:5555:6666:7777:8888]:PORT || [1111:2222:3333::]:PORT || [::ffff:IPv4]:PORT
+	 * IPv6 - [1111:2222:3333:4444:5555:6666:7777:8888]:PORT || [1111:2222:3333::]:PORT || [::ffff:IPv4]:PORT || any of these without the brackets and port
 	 * IPv4 - aaa.bbb.ccc.ddd:PORT || 127.0.0.1:PORT
 	 *
 	 * @param InAddr the string containing the new ip address to use
