Introducing MaxNetTickRate, default 120, to limit how often the net connection ticks. This allows very high FPS clients to still be able to replicate as normal. Without this the bandwidth settings would need to be bumped considerably.

diff --git a/Engine/Config/BaseEngine.ini b/Engine/Config/BaseEngine.ini
--- a/Engine/Config/BaseEngine.ini
+++ b/Engine/Config/BaseEngine.ini
@@ -873,6 +873,7 @@ RelevantTimeout=5.0
 SpawnPrioritySeconds=1.0
 ServerTravelPause=4.0
 NetServerMaxTickRate=30
+MaxNetTickRate=120
 NetConnectionClassName="/Script/OnlineSubsystemUtils.IpConnection"
 MaxPortCountToTry=512
 
diff --git a/Engine/Source/Runtime/Engine/Classes/Engine/NetDriver.h b/Engine/Source/Runtime/Engine/Classes/Engine/NetDriver.h
--- a/Engine/Source/Runtime/Engine/Classes/Engine/NetDriver.h
+++ b/Engine/Source/Runtime/Engine/Classes/Engine/NetDriver.h
@@ -634,6 +634,10 @@ public:
 	UPROPERTY(Config)
 	int32 NetServerMaxTickRate;
 
+	/** Limit tick rate of replication to allow very high frame rates to still replicate data. A value less or equal to zero means use the engine tick rate. A value greater than zero will clamp the net tick rate to this value.  */
+	UPROPERTY(Config)
+	int32 MaxNetTickRate;
+
 	/** @todo document */
 	UPROPERTY(Config)
 	int32 MaxInternetClientRate;
diff --git a/Engine/Source/Runtime/Engine/Private/NetConnection.cpp b/Engine/Source/Runtime/Engine/Private/NetConnection.cpp
--- a/Engine/Source/Runtime/Engine/Private/NetConnection.cpp
+++ b/Engine/Source/Runtime/Engine/Private/NetConnection.cpp
@@ -2701,6 +2701,18 @@ void UNetConnection::Tick()
 	}
 
 	FrameTime = CurrentRealtimeSeconds - LastTime;
+	const int32 MaxNetTickRate = Driver->MaxNetTickRate;
+	const float MaxNetTickRateFloat = MaxNetTickRate > 0 ? float(MaxNetTickRate) : FLT_MAX;
+	const float DesiredTickRate = FMath::Clamp(GEngine->GetMaxTickRate(0.0f, false), 0.0f, MaxNetTickRateFloat);
+	if (!InternalAck && MaxNetTickRate > 0 && DesiredTickRate > 0.0f)
+	{
+		const float MinNetFrameTime = 1.0f/DesiredTickRate;
+		if (FrameTime < MinNetFrameTime)
+		{
+			return;
+		}
+	}
+
 	LastTime = CurrentRealtimeSeconds;
 	CumulativeTime += FrameTime;
 	CountedFrames++;
@@ -2970,7 +2982,6 @@ void UNetConnection::Tick()
 
 	// Clamp DeltaTime for bandwidth limiting so that if there is a hitch, we don't try to send
 	// a large burst on the next frame, which can cause another hitch if a lot of additional replication occurs.
-	const float DesiredTickRate = GEngine->GetMaxTickRate(0.0f, false);
 	float BandwidthDeltaTime = DeltaTime;
 	if (DesiredTickRate != 0.0f)
 	{
