Fixed ensure in TResizableCircularQueue

diff --git a/Engine/Source/Runtime/Engine/Public/Net/Util/ResizableCircularQueue.h b/Engine/Source/Runtime/Engine/Public/Net/Util/ResizableCircularQueue.h
--- a/Engine/Source/Runtime/Engine/Public/Net/Util/ResizableCircularQueue.h
+++ b/Engine/Source/Runtime/Engine/Public/Net/Util/ResizableCircularQueue.h
@@ -129,7 +129,7 @@ template<typename T, typename AllocatorT>
 void
 TResizableCircularQueue<T, AllocatorT>::Pop(SIZE_T PopCount)
 {
-	if (ensure(Count() > PopCount))
+	if (ensure(Count() >= PopCount))
 	{
 		PopNoCheck(PopCount);
 	}
