Fixing a typo I made when writing the unit tests for v6 mapped addresses.

#rnx

diff --git a/Engine/Source/Runtime/Engine/Private/URL.cpp b/Engine/Source/Runtime/Engine/Private/URL.cpp
--- a/Engine/Source/Runtime/Engine/Private/URL.cpp
+++ b/Engine/Source/Runtime/Engine/Private/URL.cpp
@@ -676,13 +676,13 @@ static bool URLSerializationTests(UWorld* InWorld, const TCHAR* Cmd, FOutputDevi
 		TestCases.Push(FURLTestCase(TEXT("epic://2001:db8:85a3::8a2e:370:7334"), TEXT("2001:db8:85a3::8a2e:370:7334"), TEXT("epic")));
 		TestCases.Push(FURLTestCase(TEXT("192.168.0.1"), TEXT("192.168.0.1")));
 		TestCases.Push(FURLTestCase(TEXT("test://192.168.0.1"), TEXT("192.168.0.1"), TEXT("test")));
-		TestCases.Push(FURLTestCase(TEXT("::ffff::192.168.0.1"), TEXT("::ffff::192.168.0.1")));
-		TestCases.Push(FURLTestCase(TEXT("[::ffff::192.168.0.1]"), TEXT("::ffff::192.168.0.1")));
-		TestCases.Push(FURLTestCase(TEXT("[::ffff::192.168.0.1]:7778"), TEXT("::ffff::192.168.0.1"), TEXT(""), 7778));
-		TestCases.Push(FURLTestCase(TEXT("test://::ffff::192.168.0.1"), TEXT("::ffff::192.168.0.1"), TEXT("test")));
+		TestCases.Push(FURLTestCase(TEXT("::ffff:192.168.0.1"), TEXT("::ffff:192.168.0.1")));
+		TestCases.Push(FURLTestCase(TEXT("[::ffff:192.168.0.1]"), TEXT("::ffff:192.168.0.1")));
+		TestCases.Push(FURLTestCase(TEXT("[::ffff:192.168.0.1]:7778"), TEXT("::ffff:192.168.0.1"), TEXT(""), 7778));
+		TestCases.Push(FURLTestCase(TEXT("test://::ffff:192.168.0.1"), TEXT("::ffff:192.168.0.1"), TEXT("test")));
 		TestCases.Push(FURLTestCase(TEXT("unreal:192.168.0.1:7776"), TEXT("192.168.0.1"), TEXT("unreal"), 7776));
 		TestCases.Push(FURLTestCase(TEXT("192.168.0.1"), TEXT("192.168.0.1")));
-		TestCases.Push(FURLTestCase(TEXT("http://[::ffff::192.168.0.1]:8080"), TEXT("::ffff::192.168.0.1"), TEXT("http"), 8080));
+		TestCases.Push(FURLTestCase(TEXT("http://[::ffff:192.168.0.1]:8080"), TEXT("::ffff:192.168.0.1"), TEXT("http"), 8080));
 		TestCases.Push(FURLTestCase(TEXT("https:[2001:db8:85a3::8a2e:370:7334]:443"), TEXT("2001:db8:85a3::8a2e:370:7334"), TEXT("https"), 443));
 		TestCases.Push(FURLTestCase(TEXT("steam.76561197993275299:20/"), TEXT("steam.76561197993275299"), TEXT(""), 20));
 
