PR #5228: Fix net driver cleanup (Contributed by [REDACTED_PERSON])

diff --git a/Engine/Source/Runtime/Engine/Private/World.cpp b/Engine/Source/Runtime/Engine/Private/World.cpp
--- a/Engine/Source/Runtime/Engine/Private/World.cpp
+++ b/Engine/Source/Runtime/Engine/Private/World.cpp
@@ -5139,7 +5139,7 @@ bool UWorld::Listen( FURL& InURL )
 		}
 	}
 
-	if (NetDriver == NULL)
+	if (NetDriver == nullptr)
 	{
 		GEngine->BroadcastNetworkFailure(this, NULL, ENetworkFailure::NetDriverCreateFailure);
 		return false;
@@ -5150,8 +5150,8 @@ bool UWorld::Listen( FURL& InURL )
 	{
 		GEngine->BroadcastNetworkFailure(this, NetDriver, ENetworkFailure::NetDriverListenFailure, Error);
 		UE_LOG(LogWorld, Log,  TEXT("Failed to listen: %s"), *Error );
-		NetDriver->SetWorld(NULL);
-		NetDriver = NULL;
+		GEngine->DestroyNamedNetDriver(this, NetDriver->NetDriverName);
+		NetDriver = nullptr;
 		FLevelCollection* SourceCollection = FindCollectionByType(ELevelCollectionType::DynamicSourceLevels);
 		if (SourceCollection)
 		{
