Calling ForceRelevant on dormant static actors will now force the actor to be processed until all connections have replicated the actor.
Deprecated the FNetworkObjectInfo::bForceRelevantNextUpdate variable and relying on the new ForceRelevantFrame variable instead

#rn fix Network/Engine

diff --git a/Engine/Source/Runtime/Engine/Classes/Engine/NetConnection.h b/Engine/Source/Runtime/Engine/Classes/Engine/NetConnection.h
--- a/Engine/Source/Runtime/Engine/Classes/Engine/NetConnection.h
+++ b/Engine/Source/Runtime/Engine/Classes/Engine/NetConnection.h
@@ -368,6 +368,7 @@ public:
 	double			LastTickTime;			// Last time of polling.
 	int32			QueuedBits;			// Bits assumed to be queued up.
 	int32			TickCount;				// Count of ticks.
+	uint32			LastProcessedFrame;   // The last frame where we gathered and processed actors for this connection
 	/** The last time an ack was received */
 	float			LastRecvAckTime;
 	/** Time when connection request was first initiated */
@@ -932,7 +933,8 @@ public:
 	void PurgeAcks();
 
 	/** Send package map to the remote. */
-	void SendPackageMap();
+	UE_DEPRECATED(4.23, "This method will be removed.")
+	void SendPackageMap() {}
 
 	/** 
 	 * Appends the passed in data to the SendBuffer to be sent when FlushNet is called
diff --git a/Engine/Source/Runtime/Engine/Classes/Engine/NetDriver.h b/Engine/Source/Runtime/Engine/Classes/Engine/NetDriver.h
--- a/Engine/Source/Runtime/Engine/Classes/Engine/NetDriver.h
+++ b/Engine/Source/Runtime/Engine/Classes/Engine/NetDriver.h
@@ -1571,6 +1571,4 @@ private:
 
 	/** NetDriver time to end packet loss burst simulation. */
 	float PacketLossBurstEndTime;
-
-	TArray<FNetworkObjectInfo*> ProcessedForceRelevantActors;
 };
diff --git a/Engine/Source/Runtime/Engine/Classes/Engine/NetworkObjectList.h b/Engine/Source/Runtime/Engine/Classes/Engine/NetworkObjectList.h
--- a/Engine/Source/Runtime/Engine/Classes/Engine/NetworkObjectList.h
+++ b/Engine/Source/Runtime/Engine/Classes/Engine/NetworkObjectList.h
@@ -46,11 +46,15 @@ struct FNetworkObjectInfo
 	uint8 bPendingNetUpdate : 1;
 
 	/** Force this object to be considered relevant for at least one update */
+	UE_DEPRECATED(4.23, "Use the ForceRelevantFrame variable since this variable is not set anymore.")
 	uint8 bForceRelevantNextUpdate : 1;
 
 	/** Should this object be considered for replay checkpoint writes */
 	uint8 bDirtyForReplay : 1;
 
+	/** Force this object to be considered relevant for at least one update */
+	uint32 ForceRelevantFrame = 0;
+
 	FNetworkObjectInfo()
 		: Actor(nullptr)
 		, NextUpdateTime(0.0)
diff --git a/Engine/Source/Runtime/Engine/Private/DemoNetDriver.cpp b/Engine/Source/Runtime/Engine/Private/DemoNetDriver.cpp
--- a/Engine/Source/Runtime/Engine/Private/DemoNetDriver.cpp
+++ b/Engine/Source/Runtime/Engine/Private/DemoNetDriver.cpp
@@ -2656,9 +2656,7 @@ void UDemoNetDriver::TickDemoRecordFrame(float DeltaSeconds)
 					// We check ActorInfo->LastNetUpdateTime < KINDA_SMALL_NUMBER to force at least one update for each actor
 					const bool bWasRecentlyRelevant = (ActorInfo->LastNetUpdateTime < KINDA_SMALL_NUMBER) || ((Time - ActorInfo->LastNetUpdateTime) < RelevantTimeout);
 
-					bool bIsRelevant = !bUseNetRelevancy || Actor->bAlwaysRelevant || Actor == ClientConnection->PlayerController || ActorInfo->bForceRelevantNextUpdate;
-
-					ActorInfo->bForceRelevantNextUpdate = false;
+					bool bIsRelevant = !bUseNetRelevancy || Actor->bAlwaysRelevant || Actor == ClientConnection->PlayerController || (ActorInfo->ForceRelevantFrame >= ReplicationFrame);
 
 					if (!bIsRelevant)
 					{
diff --git a/Engine/Source/Runtime/Engine/Private/NetConnection.cpp b/Engine/Source/Runtime/Engine/Private/NetConnection.cpp
--- a/Engine/Source/Runtime/Engine/Private/NetConnection.cpp
+++ b/Engine/Source/Runtime/Engine/Private/NetConnection.cpp
@@ -116,6 +116,7 @@ UNetConnection::UNetConnection(const FObjectInitializer& ObjectInitializer)
 
 ,	QueuedBits			( 0 )
 ,	TickCount			( 0 )
+,	LastProcessedFrame	( 0 )
 ,	ConnectTime			( 0.0 )
 
 ,	AllowMerge			( false )
@@ -842,9 +843,6 @@ void UNetConnection::AssertValid()
 	check(State==USOCK_Closed || State==USOCK_Pending || State==USOCK_Open);
 
 }
-void UNetConnection::SendPackageMap()
-{
-}
 
 bool UNetConnection::ClientHasInitializedLevelFor(const AActor* TestActor) const
 {
diff --git a/Engine/Source/Runtime/Engine/Private/NetworkDriver.cpp b/Engine/Source/Runtime/Engine/Private/NetworkDriver.cpp
--- a/Engine/Source/Runtime/Engine/Private/NetworkDriver.cpp
+++ b/Engine/Source/Runtime/Engine/Private/NetworkDriver.cpp
@@ -4023,12 +4023,7 @@ int32 UNetDriver::ServerReplicateActors_ProcessPrioritizedActors( UNetConnection
 			}
 
 			// if the actor is now relevant or was recently relevant
-			const bool bIsRecentlyRelevant = bIsRelevant || ( Channel && Time - Channel->RelevantTime < RelevantTimeout ) || ActorInfo->bForceRelevantNextUpdate;
-
-			if (ActorInfo->bForceRelevantNextUpdate)
-			{
-				ProcessedForceRelevantActors.AddUnique(ActorInfo);
-			}
+			const bool bIsRecentlyRelevant = bIsRelevant || ( Channel && Time - Channel->RelevantTime < RelevantTimeout ) || (ActorInfo->ForceRelevantFrame >= Connection->LastProcessedFrame);
 
 			if ( bIsRecentlyRelevant )
 			{
@@ -4467,22 +4462,24 @@ int32 UNetDriver::ServerReplicateActors(float DeltaSeconds)
 						Channel->RelevantTime = Time + 0.5f * FMath::SRand();
 					}
 				}
+
+				// If the actor was forced to relevant and didn't get processed, try again on the next update;
+				if (PriorityActors[k]->ActorInfo->ForceRelevantFrame >= Connection->LastProcessedFrame)
+				{
+					PriorityActors[k]->ActorInfo->ForceRelevantFrame = ReplicationFrame+1;
+				}
 			}
 			RelevantActorMark.Pop();
 
 			ConnectionViewers.Reset();
 
+			Connection->LastProcessedFrame = ReplicationFrame;
+
 			const bool bWasSaturated = GNumSaturatedConnections > LocalNumSaturated;
 			Connection->TrackReplicationForAnalytics(bWasSaturated);
 		}
 	}
 
-	for (FNetworkObjectInfo* ActorInfo : ProcessedForceRelevantActors)
-	{
-		ActorInfo->bForceRelevantNextUpdate = false;
-	}
-	ProcessedForceRelevantActors.Reset();
-
 	// shuffle the list of connections if not all connections were ticked
 	if (NumClientsToTick < ClientConnections.Num())
 	{
diff --git a/Engine/Source/Runtime/Engine/Private/NetworkObjectList.cpp b/Engine/Source/Runtime/Engine/Private/NetworkObjectList.cpp
--- a/Engine/Source/Runtime/Engine/Private/NetworkObjectList.cpp
+++ b/Engine/Source/Runtime/Engine/Private/NetworkObjectList.cpp
@@ -327,8 +327,8 @@ void FNetworkObjectList::ForceActorRelevantNextUpdate(AActor* const Actor, UNetD
 	}
 
 	FNetworkObjectInfo* NetworkObjectInfo = NetworkObjectInfoPtr->Get();
-	
-	NetworkObjectInfo->bForceRelevantNextUpdate = true;
+
+	NetworkObjectInfo->ForceRelevantFrame = NetDriver->ReplicationFrame + 1;
 }
 
 void FNetworkObjectList::Reset()
diff --git a/Engine/Source/Runtime/Engine/Private/World.cpp b/Engine/Source/Runtime/Engine/Private/World.cpp
--- a/Engine/Source/Runtime/Engine/Private/World.cpp
+++ b/Engine/Source/Runtime/Engine/Private/World.cpp
@@ -4572,7 +4572,6 @@ void UWorld::WelcomePlayer(UNetConnection* Connection)
 #endif
 
 	check(CurrentLevel);
-	Connection->SendPackageMap();
 
 	FString LevelName;
 
