Added a check for sequence length being greater than 0 when checking if notifies are available.

diff --git a/Engine/Source/Runtime/Engine/Private/Animation/AnimComposite.cpp b/Engine/Source/Runtime/Engine/Private/Animation/AnimComposite.cpp
--- a/Engine/Source/Runtime/Engine/Private/Animation/AnimComposite.cpp
+++ b/Engine/Source/Runtime/Engine/Private/Animation/AnimComposite.cpp
@@ -26,7 +26,7 @@ void UAnimComposite::ReplaceReferredAnimations(const TMap<UAnimationAsset*, UAni
 
 bool UAnimComposite::IsNotifyAvailable() const
 {
-	return (Super::IsNotifyAvailable() || AnimationTrack.IsNotifyAvailable());
+	return (SequenceLength > 0.f && (Super::IsNotifyAvailable() || AnimationTrack.IsNotifyAvailable()));
 }
 
 void UAnimComposite::GetAnimNotifiesFromDeltaPositions(const float& PreviousPosition, const float & CurrentPosition, TArray<FAnimNotifyEventReference>& OutActiveNotifies) const
