Made AnimNotifyState_TimedParticleEffect to be extensible as requested by licensee

diff --git a/Engine/Source/Runtime/Engine/Classes/Animation/AnimNotifies/AnimNotifyState_TimedParticleEffect.h b/Engine/Source/Runtime/Engine/Classes/Animation/AnimNotifies/AnimNotifyState_TimedParticleEffect.h
--- a/Engine/Source/Runtime/Engine/Classes/Animation/AnimNotifies/AnimNotifyState_TimedParticleEffect.h
+++ b/Engine/Source/Runtime/Engine/Classes/Animation/AnimNotifies/AnimNotifyState_TimedParticleEffect.h
@@ -13,8 +13,8 @@ class USkeletalMeshComponent;
 // Timed Particle Effect Notify
 // Allows a looping particle effect to be played in an animation that will activate
 // at the beginning of the notify and deactivate at the end.
-UCLASS(MinimalAPI, Blueprintable, meta = (DisplayName = "Timed Particle Effect"))
-class UAnimNotifyState_TimedParticleEffect : public UAnimNotifyState
+UCLASS(Blueprintable, meta = (DisplayName = "Timed Particle Effect"))
+class ENGINE_API UAnimNotifyState_TimedParticleEffect : public UAnimNotifyState
 {
 	GENERATED_UCLASS_BODY()
 
@@ -64,6 +64,6 @@ class UAnimNotifyState_TimedParticleEffect : public UAnimNotifyState
 	// Overridden from UAnimNotifyState to provide custom notify name.
 	FString GetNotifyName_Implementation() const override;
 
-private:
+protected:
 	bool ValidateParameters(USkeletalMeshComponent* MeshComp);
 };
