Fix issue with thread of default curve UID list access

diff --git a/Engine/Source/Runtime/Engine/Classes/Components/SkeletalMeshComponent.h b/Engine/Source/Runtime/Engine/Classes/Components/SkeletalMeshComponent.h
--- a/Engine/Source/Runtime/Engine/Classes/Components/SkeletalMeshComponent.h
+++ b/Engine/Source/Runtime/Engine/Classes/Components/SkeletalMeshComponent.h
@@ -1221,6 +1221,11 @@ private:
 	 **/
 	TMap<FName, float>	MorphTargetCurves;
 
+	/** 
+	 * Temporary storage for Curve UIDList of evaluating Animation 
+	 */
+	TArray<uint16> CachedCurveUIDList;
+
 public:
 	const TMap<FName, float>& GetMorphTargetCurves() const { return MorphTargetCurves;  }
 	// 
diff --git a/Engine/Source/Runtime/Engine/Private/Components/SkeletalMeshComponent.cpp b/Engine/Source/Runtime/Engine/Private/Components/SkeletalMeshComponent.cpp
--- a/Engine/Source/Runtime/Engine/Private/Components/SkeletalMeshComponent.cpp
+++ b/Engine/Source/Runtime/Engine/Private/Components/SkeletalMeshComponent.cpp
@@ -1468,6 +1468,11 @@ void USkeletalMeshComponent::RecalcRequiredCurves()
 		return;
 	}
 
+	if (SkeletalMesh->Skeleton)
+	{
+		CachedCurveUIDList = SkeletalMesh->Skeleton->GetDefaultCurveUIDList();
+	}
+
 	const FCurveEvaluationOption CurveEvalOption(bAllowAnimCurveEvaluation, &DisallowedAnimCurves, PredictedLODLevel);
 
 	// make sure animation requiredcurve to mark as dirty
@@ -1728,11 +1733,7 @@ void USkeletalMeshComponent::EvaluateAnimation(const USkeletalMesh* InSkeletalMe
 	}
 	else
 	{
-		// unfortunately it's possible they might not have skeleton, in that case, we don't have any place to copy the curve from
-		if (InSkeletalMesh->Skeleton)
-		{
-			OutCurve.InitFrom(&InSkeletalMesh->Skeleton->GetDefaultCurveUIDList());
-		}
+		OutCurve.InitFrom(&CachedCurveUIDList);
 	}
 }
 
