One more fix for montage hardening

diff --git a/Engine/Source/Runtime/Engine/Private/Animation/AnimNode_SequencePlayer.cpp b/Engine/Source/Runtime/Engine/Private/Animation/AnimNode_SequencePlayer.cpp
--- a/Engine/Source/Runtime/Engine/Private/Animation/AnimNode_SequencePlayer.cpp
+++ b/Engine/Source/Runtime/Engine/Private/Animation/AnimNode_SequencePlayer.cpp
@@ -33,7 +33,7 @@ void FAnimNode_SequencePlayer::Initialize_AnyThread(const FAnimationInitializeCo
 
 	GetEvaluateGraphExposedInputs().Execute(Context);
 
-	if(ensureMsgf(!Sequence->IsA<UAnimMontage>(), TEXT("Sequence players do not support anim montages.")))
+	if(Sequence && !ensureMsgf(!Sequence->IsA<UAnimMontage>(), TEXT("Sequence players do not support anim montages.")))
 	{
 		Sequence = nullptr;
 	}
@@ -61,7 +61,7 @@ void FAnimNode_SequencePlayer::UpdateAssetPlayer(const FAnimationUpdateContext&
 {
 	GetEvaluateGraphExposedInputs().Execute(Context);
 
-	if(ensureMsgf(!Sequence->IsA<UAnimMontage>(), TEXT("Sequence players do not support anim montages.")))
+	if(Sequence && !ensureMsgf(!Sequence->IsA<UAnimMontage>(), TEXT("Sequence players do not support anim montages.")))
 	{
 		Sequence = nullptr;
 	}
