ensure bonespace transforms are not accessible or deprecate it somewhere

diff --git a/Engine/Source/Editor/Persona/Private/AnimationEditorViewportClient.cpp b/Engine/Source/Editor/Persona/Private/AnimationEditorViewportClient.cpp
--- a/Engine/Source/Editor/Persona/Private/AnimationEditorViewportClient.cpp
+++ b/Engine/Source/Editor/Persona/Private/AnimationEditorViewportClient.cpp
@@ -862,11 +862,12 @@ FText FAnimationViewportClient::GetDisplayInfo(bool bDisplayAllInfo) const
 				NumSectionsInUse
 				));
 
+			TArray<FTransform> LocalBoneTransforms = PreviewMeshComponent->GetBoneSpaceTransforms();
 			if (PreviewMeshComponent->BonesOfInterest.Num() > 0)
 			{
 				int32 BoneIndex = PreviewMeshComponent->BonesOfInterest[0];
 				FTransform ReferenceTransform = PreviewMeshComponent->GetReferenceSkeleton().GetRefBonePose()[BoneIndex];
-				FTransform LocalTransform = PreviewMeshComponent->BoneSpaceTransforms[BoneIndex];
+				FTransform LocalTransform = LocalBoneTransforms[BoneIndex];
 				FTransform ComponentTransform = PreviewMeshComponent->GetDrawTransform(BoneIndex);
 
 				TextValue = ConcatenateLine(TextValue, FText::Format(LOCTEXT("LocalTransform", "Local: {0}"), FText::FromString(LocalTransform.ToHumanReadableString())));
diff --git a/Engine/Source/Editor/SkeletonEditor/Private/BoneProxy.cpp b/Engine/Source/Editor/SkeletonEditor/Private/BoneProxy.cpp
--- a/Engine/Source/Editor/SkeletonEditor/Private/BoneProxy.cpp
+++ b/Engine/Source/Editor/SkeletonEditor/Private/BoneProxy.cpp
@@ -26,10 +26,12 @@ void UBoneProxy::Tick(float DeltaTime)
 	{
 		if (UDebugSkelMeshComponent* Component = SkelMeshComponent.Get())
 		{
+			TArray<FTransform> LocalBoneTransforms = Component->GetBoneSpaceTransforms();
+
 			int32 BoneIndex = Component->GetBoneIndex(BoneName);
-			if (Component->BoneSpaceTransforms.IsValidIndex(BoneIndex))
+			if (LocalBoneTransforms.IsValidIndex(BoneIndex))
 			{
-				FTransform LocalTransform = Component->BoneSpaceTransforms[BoneIndex];
+				FTransform LocalTransform = LocalBoneTransforms[BoneIndex];
 				FTransform BoneTransform = Component->GetBoneTransform(BoneIndex);
 
 				if (bLocalLocation)
diff --git a/Engine/Source/Editor/UnrealEd/Private/Fbx/FbxAnimationExport.cpp b/Engine/Source/Editor/UnrealEd/Private/Fbx/FbxAnimationExport.cpp
--- a/Engine/Source/Editor/UnrealEd/Private/Fbx/FbxAnimationExport.cpp
+++ b/Engine/Source/Editor/UnrealEd/Private/Fbx/FbxAnimationExport.cpp
@@ -505,6 +505,7 @@ void FFbxExporter::ExportAnimTrack(IAnimTrackAdapter& AnimTrackAdapter, AActor*
 	}
 
 	FTransform InitialInvParentTransform;
+	TArray<FTransform> LocalBoneTransforms = SkeletalMeshComponent->GetBoneSpaceTransforms();
 
 	float SampleTime;
 	for(SampleTime = AnimationStart; SampleTime <= AnimationEnd; SampleTime += SamplingRate)
@@ -560,7 +561,7 @@ void FFbxExporter::ExportAnimTrack(IAnimTrackAdapter& AnimTrackAdapter, AActor*
 				Curves[i]->KeyModifyBegin();
 			}
 
-			FTransform BoneTransform = SkeletalMeshComponent->BoneSpaceTransforms[BoneIndex];
+			FTransform BoneTransform = LocalBoneTransforms[BoneIndex];
 
 			if (GetExportOptions()->MapSkeletalMotionToRoot && BoneIndex == 0)
 			{
diff --git a/Engine/Source/Runtime/Engine/Classes/Components/PoseableMeshComponent.h b/Engine/Source/Runtime/Engine/Classes/Components/PoseableMeshComponent.h
--- a/Engine/Source/Runtime/Engine/Classes/Components/PoseableMeshComponent.h
+++ b/Engine/Source/Runtime/Engine/Classes/Components/PoseableMeshComponent.h
@@ -18,8 +18,15 @@ class ENGINE_API UPoseableMeshComponent : public USkinnedMeshComponent
 {
 	GENERATED_UCLASS_BODY()
 
+	// we need this for template function in the cpp
+	// skeletalmeshcomponent hides this property, so now creating extra getter for poseablemeshcomponent
+	// although it's fine for poseablemeshcomponent to have this as external
+	// note that their signature is different, but the template would workf ine
+	const TArray<FTransform>& GetBoneSpaceTransforms() const { return BoneSpaceTransforms; }
+
 	/** Temporary array of local-space (ie relative to parent bone) rotation/translation/scale for each bone. */
 	TArray<FTransform> BoneSpaceTransforms;
+public:
 
 	FBoneContainer RequiredBones;
 
@@ -51,7 +58,7 @@ class ENGINE_API UPoseableMeshComponent : public USkinnedMeshComponent
 	void ResetBoneTransformByName(FName BoneName);
 
 	UFUNCTION(BlueprintCallable, Category="Components|PoseableMesh")
-	void CopyPoseFromSkeletalComponent(const USkeletalMeshComponent* InComponentToCopy);
+	void CopyPoseFromSkeletalComponent(USkeletalMeshComponent* InComponentToCopy);
 
 	//~ Begin USkinnedMeshComponent Interface
 	virtual void RefreshBoneTransforms(FActorComponentTickFunction* TickFunction = NULL) override;
diff --git a/Engine/Source/Runtime/Engine/Classes/Components/SkeletalMeshComponent.h b/Engine/Source/Runtime/Engine/Classes/Components/SkeletalMeshComponent.h
--- a/Engine/Source/Runtime/Engine/Classes/Components/SkeletalMeshComponent.h
+++ b/Engine/Source/Runtime/Engine/Classes/Components/SkeletalMeshComponent.h
@@ -339,9 +339,18 @@ public:
 	UPROPERTY(EditAnywhere, BlueprintReadWrite, Category=Animation, meta=(ShowOnlyInnerProperties))
 	struct FSingleAnimationPlayData AnimationData;
 
-	/** Temporary array of local-space (relative to parent bone) rotation/translation for each bone. */
+	// this is explicit copy because this buffer is reused during evaluation
+	// we want to have reference and emptied during evaluation
+	TArray<FTransform> GetBoneSpaceTransforms();
+
+	/** 
+	 * Temporary array of local-space (relative to parent bone) rotation/translation for each bone. 
+	 * This property is not safe to access during evaluation, so we created wrapper.
+	 */
+	UE_DEPRECATED(4.23, "Direct access to this property is deprecated, please use GetBoneSpaceTransforms instead. We will move to private in the future.")
 	TArray<FTransform> BoneSpaceTransforms;
-	
+
+public:
 	/** Offset of the root bone from the reference pose. Used to offset bounding box. */
 	UPROPERTY(transient)
 	FVector RootBoneTranslation;
diff --git a/Engine/Source/Runtime/Engine/Private/Animation/AnimSequence.cpp b/Engine/Source/Runtime/Engine/Private/Animation/AnimSequence.cpp
--- a/Engine/Source/Runtime/Engine/Private/Animation/AnimSequence.cpp
+++ b/Engine/Source/Runtime/Engine/Private/Animation/AnimSequence.cpp
@@ -5184,7 +5184,7 @@ bool UAnimSequence::CreateAnimation(USkeletalMeshComponent* MeshComponent)
 		RawAnimationData.AddZeroed(NumBones);
 		AnimationTrackNames.AddUninitialized(NumBones);
 
-		const TArray<FTransform>& BoneSpaceTransforms = MeshComponent->BoneSpaceTransforms;
+		const TArray<FTransform> BoneSpaceTransforms = MeshComponent->GetBoneSpaceTransforms();
 
 		check(BoneSpaceTransforms.Num() >= NumBones);
 
diff --git a/Engine/Source/Runtime/Engine/Private/Components/PoseableMeshComponent.cpp b/Engine/Source/Runtime/Engine/Private/Components/PoseableMeshComponent.cpp
--- a/Engine/Source/Runtime/Engine/Private/Components/PoseableMeshComponent.cpp
+++ b/Engine/Source/Runtime/Engine/Private/Components/PoseableMeshComponent.cpp
@@ -202,7 +202,7 @@ FTransform GetBoneTransformByNameHelper(FName BoneName, EBoneSpaces::Type BoneSp
 	}*/
 
 	FA2CSPose CSPose;
-	CSPose.AllocateLocalPoses(RequiredBones, Component->BoneSpaceTransforms);
+	CSPose.AllocateLocalPoses(RequiredBones, Component->GetBoneSpaceTransforms());
 
 	if (BoneSpace == EBoneSpaces::ComponentSpace)
 	{
@@ -278,16 +278,16 @@ void UPoseableMeshComponent::ResetBoneTransformByName(FName BoneName)
 	}
 }
 
-void UPoseableMeshComponent::CopyPoseFromSkeletalComponent(const USkeletalMeshComponent* InComponentToCopy)
+void UPoseableMeshComponent::CopyPoseFromSkeletalComponent(USkeletalMeshComponent* InComponentToCopy)
 {
 	if(RequiredBones.IsValid())
 	{
-		if(this->SkeletalMesh == InComponentToCopy->SkeletalMesh)
+		TArray<FTransform> LocalTransforms = InComponentToCopy->GetBoneSpaceTransforms();
+		if(this->SkeletalMesh == InComponentToCopy->SkeletalMesh 
+			&& LocalTransforms.Num() == BoneSpaceTransforms.Num())
 		{
-			check(BoneSpaceTransforms.Num() == InComponentToCopy->BoneSpaceTransforms.Num());
-
-			// Quick path, we know everything matches, just copy the local atoms
-			BoneSpaceTransforms = InComponentToCopy->BoneSpaceTransforms;
+			
+			Exchange(BoneSpaceTransforms, LocalTransforms);
 		}
 		else
 		{
@@ -306,7 +306,7 @@ void UPoseableMeshComponent::CopyPoseFromSkeletalComponent(const USkeletalMeshCo
 
 				if(TargetBoneIndex != INDEX_NONE)
 				{
-					BoneSpaceTransforms[TargetBoneIndex] = InComponentToCopy->BoneSpaceTransforms[SourceBoneIndex];
+					BoneSpaceTransforms[TargetBoneIndex] = LocalTransforms[SourceBoneIndex];
 				}
 			}
 		}
diff --git a/Engine/Source/Runtime/Engine/Private/Components/SkeletalMeshComponent.cpp b/Engine/Source/Runtime/Engine/Private/Components/SkeletalMeshComponent.cpp
--- a/Engine/Source/Runtime/Engine/Private/Components/SkeletalMeshComponent.cpp
+++ b/Engine/Source/Runtime/Engine/Private/Components/SkeletalMeshComponent.cpp
@@ -654,9 +654,11 @@ void USkeletalMeshComponent::InitAnim(bool bForceReinit)
 				}
 				else
 				{
+					PRAGMA_DISABLE_DEPRECATION_WARNINGS
 					BoneSpaceTransforms = SkeletalMesh->RefSkeleton.GetRefBonePose();
 					//Mini RefreshBoneTransforms (the bit we actually care about)
 					FillComponentSpaceTransforms(SkeletalMesh, BoneSpaceTransforms, GetEditableComponentSpaceTransforms());
+					PRAGMA_ENABLE_DEPRECATION_WARNINGS
 					bNeedToFlipSpaceBaseBuffers = true; // Have updated space bases so need to flip
 					FlipEditableSpaceBases();
 				}
@@ -1595,8 +1597,10 @@ void USkeletalMeshComponent::ComputeRequiredBones(TArray<FBoneIndexType>& OutReq
 
 	// Add in any bones that may be required when mirroring.
 	// JTODO: This is only required if there are mirroring nodes in the tree, but hard to know...
+	PRAGMA_DISABLE_DEPRECATION_WARNINGS
 	if (SkeletalMesh->SkelMirrorTable.Num() > 0 &&
 		SkeletalMesh->SkelMirrorTable.Num() == BoneSpaceTransforms.Num())
+	PRAGMA_ENABLE_DEPRECATION_WARNINGS
 	{
 		TArray<FBoneIndexType> MirroredDesiredBones;
 		MirroredDesiredBones.AddUninitialized(RequiredBones.Num());
@@ -1672,9 +1676,9 @@ void USkeletalMeshComponent::RecalcRequiredBones(int32 LODIndex)
 	}
 
 	ComputeRequiredBones(RequiredBones, FillComponentSpaceTransformsRequiredBones, LODIndex, /*bIgnorePhysicsAsset=*/ false);
-
+	PRAGMA_DISABLE_DEPRECATION_WARNINGS
 	BoneSpaceTransforms = SkeletalMesh->RefSkeleton.GetRefBonePose();
-
+	PRAGMA_ENABLE_DEPRECATION_WARNINGS
 	// make sure animation requiredBone to mark as dirty
 	if (AnimScriptInstance)
 	{
@@ -1979,11 +1983,13 @@ void USkeletalMeshComponent::RefreshBoneTransforms(FActorComponentTickFunction*
 
 	//Handle update rate optimization setup
 	//Dont mark cache as invalid if we aren't performing optimization anyway
+	PRAGMA_DISABLE_DEPRECATION_WARNINGS
 	const bool bInvalidCachedBones = bDoEvaluationRateOptimization &&
 									 ((BoneSpaceTransforms.Num() != SkeletalMesh->RefSkeleton.GetNum())
 									 || (BoneSpaceTransforms.Num() != CachedBoneSpaceTransforms.Num())
 									 || (GetNumComponentSpaceTransforms() != CachedComponentSpaceTransforms.Num()));
 
+	PRAGMA_ENABLE_DEPRECATION_WARNINGS
 
 	TArray<uint16> const* CurrentAnimCurveUIDFinder = (AnimScriptInstance) ? &AnimScriptInstance->GetRequiredBones().GetUIDToArrayLookupTable() : nullptr;
 	const bool bAnimInstanceHasCurveUIDList = CurrentAnimCurveUIDFinder != nullptr;
@@ -2110,8 +2116,10 @@ void USkeletalMeshComponent::RefreshBoneTransforms(FActorComponentTickFunction*
 
 				if(CachedBoneSpaceTransforms.Num())
 				{
+					PRAGMA_DISABLE_DEPRECATION_WARNINGS
 					BoneSpaceTransforms.Reset();
 					BoneSpaceTransforms.Append(CachedBoneSpaceTransforms);
+					PRAGMA_ENABLE_DEPRECATION_WARNINGS
 				}
 				if(CachedComponentSpaceTransforms.Num())
 				{
@@ -2150,7 +2158,9 @@ void USkeletalMeshComponent::SwapEvaluationContextBuffers()
 {
 	Exchange(AnimEvaluationContext.ComponentSpaceTransforms, GetEditableComponentSpaceTransforms());
 	Exchange(AnimEvaluationContext.CachedComponentSpaceTransforms, CachedComponentSpaceTransforms);
+	PRAGMA_DISABLE_DEPRECATION_WARNINGS
 	Exchange(AnimEvaluationContext.BoneSpaceTransforms, BoneSpaceTransforms);
+	PRAGMA_ENABLE_DEPRECATION_WARNINGS
 	Exchange(AnimEvaluationContext.CachedBoneSpaceTransforms, CachedBoneSpaceTransforms);
 	Exchange(AnimEvaluationContext.Curve, AnimCurves);
 	Exchange(AnimEvaluationContext.CachedCurve, CachedCurve);
@@ -2286,7 +2296,9 @@ void USkeletalMeshComponent::PostAnimEvaluation(FAnimationEvaluationContext& Eva
 			CachedComponentSpaceTransforms.Reset();
 			CachedComponentSpaceTransforms.Append(GetEditableComponentSpaceTransforms());
 			CachedBoneSpaceTransforms.Reset();
+			PRAGMA_DISABLE_DEPRECATION_WARNINGS
 			CachedBoneSpaceTransforms.Append(BoneSpaceTransforms);
+			PRAGMA_ENABLE_DEPRECATION_WARNINGS
 		}
 
 		if (EvaluationContext.bDoInterpolation)
@@ -2320,9 +2332,10 @@ void USkeletalMeshComponent::PostAnimEvaluation(FAnimationEvaluationContext& Eva
 				Alpha = ExternalInterpolationAlpha;
 			}
 
+			PRAGMA_DISABLE_DEPRECATION_WARNINGS
 			FAnimationRuntime::LerpBoneTransforms(BoneSpaceTransforms, CachedBoneSpaceTransforms, Alpha, RequiredBones);
 			FillComponentSpaceTransforms(SkeletalMesh, BoneSpaceTransforms, GetEditableComponentSpaceTransforms());
-
+			PRAGMA_ENABLE_DEPRECATION_WARNINGS
 			// interpolate curve
 			AnimCurves.LerpTo(CachedCurve, Alpha);
 		}
@@ -2595,23 +2608,28 @@ bool USkeletalMeshComponent::AllocateTransformData()
 	// Allocate transforms if not present.
 	if ( Super::AllocateTransformData() )
 	{
+		PRAGMA_DISABLE_DEPRECATION_WARNINGS
 		if(BoneSpaceTransforms.Num() != SkeletalMesh->RefSkeleton.GetNum() )
 		{
 			BoneSpaceTransforms = SkeletalMesh->RefSkeleton.GetRefBonePose();
 		}
+		PRAGMA_ENABLE_DEPRECATION_WARNINGS
 
 		return true;
 	}
 
+	PRAGMA_DISABLE_DEPRECATION_WARNINGS
 	BoneSpaceTransforms.Empty();
-	
+	PRAGMA_ENABLE_DEPRECATION_WARNINGS
 	return false;
 }
 
 void USkeletalMeshComponent::DeallocateTransformData()
 {
 	Super::DeallocateTransformData();
+	PRAGMA_DISABLE_DEPRECATION_WARNINGS
 	BoneSpaceTransforms.Empty();
+	PRAGMA_ENABLE_DEPRECATION_WARNINGS
 }
 
 void USkeletalMeshComponent::SetForceRefPose(bool bNewForceRefPose)
@@ -2751,9 +2769,9 @@ void USkeletalMeshComponent::HideBone( int32 BoneIndex, EPhysBodyOp PhysBodyOpti
 		return;
 	}
 
-	if (BoneSpaceTransforms.IsValidIndex(BoneIndex))
+	// if valid bone index
+	if (BoneIndex >= 0 && GetNumBones() > BoneIndex)
 	{
-		BoneSpaceTransforms[BoneIndex].SetScale3D(FVector::ZeroVector);
 		bRequiredBonesUpToDate = false;
 
 		if (PhysBodyOption != PBO_None)
@@ -2785,9 +2803,8 @@ void USkeletalMeshComponent::UnHideBone( int32 BoneIndex )
 		return;
 	}
 
-	if (BoneSpaceTransforms.IsValidIndex(BoneIndex))
+	if (BoneIndex >= 0 && GetNumBones() > BoneIndex)
 	{
-		BoneSpaceTransforms[BoneIndex].SetScale3D(FVector(1.0f));
 		bRequiredBonesUpToDate = false;
 
 		//FName HideBoneName = SkeletalMesh->RefSkeleton.GetBoneName(BoneIndex);
@@ -3932,4 +3949,17 @@ const TArray<FTransform>& USkeletalMeshComponent::GetCachedComponentSpaceTransfo
 	return CachedComponentSpaceTransforms;
 }
 
+TArray<FTransform> USkeletalMeshComponent::GetBoneSpaceTransforms() 
+{
+	// We may be doing parallel evaluation on the current anim instance
+	// Calling this here with true will block this init till that thread completes
+	// and it is safe to continue
+	const bool bBlockOnTask = true; // wait on evaluation task so it is safe to swap the buffers
+	const bool bPerformPostAnimEvaluation = true; // Do PostEvaluation so we make sure to swap the buffers back. 
+	HandleExistingParallelEvaluationTask(bBlockOnTask, bPerformPostAnimEvaluation);
+
+	PRAGMA_DISABLE_DEPRECATION_WARNINGS
+	return BoneSpaceTransforms;
+	PRAGMA_ENABLE_DEPRECATION_WARNINGS
+}
 #undef LOCTEXT_NAMESPACE
diff --git a/Engine/Source/Runtime/Engine/Private/PhysicsEngine/PhysicalAnimationComponent.cpp b/Engine/Source/Runtime/Engine/Private/PhysicsEngine/PhysicalAnimationComponent.cpp
--- a/Engine/Source/Runtime/Engine/Private/PhysicsEngine/PhysicalAnimationComponent.cpp
+++ b/Engine/Source/Runtime/Engine/Private/PhysicsEngine/PhysicalAnimationComponent.cpp
@@ -216,10 +216,10 @@ FTransform ComputeWorldSpaceTargetTM(const USkeletalMeshComponent& SkeletalMeshC
 	return SpaceBases[BoneIndex] * SkeletalMeshComponent.GetComponentToWorld();
 }
 
-FTransform ComputeLocalSpaceTargetTM(const USkeletalMeshComponent& SkeletalMeshComponent, const UPhysicsAsset& PhysAsset, int32 BoneIndex)
+FTransform ComputeLocalSpaceTargetTM(const USkeletalMeshComponent& SkeletalMeshComponent, const UPhysicsAsset& PhysAsset, const TArray<FTransform>& LocalTransforms, int32 BoneIndex)
 {
 	const FReferenceSkeleton& RefSkeleton = SkeletalMeshComponent.SkeletalMesh->RefSkeleton;
-	FTransform AccumulatedDelta = SkeletalMeshComponent.BoneSpaceTransforms[BoneIndex];
+	FTransform AccumulatedDelta = LocalTransforms[BoneIndex];
 	int32 CurBoneIdx = BoneIndex;
 	while ((CurBoneIdx = RefSkeleton.GetParentIndex(CurBoneIdx)) != INDEX_NONE)
 	{
@@ -246,15 +246,15 @@ FTransform ComputeLocalSpaceTargetTM(const USkeletalMeshComponent& SkeletalMeshC
 			}
 		}
 
-		AccumulatedDelta = AccumulatedDelta * SkeletalMeshComponent.BoneSpaceTransforms[CurBoneIdx];
+		AccumulatedDelta = AccumulatedDelta * LocalTransforms[CurBoneIdx];
 	}
 
 	return FTransform::Identity;
 }
 
-FTransform ComputeTargetTM(const FPhysicalAnimationData& PhysAnimData, const USkeletalMeshComponent& SkeletalMeshComponent, const UPhysicsAsset& PhysAsset, const TArray<FTransform>& SpaceBases, int32 BoneIndex)
+FTransform ComputeTargetTM(const FPhysicalAnimationData& PhysAnimData, const USkeletalMeshComponent& SkeletalMeshComponent, const UPhysicsAsset& PhysAsset, const TArray<FTransform>& LocalTransforms, const TArray<FTransform>& SpaceBases, int32 BoneIndex)
 {
-	return PhysAnimData.bIsLocalSimulation ? ComputeLocalSpaceTargetTM(SkeletalMeshComponent, PhysAsset, BoneIndex) : ComputeWorldSpaceTargetTM(SkeletalMeshComponent, SpaceBases, BoneIndex);
+	return PhysAnimData.bIsLocalSimulation ? ComputeLocalSpaceTargetTM(SkeletalMeshComponent, PhysAsset, LocalTransforms, BoneIndex) : ComputeWorldSpaceTargetTM(SkeletalMeshComponent, SpaceBases, BoneIndex);
 }
 
 void UPhysicalAnimationComponent::UpdateTargetActors(ETeleportType TeleportType)
@@ -271,6 +271,7 @@ void UPhysicalAnimationComponent::UpdateTargetActors(ETeleportType TeleportType)
 #if WITH_PHYSX
 		FPhysicsCommand::ExecuteWrite(SkeletalMeshComponent, [&]()
 		{
+			TArray<FTransform> LocalTransforms = SkeletalMeshComponent->GetBoneSpaceTransforms();
 			for (int32 DataIdx = 0; DataIdx < DriveData.Num(); ++DataIdx)
 			{
 				const FPhysicalAnimationData& PhysAnimData = DriveData[DataIdx];
@@ -280,7 +281,7 @@ void UPhysicalAnimationComponent::UpdateTargetActors(ETeleportType TeleportType)
 					const int32 BoneIdx = RefSkeleton.FindBoneIndex(PhysAnimData.BodyName);
 					if (BoneIdx != INDEX_NONE)	//It's possible the skeletal mesh has changed out from under us. In that case we should probably reset, but at the very least don't do work on non-existent bones
 					{
-						const FTransform TargetTM = ComputeTargetTM(PhysAnimData, *SkeletalMeshComponent, *PhysAsset, SpaceBases, BoneIdx);
+						const FTransform TargetTM = ComputeTargetTM(PhysAnimData, *SkeletalMeshComponent, *PhysAsset, LocalTransforms, SpaceBases, BoneIdx);
 						TargetActor->setKinematicTarget(U2PTransform(TargetTM));	//TODO: this doesn't work with sub-stepping!
 						
 						if(TeleportType == ETeleportType::TeleportPhysics)
@@ -338,6 +339,8 @@ void UPhysicalAnimationComponent::UpdatePhysicsEngine()
 #if WITH_PHYSX
 		FPhysicsCommand::ExecuteWrite(SkeletalMeshComponent, [&]()
 		{
+			TArray<FTransform> LocalTransforms = SkeletalMeshComponent->GetBoneSpaceTransforms();
+
 			for (int32 DataIdx = 0; DataIdx < DriveData.Num(); ++DataIdx)
 			{
 				bool bNewConstraint = false;
@@ -368,7 +371,7 @@ void UPhysicalAnimationComponent::UpdatePhysicsEngine()
 							ConstraintInstance->SetRefFrame(EConstraintFrame::Frame1, FTransform::Identity);
 							ConstraintInstance->SetRefFrame(EConstraintFrame::Frame2, FTransform::Identity);
 
-							const FTransform TargetTM = ComputeTargetTM(PhysAnimData, *SkeletalMeshComponent, *PhysAsset, SpaceBases, ChildBody->InstanceBoneIndex);
+							const FTransform TargetTM = ComputeTargetTM(PhysAnimData, *SkeletalMeshComponent, *PhysAsset, LocalTransforms, SpaceBases, ChildBody->InstanceBoneIndex);
 
 							// Create kinematic actor we are going to create joint with. This will be moved around with calls to SetLocation/SetRotation.
 							PxScene* PScene = PRigidActor->getScene();
