Performance optimization: Only normalize bone rotation once.

The bug identifies a n * m complexity problem when getting the animation pose out of the pose asset.

diff --git a/Engine/Source/Runtime/Engine/Private/Animation/PoseAsset.cpp b/Engine/Source/Runtime/Engine/Private/Animation/PoseAsset.cpp
--- a/Engine/Source/Runtime/Engine/Private/Animation/PoseAsset.cpp
+++ b/Engine/Source/Runtime/Engine/Private/Animation/PoseAsset.cpp
@@ -696,7 +696,7 @@ bool UPoseAsset::GetAnimationPose(struct FCompactPose& OutPose, FBlendedCurve& O
 					}
 					FAnimationRuntime::RetargetBoneTransform(MySkeleton, RetargetSource, BlendedBoneTransform[TrackIndex], LocalBoneIndices.SkeletonBoneIndex, LocalBoneIndices.CompactBoneIndex, RequiredBones, bAdditivePose);
 					OutPose[LocalBoneIndices.CompactBoneIndex] = BlendedBoneTransform[TrackIndex];
-					OutPose.NormalizeRotations();
+					OutPose[LocalBoneIndices.CompactBoneIndex].NormalizeRotation();
 				}
 			}
 
