GetAllMorphTargetNames on SkeletalMesh BlueprintCallable

diff --git a/Engine/Source/Runtime/Engine/Classes/Engine/SkeletalMesh.h b/Engine/Source/Runtime/Engine/Classes/Engine/SkeletalMesh.h
--- a/Engine/Source/Runtime/Engine/Classes/Engine/SkeletalMesh.h
+++ b/Engine/Source/Runtime/Engine/Classes/Engine/SkeletalMesh.h
@@ -721,6 +721,13 @@ public:
 	UPROPERTY(Category=Mesh, BlueprintReadWrite)
 	TArray<UMorphTarget*> MorphTargets;
 
+	/**
+	 *	Returns the list of all morph targets of this skeletal mesh
+	 *  @return	The list of morph targets
+	 */
+	UFUNCTION(BlueprintPure, Category = Mesh, meta = (DisplayName = "GetAllMorphTargetNames", ScriptName = "GetAllMorphTargetNames", Keywords = "morph shape"))
+	TArray<FString> K2_GetAllMorphTargetNames() const;
+
 	/** A fence which is used to keep track of the rendering thread releasing the static mesh resources. */
 	FRenderCommandFence ReleaseResourcesFence;
 
diff --git a/Engine/Source/Runtime/Engine/Private/SkeletalMesh.cpp b/Engine/Source/Runtime/Engine/Private/SkeletalMesh.cpp
--- a/Engine/Source/Runtime/Engine/Private/SkeletalMesh.cpp
+++ b/Engine/Source/Runtime/Engine/Private/SkeletalMesh.cpp
@@ -3226,6 +3226,17 @@ FText USkeletalMesh::GetSourceFileLabelFromIndex(int32 SourceFileIndex)
 }
 #endif //WITH_EDITOR
 
+
+TArray<FString> USkeletalMesh::K2_GetAllMorphTargetNames() const
+{
+	TArray<FString> Names;
+	for (UMorphTarget* MorphTarget : MorphTargets)
+	{
+		Names.Add(MorphTarget->GetFName().ToString());
+	}
+	return Names;
+}
+
 /*-----------------------------------------------------------------------------
 USkeletalMeshSocket
 -----------------------------------------------------------------------------*/
