Fix a bug when marking package dirty whenever accumulating metadata

diff --git a/Engine/Source/Runtime/Engine/Private/Animation/Skeleton.cpp b/Engine/Source/Runtime/Engine/Private/Animation/Skeleton.cpp
--- a/Engine/Source/Runtime/Engine/Private/Animation/Skeleton.cpp
+++ b/Engine/Source/Runtime/Engine/Private/Animation/Skeleton.cpp
@@ -1543,10 +1543,17 @@ void USkeleton::AccumulateCurveMetaData(FName CurveName, bool bMaterialSet, bool
 			if (Mapping->Exists(CurveName))
 			{
 				FCurveMetaData* CurveMetaData = GetCurveMetaData(CurveName);
+				bool bOldMaterial = CurveMetaData->Type.bMaterial;
+				bool bOldMorphtarget = CurveMetaData->Type.bMorphtarget;
 				// we don't want to undo previous flags, if it was true, we just alolw more to it. 
 				CurveMetaData->Type.bMaterial |= bMaterialSet;
 				CurveMetaData->Type.bMorphtarget |= bMorphtargetSet;
-				MarkPackageDirty();
+
+				if (bOldMaterial != CurveMetaData->Type.bMaterial 
+					|| bOldMorphtarget != CurveMetaData->Type.bMorphtarget)
+				{
+					MarkPackageDirty();
+				}
 			}
 		}
 	}
