Fix issue where additive ref pose is not handled by blendspace

diff --git a/Engine/Source/Runtime/Engine/Classes/Animation/BlendSpaceBase.h b/Engine/Source/Runtime/Engine/Classes/Animation/BlendSpaceBase.h
--- a/Engine/Source/Runtime/Engine/Classes/Animation/BlendSpaceBase.h
+++ b/Engine/Source/Runtime/Engine/Classes/Animation/BlendSpaceBase.h
@@ -475,6 +475,9 @@ protected:
 	UPROPERTY(EditAnywhere, Category = BlendParametersTest)
 	struct FBlendParameter BlendParameters[3];
 
+	/** Reset to reference pose. It does apply different refpose based on additive or not*/
+	void ResetToRefPose(FCompactPose& OutPose);
+
 #if WITH_EDITOR
 private:
 	// Track whether we have updated markers so cached data can be updated
diff --git a/Engine/Source/Runtime/Engine/Private/Animation/BlendSpaceBase.cpp b/Engine/Source/Runtime/Engine/Private/Animation/BlendSpaceBase.cpp
--- a/Engine/Source/Runtime/Engine/Private/Animation/BlendSpaceBase.cpp
+++ b/Engine/Source/Runtime/Engine/Private/Animation/BlendSpaceBase.cpp
@@ -553,6 +553,18 @@ bool UBlendSpaceBase::IsValidAdditiveType(EAdditiveAnimationType AdditiveType) c
 	return false;
 }
 
+void UBlendSpaceBase::ResetToRefPose(FCompactPose& OutPose)
+{
+	if (IsValidAdditive())
+	{
+		OutPose.ResetToAdditiveIdentity();
+	}
+	else
+	{
+		OutPose.ResetToRefPose();
+	}
+}
+
 void UBlendSpaceBase::GetAnimationPose(TArray<FBlendSampleData>& BlendSampleDataCache, /*out*/ FCompactPose& OutPose, /*out*/ FBlendedCurve& OutCurve)
 {
 	SCOPE_CYCLE_COUNTER(STAT_BlendSpace_GetAnimPose);
@@ -560,7 +572,7 @@ void UBlendSpaceBase::GetAnimationPose(TArray<FBlendSampleData>& BlendSampleData
 
 	if(BlendSampleDataCache.Num() == 0)
 	{
-		OutPose.ResetToRefPose();
+		ResetToRefPose(OutPose);
 		return;
 	}
 
@@ -605,12 +617,12 @@ void UBlendSpaceBase::GetAnimationPose(TArray<FBlendSampleData>& BlendSampleData
 			}
 			else
 			{
-				Pose.ResetToRefPose();
+				ResetToRefPose(Pose);
 			}
 		}
 		else
 		{
-			Pose.ResetToRefPose();
+			ResetToRefPose(Pose);
 		}
 	}
 
