This is code to fix https://jira.it.epicgames.net/browse/UE-62022
The issue is that we're not filling up to end, and have garbage value in the array, may cause other issues later. So just filling up to 0 when it's not there.

diff --git a/Engine/Source/Runtime/Engine/Private/SkeletalMeshMerge.cpp b/Engine/Source/Runtime/Engine/Private/SkeletalMeshMerge.cpp
--- a/Engine/Source/Runtime/Engine/Private/SkeletalMeshMerge.cpp
+++ b/Engine/Source/Runtime/Engine/Private/SkeletalMeshMerge.cpp
@@ -397,7 +397,8 @@ void FSkeletalMeshMerge::CopyVertexFromSource(VertexDataType& DestVert, const FS
 
 	// Copy all UVs that are available
 	uint32 LODNumTexCoords = SrcLODData.StaticVertexBuffers.StaticMeshVertexBuffer.GetNumTexCoords();
-	for (uint32 UVIndex = 0; UVIndex < LODNumTexCoords && UVIndex < VertexDataType::NumTexCoords; ++UVIndex)
+	const uint32 ValidLoopCount = FMath::Min(VertexDataType::NumTexCoords, LODNumTexCoords);
+	for (uint32 UVIndex = 0; UVIndex < ValidLoopCount; ++UVIndex)
 	{
 		FVector2D UVs = SrcLODData.StaticVertexBuffers.StaticMeshVertexBuffer.GetVertexUV_Typed<VertexDataType::StaticMeshVertexUVType>(SourceVertIdx, UVIndex);
 		if (UVIndex < (uint32)MergeSectionInfo.UVTransforms.Num())
@@ -407,6 +408,12 @@ void FSkeletalMeshMerge::CopyVertexFromSource(VertexDataType& DestVert, const FS
 		}
 		DestVert.UVs[UVIndex] = UVs;
 	}
+	
+	// now just fill up zero value if we didn't reach till end
+	for (uint32 UVIndex = ValidLoopCount; UVIndex < VertexDataType::NumTexCoords; ++UVIndex)
+	{
+		DestVert.UVs[UVIndex] = FVector2D::ZeroVector;
+	}
 }
 
 template<typename SkinWeightType, bool bHasExtraBoneInfluences>
