Tooltips for Merge Actors window are vague or unhelpful
#fix Added comments to designated properties

diff --git a/Engine/Source/Runtime/Engine/Classes/Engine/MaterialMerging.h b/Engine/Source/Runtime/Engine/Classes/Engine/MaterialMerging.h
--- a/Engine/Source/Runtime/Engine/Classes/Engine/MaterialMerging.h
+++ b/Engine/Source/Runtime/Engine/Classes/Engine/MaterialMerging.h
@@ -32,105 +32,115 @@ struct FMaterialProxySettings
 	// Size of generated BaseColor map
 	UPROPERTY(Category = Material, BlueprintReadWrite, EditAnywhere, meta =(ClampMin = "1", UIMin = "1"))
 	FIntPoint TextureSize;
-
+	
+	// Gutter space to take into account 
 	UPROPERTY(Category = Material, BlueprintReadWrite, AdvancedDisplay, EditAnywhere, meta=(DisplayAfter="TextureSizingType"))
 	float GutterSpace;
 
-	// Metallic constant
+	// Constant value to use for the Metallic property
 	UPROPERTY(Category = Material, BlueprintReadWrite, EditAnywhere, meta = (DisplayAfter="bMetallicMap", ClampMin = "0", ClampMax = "1", UIMin = "0", UIMax = "1", editcondition = "!bMetallicMap"))
 	float MetallicConstant;
 
-	// Roughness constant
+	// Constant value to use for the Roughness property
 	UPROPERTY(Category = Material, BlueprintReadWrite, EditAnywhere, meta = (DisplayAfter="bRoughnessMap", ClampMin = "0", ClampMax = "1", UIMin = "0", UIMax = "1", editcondition = "!bRoughnessMap"))
 	float RoughnessConstant;
 
-	// Specular constant
+	// Constant value to use for the Specular property
 	UPROPERTY(Category = Material, BlueprintReadWrite, EditAnywhere, meta = (DisplayAfter="bSpecularMap", ClampMin = "0", ClampMax = "1", UIMin = "0", UIMax = "1", editcondition = "!bSpecularMap"))
 	float SpecularConstant;
 
+	// Constant value to use for the Opacity property
 	UPROPERTY(Category = Material, BlueprintReadWrite, EditAnywhere, meta = (DisplayAfter="bOpacityMap", ClampMin = "0", ClampMax = "1", UIMin = "0", UIMax = "1", editcondition = "!bOpacityMap"))
 	float OpacityConstant;
 
+	// Constant value to use for the Opacity mask property
 	UPROPERTY(Category = Material, BlueprintReadWrite, EditAnywhere, meta = (DisplayAfter="bOpacityMaskMap", ClampMin = "0", ClampMax = "1", UIMin = "0", UIMax = "1", editcondition = "!bOpacityMaskMap"))
 	float OpacityMaskConstant;
 
+	// Constant value to use for the Ambient Occlusion property
 	UPROPERTY(Category = Material, BlueprintReadWrite, EditAnywhere, meta = (DisplayAfter="bAmbientOcclusionMap", ClampMin = "0", ClampMax = "1", UIMin = "0", UIMax = "1", editcondition = "!bAmbientOcclusionMap"))
 	float AmbientOcclusionConstant;
 
+	// Method that should be used to generate the sizes of the output textures
 	UPROPERTY(Category = Material, BlueprintReadWrite, EditAnywhere)
 	TEnumAsByte<ETextureSizingType> TextureSizingType;
 
 	UPROPERTY()
 	TEnumAsByte<EMaterialMergeType> MaterialMergeType;
 
+	// Target blend mode for the generated material
 	UPROPERTY(Category = Material, BlueprintReadWrite, AdvancedDisplay, EditAnywhere, meta=(DisplayAfter="AmbientOcclusionTextureSize"))
 	TEnumAsByte<EBlendMode> BlendMode;
 
+	// Whether or not to allow the generated material can be two-sided
 	UPROPERTY(Category = Material, BlueprintReadWrite, EditAnywhere, meta = (DisplayAfter = "BlendMode"))
 	uint8 bAllowTwoSidedMaterial : 1;
 
-	// Whether to generate normal map
+	// Whether to generate a texture for the Normal property
 	UPROPERTY(Category = Material, BlueprintReadWrite, EditAnywhere)
 	uint8 bNormalMap:1;
 
-	// Whether to generate metallic map
+	// Whether to generate a texture for the Metallic property
 	UPROPERTY(Category = Material, BlueprintReadWrite, EditAnywhere)
 	uint8 bMetallicMap:1;
 
-	// Whether to generate roughness map
+	// Whether to generate a texture for the Roughness property
 	UPROPERTY(Category = Material, BlueprintReadWrite, EditAnywhere)
 	uint8 bRoughnessMap:1;
 
-	// Whether to generate specular map
+	// Whether to generate a texture for the Specular property
 	UPROPERTY(Category = Material, BlueprintReadWrite, EditAnywhere)
 	uint8 bSpecularMap:1;
 
-	// Whether to generate emissive map
+	// Whether to generate a texture for the Emissive property
 	UPROPERTY(Category = Material, BlueprintReadWrite, EditAnywhere)
 	uint8 bEmissiveMap:1;
 
-	// Whether to generate opacity map
+	// Whether to generate a texture for the Opacity property
 	UPROPERTY(Category = Material, BlueprintReadWrite, EditAnywhere)
 	uint8 bOpacityMap:1;
 
+	// Whether to generate a texture for the Opacity Mask property
 	UPROPERTY(Category = Material, BlueprintReadWrite, EditAnywhere)
 	uint8 bOpacityMaskMap:1;
 
+	// Whether to generate a texture for the Ambient Occlusion property
 	UPROPERTY(Category = Material, BlueprintReadWrite, EditAnywhere)
 	uint8 bAmbientOcclusionMap:1;
 
-	// Override diffuse map size
+	// Override Diffuse texture size
 	UPROPERTY(Category = Material, BlueprintReadWrite, AdvancedDisplay, EditAnywhere, meta = (ClampMin = "1", UIMin = "1"))
 	FIntPoint DiffuseTextureSize;
 
-	// Override normal map size
+	// Override Normal texture size
 	UPROPERTY(Category = Material, BlueprintReadWrite, AdvancedDisplay, EditAnywhere, meta = (ClampMin = "1", UIMin = "1"))
 	FIntPoint NormalTextureSize;
 
-	// Override metallic map size
+	// Override Metallic texture size
 	UPROPERTY(Category = Material, BlueprintReadWrite, AdvancedDisplay, EditAnywhere, meta = (ClampMin = "1", UIMin = "1"))
 	FIntPoint MetallicTextureSize;
 
-	// Override roughness map size
+	// Override Roughness texture size
 	UPROPERTY(Category = Material, BlueprintReadWrite, AdvancedDisplay, EditAnywhere, meta = (ClampMin = "1", UIMin = "1"))
 	FIntPoint RoughnessTextureSize;
 
-	// Override specular map size
+	// Override Specular texture size
 	UPROPERTY(Category = Material, BlueprintReadWrite, AdvancedDisplay, EditAnywhere, meta = (ClampMin = "1", UIMin = "1"))
 	FIntPoint SpecularTextureSize;
 
-	// Override emissive map size
+	// Override Emissive texture size
 	UPROPERTY(Category = Material, BlueprintReadWrite, AdvancedDisplay, EditAnywhere, meta = (ClampMin = "1", UIMin = "1"))
 	FIntPoint EmissiveTextureSize;
 
-	// Override opacity map size
+	// Override Opacity texture size
 	UPROPERTY(Category = Material, BlueprintReadWrite, AdvancedDisplay, EditAnywhere, meta = (ClampMin = "1", UIMin = "1"))
 	FIntPoint OpacityTextureSize;
 	
-	// Override opacity map size
+	// Override Opacity Mask texture size
 	UPROPERTY(Category = Material, BlueprintReadWrite, AdvancedDisplay, EditAnywhere, meta = (ClampMin = "1", UIMin = "1"))
 	FIntPoint OpacityMaskTextureSize;
 
+	// Override Ambient Occlusion texture size
 	UPROPERTY(Category = Material, BlueprintReadWrite, AdvancedDisplay, EditAnywhere, meta = (ClampMin = "1", UIMin = "1"))
 	FIntPoint AmbientOcclusionTextureSize;
 
diff --git a/Engine/Source/Runtime/Engine/Classes/Engine/MeshMerging.h b/Engine/Source/Runtime/Engine/Classes/Engine/MeshMerging.h
--- a/Engine/Source/Runtime/Engine/Classes/Engine/MeshMerging.h
+++ b/Engine/Source/Runtime/Engine/Classes/Engine/MeshMerging.h
@@ -425,7 +425,7 @@ struct FMeshMergingSettings
 {
 	GENERATED_USTRUCT_BODY()
 
-	/** Target lightmap resolution */
+	/** The lightmap resolution used both for generating lightmap UV coordinates, and also set on the generated static mesh */
 	UPROPERTY(EditAnywhere, AdvancedDisplay, BlueprintReadWrite, Category = MeshSettings, meta=(ClampMax = 4096, EditCondition = "!bComputedLightMapResolution", DisplayAfter="bGenerateLightMapUV"))
 	int32 TargetLightMapResolution;
 
@@ -445,6 +445,7 @@ struct FMeshMergingSettings
 	UPROPERTY(EditAnywhere, Category = MeshSettings, BlueprintReadWrite, meta = (DisplayAfter="LODSelectionType", ClampMin = "0", ClampMax = "7", UIMin = "0", UIMax = "7", EnumCondition = 1))
 	int32 SpecificLOD;
 
+	/** Which selection mode should be used when generating the merged static mesh */
 	UPROPERTY(EditAnywhere, BlueprintReadWrite, Category = MeshSettings, meta = (DisplayAfter="bBakeVertexDataToMesh"))
 	EMeshLODSelectionType LODSelectionType;
 
@@ -492,6 +493,7 @@ struct FMeshMergingSettings
 	UPROPERTY(EditAnywhere, BlueprintReadWrite, Category = LandscapeCulling)
 	uint8 bUseLandscapeCulling:1;
 
+	/** Whether or not to include any imposter LODs that are part of the source static meshes */
 	UPROPERTY(EditAnywhere, BlueprintReadWrite, Category = MeshSettings)
 	uint8 bIncludeImposters:1;
 
