PR #5839: Fix for memory corruption in FSkeletalMeshMerge::GenerateLODModel() (Contributed by [REDACTED_PERSON])

diff --git a/Engine/Source/Runtime/Engine/Private/SkeletalMeshMerge.cpp b/Engine/Source/Runtime/Engine/Private/SkeletalMeshMerge.cpp
--- a/Engine/Source/Runtime/Engine/Private/SkeletalMeshMerge.cpp
+++ b/Engine/Source/Runtime/Engine/Private/SkeletalMeshMerge.cpp
@@ -702,7 +702,7 @@ void FSkeletalMeshMerge::GenerateLODModel( int32 LODIdx )
                         Section.DuplicatedVerticesBuffer.DupVertData = MergeSectionInfo.Section->DuplicatedVerticesBuffer.DupVertData;
                         Section.DuplicatedVerticesBuffer.DupVertIndexData = MergeSectionInfo.Section->DuplicatedVerticesBuffer.DupVertIndexData;
                         uint8* VertData = Section.DuplicatedVerticesBuffer.DupVertData.GetDataPointer();
-                        for (uint32 i = 0; i < MergeSectionInfo.Section->NumVertices; ++i)
+                        for (int32 i = 0; i < Section.DuplicatedVerticesBuffer.DupVertData.Num(); ++i)
                         {
                             *((uint32*)(VertData + i * sizeof(uint32))) += CurrentBaseVertexIndex - MergeSectionInfo.Section->BaseVertexIndex;
                         }
