Fix build error

diff --git a/Engine/Source/Runtime/Engine/Private/PhysicsEngine/PhysAnim.cpp b/Engine/Source/Runtime/Engine/Private/PhysicsEngine/PhysAnim.cpp
--- a/Engine/Source/Runtime/Engine/Private/PhysicsEngine/PhysAnim.cpp
+++ b/Engine/Source/Runtime/Engine/Private/PhysicsEngine/PhysAnim.cpp
@@ -364,6 +364,7 @@ void USkeletalMeshComponent::BlendInPhysics(FTickFunction& ThisTickFunction)
 		const bool bParallelBlend = !!CVarUseParallelBlendPhysics.GetValueOnGameThread() && FApp::ShouldUseThreadingForPerformance();
 		if(bParallelBlend)
 		{
+			PRAGMA_DISABLE_DEPRECATION_WARNINGS
 			if (SkeletalMesh->RefSkeleton.GetNum() != AnimEvaluationContext.BoneSpaceTransforms.Num())
 			{
 				// Initialize Parallel Task arrays
@@ -372,6 +373,7 @@ void USkeletalMeshComponent::BlendInPhysics(FTickFunction& ThisTickFunction)
 
 			AnimEvaluationContext.BoneSpaceTransforms.Reset(BoneSpaceTransforms.Num());
 			AnimEvaluationContext.BoneSpaceTransforms.Append(BoneSpaceTransforms);
+			PRAGMA_ENABLE_DEPRECATION_WARNINGS
 
 			ParallelAnimationEvaluationTask = TGraphTask<FParallelBlendPhysicsTask>::CreateTask().ConstructAndDispatchWhenReady(this);
 
@@ -387,7 +389,9 @@ void USkeletalMeshComponent::BlendInPhysics(FTickFunction& ThisTickFunction)
 		}
 		else
 		{
+			PRAGMA_DISABLE_DEPRECATION_WARNINGS
 			PerformBlendPhysicsBones(RequiredBones, BoneSpaceTransforms);
+			PRAGMA_ENABLE_DEPRECATION_WARNINGS
 			FinalizeAnimationUpdate();
 		}
 	}
@@ -454,8 +458,10 @@ void USkeletalMeshComponent::FinalizeAnimationUpdate()
 
 void USkeletalMeshComponent::CompleteParallelBlendPhysics()
 {
+	PRAGMA_DISABLE_DEPRECATION_WARNINGS
 	Exchange(AnimEvaluationContext.BoneSpaceTransforms, AnimEvaluationContext.bDoInterpolation ? CachedBoneSpaceTransforms : BoneSpaceTransforms);
-		
+	PRAGMA_ENABLE_DEPRECATION_WARNINGS
+
 	FinalizeAnimationUpdate();
 
 	ParallelAnimationEvaluationTask.SafeRelease();
@@ -719,10 +725,12 @@ void USkeletalMeshComponent::UpdateRBJointMotors()
 				(GetBoneVisibilityStates()[BoneIndex] == BVS_Visible) &&
 				(CI->IsAngularOrientationDriveEnabled()) )
 			{
+				PRAGMA_DISABLE_DEPRECATION_WARNINGS
 				check(BoneIndex < BoneSpaceTransforms.Num());
 
 				// If we find the joint - get the local-space animation between this bone and its parent.
 				FQuat LocalQuat = BoneSpaceTransforms[BoneIndex].GetRotation();
+				PRAGMA_ENABLE_DEPRECATION_WARNINGS
 				FQuatRotationTranslationMatrix LocalRot(LocalQuat, FVector::ZeroVector);
 
 				// We loop from the graphics parent bone up to the bone that has the body which the joint is attached to, to calculate the relative transform.
@@ -734,16 +742,20 @@ void USkeletalMeshComponent::UpdateRBJointMotors()
 
 				while(!bFoundControlBody)
 				{
+					PRAGMA_DISABLE_DEPRECATION_WARNINGS
 					// Abort if we find a bone scaled to zero.
 					const FVector Scale3D = BoneSpaceTransforms[TestBoneIndex].GetScale3D();
+					PRAGMA_ENABLE_DEPRECATION_WARNINGS
 					const float ScaleSum = Scale3D.X + Scale3D.Y + Scale3D.Z;
 					if(ScaleSum < KINDA_SMALL_NUMBER)
 					{
 						break;
 					}
 
+					PRAGMA_DISABLE_DEPRECATION_WARNINGS
 					// Add the current animated local transform into the overall controlling body->parent bone TM
 					FMatrix RelTM = BoneSpaceTransforms[TestBoneIndex].ToMatrixNoScale();
+					PRAGMA_ENABLE_DEPRECATION_WARNINGS
 					RelTM.SetOrigin(FVector::ZeroVector);
 					ControlBodyToParentBoneTM = ControlBodyToParentBoneTM * RelTM;
 
