AnimInstance: Fixed crash in TriggerMontageEndedEvent when a montage end notification triggers actor destruction

diff --git a/Engine/Source/Runtime/Engine/Private/Animation/AnimInstance.cpp b/Engine/Source/Runtime/Engine/Private/Animation/AnimInstance.cpp
--- a/Engine/Source/Runtime/Engine/Private/Animation/AnimInstance.cpp
+++ b/Engine/Source/Runtime/Engine/Private/Animation/AnimInstance.cpp
@@ -1578,7 +1578,7 @@ void UAnimInstance::TriggerMontageEndedEvent(const FQueuedMontageEndedEvent& Mon
 
 	if (SkelMeshComp != nullptr)
 	{
-		for (int32 Index = ActiveAnimNotifyState.Num() - 1; Index >= 0; --Index)
+		for (int32 Index = ActiveAnimNotifyState.Num() - 1; ActiveAnimNotifyState.IsValidIndex(Index); --Index)
 		{
 			const FAnimNotifyEvent& AnimNotifyEvent = ActiveAnimNotifyState[Index];
 			UAnimMontage* NotifyMontage = Cast<UAnimMontage>(AnimNotifyEvent.NotifyStateClass->GetOuter());
@@ -1589,7 +1589,19 @@ void UAnimInstance::TriggerMontageEndedEvent(const FQueuedMontageEndedEvent& Mon
 				{
 					AnimNotifyEvent.NotifyStateClass->NotifyEnd(SkelMeshComp, NotifyMontage);
 				}
-				ActiveAnimNotifyState.RemoveAtSwap(Index);
+
+				if (ActiveAnimNotifyState.IsValidIndex(Index))
+				{
+					ActiveAnimNotifyState.RemoveAtSwap(Index);
+				}
+				else
+				{
+					// The NotifyEnd callback above may have triggered actor destruction and the tear down
+					// of this instance via UninitializeAnimation which empties ActiveAnimNotifyState.
+					// If that happened, we should stop iterating the ActiveAnimNotifyState array and bail 
+					// without attempting to send MontageEnded events.
+					return;
+				}
 			}
 		}
 	}
