Rendering: Fixed Translucent User Interface materials not being drawn on asset thumbnails

diff --git a/Engine/Source/Runtime/Engine/Private/Materials/MaterialInterface.cpp b/Engine/Source/Runtime/Engine/Private/Materials/MaterialInterface.cpp
--- a/Engine/Source/Runtime/Engine/Private/Materials/MaterialInterface.cpp
+++ b/Engine/Source/Runtime/Engine/Private/Materials/MaterialInterface.cpp
@@ -138,7 +138,9 @@ FMaterialRelevance UMaterialInterface::GetRelevance_Internal(const UMaterial* Ma
 		}
 		else
 		{
-			bool bMaterialSeparateTranclucency = (InFeatureLevel > ERHIFeatureLevel::ES3_1 ? Material->bEnableSeparateTranslucency : Material->bEnableMobileSeparateTranslucency);
+			// Check whether the material can be drawn in the separate translucency pass as per FMaterialResource::IsTranslucencyAfterDOFEnabled and IsMobileSeparateTranslucencyEnabled
+			bool bSupportsSeparateTranclucency = Material->MaterialDomain != MD_UI && Material->MaterialDomain != MD_DeferredDecal;
+			bool bMaterialSeparateTranclucency = bSupportsSeparateTranclucency && (InFeatureLevel > ERHIFeatureLevel::ES3_1 ? Material->bEnableSeparateTranslucency : Material->bEnableMobileSeparateTranslucency);
 			
 			MaterialRelevance.bOpaque = !bIsTranslucent;
 			MaterialRelevance.bMasked = IsMasked();
