Sequencer: SetNextPositionWithEvents when playing to fix notifies not firing on first frame. This in conjunction with ticking sequencer first will fixes a double ticking of montage if prerequisites aren't set up first.

diff --git a/Engine/Source/Runtime/Engine/Private/Animation/AnimMontage.cpp b/Engine/Source/Runtime/Engine/Private/Animation/AnimMontage.cpp
--- a/Engine/Source/Runtime/Engine/Private/Animation/AnimMontage.cpp
+++ b/Engine/Source/Runtime/Engine/Private/Animation/AnimMontage.cpp
@@ -2695,7 +2695,16 @@ UAnimMontage* FAnimMontageInstance::SetSequencerMontagePosition(FName SlotName,
 			// ensure full weighting to this instance
 			MontageInstanceToUpdate->Blend.SetDesiredValue(Weight);
 			MontageInstanceToUpdate->Blend.SetAlpha(Weight);
-			MontageInstanceToUpdate->SetPosition(InPosition);
+			
+			if (bInPlaying)
+			{
+				MontageInstanceToUpdate->SetNextPositionWithEvents(InPosition);
+			}
+			else
+			{
+				MontageInstanceToUpdate->SetPosition(InPosition);
+			}
+
 			MontageInstanceToUpdate->bPlaying = bInPlaying;
 			return PlayingMontage;
 		}
