Engine: Added the ability to override previous transforms for components

  - This work is foundational to incoming motion-blur on cut-frames work and allows us to arbitrarily simulate previous frame transforms for components that will be used for motion vector computation.
  - Cameras with simulated previous frame transforms will also forward these transforms onto their scene-views which informs correct camera motion vectors on cut frames.

diff --git a/Engine/Source/Editor/UnrealEd/Private/EditorViewportClient.cpp b/Engine/Source/Editor/UnrealEd/Private/EditorViewportClient.cpp
--- a/Engine/Source/Editor/UnrealEd/Private/EditorViewportClient.cpp
+++ b/Engine/Source/Editor/UnrealEd/Private/EditorViewportClient.cpp
@@ -1057,6 +1057,8 @@ FSceneView* FEditorViewportClient::CalcSceneView(FSceneViewFamily* ViewFamily, c
 
 	if (bUseControllingActorViewInfo)
 	{
+		// Pass on the previous view transform of the controlling actor to the view
+		View->PreviousViewTransform = ControllingActorViewInfo.PreviousViewTransform;
 		View->OverridePostProcessSettings(ControllingActorViewInfo.PostProcessSettings, ControllingActorViewInfo.PostProcessBlendWeight);
 
 		for (int32 ExtraPPBlendIdx = 0; ExtraPPBlendIdx < ControllingActorExtraPostProcessBlends.Num(); ++ExtraPPBlendIdx)
diff --git a/Engine/Source/Runtime/Engine/Classes/Camera/CameraTypes.h b/Engine/Source/Runtime/Engine/Classes/Camera/CameraTypes.h
--- a/Engine/Source/Runtime/Engine/Classes/Camera/CameraTypes.h
+++ b/Engine/Source/Runtime/Engine/Classes/Camera/CameraTypes.h
@@ -95,6 +95,9 @@ struct FMinimalViewInfo
 	UPROPERTY(Transient, VisibleInstanceOnly, BlueprintReadWrite, Category = Camera)
 	FVector2D OffCenterProjectionOffset;
 
+	/** Optional transform to be considered as this view's previous transform */
+	TOptional<FTransform> PreviousViewTransform;
+
 	FMinimalViewInfo()
 		: Location(ForceInit)
 		, Rotation(ForceInit)
diff --git a/Engine/Source/Runtime/Engine/Classes/Components/SkinnedMeshComponent.h b/Engine/Source/Runtime/Engine/Classes/Components/SkinnedMeshComponent.h
--- a/Engine/Source/Runtime/Engine/Classes/Components/SkinnedMeshComponent.h
+++ b/Engine/Source/Runtime/Engine/Classes/Components/SkinnedMeshComponent.h
@@ -238,6 +238,9 @@ class ENGINE_API USkinnedMeshComponent : public UMeshComponent
 	/* this update renderer with new revision number twice so to clear bone velocity for motion blur or temporal AA */
 	void ClearMotionVector();
 	
+	/* Forcibly update the renderer with a new revision number to assign the current bone velocity for motion blur or temporal AA */
+	void ForceMotionVector();
+
 private:
 	/** Temporary array of of component-space bone matrices, update each frame and used for rendering the mesh. */
 	TArray<FTransform> ComponentSpaceTransformsArray[2];
diff --git a/Engine/Source/Runtime/Engine/Private/Camera/CameraComponent.cpp b/Engine/Source/Runtime/Engine/Private/Camera/CameraComponent.cpp
--- a/Engine/Source/Runtime/Engine/Private/Camera/CameraComponent.cpp
+++ b/Engine/Source/Runtime/Engine/Private/Camera/CameraComponent.cpp
@@ -13,6 +13,7 @@
 #include "Logging/TokenizedMessage.h"
 #include "Logging/MessageLog.h"
 #include "Misc/UObjectToken.h"
+#include "Rendering/MotionVectorSimulation.h"
 #include "Misc/MapErrors.h"
 #include "Components/DrawFrustumComponent.h"
 #include "IHeadMountedDisplay.h"
@@ -321,6 +322,9 @@ void UCameraComponent::GetCameraView(float DeltaTime, FMinimalViewInfo& DesiredV
 	{
 		DesiredView.PostProcessSettings = PostProcessSettings;
 	}
+
+	// If this camera component has a motion vector simumlation transform, use that for the current view's previous transform
+	DesiredView.PreviousViewTransform = FMotionVectorSimulation::Get().GetPreviousTransform(this);
 }
 
 #if WITH_EDITOR
diff --git a/Engine/Source/Runtime/Engine/Private/Components/SkinnedMeshComponent.cpp b/Engine/Source/Runtime/Engine/Private/Components/SkinnedMeshComponent.cpp
--- a/Engine/Source/Runtime/Engine/Private/Components/SkinnedMeshComponent.cpp
+++ b/Engine/Source/Runtime/Engine/Private/Components/SkinnedMeshComponent.cpp
@@ -688,6 +688,15 @@ void USkinnedMeshComponent::ClearMotionVector()
 	}
 }
 
+void USkinnedMeshComponent::ForceMotionVector()
+{
+	if (MeshObject)
+	{
+		++CurrentBoneTransformRevisionNumber;
+		MeshObject->Update(PredictedLODLevel, this, ActiveMorphTargets, MorphTargetWeights, EPreviousBoneTransformUpdateMode::None);
+	}
+}
+
 #if WITH_EDITOR
 
 bool USkinnedMeshComponent::CanEditChange(const UProperty* InProperty) const
diff --git a/Engine/Source/Runtime/Engine/Private/LocalPlayer.cpp b/Engine/Source/Runtime/Engine/Private/LocalPlayer.cpp
--- a/Engine/Source/Runtime/Engine/Private/LocalPlayer.cpp
+++ b/Engine/Source/Runtime/Engine/Private/LocalPlayer.cpp
@@ -841,6 +841,8 @@ FSceneView* ULocalPlayer::CalcSceneView( class FSceneViewFamily* ViewFamily,
 
 	View->ViewLocation = OutViewLocation;
 	View->ViewRotation = OutViewRotation;
+	// Pass on the previous view transform from the view info (probably provided by the camera if set)
+	View->PreviousViewTransform = ViewInfo.PreviousViewTransform;
 
 	ViewFamily->Views.Add(View);
 
diff --git a/Engine/Source/Runtime/Engine/Private/Rendering/MotionVectorSimulation.cpp b/Engine/Source/Runtime/Engine/Private/Rendering/MotionVectorSimulation.cpp
new file mode 100644
--- /dev/null
+++ b/Engine/Source/Runtime/Engine/Private/Rendering/MotionVectorSimulation.cpp
@@ -0,0 +1,105 @@
+// Copyright 1998-2019 Epic Games, Inc. All Rights Reserved.
+
+#include "Rendering/MotionVectorSimulation.h"
+#include "Components/SceneComponent.h"
+
+// Global switch for enabling/disabling motion vector simulation under some circumstances (ie on camera cut frames)
+int32 GMotionVectorSimulation = 0;
+static FAutoConsoleVariableRef CVarEnabled(
+	TEXT("r.MotionVectorSimulation"),
+	GMotionVectorSimulation,
+	TEXT("Controls whether to allow simulated motion vectors on scene components, geometry caches and skinned meshes on camera cut frames.")
+	);
+
+FMotionVectorSimulation::FMotionVectorSimulation()
+{}
+
+FMotionVectorSimulation::~FMotionVectorSimulation()
+{}
+
+FMotionVectorSimulation& FMotionVectorSimulation::Get()
+{
+	static FMotionVectorSimulation Singleton;
+	return Singleton;
+}
+
+bool FMotionVectorSimulation::IsEnabled()
+{
+	return GMotionVectorSimulation != 0;
+}
+
+void FMotionVectorSimulation::NotifyUObjectDeleted(const UObjectBase* Object, int32 Index)
+{
+	SimulatedTransforms.Remove(Object);
+}
+
+TOptional<FTransform> FMotionVectorSimulation::GetPreviousTransform(USceneComponent* Component) const
+{
+	if (IsEnabled())
+	{
+		FScopeLock Lock(&MapCriticalSection);
+		const FSimulatedTransform* PreviousTransform = SimulatedTransforms.Find(Component);
+
+		// Only return the transform if it pertains to the current frame number
+		if (PreviousTransform && PreviousTransform->FrameNumber == GFrameCounter)
+		{
+			return PreviousTransform->Transform;
+		}
+	}
+
+	return TOptional<FTransform>();
+}
+
+bool FMotionVectorSimulation::GetPreviousTransform(USceneComponent* Component, FTransform* OutTransform) const
+{
+	if (IsEnabled())
+	{
+		FScopeLock Lock(&MapCriticalSection);
+		const FSimulatedTransform* PreviousTransform = SimulatedTransforms.Find(Component);
+
+		// Only return the transform if it pertains to the current frame number
+		if (PreviousTransform && PreviousTransform->FrameNumber == GFrameCounter)
+		{
+			*OutTransform = PreviousTransform->Transform;
+			return true;
+		}
+	}
+
+	return false;
+}
+
+void FMotionVectorSimulation::SetPreviousTransform(USceneComponent* Component, const FTransform& InNewPreviousTransform)
+{
+	FScopeLock Lock(&MapCriticalSection);
+
+	FSimulatedTransform SimulatedTransform;
+	SimulatedTransform.Transform = InNewPreviousTransform;
+	SimulatedTransform.FrameNumber = GFrameCounter;
+
+	SimulatedTransforms.Add(Component, SimulatedTransform);
+}
+
+void FMotionVectorSimulation::ClearPreviousTransform(USceneComponent* Component)
+{
+	FScopeLock Lock(&MapCriticalSection);
+
+	SimulatedTransforms.Remove(Component);
+}
+
+void FMotionVectorSimulation::Tick(float DeltaTime)
+{
+	FScopeLock Lock(&MapCriticalSection);
+
+	for (auto It = SimulatedTransforms.CreateIterator(); It; ++It)
+	{
+		if (It.Value().FrameNumber < GFrameCounter)
+		{
+			It.RemoveCurrent();
+		}
+	}
+}
+
+TStatId FMotionVectorSimulation::GetStatId() const
+{
+	RETURN_QUICK_DECLARE_CYCLE_STAT(MotionVectorSimulation, STATGROUP_Tickables);
+}
\ No newline at end of file
diff --git a/Engine/Source/Runtime/Engine/Public/Rendering/MotionVectorSimulation.h b/Engine/Source/Runtime/Engine/Public/Rendering/MotionVectorSimulation.h
new file mode 100644
--- /dev/null
+++ b/Engine/Source/Runtime/Engine/Public/Rendering/MotionVectorSimulation.h
@@ -0,0 +1,103 @@
+// Copyright 1998-2019 Epic Games, Inc. All Rights Reserved.
+
+#pragma once
+
+#include "CoreTypes.h"
+#include "Misc/Optional.h"
+#include "Tickable.h"
+#include "UObject/UObjectArray.h"
+#include "Math/Transform.h"
+#include "HAL/IConsoleManager.h"
+
+class UObjectBase;
+class USceneComponent;
+
+/**
+ * Singleton class used for optionally overriding previous transforms for motion vector computation.
+ * This predominantly used by Sequencer on camera cut frames to forcibly inform the renderer of simulated trajectories for primitive components and cameras.
+ * Transforms are stored in an unintrusive sparse map per-component to avoid paying a memory cost for all components.
+ * Simulated transforms are only valid for the frame that they were added, and are removed on a subsequent tick.
+ */
+class ENGINE_API FMotionVectorSimulation : public FTickableGameObject, public FUObjectArray::FUObjectDeleteListener
+{
+public:
+
+	/**
+	 * Access the singleton instance for motion vector simulation
+	 * @return the singleton instance
+	 */
+	static FMotionVectorSimulation& Get();
+
+
+	/**
+	 * Check whether motion vector simulation is enabled. When disabled, no transforms will be returned.
+	 * @return Whether motion vector simulation is currently enabled or not.
+	 */
+	static bool IsEnabled();
+
+public:
+
+
+	/**
+	 * Check if the specified scene component has a simulated transform, setting the specified transform if so.
+	 * @note: Previous simulated transforms are only valid for the frame on which they were added, and are removed on the next frame
+	 *
+	 * @param Component      The Component to retrieve a previous transform for
+	 * @param OutTransform   A valid (non-null) pointer to a transform to receive the simulated transform if possible.
+	 * @return True if OutTransform was overwritten with a valid transform, false otherwise.
+	 */
+	bool GetPreviousTransform(USceneComponent* Component, FTransform* OutTransform) const;
+
+
+	/**
+	 * Check if the specified scene component has a simulated transform and return the result as an optional transform.
+	 * @note: Previous simulated transforms are only valid for the frame on which they were added, and are removed on the next frame
+	 *
+	 * @param Component      The Component to retrieve a previous transform for
+	 * @return An optional transform that is set if the specified component has a previous simulated transform, otherwise an empty optional.
+	 */
+	TOptional<FTransform> GetPreviousTransform(USceneComponent* Component) const;
+
+
+	/**
+	 * Assign a simulated previous frame transform for the specified component. Overwrites any existing simulated transform.
+	 * @note: Previous simulated transforms are only valid for the frame on which they were added, and are removed on the next frame
+	 *
+	 * @param Component                   The Component to assign the previous transform to
+	 * @param SimulatedPreviousTransform  The simulated transform that this component had on the last frame
+	 */
+	void SetPreviousTransform(USceneComponent* Component, const FTransform& SimulatedPreviousTransform);
+
+
+	/**
+	 * Clear the simulated previous transform for the specified component
+	 *
+	 * @param Component The Component to clear the previous transform for
+	 */
+	void ClearPreviousTransform(USceneComponent* Component);
+
+private:
+
+	//~ FTickableGameObject interface - used for removing transforms that are no longer valid
+	virtual bool IsTickable() const { return SimulatedTransforms.Num(); }
+	virtual bool IsTickableInEditor() const override { return true; }
+	virtual void Tick( float DeltaTime ) override;
+	virtual TStatId GetStatId() const override;
+
+	//~ FUObjectDeleteListener interface - used for removing transforms for components that are no longer valid
+	virtual void NotifyUObjectDeleted(const UObjectBase* Object, int32 Index) override;
+
+	//~ Private constructor/destructor to prevent non-singleton use
+	FMotionVectorSimulation();
+	~FMotionVectorSimulation();
+
+	struct FSimulatedTransform
+	{
+		/** The simulated transform for the component from the last frame */
+		FTransform Transform;
+		/** The frame number that this simulated transform relates to */
+		uint64 FrameNumber;
+	};
+	mutable FCriticalSection MapCriticalSection;
+	TMap<const UObjectBase *, FSimulatedTransform> SimulatedTransforms;
+};
diff --git a/Engine/Source/Runtime/Engine/Public/SceneView.h b/Engine/Source/Runtime/Engine/Public/SceneView.h
--- a/Engine/Source/Runtime/Engine/Public/SceneView.h
+++ b/Engine/Source/Runtime/Engine/Public/SceneView.h
@@ -289,7 +289,7 @@ private:
 	// ----------------
 
 public:
-	void UpdateViewMatrix(const FVector& ViewLocation, const FRotator& ViewRotation);
+	ENGINE_API void UpdateViewMatrix(const FVector& ViewLocation, const FRotator& ViewRotation);
 
 	void UpdatePlanarReflectionViewMatrix(const FSceneView& SourceView, const FMirrorMatrix& MirrorMatrix);
 
@@ -849,6 +849,7 @@ public:
 	FQuat		BaseHmdOrientation;
 	FVector		BaseHmdLocation;
 	float		WorldToMetersScale;
+	TOptional<FTransform> PreviousViewTransform;
 
 	// normally the same as ViewMatrices unless "r.Shadow.FreezeCamera" is activated
 	FViewMatrices ShadowViewMatrices;
diff --git a/Engine/Source/Runtime/Renderer/Private/RendererScene.cpp b/Engine/Source/Runtime/Renderer/Private/RendererScene.cpp
--- a/Engine/Source/Runtime/Renderer/Private/RendererScene.cpp
+++ b/Engine/Source/Runtime/Renderer/Private/RendererScene.cpp
@@ -31,6 +31,7 @@
 #include "Components/ReflectionCaptureComponent.h"
 #include "ScenePrivateBase.h"
 #include "SceneCore.h"
+#include "Rendering/MotionVectorSimulation.h"
 #include "PrimitiveSceneInfo.h"
 #include "LightSceneInfo.h"
 #include "LightMapRendering.h"
@@ -1227,11 +1228,20 @@ void FScene::AddPrimitive(UPrimitiveComponent* Primitive)
 
 	// Send a command to the rendering thread to add the primitive to the scene.
 	FScene* Scene = this;
+
+	// If this primitive has a simulated previous transform, ensure that the velocity data for the scene representation is correct
+	TOptional<FTransform> PreviousTransform = FMotionVectorSimulation::Get().GetPreviousTransform(Primitive);
+
 	ENQUEUE_RENDER_COMMAND(AddPrimitiveCommand)(
-		[Scene, PrimitiveSceneInfo](FRHICommandListImmediate& RHICmdList)
+		[Scene, PrimitiveSceneInfo, PreviousTransform](FRHICommandListImmediate& RHICmdList)
 		{
 			FScopeCycleCounter Context(PrimitiveSceneInfo->Proxy->GetStatId());
 			Scene->AddPrimitiveSceneInfo_RenderThread(RHICmdList, PrimitiveSceneInfo);
+
+			if (PreviousTransform.IsSet())
+			{
+				Scene->VelocityData.OverridePreviousTransform(PrimitiveSceneInfo->PrimitiveComponentId, PreviousTransform->ToMatrixWithScale());
+			}
 		});
 
 }
@@ -1332,6 +1342,7 @@ void FScene::UpdatePrimitiveTransform(UPrimitiveComponent* Primitive)
 				FBoxSphereBounds WorldBounds;
 				FBoxSphereBounds LocalBounds;
 				FMatrix LocalToWorld;
+				TOptional<FTransform> PreviousTransform;
 				FVector AttachmentRootPosition;
 			};
 
@@ -1342,6 +1353,7 @@ void FScene::UpdatePrimitiveTransform(UPrimitiveComponent* Primitive)
 			UpdateParams.LocalToWorld = Primitive->GetRenderMatrix();
 			UpdateParams.AttachmentRootPosition = AttachmentRootPosition;
 			UpdateParams.LocalBounds = Primitive->CalcBounds(FTransform::Identity);
+			UpdateParams.PreviousTransform = FMotionVectorSimulation::Get().GetPreviousTransform(Primitive);
 
 			// Help track down primitive with bad bounds way before the it gets to the Renderer
 			ensureMsgf(!Primitive->Bounds.BoxExtent.ContainsNaN() && !Primitive->Bounds.Origin.ContainsNaN() && !FMath::IsNaN(Primitive->Bounds.SphereRadius) && FMath::IsFinite(Primitive->Bounds.SphereRadius),
@@ -1352,6 +1364,11 @@ void FScene::UpdatePrimitiveTransform(UPrimitiveComponent* Primitive)
 				{
 					FScopeCycleCounter Context(UpdateParams.PrimitiveSceneProxy->GetStatId());
 					UpdateParams.Scene->UpdatePrimitiveTransform_RenderThread(RHICmdList, UpdateParams.PrimitiveSceneProxy, UpdateParams.WorldBounds, UpdateParams.LocalBounds, UpdateParams.LocalToWorld, UpdateParams.AttachmentRootPosition);
+
+					if (UpdateParams.PreviousTransform.IsSet())
+					{
+						UpdateParams.Scene->VelocityData.OverridePreviousTransform(UpdateParams.PrimitiveSceneProxy->PrimitiveComponentId, UpdateParams.PreviousTransform->ToMatrixWithScale());
+					}
 				});
 		}
 	}
diff --git a/Engine/Source/Runtime/Renderer/Private/SceneCaptureRendering.cpp b/Engine/Source/Runtime/Renderer/Private/SceneCaptureRendering.cpp
--- a/Engine/Source/Runtime/Renderer/Private/SceneCaptureRendering.cpp
+++ b/Engine/Source/Runtime/Renderer/Private/SceneCaptureRendering.cpp
@@ -40,6 +40,7 @@
 #include "ClearQuad.h"
 #include "PipelineStateCache.h"
 #include "RendererModule.h"
+#include "Rendering/MotionVectorSimulation.h"
 #include "SceneViewExtension.h"
 
 const TCHAR* GShaderSourceModeDefineName[] =
@@ -731,6 +732,9 @@ void FScene::UpdateSceneCaptureContents(USceneCaptureComponent2D* CaptureCompone
 		SceneRenderer->ViewFamily.SceneCaptureSource = CaptureComponent->CaptureSource;
 		SceneRenderer->ViewFamily.SceneCaptureCompositeMode = CaptureComponent->CompositeMode;
 
+		// Ensure that the views for this scene capture reflect any simulated camera motion for this frame
+		TOptional<FTransform> PreviousTransform = FMotionVectorSimulation::Get().GetPreviousTransform(CaptureComponent);
+
 		// Process Scene View extensions for the capture component
 		{
 			for (int32 Index = 0; Index < CaptureComponent->SceneViewExtensions.Num(); ++Index)
@@ -761,6 +765,11 @@ void FScene::UpdateSceneCaptureContents(USceneCaptureComponent2D* CaptureCompone
 
 			for (FSceneView& View : SceneRenderer->Views)
 			{
+				if (PreviousTransform.IsSet())
+				{
+					View.PreviousViewTransform = PreviousTransform.GetValue();
+				}
+
 				View.bCameraCut = CaptureComponent->bCameraCutThisFrame;
 
 				if (CaptureComponent->bEnableClipPlane)
diff --git a/Engine/Source/Runtime/Renderer/Private/ScenePrivate.h b/Engine/Source/Runtime/Renderer/Private/ScenePrivate.h
--- a/Engine/Source/Runtime/Renderer/Private/ScenePrivate.h
+++ b/Engine/Source/Runtime/Renderer/Private/ScenePrivate.h
@@ -2228,6 +2228,19 @@ public:
 		}
 	}
 
+	/** 
+	 * Overrides a primitive's previous LocalToWorld matrix for this frame only
+	 */
+	void OverridePreviousTransform(FPrimitiveComponentId PrimitiveComponentId, const FMatrix& PreviousLocalToWorld)
+	{
+		FComponentVelocityData* VelocityData = ComponentData.Find(PrimitiveComponentId);
+		if (VelocityData)
+		{
+			VelocityData->PreviousLocalToWorld = PreviousLocalToWorld;
+			VelocityData->bPreviousLocalToWorldValid = true;
+		}
+	}
+
 	void ApplyOffset(FVector Offset)
 	{
 		for (TMap<FPrimitiveComponentId, FComponentVelocityData>::TIterator It(ComponentData); It; ++It)
diff --git a/Engine/Source/Runtime/Renderer/Private/SceneVisibility.cpp b/Engine/Source/Runtime/Renderer/Private/SceneVisibility.cpp
--- a/Engine/Source/Runtime/Renderer/Private/SceneVisibility.cpp
+++ b/Engine/Source/Runtime/Renderer/Private/SceneVisibility.cpp
@@ -3223,6 +3223,13 @@ void FSceneRenderer::PreVisibilityFrameSetup(FRHICommandListImmediate& RHICmdLis
 				ViewState->PrevFrameViewInfo = NewPrevViewInfo;
 			}
 
+			// If the view has a previous view transform, then overwrite the previous view info for the _current_ frame.
+			if (View.PreviousViewTransform.IsSet())
+			{
+				// Note that we must ensure this transform ends up in ViewState->PrevFrameViewInfo else it will be used to calculate the next frame's motion vectors as well
+				View.PrevViewInfo.ViewMatrices.UpdateViewMatrix(View.PreviousViewTransform->GetTranslation(), View.PreviousViewTransform->GetRotation().Rotator());
+			}
+
 			// detect conditions where we should reset occlusion queries
 			if (bFirstFrameOrTimeWasReset || 
 				ViewState->LastRenderTime + GEngine->PrimitiveProbablyVisibleTime < View.Family->CurrentRealTime ||
diff --git a/Engine/Source/Runtime/Renderer/Private/VelocityRendering.cpp b/Engine/Source/Runtime/Renderer/Private/VelocityRendering.cpp
--- a/Engine/Source/Runtime/Renderer/Private/VelocityRendering.cpp
+++ b/Engine/Source/Runtime/Renderer/Private/VelocityRendering.cpp
@@ -527,7 +527,9 @@ bool FVelocityRendering::PrimitiveHasVelocity(ERHIFeatureLevel::Type FeatureLeve
 
 bool FVelocityRendering::PrimitiveHasVelocityForView(const FViewInfo& View, const FBoxSphereBounds& Bounds, const FPrimitiveSceneInfo* PrimitiveSceneInfo)
 {
-	if (View.bCameraCut)
+	// Skip camera cuts which effectively reset velocity for the new frame.
+	bool bRawCameraCut = View.bCameraCut && !View.PreviousViewTransform.IsSet();
+	if (bRawCameraCut)
 	{
 		return false;
 	}
