Change FViewport to default HasMouseCapture to false, so derived classes do not claim to have active user input by default

diff --git a/Engine/Source/Runtime/Engine/Public/UnrealClient.h b/Engine/Source/Runtime/Engine/Public/UnrealClient.h
--- a/Engine/Source/Runtime/Engine/Public/UnrealClient.h
+++ b/Engine/Source/Runtime/Engine/Public/UnrealClient.h
@@ -317,7 +317,7 @@ public:
 	virtual void Destroy() = 0;
 
 	// New MouseCapture/MouseLock API
-	virtual bool HasMouseCapture() const				{ return true; }
+	virtual bool HasMouseCapture() const				{ return false; }
 	virtual bool HasFocus() const					{ return true; }
 	virtual bool IsForegroundWindow() const			{ return true; }
 	virtual void CaptureMouse( bool bCapture )		{ }
