Motion Vector Simulation: Implement OnUObjectArrayShutdown

diff --git a/Engine/Source/Runtime/Engine/Private/Rendering/MotionVectorSimulation.cpp b/Engine/Source/Runtime/Engine/Private/Rendering/MotionVectorSimulation.cpp
--- a/Engine/Source/Runtime/Engine/Private/Rendering/MotionVectorSimulation.cpp
+++ b/Engine/Source/Runtime/Engine/Private/Rendering/MotionVectorSimulation.cpp
@@ -23,6 +23,12 @@ FMotionVectorSimulation& FMotionVectorSimulation::Get()
 	return Singleton;
 }
 
+void FMotionVectorSimulation::OnUObjectArrayShutdown()
+{
+	SimulatedTransforms.Empty();
+	GUObjectArray.RemoveUObjectDeleteListener(this);
+}
+
 bool FMotionVectorSimulation::IsEnabled()
 {
 	return GMotionVectorSimulation != 0;
diff --git a/Engine/Source/Runtime/Engine/Public/Rendering/MotionVectorSimulation.h b/Engine/Source/Runtime/Engine/Public/Rendering/MotionVectorSimulation.h
--- a/Engine/Source/Runtime/Engine/Public/Rendering/MotionVectorSimulation.h
+++ b/Engine/Source/Runtime/Engine/Public/Rendering/MotionVectorSimulation.h
@@ -86,7 +86,8 @@ private:
 
 	//~ FUObjectDeleteListener interface - used for removing transforms for components that are no longer valid
 	virtual void NotifyUObjectDeleted(const UObjectBase* Object, int32 Index) override;
-
+	virtual void OnUObjectArrayShutdown() override;
+	
 	//~ Private constructor/destructor to prevent non-singleton use
 	FMotionVectorSimulation();
 	~FMotionVectorSimulation();
