Integrate CL 6734445 for changing the default offset.

diff --git a/Engine/Source/Runtime/Engine/Private/PhysicsEngine/PhysicsGeometryPhysX.cpp b/Engine/Source/Runtime/Engine/Private/PhysicsEngine/PhysicsGeometryPhysX.cpp
--- a/Engine/Source/Runtime/Engine/Private/PhysicsEngine/PhysicsGeometryPhysX.cpp
+++ b/Engine/Source/Runtime/Engine/Private/PhysicsEngine/PhysicsGeometryPhysX.cpp
@@ -71,7 +71,7 @@ void FBodySetupShapeIterator::ForEachShape(const TArrayView<ElemType>& Elements,
 		if(PopulatePhysXGeometryAndTransform(Elem, Geom, PLocalPose))
 		{
 			const float RestOffset = ComputeRestOffset(Elem);
-			const float ContactOffset = FMath::Max(ComputeContactOffset(Geom), RestOffset + 1.f);	//make sure contact offset is always at least rest offset + 1 cm
+			const float ContactOffset = ComputeContactOffset(Geom) + RestOffset;
 			VisitorFunc(Elem, Geom, PLocalPose, ContactOffset, RestOffset);
 		}
 		else
