Fixing reverb for google resonance and adding ability to disable HRTF spatialization via cvar.

diff --git a/Engine/Plugins/Runtime/ResonanceAudio/Source/ResonanceAudio/Private/ResonanceAudioReverb.h b/Engine/Plugins/Runtime/ResonanceAudio/Source/ResonanceAudio/Private/ResonanceAudioReverb.h
--- a/Engine/Plugins/Runtime/ResonanceAudio/Source/ResonanceAudio/Private/ResonanceAudioReverb.h
+++ b/Engine/Plugins/Runtime/ResonanceAudio/Source/ResonanceAudio/Private/ResonanceAudioReverb.h
@@ -113,6 +113,7 @@ namespace ResonanceAudio
 		virtual void OnReleaseSource(const uint32 SourceId) override;
 		virtual class FSoundEffectSubmix* GetEffectSubmix(class USoundSubmix* Submix) override;
 		virtual void ProcessSourceAudio(const FAudioPluginSourceInputData& InputData, FAudioPluginSourceOutputData& OutputData) override;
+		virtual bool DoesReverbOverrideMasterReverb() const override { return false; }
 
 		void SetResonanceAudioApi(vraudio::VrAudioApi* InResonanceAudioApi) { ResonanceAudioApi = InResonanceAudioApi; };
 		void SetPreset(UResonanceAudioReverbPluginPreset* InPreset);
diff --git a/Engine/Source/Runtime/AudioMixer/Private/AudioMixerDevice.cpp b/Engine/Source/Runtime/AudioMixer/Private/AudioMixerDevice.cpp
--- a/Engine/Source/Runtime/AudioMixer/Private/AudioMixerDevice.cpp
+++ b/Engine/Source/Runtime/AudioMixer/Private/AudioMixerDevice.cpp
@@ -642,6 +642,8 @@ namespace Audio
 			FSoundEffectSubmixInitData InitData;
 			InitData.SampleRate = GetSampleRate();
 	
+			bool bIsMasterReverbBypassed = false;
+
 			// Setup the master reverb plugin
 			if (ReverbPluginInterface.IsValid() && MasterSubmixInstances[EMasterSubmixType::ReverbPlugin].IsValid())
 			{
@@ -656,8 +658,11 @@ namespace Audio
 				MasterReverbPluginSubmix->AddSoundEffectSubmix(ReverbPluginId, ReverbPluginEffectSubmix);
 				MasterReverbPluginSubmix->SetParentSubmix(MasterSubmixInstance);
 				MasterSubmixInstance->AddChildSubmix(MasterReverbPluginSubmix);
+
+				bIsMasterReverbBypassed = ReverbPluginInterface->DoesReverbOverrideMasterReverb();
 			}
-			else if (MasterSubmixInstances[EMasterSubmixType::Reverb].IsValid())
+					
+			if (MasterSubmixInstances[EMasterSubmixType::Reverb].IsValid() && !bIsMasterReverbBypassed)
 			{
 				// Setup the master reverb only if we don't have a reverb plugin
 				USoundEffectSubmixPreset* ReverbPreset = nullptr;
diff --git a/Engine/Source/Runtime/AudioMixer/Private/AudioMixerSource.cpp b/Engine/Source/Runtime/AudioMixer/Private/AudioMixerSource.cpp
--- a/Engine/Source/Runtime/AudioMixer/Private/AudioMixerSource.cpp
+++ b/Engine/Source/Runtime/AudioMixer/Private/AudioMixerSource.cpp
@@ -12,14 +12,6 @@
 // Link to "Audio" profiling category
 CSV_DECLARE_CATEGORY_MODULE_EXTERN(AUDIOMIXER_API, Audio);
 
-static int32 DisableHRTFCvar = 0;
-FAutoConsoleVariableRef CVarDisableHRTF(
-	TEXT("au.DisableHRTF"),
-	DisableHRTFCvar,
-	TEXT("Disables HRTF\n")
-	TEXT("0: Not Disabled, 1: Disabled"),
-	ECVF_Default);
-
 static int32 UseListenerOverrideForSpreadCVar = 0;
 FAutoConsoleVariableRef CVarUseListenerOverrideForSpread(
 	TEXT("au.UseListenerOverrideForSpread"),
@@ -123,6 +115,7 @@ namespace Audio
 			InitParams.NumInputFrames = NumFrames;
 			InitParams.SourceVoice = MixerSourceVoice;
 			InitParams.bUseHRTFSpatialization = UseObjectBasedSpatialization();
+			InitParams.bIsExternalSend = MixerDevice->bSpatializationIsExternalSend;
 			InitParams.bIsAmbisonics = WaveInstance->bIsAmbisonics;
 
 			if (InitParams.bIsAmbisonics)
@@ -943,13 +936,11 @@ namespace Audio
 					MixerSourceVoice->SetSubmixSendInfo(MixerDevice->GetMasterReverbPluginSubmix(), ReverbSendLevel);
 				}
 			}
-			else
+
+			// Send the source audio to the master reverb
+			if (!AudioDevice->IsReverbPluginBypassingMasterReverb() && MixerDevice->GetMasterReverbSubmix().IsValid())
 			{
-				// Send the source audio to the master reverb
-				if (MixerDevice->GetMasterReverbSubmix().IsValid())
-				{
-					MixerSourceVoice->SetSubmixSendInfo(MixerDevice->GetMasterReverbSubmix(), ReverbSendLevel);
-				}
+				MixerSourceVoice->SetSubmixSendInfo(MixerDevice->GetMasterReverbSubmix(), ReverbSendLevel);
 			}
 		}
 
@@ -1183,9 +1174,8 @@ namespace Audio
 	bool FMixerSource::UseObjectBasedSpatialization() const
 	{
 		return (Buffer->NumChannels <= MixerDevice->MaxChannelsSupportedBySpatializationPlugin &&
-			AudioDevice->IsSpatializationPluginEnabled() &&
-			DisableHRTFCvar == 0 &&
-			WaveInstance->SpatializationMethod == ESoundSpatializationAlgorithm::SPATIALIZATION_HRTF);
+				AudioDevice->IsSpatializationPluginEnabled() &&
+				WaveInstance->SpatializationMethod == ESoundSpatializationAlgorithm::SPATIALIZATION_HRTF);
 	}
 
 	bool FMixerSource::UseSpatializationPlugin() const
diff --git a/Engine/Source/Runtime/AudioMixer/Private/AudioMixerSourceManager.cpp b/Engine/Source/Runtime/AudioMixer/Private/AudioMixerSourceManager.cpp
--- a/Engine/Source/Runtime/AudioMixer/Private/AudioMixerSourceManager.cpp
+++ b/Engine/Source/Runtime/AudioMixer/Private/AudioMixerSourceManager.cpp
@@ -496,6 +496,7 @@ namespace Audio
 		SourceInfo.bIsStopping = false;
 		SourceInfo.bIsBusy = false;
 		SourceInfo.bUseHRTFSpatializer = false;
+		SourceInfo.bIsExternalSend = false;
 		SourceInfo.bUseOcclusionPlugin = false;
 		SourceInfo.bUseReverbPlugin = false;
 		SourceInfo.bUseModulationPlugin = false;
@@ -631,6 +632,7 @@ namespace Audio
 			SourceInfo.bIsDone = false;
 			SourceInfo.bIsLastBuffer = false;
 			SourceInfo.bUseHRTFSpatializer = InitParams.bUseHRTFSpatialization;
+			SourceInfo.bIsExternalSend = InitParams.bIsExternalSend;
 			SourceInfo.bIsVorbis = InitParams.bIsVorbis;
 			SourceInfo.AudioComponentID = InitParams.AudioComponentID;
 
@@ -775,7 +777,7 @@ namespace Audio
 
 			// Initialize a new downmix data:
 			check(SourceId < SourceInfos.Num());
-			const int32 SourceInputChannels = SourceInfo.bUseHRTFSpatializer ? 2 : SourceInfo.NumInputChannels;
+			const int32 SourceInputChannels = (SourceInfo.bUseHRTFSpatializer && !SourceInfo.bIsExternalSend) ? 2 : SourceInfo.NumInputChannels;
 			FSourceDownmixData& DownmixData = InitializeDownmixForSource(SourceId, SourceInputChannels, MixerDevice->GetDeviceOutputChannels(), NumOutputFrames);
 
 			for (int32 i = 0; i < InitParams.SubmixSends.Num(); ++i)
diff --git a/Engine/Source/Runtime/AudioMixer/Private/AudioMixerSourceManager.h b/Engine/Source/Runtime/AudioMixer/Private/AudioMixerSourceManager.h
--- a/Engine/Source/Runtime/AudioMixer/Private/AudioMixerSourceManager.h
+++ b/Engine/Source/Runtime/AudioMixer/Private/AudioMixerSourceManager.h
@@ -109,6 +109,7 @@ namespace Audio
 		uint64 AudioComponentID;
 		uint8 bPlayEffectChainTails : 1;
 		uint8 bUseHRTFSpatialization : 1;
+		uint8 bIsExternalSend : 1;
 		uint8 bIsDebugMode : 1;
 		uint8 bOutputToBusOnly : 1;
 		uint8 bIsVorbis : 1;
@@ -133,6 +134,7 @@ namespace Audio
 			, AudioComponentID(0)
 			, bPlayEffectChainTails(false)
 			, bUseHRTFSpatialization(false)
+			, bIsExternalSend(false)
 			, bIsDebugMode(false)
 			, bOutputToBusOnly(false)
 			, bIsVorbis(false)
@@ -491,6 +493,7 @@ namespace Audio
 			uint8 bHasStarted:1;
 			uint8 bIsBusy:1;
 			uint8 bUseHRTFSpatializer:1;
+			uint8 bIsExternalSend:1;
 			uint8 bUseOcclusionPlugin:1;
 			uint8 bUseReverbPlugin:1;
 			uint8 bUseModulationPlugin:1;
diff --git a/Engine/Source/Runtime/Engine/Private/ActiveSound.cpp b/Engine/Source/Runtime/Engine/Private/ActiveSound.cpp
--- a/Engine/Source/Runtime/Engine/Private/ActiveSound.cpp
+++ b/Engine/Source/Runtime/Engine/Private/ActiveSound.cpp
@@ -1635,6 +1635,10 @@ void FActiveSound::ApplyAttenuation(FSoundParseParameters& ParseParams, const FL
 	{
 		ParseParams.SpatializationMethod = ESoundSpatializationAlgorithm::SPATIALIZATION_HRTF;
 	}
+	else if (Settings->SpatializationAlgorithm == ESoundSpatializationAlgorithm::SPATIALIZATION_HRTF && AudioDevice->IsHRTFDisabled())
+	{
+		ParseParams.SpatializationMethod = ESoundSpatializationAlgorithm::SPATIALIZATION_Default;
+	}
 	else
 	{
 		ParseParams.SpatializationMethod = Settings->SpatializationAlgorithm;
diff --git a/Engine/Source/Runtime/Engine/Private/AudioDevice.cpp b/Engine/Source/Runtime/Engine/Private/AudioDevice.cpp
--- a/Engine/Source/Runtime/Engine/Private/AudioDevice.cpp
+++ b/Engine/Source/Runtime/Engine/Private/AudioDevice.cpp
@@ -145,6 +145,12 @@ FAutoConsoleVariableRef CVarEnableBinauralAudioForAllSpatialSounds(
 	TEXT("Toggles binaural audio rendering for all spatial sounds if binaural rendering is available.\n"),
 	ECVF_Default);
 
+static int32 DisableBinauralSpatializationCVar = 0;
+FAutoConsoleVariableRef CVarDisableBinauralSpatialization(
+	TEXT("au.DisableBinauralSpatialization"),
+	DisableBinauralSpatializationCVar,
+	TEXT("Disables binaural spatialization.\n"),
+	ECVF_Default);
 
 using FVirtualLoopPair = TPair<FActiveSound*, FAudioVirtualLoop>;
 
@@ -184,6 +190,7 @@ FAudioDevice::FAudioDevice()
 	, PlatformAudioHeadroom(1.0f)
 	, DefaultReverbSendLevel(0.0f)
 	, bHRTFEnabledForAll_OnGameThread(false)
+	, bHRTFDisabled_OnGameThread(false)
 	, bGameWasTicking(true)
 	, bDisableAudioCaching(false)
 	, bIsAudioDeviceHardwareInitialized(false)
@@ -198,9 +205,11 @@ FAudioDevice::FAudioDevice()
 	, bSpatializationInterfaceEnabled(false)
 	, bOcclusionInterfaceEnabled(false)
 	, bReverbInterfaceEnabled(false)
+	, bReverbPluginBypassesMasterReverb(true)
 	, bModulationInterfaceEnabled(false)
 	, bPluginListenersInitialized(false)
 	, bHRTFEnabledForAll(false)
+	, bHRTFDisabled(false)
 	, bIsDeviceMuted(false)
 	, bIsInitialized(false)
 	, AudioClock(0.0)
@@ -351,9 +360,9 @@ bool FAudioDevice::Init(int32 InMaxChannels)
 	{
 		ReverbPluginInterface = ReverbPluginFactory->CreateNewReverbPlugin(this);
 		bReverbInterfaceEnabled = true;
+		bReverbPluginBypassesMasterReverb = ReverbPluginInterface->DoesReverbOverrideMasterReverb();
 		bReverbIsExternalSend = ReverbPluginFactory->IsExternalSend();
 		UE_LOG(LogAudio, Log, TEXT("Audio Reverb Plugin: %s"), *(ReverbPluginFactory->GetDisplayName()));
-
 	}
 	else
 	{
@@ -1708,6 +1717,17 @@ bool FAudioDevice::IsHRTFEnabledForAll() const
 	return (bHRTFEnabledForAll_OnGameThread || EnableBinauralAudioForAllSpatialSoundsCVar == 1) && IsSpatializationPluginEnabled();
 }
 
+bool FAudioDevice::IsHRTFDisabled() const
+{
+	if (IsInAudioThread())
+	{
+		return (bHRTFDisabled || DisableBinauralSpatializationCVar == 1);
+	}
+
+	check(IsInGameThread());
+	return (bHRTFDisabled_OnGameThread || DisableBinauralSpatializationCVar == 1);
+}
+
 void FAudioDevice::SetMixDebugState(EDebugState InDebugState)
 {
 	if (!IsInAudioThread())
diff --git a/Engine/Source/Runtime/Engine/Public/AudioDevice.h b/Engine/Source/Runtime/Engine/Public/AudioDevice.h
--- a/Engine/Source/Runtime/Engine/Public/AudioDevice.h
+++ b/Engine/Source/Runtime/Engine/Public/AudioDevice.h
@@ -950,6 +950,25 @@ public:
 		}, GET_STATID(STAT_SetHRTFEnabledForAll));
 	}
 
+	/** Whether or not HRTF is disabled. */
+	bool IsHRTFDisabled() const;
+
+	void SetHRTFDisabled(bool InIsHRTFDisabled)
+	{
+		const bool bNewHRTFDisabled = InIsHRTFDisabled;
+
+		bHRTFDisabled_OnGameThread = bNewHRTFDisabled;
+
+		DECLARE_CYCLE_STAT(TEXT("FAudioThreadTask.SetHRTFDisabled"), STAT_SetHRTFDisabled, STATGROUP_AudioThreadCommands);
+
+		FAudioDevice* AudioDevice = this;
+		FAudioThread::RunCommandOnAudioThread([AudioDevice, bNewHRTFDisabled]()
+		{
+			AudioDevice->bHRTFDisabled = bNewHRTFDisabled;
+
+		}, GET_STATID(STAT_SetHRTFDisabled));
+	}
+
 	void SetSpatializationInterfaceEnabled(bool InbSpatializationInterfaceEnabled)
 	{
 		FAudioThread::SuspendAudioThread();
@@ -1085,6 +1104,12 @@ public:
 		return bReverbInterfaceEnabled;
 	}
 
+	/** Whether or not the reverb plugin is bypassing the master reverb. */
+	bool IsReverbPluginBypassingMasterReverb() const
+	{
+		return bReverbPluginBypassesMasterReverb;
+	}
+
 	static bool IsReverbPluginLoaded()
 	{
 		if (FAudioDevice* MainAudioDevice = GEngine->GetMainAudioDevice())
@@ -1679,7 +1704,10 @@ private:
 	FActivatedReverb HighestPriorityActivatedReverb;
 
 	/** Gamethread representation of whether HRTF is enabled for all 3d sounds. (not bitpacked to avoid thread issues) */
-	bool bHRTFEnabledForAll_OnGameThread:1;
+	bool bHRTFEnabledForAll_OnGameThread;
+
+	/** Gamethread representation of whether HRTF is disabbled for all 3d sounds. */
+	bool bHRTFDisabled_OnGameThread;
 
 	uint8 bGameWasTicking:1;
 
@@ -1714,14 +1742,18 @@ private:
 	uint8 bSpatializationInterfaceEnabled:1;
 	uint8 bOcclusionInterfaceEnabled:1;
 	uint8 bReverbInterfaceEnabled:1;
+	uint8 bReverbPluginBypassesMasterReverb:1;
 	uint8 bModulationInterfaceEnabled:1;
 
 	/** Whether or not we've initialized plugin listeners array. */
 	uint8 bPluginListenersInitialized:1;
 
-	/** Whether HRTF is enabled for all 3d sounds. */
+	/** Whether HRTF is enabled for all 3d sounds. This will automatically make all 3d mono sounds HRTF spatialized. */
 	uint8 bHRTFEnabledForAll:1;
 
+	/** Whether or not HRTF is disabled. This will make any sounds which are set to HRTF spatialize to spatialize with panning. */
+	uint8 bHRTFDisabled:1;
+
 	/** Whether the audio device is active (current audio device in-focus in PIE) */
 	uint8 bIsDeviceMuted:1;
 
diff --git a/Engine/Source/Runtime/Engine/Public/IAudioExtensionPlugin.h b/Engine/Source/Runtime/Engine/Public/IAudioExtensionPlugin.h
--- a/Engine/Source/Runtime/Engine/Public/IAudioExtensionPlugin.h
+++ b/Engine/Source/Runtime/Engine/Public/IAudioExtensionPlugin.h
@@ -640,6 +640,9 @@ public:
 	virtual void ProcessSourceAudio(const FAudioPluginSourceInputData& InputData, FAudioPluginSourceOutputData& OutputData)
 	{
 	}
+
+	/** Returns whether or not the plugin reverb overrides the master reverb. If true, then the built in reverb will be uninitialized and bypassed. */
+	virtual bool DoesReverbOverrideMasterReverb() const { return true; }
 };
 
 /************************************************************************/
