Fixing a crash in ParticleSystemRenderer when the view is set to wireframe and using forward rendering

diff --git a/Engine/Source/Runtime/Engine/Private/Particles/ParticleSystemRender.cpp b/Engine/Source/Runtime/Engine/Private/Particles/ParticleSystemRender.cpp
--- a/Engine/Source/Runtime/Engine/Private/Particles/ParticleSystemRender.cpp
+++ b/Engine/Source/Runtime/Engine/Private/Particles/ParticleSystemRender.cpp
@@ -1870,7 +1870,7 @@ void FDynamicMeshEmitterData::GetDynamicMeshElementsEmitter(const FParticleSyste
 
 					if (bIsWireframe)
 					{
-						if (LODModel.AdditionalIndexBuffers->WireframeIndexBuffer.IsInitialized()
+						if (LODModel.AdditionalIndexBuffers && LODModel.AdditionalIndexBuffers->WireframeIndexBuffer.IsInitialized()
 							&& !(RHISupportsTessellation(ShaderPlatform) && Mesh.VertexFactory->GetType()->SupportsTessellationShaders()))
 						{
 							Mesh.Type = PT_LineList;
