Display max wave instance volume instead of falling back to ActiveSound volume in debug

diff --git a/Engine/Source/Runtime/Engine/Private/AudioDebug.cpp b/Engine/Source/Runtime/Engine/Private/AudioDebug.cpp
--- a/Engine/Source/Runtime/Engine/Private/AudioDebug.cpp
+++ b/Engine/Source/Runtime/Engine/Private/AudioDebug.cpp
@@ -277,11 +277,9 @@ void FAudioDebugger::DrawDebugInfo(const FActiveSound& ActiveSound, const TArray
 		float DisplayValue = 0.0f;
 		if (ActiveSoundVisualizeModeCVar == 1 || ActiveSoundVisualizeModeCVar == 2)
 		{
-			DisplayValue = ActiveSound.GetVolume();
-			if (ThisSoundsWaveInstances.Num() == 1)
+			for (FWaveInstance* WaveInstance : *ThisSoundsWaveInstances)
 			{
-				const FWaveInstance* WaveInstance = ThisSoundsWaveInstances[0];
-				DisplayValue = WaveInstance->GetVolumeWithDistanceAttenuation() * WaveInstance->GetDynamicVolume();
+				DisplayValue = FMath::Max(DisplayValue, WaveInstance->GetVolumeWithDistanceAttenuation() * WaveInstance->GetDynamicVolume());
 			}
 		}
 		else if (ActiveSoundVisualizeModeCVar == 3)
