Fix for migrated function signature changed to const ref instead of ptr

diff --git a/Engine/Source/Runtime/Engine/Private/AudioDebug.cpp b/Engine/Source/Runtime/Engine/Private/AudioDebug.cpp
--- a/Engine/Source/Runtime/Engine/Private/AudioDebug.cpp
+++ b/Engine/Source/Runtime/Engine/Private/AudioDebug.cpp
@@ -277,7 +277,7 @@ void FAudioDebugger::DrawDebugInfo(const FActiveSound& ActiveSound, const TArray
 		float DisplayValue = 0.0f;
 		if (ActiveSoundVisualizeModeCVar == 1 || ActiveSoundVisualizeModeCVar == 2)
 		{
-			for (FWaveInstance* WaveInstance : *ThisSoundsWaveInstances)
+			for (FWaveInstance* WaveInstance : ThisSoundsWaveInstances)
 			{
 				DisplayValue = FMath::Max(DisplayValue, WaveInstance->GetVolumeWithDistanceAttenuation() * WaveInstance->GetDynamicVolume());
 			}
