Now support static switches, quality switches and feature switches for the Shading Model from Material Expression. Done by collecting the shading models used during material translation.

#rnx

diff --git a/Engine/Source/Developer/MaterialBaking/Private/ExportMaterialProxy.h b/Engine/Source/Developer/MaterialBaking/Private/ExportMaterialProxy.h
--- a/Engine/Source/Developer/MaterialBaking/Private/ExportMaterialProxy.h
+++ b/Engine/Source/Developer/MaterialBaking/Private/ExportMaterialProxy.h
@@ -473,6 +473,7 @@ public:
 	virtual bool IsMasked() const override { return false; }
 	virtual enum EBlendMode GetBlendMode() const override { return BLEND_Opaque; }
 	virtual FMaterialShadingModelField GetShadingModels() const override { return MSM_DefaultLit; }
+	virtual bool IsShadingModelFromMaterialExpression() const override { return false; }
 	virtual float GetOpacityMaskClipValue() const override { return 0.5f; }
 	virtual bool GetCastDynamicShadowAsMasked() const override { return false; }
 	virtual FString GetFriendlyName() const override { return FString::Printf(TEXT("FExportMaterialRenderer %s"), MaterialInterface ? *MaterialInterface->GetName() : TEXT("NULL")); }
diff --git a/Engine/Source/Developer/MaterialUtilities/Private/MaterialUtilities.cpp b/Engine/Source/Developer/MaterialUtilities/Private/MaterialUtilities.cpp
--- a/Engine/Source/Developer/MaterialUtilities/Private/MaterialUtilities.cpp
+++ b/Engine/Source/Developer/MaterialUtilities/Private/MaterialUtilities.cpp
@@ -651,6 +651,7 @@ public:
 	virtual bool IsMasked() const override									{ return false; }
 	virtual enum EBlendMode GetBlendMode() const override					{ return BLEND_Opaque; }
 	virtual FMaterialShadingModelField GetShadingModels() const override	{ return MSM_Unlit; }
+	virtual bool IsShadingModelFromMaterialExpression() const override		{ return false; }
 	virtual float GetOpacityMaskClipValue() const override					{ return 0.5f; }
 	virtual bool GetCastDynamicShadowAsMasked() const override				{ return false; }
 	virtual FString GetFriendlyName() const override { return FString::Printf(TEXT("FExportMaterialRenderer %s"), MaterialInterface ? *MaterialInterface->GetName() : TEXT("NULL")); }
diff --git a/Engine/Source/Editor/MaterialEditor/Private/MaterialEditor.h b/Engine/Source/Editor/MaterialEditor/Private/MaterialEditor.h
--- a/Engine/Source/Editor/MaterialEditor/Private/MaterialEditor.h
+++ b/Engine/Source/Editor/MaterialEditor/Private/MaterialEditor.h
@@ -156,6 +156,7 @@ public:
 	virtual bool IsMasked() const override { return false; }
 	virtual enum EBlendMode GetBlendMode() const override { return BLEND_Translucent; }
 	virtual FMaterialShadingModelField GetShadingModels() const override { return MSM_Unlit; }
+	virtual bool IsShadingModelFromMaterialExpression() const override { return false; }
 	virtual float GetOpacityMaskClipValue() const override { return 0.5f; }
 	virtual bool GetCastDynamicShadowAsMasked() const override { return false; }
 	virtual FString GetFriendlyName() const override { return FString::Printf(TEXT("FMatExpressionPreview %s"), Expression.IsValid() ? *Expression->GetName() : TEXT("NULL")); }
diff --git a/Engine/Source/Editor/UnrealEd/Private/Lightmass/LightmassRender.cpp b/Engine/Source/Editor/UnrealEd/Private/Lightmass/LightmassRender.cpp
--- a/Engine/Source/Editor/UnrealEd/Private/Lightmass/LightmassRender.cpp
+++ b/Engine/Source/Editor/UnrealEd/Private/Lightmass/LightmassRender.cpp
@@ -532,6 +532,7 @@ public:
 	virtual bool IsMasked() const override									{ return false; }
 	virtual enum EBlendMode GetBlendMode() const override					{ return BLEND_Opaque; }
 	virtual FMaterialShadingModelField GetShadingModels() const override	{ return MSM_Unlit; }
+	virtual bool IsShadingModelFromMaterialExpression() const override		{ return false; }
 	virtual float GetOpacityMaskClipValue() const override					{ return 0.5f; }
 	virtual bool GetCastDynamicShadowAsMasked() const override				{ return false; }
 	virtual FString GetFriendlyName() const override { return FString::Printf(TEXT("FLightmassMaterialRenderer %s"), MaterialInterface ? *MaterialInterface->GetName() : TEXT("NULL")); }
diff --git a/Engine/Source/Runtime/Engine/Classes/Engine/EngineTypes.h b/Engine/Source/Runtime/Engine/Classes/Engine/EngineTypes.h
--- a/Engine/Source/Runtime/Engine/Classes/Engine/EngineTypes.h
+++ b/Engine/Source/Runtime/Engine/Classes/Engine/EngineTypes.h
@@ -481,7 +481,7 @@ enum EMaterialShadingModel
 static_assert(MSM_NUM <= 16, "Do not exceed 16 shading models without expanding FMaterialShadingModelField to support uint32 instead of uint16!");
 
 /** Wrapper for a bitfield of shading models. A material contains one of these to describe what possible shading models can be used by that material. */
-USTRUCT(BlueprintType)
+USTRUCT()
 struct ENGINE_API FMaterialShadingModelField
 {
 	GENERATED_USTRUCT_BODY()
diff --git a/Engine/Source/Runtime/Engine/Private/DebugViewModeMaterialProxy.cpp b/Engine/Source/Runtime/Engine/Private/DebugViewModeMaterialProxy.cpp
--- a/Engine/Source/Runtime/Engine/Private/DebugViewModeMaterialProxy.cpp
+++ b/Engine/Source/Runtime/Engine/Private/DebugViewModeMaterialProxy.cpp
@@ -168,6 +168,11 @@ FMaterialShadingModelField FDebugViewModeMaterialProxy::GetShadingModels() const
 	return Material ? Material->GetShadingModels() : MSM_Unlit;
 }
 
+bool FDebugViewModeMaterialProxy::IsShadingModelFromMaterialExpression() const
+{ 
+	return Material ? Material->IsShadingModelFromMaterialExpression() : false;
+}
+
 float FDebugViewModeMaterialProxy::GetOpacityMaskClipValue() const
 { 
 	return Material ? Material->GetOpacityMaskClipValue() : .5f;
diff --git a/Engine/Source/Runtime/Engine/Private/DebugViewModeMaterialProxy.h b/Engine/Source/Runtime/Engine/Private/DebugViewModeMaterialProxy.h
--- a/Engine/Source/Runtime/Engine/Private/DebugViewModeMaterialProxy.h
+++ b/Engine/Source/Runtime/Engine/Private/DebugViewModeMaterialProxy.h
@@ -142,6 +142,7 @@ public:
 	virtual bool IsMasked() const override;
 	virtual enum EBlendMode GetBlendMode() const override;
 	virtual FMaterialShadingModelField GetShadingModels() const override;
+	virtual bool IsShadingModelFromMaterialExpression() const override;
 	virtual float GetOpacityMaskClipValue() const override;
 	virtual bool GetCastDynamicShadowAsMasked() const override;
 	virtual void GatherCustomOutputExpressions(TArray<class UMaterialExpressionCustomOutput*>& OutCustomOutputs) const override;
diff --git a/Engine/Source/Runtime/Engine/Private/Materials/HLSLMaterialTranslator.h b/Engine/Source/Runtime/Engine/Private/Materials/HLSLMaterialTranslator.h
--- a/Engine/Source/Runtime/Engine/Private/Materials/HLSLMaterialTranslator.h
+++ b/Engine/Source/Runtime/Engine/Private/Materials/HLSLMaterialTranslator.h
@@ -253,6 +253,10 @@ protected:
 	uint32 NumUserVertexTexCoords;
 
 	uint32 DynamicParticleParameterMask;
+
+	/** Will contain all the shading models picked up from the material expression graph */
+	FMaterialShadingModelField ShadingModelsFromCompilation;
+
 public: 
 
 	FHLSLMaterialTranslator(FMaterial* InMaterial,
@@ -591,8 +595,8 @@ public:
 				}
 			}
 
+			
 			const EShaderFrequency NormalShaderFrequency = FMaterialAttributeDefinitionMap::GetShaderFrequency(MP_Normal);
-			const FMaterialShadingModelField MaterialShadingModels = Material->GetShadingModels();
 			const EMaterialDomain Domain = Material->GetMaterialDomain();
 			const EBlendMode BlendMode = Material->GetBlendMode();
 
@@ -626,6 +630,19 @@ public:
 			Chunk[MP_WorldDisplacement]				= Material->CompilePropertyAndSetMaterialProperty(MP_WorldDisplacement		,this);
 			Chunk[MP_TessellationMultiplier]		= Material->CompilePropertyAndSetMaterialProperty(MP_TessellationMultiplier	,this);			
 
+			// Make sure to compile this property before using ShadingModelsFromCompilation
+			Chunk[MP_ShadingModel]					= Material->CompilePropertyAndSetMaterialProperty(MP_ShadingModel			,this);
+			
+			// Get shading models from material.
+			FMaterialShadingModelField MaterialShadingModels = Material->GetShadingModels(); 
+
+			// If the material gets its shading model from material expressions and we have compiled one or more shading model expressions, 
+			// then use that shading model field instead. It's the most optimal set of shading models
+			if (Material->IsShadingModelFromMaterialExpression() && ShadingModelsFromCompilation.IsValid())
+			{
+				MaterialShadingModels = ShadingModelsFromCompilation;
+			}
+
 			if (Domain == MD_Surface && IsSubsurfaceShadingModel(MaterialShadingModels))
 			{
 				// Note we don't test for the blend mode as you can have a translucent material using the subsurface shading model
@@ -667,7 +684,6 @@ public:
 
 			Chunk[MP_PixelDepthOffset] = Material->CompilePropertyAndSetMaterialProperty(MP_PixelDepthOffset, this);
 
-			Chunk[MP_ShadingModel] = Material->CompilePropertyAndSetMaterialProperty(MP_ShadingModel, this);
 
 			// No more calls to non-vertex shader CompilePropertyAndSetMaterialProperty beyond this point
 			const uint32 SavedNumUserTexCoords = NumUserTexCoords;
@@ -1092,6 +1108,89 @@ public:
 			FShaderUniformBufferParameter::ModifyCompilationEnvironment(*CollectionName, ParameterCollections[CollectionIndex]->GetUniformBufferStruct(), InPlatform, OutEnvironment);
 		}
 		OutEnvironment.SetDefine(TEXT("IS_MATERIAL_SHADER"), TEXT("1"));
+
+		// Set all the shading models for this material here 
+		FMaterialShadingModelField ShadingModels = Material->GetShadingModels();
+
+		// If the material gets its shading model from the material expressions, then we use the result from the compilation (assuming it's valid).
+		// This result will potentially be tighter than what GetShadingModels() returns, because it only picks up the shading models from the expressions that get compiled for a specific feature level and quality level
+		// For example, the material might have shading models behind static switches. GetShadingModels() will return both the true and the false paths from that switch, whereas the shading model field from the compilation will only contain the actual shading model selected 
+		if (Material->IsShadingModelFromMaterialExpression() && ShadingModelsFromCompilation.IsValid())
+		{
+			// Shading models fetched from the compilation of the expression graph
+			ShadingModels = ShadingModelsFromCompilation;
+		}
+
+		ensure(ShadingModels.IsValid());
+
+		if (ShadingModels.IsLit())
+		{	
+			int NumSetMaterials = 0;
+			if (ShadingModels.HasShadingModel(MSM_DefaultLit))
+			{
+				OutEnvironment.SetDefine(TEXT("MATERIAL_SHADINGMODEL_DEFAULT_LIT"), TEXT("1"));
+				NumSetMaterials++;
+			}
+			if (ShadingModels.HasShadingModel(MSM_Subsurface))
+			{
+				OutEnvironment.SetDefine(TEXT("MATERIAL_SHADINGMODEL_SUBSURFACE"), TEXT("1"));
+				NumSetMaterials++;
+			}
+			if (ShadingModels.HasShadingModel(MSM_PreintegratedSkin))
+			{
+				OutEnvironment.SetDefine(TEXT("MATERIAL_SHADINGMODEL_PREINTEGRATED_SKIN"), TEXT("1"));
+				NumSetMaterials++;
+			}
+			if (ShadingModels.HasShadingModel(MSM_SubsurfaceProfile))
+			{
+				OutEnvironment.SetDefine(TEXT("MATERIAL_SHADINGMODEL_SUBSURFACE_PROFILE"), TEXT("1"));
+				NumSetMaterials++;
+			}
+			if (ShadingModels.HasShadingModel(MSM_ClearCoat))
+			{
+				OutEnvironment.SetDefine(TEXT("MATERIAL_SHADINGMODEL_CLEAR_COAT"), TEXT("1"));
+				NumSetMaterials++;
+			}
+			if (ShadingModels.HasShadingModel(MSM_TwoSidedFoliage))
+			{
+				OutEnvironment.SetDefine(TEXT("MATERIAL_SHADINGMODEL_TWOSIDED_FOLIAGE"), TEXT("1"));
+				NumSetMaterials++;
+			}
+			if (ShadingModels.HasShadingModel(MSM_Hair))
+			{
+				OutEnvironment.SetDefine(TEXT("MATERIAL_SHADINGMODEL_HAIR"), TEXT("1"));
+				NumSetMaterials++;
+			}
+			if (ShadingModels.HasShadingModel(MSM_Cloth))
+			{
+				OutEnvironment.SetDefine(TEXT("MATERIAL_SHADINGMODEL_CLOTH"), TEXT("1"));
+				NumSetMaterials++;
+			}
+			if (ShadingModels.HasShadingModel(MSM_Eye))
+			{
+				OutEnvironment.SetDefine(TEXT("MATERIAL_SHADINGMODEL_EYE"), TEXT("1"));
+				NumSetMaterials++;
+			}
+
+			if (NumSetMaterials == 1)
+			{
+				OutEnvironment.SetDefine(TEXT("MATERIAL_SINGLE_SHADINGMODEL"), TEXT("1"));
+			}
+
+			ensure(NumSetMaterials != 0);
+			if (NumSetMaterials == 0)
+			{
+				// Should not really end up here
+				UE_LOG(LogMaterial, Warning, TEXT("Unknown material shading model(s). Setting to MSM_DefaultLit"));
+				OutEnvironment.SetDefine(TEXT("MATERIAL_SHADINGMODEL_DEFAULT_LIT"),TEXT("1"));
+			}
+		}
+		else
+		{
+			// Unlit shading model can only exist by itself
+			OutEnvironment.SetDefine(TEXT("MATERIAL_SINGLE_SHADINGMODEL"), TEXT("1"));
+			OutEnvironment.SetDefine(TEXT("MATERIAL_SHADINGMODEL_UNLIT"), TEXT("1"));
+		}
 	}
 
 	void GetSharedInputsMaterialCode(FString& PixelMembersDeclaration, FString& NormalAssignment, FString& PixelMembersInitializationEpilog)
@@ -5544,6 +5643,7 @@ protected:
 
 	virtual int32 ShadingModel(EMaterialShadingModel InSelectedShadingModel) override
 	{
+		ShadingModelsFromCompilation.AddShadingModel(InSelectedShadingModel);
 		return AddInlinedCodeChunk(MCT_ShadingModel, TEXT("%d"), InSelectedShadingModel);
 	}
 
diff --git a/Engine/Source/Runtime/Engine/Private/Materials/MaterialShared.cpp b/Engine/Source/Runtime/Engine/Private/Materials/MaterialShared.cpp
--- a/Engine/Source/Runtime/Engine/Private/Materials/MaterialShared.cpp
+++ b/Engine/Source/Runtime/Engine/Private/Materials/MaterialShared.cpp
@@ -1030,7 +1030,6 @@ bool FMaterialResource::IsVolumetricPrimitive() const { return Material->Materia
 bool FMaterialResource::IsSpecialEngineMaterial() const { return Material->bUsedAsSpecialEngineMaterial; }
 bool FMaterialResource::HasVertexPositionOffsetConnected() const { return HasMaterialAttributesConnected() || (!Material->bUseMaterialAttributes && Material->WorldPositionOffset.IsConnected()); }
 bool FMaterialResource::HasPixelDepthOffsetConnected() const { return HasMaterialAttributesConnected() || (!Material->bUseMaterialAttributes && Material->PixelDepthOffset.IsConnected()); }
-bool FMaterialResource::HasShadingModelFromMaterialExpressionConnected() const { return HasMaterialAttributesConnected() || (!Material->bUseMaterialAttributes && Material->ShadingModelFromMaterialExpression.IsConnected()); }
 bool FMaterialResource::HasMaterialAttributesConnected() const { return Material->bUseMaterialAttributes && Material->MaterialAttributes.IsConnected(); }
 FString FMaterialResource::GetBaseMaterialPathName() const { return Material->GetPathName(); }
 FString FMaterialResource::GetDebugName() const
@@ -1227,6 +1226,11 @@ FMaterialShadingModelField FMaterialResource::GetShadingModels() const
 	return MaterialInstance ? MaterialInstance->GetShadingModels() : Material->GetShadingModels();
 }
 
+bool FMaterialResource::IsShadingModelFromMaterialExpression() const 
+{
+	return MaterialInstance ? MaterialInstance->IsShadingModelFromMaterialExpression() : Material->IsShadingModelFromMaterialExpression(); 
+}
+
 bool FMaterialResource::IsTwoSided() const 
 {
 	return MaterialInstance ? MaterialInstance->IsTwoSided() : Material->IsTwoSided();
@@ -1592,79 +1596,6 @@ void FMaterial::SetupMaterialEnvironment(
 			OutEnvironment.SetDefine(TEXT("MATERIAL_DOMAIN_SURFACE"),TEXT("1"));
 	};
 
-	// Set all the shading models for this material here 
-	FMaterialShadingModelField ShadingModels = GetShadingModels();
-	ensure(ShadingModels.IsValid());
-
-	if (ShadingModels.IsLit())
-	{	
-		int NumSetMaterials = 0;
-		if (ShadingModels.HasShadingModel(MSM_DefaultLit))
-		{
-			OutEnvironment.SetDefine(TEXT("MATERIAL_SHADINGMODEL_DEFAULT_LIT"), TEXT("1"));
-			NumSetMaterials++;
-		}
-		if (ShadingModels.HasShadingModel(MSM_Subsurface))
-		{
-			OutEnvironment.SetDefine(TEXT("MATERIAL_SHADINGMODEL_SUBSURFACE"), TEXT("1"));
-			NumSetMaterials++;
-		}
-		if (ShadingModels.HasShadingModel(MSM_PreintegratedSkin))
-		{
-			OutEnvironment.SetDefine(TEXT("MATERIAL_SHADINGMODEL_PREINTEGRATED_SKIN"), TEXT("1"));
-			NumSetMaterials++;
-		}
-		if (ShadingModels.HasShadingModel(MSM_SubsurfaceProfile))
-		{
-			OutEnvironment.SetDefine(TEXT("MATERIAL_SHADINGMODEL_SUBSURFACE_PROFILE"), TEXT("1"));
-			NumSetMaterials++;
-		}
-		if (ShadingModels.HasShadingModel(MSM_ClearCoat))
-		{
-			OutEnvironment.SetDefine(TEXT("MATERIAL_SHADINGMODEL_CLEAR_COAT"), TEXT("1"));
-			NumSetMaterials++;
-		}
-		if (ShadingModels.HasShadingModel(MSM_TwoSidedFoliage))
-		{
-			OutEnvironment.SetDefine(TEXT("MATERIAL_SHADINGMODEL_TWOSIDED_FOLIAGE"), TEXT("1"));
-			NumSetMaterials++;
-		}
-		if (ShadingModels.HasShadingModel(MSM_Hair))
-		{
-			OutEnvironment.SetDefine(TEXT("MATERIAL_SHADINGMODEL_HAIR"), TEXT("1"));
-			NumSetMaterials++;
-		}
-		if (ShadingModels.HasShadingModel(MSM_Cloth))
-		{
-			OutEnvironment.SetDefine(TEXT("MATERIAL_SHADINGMODEL_CLOTH"), TEXT("1"));
-			NumSetMaterials++;
-		}
-		if (ShadingModels.HasShadingModel(MSM_Eye))
-		{
-			OutEnvironment.SetDefine(TEXT("MATERIAL_SHADINGMODEL_EYE"), TEXT("1"));
-			NumSetMaterials++;
-		}
-
-		if (NumSetMaterials == 1)
-		{
-			OutEnvironment.SetDefine(TEXT("MATERIAL_SINGLE_SHADINGMODEL"), TEXT("1"));
-		}
-
-		ensure(NumSetMaterials != 0);
-		if (NumSetMaterials == 0)
-		{
-			// Should not really end up here
-			UE_LOG(LogMaterial, Warning, TEXT("Unknown material shading model(s). Setting to MSM_DefaultLit"));
-			OutEnvironment.SetDefine(TEXT("MATERIAL_SHADINGMODEL_DEFAULT_LIT"),TEXT("1"));
-		}
-	}
-	else
-	{
-		// Unlit shading model can only exist by itself
-		OutEnvironment.SetDefine(TEXT("MATERIAL_SINGLE_SHADINGMODEL"), TEXT("1"));
-		OutEnvironment.SetDefine(TEXT("MATERIAL_SHADINGMODEL_UNLIT"), TEXT("1"));
-	}
-
 	if (IsTranslucentBlendMode(GetBlendMode()))
 	{
 		switch(GetTranslucencyLightingMode())
diff --git a/Engine/Source/Runtime/Engine/Public/MaterialShared.h b/Engine/Source/Runtime/Engine/Public/MaterialShared.h
--- a/Engine/Source/Runtime/Engine/Public/MaterialShared.h
+++ b/Engine/Source/Runtime/Engine/Public/MaterialShared.h
@@ -1307,6 +1307,7 @@ public:
 	virtual enum EBlendMode GetBlendMode() const = 0;
 	ENGINE_API virtual enum ERefractionMode GetRefractionMode() const;
 	virtual FMaterialShadingModelField GetShadingModels() const = 0;
+	virtual bool IsShadingModelFromMaterialExpression() const = 0;
 	virtual enum ETranslucencyLightingMode GetTranslucencyLightingMode() const { return TLM_VolumetricNonDirectional; };
 	virtual float GetOpacityMaskClipValue() const = 0;
 	virtual bool GetCastDynamicShadowAsMasked() const = 0;
@@ -1328,7 +1329,6 @@ public:
 	virtual FString GetFriendlyName() const = 0;
 	virtual bool HasVertexPositionOffsetConnected() const { return false; }
 	virtual bool HasPixelDepthOffsetConnected() const { return false; }
-	virtual bool HasShadingModelFromMaterialExpressionConnected() const { return false; }
 	virtual bool HasMaterialAttributesConnected() const { return false; }
 	virtual uint32 GetDecalBlendMode() const { return 0; }
 	virtual uint32 GetMaterialDecalResponse() const { return 0; }
@@ -1998,6 +1998,7 @@ public:
 	ENGINE_API virtual bool HasNormalConnected() const override;
 	ENGINE_API virtual bool HasEmissiveColorConnected() const override;
 	ENGINE_API virtual FMaterialShadingModelField GetShadingModels() const override;
+	ENGINE_API virtual bool IsShadingModelFromMaterialExpression() const override;
 	ENGINE_API virtual enum ETranslucencyLightingMode GetTranslucencyLightingMode() const override;
 	ENGINE_API virtual float GetOpacityMaskClipValue() const override;
 	ENGINE_API virtual bool GetCastDynamicShadowAsMasked() const override;
@@ -2065,7 +2066,6 @@ protected:
 	ENGINE_API virtual void GatherExpressionsForCustomInterpolators(TArray<class UMaterialExpression*>& OutExpressions) const override;
 	ENGINE_API virtual bool HasVertexPositionOffsetConnected() const override;
 	ENGINE_API virtual bool HasPixelDepthOffsetConnected() const override;
-	ENGINE_API virtual bool HasShadingModelFromMaterialExpressionConnected() const override;
 	ENGINE_API virtual bool HasMaterialAttributesConnected() const override;
 	/** Useful for debugging. */
 	ENGINE_API virtual FString GetBaseMaterialPathName() const override;
