Prevent crashing if Seamless Travel causes improper cleanup of GUIDs.
This will still result in replication errors, but net.AllowClientRemapCacheObject may help.

diff --git a/Engine/Source/Runtime/Engine/Private/DataReplication.cpp b/Engine/Source/Runtime/Engine/Private/DataReplication.cpp
--- a/Engine/Source/Runtime/Engine/Private/DataReplication.cpp
+++ b/Engine/Source/Runtime/Engine/Private/DataReplication.cpp
@@ -604,7 +604,16 @@ void FObjectReplicator::StartReplicating(class UActorChannel * InActorChannel)
 
 	// Cache off netGUID so if this object gets deleted we can close it
 	ObjectNetGUID = ConnectionNetDriver->GuidCache->GetOrAssignNetGUID( Object );
-	check(!ObjectNetGUID.IsDefault() && ObjectNetGUID.IsValid());
+
+	const bool bIsValidToReplicate = !ObjectNetGUID.IsDefault() && ObjectNetGUID.IsValid();
+	if (UNLIKELY(!bIsValidToReplicate))
+	{
+		// This has mostly been seen when doing a Seamless Travel Restart.
+		// In that case, the server can think the client is still on the same map and replicate objects before
+		// the client has finished the travel, and the client will later remove those references from the Package Map.
+		UE_LOG(LogRep, Error, TEXT("StartReplicating: Invalid Net GUID. Object may fail to replicate properties or handle RPCs. Object %s"), *Object->GetPathName());
+		return;
+	}
 
 	if (ConnectionNetDriver->IsServer() || ConnectionNetDriver->MaySendProperties())
 	{
diff --git a/Engine/Source/Runtime/Engine/Private/PackageMapClient.cpp b/Engine/Source/Runtime/Engine/Private/PackageMapClient.cpp
--- a/Engine/Source/Runtime/Engine/Private/PackageMapClient.cpp
+++ b/Engine/Source/Runtime/Engine/Private/PackageMapClient.cpp
@@ -76,6 +76,13 @@ static FAutoConsoleVariableRef CVarDelinquencyNumberOfTopOffendersToTrack(
 		" items are tracked.")
 );
 
+static bool GbAllowClientRemapCacheObject = false;
+static FAutoConsoleVariableRef CVarAllowClientRemapCacheObject(
+	TEXT("net.AllowClientRemapCacheObject"),
+	GbAllowClientRemapCacheObject,
+	TEXT("When enabled, we will allow clients to remap read only cache objects and keep the same NetGUID.")
+);
+
 void BroadcastNetFailure(UNetDriver* Driver, ENetworkFailure::Type FailureType, const FString& ErrorStr)
 {
 	UWorld* World = Driver->GetWorld();
@@ -2247,16 +2254,18 @@ public:
 
 void FNetGUIDCache::CleanReferences()
 {
+	const double Time = FPlatformTime::Seconds();
+
 	// Mark all static or non valid dynamic guids to timeout after NETWORK_GUID_TIMEOUT seconds
 	// We want to leave them around for a certain amount of time to allow in-flight references to these guids to continue to resolve
-	for ( auto It = ObjectLookup.CreateIterator(); It; ++It )
+	for (auto It = ObjectLookup.CreateIterator(); It; ++It)
 	{
-		if ( It.Value().ReadOnlyTimestamp != 0 )
+		if (It.Value().ReadOnlyTimestamp != 0)
 		{
 			// If this guid was suppose to time out, check to see if it has, otherwise ignore it
 			const double NETWORK_GUID_TIMEOUT = 90;
 
-			if ( FPlatformTime::Seconds() - It.Value().ReadOnlyTimestamp > NETWORK_GUID_TIMEOUT )
+			if (Time - It.Value().ReadOnlyTimestamp > NETWORK_GUID_TIMEOUT)
 			{
 				It.RemoveCurrent();
 			}
@@ -2264,16 +2273,16 @@ void FNetGUIDCache::CleanReferences()
 			continue;
 		}
 
-		if ( !It.Value().Object.IsValid() || It.Key().IsStatic() )
+		if (!It.Value().Object.IsValid())
 		{
 			// We will leave this guid around for NETWORK_GUID_TIMEOUT seconds to make sure any in-flight guids can be resolved
-			It.Value().ReadOnlyTimestamp = FPlatformTime::Seconds();
+			It.Value().ReadOnlyTimestamp = Time;
 		}
 	}
 
-	for ( auto It = NetGUIDLookup.CreateIterator(); It; ++It )
+	for (auto It = NetGUIDLookup.CreateIterator(); It; ++It)
 	{
-		if ( !It.Key().IsValid() || !ObjectLookup.Contains( It.Value() ) )
+		if (!It.Key().IsValid() || !ObjectLookup.Contains(It.Value()))
 		{
 			It.RemoveCurrent();
 		}
@@ -2281,19 +2290,19 @@ void FNetGUIDCache::CleanReferences()
 
 	// Sanity check
 	// (make sure look-ups are reciprocal)
-	for ( auto It = ObjectLookup.CreateIterator(); It; ++It )
+	for (auto It = ObjectLookup.CreateIterator(); It; ++It)
 	{
-		check( !It.Key().IsDefault() );
-		check( It.Key().IsStatic() != It.Key().IsDynamic() );
+		check(!It.Key().IsDefault());
+		check(It.Key().IsStatic() != It.Key().IsDynamic());
 
-		checkf( !It.Value().Object.IsValid() || NetGUIDLookup.FindRef( It.Value().Object ) == It.Key() || It.Value().ReadOnlyTimestamp != 0, TEXT( "Failed to validate ObjectLookup map in UPackageMap. Object '%s' was not in the NetGUIDLookup map with with value '%s'." ), *It.Value().Object.Get()->GetPathName(), *It.Key().ToString() );
+		checkf(!It.Value().Object.IsValid() || NetGUIDLookup.FindRef(It.Value().Object) == It.Key() || It.Value().ReadOnlyTimestamp != 0, TEXT("Failed to validate ObjectLookup map in UPackageMap. Object '%s' was not in the NetGUIDLookup map with with value '%s'." ), *It.Value().Object.Get()->GetPathName(), *It.Key().ToString());
 	}
 
 #if !UE_BUILD_SHIPPING || !UE_BUILD_TEST
-	for ( auto It = NetGUIDLookup.CreateIterator(); It; ++It )
+	for (auto It = NetGUIDLookup.CreateIterator(); It; ++It)
 	{
-		check( It.Key().IsValid() );
-		checkf( ObjectLookup.FindRef( It.Value() ).Object == It.Key(), TEXT("Failed to validate NetGUIDLookup map in UPackageMap. GUID '%s' was not in the ObjectLookup map with with object '%s'."), *It.Value().ToString(), *It.Key().Get()->GetPathName());
+		check(It.Key().IsValid());
+		checkf(ObjectLookup.FindRef(It.Value() ).Object == It.Key(), TEXT("Failed to validate NetGUIDLookup map in UPackageMap. GUID '%s' was not in the ObjectLookup map with with object '%s'."), *It.Value().ToString(), *It.Key().Get()->GetPathName());
 	}
 #endif
 
@@ -2302,7 +2311,7 @@ void FNetGUIDCache::CleanReferences()
 	ObjectLookup.CountBytes( CountBytesAr );
 	NetGUIDLookup.CountBytes( CountBytesAr );
 
-	UE_LOG( LogNetPackageMap, Log, TEXT( "FNetGUIDCache::CleanReferences: ObjectLookup: %i, NetGUIDLookup: %i, Mem: %i kB" ), ObjectLookup.Num(), NetGUIDLookup.Num(), ( CountBytesAr.Mem / 1024 ) );
+	UE_LOG(LogNetPackageMap, Log, TEXT("FNetGUIDCache::CleanReferences: ObjectLookup: %i, NetGUIDLookup: %i, Mem: %i kB"), ObjectLookup.Num(), NetGUIDLookup.Num(), (CountBytesAr.Mem / 1024));
 }
 
 bool FNetGUIDCache::SupportsObject( const UObject* Object, const TWeakObjectPtr<UObject>* WeakObjectPtr ) const
@@ -2362,34 +2371,48 @@ bool FNetGUIDCache::IsNetGUIDAuthority() const
 }
 
 /** Gets or assigns a new NetGUID to this object. Returns whether the object is fully mapped or not */
-FNetworkGUID FNetGUIDCache::GetOrAssignNetGUID( UObject* Object, const TWeakObjectPtr<UObject>* WeakObjectPtr)
+FNetworkGUID FNetGUIDCache::GetOrAssignNetGUID(UObject* Object, const TWeakObjectPtr<UObject>* WeakObjectPtr)
 {
 	// Construct WeakPtr once: either use the passed in one or create a new one.
-	const TWeakObjectPtr<UObject>& WeakObject = WeakObjectPtr ? *WeakObjectPtr : MakeWeakObjectPtr( const_cast<UObject*>( Object ) );
+	const TWeakObjectPtr<UObject>& WeakObject = WeakObjectPtr ? *WeakObjectPtr : MakeWeakObjectPtr(const_cast<UObject*>(Object));
 
-	if ( !Object || !SupportsObject( Object, &WeakObject ) )
+	if (!Object || !SupportsObject(Object, &WeakObject))
 	{
 		// Null of unsupported object, leave as default NetGUID and just return mapped=true
+		UE_LOG(LogNetPackageMap, Verbose, TEXT("GetOrAssignNetGUID: Object is not supported. Object %s"), *GetPathNameSafe(Object));
 		return FNetworkGUID();
 	}
 
 	// ----------------
 	// Assign NetGUID if necessary
-	// ----------------	
-	FNetworkGUID NetGUID = NetGUIDLookup.FindRef( WeakObject );
+	// ----------------
+	
+	const bool bIsNetGUIDAuthority = IsNetGUIDAuthority();
+	FNetworkGUID NetGUID = NetGUIDLookup.FindRef(WeakObject);
 
-	if ( NetGUID.IsValid() )
+	if (NetGUID.IsValid())
 	{
-		const FNetGuidCacheObject* CacheObject = ObjectLookup.Find( NetGUID );
+		FNetGuidCacheObject* CacheObject = ObjectLookup.Find(NetGUID);
 
 		// Check to see if this guid is read only
 		// If so, we should ignore this entry, and create a new one (or send default as client)
-		const bool bReadOnly = CacheObject != NULL && CacheObject->ReadOnlyTimestamp > 0;
+		const bool bReadOnly = CacheObject != nullptr && CacheObject->ReadOnlyTimestamp > 0;
 
-		if ( bReadOnly )
+		if (bReadOnly)
 		{
 			// Reset this object's guid, we will re-assign below (or send default as a client)
-			NetGUIDLookup.Remove( WeakObject );
+			UE_LOG(LogNetPackageMap, Warning, TEXT("GetOrAssignNetGUID: Attempt to reassign read-only guid. FullNetGUIDPath: %s"), *FullNetGUIDPath(NetGUID));
+
+			const bool bAllowClientRemap = !bIsNetGUIDAuthority && GbAllowClientRemapCacheObject;
+			if (bAllowClientRemap)
+			{
+				CacheObject->ReadOnlyTimestamp = 0;
+				return NetGUID;
+			}
+			else
+			{
+				NetGUIDLookup.Remove(WeakObject);
+			}
 		}
 		else
 		{
@@ -2397,15 +2420,16 @@ FNetworkGUID FNetGUIDCache::GetOrAssignNetGUID( UObject* Object, const TWeakObje
 		}
 	}
 
-	if ( !IsNetGUIDAuthority() )
+	if (!bIsNetGUIDAuthority)
 	{
 		// We cannot make or assign new NetGUIDs
 		// Generate a default GUID, which signifies we write the full path
 		// The server should detect this, and assign a full-time guid, and send that back to us
+		UE_LOG(LogNetPackageMap, Verbose, TEXT("GetOrAssignNetGUID: NetGUIDLookup did not contain object on client, returning default. Object %s"), *Object->GetPathName());
 		return FNetworkGUID::GetDefault();
 	}
 
-	return AssignNewNetGUID_Server( Object );
+	return AssignNewNetGUID_Server(Object);
 }
 
 FNetworkGUID FNetGUIDCache::GetNetGUID(const UObject* Object) const
@@ -2940,10 +2964,11 @@ UObject* FNetGUIDCache::GetObjectFromNetGUID( const FNetworkGUID& NetGUID, const
 
 	// Assume this is a package if the outer is invalid and this is a static guid
 	const bool bIsPackage = NetGUID.IsStatic() && !CacheObjectPtr->OuterGUID.IsValid();
+	const bool bIsNetGUIDAuthority = IsNetGUIDAuthority();
 
 	if ( Object == NULL && !CacheObjectPtr->bNoLoad )
 	{
-		if (IsNetGUIDAuthority())
+		if (bIsNetGUIDAuthority)
 		{
 			// Log when the server needs to re-load an object, it's probably due to a GC after initially loading as default guid
 			UE_LOG(LogNetPackageMap, Warning, TEXT("GetObjectFromNetGUID: Server re-loading object (might have been GC'd). FullNetGUIDPath: %s"), *FullNetGUIDPath(NetGUID));
@@ -3070,11 +3095,19 @@ UObject* FNetGUIDCache::GetObjectFromNetGUID( const FNetworkGUID& NetGUID, const
 	// Assign the guid to the object 
 	// We don't want to assign this guid to the object if this guid is timing out
 	// But we'll have to if there is no other guid yet
-	if ( CacheObjectPtr->ReadOnlyTimestamp == 0 || !NetGUIDLookup.Contains( Object ) )
+	const bool bAllowClientRemap = !bIsNetGUIDAuthority && GbAllowClientRemapCacheObject;
+	const bool bIsNotReadOnlyOrAllowRemap = (CacheObjectPtr->ReadOnlyTimestamp == 0 || bAllowClientRemap);
+
+	if (bIsNotReadOnlyOrAllowRemap || !NetGUIDLookup.Contains(Object))
 	{
-		if ( CacheObjectPtr->ReadOnlyTimestamp > 0 )
+		if (CacheObjectPtr->ReadOnlyTimestamp > 0)
 		{
 			UE_LOG( LogNetPackageMap, Warning, TEXT( "GetObjectFromNetGUID: Attempt to reassign read-only guid. FullNetGUIDPath: %s" ), *FullNetGUIDPath( NetGUID ) );
+
+			if (bAllowClientRemap)
+			{
+				CacheObjectPtr->ReadOnlyTimestamp = 0;
+			}
 		}
 
 		NetGUIDLookup.Add( Object, NetGUID );
