Scene Capture Component - Adding a way to configure it using the minimal view from a camera.

diff --git a/Engine/Source/Runtime/Engine/Private/Components/SceneCaptureComponent.cpp b/Engine/Source/Runtime/Engine/Private/Components/SceneCaptureComponent.cpp
--- a/Engine/Source/Runtime/Engine/Private/Components/SceneCaptureComponent.cpp
+++ b/Engine/Source/Runtime/Engine/Private/Components/SceneCaptureComponent.cpp
@@ -514,6 +514,16 @@ void USceneCaptureComponent2D::TickComponent(float DeltaTime, enum ELevelTick Ti
 	}
 }
 
+void USceneCaptureComponent2D::SetCameraView(const FMinimalViewInfo& DesiredView)
+{
+	SetWorldLocation(DesiredView.Location);
+	SetWorldRotation(DesiredView.Rotation);
+
+	FOVAngle = DesiredView.FOV;
+	ProjectionType = DesiredView.ProjectionMode;
+	OrthoWidth = DesiredView.OrthoWidth;
+}
+
 void USceneCaptureComponent2D::GetCameraView(float DeltaTime, FMinimalViewInfo& OutMinimalViewInfo)
 {
 	OutMinimalViewInfo.Location = GetComponentLocation();
