Scene Capture Component - Adding a way to configure it using the minimal view from a camera.

Michael.Noland

diff --git a/Engine/Source/Runtime/Engine/Classes/Components/SceneCaptureComponent2D.h b/Engine/Source/Runtime/Engine/Classes/Components/SceneCaptureComponent2D.h
--- a/Engine/Source/Runtime/Engine/Classes/Components/SceneCaptureComponent2D.h
+++ b/Engine/Source/Runtime/Engine/Classes/Components/SceneCaptureComponent2D.h
@@ -122,7 +122,9 @@ public:
 
 	//~ End UObject Interface
 
-	virtual void GetCameraView(float DeltaTime, FMinimalViewInfo& DesiredView);
+	void SetCameraView(const FMinimalViewInfo& DesiredView);
+
+	virtual void GetCameraView(float DeltaTime, FMinimalViewInfo& OutDesiredView);
 
 	/** Adds an Blendable (implements IBlendableInterface) to the array of Blendables (if it doesn't exist) and update the weight */
 	UFUNCTION(BlueprintCallable, Category="Rendering")
