Reverting github PR causing crash. (Was attempting to force a synchronous load of a soundwave.)

diff --git a/Engine/Source/Runtime/Engine/Private/SoundNodeWavePlayer.cpp b/Engine/Source/Runtime/Engine/Private/SoundNodeWavePlayer.cpp
--- a/Engine/Source/Runtime/Engine/Private/SoundNodeWavePlayer.cpp
+++ b/Engine/Source/Runtime/Engine/Private/SoundNodeWavePlayer.cpp
@@ -185,10 +185,6 @@ void USoundNodeWavePlayer::ParseNodes( FAudioDevice* AudioDevice, const UPTRINT
 float USoundNodeWavePlayer::GetDuration()
 {
 	float Duration = 0.f;
-	if (!SoundWave && !IsAsyncLoading())
-	{
-		SoundWave = SoundWaveAssetPtr.LoadSynchronous();
-	}
 	if (SoundWave)
 	{
 		if (bLooping)
