- Add cvar to control forcing virtual loop update on listener transform jumps.

diff --git a/Engine/Source/Runtime/Engine/Private/AudioDevice.cpp b/Engine/Source/Runtime/Engine/Private/AudioDevice.cpp
--- a/Engine/Source/Runtime/Engine/Private/AudioDevice.cpp
+++ b/Engine/Source/Runtime/Engine/Private/AudioDevice.cpp
@@ -152,7 +152,10 @@ FAutoConsoleVariableRef CVarDisableBinauralSpatialization(
 	TEXT("Disables binaural spatialization.\n"),
 	ECVF_Default);
 
-using FVirtualLoopPair = TPair<FActiveSound*, FAudioVirtualLoop>;
+namespace
+{
+	using FVirtualLoopPair = TPair<FActiveSound*, FAudioVirtualLoop>;
+} // namespace <>
 
 /*-----------------------------------------------------------------------------
 	FAudioDevice implementation.
@@ -2814,22 +2817,21 @@ void FAudioDevice::SetListener(UWorld* World, const int32 InViewportIndex, const
 	}
 
 
-	FAudioDevice* AudioDevice = this;
-	FAudioThread::RunCommandOnAudioThread([AudioDevice, WorldID, InViewportIndex, ListenerTransformCopy, InDeltaSeconds]()
+	FAudioThread::RunCommandOnAudioThread([this, WorldID, InViewportIndex, ListenerTransformCopy, InDeltaSeconds]()
 	{
 		// Broadcast to a 3rd party plugin listener observer if enabled
-		for (TAudioPluginListenerPtr PluginManager : AudioDevice->PluginListeners)
+		for (TAudioPluginListenerPtr PluginManager : PluginListeners)
 		{
-			PluginManager->OnListenerUpdated(AudioDevice, InViewportIndex, ListenerTransformCopy, InDeltaSeconds);
+			PluginManager->OnListenerUpdated(this, InViewportIndex, ListenerTransformCopy, InDeltaSeconds);
 		}
 
-		TArray<FListener>& AudioThreadListeners = AudioDevice->Listeners;
+		TArray<FListener>& AudioThreadListeners = Listeners;
 		if (InViewportIndex >= AudioThreadListeners.Num())
 		{
 			const int32 NumListeners = InViewportIndex - AudioThreadListeners.Num() + 1;
 			for (int32 i = 0; i < NumListeners; ++i)
 			{
-				AudioThreadListeners.Add(FListener(AudioDevice));
+				AudioThreadListeners.Add(FListener(this));
 			}
 		}
 
@@ -2845,6 +2847,12 @@ void FAudioDevice::SetListener(UWorld* World, const int32 InViewportIndex, const
 		}
 #endif
 
+		if (FAudioVirtualLoop::ShouldListenerMoveForceUpdate(Listener.Transform, ListenerTransformCopy))
+		{
+			const bool bForceUpdate = true;
+			UpdateVirtualLoops(bForceUpdate);
+		}
+
 		Listener.WorldID = WorldID;
 		Listener.Transform = ListenerTransformCopy;
 
@@ -3842,7 +3850,8 @@ void FAudioDevice::Update(bool bGameTicking)
 		SCOPED_NAMED_EVENT(FAudioDevice_UpdateVirtualLoops, FColor::Blue);
 		// Update which loops should re-trigger due to coming back into proximity
 		// or allowed by concurrency re-evaluating in context of other sounds stopping
-		UpdateVirtualLoops();
+		const bool bForceUpdate = false;
+		UpdateVirtualLoops(bForceUpdate);
 	}
 
 	// update if baked analysis is enabled
@@ -5742,7 +5751,7 @@ float FAudioDevice::GetGameDeltaTime() const
 	return FMath::Min(DeltaTime, 0.5f);
 }
 
-void FAudioDevice::UpdateVirtualLoops()
+void FAudioDevice::UpdateVirtualLoops(bool bForceUpdate)
 {
 	if (FAudioVirtualLoop::IsEnabled())
 	{
@@ -5775,7 +5784,7 @@ void FAudioDevice::UpdateVirtualLoops()
 
 			// If the loop is ready to realize, add to array to be re-triggered outside of the loop
 			// to avoid map manipulation while iterating.
-			if (VirtualLoop.CanRealize(DeltaTime))
+			if (VirtualLoop.CanRealize(DeltaTime, bForceUpdate))
 			{
 				VirtualLoopsToRetrigger.Add(VirtualLoop);
 			}
diff --git a/Engine/Source/Runtime/Engine/Private/AudioVirtualLoop.cpp b/Engine/Source/Runtime/Engine/Private/AudioVirtualLoop.cpp
--- a/Engine/Source/Runtime/Engine/Private/AudioVirtualLoop.cpp
+++ b/Engine/Source/Runtime/Engine/Private/AudioVirtualLoop.cpp
@@ -21,6 +21,13 @@ FAutoConsoleVariableRef CVarVirtualLoopsPerfDistance(
 	TEXT("Sets virtual loop distance to scale update rate between min and max beyond max audible distance of sound.\n"),
 	ECVF_Default);
 
+static float VirtualLoopsForceUpdateListenerMoveDistanceCVar = 2500.0f;
+FAutoConsoleVariableRef CVarVirtualLoopsForceUpdateListenerMoveDistance(
+	TEXT("au.VirtualLoops.ForceUpdateListenerMoveDistance"),
+	VirtualLoopsForceUpdateListenerMoveDistanceCVar,
+	TEXT("Sets distance threshold required to force an update to check for virtualized sounds to realize if listener moves in a single frame over the given distance.\n"),
+	ECVF_Default);
+
 static float VirtualLoopsUpdateRateMinCVar = 0.1f;
 FAutoConsoleVariableRef CVarVirtualLoopsUpdateRateMin(
 	TEXT("au.VirtualLoops.UpdateRate.Min"),
@@ -139,7 +146,7 @@ bool FAudioVirtualLoop::IsInAudibleRange(const FActiveSound& InActiveSound, cons
 			return true;
 		}
 
-		DistanceScale = AttenuationSettingsToApply->GetFocusDistanceScale(AudioDevice->GetGlobalFocusSettings(), InActiveSound.FocusDistanceScale);
+		DistanceScale = InActiveSound.FocusDistanceScale;
 	}
 
 	DistanceScale = FMath::Max(DistanceScale, 0.0001f);
@@ -147,9 +154,13 @@ bool FAudioVirtualLoop::IsInAudibleRange(const FActiveSound& InActiveSound, cons
 	return AudioDevice->LocationIsAudible(Location, InActiveSound.MaxDistance / DistanceScale);
 }
 
-bool FAudioVirtualLoop::CanRealize(float DeltaTime)
+bool FAudioVirtualLoop::CanRealize(float DeltaTime, bool bForceUpdate)
 {
-	if (UpdateInterval > 0.0f)
+	if (bForceUpdate)
+	{
+		TimeSinceLastUpdate = 0.0f;
+	}
+	else if (UpdateInterval > 0.0f)
 	{
 		TimeSinceLastUpdate += DeltaTime;
 		if (UpdateInterval > TimeSinceLastUpdate)
@@ -172,3 +183,10 @@ bool FAudioVirtualLoop::CanRealize(float DeltaTime)
 
 	return true;
 }
+
+bool FAudioVirtualLoop::ShouldListenerMoveForceUpdate(const FTransform& LastTransform, const FTransform& CurrentTransform)
+{
+	const float DistanceSq = FVector::DistSquared(LastTransform.GetTranslation(), CurrentTransform.GetTranslation());
+	const float ForceUpdateDistSq = VirtualLoopsForceUpdateListenerMoveDistanceCVar * VirtualLoopsForceUpdateListenerMoveDistanceCVar;
+	return DistanceSq > ForceUpdateDistSq;
+}
diff --git a/Engine/Source/Runtime/Engine/Public/AudioDevice.h b/Engine/Source/Runtime/Engine/Public/AudioDevice.h
--- a/Engine/Source/Runtime/Engine/Public/AudioDevice.h
+++ b/Engine/Source/Runtime/Engine/Public/AudioDevice.h
@@ -1474,7 +1474,7 @@ public:
 	/** Returns the game's delta time */
 	float GetGameDeltaTime() const;
 
-	void UpdateVirtualLoops();
+	void UpdateVirtualLoops(bool bForceUpdate);
 
 	/** Sets the update delta time for the audio frame */
 	virtual void UpdateDeviceDeltaTime()
diff --git a/Engine/Source/Runtime/Engine/Public/AudioVirtualLoop.h b/Engine/Source/Runtime/Engine/Public/AudioVirtualLoop.h
--- a/Engine/Source/Runtime/Engine/Public/AudioVirtualLoop.h
+++ b/Engine/Source/Runtime/Engine/Public/AudioVirtualLoop.h
@@ -17,9 +17,19 @@ private:
 
 	FActiveSound* ActiveSound;
 
+	/**
+	  * Check if provided active sound is in audible range.
+	  */
+	static bool IsInAudibleRange(const FActiveSound& InActiveSound, const FAudioDevice* InAudioDevice = nullptr);
+
 public:
 	FAudioVirtualLoop();
 
+	/**
+	 * Check to see if listener move is far enough such that a check for virtual loop realization is necessary
+	 */
+	static bool ShouldListenerMoveForceUpdate(const FTransform& LastTransform, const FTransform& CurrentTransform);
+
 	/**
 	 * Checks if provided active sound is available to be virtualized.  If so, returns new active sound ready to be
 	 * added to virtual loop management by parent audio device.
@@ -47,11 +57,6 @@ public:
 	 */
 	const FActiveSound& GetActiveSound() const;
 
-	/**
-	 * Check if provided active sound is in audible range.
-	 */
-	static bool IsInAudibleRange(const FActiveSound& InActiveSound, const FAudioDevice* InAudioDevice = nullptr);
-
 	/**
 	 * Overrides the update interval to the provided length.
 	 */
@@ -61,5 +66,5 @@ public:
 	  * Updates the loop and checks if ready to play (or 'realize').
 	  * Returns whether or not the sound is ready to be realized.
 	  */
-	bool CanRealize(float DeltaTime);
+	bool CanRealize(float DeltaTime, bool bForceUpdate);
 };
