Minor prep for code extraction/refactor to make focus updates more accurate
- Remove unused parameters passed to GetAzimuth
- Minor clean-up

diff --git a/Engine/Source/Runtime/Engine/Private/ActiveSound.cpp b/Engine/Source/Runtime/Engine/Private/ActiveSound.cpp
--- a/Engine/Source/Runtime/Engine/Private/ActiveSound.cpp
+++ b/Engine/Source/Runtime/Engine/Private/ActiveSound.cpp
@@ -1449,8 +1449,6 @@ void FActiveSound::ApplyAttenuation(FSoundParseParameters& ParseParams, const FL
 	FocusDistanceScale = 1.0f;
 	FocusPriorityScale = 1.0f;
 
-	check(Sound);
-
 	if (Settings->bEnableReverbSend)
 	{
 		ParseParams.ReverbSendMethod = Settings->ReverbSendMethod;
@@ -1462,16 +1460,13 @@ void FActiveSound::ApplyAttenuation(FSoundParseParameters& ParseParams, const FL
 
 	if (Settings->bSpatialize || Settings->bEnableListenerFocus)
 	{
-		AudioDevice->GetAzimuth(ListenerData, Sound, SoundTransform, *Settings, Listener.Transform, Azimuth, AbsoluteAzimuth);
+		AudioDevice->GetAzimuth(ListenerData, SoundTransform, *Settings, Listener.Transform, Azimuth, AbsoluteAzimuth);
 
 		if (Settings->bSpatialize)
 		{
 			ParseParams.AttenuationDistance = ListenerData.AttenuationDistance;
-
 			ParseParams.ListenerToSoundDistance = ListenerData.ListenerToSoundDistance;
-
 			ParseParams.ListenerToSoundDistanceForPanning = ListenerData.ListenerToSoundDistanceForPanning;
-
 			ParseParams.AbsoluteAzimuth = AbsoluteAzimuth;
 		}
 
@@ -1481,7 +1476,7 @@ void FActiveSound::ApplyAttenuation(FSoundParseParameters& ParseParams, const FL
 			const FGlobalFocusSettings& FocusSettings = AudioDevice->GetGlobalFocusSettings();
 
 			// Get the current target focus factor
-			const float TargetFocusFactor = AudioDevice->GetFocusFactor(ListenerData, Sound, Azimuth, *Settings);
+			const float TargetFocusFactor = AudioDevice->GetFocusFactor(Azimuth, *Settings);
 
 			// User opt-in for focus interpolation
 			if (Settings->bEnableFocusInterpolation)
diff --git a/Engine/Source/Runtime/Engine/Private/AudioDevice.cpp b/Engine/Source/Runtime/Engine/Private/AudioDevice.cpp
--- a/Engine/Source/Runtime/Engine/Private/AudioDevice.cpp
+++ b/Engine/Source/Runtime/Engine/Private/AudioDevice.cpp
@@ -4663,17 +4663,15 @@ void FAudioDevice::GetMaxDistanceAndFocusFactor(USoundBase* Sound, const UWorld*
 			float Azimuth = 0.0f;
 			float AbsoluteAzimuth = 0.0f;
 			const int32 ClosestListenerIndex = FindClosestListenerIndex(SoundTransform);
-			if (ClosestListenerIndex >= 0 && ClosestListenerIndex < ListenerTransforms.Num())
+			if (ClosestListenerIndex >= ListenerTransforms.Num())
 			{
-				const FTransform& ListenerTransform = ListenerTransforms[ClosestListenerIndex];
-				GetAzimuth(ListenerData, Sound, SoundTransform, *AttenuationSettingsToApply, ListenerTransform, Azimuth, AbsoluteAzimuth);
-
-				OutFocusFactor = GetFocusFactor(ListenerData, Sound, Azimuth, *AttenuationSettingsToApply);
-			}
-			else
-			{
-				UE_LOG(LogAudio, Warning, TEXT("Failed to get max distance and focus factor of sound."));
+				UE_LOG(LogAudio, Warning, TEXT("Invalid ClosestListenerIndex. Sound max distance and focus factor calculation failed."));
+				return;
 			}
+
+			const FTransform& ListenerTransform = ListenerTransforms[ClosestListenerIndex];
+			GetAzimuth(ListenerData, SoundTransform, *AttenuationSettingsToApply, ListenerTransform, Azimuth, AbsoluteAzimuth);
+			OutFocusFactor = GetFocusFactor(Azimuth, *AttenuationSettingsToApply);
 		}
 	}
 	else
@@ -4854,7 +4852,7 @@ void FAudioDevice::GetAttenuationListenerData(FAttenuationListenerData& OutListe
 	}
 }
 
-void FAudioDevice::GetAzimuth(FAttenuationListenerData& OutListenerData, const USoundBase* Sound, const FTransform& SoundTransform, const FSoundAttenuationSettings& AttenuationSettings, const FTransform& ListenerTransform, float& OutAzimuth, float& OutAbsoluteAzimuth) const
+void FAudioDevice::GetAzimuth(FAttenuationListenerData& OutListenerData, const FTransform& SoundTransform, const FSoundAttenuationSettings& AttenuationSettings, const FTransform& ListenerTransform, float& OutAzimuth, float& OutAbsoluteAzimuth) const
 {
 	GetAttenuationListenerData(OutListenerData, SoundTransform, AttenuationSettings, &ListenerTransform);
 
@@ -4878,7 +4876,7 @@ void FAudioDevice::GetAzimuth(FAttenuationListenerData& OutListenerData, const U
 
 	if (AbsAzimuthVector2D.X > 0.0f && AbsAzimuthVector2D.Y < 0.0f)
 	{
-		OutAbsoluteAzimuth = 360 - OutAbsoluteAzimuth;
+		OutAbsoluteAzimuth = 360.0f - OutAbsoluteAzimuth;
 	}
 	else if (AbsAzimuthVector2D.X < 0.0f && AbsAzimuthVector2D.Y < 0.0f)
 	{
@@ -4886,14 +4884,12 @@ void FAudioDevice::GetAzimuth(FAttenuationListenerData& OutListenerData, const U
 	}
 	else if (AbsAzimuthVector2D.X < 0.0f && AbsAzimuthVector2D.Y > 0.0f)
 	{
-		OutAbsoluteAzimuth = 180 - OutAbsoluteAzimuth;
+		OutAbsoluteAzimuth = 180.0f - OutAbsoluteAzimuth;
 	}
 }
 
-float FAudioDevice::GetFocusFactor(FAttenuationListenerData& OutListenerData, const USoundBase* Sound, const float Azimuth, const FSoundAttenuationSettings& AttenuationSettings) const
+float FAudioDevice::GetFocusFactor(const float Azimuth, const FSoundAttenuationSettings& AttenuationSettings) const
 {
-	check(Sound);
-
 	// 0.0f means we are in focus, 1.0f means we are out of focus
 	float FocusFactor = 0.0f;
 
diff --git a/Engine/Source/Runtime/Engine/Public/AudioDevice.h b/Engine/Source/Runtime/Engine/Public/AudioDevice.h
--- a/Engine/Source/Runtime/Engine/Public/AudioDevice.h
+++ b/Engine/Source/Runtime/Engine/Public/AudioDevice.h
@@ -992,10 +992,10 @@ public:
 	void GetAttenuationListenerData(FAttenuationListenerData& OutListenerData, const FTransform& SoundTransform, const FSoundAttenuationSettings& AttenuationSettings, const FTransform* InListenerTransform = nullptr) const;
 
 	/** Returns the azimuth angle of the sound relative to the sound's nearest listener. Used for 3d audio calculations. */
-	void GetAzimuth(FAttenuationListenerData& OutListenerData, const USoundBase* Sound, const FTransform& SoundTransform, const FSoundAttenuationSettings& AttenuationSettings, const FTransform& ListenerTransform, float& OutAzimuth, float& AbsoluteAzimuth) const;
+	void GetAzimuth(FAttenuationListenerData& OutListenerData, const FTransform& SoundTransform, const FSoundAttenuationSettings& AttenuationSettings, const FTransform& ListenerTransform, float& OutAzimuth, float& AbsoluteAzimuth) const;
 
 	/** Returns the focus factor of a sound based on its position and listener data. */
-	float GetFocusFactor(FAttenuationListenerData& OutListenerData, const USoundBase* Sound, const float Azimuth, const FSoundAttenuationSettings& AttenuationSettings) const;
+	float GetFocusFactor(const float Azimuth, const FSoundAttenuationSettings& AttenuationSettings) const;
 
 	/** Gets the max distance and focus factor of a sound. */
 	void GetMaxDistanceAndFocusFactor(USoundBase* Sound, const UWorld* World, const FVector& Location, const FSoundAttenuationSettings* AttenuationSettingsToApply, float& OutMaxDistance, float& OutFocusFactor);
