Ensure that unregistered instanced components are correctly in their BP created component parent's AttachChildren array prior to the instance data cache getting created.

diff --git a/Engine/Source/Runtime/Engine/Private/ActorConstruction.cpp b/Engine/Source/Runtime/Engine/Private/ActorConstruction.cpp
--- a/Engine/Source/Runtime/Engine/Private/ActorConstruction.cpp
+++ b/Engine/Source/Runtime/Engine/Private/ActorConstruction.cpp
@@ -272,6 +272,26 @@ void AActor::RerunConstructionScripts()
 		// Save info about attached actors
 		TArray<FAttachedActorInfo> AttachedActorInfos;
 
+		// Before we build the component instance data cache we need to make sure that instance components 
+		// are correctly in their AttachParent's AttachChildren array which may not be the case if they
+		// have not yet been registered
+		for (UActorComponent* Component : InstanceComponents)
+		{
+			if (Component && !Component->IsRegistered())
+			{
+				if (USceneComponent* SceneComp = Cast<USceneComponent>(Component))
+				{
+					if (USceneComponent* AttachParent = SceneComp->GetAttachParent())
+					{
+						if (AttachParent->IsCreatedByConstructionScript())
+						{
+							SceneComp->AttachToComponent(AttachParent, FAttachmentTransformRules::KeepRelativeTransform, SceneComp->GetAttachSocketName());
+						}
+					}
+				}
+			}
+		}
+
 #if WITH_EDITOR
 		if (!CurrentTransactionAnnotation.IsValid())
 		{
