Fixed race condition with BSP models and cached mesh draw commands.

 - In the level BSP model invalidation / undo methods, this change moves the component unregistration to happen prior to releasing the meshes. This is likely not the cause of the crash, but it seems safer to unregister components before releasing their underlying resources.

 - When models are updating, a legacy optimization was used to only update material index buffers. This would *not* first remove model components from the scene, which leaves static meshes intact. Then, it would queue a re-creation of static meshes until the end of the
   frame. This creates a race condition where if static meshes are rebuilt prior to the end of the frame (e.g. UpdateStaticMeshesForMaterials), the FMeshBatch is now pointing at bogus vertex and index buffers. The fix is to perform an immediate recreation instead of deferring one.

diff --git a/Engine/Source/Runtime/Engine/Private/Level.cpp b/Engine/Source/Runtime/Engine/Private/Level.cpp
--- a/Engine/Source/Runtime/Engine/Private/Level.cpp
+++ b/Engine/Source/Runtime/Engine/Private/Level.cpp
@@ -1364,10 +1364,6 @@ void ULevel::PreEditUndo()
 {
 	Super::PreEditUndo();
 
-	// Release the model's resources.
-	Model->BeginReleaseResources();
-	Model->ReleaseResourcesFence.Wait();
-
 	// Detach existing model components.  These are left in the array, so they are saved for undoing the undo.
 	for(int32 ComponentIndex = 0;ComponentIndex < ModelComponents.Num();ComponentIndex++)
 	{
@@ -1377,6 +1373,10 @@ void ULevel::PreEditUndo()
 		}
 	}
 
+	// Release the model's resources.
+	Model->BeginReleaseResources();
+	Model->ReleaseResourcesFence.Wait();
+
 	ReleaseRenderingResources();
 
 	// Wait for the components to be detached.
@@ -1483,9 +1483,6 @@ void ULevel::InvalidateModelGeometry()
 	Model->Modify();
 	Modify();
 
-	// Begin releasing the model's resources.
-	Model->BeginReleaseResources();
-
 	// Remove existing model components.
 	for(int32 ComponentIndex = 0;ComponentIndex < ModelComponents.Num();ComponentIndex++)
 	{
@@ -1496,6 +1493,9 @@ void ULevel::InvalidateModelGeometry()
 		}
 	}
 	ModelComponents.Empty();
+
+	// Begin releasing the model's resources.
+	Model->BeginReleaseResources();
 }
 
 
@@ -1562,7 +1562,10 @@ void ULevel::CommitModelSurfaces()
 				{
 					if (Model->bOnlyRebuildMaterialIndexBuffers)
 					{
-						ModelComponents[ComponentIndex]->MarkRenderStateDirty();
+						// This is intentionally updated immediately. We just re-created vertex and index buffers
+						// without invalidating static meshes. Re-create all static meshes now so that mesh draw
+						// commands are refreshed.
+						ModelComponents[ComponentIndex]->RecreateRenderState_Concurrent();
 					}
 					else
 					{
