#ue4 - Adjusted log levels for LogNetPlayerMovement to make Log and Verbose less spammy and moved some of those down a level.

diff --git a/Engine/Source/Runtime/Engine/Private/Components/CharacterMovementComponent.cpp b/Engine/Source/Runtime/Engine/Private/Components/CharacterMovementComponent.cpp
--- a/Engine/Source/Runtime/Engine/Private/Components/CharacterMovementComponent.cpp
+++ b/Engine/Source/Runtime/Engine/Private/Components/CharacterMovementComponent.cpp
@@ -7672,7 +7672,7 @@ bool UCharacterMovementComponent::ClientUpdatePositionAfterServerUpdate()
 
 	if (ClientData->SavedMoves.Num() == 0)
 	{
-		UE_LOG(LogNetPlayerMovement, VeryVerbose, TEXT("ClientUpdatePositionAfterServerUpdate No saved moves to replay"), ClientData->SavedMoves.Num());
+		UE_LOG(LogNetPlayerMovement, Verbose, TEXT("ClientUpdatePositionAfterServerUpdate No saved moves to replay"), ClientData->SavedMoves.Num());
 
 		// With no saved moves to resimulate, the move the server updated us with is the last move we've done, no resimulation needed.
 		CharacterOwner->bClientResimulateRootMotion = false;
@@ -7699,7 +7699,7 @@ bool UCharacterMovementComponent::ClientUpdatePositionAfterServerUpdate()
 	bForceNextFloorCheck = true;
 
 	// Replay moves that have not yet been acked.
-	UE_LOG(LogNetPlayerMovement, VeryVerbose, TEXT("ClientUpdatePositionAfterServerUpdate Replaying %d Moves, starting at Timestamp %f"), ClientData->SavedMoves.Num(), ClientData->SavedMoves[0]->TimeStamp);
+	UE_LOG(LogNetPlayerMovement, Verbose, TEXT("ClientUpdatePositionAfterServerUpdate Replaying %d Moves, starting at Timestamp %f"), ClientData->SavedMoves.Num(), ClientData->SavedMoves[0]->TimeStamp);
 	for (int32 i=0; i<ClientData->SavedMoves.Num(); i++)
 	{
 		FSavedMove_Character* const CurrentMove = ClientData->SavedMoves[i].Get();
@@ -8020,12 +8020,12 @@ void UCharacterMovementComponent::ReplicateMoveToServer(float DeltaTime, const F
 			}
 			else
 			{
-				UE_LOG(LogNetPlayerMovement, Log, TEXT("Not combining move [would collide at start location]"));
+				UE_LOG(LogNetPlayerMovement, Verbose, TEXT("Not combining move [would collide at start location]"));
 			}
 		}
 		else
 		{
-			UE_LOG(LogNetPlayerMovement, Log, TEXT("Not combining move [not allowed by CanCombineWith()]"));
+			UE_LOG(LogNetPlayerMovement, Verbose, TEXT("Not combining move [not allowed by CanCombineWith()]"));
 		}
 	}
 
@@ -8063,7 +8063,7 @@ void UCharacterMovementComponent::ReplicateMoveToServer(float DeltaTime, const F
 
 		ClientData->ClientUpdateTime = MyWorld->TimeSeconds;
 
-		UE_CLOG(CharacterOwner && UpdatedComponent, LogNetPlayerMovement, Verbose, TEXT("ClientMove Time %f Acceleration %s Velocity %s Position %s DeltaTime %f Mode %s MovementBase %s.%s (Dynamic:%d) DualMove? %d"),
+		UE_CLOG(CharacterOwner && UpdatedComponent, LogNetPlayerMovement, VeryVerbose, TEXT("ClientMove Time %f Acceleration %s Velocity %s Position %s DeltaTime %f Mode %s MovementBase %s.%s (Dynamic:%d) DualMove? %d"),
 			NewMove->TimeStamp, *NewMove->Acceleration.ToString(), *Velocity.ToString(), *UpdatedComponent->GetComponentLocation().ToString(), NewMove->DeltaTime, *GetMovementName(),
 			*GetNameSafe(NewMove->EndBase.Get()), *NewMove->EndBoneName.ToString(), MovementBaseUtility::IsDynamicBase(NewMove->EndBase.Get()) ? 1 : 0, ClientData->PendingMove.IsValid() ? 1 : 0);
 
@@ -8212,11 +8212,11 @@ void UCharacterMovementComponent::ServerMoveOld_Implementation
 
 	if( !VerifyClientTimeStamp(OldTimeStamp, *ServerData) )
 	{
-		UE_LOG(LogNetPlayerMovement, Verbose, TEXT("ServerMoveOld: TimeStamp expired. %f, CurrentTimeStamp: %f"), OldTimeStamp, ServerData->CurrentClientTimeStamp);
+		UE_LOG(LogNetPlayerMovement, VeryVerbose, TEXT("ServerMoveOld: TimeStamp expired. %f, CurrentTimeStamp: %f"), OldTimeStamp, ServerData->CurrentClientTimeStamp);
 		return;
 	}
 
-	UE_LOG(LogNetPlayerMovement, Log, TEXT("Recovered move from OldTimeStamp %f, DeltaTime: %f"), OldTimeStamp, OldTimeStamp - ServerData->CurrentClientTimeStamp);
+	UE_LOG(LogNetPlayerMovement, Verbose, TEXT("Recovered move from OldTimeStamp %f, DeltaTime: %f"), OldTimeStamp, OldTimeStamp - ServerData->CurrentClientTimeStamp);
 
 	const UWorld* MyWorld = GetWorld();
 	const float DeltaTime = ServerData->GetServerMoveDeltaTime(OldTimeStamp, CharacterOwner->GetActorTimeDilation(*MyWorld));
@@ -8660,7 +8660,7 @@ void UCharacterMovementComponent::ServerMove_Implementation(
 		MoveAutonomous(TimeStamp, DeltaTime, MoveFlags, Accel);
 	}
 
-	UE_CLOG(CharacterOwner && UpdatedComponent, LogNetPlayerMovement, Verbose, TEXT("ServerMove Time %f Acceleration %s Velocity %s Position %s DeltaTime %f Mode %s MovementBase %s.%s (Dynamic:%d)"),
+	UE_CLOG(CharacterOwner && UpdatedComponent, LogNetPlayerMovement, VeryVerbose, TEXT("ServerMove Time %f Acceleration %s Velocity %s Position %s DeltaTime %f Mode %s MovementBase %s.%s (Dynamic:%d)"),
 			TimeStamp, *Accel.ToString(), *Velocity.ToString(), *UpdatedComponent->GetComponentLocation().ToString(), DeltaTime, *GetMovementName(),
 			*GetNameSafe(GetMovementBase()), *CharacterOwner->GetBasedMovement().BoneName.ToString(), MovementBaseUtility::IsDynamicBase(GetMovementBase()) ? 1 : 0);
 
