Dead code removal

diff --git a/Engine/Source/Runtime/Engine/Classes/Components/SkyLightComponent.h b/Engine/Source/Runtime/Engine/Classes/Components/SkyLightComponent.h
--- a/Engine/Source/Runtime/Engine/Classes/Components/SkyLightComponent.h
+++ b/Engine/Source/Runtime/Engine/Classes/Components/SkyLightComponent.h
@@ -12,10 +12,6 @@
 #include "Math/SHMath.h"
 #include "SkyLightComponent.generated.h"
 
-#ifndef ENVIRONMENT_TEXTURE_ARRAY_WORKAROUND // RHI_RAYTRACING
-#define ENVIRONMENT_TEXTURE_ARRAY_WORKAROUND	1
-#endif
-
 class FSkyLightSceneProxy;
 class UTextureCube;
 
diff --git a/Engine/Source/Runtime/Engine/Public/SceneManagement.h b/Engine/Source/Runtime/Engine/Public/SceneManagement.h
--- a/Engine/Source/Runtime/Engine/Public/SceneManagement.h
+++ b/Engine/Source/Runtime/Engine/Public/SceneManagement.h
@@ -27,10 +27,6 @@
 #include "SceneUtils.h"
 #include "LightmapUniformShaderParameters.h"
 
-#ifndef ENVIRONMENT_TEXTURE_ARRAY_WORKAROUND // RHI_RAYTRACING
-#define ENVIRONMENT_TEXTURE_ARRAY_WORKAROUND	1
-#endif
-
 class FCanvas;
 class FLightMap;
 class FLightmapResourceCluster;
