Removed r.DeferUniformBufferUpdatesUntilVisible since it is no longer supported with MeshDrawCommands, which cache the uniform buffer pointer at AddToScene time.  Removed dead code referencing it.

diff --git a/Engine/Source/Runtime/Engine/Private/PrimitiveSceneProxy.cpp b/Engine/Source/Runtime/Engine/Private/PrimitiveSceneProxy.cpp
--- a/Engine/Source/Runtime/Engine/Private/PrimitiveSceneProxy.cpp
+++ b/Engine/Source/Runtime/Engine/Private/PrimitiveSceneProxy.cpp
@@ -232,41 +232,6 @@ FPrimitiveViewRelevance FPrimitiveSceneProxy::GetViewRelevance(const FSceneView*
 	return FPrimitiveViewRelevance();
 }
 
-static TAutoConsoleVariable<int32> CVarDeferUniformBufferUpdatesUntilVisible(
-	TEXT("r.DeferUniformBufferUpdatesUntilVisible"),
-	0,
-	TEXT("If > 0, then don't update the primitive uniform buffer until it is visible. Incompatible with the Mesh Draw Command pipeline."));
-
-void FPrimitiveSceneProxy::UpdateUniformBufferMaybeLazy()
-{
-	//@todo MeshCommandPipeline r.DeferUniformBufferUpdatesUntilVisible isn't currently supported.
-	/*
-	if (PrimitiveSceneInfo && CVarDeferUniformBufferUpdatesUntilVisible.GetValueOnAnyThread() > 0)
-	{
-		PrimitiveSceneInfo->SetNeedsUniformBufferUpdate(true);
-	}
-	else
-	{
-		UpdateUniformBuffer();
-	}
-	*/
-
-	UpdateUniformBuffer();
-}
-
-bool FPrimitiveSceneProxy::NeedsUniformBufferUpdate() const
-{
-	//@todo MeshCommandPipeline r.DeferUniformBufferUpdatesUntilVisible isn't currently supported.
-	/*
-	if (PrimitiveSceneInfo && CVarDeferUniformBufferUpdatesUntilVisible.GetValueOnAnyThread() > 0)
-	{
-		return PrimitiveSceneInfo->NeedsUniformBufferUpdate();
-	}
-	*/
-
-	return false;
-}
-
 void FPrimitiveSceneProxy::UpdateUniformBuffer()
 {
 	QUICK_SCOPE_CYCLE_COUNTER(STAT_FPrimitiveSceneProxy_UpdateUniformBuffer);
@@ -328,7 +293,7 @@ void FPrimitiveSceneProxy::SetTransform(const FMatrix& InLocalToWorld, const FBo
 	LocalBounds = InLocalBounds;
 	ActorPosition = InActorPosition;
 	
-	UpdateUniformBufferMaybeLazy();
+	UpdateUniformBuffer();
 	
 	// Notify the proxy's implementation of the change.
 	OnTransformChanged();
diff --git a/Engine/Source/Runtime/Engine/Public/PrimitiveSceneProxy.h b/Engine/Source/Runtime/Engine/Public/PrimitiveSceneProxy.h
--- a/Engine/Source/Runtime/Engine/Public/PrimitiveSceneProxy.h
+++ b/Engine/Source/Runtime/Engine/Public/PrimitiveSceneProxy.h
@@ -606,10 +606,6 @@ public:
 	 * Updates the primitive proxy's uniform buffer.
 	 */
 	ENGINE_API void UpdateUniformBuffer();
-	/**
-	 * Updates the primitive proxy's uniform buffer.
-	 */
-	ENGINE_API bool NeedsUniformBufferUpdate() const;
 
 #if !UE_BUILD_SHIPPING
 
@@ -1042,11 +1038,6 @@ private:
 
 	ENGINE_API bool WouldSetTransformBeRedundant(const FMatrix& InLocalToWorld, const FBoxSphereBounds& InBounds, const FBoxSphereBounds& InLocalBounds, FVector InActorPosition);
 
-	/**
-	 * Either updates the uniform buffer or defers it until it becomes visible depending on a cvar
-	 */
-	ENGINE_API void UpdateUniformBufferMaybeLazy();
-
 	/** Updates the hidden editor view visibility map on the render thread */
 	void SetHiddenEdViews_RenderThread( uint64 InHiddenEditorViews );
 
diff --git a/Engine/Source/Runtime/Landscape/Private/LandscapeRender.cpp b/Engine/Source/Runtime/Landscape/Private/LandscapeRender.cpp
--- a/Engine/Source/Runtime/Landscape/Private/LandscapeRender.cpp
+++ b/Engine/Source/Runtime/Landscape/Private/LandscapeRender.cpp
@@ -950,10 +950,6 @@ void FLandscapeComponentSceneProxy::CreateRenderThreadResources()
 		BatchElementParams->SubY = -1;
 		BatchElementParams->CurrentLOD = 0;
 		GrassBatchElement->UserData = BatchElementParams;
-		if (NeedsUniformBufferUpdate())
-		{
-			UpdateUniformBuffer();
-		}
 		GrassBatchElement->PrimitiveUniformBuffer = GetUniformBuffer();
 		GrassBatchElement->IndexBuffer = SharedBuffers->GrassIndexBuffer;
 		GrassBatchElement->NumPrimitives = FMath::Square(NumSubsections) * FMath::Square(SubsectionSizeVerts);
