DR - Minor spelling fix
PR #5538: Fix "can be use" typos in Material.h property descriptions (Contributed by [REDACTED_PERSON])
#rnx

diff --git a/Engine/Source/Runtime/Engine/Classes/Materials/Material.h b/Engine/Source/Runtime/Engine/Classes/Materials/Material.h
--- a/Engine/Source/Runtime/Engine/Classes/Materials/Material.h
+++ b/Engine/Source/Runtime/Engine/Classes/Materials/Material.h
@@ -603,35 +603,35 @@ public:
 	uint32 bUsedAsSpecialEngineMaterial:1;
 
 	/** 
-	 * Indicates that the material and its instances can be use with skeletal meshes.  
+	 * Indicates that the material and its instances can be used with skeletal meshes.  
 	 * This will result in the shaders required to support skeletal meshes being compiled which will increase shader compile time and memory usage.
 	 */
 	UPROPERTY(EditAnywhere, BlueprintReadOnly, Category=Usage)
 	uint32 bUsedWithSkeletalMesh:1;
 
 	/** 
-	 * Indicates that the material and its instances can be use with editor compositing  
+	 * Indicates that the material and its instances can be used with editor compositing  
 	 * This will result in the shaders required to support editor compositing being compiled which will increase shader compile time and memory usage.
 	 */
 	UPROPERTY(EditAnywhere, BlueprintReadOnly, Category=Usage)
 	uint32 bUsedWithEditorCompositing:1;
 
 	/** 
-	 * Indicates that the material and its instances can be use with particle sprites 
+	 * Indicates that the material and its instances can be used with particle sprites 
 	 * This will result in the shaders required to support particle sprites being compiled which will increase shader compile time and memory usage.
 	 */
 	UPROPERTY(EditAnywhere, BlueprintReadOnly, Category=Usage)
 	uint32 bUsedWithParticleSprites:1;
 
 	/** 
-	 * Indicates that the material and its instances can be use with beam trails
+	 * Indicates that the material and its instances can be used with beam trails
 	 * This will result in the shaders required to support beam trails being compiled which will increase shader compile time and memory usage.
 	 */
 	UPROPERTY(EditAnywhere, BlueprintReadOnly, Category=Usage)
 	uint32 bUsedWithBeamTrails:1;
 
 	/** 
-	 * Indicates that the material and its instances can be use with mesh particles
+	 * Indicates that the material and its instances can be used with mesh particles
 	 * This will result in the shaders required to support mesh particles being compiled which will increase shader compile time and memory usage.
 	 */
 	UPROPERTY(EditAnywhere, BlueprintReadOnly, Category=Usage)
@@ -639,7 +639,7 @@ public:
 
 
 	/**
-	* Indicates that the material and its instances can be use with Niagara sprites (meshes and ribbons, respectively)
+	* Indicates that the material and its instances can be used with Niagara sprites (meshes and ribbons, respectively)
 	* This will result in the shaders required to support Niagara sprites being compiled which will increase shader compile time and memory usage.
 	*/
 	UPROPERTY(EditAnywhere, BlueprintReadOnly, Category = Usage)
@@ -655,49 +655,49 @@ public:
 	uint32 bUsedWithGeometryCache : 1;
 
 	/** 
-	 * Indicates that the material and its instances can be use with static lighting
+	 * Indicates that the material and its instances can be used with static lighting
 	 * This will result in the shaders required to support static lighting being compiled which will increase shader compile time and memory usage.
 	 */
 	UPROPERTY(EditAnywhere, BlueprintReadOnly, Category=Usage)
 	uint32 bUsedWithStaticLighting:1;
 
 	/** 
-	 * Indicates that the material and its instances can be use with morph targets
+	 * Indicates that the material and its instances can be used with morph targets
 	 * This will result in the shaders required to support morph targets being compiled which will increase shader compile time and memory usage.
 	 */
 	UPROPERTY(EditAnywhere, BlueprintReadOnly, Category=Usage)
 	uint32 bUsedWithMorphTargets:1;
 
 	/** 
-	 * Indicates that the material and its instances can be use with spline meshes
+	 * Indicates that the material and its instances can be used with spline meshes
 	 * This will result in the shaders required to support spline meshes being compiled which will increase shader compile time and memory usage.
 	 */
 	UPROPERTY(EditAnywhere, BlueprintReadOnly, Category=Usage)
 	uint32 bUsedWithSplineMeshes:1;
 
 	/** 
-	 * Indicates that the material and its instances can be use with instanced static meshes
+	 * Indicates that the material and its instances can be used with instanced static meshes
 	 * This will result in the shaders required to support instanced static meshes being compiled which will increase shader compile time and memory usage.
 	 */
 	UPROPERTY(EditAnywhere, BlueprintReadOnly, Category=Usage)
 	uint32 bUsedWithInstancedStaticMeshes:1;
 
 	/** 
-	 * Indicates that the material and its instances can be use with distortion
+	 * Indicates that the material and its instances can be used with distortion
 	 * This will result in the shaders required to support distortion being compiled which will increase shader compile time and memory usage.
 	 */
 	UPROPERTY()
 	uint32 bUsesDistortion:1;
 
 	/** 
-	 * Indicates that the material and its instances can be use with clothing
+	 * Indicates that the material and its instances can be used with clothing
 	 * This will result in the shaders required to support clothing being compiled which will increase shader compile time and memory usage.
 	 */
 	UPROPERTY(EditAnywhere, BlueprintReadOnly, Category=Usage)
 	uint32 bUsedWithClothing:1;
 
 	/** 
-	 * Indicates that the material and its instances can be use with Slate UI and UMG
+	 * Indicates that the material and its instances can be used with Slate UI and UMG
 	 * This will result in the shaders required to support UI materials being compiled which will increase shader compile time and memory usage.
 	 */
 	UPROPERTY()
