[BUG] UE-65233: Need to code gen for Custom Output Nodes before calculating tex samplers used. Harden 2 phases of translator: Code Gen & Usage Query.

diff --git a/Engine/Source/Runtime/Engine/Private/Materials/HLSLMaterialTranslator.h b/Engine/Source/Runtime/Engine/Private/Materials/HLSLMaterialTranslator.h
--- a/Engine/Source/Runtime/Engine/Private/Materials/HLSLMaterialTranslator.h
+++ b/Engine/Source/Runtime/Engine/Private/Materials/HLSLMaterialTranslator.h
@@ -247,6 +247,8 @@
 	uint32 bUsesDistanceCullFade : 1;
 	/** true if the Roughness input evaluates to a constant 1.0 */
 	uint32 bIsFullyRough : 1;
+	/** true if allowed to generate code chunks. Translator operates in two phases; generate all code chunks & query meta data based on generated code chunks. */
+	uint32 bAllowCodeChunkGeneration : 1;
 	/** Tracks the number of texture coordinates used by this material. */
 	uint32 NumUserTexCoords;
 	/** Tracks the number of texture coordinates used by the vertex shader in this material. */
@@ -303,6 +305,7 @@
 	,	bUsesEmissiveColor(false)
 	,	bUsesDistanceCullFade(false)
 	,	bIsFullyRough(0)
+	,	bAllowCodeChunkGeneration(true)
 	,	NumUserTexCoords(0)
 	,	NumUserVertexTexCoords(0)
 	,	DynamicParticleParameterMask(0)
@@ -482,7 +485,7 @@
 			}
 			else
 			{
-				SeenCustomOutputExpressionsClasses.Add(CustomOutput->GetClass());		
+				SeenCustomOutputExpressionsClasses.Add(CustomOutput->GetClass());
 				int32 NumOutputs = CustomOutput->GetNumOutputs();
 
 				if (CustomOutput->NeedsCustomOutputDefines())
@@ -551,7 +554,7 @@
 
 			int32 NormalCodeChunkEnd = -1;
 			int32 Chunk[CompiledMP_MAX];
-
+			
 			memset(Chunk, INDEX_NONE, sizeof(Chunk));
 
 			// Translate all custom vertex interpolators before main attributes so type information is available
@@ -666,6 +669,83 @@
 
 			Chunk[MP_PixelDepthOffset] = Material->CompilePropertyAndSetMaterialProperty(MP_PixelDepthOffset, this);
 
+			ResourcesString = TEXT("");
+
+#if HANDLE_CUSTOM_OUTPUTS_AS_MATERIAL_ATTRIBUTES
+			// Handle custom outputs when using material attribute output
+			if (Material->HasMaterialAttributesConnected())
+			{
+				TArray<FMaterialCustomOutputAttributeDefintion> CustomAttributeList;
+				FMaterialAttributeDefinitionMap::GetCustomAttributeList(CustomAttributeList);
+				TArray<FShaderCodeChunk> CustomExpressionChunks;
+
+				for (FMaterialCustomOutputAttributeDefintion& Attribute : CustomAttributeList)
+				{
+					// Compile all outputs for attribute
+					bool bValidResultCompiled = false;
+					int32 NumOutputs = 1;//CustomOutput->GetNumOutputs();
+
+					for (int32 OutputIndex = 0; OutputIndex < NumOutputs; ++OutputIndex)
+					{
+						MaterialProperty = Attribute.Property;
+						ShaderFrequency = Attribute.ShaderFrequency;
+						FunctionStacks[ShaderFrequency].Empty();
+						FunctionStacks[ShaderFrequency].Add(FMaterialFunctionCompileState(nullptr));
+
+						CustomExpressionChunks.Empty();
+						CurrentScopeChunks = &CustomExpressionChunks;
+						int32 Result = Material->CompileCustomAttribute(Attribute.AttributeID, this);
+
+						// Consider attribute used if varies from default value
+						if (Result != INDEX_NONE)
+						{
+							bool bValueNonDefault = true;
+
+							if (FMaterialUniformExpression* Expression = GetParameterUniformExpression(Result))
+							{
+								FLinearColor Value;
+								FMaterialRenderContext DummyContext(nullptr, *Material, nullptr);
+								Expression->GetNumberValue(DummyContext, Value);
+
+								bool bEqualValue = Value.R == Attribute.DefaultValue.X;
+								bEqualValue &= Value.G == Attribute.DefaultValue.Y || Attribute.ValueType < MCT_Float2;
+								bEqualValue &= Value.B == Attribute.DefaultValue.Z || Attribute.ValueType < MCT_Float3;
+								bEqualValue &= Value.A == Attribute.DefaultValue.W || Attribute.ValueType < MCT_Float4;
+
+								if (Expression->IsConstant() && bEqualValue)
+								{
+									bValueNonDefault = false;
+								}
+							}
+
+							// Valid, non-default value so generate shader code
+							if (bValueNonDefault)
+							{
+								GenerateCustomAttributeCode(OutputIndex, Result, Attribute.ValueType, Attribute.FunctionName);
+								bValidResultCompiled = true;
+							}
+						}
+					}
+
+					// If used, add compile data
+					if (bValidResultCompiled)
+					{
+						ResourcesString += FString::Printf(TEXT("#define NUM_MATERIAL_OUTPUTS_%s %d\r\n"), *Attribute.FunctionName.ToUpper(), NumOutputs);
+					}
+				}
+			}
+			else
+#endif // #if HANDLE_CUSTOM_OUTPUTS_AS_MATERIAL_ATTRIBUTES
+			{
+				CompileCustomOutputs(CustomOutputExpressions, SeenCustomOutputExpressionsClasses, false);
+			}
+
+			// Output the implementation for any custom expressions we will call below.
+			for (int32 ExpressionIndex = 0; ExpressionIndex < CustomExpressionImplementations.Num(); ExpressionIndex++)
+			{
+				ResourcesString += CustomExpressionImplementations[ExpressionIndex] + "\r\n\r\n";
+			}
+
 			// No more calls to non-vertex shader CompilePropertyAndSetMaterialProperty beyond this point
 			const uint32 SavedNumUserTexCoords = NumUserTexCoords;
 
@@ -680,6 +760,10 @@
 				}
 			}
 
+			// Translation is designed to have a code chunk generation phase followed by several passes that only has readonly access to the code chunks.
+			// At this point we mark the code chunk generation complete.
+			bAllowCodeChunkGeneration = false;
+
 			bUsesEmissiveColor = IsMaterialPropertyUsed(MP_EmissiveColor, Chunk[MP_EmissiveColor], FLinearColor(0, 0, 0, 0), 3);
 			bUsesPixelDepthOffset = IsMaterialPropertyUsed(MP_PixelDepthOffset, Chunk[MP_PixelDepthOffset], FLinearColor(0, 0, 0, 0), 1)
 				|| (Domain == MD_DeferredDecal && Material->GetDecalBlendMode() == DBM_Volumetric_DistanceFunction);
@@ -804,83 +888,6 @@
 			MaterialCompilationOutput.NumUsedUVScalars = NumUserTexCoords * 2;
 			MaterialCompilationOutput.NumUsedCustomInterpolatorScalars = CurrentCustomVertexInterpolatorOffset;
 
-			ResourcesString = TEXT("");
-
-#if HANDLE_CUSTOM_OUTPUTS_AS_MATERIAL_ATTRIBUTES
-			// Handle custom outputs when using material attribute output
-			if (Material->HasMaterialAttributesConnected())
-			{
-				TArray<FMaterialCustomOutputAttributeDefintion> CustomAttributeList;
-				FMaterialAttributeDefinitionMap::GetCustomAttributeList(CustomAttributeList);
-				TArray<FShaderCodeChunk> CustomExpressionChunks;
-				
-				for (FMaterialCustomOutputAttributeDefintion& Attribute : CustomAttributeList)
-				{
-					// Compile all outputs for attribute
-					bool bValidResultCompiled = false;
-					int32 NumOutputs = 1;//CustomOutput->GetNumOutputs();
-
-					for (int32 OutputIndex = 0; OutputIndex < NumOutputs; ++OutputIndex)
-					{
-						MaterialProperty = Attribute.Property;
-						ShaderFrequency = Attribute.ShaderFrequency;
-						FunctionStacks[ShaderFrequency].Empty();
-						FunctionStacks[ShaderFrequency].Add(FMaterialFunctionCompileState(nullptr));
-
-						CustomExpressionChunks.Empty();
-						CurrentScopeChunks = &CustomExpressionChunks;
-						int32 Result = Material->CompileCustomAttribute(Attribute.AttributeID, this);	
-
-						// Consider attribute used if varies from default value
-						if (Result != INDEX_NONE)
-						{
-							bool bValueNonDefault = true;
-
-							if (FMaterialUniformExpression* Expression = GetParameterUniformExpression(Result))
-							{
-								FLinearColor Value;
-								FMaterialRenderContext DummyContext(nullptr, *Material, nullptr);
-								Expression->GetNumberValue(DummyContext, Value);
-
-								bool bEqualValue = Value.R == Attribute.DefaultValue.X;
-								bEqualValue &= Value.G == Attribute.DefaultValue.Y || Attribute.ValueType < MCT_Float2;
-								bEqualValue &= Value.B == Attribute.DefaultValue.Z || Attribute.ValueType < MCT_Float3;
-								bEqualValue &= Value.A == Attribute.DefaultValue.W || Attribute.ValueType < MCT_Float4;
-
-								if (Expression->IsConstant() && bEqualValue)
-								{
-									bValueNonDefault = false;
-								}
-							}
-
-							// Valid, non-default value so generate shader code
-							if (bValueNonDefault)
-							{
-								GenerateCustomAttributeCode(OutputIndex, Result, Attribute.ValueType, Attribute.FunctionName);
-								bValidResultCompiled = true;
-							}
-						}
-					}
-
-					// If used, add compile data
-					if (bValidResultCompiled)
-					{
-						ResourcesString += FString::Printf(TEXT("#define NUM_MATERIAL_OUTPUTS_%s %d\r\n"), *Attribute.FunctionName.ToUpper(), NumOutputs);
-					}
-				}
-			}
-			else
-#endif // #if HANDLE_CUSTOM_OUTPUTS_AS_MATERIAL_ATTRIBUTES
-			{
-				CompileCustomOutputs(CustomOutputExpressions, SeenCustomOutputExpressionsClasses, false);
-			}
-
-			// Output the implementation for any custom expressions we will call below.
-			for(int32 ExpressionIndex = 0;ExpressionIndex < CustomExpressionImplementations.Num();ExpressionIndex++)
-			{
-				ResourcesString += CustomExpressionImplementations[ExpressionIndex] + "\r\n\r\n";
-			}
-
 			// Do Normal Chunk first
 			{
 				GetFixedParameterCode(
@@ -1509,6 +1516,8 @@
 	/** Adds an already formatted inline or referenced code chunk */
 	int32 AddCodeChunkInner(const TCHAR* FormattedCode,EMaterialValueType Type,bool bInlined)
 	{
+		check(bAllowCodeChunkGeneration);
+
 		if (Type == MCT_Unknown)
 		{
 			return INDEX_NONE;
@@ -1606,6 +1615,8 @@
 	// AddUniformExpression - Adds an input to the Code array and returns its index.
 	int32 AddUniformExpression(FMaterialUniformExpression* UniformExpression,EMaterialValueType Type,const TCHAR* Format,...)
 	{
+		check(bAllowCodeChunkGeneration);
+
 		if (Type == MCT_Unknown)
 		{
 			return INDEX_NONE;
