UE-69305 - Expand skylight & reflection cubemap limits to device limits. Warn if cubemap may use large amount of memory based on request.

diff --git a/Engine/Source/Runtime/Engine/Classes/Components/SkyLightComponent.h b/Engine/Source/Runtime/Engine/Classes/Components/SkyLightComponent.h
--- a/Engine/Source/Runtime/Engine/Classes/Components/SkyLightComponent.h
+++ b/Engine/Source/Runtime/Engine/Classes/Components/SkyLightComponent.h
@@ -171,6 +171,7 @@ class ENGINE_API USkyLightComponent : public ULightComponentBase
 	virtual void PostLoad() override;
 	virtual void PostInterpChange(UProperty* PropertyThatChanged) override;
 #if WITH_EDITOR
+	virtual void PreEditChange(UProperty* PropertyAboutToChange) override;
 	virtual void PostEditChangeProperty(FPropertyChangedEvent& PropertyChangedEvent) override;
 	virtual bool CanEditChange(const UProperty* InProperty) const override;
 	virtual void CheckForErrors() override;
@@ -260,6 +261,11 @@ public:
 
 protected:
 
+#if WITH_EDITOR
+	/** shadow copy saved before effects of PostEditChange() to provide option to roll back edit. */
+	int32 PreEditCubemapResolution = 128;
+#endif
+
 	/** Indicates whether the cached data stored in GetComponentInstanceData is valid to be applied in ApplyComponentInstanceData. */
 	bool bSavedConstructionScriptValuesValid;
 
diff --git a/Engine/Source/Runtime/Engine/Classes/Engine/RendererSettings.h b/Engine/Source/Runtime/Engine/Classes/Engine/RendererSettings.h
--- a/Engine/Source/Runtime/Engine/Classes/Engine/RendererSettings.h
+++ b/Engine/Source/Runtime/Engine/Classes/Engine/RendererSettings.h
@@ -696,6 +696,7 @@ public:
 	virtual void PostInitProperties() override;
 
 #if WITH_EDITOR
+	virtual void PreEditChange(UProperty* PropertyAboutToChange) override;
 	virtual void PostEditChangeProperty(FPropertyChangedEvent& PropertyChangedEvent) override;
 	virtual bool CanEditChange(const UProperty* InProperty) const override;
 #endif
@@ -703,6 +704,11 @@ public:
 	//~ End UObject Interface
 
 private:
+#if WITH_EDITOR
+	/** shadow copy saved before effects of PostEditChange() to provide option to roll back edit. */
+	int32 PreEditReflectionCaptureResolution = 128;
+#endif // WITH_EDITOR
+
 	void SanatizeReflectionCaptureResolution();
 };
 
diff --git a/Engine/Source/Runtime/Engine/Private/Components/ReflectionCaptureComponent.cpp b/Engine/Source/Runtime/Engine/Private/Components/ReflectionCaptureComponent.cpp
--- a/Engine/Source/Runtime/Engine/Private/Components/ReflectionCaptureComponent.cpp
+++ b/Engine/Source/Runtime/Engine/Private/Components/ReflectionCaptureComponent.cpp
@@ -59,8 +59,8 @@ ENGINE_API TAutoConsoleVariable<int32> CVarReflectionCaptureSize(
 
 static int32 SanitizeReflectionCaptureSize(int32 ReflectionCaptureSize)
 {
-	static const int32 MaxReflectionCaptureSize = 1024;
-	static const int32 MinReflectionCaptureSize = 1;
+	const int32 MaxReflectionCaptureSize = GetMaxCubeTextureDimension();
+	const int32 MinReflectionCaptureSize = 1;
 
 	return FMath::Clamp(ReflectionCaptureSize, MinReflectionCaptureSize, MaxReflectionCaptureSize);
 }
diff --git a/Engine/Source/Runtime/Engine/Private/Components/SkyLightComponent.cpp b/Engine/Source/Runtime/Engine/Private/Components/SkyLightComponent.cpp
--- a/Engine/Source/Runtime/Engine/Private/Components/SkyLightComponent.cpp
+++ b/Engine/Source/Runtime/Engine/Private/Components/SkyLightComponent.cpp
@@ -21,6 +21,7 @@
 #include "Components/BillboardComponent.h"
 #include "UObject/ReleaseObjectVersion.h"
 #include "Modules/ModuleManager.h"
+#include "Internationalization/Text.h"
 
 #if RHI_RAYTRACING
 #include "GlobalShader.h"
@@ -58,6 +59,8 @@ FAutoConsoleVariableRef CVarUpdateSkylightsEveryFrame(
 	ECVF_Default
 	);
 
+constexpr EPixelFormat SKYLIGHT_CUBEMAP_FORMAT = PF_FloatRGBA;
+
 void FSkyTextureCubeResource::InitRHI()
 {
 	if (GetFeatureLevel() >= ERHIFeatureLevel::SM4 || GSupportsRenderTargetFormat_PF_FloatRGBA)
@@ -247,9 +250,43 @@ void USkyLightComponent::SetCaptureIsDirty()
 
 void USkyLightComponent::SanitizeCubemapSize()
 {
-	static const int32 MaxCubemapResolution = 1024;
-	static const int32 MinCubemapResolution = 64;
+	const int32 MaxCubemapResolution = GetMaxCubeTextureDimension();
+	const int32 MinCubemapResolution = 8;
+
 	CubemapResolution = FMath::Clamp(int32(FMath::RoundUpToPowerOfTwo(CubemapResolution)), MinCubemapResolution, MaxCubemapResolution);
+
+#if WITH_EDITOR
+	SIZE_T TexMemRequired = CalcTextureSize(CubemapResolution, CubemapResolution, SKYLIGHT_CUBEMAP_FORMAT, FMath::CeilLogTwo(CubemapResolution) + 1) * CubeFace_MAX;
+
+	FTextureMemoryStats TextureMemStats;
+	RHIGetTextureMemoryStats(TextureMemStats);
+
+	if (TexMemRequired > SIZE_T(TextureMemStats.DedicatedVideoMemory / 2))
+	{
+		FNumberFormattingOptions FmtOpts = FNumberFormattingOptions()
+			.SetUseGrouping(false)
+			.SetMaximumFractionalDigits(2)
+			.SetMinimumFractionalDigits(0)
+			.SetRoundingMode(HalfFromZero);
+
+		EAppReturnType::Type Response = FPlatformMisc::MessageBoxExt(
+			EAppMsgType::YesNo, 
+			*FText::Format(
+				LOCTEXT("MemAllocWarning_Message_SkylightCubemap", "A resolution of {0} will require {1} of video memory. Are you sure?"), 
+				FText::AsNumber(CubemapResolution, &FmtOpts), 
+				FText::AsMemory(TexMemRequired, &FmtOpts)
+				).ToString(),
+			*LOCTEXT("MemAllocWarning_Title_SkylightCubemap", "Memory Allocation Warning").ToString()
+			);
+
+		if (Response == EAppReturnType::No)
+		{
+			CubemapResolution = PreEditCubemapResolution;
+		}
+	}
+
+	PreEditCubemapResolution = CubemapResolution;
+#endif // WITH_EDITOR
 }
 
 void USkyLightComponent::SetBlendDestinationCaptureIsDirty()
@@ -413,11 +450,23 @@ void USkyLightComponent::DestroyRenderState_Concurrent()
 }
 
 #if WITH_EDITOR
+void USkyLightComponent::PreEditChange(UProperty* PropertyAboutToChange)
+{
+	Super::PreEditChange(PropertyAboutToChange);
+
+	PreEditCubemapResolution = CubemapResolution;
+}
+
 void USkyLightComponent::PostEditChangeProperty(FPropertyChangedEvent& PropertyChangedEvent)
 {
 	Super::PostEditChangeProperty(PropertyChangedEvent);
 
-	SanitizeCubemapSize();
+	if (PropertyChangedEvent.Property->GetFName() == GET_MEMBER_NAME_CHECKED(USkyLightComponent, CubemapResolution) &&
+		PropertyChangedEvent.ChangeType != EPropertyChangeType::Interactive)
+	{
+		SanitizeCubemapSize();
+	}
+
 	SetCaptureIsDirty();
 }
 
@@ -576,7 +625,7 @@ void USkyLightComponent::UpdateSkyCaptureContentsArray(UWorld* WorldToUpdate, TA
 					if (!CaptureComponent->ProcessedSkyTexture || CaptureComponent->ProcessedSkyTexture->GetSizeX() != CaptureComponent->CubemapResolution)
 					{
 						CaptureComponent->ProcessedSkyTexture = new FSkyTextureCubeResource();
-						CaptureComponent->ProcessedSkyTexture->SetupParameters(CaptureComponent->CubemapResolution, FMath::CeilLogTwo(CaptureComponent->CubemapResolution) + 1, PF_FloatRGBA);
+						CaptureComponent->ProcessedSkyTexture->SetupParameters(CaptureComponent->CubemapResolution, FMath::CeilLogTwo(CaptureComponent->CubemapResolution) + 1, SKYLIGHT_CUBEMAP_FORMAT);
 						BeginInitResource(CaptureComponent->ProcessedSkyTexture);
 						CaptureComponent->MarkRenderStateDirty();
 					}
@@ -591,7 +640,7 @@ void USkyLightComponent::UpdateSkyCaptureContentsArray(UWorld* WorldToUpdate, TA
 					if (!CaptureComponent->BlendDestinationProcessedSkyTexture || CaptureComponent->BlendDestinationProcessedSkyTexture->GetSizeX() != CaptureComponent->CubemapResolution)
 					{
 						CaptureComponent->BlendDestinationProcessedSkyTexture = new FSkyTextureCubeResource();
-						CaptureComponent->BlendDestinationProcessedSkyTexture->SetupParameters(CaptureComponent->CubemapResolution, FMath::CeilLogTwo(CaptureComponent->CubemapResolution) + 1, PF_FloatRGBA);
+						CaptureComponent->BlendDestinationProcessedSkyTexture->SetupParameters(CaptureComponent->CubemapResolution, FMath::CeilLogTwo(CaptureComponent->CubemapResolution) + 1, SKYLIGHT_CUBEMAP_FORMAT);
 						BeginInitResource(CaptureComponent->BlendDestinationProcessedSkyTexture);
 						CaptureComponent->MarkRenderStateDirty(); 
 					}
diff --git a/Engine/Source/Runtime/Engine/Private/RendererSettings.cpp b/Engine/Source/Runtime/Engine/Private/RendererSettings.cpp
--- a/Engine/Source/Runtime/Engine/Private/RendererSettings.cpp
+++ b/Engine/Source/Runtime/Engine/Private/RendererSettings.cpp
@@ -10,6 +10,8 @@
 extern UNREALED_API class UEditorEngine* GEditor;
 #endif // #if WITH_EDITOR
 
+#define LOCTEXT_NAMESPACE "RenderSettings"
+
 namespace EAlphaChannelMode
 {
 	EAlphaChannelMode::Type FromInt(int32 InAlphaChannelMode)
@@ -84,6 +86,13 @@ void URendererSettings::PostInitProperties()
 }
 
 #if WITH_EDITOR
+void URendererSettings::PreEditChange(UProperty* PropertyAboutToChange)
+{
+	Super::PreEditChange(PropertyAboutToChange);
+
+	PreEditReflectionCaptureResolution = ReflectionCaptureResolution;
+}
+
 void URendererSettings::PostEditChangeProperty(FPropertyChangedEvent& PropertyChangedEvent)
 {
 	Super::PostEditChangeProperty(PropertyChangedEvent);
@@ -134,9 +143,43 @@ bool URendererSettings::CanEditChange(const UProperty* InProperty) const
 
 void URendererSettings::SanatizeReflectionCaptureResolution()
 {
-	static const int32 MaxReflectionCaptureResolution = 1024;
-	static const int32 MinReflectionCaptureResolution = 64;
-	ReflectionCaptureResolution = FMath::Clamp(int32(FMath::RoundUpToPowerOfTwo(ReflectionCaptureResolution)), MinReflectionCaptureResolution, MaxReflectionCaptureResolution);
+	const int32 MaxCubemapResolution = GetMaxCubeTextureDimension();
+	const int32 MinCubemapResolution = 8;
+
+	ReflectionCaptureResolution = FMath::Clamp(int32(FMath::RoundUpToPowerOfTwo(ReflectionCaptureResolution)), MinCubemapResolution, MaxCubemapResolution);
+
+#if WITH_EDITOR
+	SIZE_T TexMemRequired = CalcTextureSize(ReflectionCaptureResolution, ReflectionCaptureResolution, PF_FloatRGBA, FMath::CeilLogTwo(ReflectionCaptureResolution) + 1) * CubeFace_MAX;
+
+	FTextureMemoryStats TextureMemStats;
+	RHIGetTextureMemoryStats(TextureMemStats);
+
+	if (TexMemRequired > SIZE_T(TextureMemStats.DedicatedVideoMemory / 4))
+	{
+		FNumberFormattingOptions FmtOpts = FNumberFormattingOptions()
+			.SetUseGrouping(false)
+			.SetMaximumFractionalDigits(2)
+			.SetMinimumFractionalDigits(0)
+			.SetRoundingMode(HalfFromZero);
+
+		EAppReturnType::Type Response = FPlatformMisc::MessageBoxExt(
+			EAppMsgType::YesNo,
+			*FText::Format(
+				LOCTEXT("MemAllocWarning_Message_ReflectionCubemap", "A resolution of {0} will require {1} of video memory PER reflection capture component. Are you sure?"),
+				FText::AsNumber(ReflectionCaptureResolution, &FmtOpts),
+				FText::AsMemory(TexMemRequired, &FmtOpts)
+				).ToString(),
+			*LOCTEXT("MemAllocWarning_Title_ReflectionCubemap", "Memory Allocation Warning").ToString()
+			);
+
+		if (Response == EAppReturnType::No)
+		{
+			ReflectionCaptureResolution = PreEditReflectionCaptureResolution;
+		}
+	}
+
+	PreEditReflectionCaptureResolution = ReflectionCaptureResolution;
+#endif // WITH_EDITOR
 }
 
 URendererOverrideSettings::URendererOverrideSettings(const FObjectInitializer& ObjectInitializer)
@@ -168,3 +211,5 @@ void URendererOverrideSettings::PostEditChangeProperty(FPropertyChangedEvent& Pr
 	}
 }
 #endif // #if WITH_EDITOR
+
+#undef LOCTEXT_NAMESPACE
