Fixed issue where a "Draw Material to RenderTarget" node invoked in the first frame of an application would result in the render target being immediately cleared afterwards.

In the case where a blueprint is invoked on the first frame to render a material into a render target, the render target can be simultaneously queued for a deferred update. After rendering the material to the render target, the actual scene is rendered, which flushes this deferred update queue. The render target is then immediately cleared. This change flushed the queue prior to rendering a canvas to a render target, because we need all deferred clears to occur beforehand.

diff --git a/Engine/Source/Runtime/Engine/Private/KismetRenderingLibrary.cpp b/Engine/Source/Runtime/Engine/Private/KismetRenderingLibrary.cpp
--- a/Engine/Source/Runtime/Engine/Private/KismetRenderingLibrary.cpp
+++ b/Engine/Source/Runtime/Engine/Private/KismetRenderingLibrary.cpp
@@ -130,8 +130,14 @@ void UKismetRenderingLibrary::DrawMaterialToRenderTarget(UObject* WorldContextOb
 
 		FName RTName = TextureRenderTarget->GetFName();
 		ENQUEUE_RENDER_COMMAND(BeginDrawEventCommand)(
-			[RTName, DrawMaterialToTargetEvent](FRHICommandList& RHICmdList)
+			[RTName, DrawMaterialToTargetEvent](FRHICommandListImmediate& RHICmdList)
 			{
+				// Update resources that were marked for deferred update. This is important
+				// in cases where the blueprint is invoked in the same frame that the render
+				// target is created. Among other things, this will perform deferred render
+				// target clears.
+				FDeferredUpdateResource::UpdateResources(RHICmdList);
+
 				BEGIN_DRAW_EVENTF(
 					RHICmdList, 
 					DrawCanvasToTarget, 
