Added a GetCollisionMeshElement variant to FStaticMeshSceneProxy which fixes the collision view mode when tessellation is enabled. Also refactored the four mesh element implementations (Wireframe, Shadows, Base, Collision) to share some code. Uncrustified some of the logic to make it clearer and more maintainable in the future.

#4004

diff --git a/Engine/Source/Runtime/Engine/Private/StaticMeshRender.cpp b/Engine/Source/Runtime/Engine/Private/StaticMeshRender.cpp
--- a/Engine/Source/Runtime/Engine/Private/StaticMeshRender.cpp
+++ b/Engine/Source/Runtime/Engine/Private/StaticMeshRender.cpp
@@ -362,46 +362,50 @@ bool FStaticMeshSceneProxy::GetShadowMeshElement(int32 LODIndex, int32 BatchInde
 	const FLODInfo& ProxyLODInfo = LODs[LODIndex];
 
 	const bool bUseReversedIndices = GUseReversedIndexBuffer && IsLocalToWorldDeterminantNegative() && LOD.bHasReversedDepthOnlyIndices;
+	const bool bNoIndexBufferAvailable = !bUseReversedIndices && !LOD.bHasDepthOnlyIndices;
+
+	if (bNoIndexBufferAvailable)
+	{
+		return false;
+	}
+
+	FMeshBatchElement& OutMeshBatchElement = OutMeshBatch.Elements[0];
+
+	if (ProxyLODInfo.OverrideColorVertexBuffer)
+	{
+		OutMeshBatch.VertexFactory = &VFs.VertexFactoryOverrideColorVertexBuffer;
+		OutMeshBatchElement.VertexFactoryUserData = ProxyLODInfo.OverrideColorVFUniformBuffer.GetReference();
+	}
+	else
+	{
+		OutMeshBatch.VertexFactory = &VFs.VertexFactory;
+		OutMeshBatchElement.VertexFactoryUserData = VFs.VertexFactory.GetUniformBuffer();
+	}
+
+	OutMeshBatchElement.IndexBuffer = bUseReversedIndices ? &LOD.ReversedDepthOnlyIndexBuffer : &LOD.DepthOnlyIndexBuffer;
+	OutMeshBatchElement.FirstIndex = 0;
+	OutMeshBatchElement.NumPrimitives = LOD.DepthOnlyNumTriangles;
+	OutMeshBatchElement.MinVertexIndex = 0;
+	OutMeshBatchElement.MaxVertexIndex = LOD.VertexBuffers.PositionVertexBuffer.GetNumVertices() - 1;
 
-	FMeshBatchElement& OutBatchElement = OutMeshBatch.Elements[0];
-	OutBatchElement.VertexFactoryUserData = ProxyLODInfo.OverrideColorVertexBuffer ? ProxyLODInfo.OverrideColorVFUniformBuffer.GetReference() : VFs.VertexFactory.GetUniformBuffer();
-	OutMeshBatch.MaterialRenderProxy = UMaterial::GetDefaultMaterial(MD_Surface)->GetRenderProxy();
-	OutMeshBatch.VertexFactory = ProxyLODInfo.OverrideColorVertexBuffer ? &VFs.VertexFactoryOverrideColorVertexBuffer : &VFs.VertexFactory;
-	OutBatchElement.IndexBuffer = bUseReversedIndices ? &LOD.ReversedDepthOnlyIndexBuffer : &LOD.DepthOnlyIndexBuffer;
-	OutMeshBatch.Type = PT_TriangleList;
-	OutBatchElement.FirstIndex = 0;
-	OutBatchElement.NumPrimitives = LOD.DepthOnlyNumTriangles;
-	OutBatchElement.MinVertexIndex = 0;
-	OutBatchElement.MaxVertexIndex = LOD.VertexBuffers.PositionVertexBuffer.GetNumVertices() - 1;
-	OutMeshBatch.DepthPriorityGroup = InDepthPriorityGroup;
-	OutMeshBatch.ReverseCulling = (bReverseCulling || IsLocalToWorldDeterminantNegative()) && !bUseReversedIndices;
 	OutMeshBatch.LODIndex = LODIndex;
 #if STATICMESH_ENABLE_DEBUG_RENDERING
 	OutMeshBatch.VisualizeLODIndex = LODIndex;
 	OutMeshBatch.VisualizeHLODIndex = HierarchicalLODIndex;
 #endif
+	OutMeshBatch.ReverseCulling = IsReversedCullingNeeded(bUseReversedIndices);
+	OutMeshBatch.Type = PT_TriangleList;
+	OutMeshBatch.DepthPriorityGroup = InDepthPriorityGroup;
 	OutMeshBatch.LCI = &ProxyLODInfo;
-	if (ForcedLodModel > 0) 
-	{
-		OutBatchElement.MaxScreenSize = 0.0f;
-		OutBatchElement.MinScreenSize = -1.0f;
-	}
-	else
-	{
-		OutMeshBatch.bDitheredLODTransition = bDitheredLODTransition;
-		OutBatchElement.MaxScreenSize = GetScreenSize(LODIndex);
-		OutBatchElement.MinScreenSize = 0.0f;
-		if (LODIndex < MAX_STATIC_MESH_LODS - 1)
-		{
-			OutBatchElement.MinScreenSize = GetScreenSize(LODIndex + 1);
-		}
-	}
+	OutMeshBatch.MaterialRenderProxy = UMaterial::GetDefaultMaterial(MD_Surface)->GetRenderProxy();
 
 	// By default this will be a shadow only mesh.
 	OutMeshBatch.bUseForMaterial = false;
 	OutMeshBatch.bUseForDepthPass = false;
 	OutMeshBatch.bUseAsOccluder = false;
 
+	SetMeshElementScreenSize(LODIndex, bDitheredLODTransition, OutMeshBatch);
+
 	return true;
 }
 
@@ -415,18 +419,19 @@ bool FStaticMeshSceneProxy::GetMeshElement(
 	bool bAllowPreCulledIndices, 
 	FMeshBatch& OutMeshBatch) const
 {
-	OutMeshBatch.SegmentIndex = SectionIndex;
-
 	const ERHIFeatureLevel::Type FeatureLevel = GetScene().GetFeatureLevel();
 	const FStaticMeshLODResources& LOD = RenderData->LODResources[LODIndex];
 	const FStaticMeshVertexFactories& VFs = RenderData->LODVertexFactories[LODIndex];
 	const FStaticMeshSection& Section = LOD.Sections[SectionIndex];
-	
 	const FLODInfo& ProxyLODInfo = LODs[LODIndex];
+
 	UMaterialInterface* MaterialInterface = ProxyLODInfo.Sections[SectionIndex].Material;
-	
-	OutMeshBatch.MaterialRenderProxy = MaterialInterface->GetRenderProxy();
-	OutMeshBatch.VertexFactory = &VFs.VertexFactory;
+	FMaterialRenderProxy* MaterialRenderProxy = MaterialInterface->GetRenderProxy();
+	const FMaterial* Material = MaterialRenderProxy->GetMaterial(FeatureLevel);
+
+	const FVertexFactory* VertexFactory = nullptr;
+
+	FMeshBatchElement& OutMeshBatchElement = OutMeshBatch.Elements[0];
 
 #if WITH_EDITORONLY_DATA
 	// If material is hidden, then skip the draw.
@@ -443,63 +448,60 @@ bool FStaticMeshSceneProxy::GetMeshElement(
 	OutMeshBatch.bUseSelectionOutline = bPerSectionSelection ? bUseSelectionOutline : true;
 #endif
 
+	// Has the mesh component overridden the vertex color stream for this mesh LOD?
+	if (ProxyLODInfo.OverrideColorVertexBuffer)
+	{
+		// Make sure the indices are accessing data within the vertex buffer's
+		check(Section.MaxVertexIndex < ProxyLODInfo.OverrideColorVertexBuffer->GetNumVertices())
+
+		// Use the instanced colors vertex factory.
+		VertexFactory = &VFs.VertexFactoryOverrideColorVertexBuffer;
+
+		OutMeshBatchElement.VertexFactoryUserData = ProxyLODInfo.OverrideColorVFUniformBuffer.GetReference();
+		OutMeshBatchElement.UserData = ProxyLODInfo.OverrideColorVertexBuffer;
+		OutMeshBatchElement.bUserDataIsColorVertexBuffer = true;
+	}
+	else
+	{
+		VertexFactory = &VFs.VertexFactory;
+
+		OutMeshBatchElement.VertexFactoryUserData = VFs.VertexFactory.GetUniformBuffer();
+	}
+
 	const bool bWireframe = false;
-	const bool bRequiresAdjacencyInformation = RequiresAdjacencyInformation( MaterialInterface, OutMeshBatch.VertexFactory->GetType(), FeatureLevel );
-	const FMaterial* Material = OutMeshBatch.MaterialRenderProxy->GetMaterial(FeatureLevel);
-	
-	// Two sided material use bIsFrontFace which is wrong with Reversed Indices. AdjacencyInformation use another index buffer.
-	CA_SUPPRESS(6239);
-	const bool bUseReversedIndices = !bWireframe && GUseReversedIndexBuffer && IsLocalToWorldDeterminantNegative() && LOD.bHasReversedIndices && !bRequiresAdjacencyInformation && !Material->IsTwoSided();
+	const bool bRequiresAdjacencyInformation = RequiresAdjacencyInformation(MaterialInterface, VertexFactory->GetType(), FeatureLevel);
 
-	SetIndexSource(LODIndex, SectionIndex, OutMeshBatch, bWireframe, bRequiresAdjacencyInformation, bUseReversedIndices, bAllowPreCulledIndices);
+	// Two sided material use bIsFrontFace which is wrong with Reversed Indices. AdjacencyInformation use another index buffer.
+	const bool bUseReversedIndices = !bWireframe && GUseReversedIndexBuffer && IsLocalToWorldDeterminantNegative() && (LOD.bHasReversedIndices != 0) && !bRequiresAdjacencyInformation && !Material->IsTwoSided();
 
-	FMeshBatchElement& OutBatchElement = OutMeshBatch.Elements[0];
+	// No support for stateless dithered LOD transitions for movable meshes
+	const bool bDitheredLODTransition = !IsMovable() && Material->IsDitheredLODTransition();
 
-	OutBatchElement.VertexFactoryUserData = ProxyLODInfo.OverrideColorVertexBuffer ? ProxyLODInfo.OverrideColorVFUniformBuffer.GetReference() : VFs.VertexFactory.GetUniformBuffer();
+	const uint32 NumPrimitives = SetMeshElementGeometrySource(LODIndex, SectionIndex, bWireframe, bRequiresAdjacencyInformation, bUseReversedIndices, bAllowPreCulledIndices, VertexFactory, OutMeshBatch);
 
-	// Has the mesh component overridden the vertex color stream for this mesh LOD?
-	if (ProxyLODInfo.OverrideColorVertexBuffer != NULL)
+	if(NumPrimitives > 0)
 	{
-		// Make sure the indices are accessing data within the vertex buffer's
-		check(Section.MaxVertexIndex < ProxyLODInfo.OverrideColorVertexBuffer->GetNumVertices());
-		// Switch out the stock mesh vertex factory with the instanced colors one
-		OutMeshBatch.VertexFactory = &VFs.VertexFactoryOverrideColorVertexBuffer;
-		OutBatchElement.UserData = ProxyLODInfo.OverrideColorVertexBuffer;
-		OutBatchElement.bUserDataIsColorVertexBuffer = true;
-	}
+		OutMeshBatch.SegmentIndex = SectionIndex;
 
-	if(OutBatchElement.NumPrimitives > 0)
-	{
-		OutMeshBatch.LCI = &ProxyLODInfo;
-		OutBatchElement.MinVertexIndex = Section.MinVertexIndex;
-		OutBatchElement.MaxVertexIndex = Section.MaxVertexIndex;
 		OutMeshBatch.LODIndex = LODIndex;
 #if STATICMESH_ENABLE_DEBUG_RENDERING
-		OutBatchElement.VisualizeElementIndex = SectionIndex;
 		OutMeshBatch.VisualizeLODIndex = LODIndex;
 		OutMeshBatch.VisualizeHLODIndex = HierarchicalLODIndex;
 #endif
-		OutMeshBatch.ReverseCulling = (bReverseCulling || IsLocalToWorldDeterminantNegative()) && !bUseReversedIndices;
+		OutMeshBatch.ReverseCulling = IsReversedCullingNeeded(bUseReversedIndices);
 		OutMeshBatch.CastShadow = bCastShadow && Section.bCastShadow;
 		OutMeshBatch.DepthPriorityGroup = (ESceneDepthPriorityGroup)InDepthPriorityGroup;
-		if (ForcedLodModel > 0) 
-		{
-			OutBatchElement.MaxScreenSize = 0.0f;
-			OutBatchElement.MinScreenSize = -1.0f;
-		}
-		else
-		{
-			// no support for stateless dithered LOD transitions for movable meshes
-			OutMeshBatch.bDitheredLODTransition = !IsMovable() && Material->IsDitheredLODTransition();
+		OutMeshBatch.LCI = &ProxyLODInfo;
+		OutMeshBatch.MaterialRenderProxy = MaterialRenderProxy;
+
+		OutMeshBatchElement.MinVertexIndex = Section.MinVertexIndex;
+		OutMeshBatchElement.MaxVertexIndex = Section.MaxVertexIndex;
+#if STATICMESH_ENABLE_DEBUG_RENDERING
+		OutMeshBatchElement.VisualizeElementIndex = SectionIndex;
+#endif
+
+		SetMeshElementScreenSize(LODIndex, bDitheredLODTransition, OutMeshBatch);
 
-			OutBatchElement.MaxScreenSize = GetScreenSize(LODIndex);
-			OutBatchElement.MinScreenSize = 0.0f;
-			if (LODIndex < MAX_STATIC_MESH_LODS - 1)
-			{
-				OutBatchElement.MinScreenSize = GetScreenSize(LODIndex + 1);
-			}
-		}
-			
 		return true;
 	}
 	else
@@ -525,39 +527,108 @@ bool FStaticMeshSceneProxy::GetWireframeMeshElement(int32 LODIndex, int32 BatchI
 	const FStaticMeshLODResources& LODModel = RenderData->LODResources[LODIndex];
 	const FStaticMeshVertexFactories& VFs = RenderData->LODVertexFactories[LODIndex];
 	const FLODInfo& ProxyLODInfo = LODs[LODIndex];
+	const FVertexFactory* VertexFactory = nullptr;
 
 	FMeshBatchElement& OutBatchElement = OutMeshBatch.Elements[0];
 
-	OutBatchElement.VertexFactoryUserData = ProxyLODInfo.OverrideColorVertexBuffer ? ProxyLODInfo.OverrideColorVFUniformBuffer.GetReference() : VFs.VertexFactory.GetUniformBuffer();
-	OutMeshBatch.VertexFactory = ProxyLODInfo.OverrideColorVertexBuffer ? &VFs.VertexFactoryOverrideColorVertexBuffer : &VFs.VertexFactory;
-	OutMeshBatch.MaterialRenderProxy = WireframeRenderProxy;
-	OutBatchElement.MinVertexIndex = 0;
-	OutBatchElement.MaxVertexIndex = LODModel.GetNumVertices() - 1;
-	OutMeshBatch.ReverseCulling = bReverseCulling || IsLocalToWorldDeterminantNegative();
-	OutMeshBatch.CastShadow = bCastShadow;
-	OutMeshBatch.DepthPriorityGroup = (ESceneDepthPriorityGroup)InDepthPriorityGroup;
-	if (ForcedLodModel > 0) 
+	if (ProxyLODInfo.OverrideColorVertexBuffer)
 	{
-		OutBatchElement.MaxScreenSize = 0.0f;
-		OutBatchElement.MinScreenSize = -1.0f;
+		VertexFactory = &VFs.VertexFactoryOverrideColorVertexBuffer;
+
+		OutBatchElement.VertexFactoryUserData = ProxyLODInfo.OverrideColorVFUniformBuffer.GetReference();
 	}
 	else
 	{
-		OutBatchElement.MaxScreenSize = GetScreenSize(LODIndex);
-		OutBatchElement.MinScreenSize = 0.0f;
-		if (LODIndex < MAX_STATIC_MESH_LODS - 1)
-		{
-			OutBatchElement.MinScreenSize = GetScreenSize(LODIndex + 1);
-		}
+		VertexFactory = &VFs.VertexFactory;
+
+		OutBatchElement.VertexFactoryUserData = VFs.VertexFactory.GetUniformBuffer();
 	}
 
 	const bool bWireframe = true;
 	const bool bRequiresAdjacencyInformation = false;
 	const bool bUseReversedIndices = false;
+	const bool bDitheredLODTransition = false;
 
-	SetIndexSource(LODIndex, 0, OutMeshBatch, bWireframe, bRequiresAdjacencyInformation, bUseReversedIndices, bAllowPreCulledIndices);
+	OutMeshBatch.ReverseCulling = IsReversedCullingNeeded(bUseReversedIndices);
+	OutMeshBatch.CastShadow = bCastShadow;
+	OutMeshBatch.DepthPriorityGroup = (ESceneDepthPriorityGroup)InDepthPriorityGroup;
+	OutMeshBatch.MaterialRenderProxy = WireframeRenderProxy;
 
-	return OutBatchElement.NumPrimitives > 0;
+	OutBatchElement.MinVertexIndex = 0;
+	OutBatchElement.MaxVertexIndex = LODModel.GetNumVertices() - 1;
+
+	const uint32_t NumPrimitives = SetMeshElementGeometrySource(LODIndex, 0, bWireframe, bRequiresAdjacencyInformation, bUseReversedIndices, bAllowPreCulledIndices, VertexFactory, OutMeshBatch);
+	SetMeshElementScreenSize(LODIndex, bDitheredLODTransition, OutMeshBatch);
+
+	return NumPrimitives > 0;
+}
+
+bool FStaticMeshSceneProxy::GetCollisionMeshElement(
+	int32 LODIndex,
+	int32 BatchIndex,
+	int32 SectionIndex,
+	uint8 InDepthPriorityGroup,
+	const FMaterialRenderProxy* RenderProxy,
+	FMeshBatch& OutMeshBatch) const
+{
+	const FStaticMeshLODResources& LOD = RenderData->LODResources[LODIndex];
+	const FStaticMeshVertexFactories& VFs = RenderData->LODVertexFactories[LODIndex];
+	const FStaticMeshSection& Section = LOD.Sections[SectionIndex];
+	const FVertexFactory* VertexFactory = nullptr;
+
+	const FLODInfo& ProxyLODInfo = LODs[LODIndex];
+
+	const bool bWireframe = false;
+	const bool bRequiresAdjacencyInformation = false;
+	const bool bUseReversedIndices = false;
+	const bool bAllowPreCulledIndices = true;
+	const bool bDitheredLODTransition = false;
+
+	SetMeshElementGeometrySource(LODIndex, SectionIndex, bWireframe, bRequiresAdjacencyInformation, bUseReversedIndices, bAllowPreCulledIndices, VertexFactory, OutMeshBatch);
+
+	FMeshBatchElement& OutMeshBatchElement = OutMeshBatch.Elements[0];
+
+	if (ProxyLODInfo.OverrideColorVertexBuffer)
+	{
+		VertexFactory = &VFs.VertexFactoryOverrideColorVertexBuffer;
+
+		OutMeshBatchElement.VertexFactoryUserData = ProxyLODInfo.OverrideColorVFUniformBuffer.GetReference();
+	}
+	else
+	{
+		VertexFactory = &VFs.VertexFactory;
+
+		OutMeshBatchElement.VertexFactoryUserData = VFs.VertexFactory.GetUniformBuffer();
+	}
+
+	if (OutMeshBatchElement.NumPrimitives > 0)
+	{
+		OutMeshBatch.LODIndex = LODIndex;
+#if STATICMESH_ENABLE_DEBUG_RENDERING
+		OutMeshBatch.VisualizeLODIndex = LODIndex;
+		OutMeshBatch.VisualizeHLODIndex = HierarchicalLODIndex;
+#endif
+		OutMeshBatch.ReverseCulling = IsReversedCullingNeeded(bUseReversedIndices);
+		OutMeshBatch.CastShadow = false;
+		OutMeshBatch.DepthPriorityGroup = (ESceneDepthPriorityGroup)InDepthPriorityGroup;
+		OutMeshBatch.LCI = &ProxyLODInfo;
+		OutMeshBatch.VertexFactory = VertexFactory;
+		OutMeshBatch.MaterialRenderProxy = RenderProxy;
+
+		OutMeshBatchElement.MinVertexIndex = Section.MinVertexIndex;
+		OutMeshBatchElement.MaxVertexIndex = Section.MaxVertexIndex;
+#if STATICMESH_ENABLE_DEBUG_RENDERING
+		OutMeshBatchElement.VisualizeElementIndex = SectionIndex;
+#endif
+
+		SetMeshElementScreenSize(LODIndex, bDitheredLODTransition, OutMeshBatch);
+
+		return true;
+	}
+	else
+	{
+		return false;
+	}
 }
 
 #if WITH_EDITORONLY_DATA
@@ -641,75 +712,118 @@ bool FStaticMeshSceneProxy::GetMaterialTextureScales(int32 LODIndex, int32 Secti
 }
 #endif
 
-/**
- * Sets IndexBuffer, FirstIndex and NumPrimitives of OutMeshElement.
- */
-void FStaticMeshSceneProxy::SetIndexSource(int32 LODIndex, int32 SectionIndex, FMeshBatch& OutMeshElement, bool bWireframe, bool bRequiresAdjacencyInformation, bool bUseReversedIndices, bool bAllowPreCulledIndices) const
+uint32 FStaticMeshSceneProxy::SetMeshElementGeometrySource(
+	int32 LODIndex,
+	int32 SectionIndex,
+	bool bWireframe,
+	bool bRequiresAdjacencyInformation,
+	bool bUseReversedIndices,
+	bool bAllowPreCulledIndices,
+	const FVertexFactory* VertexFactory,
+	FMeshBatch& OutMeshBatch) const
 {
-	FMeshBatchElement& OutElement = OutMeshElement.Elements[0];
 	const FStaticMeshLODResources& LODModel = RenderData->LODResources[LODIndex];
+	const FStaticMeshSection& Section = LODModel.Sections[SectionIndex];
+	const FLODInfo& LODInfo = LODs[LODIndex];
+	const FLODInfo::FSectionInfo& SectionInfo = LODInfo.Sections[SectionIndex];
+
+	FMeshBatchElement& OutMeshBatchElement = OutMeshBatch.Elements[0];
+	uint32 NumPrimitives = 0;
+
+	const bool bHasPreculledTriangles = LODInfo.Sections[SectionIndex].NumPreCulledTriangles >= 0;
+	const bool bUsePreculledIndices = bAllowPreCulledIndices && GUsePreCulledIndexBuffer && bHasPreculledTriangles;
+
 	if (bWireframe)
 	{
-		if( LODModel.WireframeIndexBuffer.IsInitialized()
-			&& !(RHISupportsTessellation(GetScene().GetShaderPlatform()) && OutMeshElement.VertexFactory->GetType()->SupportsTessellationShaders())
-			)
+		const bool bSupportsTessellation = RHISupportsTessellation(GetScene().GetShaderPlatform()) && VertexFactory->GetType()->SupportsTessellationShaders();
+
+		if (LODModel.WireframeIndexBuffer.IsInitialized() && !bSupportsTessellation)
 		{
-			OutMeshElement.Type = PT_LineList;
-			OutElement.FirstIndex = 0;
-			OutElement.IndexBuffer = &LODModel.WireframeIndexBuffer;
-			OutElement.NumPrimitives = LODModel.WireframeIndexBuffer.GetNumIndices() / 2;
+			OutMeshBatch.Type = PT_LineList;
+			OutMeshBatchElement.FirstIndex = 0;
+			OutMeshBatchElement.IndexBuffer = &LODModel.WireframeIndexBuffer;
+			NumPrimitives = LODModel.WireframeIndexBuffer.GetNumIndices() / 2;
 		}
 		else
 		{
-			OutMeshElement.Type = PT_TriangleList;
+			OutMeshBatch.Type = PT_TriangleList;
 
-			if (bAllowPreCulledIndices
-				&& GUsePreCulledIndexBuffer 
-				&& LODs[LODIndex].Sections[SectionIndex].NumPreCulledTriangles >= 0)
+			if (bUsePreculledIndices)
 			{
-				OutElement.IndexBuffer = LODs[LODIndex].PreCulledIndexBuffer;
-				OutElement.FirstIndex = 0;
-				OutElement.NumPrimitives = LODs[LODIndex].PreCulledIndexBuffer->GetNumIndices() / 3;
+				OutMeshBatchElement.IndexBuffer = LODInfo.PreCulledIndexBuffer;
+				OutMeshBatchElement.FirstIndex = 0;
+				NumPrimitives = LODInfo.PreCulledIndexBuffer->GetNumIndices() / 3;
 			}
 			else
 			{
-				OutElement.FirstIndex = 0;
-				OutElement.IndexBuffer = &LODModel.IndexBuffer;
-				OutElement.NumPrimitives = LODModel.IndexBuffer.GetNumIndices() / 3;
+				OutMeshBatchElement.FirstIndex = 0;
+				OutMeshBatchElement.IndexBuffer = &LODModel.IndexBuffer;
+				NumPrimitives = LODModel.IndexBuffer.GetNumIndices() / 3;
 			}
 
-			OutMeshElement.bWireframe = true;
-			OutMeshElement.bDisableBackfaceCulling = true;
+			OutMeshBatch.bWireframe = true;
+			OutMeshBatch.bDisableBackfaceCulling = true;
 		}
 	}
 	else
 	{
-		const FStaticMeshSection& Section = LODModel.Sections[SectionIndex];
-		OutMeshElement.Type = PT_TriangleList;
+		OutMeshBatch.Type = PT_TriangleList;
 
-		if (bAllowPreCulledIndices
-			&& GUsePreCulledIndexBuffer 
-			&& LODs[LODIndex].Sections[SectionIndex].NumPreCulledTriangles >= 0)
+		if (bUsePreculledIndices)
 		{
-			OutElement.IndexBuffer = LODs[LODIndex].PreCulledIndexBuffer;
-			OutElement.FirstIndex = LODs[LODIndex].Sections[SectionIndex].FirstPreCulledIndex;
-			OutElement.NumPrimitives = LODs[LODIndex].Sections[SectionIndex].NumPreCulledTriangles;
+			OutMeshBatchElement.IndexBuffer = LODInfo.PreCulledIndexBuffer;
+			OutMeshBatchElement.FirstIndex = SectionInfo.FirstPreCulledIndex;
+			NumPrimitives = SectionInfo.NumPreCulledTriangles;
 		}
 		else
 		{
-			OutElement.IndexBuffer = bUseReversedIndices ? &LODModel.ReversedIndexBuffer : &LODModel.IndexBuffer;
-			OutElement.FirstIndex = Section.FirstIndex;
-			OutElement.NumPrimitives = Section.NumTriangles;
+			OutMeshBatchElement.IndexBuffer = bUseReversedIndices ? &LODModel.ReversedIndexBuffer : &LODModel.IndexBuffer;
+			OutMeshBatchElement.FirstIndex = Section.FirstIndex;
+			NumPrimitives = Section.NumTriangles;
 		}
 	}
 
-	if ( bRequiresAdjacencyInformation )
+	if (bRequiresAdjacencyInformation)
 	{
 		check( LODModel.bHasAdjacencyInfo );
-		OutElement.IndexBuffer = &LODModel.AdjacencyIndexBuffer;
-		OutMeshElement.Type = PT_12_ControlPointPatchList;
-		OutElement.FirstIndex *= 4;
+		OutMeshBatchElement.IndexBuffer = &LODModel.AdjacencyIndexBuffer;
+		OutMeshBatch.Type = PT_12_ControlPointPatchList;
+		OutMeshBatchElement.FirstIndex *= 4;
 	}
+
+	OutMeshBatchElement.NumPrimitives = NumPrimitives;
+	OutMeshBatch.VertexFactory = VertexFactory;
+
+	return NumPrimitives;
+}
+
+void FStaticMeshSceneProxy::SetMeshElementScreenSize(int32 LODIndex, bool bDitheredLODTransition, FMeshBatch& OutMeshBatch) const
+{
+	FMeshBatchElement& OutBatchElement = OutMeshBatch.Elements[0];
+
+	if (ForcedLodModel > 0)
+	{
+		OutMeshBatch.bDitheredLODTransition = false;
+
+		OutBatchElement.MaxScreenSize = 0.0f;
+		OutBatchElement.MinScreenSize = -1.0f;
+	}
+	else
+	{
+		OutMeshBatch.bDitheredLODTransition = bDitheredLODTransition;
+
+		OutBatchElement.MaxScreenSize = GetScreenSize(LODIndex);
+		OutBatchElement.MinScreenSize = 0.0f;
+		if (LODIndex < MAX_STATIC_MESH_LODS - 1)
+		{
+			OutBatchElement.MinScreenSize = GetScreenSize(LODIndex + 1);
+		}
+	}
+}
+
+bool FStaticMeshSceneProxy::IsReversedCullingNeeded(bool bUseReversedIndices) const
+{
+	return (bReverseCulling || IsLocalToWorldDeterminantNegative()) && !bUseReversedIndices;
 }
 
 // FPrimitiveSceneProxy interface.
@@ -1233,9 +1347,8 @@ void FStaticMeshSceneProxy::GetDynamicMeshElements(const TArray<const FSceneView
 								for (int32 BatchIndex = 0; BatchIndex < GetNumMeshBatches(); BatchIndex++)
 								{
 									FMeshBatch& CollisionElement = Collector.AllocateMesh();
-									if (GetMeshElement(DrawLOD, BatchIndex, SectionIndex, SDPG_World, bSectionIsSelected, true, CollisionElement))
+									if (GetCollisionMeshElement(DrawLOD, BatchIndex, SectionIndex, SDPG_World, CollisionMaterialInstance, CollisionElement))
 									{
-										CollisionElement.MaterialRenderProxy = CollisionMaterialInstance;
 										Collector.AddMesh(ViewIndex, CollisionElement);
 										INC_DWORD_STAT_BY(STAT_StaticMeshTriangles, CollisionElement.GetNumPrimitives());
 									}
diff --git a/Engine/Source/Runtime/Engine/Public/StaticMeshResources.h b/Engine/Source/Runtime/Engine/Public/StaticMeshResources.h
--- a/Engine/Source/Runtime/Engine/Public/StaticMeshResources.h
+++ b/Engine/Source/Runtime/Engine/Public/StaticMeshResources.h
@@ -572,11 +572,33 @@ public:
 	/** Sets up a wireframe FMeshBatch for a specific LOD. */
 	virtual bool GetWireframeMeshElement(int32 LODIndex, int32 BatchIndex, const FMaterialRenderProxy* WireframeRenderProxy, uint8 InDepthPriorityGroup, bool bAllowPreCulledIndices, FMeshBatch& OutMeshBatch) const;
 
+	/** Sets up a collision FMeshBatch for a specific LOD and element. */
+	virtual bool GetCollisionMeshElement(
+		int32 LODIndex,
+		int32 BatchIndex,
+		int32 ElementIndex,
+		uint8 InDepthPriorityGroup,
+		const FMaterialRenderProxy* RenderProxy,
+		FMeshBatch& OutMeshBatch) const;
+
 protected:
-	/**
-	 * Sets IndexBuffer, FirstIndex and NumPrimitives of OutMeshElement.
-	 */
-	virtual void SetIndexSource(int32 LODIndex, int32 ElementIndex, FMeshBatch& OutMeshElement, bool bWireframe, bool bRequiresAdjacencyInformation, bool bUseInversedIndices, bool bAllowPreCulledIndices) const;
+	/** Configures mesh batch vertex / index state. Returns the number of primitives used in the element. */
+	uint32 SetMeshElementGeometrySource(
+		int32 LODIndex,
+		int32 ElementIndex,
+		bool bWireframe,
+		bool bRequiresAdjacencyInformation,
+		bool bUseInversedIndices,
+		bool bAllowPreCulledIndices,
+		const FVertexFactory* VertexFactory,
+		FMeshBatch& OutMeshElement) const;
+
+	/** Sets the screen size on a mesh element. */
+	void SetMeshElementScreenSize(int32 LODIndex, bool bDitheredLODTransition, FMeshBatch& OutMeshBatch) const;
+
+	/** Returns whether this mesh needs reverse culling when using reversed indices. */
+	bool IsReversedCullingNeeded(bool bUseReversedIndices) const;
+
 	bool IsCollisionView(const FEngineShowFlags& EngineShowFlags, bool& bDrawSimpleCollision, bool& bDrawComplexCollision) const;
 
 public:
