Fixed warning around constant in conditional expression.

diff --git a/Engine/Source/Runtime/Engine/Private/StaticMeshRender.cpp b/Engine/Source/Runtime/Engine/Private/StaticMeshRender.cpp
--- a/Engine/Source/Runtime/Engine/Private/StaticMeshRender.cpp
+++ b/Engine/Source/Runtime/Engine/Private/StaticMeshRender.cpp
@@ -472,7 +472,7 @@ bool FStaticMeshSceneProxy::GetMeshElement(
 	const bool bRequiresAdjacencyInformation = RequiresAdjacencyInformation(MaterialInterface, VertexFactory->GetType(), FeatureLevel);
 
 	// Two sided material use bIsFrontFace which is wrong with Reversed Indices. AdjacencyInformation use another index buffer.
-	const bool bUseReversedIndices = !bWireframe && GUseReversedIndexBuffer && IsLocalToWorldDeterminantNegative() && (LOD.bHasReversedIndices != 0) && !bRequiresAdjacencyInformation && !Material->IsTwoSided();
+	const bool bUseReversedIndices = GUseReversedIndexBuffer && IsLocalToWorldDeterminantNegative() && (LOD.bHasReversedIndices != 0) && !bRequiresAdjacencyInformation && !Material->IsTwoSided();
 
 	// No support for stateless dithered LOD transitions for movable meshes
 	const bool bDitheredLODTransition = !IsMovable() && Material->IsDitheredLODTransition();
