Fix nasty bug where chunks were being deserialized improperly if you are cooking a packaged build and have retrieved the chunk from the DDC.

diff --git a/Engine/Source/Runtime/Engine/Private/AudioDerivedData.cpp b/Engine/Source/Runtime/Engine/Private/AudioDerivedData.cpp
--- a/Engine/Source/Runtime/Engine/Private/AudioDerivedData.cpp
+++ b/Engine/Source/Runtime/Engine/Private/AudioDerivedData.cpp
@@ -44,7 +44,7 @@ Derived data key generation.
 
 // If you want to bump this version, generate a new guid using
 // VS->Tools->Create GUID and paste it here. https://www.guidgen.com works too.
-#define STREAMEDAUDIO_DERIVEDDATA_VER		TEXT("BC6E92FBBD314E3B9B9EC6778749EA7F")
+#define STREAMEDAUDIO_DERIVEDDATA_VER		TEXT("BC6E92FBBD314E3B9B9EC6778749EB5E")
 
 /**
  * Computes the derived data key suffix for a SoundWave's Streamed Audio.
@@ -363,7 +363,7 @@ class FStreamedAudioCacheDerivedDataWorker : public FNonAbandonableTask
 #endif
 
 						NewChunk->BulkData.Lock(LOCK_READ_WRITE);
-						void* NewChunkData = NewChunk->BulkData.Realloc(NewChunk->DataSize);
+						void* NewChunkData = NewChunk->BulkData.Realloc(NewChunk->AudioDataSize);
 						FMemory::Memcpy(NewChunkData, ChunkBuffers[0].GetData(), AudioDataSize);
 						NewChunk->BulkData.Unlock();
 					}
@@ -374,7 +374,6 @@ class FStreamedAudioCacheDerivedDataWorker : public FNonAbandonableTask
 						// Zero pad the reallocation if the chunk isn't precisely the max chunk size to keep the reads aligned to MaxChunkSize
 						const int32 AudioDataSize = ChunkBuffers[ChunkIndex].Num();
 						check(AudioDataSize != 0 && AudioDataSize <= MaxChunkSize);
-						int32 AlignedChunkSize = FMath::RoundUpToPowerOfTwo(MaxChunkSize);
 						const int32 ZeroPadBytes = FMath::Max(MaxChunkSize - AudioDataSize, 0);
 
 						FStreamedAudioChunk* NewChunk = new FStreamedAudioChunk();
@@ -426,7 +425,7 @@ class FStreamedAudioCacheDerivedDataWorker : public FNonAbandonableTask
 				// Store it in the cache.
 				// @todo: This will remove the streaming bulk data, which we immediately reload below!
 				// Should ideally avoid this redundant work, but it only happens when we actually have
-				// to build the texture, which should only ever be once.
+				// to build the compressed audio, which should only ever be once.
 				this->BytesCached = PutDerivedDataInCache(DerivedData, KeySuffix);
 
 				check(this->BytesCached != 0);
@@ -663,8 +662,10 @@ bool FStreamedAudioPlatformData::TryInlineChunkData()
 			if (bLoadedFromDDC)
 			{
 				int32 ChunkSize = 0;
+				int32 AudioDataSize = 0;
 				FMemoryReader Ar(TempData, /*bIsPersistent=*/ true);
 				Ar << ChunkSize;
+				Ar << AudioDataSize; // Unused for the purposes of this function.
 
 #if WITH_EDITORONLY_DATA
 				if (Chunk.BulkData.IsLocked())
@@ -846,6 +847,7 @@ void FStreamedAudioPlatformData::Serialize(FArchive& Ar, USoundWave* Owner)
 			Chunks.Add(new FStreamedAudioChunk());
 		}
 	}
+
 	for (int32 ChunkIndex = 0; ChunkIndex < Chunks.Num(); ++ChunkIndex)
 	{
 		Chunks[ChunkIndex].Serialize(Ar, Owner, ChunkIndex);
