Merging iOS audio stutter/pop fix.

diff --git a/Engine/Source/Runtime/Engine/Private/ADPCMAudioInfo.cpp b/Engine/Source/Runtime/Engine/Private/ADPCMAudioInfo.cpp
--- a/Engine/Source/Runtime/Engine/Private/ADPCMAudioInfo.cpp
+++ b/Engine/Source/Runtime/Engine/Private/ADPCMAudioInfo.cpp
@@ -194,13 +194,14 @@ bool FADPCMAudioInfo::ReadCompressedInfo(const uint8* InSrcBufferData, uint32 In
 		BlockSize = *WaveInfo.pBlockAlign;
 
 		// ADPCM starts with 2 uncompressed samples and then the remaining compressed sample data has 2 samples per byte
-		const uint32 UncompressedBlockSamples = (2 + (BlockSize - PreambleSize) * 2);
-		UncompressedBlockSize = UncompressedBlockSamples * sizeof(int16);
+		UncompressedBlockSize = (2 + (BlockSize - PreambleSize) * 2) * sizeof(int16);
 		CompressedBlockSize = BlockSize;
 
-		const uint32 TargetBlocks = MONO_PCM_BUFFER_SAMPLES / UncompressedBlockSamples;
-		StreamBufferSize = TargetBlocks * UncompressedBlockSize;
-		TotalDecodedSize = TotalSamplesPerChannel * NumChannels * sizeof(int16);
+		const uint32 uncompressedBlockSamples = (2 + (BlockSize - PreambleSize) * 2);
+		const uint32 targetBlocks = MONO_PCM_BUFFER_SAMPLES / uncompressedBlockSamples;
+		StreamBufferSize = targetBlocks * UncompressedBlockSize;
+		// Ensure TotalDecodedSize is a even multiple of the compressed block size so that the buffer is not over read on the last block
+		TotalDecodedSize = ((WaveInfo.SampleDataSize + CompressedBlockSize - 1) / CompressedBlockSize) * UncompressedBlockSize;
 
 		UncompressedBlockData = (uint8*)FMemory::Realloc(UncompressedBlockData, NumChannels * UncompressedBlockSize);
 		check(UncompressedBlockData != nullptr);
