Fix for crash in EnsureZerothChunkIsLoaded. This function should noop when stream caching is disabled.

diff --git a/Engine/Source/Runtime/Engine/Private/SoundWave.cpp b/Engine/Source/Runtime/Engine/Private/SoundWave.cpp
--- a/Engine/Source/Runtime/Engine/Private/SoundWave.cpp
+++ b/Engine/Source/Runtime/Engine/Private/SoundWave.cpp
@@ -798,7 +798,7 @@ void USoundWave::PostLoad()
 void USoundWave::EnsureZerothChunkIsLoaded()
 {
 	// If the zeroth chunk is already loaded, early exit.
-	if (ZerothChunkData.Num() > 0)
+	if (ZerothChunkData.Num() > 0 || !ShouldUseStreamCaching())
 	{
 		return;
 	}
