Making the Auto attach parameters code part of FXComponent, so that calls can be shared with Niagara

diff --git a/Engine/Plugins/FX/Niagara/Source/Niagara/Public/NiagaraComponent.h b/Engine/Plugins/FX/Niagara/Source/Niagara/Public/NiagaraComponent.h
--- a/Engine/Plugins/FX/Niagara/Source/Niagara/Public/NiagaraComponent.h
+++ b/Engine/Plugins/FX/Niagara/Source/Niagara/Public/NiagaraComponent.h
@@ -418,8 +418,7 @@ public:
 	 * @param  ScaleRule		Option for how we handle our scale when we attach to Parent.
 	 * @see bAutoManageAttachment, AutoAttachParent, AutoAttachSocketName, AutoAttachLocationRule, AutoAttachRotationRule, AutoAttachScaleRule
 	 */
-	UFUNCTION(BlueprintCallable, Category = Niagara)
-	void SetAutoAttachmentParameters(USceneComponent* Parent, FName SocketName, EAttachmentRule LocationRule, EAttachmentRule RotationRule, EAttachmentRule ScaleRule);
+	void SetAutoAttachmentParameters(USceneComponent* Parent, FName SocketName, EAttachmentRule LocationRule, EAttachmentRule RotationRule, EAttachmentRule ScaleRule) override;
 
 	UPROPERTY(EditAnywhere, Category = Preview, Transient, meta=(EditCondition=bEnablePreviewDetailLevel))
 	int32 PreviewDetailLevel;
diff --git a/Engine/Source/Runtime/Engine/Classes/Particles/ParticleSystemComponent.h b/Engine/Source/Runtime/Engine/Classes/Particles/ParticleSystemComponent.h
--- a/Engine/Source/Runtime/Engine/Classes/Particles/ParticleSystemComponent.h
+++ b/Engine/Source/Runtime/Engine/Classes/Particles/ParticleSystemComponent.h
@@ -369,6 +369,19 @@ public:
 	 */
 	UFUNCTION(BlueprintCallable, Category = "Effects|Components|ParticleSystem")
 	virtual void SetEmitterEnable(FName EmitterName, bool bNewEnableState) {};
+
+
+	/**
+	 * Set AutoAttachParent, AutoAttachSocketName, AutoAttachLocationRule, AutoAttachRotationRule, AutoAttachScaleRule to the specified parameters. Does not change bAutoManageAttachment; that must be set separately.
+	 * @param  Parent			Component to attach to.
+	 * @param  SocketName		Socket on Parent to attach to.
+	 * @param  LocationRule		Option for how we handle our location when we attach to Parent.
+	 * @param  RotationRule		Option for how we handle our rotation when we attach to Parent.
+	 * @param  ScaleRule		Option for how we handle our scale when we attach to Parent.
+	 * @see bAutoManageAttachment, AutoAttachParent, AutoAttachSocketName, AutoAttachLocationRule, AutoAttachRotationRule, AutoAttachScaleRule
+	 */
+	UFUNCTION(BlueprintCallable, Category = "Effects|Components|ParticleSystem")
+	virtual void SetAutoAttachmentParameters(USceneComponent* Parent, FName SocketName, EAttachmentRule LocationRule, EAttachmentRule RotationRule, EAttachmentRule ScaleRule) {}
 };
 
 
@@ -735,8 +748,7 @@ public:
 	 * @param  ScaleRule		Option for how we handle our scale when we attach to Parent.
 	 * @see bAutoManageAttachment, AutoAttachParent, AutoAttachSocketName, AutoAttachLocationRule, AutoAttachRotationRule, AutoAttachScaleRule
 	 */
-	UFUNCTION(BlueprintCallable, Category = "Effects|Components|ParticleSystem")
-	void SetAutoAttachmentParameters(USceneComponent* Parent, FName SocketName, EAttachmentRule LocationRule, EAttachmentRule RotationRule, EAttachmentRule ScaleRule);
+	void SetAutoAttachmentParameters(USceneComponent* Parent, FName SocketName, EAttachmentRule LocationRule, EAttachmentRule RotationRule, EAttachmentRule ScaleRule) override;
 
 private:
 
@@ -1575,11 +1587,4 @@ FORCEINLINE_DEBUGGABLE void UParticleSystemComponent::SetAutoAttachParams(UScene
 	USceneComponent::ConvertAttachLocation(LocationType, AutoAttachLocationRule, AutoAttachRotationRule, AutoAttachScaleRule);
 }
 
-FORCEINLINE_DEBUGGABLE void UParticleSystemComponent::SetAutoAttachmentParameters(USceneComponent* Parent, FName SocketName, EAttachmentRule LocationRule, EAttachmentRule RotationRule, EAttachmentRule ScaleRule)
-{
-	AutoAttachParent = Parent;
-	AutoAttachSocketName = SocketName;
-	AutoAttachLocationRule = LocationRule;
-	AutoAttachRotationRule = RotationRule;
-	AutoAttachScaleRule = ScaleRule;
-}
+
diff --git a/Engine/Source/Runtime/Engine/Private/Particles/ParticleComponents.cpp b/Engine/Source/Runtime/Engine/Private/Particles/ParticleComponents.cpp
--- a/Engine/Source/Runtime/Engine/Private/Particles/ParticleComponents.cpp
+++ b/Engine/Source/Runtime/Engine/Private/Particles/ParticleComponents.cpp
@@ -8554,5 +8554,14 @@ FParticleSystemWorldManager* UParticleSystemComponent::GetWorldManager()const
 	return FParticleSystemWorldManager::Get(GetWorld());
 }
 
+void UParticleSystemComponent::SetAutoAttachmentParameters(USceneComponent* Parent, FName SocketName, EAttachmentRule LocationRule, EAttachmentRule RotationRule, EAttachmentRule ScaleRule)
+{
+	AutoAttachParent = Parent;
+	AutoAttachSocketName = SocketName;
+	AutoAttachLocationRule = LocationRule;
+	AutoAttachRotationRule = RotationRule;
+	AutoAttachScaleRule = ScaleRule;
+}
+
 #undef LOCTEXT_NAMESPACE
 
