Defer replay playback from console command to the next frame to avoid a potential crash

diff --git a/Engine/Source/Runtime/Engine/Private/World.cpp b/Engine/Source/Runtime/Engine/Private/World.cpp
--- a/Engine/Source/Runtime/Engine/Private/World.cpp
+++ b/Engine/Source/Runtime/Engine/Private/World.cpp
@@ -3787,7 +3787,7 @@ bool UWorld::HandleDemoPlayCommand( const TCHAR* Cmd, FOutputDevice& Ar, UWorld*
 		{
 			if (Context.World()->DemoNetDriver != nullptr && Context.World()->DemoNetDriver->IsPlaying())
 			{
-				ErrorString = TEXT("A demo is already in progress, cannot play more than one demo at a time");
+				ErrorString = TEXT("A demo is already in progress, cannot play more than one demo at a time in PIE.");
 				break;
 			}
 		}
@@ -3804,18 +3804,27 @@ bool UWorld::HandleDemoPlayCommand( const TCHAR* Cmd, FOutputDevice& Ar, UWorld*
 	}
 	else
 	{
-		// Allow additional url arguments after the demo name
-		TArray<FString> Options;
-		if (Temp.ParseIntoArray(Options, TEXT("?")) > 1)
-		{
-			Temp = Options[0];
-			Options.RemoveAtSwap(0);
-			InWorld->GetGameInstance()->PlayReplay(Temp, nullptr, Options);
-		}
-		else
+		// defer playback to the next frame
+		GetTimerManager().SetTimerForNextTick(FTimerDelegate::CreateWeakLambda(InWorld, [InWorld, Temp]()
 		{
-			InWorld->GetGameInstance()->PlayReplay(Temp);
-		}
+			if (InWorld->GetGameInstance())
+			{
+				FString ReplayName = Temp;
+				// Allow additional url arguments after the demo name
+				TArray<FString> Options;
+				if (Temp.ParseIntoArray(Options, TEXT("?")) > 1)
+				{
+					ReplayName = Options[0];
+					Options.RemoveAtSwap(0);
+
+					InWorld->GetGameInstance()->PlayReplay(ReplayName, nullptr, Options);
+				}
+				else
+				{
+					InWorld->GetGameInstance()->PlayReplay(ReplayName);
+				}
+			}
+		}));
 	}
 
 	return true;
