Tuning GC interval from 10->15 seconds for new GLowMemoryTimeBetweenPurgingPendingLevels setting.  Testing showed diminishing returns when speeding up GC faster than this.

diff --git a/Engine/Source/Runtime/Engine/Private/UnrealEngine.cpp b/Engine/Source/Runtime/Engine/Private/UnrealEngine.cpp
--- a/Engine/Source/Runtime/Engine/Private/UnrealEngine.cpp
+++ b/Engine/Source/Runtime/Engine/Private/UnrealEngine.cpp
@@ -1154,7 +1154,7 @@ static FAutoConsoleVariableRef CVarLowMemoryTimeBetweenPurgingPendingKillObjects
 );
 
 // see also: s.ForceGCAfterLevelStreamedOut, s.ContinuouslyIncrementalGCWhileLevelsPendingPurge
-static float GLowMemoryTimeBetweenPurgingPendingLevels = 10.0f;
+static float GLowMemoryTimeBetweenPurgingPendingLevels = 15.0f;
 static FAutoConsoleVariableRef CVarLowMemoryTimeBetweenPurgingPendingLevels(
 	TEXT("gc.LowMemory.TimeBetweenPurgingPendingLevels"),
 	GLowMemoryTimeBetweenPurgingPendingLevels,
