Fix SimulatedProxy additive root motion source dash overshoot by delaying cleanup of authoritative root motion sources until right before use (so we don't throw away useful server state prematurely) and by allowing CalcVelocity() to work for SimulatedProxies during root motion sources (so that PerformMovement is accurate and friction is accounted for)
- Situation: AutonomousProxy triggers a short ConstantForce additive RootMotionSource via an GameplayAbility task. The server does the same. SimulatedProxy versions of the character, say on another client first receives an ACharacter::OnRep_ReplicatedMovement call, containing the location and velocity *after a frame of the root motion source has been applied* on the server - SimulatedProxy checks whether it has active root motion or not, which it does not have because it has not yet processed the GameplayAbility task/logic that will add the RMS locally, so the SimulatedProxy sets location/velocity to the "post-RMS tick" location/velocity. SimulatedProxy then gets OnRep_RootMotion called, which prior to this change immediately cleans up any root motion sources from the server that we don't have locally. Since the client hasn't yet processed the GameplayAbility task/logic, it throws out the server RMS data. This is unfortunate because we may not receive any authoritative server state of the RMS for many frames after. The client THEN gets the gameplay task logic triggered, adds the RMS to pending, and then ticks PerformMovement() locally - unfortunately, the SimulatedProxy is now running the first tick of RMS starting from a post-RMS tick location received from the server, so the SimulatedProxy will be a tick ahead of where it should be throughout the remainder of the RootMotionSource. Additionally, on the last frame of the RMS on the SimulatedProxy, it calls PerformMovement() which removes the RMS due to expiration, and then processed PerformMovement as normal - unfortunately in CalcVelocity() we guard against SimulatedProxy ever using it, which means friction does not get applied for this tick. For a high-velocity dash, this can mean a significant velocity remains, and the SimulatedProxy smoothing/correction code will overshoot end destination.
- Fix: This change moves the cleanup of authoritative root motion source state until it is actually used, and additionally makes the "update SimulatedProxy RMS to server state" code more robust by looking at both PendingRootMotionSources as well as active ones. Now, when the SimulatedTick() happens on the client, it has already run the logic that adds the RMS locally, still has the authoritative state available, and activates the RMS and snaps it to the location the server was at. Now the SimulatedProxy RMS simulation matches what happens on Autonomous and Authority versions. The other fix is to allow CalcVelocity function (which adds friction) to run during RootMotionSource simulation since PerformMovement needs to be the same between clients and servers in order to have a hope of coming up with the same movement results.

#rnx

diff --git a/Engine/Source/Runtime/Engine/Private/Character.cpp b/Engine/Source/Runtime/Engine/Private/Character.cpp
--- a/Engine/Source/Runtime/Engine/Private/Character.cpp
+++ b/Engine/Source/Runtime/Engine/Private/Character.cpp
@@ -1179,15 +1179,6 @@ void ACharacter::OnRep_RootMotion()
 			FSimulatedRootMotionReplicatedMove& NewMove = RootMotionRepMoves.Last();
 			NewMove.RootMotion = RepRootMotion;
 			NewMove.Time = GetWorld()->GetTimeSeconds();
-
-			// Convert RootMotionSource Server IDs -> Local IDs in AuthoritativeRootMotion and cull invalid
-			// so that when we use this root motion it has the correct IDs
-			if (CharacterMovement)
-			{
-				CharacterMovement->ConvertRootMotionServerIDsToLocalIDs(CharacterMovement->CurrentRootMotion, NewMove.RootMotion.AuthoritativeRootMotion, NewMove.Time);
-			}
-			
-			NewMove.RootMotion.AuthoritativeRootMotion.CullInvalidSources();
 		}
 		else
 		{
diff --git a/Engine/Source/Runtime/Engine/Private/Components/CharacterMovementComponent.cpp b/Engine/Source/Runtime/Engine/Private/Components/CharacterMovementComponent.cpp
--- a/Engine/Source/Runtime/Engine/Private/Components/CharacterMovementComponent.cpp
+++ b/Engine/Source/Runtime/Engine/Private/Components/CharacterMovementComponent.cpp
@@ -1528,7 +1528,7 @@ void UCharacterMovementComponent::SimulatedTick(float DeltaSeconds)
 		if( CharacterOwner->RootMotionRepMoves.Num() > 0 )
 		{
 			// Move Actor back to position of that buffered move. (server replicated position).
-			const FSimulatedRootMotionReplicatedMove& RootMotionRepMove = CharacterOwner->RootMotionRepMoves.Last();
+			FSimulatedRootMotionReplicatedMove& RootMotionRepMove = CharacterOwner->RootMotionRepMoves.Last();
 			if( CharacterOwner->RestoreReplicatedMove(RootMotionRepMove) )
 			{
 				bCorrectedToServer = true;
@@ -1538,6 +1538,11 @@ void UCharacterMovementComponent::SimulatedTick(float DeltaSeconds)
 			CharacterOwner->PostNetReceiveVelocity(RootMotionRepMove.RootMotion.LinearVelocity);
 			LastUpdateVelocity = RootMotionRepMove.RootMotion.LinearVelocity;
 
+			// Convert RootMotionSource Server IDs -> Local IDs in AuthoritativeRootMotion and cull invalid
+			// so that when we use this root motion it has the correct IDs
+			ConvertRootMotionServerIDsToLocalIDs(CurrentRootMotion, RootMotionRepMove.RootMotion.AuthoritativeRootMotion, RootMotionRepMove.Time);
+			RootMotionRepMove.RootMotion.AuthoritativeRootMotion.CullInvalidSources();
+
 			// Set root motion states to that of repped in state
 			CurrentRootMotion.UpdateStateFrom(RootMotionRepMove.RootMotion.AuthoritativeRootMotion, true);
 
@@ -3154,8 +3159,8 @@ bool UCharacterMovementComponent::IsCrouching() const
 
 void UCharacterMovementComponent::CalcVelocity(float DeltaTime, float Friction, bool bFluid, float BrakingDeceleration)
 {
-	// Do not update velocity when using root motion or when SimulatedProxy - SimulatedProxy are repped their Velocity
-	if (!HasValidData() || HasAnimRootMotion() || DeltaTime < MIN_TICK_TIME || (CharacterOwner && CharacterOwner->Role == ROLE_SimulatedProxy))
+	// Do not update velocity when using root motion or when SimulatedProxy and not simulating root motion - SimulatedProxy are repped their Velocity
+	if (!HasValidData() || HasAnimRootMotion() || DeltaTime < MIN_TICK_TIME || (CharacterOwner && CharacterOwner->Role == ROLE_SimulatedProxy && !bWasSimulatingRootMotion))
 	{
 		return;
 	}
@@ -10044,7 +10049,11 @@ void UCharacterMovementComponent::ConvertRootMotionServerIDsToLocalIDs(const FRo
 
 			// If no mapping found, find match out of Local RootMotionSources that are not already mapped
 			bool bMatchFound = false;
-			for (const TSharedPtr<FRootMotionSource>& LocalRootMotionSource : LocalRootMotionToMatchWith.RootMotionSources)
+			TArray<TSharedPtr<FRootMotionSource>> LocalRootMotionSources;
+			LocalRootMotionSources.Reserve(LocalRootMotionToMatchWith.RootMotionSources.Num() + LocalRootMotionToMatchWith.PendingAddRootMotionSources.Num());
+			LocalRootMotionSources.Append(LocalRootMotionToMatchWith.RootMotionSources);
+			LocalRootMotionSources.Append(LocalRootMotionToMatchWith.PendingAddRootMotionSources);
+			for (const TSharedPtr<FRootMotionSource>& LocalRootMotionSource : LocalRootMotionSources)
 			{
 				if (LocalRootMotionSource.IsValid())
 				{
diff --git a/Engine/Source/Runtime/Engine/Private/GameFramework/RootMotionSource.cpp b/Engine/Source/Runtime/Engine/Private/GameFramework/RootMotionSource.cpp
--- a/Engine/Source/Runtime/Engine/Private/GameFramework/RootMotionSource.cpp
+++ b/Engine/Source/Runtime/Engine/Private/GameFramework/RootMotionSource.cpp
@@ -1838,6 +1838,28 @@ void FRootMotionSourceGroup::UpdateStateFrom(const FRootMotionSourceGroup& Group
 	bIsAdditiveVelocityApplied = GroupToTakeStateFrom.bIsAdditiveVelocityApplied;
 	LastPreAdditiveVelocity = GroupToTakeStateFrom.LastPreAdditiveVelocity;
 
+	// If we have a PendingAdd root motion source that is already active in GroupToTakeStateFrom, make it active
+	PendingAddRootMotionSources.RemoveAll([this, &GroupToTakeStateFrom](const TSharedPtr<FRootMotionSource>& RootSource) 
+		{ 
+			if (RootSource.IsValid())
+			{
+				if (RootSource->LocalID != (uint16)ERootMotionSourceID::Invalid)
+				{
+					for (const TSharedPtr<FRootMotionSource>& TakeFromRootMotionSource : GroupToTakeStateFrom.RootMotionSources)
+					{
+						if (TakeFromRootMotionSource.IsValid() && (RootSource->LocalID == TakeFromRootMotionSource->LocalID))
+						{
+							// Matches, move to active and remove from pending
+							UE_LOG(LogRootMotion, VeryVerbose, TEXT("UpdateStateFream moving PendingAdd RMS to active: %s"), *RootSource->ToSimpleString());
+							RootMotionSources.Add(RootSource);
+							return true;
+						}
+					}
+				}
+			}
+			return false;
+		});
+
 	// For each matching Source in GroupToTakeStateFrom, move state over to this group's Sources
 	// Can do all matching with LocalID only, since anything passed into this function should have
 	// already been "matched" to LocalIDs
