[NO NON-PII COMMIT MESSAGE]

diff --git a/Engine/Source/Runtime/Engine/Private/SoundWave.cpp b/Engine/Source/Runtime/Engine/Private/SoundWave.cpp
--- a/Engine/Source/Runtime/Engine/Private/SoundWave.cpp
+++ b/Engine/Source/Runtime/Engine/Private/SoundWave.cpp
@@ -853,12 +853,7 @@ uint32 USoundWave::GetSizeOfChunk(uint32 ChunkIndex)
 {
 	check(ChunkIndex < GetNumChunks());
 
-	if (ChunkIndex == 0)
-	{
-		EnsureZerothChunkIsLoaded();
-		return ZerothChunkData.Num();
-	}
-	else if(RunningPlatformData)
+	if(RunningPlatformData)
 	{
 		return RunningPlatformData->Chunks[ChunkIndex].AudioDataSize;
 	}
