Update dirty method for getting access to persist mapped gpu buffer for VT on certain platforms

diff --git a/Engine/Source/Runtime/Engine/Private/VT/VirtualTextureUploadCache.cpp b/Engine/Source/Runtime/Engine/Private/VT/VirtualTextureUploadCache.cpp
--- a/Engine/Source/Runtime/Engine/Private/VT/VirtualTextureUploadCache.cpp
+++ b/Engine/Source/Runtime/Engine/Private/VT/VirtualTextureUploadCache.cpp
@@ -209,7 +209,11 @@ FVTUploadTileHandle FVirtualTextureUploadCache::PrepareTileForUpload(FVTUploadTi
 		{
 			FRHIResourceCreateInfo CreateInfo;
 			NewEntry.RHIStagingBuffer = RHICreateStructuredBuffer(FormatInfo.BlockBytes, MemorySize, BUF_ShaderResource | BUF_Static, CreateInfo);
-			NewEntry.Memory = RHILockStructuredBuffer(NewEntry.RHIStagingBuffer, 0u, MemorySize, RLM_WriteOnly);
+
+			// Here we bypass 'normal' RHI operations in order to get a persistent pointer to GPU memory, on supported platforms
+			// This should be encapsulated into a proper RHI method at some point
+			NewEntry.Memory = GDynamicRHI->RHILockStructuredBuffer(NewEntry.RHIStagingBuffer, 0u, MemorySize, RLM_WriteOnly);
+
 			INC_MEMORY_STAT_BY(STAT_TotalGPUUploadSize, MemorySize);
 		}
 		else
