Fix missing newline in LowLevelOutputDebugString.

diff --git a/Engine/Source/Runtime/Engine/Private/UnrealEngine.cpp b/Engine/Source/Runtime/Engine/Private/UnrealEngine.cpp
--- a/Engine/Source/Runtime/Engine/Private/UnrealEngine.cpp
+++ b/Engine/Source/Runtime/Engine/Private/UnrealEngine.cpp
@@ -689,7 +689,7 @@ void SystemResolutionSinkCallback()
 
 			if (GEngine && GEngine->GameViewport && GEngine->GameViewport->ViewportFrame)
 			{
-				FPlatformMisc::LowLevelOutputDebugStringf(TEXT("Resizing viewport due to setres change, %d x %d"), ResX, ResY);
+				FPlatformMisc::LowLevelOutputDebugStringf(TEXT("Resizing viewport due to setres change, %d x %d\n"), ResX, ResY);
 				GEngine->GameViewport->ViewportFrame->ResizeFrame(ResX, ResY, WindowMode);
 			}
 		}
