Fix for Skeletal Mesh loading before its Skeleton asset

diff --git a/Engine/Source/Runtime/Engine/Classes/Engine/SkeletalMesh.h b/Engine/Source/Runtime/Engine/Classes/Engine/SkeletalMesh.h
--- a/Engine/Source/Runtime/Engine/Classes/Engine/SkeletalMesh.h
+++ b/Engine/Source/Runtime/Engine/Classes/Engine/SkeletalMesh.h
@@ -920,6 +920,7 @@ public:
 	virtual FString GetDetailedInfoInternal() const override;
 	virtual void GetResourceSizeEx(FResourceSizeEx& CumulativeResourceSize) override;
 	static void AddReferencedObjects(UObject* InThis, FReferenceCollector& Collector);
+	virtual void GetPreloadDependencies(TArray<UObject*>& OutDeps) override;
 	//~ End UObject Interface.
 
 	/** Setup-only routines - not concerned with the instance. */
diff --git a/Engine/Source/Runtime/Engine/Private/SkeletalMesh.cpp b/Engine/Source/Runtime/Engine/Private/SkeletalMesh.cpp
--- a/Engine/Source/Runtime/Engine/Private/SkeletalMesh.cpp
+++ b/Engine/Source/Runtime/Engine/Private/SkeletalMesh.cpp
@@ -1227,6 +1227,12 @@ void USkeletalMesh::AddReferencedObjects(UObject* InThis, FReferenceCollector& C
 	Super::AddReferencedObjects( This, Collector );
 }
 
+void USkeletalMesh::GetPreloadDependencies(TArray<UObject*>& OutDeps)
+{
+	Super::GetPreloadDependencies(OutDeps);
+	OutDeps.Add(Skeleton);
+}
+
 void USkeletalMesh::FlushRenderState()
 {
 	//TComponentReregisterContext<USkeletalMeshComponent> ReregisterContext;
