Do not render the UI if rendering is disabled by the app.

diff --git a/Engine/Source/Runtime/Engine/Private/GameViewportClient.cpp b/Engine/Source/Runtime/Engine/Private/GameViewportClient.cpp
--- a/Engine/Source/Runtime/Engine/Private/GameViewportClient.cpp
+++ b/Engine/Source/Runtime/Engine/Private/GameViewportClient.cpp
@@ -1566,6 +1566,7 @@ void UGameViewportClient::Draw(FViewport* InViewport, FCanvas* SceneCanvas)
 	}
 
 	// Render the UI.
+	if (FSlateApplication::Get().GetPlatformApplication()->IsAllowedToRender())
 	{
 		SCOPE_CYCLE_COUNTER(STAT_UIDrawingTime);
 		CSV_SCOPED_TIMING_STAT_EXCLUSIVE(UI);
