SpawnEmitterAttached/AtLocation now optionally auto activate.

diff --git a/Engine/Source/Runtime/Engine/Classes/Kismet/GameplayStatics.h b/Engine/Source/Runtime/Engine/Classes/Kismet/GameplayStatics.h
--- a/Engine/Source/Runtime/Engine/Classes/Kismet/GameplayStatics.h
+++ b/Engine/Source/Runtime/Engine/Classes/Kismet/GameplayStatics.h
@@ -372,12 +372,13 @@ class ENGINE_API UGameplayStatics : public UBlueprintFunctionLibrary
 	 * @param Scale - scale to create the effect at
 	 * @param bAutoDestroy - Whether the component will automatically be destroyed when the particle system completes playing or whether it can be reactivated
 	 * @param PoolingMethod - Method used for pooling this component. Defaults to none.
+	 * @param bAutoActivate - Whether the component will be automatically activated on creation.
 	 */
 	UFUNCTION(BlueprintCallable, Category="Effects|Components|ParticleSystem", meta=(Keywords = "particle system", WorldContext="WorldContextObject", UnsafeDuringActorConstruction = "true"))
-	static UParticleSystemComponent* SpawnEmitterAtLocation(const UObject* WorldContextObject, UParticleSystem* EmitterTemplate, FVector Location, FRotator Rotation = FRotator::ZeroRotator, FVector Scale = FVector(1.f), bool bAutoDestroy = true, EPSCPoolMethod PoolingMethod = EPSCPoolMethod::None);
+	static UParticleSystemComponent* SpawnEmitterAtLocation(const UObject* WorldContextObject, UParticleSystem* EmitterTemplate, FVector Location, FRotator Rotation = FRotator::ZeroRotator, FVector Scale = FVector(1.f), bool bAutoDestroy = true, EPSCPoolMethod PoolingMethod = EPSCPoolMethod::None, bool bAutoActivate = false);
 
 	// Backwards compatible version of SpawnEmitterAttached for C++ without Scale
-	static UParticleSystemComponent* SpawnEmitterAtLocation(const UObject* WorldContextObject, UParticleSystem* EmitterTemplate, FVector Location, FRotator Rotation, bool bAutoDestroy, EPSCPoolMethod PoolingMethod = EPSCPoolMethod::None);
+	static UParticleSystemComponent* SpawnEmitterAtLocation(const UObject* WorldContextObject, UParticleSystem* EmitterTemplate, FVector Location, FRotator Rotation, bool bAutoDestroy, EPSCPoolMethod PoolingMethod = EPSCPoolMethod::None, bool bAutoActivate=false);
 
 	/** Plays the specified effect at the given location and rotation, fire and forget. The system will go away when the effect is complete. Does not replicate.
 	 * @param World - The World to spawn in
@@ -385,11 +386,12 @@ class ENGINE_API UGameplayStatics : public UBlueprintFunctionLibrary
 	 * @param SpawnTransform - transform with which to place the effect in world space
 	 * @param bAutoDestroy - Whether the component will automatically be destroyed when the particle system completes playing or whether it can be reactivated
 	 * @param PoolingMethod - Method used for pooling this component. Defaults to none.
+	 * @param bAutoActivate - Whether the component will be automatically activated on creation.
 	 */
-	static UParticleSystemComponent* SpawnEmitterAtLocation(UWorld* World, UParticleSystem* EmitterTemplate, const FTransform& SpawnTransform, bool bAutoDestroy = true, EPSCPoolMethod PoolingMethod = EPSCPoolMethod::None);
+	static UParticleSystemComponent* SpawnEmitterAtLocation(UWorld* World, UParticleSystem* EmitterTemplate, const FTransform& SpawnTransform, bool bAutoDestroy = true, EPSCPoolMethod PoolingMethod = EPSCPoolMethod::None, bool bAutoActivate = true);
 
 private:
-	static UParticleSystemComponent* InternalSpawnEmitterAtLocation(UWorld* World, UParticleSystem* EmitterTemplate, FVector Location, FRotator Rotation, FVector Scale, bool bAutoDestroy, EPSCPoolMethod PoolingMethod = EPSCPoolMethod::None);
+	static UParticleSystemComponent* InternalSpawnEmitterAtLocation(UWorld* World, UParticleSystem* EmitterTemplate, FVector Location, FRotator Rotation, FVector Scale, bool bAutoDestroy, EPSCPoolMethod PoolingMethod = EPSCPoolMethod::None, bool bAutoActivate = true);
 
 public:
 
@@ -403,12 +405,13 @@ public:
 	 * @param LocationType - Specifies whether Location is a relative offset or an absolute world position
 	 * @param bAutoDestroy - Whether the component will automatically be destroyed when the particle system completes playing or whether it can be reactivated
 	 * @param PoolingMethod - Method used for pooling this component. Defaults to none.
+	 * @param bAutoActivate - Whether the component will be automatically activated on creation.
 	 */
 	UFUNCTION(BlueprintCallable, Category="Effects|Components|ParticleSystem", meta=(Keywords = "particle system", UnsafeDuringActorConstruction = "true"))
-	static UParticleSystemComponent* SpawnEmitterAttached(class UParticleSystem* EmitterTemplate, class USceneComponent* AttachToComponent, FName AttachPointName = NAME_None, FVector Location = FVector(ForceInit), FRotator Rotation = FRotator::ZeroRotator, FVector Scale = FVector(1.f), EAttachLocation::Type LocationType = EAttachLocation::KeepRelativeOffset, bool bAutoDestroy = true, EPSCPoolMethod PoolingMethod = EPSCPoolMethod::None);
+	static UParticleSystemComponent* SpawnEmitterAttached(class UParticleSystem* EmitterTemplate, class USceneComponent* AttachToComponent, FName AttachPointName = NAME_None, FVector Location = FVector(ForceInit), FRotator Rotation = FRotator::ZeroRotator, FVector Scale = FVector(1.f), EAttachLocation::Type LocationType = EAttachLocation::KeepRelativeOffset, bool bAutoDestroy = true, EPSCPoolMethod PoolingMethod = EPSCPoolMethod::None, bool bAutoActivate=true);
 
 	// Backwards compatible version of SpawnEmitterAttached for C++ without Scale
-	static UParticleSystemComponent* SpawnEmitterAttached(class UParticleSystem* EmitterTemplate, class USceneComponent* AttachToComponent, FName AttachPointName, FVector Location, FRotator Rotation, EAttachLocation::Type LocationType, bool bAutoDestroy = true, EPSCPoolMethod PoolingMethod = EPSCPoolMethod::None);
+	static UParticleSystemComponent* SpawnEmitterAttached(class UParticleSystem* EmitterTemplate, class USceneComponent* AttachToComponent, FName AttachPointName, FVector Location, FRotator Rotation, EAttachLocation::Type LocationType, bool bAutoDestroy = true, EPSCPoolMethod PoolingMethod = EPSCPoolMethod::None, bool bAutoActivate=true);
 
 	// --- Sound functions ------------------------------
 	
diff --git a/Engine/Source/Runtime/Engine/Private/GameplayStatics.cpp b/Engine/Source/Runtime/Engine/Private/GameplayStatics.cpp
--- a/Engine/Source/Runtime/Engine/Private/GameplayStatics.cpp
+++ b/Engine/Source/Runtime/Engine/Private/GameplayStatics.cpp
@@ -949,12 +949,12 @@ UParticleSystemComponent* CreateParticleSystem(UParticleSystem* EmitterTemplate,
 	return PSC;
 }
 
-UParticleSystemComponent* UGameplayStatics::SpawnEmitterAtLocation(const UObject* WorldContextObject, UParticleSystem* EmitterTemplate, FVector SpawnLocation, FRotator SpawnRotation, bool bAutoDestroy, EPSCPoolMethod PoolingMethod)
+UParticleSystemComponent* UGameplayStatics::SpawnEmitterAtLocation(const UObject* WorldContextObject, UParticleSystem* EmitterTemplate, FVector SpawnLocation, FRotator SpawnRotation, bool bAutoDestroy, EPSCPoolMethod PoolingMethod, bool bAutoActivate)
 {
-	return SpawnEmitterAtLocation(WorldContextObject, EmitterTemplate, SpawnLocation, SpawnRotation, FVector(1.f), bAutoDestroy, PoolingMethod);
+	return SpawnEmitterAtLocation(WorldContextObject, EmitterTemplate, SpawnLocation, SpawnRotation, FVector(1.f), bAutoDestroy, PoolingMethod, bAutoActivate);
 }
 
-UParticleSystemComponent* UGameplayStatics::InternalSpawnEmitterAtLocation(UWorld* World, UParticleSystem* EmitterTemplate, FVector SpawnLocation, FRotator SpawnRotation, FVector SpawnScale, bool bAutoDestroy, EPSCPoolMethod PoolingMethod)
+UParticleSystemComponent* UGameplayStatics::InternalSpawnEmitterAtLocation(UWorld* World, UParticleSystem* EmitterTemplate, FVector SpawnLocation, FRotator SpawnRotation, FVector SpawnScale, bool bAutoDestroy, EPSCPoolMethod PoolingMethod, bool bAutoActivate)
 {
 	check(World && EmitterTemplate);
 
@@ -968,7 +968,11 @@ UParticleSystemComponent* UGameplayStatics::InternalSpawnEmitterAtLocation(UWorl
 		PSC->RelativeScale3D = SpawnScale;
 
 		PSC->RegisterComponentWithWorld(World);
-		PSC->ActivateSystem(true);
+		if (bAutoActivate)
+		{
+			PSC->ActivateSystem(true);
+		}
+	
 
 		// Notify the texture streamer so that PSC gets managed as a dynamic component.
 		IStreamingManager::Get().NotifyPrimitiveUpdated(PSC);
@@ -988,35 +992,35 @@ UParticleSystemComponent* UGameplayStatics::InternalSpawnEmitterAtLocation(UWorl
 	return nullptr;
 }
 
-UParticleSystemComponent* UGameplayStatics::SpawnEmitterAtLocation(const UObject* WorldContextObject, UParticleSystem* EmitterTemplate, FVector SpawnLocation, FRotator SpawnRotation, FVector SpawnScale, bool bAutoDestroy, EPSCPoolMethod PoolingMethod)
+UParticleSystemComponent* UGameplayStatics::SpawnEmitterAtLocation(const UObject* WorldContextObject, UParticleSystem* EmitterTemplate, FVector SpawnLocation, FRotator SpawnRotation, FVector SpawnScale, bool bAutoDestroy, EPSCPoolMethod PoolingMethod, bool bAutoActivate)
 {
 	UParticleSystemComponent* PSC = nullptr;
 	if (EmitterTemplate)
 	{
 		if (UWorld* World = GEngine->GetWorldFromContextObject(WorldContextObject, EGetWorldErrorMode::LogAndReturnNull))
 		{
-			PSC = InternalSpawnEmitterAtLocation(World, EmitterTemplate, SpawnLocation, SpawnRotation, SpawnScale, bAutoDestroy, PoolingMethod);
+			PSC = InternalSpawnEmitterAtLocation(World, EmitterTemplate, SpawnLocation, SpawnRotation, SpawnScale, bAutoDestroy, PoolingMethod, bAutoActivate);
 		}
 	}
 	return PSC;
 }
 
-UParticleSystemComponent* UGameplayStatics::SpawnEmitterAtLocation(UWorld* World, UParticleSystem* EmitterTemplate, const FTransform& SpawnTransform, bool bAutoDestroy, EPSCPoolMethod PoolingMethod)
+UParticleSystemComponent* UGameplayStatics::SpawnEmitterAtLocation(UWorld* World, UParticleSystem* EmitterTemplate, const FTransform& SpawnTransform, bool bAutoDestroy, EPSCPoolMethod PoolingMethod, bool bAutoActivate)
 {
 	UParticleSystemComponent* PSC = nullptr;
 	if (World && EmitterTemplate)
 	{
-		PSC = InternalSpawnEmitterAtLocation(World, EmitterTemplate, SpawnTransform.GetLocation(), SpawnTransform.GetRotation().Rotator(), SpawnTransform.GetScale3D(), bAutoDestroy, PoolingMethod);
+		PSC = InternalSpawnEmitterAtLocation(World, EmitterTemplate, SpawnTransform.GetLocation(), SpawnTransform.GetRotation().Rotator(), SpawnTransform.GetScale3D(), bAutoDestroy, PoolingMethod, bAutoActivate);
 	}
 	return PSC;
 }
 
-UParticleSystemComponent* UGameplayStatics::SpawnEmitterAttached(UParticleSystem* EmitterTemplate, USceneComponent* AttachToComponent, FName AttachPointName, FVector Location, FRotator Rotation, EAttachLocation::Type LocationType, bool bAutoDestroy, EPSCPoolMethod PoolingMethod)
+UParticleSystemComponent* UGameplayStatics::SpawnEmitterAttached(UParticleSystem* EmitterTemplate, USceneComponent* AttachToComponent, FName AttachPointName, FVector Location, FRotator Rotation, EAttachLocation::Type LocationType, bool bAutoDestroy, EPSCPoolMethod PoolingMethod, bool bAutoActivate)
 {
-	return SpawnEmitterAttached(EmitterTemplate, AttachToComponent, AttachPointName, Location, Rotation, FVector(1.f), LocationType, bAutoDestroy, PoolingMethod);
+	return SpawnEmitterAttached(EmitterTemplate, AttachToComponent, AttachPointName, Location, Rotation, FVector(1.f), LocationType, bAutoDestroy, PoolingMethod, bAutoActivate);
 }
 
-UParticleSystemComponent* UGameplayStatics::SpawnEmitterAttached(UParticleSystem* EmitterTemplate, USceneComponent* AttachToComponent, FName AttachPointName, FVector Location, FRotator Rotation, FVector Scale, EAttachLocation::Type LocationType, bool bAutoDestroy, EPSCPoolMethod PoolingMethod)
+UParticleSystemComponent* UGameplayStatics::SpawnEmitterAttached(UParticleSystem* EmitterTemplate, USceneComponent* AttachToComponent, FName AttachPointName, FVector Location, FRotator Rotation, FVector Scale, EAttachLocation::Type LocationType, bool bAutoDestroy, EPSCPoolMethod PoolingMethod, bool bAutoActivate)
 {
 	UParticleSystemComponent* PSC = nullptr;
 	if (EmitterTemplate)
@@ -1064,7 +1068,10 @@ UParticleSystemComponent* UGameplayStatics::SpawnEmitterAttached(UParticleSystem
 					}
 
 					PSC->RegisterComponentWithWorld(World);
-					PSC->ActivateSystem(true);
+					if(bAutoActivate)
+					{
+						PSC->ActivateSystem(true);
+					}
 
 					// Notify the texture streamer so that PSC gets managed as a dynamic component.
 					IStreamingManager::Get().NotifyPrimitiveUpdated(PSC);
