Expose some PrimitiveComponent HLOD properties to blueprints
#rnx

diff --git a/Engine/Source/Runtime/Engine/Classes/Components/PrimitiveComponent.h b/Engine/Source/Runtime/Engine/Classes/Components/PrimitiveComponent.h
--- a/Engine/Source/Runtime/Engine/Classes/Components/PrimitiveComponent.h
+++ b/Engine/Source/Runtime/Engine/Classes/Components/PrimitiveComponent.h
@@ -222,12 +222,12 @@ public:
 	TArray<int32> ExcludeForSpecificHLODLevels;
 
 	/** If true, and if World setting has bEnableHierarchicalLOD equal to true, then this component will be included when generating a Proxy mesh for the parent Actor */
-	UPROPERTY(EditAnywhere, AdvancedDisplay, Category = HLOD, meta = (DisplayName = "Include Component for HLOD Mesh generation"))
+	UPROPERTY(EditAnywhere, AdvancedDisplay, BlueprintReadWrite, Category = HLOD, meta = (DisplayName = "Include Component for HLOD Mesh generation"))
 	uint8 bEnableAutoLODGeneration : 1;
 #endif 
 
 	/** Use the Maximum LOD Mesh (imposter) instead of including Mesh data from this component in the Proxy Generation process */
-	UPROPERTY(EditAnywhere, AdvancedDisplay, Category = HLOD)
+	UPROPERTY(EditAnywhere, AdvancedDisplay, BlueprintReadWrite, Category = HLOD)
 	uint8 bUseMaxLODAsImposter : 1;
 
 	/**
