Now that stream caching ref counts chunk handles, this CL rolls all audio decoders back to how they were previous to the Stream Caching CL, except that an FAudioChunkHandle is retained by the ICompressedAudioInfo instance.

diff --git a/Engine/Source/Runtime/Engine/Private/AudioDecompress.cpp b/Engine/Source/Runtime/Engine/Private/AudioDecompress.cpp
--- a/Engine/Source/Runtime/Engine/Private/AudioDecompress.cpp
+++ b/Engine/Source/Runtime/Engine/Private/AudioDecompress.cpp
@@ -8,13 +8,12 @@
 #include "HAL/LowLevelMemTracker.h"
 #include "ADPCMAudioInfo.h"
 
-PRAGMA_DISABLE_OPTIMIZATION
-
 IStreamedCompressedInfo::IStreamedCompressedInfo()
 	: bIsStreaming(false)
+	,SrcBufferData(nullptr)
+	, SrcBufferDataSize(0)
 	, SrcBufferOffset(0)
 	, AudioDataOffset(0)
-	, AudioDataOffsetChunkIndex(0)
 	, SampleRate(0)
 	, TrueSampleCount(0)
 	, CurrentSampleCount(0)
@@ -27,27 +26,15 @@ IStreamedCompressedInfo::IStreamedCompressedInfo()
 	, CurrentChunkIndex(0)
 	, bPrintChunkFailMessage(true)
 	, SrcBufferPadding(0)
-	, NonStreamingAudio(nullptr)
-	, NonStreamingAudioSize(0)
 {
 }
 
 uint32 IStreamedCompressedInfo::Read(void *OutBuffer, uint32 DataSize)
 {
-	const uint8* SrcBufferPtr = nullptr;
-	uint32 SrcBufferSize = 0;
-
-	GetChunkPtr(SrcBufferPtr, SrcBufferSize);
-
-	uint32 BytesToRead = 0;
-	if (SrcBufferOffset < SrcBufferSize)
-	{
-		BytesToRead = FMath::Min(DataSize, SrcBufferSize - SrcBufferOffset);
-	}
-
+	uint32 BytesToRead = FMath::Min(DataSize, SrcBufferDataSize - SrcBufferOffset);
 	if (BytesToRead > 0)
 	{
-		FMemory::Memcpy(OutBuffer, SrcBufferPtr + SrcBufferOffset, BytesToRead);
+		FMemory::Memcpy(OutBuffer, SrcBufferData + SrcBufferOffset, BytesToRead);
 		SrcBufferOffset += BytesToRead;
 	}
 	return BytesToRead;
@@ -55,28 +42,19 @@ uint32 IStreamedCompressedInfo::Read(void *OutBuffer, uint32 DataSize)
 
 bool IStreamedCompressedInfo::ReadCompressedInfo(const uint8* InSrcBufferData, uint32 InSrcBufferDataSize, FSoundQualityInfo* QualityInfo)
 {
-	check(!SrcBufferHandle.IsValid());
+	check(!SrcBufferData);
 
 	SCOPE_CYCLE_COUNTER(STAT_AudioStreamedDecompressTime);
 
-	// ReadCompressedInfo is called from StreamCompressedInfo(), in which case bIsStreaming will be true.
-	if (!bIsStreaming)
-	{
-		NonStreamingAudio = InSrcBufferData;
-		NonStreamingAudioSize = InSrcBufferDataSize;
-	}
-
-	const uint8* SrcBufferPtr = nullptr;
-	uint32 SrcBufferSize = 0;
-
-	GetChunkPtr(SrcBufferPtr, SrcBufferSize);
-
 	// Parse the format header, this is done different for each format
-	if (!ParseHeader(SrcBufferPtr, SrcBufferSize, QualityInfo))
+	if (!ParseHeader(InSrcBufferData, InSrcBufferDataSize, QualityInfo))
 	{
 		return false;
 	}
 
+	// After parsing the header, the SrcBufferData should be none-null
+	check(SrcBufferData != nullptr);
+
 	// Sample Stride is 
 	SampleStride = NumChannels * sizeof(int16);
 
@@ -97,15 +75,10 @@ bool IStreamedCompressedInfo::ReadCompressedData(uint8* Destination, bool bLoopi
 	bool bFinished = false;
 	uint32 TotalBytesDecoded = 0;
 
-	const uint8* SrcBufferPtr = nullptr;
-	uint32 SrcBufferSize = 0;
-
-	GetChunkPtr(SrcBufferPtr, SrcBufferSize);
-
 	while (TotalBytesDecoded < BufferSize)
 	{
-		const uint8* EncodedSrcPtr = SrcBufferPtr + SrcBufferOffset;
-		const uint32 RemainingEncodedSrcSize = SrcBufferSize - SrcBufferOffset;
+		const uint8* EncodedSrcPtr = SrcBufferData + SrcBufferOffset;
+		const uint32 RemainingEncodedSrcSize = SrcBufferDataSize - SrcBufferOffset;
 
 		const FDecodeResult DecodeResult = Decode(EncodedSrcPtr, RemainingEncodedSrcSize, Destination, BufferSize - TotalBytesDecoded);
 		if (!DecodeResult.NumPcmBytesProduced)
@@ -174,36 +147,14 @@ bool IStreamedCompressedInfo::StreamCompressedInfoInternal(USoundWave* Wave, str
 
 	// Get the first chunk of audio data (should always be loaded)
 	CurrentChunkIndex = 0;
-	bIsStreaming = true;
-
-	const uint8* SrcBufferPtr = nullptr;
-	uint32 SrcBufferSize = 0;
-
-	GetChunkPtr(SrcBufferPtr, SrcBufferSize);
-
-	if (ReadCompressedInfo(SrcBufferPtr, SrcBufferSize, QualityInfo))
-	{
-		while (AudioDataOffset >= SrcBufferSize)
-		{
-			// scrub through chunks until we find the beginning of the audio data.
-			CurrentChunkIndex++;
-			
-			AudioDataOffset = AudioDataOffset - SrcBufferSize;
-			SrcBufferSize = Wave->GetSizeOfChunk(CurrentChunkIndex);
-			// Setting SrcBufferData to null will ensure that we try to load this chunk on the next iteration.
-			SrcBufferHandle = FAudioChunkHandle();
-		}
+	uint32 ChunkSize = 0;
+	const uint8* FirstChunk = GetLoadedChunk(StreamingSoundWave, CurrentChunkIndex, ChunkSize);
 
-		AudioDataOffsetChunkIndex = CurrentChunkIndex;
-		SrcBufferOffset = AudioDataOffset;
-		
-		return true;
-	}
-	else
+	bIsStreaming = true;
+	if (FirstChunk)
 	{
-		UE_LOG(LogAudio, Log, TEXT("Read Compressed Info failed for SoundWave'%s'"), *StreamingSoundWave->GetName());
+		return ReadCompressedInfo(FirstChunk, ChunkSize, QualityInfo);
 	}
-	
 
 	return false;
 }
@@ -215,28 +166,33 @@ bool IStreamedCompressedInfo::StreamCompressedData(uint8* Destination, bool bLoo
 
 	SCOPE_CYCLE_COUNTER(STAT_AudioStreamedDecompressTime);
 
-	UE_LOG(LogAudio, Verbose, TEXT("Streaming compressed data from SoundWave'%s' - Chunk %d, Offset %d"), *StreamingSoundWave->GetName(), CurrentChunkIndex, SrcBufferOffset);
+	UE_LOG(LogAudio, Log, TEXT("Streaming compressed data from SoundWave'%s' - Chunk %d, Offset %d"), *StreamingSoundWave->GetName(), CurrentChunkIndex, SrcBufferOffset);
 
 	// Write out any PCM data that was decoded during the last request
 	uint32 RawPCMOffset = WriteFromDecodedPCM(Destination, BufferSize);
 
-	const uint8* ChunkPtr = nullptr;
-	uint32 ChunkSize = 0;
-
-	GetChunkPtr(ChunkPtr, ChunkSize);
-
 	// If next chunk wasn't loaded when last one finished reading, try to get it again now
-	if (!ChunkPtr)
+	if (SrcBufferData == NULL)
 	{
-		// Still not loaded, zero remainder of current buffer
-		if (bPrintChunkFailMessage)
+		uint32 ChunkSize = 0;
+		SrcBufferData = GetLoadedChunk(StreamingSoundWave, CurrentChunkIndex, ChunkSize);
+		if (SrcBufferData)
 		{
-			UE_LOG(LogAudio, Verbose, TEXT("Chunk %d not loaded from streaming manager for SoundWave '%s'. Likely due to stall on game thread or Disc IO issue."), CurrentChunkIndex, *StreamingSoundWave->GetName());
-			bPrintChunkFailMessage = false;
+			bPrintChunkFailMessage = true;
+			SrcBufferDataSize = ChunkSize;
+			SrcBufferOffset = CurrentChunkIndex == 0 ? AudioDataOffset : 0;
+		}
+		else
+		{
+			// Still not loaded, zero remainder of current buffer
+			if (bPrintChunkFailMessage)
+			{
+				UE_LOG(LogAudio, Verbose, TEXT("Chunk %d not loaded from streaming manager for SoundWave '%s'. Likely due to stall on game thread."), CurrentChunkIndex, *StreamingSoundWave->GetName());
+				bPrintChunkFailMessage = false;
+			}
+			ZeroBuffer(Destination + RawPCMOffset, BufferSize - RawPCMOffset);
+			return false;
 		}
-
-		ZeroBuffer(Destination + RawPCMOffset, BufferSize - RawPCMOffset);
-		return false;
 	}
 
 	bool bLooped = false;
@@ -250,19 +206,6 @@ bool IStreamedCompressedInfo::StreamCompressedData(uint8* Destination, bool bLoo
 
 	while (RawPCMOffset < BufferSize)
 	{
-		GetChunkPtr(ChunkPtr, ChunkSize);
-
-		if (ChunkPtr)
-		{
-			UE_LOG(LogAudio, Verbose, TEXT("Incremented current chunk from SoundWave'%s' - Chunk %d, Offset %d"), *StreamingSoundWave->GetName(), CurrentChunkIndex, SrcBufferOffset);
-		}
-		else
-		{
-			// Zero the rest of the buffer and break when we fail to load the chunk.
-			RawPCMOffset += ZeroBuffer(Destination + RawPCMOffset, BufferSize - RawPCMOffset);
-			break;
-		}
-
 		// Get the platform-dependent size of the current "frame" of encoded audio (note: frame is used here differently than audio frame/sample)
 		uint16 FrameSize = GetFrameSize();
 
@@ -282,10 +225,10 @@ bool IStreamedCompressedInfo::StreamCompressedData(uint8* Destination, bool bLoo
 			RawPCMOffset += WriteFromDecodedPCM(Destination + RawPCMOffset, BufferSize - RawPCMOffset);
 
 			// Have we reached the end of buffer?
-			if (SrcBufferOffset >= (ChunkSize - SrcBufferPadding))
+			if (SrcBufferOffset >= SrcBufferDataSize - SrcBufferPadding)
 			{
 				// Special case for the last chunk of audio
-				if (CurrentChunkIndex == StreamingSoundWave->GetNumChunks() - 1)
+				if (CurrentChunkIndex == StreamingSoundWave->RunningPlatformData->NumChunks - 1)
 				{
 					// check whether all decoded PCM was written
 					if (LastPCMByteSize == 0)
@@ -299,7 +242,7 @@ bool IStreamedCompressedInfo::StreamCompressedData(uint8* Destination, bool bLoo
 
 					if (bLooping)
 					{
-						CurrentChunkIndex = AudioDataOffsetChunkIndex;
+						CurrentChunkIndex = 0;
 						SrcBufferOffset = AudioDataOffset;
 						CurrentSampleCount = 0;
 
@@ -316,6 +259,17 @@ bool IStreamedCompressedInfo::StreamCompressedData(uint8* Destination, bool bLoo
 					CurrentChunkIndex++;
 					SrcBufferOffset = 0;
 				}
+
+				SrcBufferData = GetLoadedChunk(StreamingSoundWave, CurrentChunkIndex, SrcBufferDataSize);
+				if (SrcBufferData)
+				{
+					UE_LOG(LogAudio, Log, TEXT("Incremented current chunk from SoundWave'%s' - Chunk %d, Offset %d"), *StreamingSoundWave->GetName(), CurrentChunkIndex, SrcBufferOffset);
+				}
+				else
+				{
+					SrcBufferDataSize = 0;
+					RawPCMOffset += ZeroBuffer(Destination + RawPCMOffset, BufferSize - RawPCMOffset);
+				}
 			}
 		}
 	}
@@ -325,32 +279,13 @@ bool IStreamedCompressedInfo::StreamCompressedData(uint8* Destination, bool bLoo
 
 int32 IStreamedCompressedInfo::DecompressToPCMBuffer(uint16 FrameSize)
 {
-
-	const uint8* ChunkPtr = nullptr;
-	uint32 ChunkSize = 0;
-
-	GetChunkPtr(ChunkPtr, ChunkSize);
-
-	if (ChunkSize == 0)
-	{
-		// We hit an underrun.
-		return 0;
-	}
-	else if (SrcBufferOffset == ChunkSize)
-	{
-		// We reached the end of the buffer.
-		return 0;
-	}
-	else if (SrcBufferOffset + FrameSize > ChunkSize)
+	if (SrcBufferOffset + FrameSize > SrcBufferDataSize)
 	{
 		// if frame size is too large, something has gone wrong
-		check(false);
 		return -1;
 	}
 
-	check(ChunkPtr && SrcBufferOffset < ChunkSize);
-
-	const uint8* SrcPtr = ChunkPtr + SrcBufferOffset;
+	const uint8* SrcPtr = SrcBufferData + SrcBufferOffset;
 	SrcBufferOffset += FrameSize;
 	LastPCMOffset = 0;
 	
@@ -380,19 +315,10 @@ uint32 IStreamedCompressedInfo::IncrementCurrentSampleCount(uint32 NewSamples)
 
 uint32 IStreamedCompressedInfo::WriteFromDecodedPCM(uint8* Destination, uint32 BufferSize)
 {
-	if (LastPCMOffset >= LastPCMByteSize)
-	{
-		LastPCMOffset = 0;
-		LastPCMByteSize = 0;
-
-		return 0;
-	}
-
-	uint32 BytesToCopy = FMath::Min<uint32>(BufferSize, LastPCMByteSize - LastPCMOffset);
-
+	uint32 BytesToCopy = FMath::Min(BufferSize, LastPCMByteSize - LastPCMOffset);
 	if (BytesToCopy > 0)
 	{
-		checkf(BytesToCopy <= LastDecodedPCM.Num() - LastPCMOffset, TEXT("Wraparound detected in WriteFromDecodedPCM. Current PCM Index %d was greater than PCM size %d, resulting in an invalid BytesToCopy value of %d."), LastPCMOffset, LastDecodedPCM.Num(), BytesToCopy);
+		check(BytesToCopy <= LastDecodedPCM.Num() - LastPCMOffset);
 		FMemory::Memcpy(Destination, LastDecodedPCM.GetData() + LastPCMOffset, BytesToCopy);
 		LastPCMOffset += BytesToCopy;
 		if (LastPCMOffset >= LastPCMByteSize)
@@ -417,24 +343,25 @@ uint32 IStreamedCompressedInfo::ZeroBuffer(uint8* Destination, uint32 BufferSize
 }
 
 
-void IStreamedCompressedInfo::GetChunkPtr(const uint8*& OutPtr, uint32& OutSize)
+const uint8* IStreamedCompressedInfo::GetLoadedChunk(USoundWave* InSoundWave, uint32 ChunkIndex, uint32& OutChunkSize)
 {
-	if (!bIsStreaming)
+	if (!InSoundWave || ChunkIndex >= InSoundWave->GetNumChunks())
 	{
-		OutPtr = NonStreamingAudio;
-		OutSize = NonStreamingAudioSize;
+		UE_LOG(LogAudio, Error, TEXT("Error calling GetLoadedChunk for ChunkIndex %d!"), ChunkIndex);
+		OutChunkSize = 0;
+		return nullptr;
 	}
-	else if (CurrentChunkIndex == 0)
+	else if (ChunkIndex == 0)
 	{
-		TArrayView<const uint8> ZerothChunk = StreamingSoundWave->GetZerothChunk();
-		OutPtr = ZerothChunk.GetData();
-		OutSize = ZerothChunk.Num();
+		TArrayView<const uint8> ZerothChunk = InSoundWave->GetZerothChunk();
+		OutChunkSize = ZerothChunk.Num();
+		return ZerothChunk.GetData();
 	}
 	else
 	{
-		SrcBufferHandle = IStreamingManager::Get().GetAudioStreamingManager().GetLoadedChunk(StreamingSoundWave, CurrentChunkIndex);
-		OutPtr = SrcBufferHandle.GetData();
-		OutSize = SrcBufferHandle.Num();
+		CurCompressedChunkHandle = IStreamingManager::Get().GetAudioStreamingManager().GetLoadedChunk(InSoundWave, ChunkIndex);
+		OutChunkSize = CurCompressedChunkHandle.Num();
+		return CurCompressedChunkHandle.GetData();
 	}
 }
 
diff --git a/Engine/Source/Runtime/Engine/Private/OpusAudioInfo.cpp b/Engine/Source/Runtime/Engine/Private/OpusAudioInfo.cpp
--- a/Engine/Source/Runtime/Engine/Private/OpusAudioInfo.cpp
+++ b/Engine/Source/Runtime/Engine/Private/OpusAudioInfo.cpp
@@ -113,6 +113,8 @@ FOpusAudioInfo::~FOpusAudioInfo()
 
 bool FOpusAudioInfo::ParseHeader(const uint8* InSrcBufferData, uint32 InSrcBufferDataSize, struct FSoundQualityInfo* QualityInfo)
 {
+	SrcBufferData = InSrcBufferData;
+	SrcBufferDataSize = InSrcBufferDataSize;
 	SrcBufferOffset = 0;
 	CurrentSampleCount = 0;
 
diff --git a/Engine/Source/Runtime/Engine/Private/VorbisAudioInfo.cpp b/Engine/Source/Runtime/Engine/Private/VorbisAudioInfo.cpp
--- a/Engine/Source/Runtime/Engine/Private/VorbisAudioInfo.cpp
+++ b/Engine/Source/Runtime/Engine/Private/VorbisAudioInfo.cpp
@@ -98,8 +98,10 @@ FVorbisAudioInfo::FVorbisAudioInfo()
 	, SrcBufferDataSize(0)
 	, BufferOffset(0)
 	, CurrentBufferChunkOffset(0)
+	, CurrentStreamingChunkData(nullptr)
 	, CurrentStreamingChunkIndex(INDEX_NONE)
 	, NextStreamingChunkIndex(0)
+	, CurrentStreamingChunksSize(0)
 	, bHeaderParsed(false)
 {
 	// Make sure we have properly allocated a VFWrapper
@@ -192,16 +194,16 @@ size_t FVorbisAudioInfo::ReadStreaming(void *Ptr, uint32 Size )
 
 	while (NumBytesRead < Size)
 	{
-		const uint8* CurrentStreamingChunkData = nullptr;
-		uint32 CurrentStreamingChunksSize = 0;
-
-		if (CurrentStreamingChunkIndex != NextStreamingChunkIndex)
+		if (!CurrentStreamingChunkData || CurrentStreamingChunkIndex != NextStreamingChunkIndex)
 		{
 			CurrentStreamingChunkIndex = NextStreamingChunkIndex;
-			CurrentBufferChunkOffset = 0;
+			CurrentStreamingChunkData = GetLoadedChunk(StreamingSoundWave, CurrentStreamingChunkIndex, CurrentStreamingChunksSize);
+			if (CurrentStreamingChunkData)
+			{
+				check(CurrentStreamingChunkIndex < StreamingSoundWave->RunningPlatformData->Chunks.Num());
+				CurrentBufferChunkOffset = 0;
+			}
 		}
-		
-		GetChunkPtr(CurrentStreamingChunkData, CurrentStreamingChunksSize);
 
 		// No chunk data -- either looping or something else happened with stream
 		if (!CurrentStreamingChunkData)
@@ -540,8 +542,6 @@ bool FVorbisAudioInfo::StreamCompressedInfoInternal(USoundWave* Wave, struct FSo
 
 	FScopeLock ScopeLock(&VorbisCriticalSection);
 
-	TotalNumChunks = Wave->GetNumChunks();
-
 	if (!VFWrapper)
 	{
 		UE_LOG(LogAudio, Error, TEXT("FVorbisAudioInfo::StreamCompressedInfoInternal failed due to no vorbis wrapper for sound '%s'."), *Wave->GetName());
@@ -570,26 +570,25 @@ bool FVorbisAudioInfo::StreamCompressedInfoInternal(USoundWave* Wave, struct FSo
 	return bHeaderParsed;
 }
 
-void FVorbisAudioInfo::GetChunkPtr(const uint8*& OutPtr, uint32& OutSize)
+const uint8* FVorbisAudioInfo::GetLoadedChunk(USoundWave* InSoundWave, uint32 ChunkIndex, uint32& OutChunkSize)
 {
-	check(StreamingSoundWave);
-
-	if (CurrentStreamingChunkIndex == 0)
+	if (!InSoundWave || ChunkIndex >= InSoundWave->GetNumChunks())
 	{
-		TArrayView<const uint8> ZerothChunk = StreamingSoundWave->GetZerothChunk();
-		OutPtr = ZerothChunk.GetData();
-		OutSize = ZerothChunk.Num();
+		UE_LOG(LogAudio, Error, TEXT("Error calling GetLoadedChunk for ChunkIndex %d!"), ChunkIndex);
+		OutChunkSize = 0;
+		return nullptr;
 	}
-	else if (((uint32)CurrentStreamingChunkIndex) < StreamingSoundWave->GetNumChunks())
+	else if (ChunkIndex == 0)
 	{
-		CurrentStreamingChunkHandle = IStreamingManager::Get().GetAudioStreamingManager().GetLoadedChunk(StreamingSoundWave, CurrentStreamingChunkIndex);
-		OutPtr = CurrentStreamingChunkHandle.GetData();
-		OutSize = CurrentStreamingChunkHandle.Num();
+		TArrayView<const uint8> ZerothChunk = InSoundWave->GetZerothChunk();
+		OutChunkSize = ZerothChunk.Num();
+		return ZerothChunk.GetData();
 	}
 	else
 	{
-		OutPtr = nullptr;
-		OutSize = 0;
+		CurCompressedChunkHandle = IStreamingManager::Get().GetAudioStreamingManager().GetLoadedChunk(InSoundWave, ChunkIndex);
+		OutChunkSize = CurCompressedChunkHandle.Num();
+		return CurCompressedChunkHandle.GetData();
 	}
 }
 
@@ -626,10 +625,11 @@ bool FVorbisAudioInfo::StreamCompressedData(uint8* InDestination, bool bLooping,
 		if( BytesActuallyRead <= 0 )
 		{
 			// if we read 0 bytes or there was an error, instead of assuming we looped, lets write out zero's.
-			// this means that the chunk wasn't loaded in time
-			const bool bHitEOF = ((uint32) CurrentStreamingChunkIndex) >= StreamingSoundWave->GetNumChunks();
+			// this means that the chunk wasn't loaded in time.
+
+			check(StreamingSoundWave->GetNumChunks() < ((uint32)TNumericLimits<int32>::Max()));
 
-			if (!bHitEOF)
+			if (NextStreamingChunkIndex < static_cast<int32>(StreamingSoundWave->GetNumChunks()))
 			{
 				// zero out the rest of the buffer
 				FMemory::Memzero(InDestination, BufferSize);
diff --git a/Engine/Source/Runtime/Engine/Public/AudioDecompress.h b/Engine/Source/Runtime/Engine/Public/AudioDecompress.h
--- a/Engine/Source/Runtime/Engine/Public/AudioDecompress.h
+++ b/Engine/Source/Runtime/Engine/Public/AudioDecompress.h
@@ -179,7 +179,7 @@ public:
 	ENGINE_API virtual void SeekToTime(const float SeekTime) override {};
 	ENGINE_API virtual void ExpandFile(uint8* DstBuffer, struct FSoundQualityInfo* QualityInfo) override;
 	ENGINE_API virtual void EnableHalfRate(bool HalfRate) override {};
-	virtual uint32 GetSourceBufferSize() const override { return SrcBufferHandle.Num(); }
+	virtual uint32 GetSourceBufferSize() const override { return SrcBufferDataSize; }
 	virtual bool UsesVorbisChannelOrdering() const override { return false; }
 	virtual int GetStreamBufferSize() const override { return  MONO_PCM_BUFFER_SIZE; }
 	virtual bool SupportsStreaming() const override { return true; }
@@ -246,20 +246,25 @@ protected:
 	*/
 	uint32	ZeroBuffer(uint8* Destination, uint32 BufferSize);
 
-	// this returns a pointer to the current chunk and the size of the audio in streaming situations,
-	// or a pointer to the entire audio buffer for non-streaming situations.
-	void GetChunkPtr(const uint8*& OutPtr, uint32& OutSize);
+	/**
+	 * Helper function for getting a chunk of compressed audio.
+	 * @param InSoundWave Pointer to the soundwave to get compressed audio from.
+	 * @param ChunkIndex the index of the chunk to get from InSoundWave.
+	 * @param[out] OutChunkSize the size of the chunk.
+	 * @return a pointer to the chunk if it's loaded, nullptr otherwise.
+	 */
+	const uint8* GetLoadedChunk(USoundWave* InSoundWave, uint32 ChunkIndex, uint32& OutChunkSize);
 
 	/** bool set before ParseHeader. Whether we are streaming a file or not. */
 	bool bIsStreaming;
 	/** Ptr to the current streamed chunk. */
-	FAudioChunkHandle SrcBufferHandle;
+	const uint8* SrcBufferData;
+	/** Size of the current streamed chunk. */
+	uint32 SrcBufferDataSize;
 	/** What byte we're currently reading in the streamed chunk. */
 	uint32 SrcBufferOffset;
 	/** Where the actual audio data starts in the current streamed chunk. Accounts for header offset. */
 	uint32 AudioDataOffset;
-	/** The chunk index where the audio data begins. */
-	uint32 AudioDataOffsetChunkIndex;
 	/** Sample rate of the source file */
 	uint16 SampleRate;
 	/** The total sample count of the source file. */
@@ -286,9 +291,8 @@ protected:
 	bool bPrintChunkFailMessage;
 	/** Number of bytes of padding used, overridden in some implementations. Defaults to 0. */
 	uint32 SrcBufferPadding;
-
-	const uint8* NonStreamingAudio;
-	uint32 NonStreamingAudioSize;
+	/** Chunk Handle to ensure that this chunk of streamed audio is not deleted while we are using it. */
+	FAudioChunkHandle CurCompressedChunkHandle;
 };
 
 
diff --git a/Engine/Source/Runtime/Engine/Public/VorbisAudioInfo.h b/Engine/Source/Runtime/Engine/Public/VorbisAudioInfo.h
--- a/Engine/Source/Runtime/Engine/Public/VorbisAudioInfo.h
+++ b/Engine/Source/Runtime/Engine/Public/VorbisAudioInfo.h
@@ -9,7 +9,6 @@
 #include "CoreMinimal.h"
 #include "HAL/ThreadSafeBool.h"
 #include "AudioDecompress.h"
-#include "ContentStreaming.h"
 
 /**
  * Whether to use OggVorbis audio format.
@@ -47,8 +46,7 @@ public:
 
 	/** Emulate read from streaming functionality */
 	size_t			ReadStreaming( void *ptr, uint32 size );
-
-	int				CloseStreaming(void);
+	int				CloseStreaming( void );
 
 	/** Common info and data functions between ReadCompressedInfo/ReadCompressedData and StreamCompressedInfo/StreamCompressedData */
 	/** Can't use ov_callbacks here so we have to pass in a void* for the callbacks instead */
@@ -95,27 +93,14 @@ public:
 	virtual bool SupportsStreaming() const override {return true;}
 	virtual bool StreamCompressedData(uint8* InDestination, bool bLooping, uint32 BufferSize) override;
 	virtual int32 GetCurrentChunkIndex() const override {return CurrentStreamingChunkIndex;}
-
-	virtual int32 GetCurrentChunkOffset() const override
-	{
-		const uint32 ChunkSize = StreamingSoundWave ? StreamingSoundWave->GetSizeOfChunk(CurrentStreamingChunkIndex) : 0;
-		if (ChunkSize)
-		{
-			return CurrentBufferChunkOffset % ChunkSize;
-		}
-		else
-		{
-			return 0;
-		}
-	}
+	virtual int32 GetCurrentChunkOffset() const override {return BufferOffset % CurrentStreamingChunksSize;}
 	// End of ICompressedAudioInfo Interface
 
 protected:
 	virtual bool StreamCompressedInfoInternal(USoundWave* Wave, struct FSoundQualityInfo* QualityInfo) override;
 
 private:
-
-	void GetChunkPtr(const uint8*& OutPtr, uint32& OutSize);
+	const uint8* GetLoadedChunk(USoundWave* InSoundWave, uint32 ChunkIndex, uint32& OutChunkSize);
 
 	struct FVorbisFileWrapper* VFWrapper;
 	const uint8* SrcBufferData;
@@ -126,10 +111,13 @@ private:
 	/** Critical section used to prevent multiple threads accessing same ogg-vorbis file handles at the same time */
 	FCriticalSection VorbisCriticalSection;
 
-	FAudioChunkHandle CurrentStreamingChunkHandle;
+	uint8 const* CurrentStreamingChunkData;
 	int32 CurrentStreamingChunkIndex;
-	int32 TotalNumChunks;
 	int32 NextStreamingChunkIndex;
+	uint32 CurrentStreamingChunksSize;
 	bool bHeaderParsed;
+
+	// This handle is used to ensure that a chunk currently being decoded isn't evicted until we are done with it.
+	FAudioChunkHandle CurCompressedChunkHandle;
 };
 #endif
