Guard against crash when blending physics bones.

diff --git a/Engine/Source/Runtime/Engine/Private/PhysicsEngine/PhysAnim.cpp b/Engine/Source/Runtime/Engine/Private/PhysicsEngine/PhysAnim.cpp
--- a/Engine/Source/Runtime/Engine/Private/PhysicsEngine/PhysAnim.cpp
+++ b/Engine/Source/Runtime/Engine/Private/PhysicsEngine/PhysAnim.cpp
@@ -243,6 +243,11 @@ void USkeletalMeshComponent::PerformBlendPhysicsBones(const TArray<FBoneIndexTyp
 					// if we wan't 'full weight' we just find 
 					if(UsePhysWeight > 0.f)
 					{
+						if (!(ensure(InBoneSpaceTransforms.Num())))
+						{
+							continue;
+						}
+
 						if(BoneIndex == 0)
 						{
 							ParentWorldTM = LocalToWorldTM;
