Avoid crashes on opening a partially-deleted Blueprint asset after a failed attempt to force delete it.

diff --git a/Engine/Source/Runtime/Engine/Classes/Engine/SimpleConstructionScript.h b/Engine/Source/Runtime/Engine/Classes/Engine/SimpleConstructionScript.h
--- a/Engine/Source/Runtime/Engine/Classes/Engine/SimpleConstructionScript.h
+++ b/Engine/Source/Runtime/Engine/Classes/Engine/SimpleConstructionScript.h
@@ -61,6 +61,9 @@ class USimpleConstructionScript : public UObject
 	/** Return the Blueprint associated with this SCS instance */
 	ENGINE_API class UClass* GetOwnerClass() const;
 
+	/** Return the parent class of the SCS owner class */
+	ENGINE_API class UClass* GetParentClass() const;
+
 #if WITH_EDITOR
 	/** Return all nodes in tree as a flat list */
 	ENGINE_API const TArray<USCS_Node*>& GetAllNodes() const;
diff --git a/Engine/Source/Runtime/Engine/Private/SimpleConstructionScript.cpp b/Engine/Source/Runtime/Engine/Private/SimpleConstructionScript.cpp
--- a/Engine/Source/Runtime/Engine/Private/SimpleConstructionScript.cpp
+++ b/Engine/Source/Runtime/Engine/Private/SimpleConstructionScript.cpp
@@ -727,6 +727,21 @@ UClass* USimpleConstructionScript::GetOwnerClass() const
 	return NULL;
 }
 
+UClass* USimpleConstructionScript::GetParentClass() const
+{
+	UClass* ParentClass = nullptr;
+	if (UBlueprint* Blueprint = GetBlueprint())
+	{
+		ParentClass = Blueprint->ParentClass;
+	}
+	else if(UClass* OwnerClass = GetOwnerClass())
+	{
+		ParentClass = OwnerClass->GetSuperClass();
+	}
+
+	return ParentClass;
+}
+
 #if WITH_EDITOR
 const TArray<USCS_Node*>& USimpleConstructionScript::GetAllNodes() const
 {
@@ -942,9 +957,8 @@ USCS_Node* USimpleConstructionScript::FindSCSNodeByGuid(const FGuid Guid) const
 #if WITH_EDITOR
 USceneComponent* USimpleConstructionScript::GetSceneRootComponentTemplate(USCS_Node** OutSCSNode) const
 {
-	UBlueprint* Blueprint = GetBlueprint();
-
 	UClass* GeneratedClass = GetOwnerClass();
+	UClass* ParentClass = GetParentClass();
 
 	if(OutSCSNode)
 	{
@@ -959,9 +973,9 @@ USceneComponent* USimpleConstructionScript::GetSceneRootComponentTemplate(USCS_N
 	}
 
 	// If the generated class does not yet have a CDO, defer to the parent class
-	if(CDO == nullptr && Blueprint->ParentClass != nullptr)
+	if(CDO == nullptr && ParentClass != nullptr)
 	{
-		CDO = Cast<AActor>(Blueprint->ParentClass->GetDefaultObject(false));
+		CDO = Cast<AActor>(ParentClass->GetDefaultObject(false));
 	}
 
 	// Check to see if we already have a native root component template
@@ -988,13 +1002,13 @@ USceneComponent* USimpleConstructionScript::GetSceneRootComponentTemplate(USCS_N
 	{
 		// Get the Blueprint hierarchy
 		TArray<UBlueprint*> BPStack;
-		if(Blueprint->GeneratedClass != nullptr)
+		if(GeneratedClass != nullptr)
 		{
-			UBlueprint::GetBlueprintHierarchyFromClass(Blueprint->GeneratedClass, BPStack);
+			UBlueprint::GetBlueprintHierarchyFromClass(GeneratedClass, BPStack);
 		}
-		else if(Blueprint->ParentClass != nullptr)
+		else if(ParentClass != nullptr)
 		{
-			UBlueprint::GetBlueprintHierarchyFromClass(Blueprint->ParentClass, BPStack);
+			UBlueprint::GetBlueprintHierarchyFromClass(ParentClass, BPStack);
 		}
 
 		// Note: Normally if the Blueprint has a parent, we can assume that the parent already has a scene root component set,
@@ -1344,7 +1358,10 @@ USCS_Node* USimpleConstructionScript::CreateNodeAndRenameComponent(UActorCompone
 void USimpleConstructionScript::ValidateNodeVariableNames(FCompilerResultsLog& MessageLog)
 {
 	UBlueprint* Blueprint = GetBlueprint();
-	check(Blueprint);
+	if (!ensureMsgf(Blueprint != nullptr, TEXT("Cannot validate SCS node variable names because the owning Blueprint could not be determined from the SCS context (perhaps it was deleted?).")))
+	{
+		return;
+	}
 
 	TSharedPtr<FKismetNameValidator> ParentBPNameValidator;
 	if( Blueprint->ParentClass != NULL )
