Allow Blueprints that own one or more child redirector objects to be deleted.

diff --git a/Engine/Source/Editor/UnrealEd/Private/ObjectTools.cpp b/Engine/Source/Editor/UnrealEd/Private/ObjectTools.cpp
--- a/Engine/Source/Editor/UnrealEd/Private/ObjectTools.cpp
+++ b/Engine/Source/Editor/UnrealEd/Private/ObjectTools.cpp
@@ -2662,6 +2662,7 @@ namespace ObjectTools
 								}
 							}
 
+							BlueprintObject->RemoveChildRedirectors();
 							BlueprintObject->RemoveGeneratedClasses();
 						}
 					}
diff --git a/Engine/Source/Runtime/Engine/Classes/Engine/Blueprint.h b/Engine/Source/Runtime/Engine/Classes/Engine/Blueprint.h
--- a/Engine/Source/Runtime/Engine/Classes/Engine/Blueprint.h
+++ b/Engine/Source/Runtime/Engine/Classes/Engine/Blueprint.h
@@ -822,6 +822,9 @@ public:
 	virtual void BeginDestroy() override;
 	//~ End UObject Interface
 
+	/** Removes any child redirectors from the root set and marks them as transient */
+	void RemoveChildRedirectors();
+
 	/** Consigns the GeneratedClass and the SkeletonGeneratedClass to oblivion, and nulls their references */
 	void RemoveGeneratedClasses();
 
diff --git a/Engine/Source/Runtime/Engine/Private/Blueprint.cpp b/Engine/Source/Runtime/Engine/Private/Blueprint.cpp
--- a/Engine/Source/Runtime/Engine/Private/Blueprint.cpp
+++ b/Engine/Source/Runtime/Engine/Private/Blueprint.cpp
@@ -643,6 +643,21 @@ UClass* UBlueprint::RegenerateClass(UClass* ClassToRegenerate, UObject* Previous
 	}
 }
 
+void UBlueprint::RemoveChildRedirectors()
+{
+	TArray<UObject*> ChildObjects;
+	GetObjectsWithOuter(this, ChildObjects);
+	for (UObject* ChildObject : ChildObjects)
+	{
+		if (ChildObject->IsA<UObjectRedirector>())
+		{
+			ChildObject->ClearFlags(RF_Public|RF_Standalone);
+			ChildObject->SetFlags(RF_Transient);
+			ChildObject->RemoveFromRoot();
+		}
+	}
+}
+
 void UBlueprint::RemoveGeneratedClasses()
 {
 	FBlueprintEditorUtils::RemoveGeneratedClasses(this);
@@ -958,7 +973,14 @@ void UBlueprint::GetAssetRegistryTags(TArray<FAssetRegistryTag>& OutTags) const
 	// Only add the FiB tags in the editor, this now gets run for standalone uncooked games
 	if ( ParentClass && GIsEditor)
 	{
-		OutTags.Add( FAssetRegistryTag(FBlueprintTags::FindInBlueprintsData, FFindInBlueprintSearchManager::Get().QuerySingleBlueprint((UBlueprint*)this, false), FAssetRegistryTag::TT_Hidden) );
+		FString Value;
+		const bool bRebuildSearchData = false;
+		if (const FSearchData* SearchData = FFindInBlueprintSearchManager::Get().QuerySingleBlueprint((UBlueprint*)this, bRebuildSearchData))
+		{
+			Value = SearchData->Value;
+		}
+		
+		OutTags.Add( FAssetRegistryTag(FBlueprintTags::FindInBlueprintsData, Value, FAssetRegistryTag::TT_Hidden) );
 	}
 
 	// Only show for strict blueprints (not animation or widget blueprints)
