CIS fix.

#rnx

diff --git a/Engine/Source/Runtime/Engine/Private/Blueprint.cpp b/Engine/Source/Runtime/Engine/Private/Blueprint.cpp
--- a/Engine/Source/Runtime/Engine/Private/Blueprint.cpp
+++ b/Engine/Source/Runtime/Engine/Private/Blueprint.cpp
@@ -973,14 +973,7 @@ void UBlueprint::GetAssetRegistryTags(TArray<FAssetRegistryTag>& OutTags) const
 	// Only add the FiB tags in the editor, this now gets run for standalone uncooked games
 	if ( ParentClass && GIsEditor)
 	{
-		FString Value;
-		const bool bRebuildSearchData = false;
-		if (const FSearchData* SearchData = FFindInBlueprintSearchManager::Get().QuerySingleBlueprint((UBlueprint*)this, bRebuildSearchData))
-		{
-			Value = SearchData->Value;
-		}
-		
-		OutTags.Add( FAssetRegistryTag(FBlueprintTags::FindInBlueprintsData, Value, FAssetRegistryTag::TT_Hidden) );
+		OutTags.Add(FAssetRegistryTag(FBlueprintTags::FindInBlueprintsData, FFindInBlueprintSearchManager::Get().QuerySingleBlueprint((UBlueprint*)this, false), FAssetRegistryTag::TT_Hidden));
 	}
 
 	// Only show for strict blueprints (not animation or widget blueprints)
