Fix crash when reparenting a blueprint component from scene to actor while an actor blueprint containing that component is opened.
Scene component hierarchy when switching from actor to scene component is now more consistent with adding a component from the UI.
Fix with editor issue from CL# 7459539
Cleared out autos.

diff --git a/Engine/Source/Editor/UnrealEd/Private/Kismet2/KismetReinstanceUtilities.cpp b/Engine/Source/Editor/UnrealEd/Private/Kismet2/KismetReinstanceUtilities.cpp
--- a/Engine/Source/Editor/UnrealEd/Private/Kismet2/KismetReinstanceUtilities.cpp
+++ b/Engine/Source/Editor/UnrealEd/Private/Kismet2/KismetReinstanceUtilities.cpp
@@ -16,6 +16,7 @@
 #include "Editor/EditorEngine.h"
 #include "Animation/AnimBlueprint.h"
 #include "Engine/BlueprintGeneratedClass.h"
+#include "Engine/SimpleConstructionScript.h"
 #include "FileHelpers.h"
 #include "Kismet2/CompilerResultsLog.h"
 #include "Kismet2/KismetEditorUtilities.h"
@@ -1713,12 +1714,7 @@ static void ReplaceObjectHelper(UObject*& OldObject, UClass* OldClass, UObject*&
 			// we need to keep track of actor instances that need 
 			// their construction scripts re-ran (since we've just 
 			// replaced a component they own)
-
-			// Skipping CDOs as CSs are not allowed for them.
-			if (!OwningActor->HasAnyFlags(RF_ClassDefaultObject))
-			{
-				OwnersToRerunConstructionScript.Add(OwningActor);
-			}
+			OwnersToRerunConstructionScript.Add(OwningActor);
 		}
 
 		if (bWasRegistered)
@@ -1874,6 +1870,7 @@ void FBlueprintCompileReinstancer::ReplaceInstancesOfClass_Inner(TMap<UClass*, U
 	USelection* SelectedActors = nullptr;
 	TArray<UObject*> ObjectsReplaced;
 	bool bSelectionChanged = false;
+	bool bFixupSCS = false;
 	const bool bLogConversions = false; // for debugging
 
 	// Map of old objects to new objects
@@ -1943,7 +1940,7 @@ void FBlueprintCompileReinstancer::ReplaceInstancesOfClass_Inner(TMap<UClass*, U
 		}
 
 		// WARNING: for (TPair<UClass*, UClass*> OldToNewClass : InOldToNewClassMap) duplicated below 
-		// to handle reconstructing actors which need to be reinstanced aftertheir owned components 
+		// to handle reconstructing actors which need to be reinstanced after their owned components 
 		// have been updated:
 		for (TPair<UClass*, UClass*> OldToNewClass : InOldToNewClassMap)
 		{
@@ -1956,6 +1953,14 @@ void FBlueprintCompileReinstancer::ReplaceInstancesOfClass_Inner(TMap<UClass*, U
 			check(OldClass != NewClass);
 #endif
 			{
+				const bool bIsComponent = NewClass->IsChildOf<UActorComponent>();
+
+				// If any of the class changes are of an actor component to scene component or reverse then we will fixup SCS of all actors affected
+				if (bIsComponent && !bFixupSCS)
+				{
+					bFixupSCS = (NewClass->IsChildOf<USceneComponent>() != OldClass->IsChildOf<USceneComponent>());
+				}
+
 				const bool bIncludeDerivedClasses = false;
 				ObjectsToReplace.Reset();
 				GetObjectsOfClass(OldClass, ObjectsToReplace, bIncludeDerivedClasses);
@@ -1967,7 +1972,6 @@ void FBlueprintCompileReinstancer::ReplaceInstancesOfClass_Inner(TMap<UClass*, U
 					AActor* OldActor = Cast<AActor>(OldObject);
 
 					// Skip archetype instances, EXCEPT for component templates and child actor templates
-					const bool bIsComponent = NewClass->IsChildOf(UActorComponent::StaticClass());
 					const bool bIsChildActorTemplate = OldActor && OldActor->GetOuter()->IsA<UChildActorComponent>();
 					if (OldObject->IsPendingKill() || 
 						(!bIsComponent && !bIsChildActorTemplate && OldObject->IsTemplate()) ||
@@ -2158,6 +2162,35 @@ void FBlueprintCompileReinstancer::ReplaceInstancesOfClass_Inner(TMap<UClass*, U
 		GEditor->NoteSelectionChange();
 	}
 
+	TSet<UBlueprintGeneratedClass*> FixedSCS;
+
+	// in the case where we're replacing component instances, we need to make 
+	// sure to re-run their owner's construction scripts
+	for (AActor* ActorInstance : OwnersToRerunConstructionScript)
+	{
+		// Before rerunning the construction script, first fix up the SCS if any component class has changed from actor to scene
+		if (bFixupSCS)
+		{
+			UBlueprintGeneratedClass* BPGC = Cast<UBlueprintGeneratedClass>(ActorInstance->GetClass());
+			while (BPGC && !FixedSCS.Contains(BPGC))
+			{
+				if (BPGC->SimpleConstructionScript)
+				{
+					BPGC->SimpleConstructionScript->FixupRootNodeParentReferences();
+					BPGC->SimpleConstructionScript->ValidateSceneRootNodes();
+				}
+				FixedSCS.Add(BPGC);
+				BPGC = Cast<UBlueprintGeneratedClass>(BPGC->GetSuperClass());
+			}
+		}
+
+		// Skipping CDOs as CSs are not allowed for them.
+		if (!ActorInstance->HasAnyFlags(RF_ClassDefaultObject))
+		{
+			ActorInstance->RerunConstructionScripts();
+		}
+	}
+
 	if (GEditor)
 	{
 		// Refresh any editors for objects that we've updated components for
@@ -2170,13 +2203,6 @@ void FBlueprintCompileReinstancer::ReplaceInstancesOfClass_Inner(TMap<UClass*, U
 			}
 		}
 	}
-
-	// in the case where we're replacing component instances, we need to make 
-	// sure to re-run their owner's construction scripts
-	for (AActor* ActorInstance : OwnersToRerunConstructionScript)
-	{
-		ActorInstance->RerunConstructionScripts();
-	}
 }
 
 void FBlueprintCompileReinstancer::ReconstructOwnerInstances(TSubclassOf<UActorComponent> ComponentClass)
diff --git a/Engine/Source/Runtime/Engine/Private/SimpleConstructionScript.cpp b/Engine/Source/Runtime/Engine/Private/SimpleConstructionScript.cpp
--- a/Engine/Source/Runtime/Engine/Private/SimpleConstructionScript.cpp
+++ b/Engine/Source/Runtime/Engine/Private/SimpleConstructionScript.cpp
@@ -302,6 +302,11 @@ void USimpleConstructionScript::FixupSceneNodeHierarchy()
 		{
 			SceneRootNode = DefaultSceneRootNode;
 			SceneRootComponentTemplate = CastChecked<USceneComponent>(DefaultSceneRootNode->ComponentTemplate);
+			if (!RootNodes.Contains(SceneRootNode))
+			{
+				RootNodes.Add(SceneRootNode);
+				AllNodes.Add(SceneRootNode);
+			}
 		}
 		// if there is no scene root (then there shouldn't be anything but the 
 		// default placeholder root).
@@ -414,27 +419,46 @@ void USimpleConstructionScript::FixupSceneNodeHierarchy()
 			}
 			else
 			{
-				UClass* ComponentClass = Node->ComponentClass ? Node->ComponentClass :
-					Node->ComponentTemplate ? Node->ComponentTemplate->GetClass() : nullptr;
-				// we don't care about non-scene nodes
-				if (ComponentClass && ComponentClass->IsChildOf<USceneComponent>())
+				auto VisitChildren = [this, Node]()
 				{
-					// scoped for the following TGuardValue
+					// recursively visit children so we can construct the hierarchy - iterate backwards so we can remove as we go
+					for (int32 ChildIndex = Node->ChildNodes.Num() - 1; ChildIndex >= 0; --ChildIndex)
 					{
-						TGuardValue<USCS_Node*> ParentStack(PendingParent, Node);
-						// recursively visit children so we can construct the hierarchy - iterate backwards so we can remove as we go
-						for (int32 ChildIndex = Node->ChildNodes.Num() - 1; ChildIndex >= 0; --ChildIndex)
+						if (!VisitNode(Node->ChildNodes[ChildIndex]))
 						{
-							if (!VisitNode(Node->ChildNodes[ChildIndex]))
-							{
-								Node->ChildNodes.RemoveAt(ChildIndex);
-							}
+							Node->ChildNodes.RemoveAt(ChildIndex);
 						}
 					}
+				};
 
-					// happens after recursing into children, so we don't add to 
-					// the orphaned list till after children are querying it
-					FixupParentage(Node);
+				if (UClass* ComponentClass = (Node->ComponentClass ? Node->ComponentClass : (Node->ComponentTemplate ? Node->ComponentTemplate->GetClass() : nullptr)))
+				{
+					if (ComponentClass->IsChildOf<USceneComponent>())
+					{
+						// scoped for the following TGuardValue
+						{
+							TGuardValue<USCS_Node*> ParentStack(PendingParent, Node);
+							VisitChildren();
+						}
+
+						// happens after recursing into children, so we don't add to 
+						// the orphaned list till after children are querying it
+						FixupParentage(Node);
+					}
+					else
+					{
+						RootNodeList.AddUnique(Node);
+						if (Node->ChildNodes.Num() > 0)
+						{
+							// If this isn't a scene component but it has children, that's not good, so shift them to the pending parent or make them orphan nodes
+							VisitChildren();
+
+							Node->ChildNodes.Reset();
+						}
+
+						// A non-scene component should never be in the child list of someone else, so return false so the parent removes it from its list
+						return false;
+					}
 				}
 			}
 			return true;
@@ -609,9 +633,8 @@ void USimpleConstructionScript::ExecuteScriptOnActor(AActor* Actor, const TInlin
 		// Get the given actor's root component (can be NULL).
 		USceneComponent* RootComponent = Actor->GetRootComponent();
 
-		for(auto NodeIt = RootNodes.CreateIterator(); NodeIt; ++NodeIt)
+		for (USCS_Node* RootNode : RootNodes)
 		{
-			USCS_Node* RootNode = *NodeIt;
 			if(RootNode != nullptr)
 			{
 				// If the root node specifies that it has a parent
@@ -696,12 +719,12 @@ void USimpleConstructionScript::RemoveNameToSCSNodeMap()
 #if WITH_EDITOR
 UBlueprint* USimpleConstructionScript::GetBlueprint() const
 {
-	if(auto OwnerClass = GetOwnerClass())
+	if (UClass* OwnerClass = GetOwnerClass())
 	{
 		return Cast<UBlueprint>(OwnerClass->ClassGeneratedBy);
 	}
 // >>> Backwards Compatibility:  VER_UE4_EDITORONLY_BLUEPRINTS
-	if(auto BP = Cast<UBlueprint>(GetOuter()))
+	if (UBlueprint* BP = Cast<UBlueprint>(GetOuter()))
 	{
 		return BP;
 	}
@@ -712,34 +735,35 @@ UBlueprint* USimpleConstructionScript::GetBlueprint() const
 
 UClass* USimpleConstructionScript::GetOwnerClass() const
 {
-	if(auto OwnerClass = Cast<UClass>(GetOuter()))
+	if (UClass* OwnerClass = Cast<UClass>(GetOuter()))
 	{
 		return OwnerClass;
 	}
 // >>> Backwards Compatibility:  VER_UE4_EDITORONLY_BLUEPRINTS
 #if WITH_EDITOR
-	if(auto BP = Cast<UBlueprint>(GetOuter()))
+	if (UBlueprint* BP = Cast<UBlueprint>(GetOuter()))
 	{
 		return BP->GeneratedClass;
 	}
 #endif
 // <<< End Backwards Compatibility
-	return NULL;
+	return nullptr;
 }
 
 UClass* USimpleConstructionScript::GetParentClass() const
 {
-	UClass* ParentClass = nullptr;
+#if WITH_EDITOR
 	if (UBlueprint* Blueprint = GetBlueprint())
 	{
-		ParentClass = Blueprint->ParentClass;
+		return Blueprint->ParentClass;
 	}
-	else if(UClass* OwnerClass = GetOwnerClass())
+#endif
+	if (UClass* OwnerClass = GetOwnerClass())
 	{
-		ParentClass = OwnerClass->GetSuperClass();
+		return OwnerClass->GetSuperClass();
 	}
 
-	return ParentClass;
+	return nullptr;
 }
 
 #if WITH_EDITOR
@@ -1027,21 +1051,31 @@ USceneComponent* USimpleConstructionScript::GetSceneRootComponentTemplate(USCS_N
 
 		for(int32 StackIndex = 0; StackIndex < SCSStack.Num() && !RootComponentTemplate; ++StackIndex)
 		{
-			// Check for any scene component nodes in the root set that are not the default scene root
 			const TArray<USCS_Node*>& SCSRootNodes = SCSStack[StackIndex]->GetRootNodes();
-			for(int32 RootNodeIndex = 0; RootNodeIndex < SCSRootNodes.Num() && RootComponentTemplate == nullptr; ++RootNodeIndex)
+
+			const bool bCanUseDefaultSceneRoot = DefaultSceneRootNode->ComponentTemplate && SCSRootNodes.Contains(DefaultSceneRootNode);
+			// Check for any scene component nodes in the root set that are not the default scene root
+			for (int32 RootNodeIndex = 0; RootNodeIndex < SCSRootNodes.Num() && RootComponentTemplate == nullptr; ++RootNodeIndex)
 			{
 				USCS_Node* RootNode = SCSRootNodes[RootNodeIndex];
-				if(RootNode != nullptr
+				if (RootNode != nullptr
 					&& RootNode != DefaultSceneRootNode
 					&& RootNode->ComponentTemplate != nullptr
 					&& RootNode->ComponentTemplate->IsA<USceneComponent>())
 				{
-					if(OutSCSNode)
+					// if we found a non-default scene root, but the default scene root is also present, then we assume that's the desired one and still return null
+					// this is to deal with the case where an actor component became scene root, but we don't want it to replace the default scene root if that was
+					// deliberately being used
+					if (bCanUseDefaultSceneRoot)
+					{
+						return nullptr;
+					}
+
+					if (OutSCSNode)
 					{
 						*OutSCSNode = RootNode;
 					}
-					
+
 					RootComponentTemplate = Cast<USceneComponent>(RootNode->ComponentTemplate);
 				}
 			}
@@ -1282,7 +1316,7 @@ FName USimpleConstructionScript::GenerateNewComponentName(const UClass* Componen
 
 USCS_Node* USimpleConstructionScript::CreateNodeImpl(UActorComponent* NewComponentTemplate, FName ComponentVariableName)
 {
-	auto NewNode = NewObject<USCS_Node>(this, MakeUniqueObjectName(this, USCS_Node::StaticClass()));
+	USCS_Node* NewNode = NewObject<USCS_Node>(this, MakeUniqueObjectName(this, USCS_Node::StaticClass()));
 	NewNode->SetFlags(RF_Transactional);
 	NewNode->ComponentClass = NewComponentTemplate->GetClass();
 	NewNode->ComponentTemplate = NewComponentTemplate;
