Dynamic resolution crash (Android) and query leaks fix

diff --git a/Engine/Source/Runtime/Engine/Private/DynamicResolution.cpp b/Engine/Source/Runtime/Engine/Private/DynamicResolution.cpp
--- a/Engine/Source/Runtime/Engine/Private/DynamicResolution.cpp
+++ b/Engine/Source/Runtime/Engine/Private/DynamicResolution.cpp
@@ -519,6 +519,12 @@ public:
 		bUseTimeQueriesThisFrame = false;
 	}
 
+	~FDefaultDynamicResolutionStateProxy()
+	{
+		check(IsInRenderingThread());
+		checkf(InFlightFrames.Num()==0, TEXT("Ensure the object is properly deinitialized by Finalize call"));
+	}
+
 	void Reset()
 	{
 		check(IsInRenderingThread());
@@ -544,6 +550,10 @@ public:
 
 		if (bUseTimeQueriesThisFrame)
 		{
+			// Create the query pool
+			if (!QueryPool.IsValid())
+				QueryPool = RHICreateRenderQueryPool(RQT_AbsoluteTime);
+
 			// Process to the in-flight frames that had their queries fully landed.
 			HandLandedQueriesToHeuristic(/* bWait = */ false);
 
@@ -556,8 +566,9 @@ public:
 			InFlightFrame.HeuristicHistoryEntry = Heuristic.CreateNewPreviousFrameTimings_RenderThread(
 				PrevGameThreadTimeMs, PrevRenderThreadTimeMs);
 
+			check(QueryPool.IsValid());
+			InFlightFrame.BeginFrameQuery = QueryPool->AllocateQuery();
 			RHICmdList.EndRenderQuery(InFlightFrame.BeginFrameQuery.GetQuery());
-			InFlightFrame.BeginFrameQuerySubmitted = true;
 		}
 		else
 		{
@@ -587,25 +598,22 @@ public:
 		{
 			InFlightFrameQueries& InFlightFrame = InFlightFrames[CurrentFrameInFlightIndex];
 
-			FRHIRenderQuery* QueryPtr = nullptr;
+			FRHIPooledRenderQuery* QueryPtr = nullptr;
 			switch (Event)
 			{
 			case EDynamicResolutionStateEvent::BeginDynamicResolutionRendering:
-				QueryPtr = InFlightFrame.BeginDynamicResolutionQuery.GetQuery(); 
-				InFlightFrame.BeginDynamicResolutionQuerySubmitted = true;
-				break;
+				QueryPtr = &InFlightFrame.BeginDynamicResolutionQuery; break;
 			case EDynamicResolutionStateEvent::EndDynamicResolutionRendering:
-				QueryPtr = InFlightFrame.EndDynamicResolutionQuery.GetQuery(); 
-				InFlightFrame.EndDynamicResolutionQuerySubmitted = true;
-				break;
+				QueryPtr = &InFlightFrame.EndDynamicResolutionQuery; break;
 			case EDynamicResolutionStateEvent::EndFrame:
-				QueryPtr = InFlightFrame.EndFrameQuery.GetQuery(); 
-				InFlightFrame.EndFrameQuerySubmitted = true;
-				break;
+				QueryPtr = &InFlightFrame.EndFrameQuery; break;
 			default: check(0);
 			}
 
-			RHICmdList.EndRenderQuery(QueryPtr);
+			check(QueryPtr != nullptr);
+			check(QueryPool.IsValid());
+			*QueryPtr = QueryPool->AllocateQuery();
+			RHICmdList.EndRenderQuery(QueryPtr->GetQuery());
 		}
 
 		// Clobber CurrentFrameInFlightIndex for internal checks.
@@ -616,12 +624,17 @@ public:
 		}
 	}
 
-	void Finish()
+	/// Called before object is to be deleted
+	void Finalize()
 	{
 		check(IsInRenderingThread());
 
 		// Wait for all queries to land.
 		HandLandedQueriesToHeuristic(/* bWait = */ true);
+
+		// Reset the frame properties
+		InFlightFrames.Empty();
+		QueryPool = nullptr;
 	}
 
 
@@ -633,7 +646,6 @@ private:
 	struct InFlightFrameQueries
 	{
 		// GPU queries.
-		FRenderQueryPoolRHIRef QueryPool;
 		FRHIPooledRenderQuery BeginFrameQuery;
 		FRHIPooledRenderQuery BeginDynamicResolutionQuery;
 		FRHIPooledRenderQuery EndDynamicResolutionQuery;
@@ -642,31 +654,25 @@ private:
 		// Heuristic's history 
 		uint64 HeuristicHistoryEntry;
 
-		uint32 BeginFrameQuerySubmitted : 1;
-		uint32 BeginDynamicResolutionQuerySubmitted : 1;
-		uint32 EndDynamicResolutionQuerySubmitted : 1;
-		uint32 EndFrameQuerySubmitted : 1;
-
 		InFlightFrameQueries()
-			: QueryPool(RHICreateRenderQueryPool(RQT_AbsoluteTime, 4))
-			, HeuristicHistoryEntry(FDynamicResolutionHeuristicProxy::kInvalidEntryId) 
-			, BeginFrameQuerySubmitted(false)
-			, BeginDynamicResolutionQuerySubmitted(false)
-			, EndDynamicResolutionQuerySubmitted(false)
-			, EndFrameQuerySubmitted(false)
 		{ 
-			InFlightFrameQueries* self = this;
-			ENQUEUE_RENDER_COMMAND(InFlightFrameQueries)(
-			[self](FRHICommandListImmediate& RHICmdList)
-			{
-				self->BeginFrameQuery = self->QueryPool->AllocateQuery();
-				self->BeginDynamicResolutionQuery = self->QueryPool->AllocateQuery();
-				self->EndDynamicResolutionQuery = self->QueryPool->AllocateQuery();
-				self->EndFrameQuery = self->QueryPool->AllocateQuery();
-			});
+			ResetValues();
+		}
+
+		/// Reset values
+		void ResetValues()
+		{
+			HeuristicHistoryEntry = FDynamicResolutionHeuristicProxy::kInvalidEntryId;
+			BeginFrameQuery.ReleaseQuery();
+			BeginDynamicResolutionQuery.ReleaseQuery();
+			EndDynamicResolutionQuery.ReleaseQuery();
+			EndFrameQuery.ReleaseQuery();
 		}
 	};
 
+	// Shared query pool for the frames in flight
+	FRenderQueryPoolRHIRef QueryPool;
+
 	// List of frame queries in flight.
 	TArray<InFlightFrameQueries> InFlightFrames;
 
@@ -701,36 +707,32 @@ private:
 
 			int32 LandingCount = 0;
 			int32 QueryCount = 0;
-			if (InFlightFrame.BeginFrameQuery.GetQuery() != nullptr && InFlightFrame.BeginFrameQuerySubmitted == true)
+			if (InFlightFrame.BeginFrameQuery.IsValid())
 			{
 				LandingCount += RHIGetRenderQueryResult(
 					InFlightFrame.BeginFrameQuery.GetQuery(), BeginFrameResult, bWait) ? 1 : 0;
 				QueryCount += 1;
-				InFlightFrame.BeginFrameQuerySubmitted = false;
 			}
 
-			if (InFlightFrame.BeginDynamicResolutionQuery.GetQuery() != nullptr && InFlightFrame.BeginDynamicResolutionQuerySubmitted == true)
+			if (InFlightFrame.BeginDynamicResolutionQuery.IsValid())
 			{
 				LandingCount += RHIGetRenderQueryResult(
 					InFlightFrame.BeginDynamicResolutionQuery.GetQuery(), BeginDynamicResolutionResult, bWait) ? 1 : 0;
 				QueryCount += 1;
-				InFlightFrame.BeginDynamicResolutionQuerySubmitted = false;
 			}
 
-			if (InFlightFrame.EndDynamicResolutionQuery.GetQuery() != nullptr && InFlightFrame.EndDynamicResolutionQuerySubmitted == true)
+			if (InFlightFrame.EndDynamicResolutionQuery.IsValid())
 			{
 				LandingCount += RHIGetRenderQueryResult(
 					InFlightFrame.EndDynamicResolutionQuery.GetQuery(), EndDynamicResolutionResult, bWait) ? 1 : 0;
 				QueryCount += 1;
-				InFlightFrame.EndDynamicResolutionQuerySubmitted = false;
 			}
 
-			if (InFlightFrame.EndFrameQuery.GetQuery() != nullptr && InFlightFrame.EndFrameQuerySubmitted == true)
+			if (InFlightFrame.EndFrameQuery.IsValid())
 			{
 				LandingCount += RHIGetRenderQueryResult(
 					InFlightFrame.EndFrameQuery.GetQuery(), EndFrameResult, bWait) ? 1 : 0;
 				QueryCount += 1;
-				InFlightFrame.EndFrameQuerySubmitted = false;
 			}
 
 			check(QueryCount == 0 || QueryCount == 4);
@@ -745,7 +747,7 @@ private:
 					/* bGPUTimingsHaveCPUBubbles = */ !GRHISupportsGPUTimestampBubblesRemoval);
 
 				// Reset this in-flight frame queries to be reused.
-				InFlightFrame = InFlightFrameQueries();
+				InFlightFrame.ResetValues();
 
 				ShouldRefreshHeuristic = true;
 			}
@@ -766,7 +768,7 @@ private:
 		for (int32 i = 0; i < InFlightFrames.Num(); i++)
 		{
 			auto& InFlightFrame = InFlightFrames[i];
-			if (InFlightFrame.BeginFrameQuery.GetQuery() == nullptr)
+			if (!InFlightFrame.BeginFrameQuery.IsValid())
 			{
 				CurrentFrameInFlightIndex = i;
 				break;
@@ -776,8 +778,7 @@ private:
 		// Allocate a new in-flight frame in the unlikely event.
 		if (CurrentFrameInFlightIndex == -1)
 		{
-			CurrentFrameInFlightIndex = InFlightFrames.Num();
-			InFlightFrames.Add(InFlightFrameQueries());
+			CurrentFrameInFlightIndex = InFlightFrames.Add(InFlightFrameQueries());
 		}
 	}
 };
@@ -806,7 +807,7 @@ public:
 		ENQUEUE_RENDER_COMMAND(DeleteDynamicResolutionProxy)(
 			[P](class FRHICommandList&)
 			{
-				P->Finish();
+				P->Finalize();
 				delete P;
 			});
 	}
diff --git a/Engine/Source/Runtime/RHI/Private/DynamicRHI.cpp b/Engine/Source/Runtime/RHI/Private/DynamicRHI.cpp
--- a/Engine/Source/Runtime/RHI/Private/DynamicRHI.cpp
+++ b/Engine/Source/Runtime/RHI/Private/DynamicRHI.cpp
@@ -366,11 +366,18 @@ FDefaultRHIRenderQueryPool::FDefaultRHIRenderQueryPool(ERenderQueryType InQueryT
 		for (uint32 i = 0; i < NumQueries; i++)
 		{
 			Queries.Push(DynamicRHI->RHICreateRenderQuery(QueryType));
+			check(Queries.Last().IsValid());
 			++AllocatedQueries;
 		}
 	}
 }
 
+FDefaultRHIRenderQueryPool::~FDefaultRHIRenderQueryPool()
+{
+	check(IsInRenderingThread());
+	checkf(AllocatedQueries == Queries.Num(), TEXT("Querypool deleted before all Queries have been released"));
+}
+
 FRHIPooledRenderQuery FDefaultRHIRenderQueryPool::AllocateQuery()
 {
 	check(IsInRenderingThread());
@@ -380,8 +387,13 @@ FRHIPooledRenderQuery FDefaultRHIRenderQueryPool::AllocateQuery()
 	}
 	else
 	{
-		ensure(++AllocatedQueries <= NumQueries);
-		return FRHIPooledRenderQuery(this, DynamicRHI->RHICreateRenderQuery(QueryType));
+		FRHIPooledRenderQuery Query = FRHIPooledRenderQuery(this, DynamicRHI->RHICreateRenderQuery(QueryType));
+		if (Query.IsValid())
+		{
+			++AllocatedQueries;
+		}
+		ensure(AllocatedQueries <= NumQueries);
+		return Query;
 	}
 }
 
@@ -395,6 +407,10 @@ void FDefaultRHIRenderQueryPool::ReleaseQuery(TRefCountPtr<FRHIRenderQuery>&& Qu
 	check(IsInRenderingThread());
 	//Hard to validate because of Resource resurrection, better to remove GetQueryRef entirely
 	//checkf(Query.IsValid() && Query.GetRefCount() <= 2, TEXT("Query has been released but reference still held: use FRHIPooledRenderQuery::GetQueryRef() with extreme caution"));
+	
+	checkf(Query.IsValid(), TEXT("Only release valid queries"));
+	checkf((uint32)Queries.Num() < NumQueries, TEXT("Pool contains more queries than it started with, double release somewhere?"));
+
 	Queries.Push(MoveTemp(Query));
 	check(!Query.IsValid());
 }
diff --git a/Engine/Source/Runtime/RHI/Public/DynamicRHI.h b/Engine/Source/Runtime/RHI/Public/DynamicRHI.h
--- a/Engine/Source/Runtime/RHI/Public/DynamicRHI.h
+++ b/Engine/Source/Runtime/RHI/Public/DynamicRHI.h
@@ -169,12 +169,7 @@ class RHI_API FDefaultRHIRenderQueryPool final : public FRHIRenderQueryPool
 {
 public:
 	FDefaultRHIRenderQueryPool(ERenderQueryType InQueryType, FDynamicRHI* InDynamicRHI, uint32 InNumQueries);
-
-	~FDefaultRHIRenderQueryPool() override 
-	{
-		check(IsInRenderingThread());
-		checkf(AllocatedQueries == Queries.Num(), TEXT("Querypool deleted before all Queries have been released"));
-	}
+	~FDefaultRHIRenderQueryPool() override;
 
 private:
 	virtual FRHIPooledRenderQuery AllocateQuery() override;
