Do not (re)set HDR parameters unconditionally.

- They can (and should) be taken out of the real device profiles

diff --git a/Engine/Source/Runtime/Engine/Private/UnrealEngine.cpp b/Engine/Source/Runtime/Engine/Private/UnrealEngine.cpp
--- a/Engine/Source/Runtime/Engine/Private/UnrealEngine.cpp
+++ b/Engine/Source/Runtime/Engine/Private/UnrealEngine.cpp
@@ -610,6 +610,7 @@ void HDRSettingChangedSinkCallback()
 		
 		int32 OutputDevice = 0;
 		int32 ColorGamut = 0;
+		bool bNewValuesSet = false;
 		
 		// If we are turning HDR on we must set the appropriate OutputDevice and ColorGamut.
 		// If we are turning it off, we'll reset back to 0/0
@@ -621,12 +622,14 @@ void HDRSettingChangedSinkCallback()
 				// ScRGB, 1000 or 2000 nits, Rec2020
 				OutputDevice = (DisplayNitLevel == 1000) ? 5 : 6;
 				ColorGamut = 2;
+				bNewValuesSet = true;
 			}
 #elif PLATFORM_PS4
 			{
 				// PQ, 1000 or 2000 nits, Rec2020
 				OutputDevice = (DisplayNitLevel == 1000) ? 3 : 4;
 				ColorGamut = 2;
+				bNewValuesSet = true;
 			}
 
 #elif PLATFORM_MAC
@@ -634,29 +637,35 @@ void HDRSettingChangedSinkCallback()
 				// ScRGB, 1000 or 2000 nits, DCI-P3
 				OutputDevice = DisplayNitLevel == 1000 ? 5 : 6;
 				ColorGamut = 1;
+				bNewValuesSet = true;
 			}
 #elif PLATFORM_IOS
 			{
 				// Linear output to Apple's specific format.
 				OutputDevice = 7;
 				ColorGamut = 0;
+				bNewValuesSet = true;
 			}
 #elif PLATFORM_XBOXONE
 			{
 				// PQ, 1000 or 2000 nits, Rec2020
 				OutputDevice = (DisplayNitLevel == 1000) ? 3 : 4;
 				ColorGamut = 2;
+				bNewValuesSet = true;
 			}
 #endif
 		}
-		
-		static IConsoleVariable* CVarHDROutputDevice = IConsoleManager::Get().FindConsoleVariable(TEXT("r.HDR.Display.OutputDevice"));
-		static IConsoleVariable* CVarHDRColorGamut = IConsoleManager::Get().FindConsoleVariable(TEXT("r.HDR.Display.ColorGamut"));
-		check(CVarHDROutputDevice);
-		check(CVarHDRColorGamut);
-		
-		CVarHDROutputDevice->Set(OutputDevice, ECVF_SetByDeviceProfile);
-		CVarHDRColorGamut->Set(ColorGamut, ECVF_SetByDeviceProfile);
+
+		if (bNewValuesSet)
+		{
+			static IConsoleVariable* CVarHDROutputDevice = IConsoleManager::Get().FindConsoleVariable(TEXT("r.HDR.Display.OutputDevice"));
+			static IConsoleVariable* CVarHDRColorGamut = IConsoleManager::Get().FindConsoleVariable(TEXT("r.HDR.Display.ColorGamut"));
+			check(CVarHDROutputDevice);
+			check(CVarHDRColorGamut);
+
+			CVarHDROutputDevice->Set(OutputDevice, ECVF_SetByDeviceProfile);
+			CVarHDRColorGamut->Set(ColorGamut, ECVF_SetByDeviceProfile);
+		}
 		
 		// Now set the HDR setting.
 		GRHIIsHDREnabled = CVarHDROutputEnabled->GetValueOnAnyThread() != 0;
