Network Emulation
* Set default PktLossMaxSize to 0 and now consider 0 to be the disabled value
* Added bNeverApplyNetworkEmulationSettings to NetDriver that allows some NetDrivers to never have emulation enabled.
* DemoNetDriver have bNeverApplyNetworkEmulationSettings set to true

[RN] minor engine/network

diff --git a/Engine/Source/Runtime/Engine/Classes/Engine/NetDriver.h b/Engine/Source/Runtime/Engine/Classes/Engine/NetDriver.h
--- a/Engine/Source/Runtime/Engine/Classes/Engine/NetDriver.h
+++ b/Engine/Source/Runtime/Engine/Classes/Engine/NetDriver.h
@@ -384,7 +384,7 @@ struct ENGINE_API FPacketSimulationSettings
 	 * Works with all other settings.
 	 */
 	UPROPERTY(EditAnywhere, Category="Simulation Settings")
-	int32	PktLoss;
+	int32	PktLoss = 0;
 
 	/**
 	* Sets the maximum size of packets in bytes that will be dropped
@@ -393,7 +393,7 @@ struct ENGINE_API FPacketSimulationSettings
 	* Works with all other settings.
 	*/
 	UPROPERTY(EditAnywhere, Category = "Simulation Settings")
-	int32	PktLossMaxSize;
+	int32	PktLossMaxSize = 0;
 
 	/**
 	* Sets the minimum size of packets in bytes that will be dropped
@@ -402,7 +402,7 @@ struct ENGINE_API FPacketSimulationSettings
 	* Works with all other settings.
 	*/
 	UPROPERTY(EditAnywhere, Category = "Simulation Settings")
-	int32	PktLossMinSize;
+	int32	PktLossMinSize = 0;
 
 	/**
 	 * When set, will cause calls to FlushNet to change ordering of packets at random.
@@ -412,7 +412,7 @@ struct ENGINE_API FPacketSimulationSettings
 	 * Takes precedence over PktDup and PktLag.
 	 */
 	UPROPERTY(EditAnywhere, Category="Simulation Settings")
-	int32	PktOrder;
+	int32	PktOrder = 0;
 
 	/**
 	 * When set, will cause calls to FlushNet to duplicate packets.
@@ -423,7 +423,7 @@ struct ENGINE_API FPacketSimulationSettings
 	 * Cannot be used with PktOrder or PktLag.
 	 */
 	UPROPERTY(EditAnywhere, Category="Simulation Settings")
-	int32	PktDup;
+	int32	PktDup = 0;
 	
 	/**
 	 * When set, will cause calls to FlushNet to delay packets.
@@ -432,7 +432,7 @@ struct ENGINE_API FPacketSimulationSettings
 	 * Cannot be used with PktOrder.
 	 */
 	UPROPERTY(EditAnywhere, Category="Simulation Settings")
-	int32	PktLag;
+	int32	PktLag = 0;
 	
 	/**
 	 * When set, will cause PktLag to use variable lag instead of constant.
@@ -441,7 +441,7 @@ struct ENGINE_API FPacketSimulationSettings
 	 * Can only be used when PktLag is enabled.
 	 */
 	UPROPERTY(EditAnywhere, Category="Simulation Settings")
-	int32	PktLagVariance;
+	int32	PktLagVariance = 0;
 
 	/**
 	 * If set lag values will randomly fluctuate between Min and Max.
@@ -453,18 +453,6 @@ struct ENGINE_API FPacketSimulationSettings
 	UPROPERTY(EditAnywhere, Category = "Simulation Settings")
 	int32	PktLagMax = 0;
 
-	/** Ctor. Zeroes the settings */
-	FPacketSimulationSettings() : 
-		PktLoss(0),
-		PktLossMaxSize(INT_MAX/8),
-		PktLossMinSize(0),
-		PktOrder(0),
-		PktDup(0),
-		PktLag(0),
-		PktLagVariance(0) 
-	{
-	}
-
 	/** reads in settings from the .ini file 
 	 * @note: overwrites all previous settings
 	 */
@@ -472,18 +460,19 @@ struct ENGINE_API FPacketSimulationSettings
 	
 	/** 
 	 * Load a preconfigured emulation profile from the .ini
+	 * Returns true if the given profile existed
 	 */
-	void LoadEmulationProfile(const TCHAR* ProfileName);
+	bool LoadEmulationProfile(const TCHAR* ProfileName);
 
 	/**
-	 * Registers commands for auto-completion, etc.
+	 * Force new emulation settings and ignore config or cmdline values
 	 */
-	void RegisterCommands();
+	void ApplySettings(const FPacketSimulationSettings& NewSettings);
 
 	/**
-	 * Unregisters commands for auto-completion, etc.
+	 * Registers commands for auto-completion, etc.
 	 */
-	void UnregisterCommands();
+	void RegisterCommands();
 
 	/**
 	 * Ensure that settings have proper values
@@ -491,6 +480,16 @@ struct ENGINE_API FPacketSimulationSettings
 	void ValidateSettings();
 	void ResetSettings();
 
+	/**
+	* Tells if a packet fits the size settings to potentially be dropped
+	*/
+	bool ShouldDropPacketOfSize(int32 NumBits) const
+	{
+		const bool bIsBigEnough = NumBits > PktLossMinSize * 8;
+		const bool bIsSmallEnough = PktLossMaxSize == 0 || NumBits < PktLossMaxSize * 8;
+		return bIsBigEnough && bIsSmallEnough;
+	}
+
 	/**
 	 * Reads the settings from a string: command line or an exec
 	 *
@@ -709,6 +708,13 @@ public:
 	UPROPERTY(Config)
 	bool bNoTimeouts;
 
+	/**
+	 * If true this NetDriver will not apply the network emulation settings that simulate
+	 * latency and packet loss in non-shippable builds
+	 */
+	UPROPERTY(Config)
+	bool bNeverApplyNetworkEmulationSettings;
+
 	/** Connection to the server (this net driver is a client) */
 	UPROPERTY()
 	class UNetConnection* ServerConnection;
@@ -1066,7 +1072,12 @@ public:
 #if DO_ENABLE_NET_TEST
 	FPacketSimulationSettings	PacketSimulationSettings;
 
-	ENGINE_API void SetPacketSimulationSettings(FPacketSimulationSettings NewSettings);
+	/**
+	 * Modify the current emulation settings
+	 */
+	ENGINE_API void SetPacketSimulationSettings(const FPacketSimulationSettings& NewSettings);
+
+	void OnPacketSimulationSettingsChanged();
 #endif
 
 	// Constructors.
diff --git a/Engine/Source/Runtime/Engine/Private/DemoNetDriver.cpp b/Engine/Source/Runtime/Engine/Private/DemoNetDriver.cpp
--- a/Engine/Source/Runtime/Engine/Private/DemoNetDriver.cpp
+++ b/Engine/Source/Runtime/Engine/Private/DemoNetDriver.cpp
@@ -640,6 +640,7 @@ UDemoNetDriver::UDemoNetDriver(const FObjectInitializer& ObjectInitializer)
 	bIsWaitingForHeaderDownload = false;
 	bIsWaitingForStream = false;
 	MaxArchiveReadPos = 0;
+	bNeverApplyNetworkEmulationSettings = true;
 
 	if (!HasAnyFlags(RF_ClassDefaultObject))
 	{
diff --git a/Engine/Source/Runtime/Engine/Private/NetConnection.cpp b/Engine/Source/Runtime/Engine/Private/NetConnection.cpp
--- a/Engine/Source/Runtime/Engine/Private/NetConnection.cpp
+++ b/Engine/Source/Runtime/Engine/Private/NetConnection.cpp
@@ -1413,9 +1413,10 @@ void UNetConnection::FlushNet(bool bIgnoreSimulation)
 bool UNetConnection::ShouldDropOutgoingPacketForLossSimulation(int64 NumBits) const
 {
 #if DO_ENABLE_NET_TEST
-	return Driver->IsSimulatingPacketLossBurst() ||
-		(NumBits > PacketSimulationSettings.PktLossMinSize * 8 && NumBits < PacketSimulationSettings.PktLossMaxSize * 8
-		&& PacketSimulationSettings.PktLoss > 0 && FMath::FRand() * 100.f < PacketSimulationSettings.PktLoss);
+	return Driver->IsSimulatingPacketLossBurst() || 
+		(PacketSimulationSettings.PktLoss > 0 && 
+         PacketSimulationSettings.ShouldDropPacketOfSize(NumBits) && 
+         FMath::FRand() * 100.f < PacketSimulationSettings.PktLoss);
 #else
 	return false;
 #endif
diff --git a/Engine/Source/Runtime/Engine/Private/NetworkDriver.cpp b/Engine/Source/Runtime/Engine/Private/NetworkDriver.cpp
--- a/Engine/Source/Runtime/Engine/Private/NetworkDriver.cpp
+++ b/Engine/Source/Runtime/Engine/Private/NetworkDriver.cpp
@@ -289,8 +289,14 @@ PRAGMA_ENABLE_DEPRECATION_WARNINGS
 void UNetDriver::InitPacketSimulationSettings()
 {
 #if DO_ENABLE_NET_TEST
+	if (bNeverApplyNetworkEmulationSettings)
+	{
+		return;
+	}
+
+	PacketSimulationSettings.ResetSettings();
+
 	// read the settings from .ini and command line, with the command line taking precedence	
-	PacketSimulationSettings = FPacketSimulationSettings();
 	PacketSimulationSettings.LoadConfig(*NetDriverName.ToString());
 	PacketSimulationSettings.RegisterCommands();
 	PacketSimulationSettings.ParseSettings(FCommandLine::Get(), *NetDriverName.ToString());
@@ -1564,10 +1570,6 @@ void UNetDriver::Shutdown()
 
 	ConnectionlessHandler.Reset(nullptr);
 
-#if DO_ENABLE_NET_TEST
-	PacketSimulationSettings.UnregisterCommands();
-#endif
-
 	SetReplicationDriver(nullptr);
 
 	if (AnalyticsAggregator.IsValid())
@@ -2752,7 +2754,8 @@ bool UNetDriver::Exec( UWorld* InWorld, const TCHAR* Cmd, FOutputDevice& Ar )
 	}
 #if DO_ENABLE_NET_TEST
 	// This will allow changing the Pkt* options at runtime
-	else if (PacketSimulationSettings.ParseSettings(Cmd, *NetDriverName.ToString()))
+	else if (bNeverApplyNetworkEmulationSettings == false && 
+			 PacketSimulationSettings.ParseSettings(Cmd, *NetDriverName.ToString()))
 	{
 		if (ServerConnection)
 		{
@@ -3218,9 +3221,19 @@ void UNetDriver::AddReferencedObjects(UObject* InThis, FReferenceCollector& Coll
 
 #if DO_ENABLE_NET_TEST
 
-void UNetDriver::SetPacketSimulationSettings(FPacketSimulationSettings NewSettings)
+void UNetDriver::SetPacketSimulationSettings(const FPacketSimulationSettings& NewSettings)
 {
+	if (bNeverApplyNetworkEmulationSettings)
+	{
+		return;
+	}
+
 	PacketSimulationSettings = NewSettings;
+	OnPacketSimulationSettingsChanged();
+}
+
+void UNetDriver::OnPacketSimulationSettingsChanged()
+{
 	if (ServerConnection)
 	{
 		ServerConnection->UpdatePacketSimulationSettings();
@@ -3268,7 +3281,7 @@ void FPacketSimulationSettings::LoadConfig(const TCHAR* OptionalQualifier)
 	ValidateSettings();
 }
 
-void FPacketSimulationSettings::LoadEmulationProfile(const TCHAR* ProfileName)
+bool FPacketSimulationSettings::LoadEmulationProfile(const TCHAR* ProfileName)
 {
 	const FString SectionName = FString::Printf(TEXT("%s.%s"), TEXT("PacketSimulationProfile"), ProfileName);
 
@@ -3277,7 +3290,7 @@ void FPacketSimulationSettings::LoadEmulationProfile(const TCHAR* ProfileName)
 	if (!bSectionExists)
 	{
 		UE_LOG(LogNet, Warning, TEXT("EmulationProfile [%s] was not found in %s. Packet settings were not changed"), *SectionName, *GEngineIni);
-		return;
+		return false;
 	}
 
 	ResetSettings();
@@ -3307,6 +3320,8 @@ void FPacketSimulationSettings::LoadEmulationProfile(const TCHAR* ProfileName)
 	}
 	
 	ValidateSettings();
+
+	return true;
 }
 
 void FPacketSimulationSettings::ResetSettings()
@@ -3378,29 +3393,9 @@ void FPacketSimulationSettings::RegisterCommands()
 		ConsoleManager.RegisterConsoleCommand(TEXT("Net PktLagMax="), TEXT("PktLagMax=<n> (maximum packet latency)"));
 		ConsoleManager.RegisterConsoleCommand(TEXT("Net EmulationProfile="), TEXT("EmulationProfile=<NAME> (apply an emulation profile)"));
 	}
-}
-
 
-void FPacketSimulationSettings::UnregisterCommands()
-{
-	// Never unregister the console commands. Since net drivers come and go, and we can sometimes have more than 1, etc. 
+	// Note we never unregister the console commands since net drivers come and go, and we can sometimes have more than 1, etc. 
 	// We could do better bookkeeping for this, but its not worth it right now. Just ensure the commands are always there for tab completion.
-#if 0
-	IConsoleManager& ConsoleManager = IConsoleManager::Get();
-
-	// Gather up all relevant console commands
-	TArray<IConsoleObject*> PacketSimulationConsoleCommands;
-	ConsoleManager.ForEachConsoleObjectThatStartsWith(
-		FConsoleObjectVisitor::CreateStatic< TArray<IConsoleObject*>& >(
-		&FPacketSimulationConsoleCommandVisitor::OnPacketSimulationConsoleCommand,
-		PacketSimulationConsoleCommands ), TEXT("Net Pkt"));
-
-	// Unregister them from the console manager
-	for(int32 i = 0; i < PacketSimulationConsoleCommands.Num(); ++i)
-	{
-		ConsoleManager.UnregisterConsoleObject(PacketSimulationConsoleCommands[i]);
-	}
-#endif
 }
 
 /**
