Remove comment.

diff --git a/Engine/Source/Runtime/Engine/Private/Components/SplineMeshComponent.cpp b/Engine/Source/Runtime/Engine/Private/Components/SplineMeshComponent.cpp
--- a/Engine/Source/Runtime/Engine/Private/Components/SplineMeshComponent.cpp
+++ b/Engine/Source/Runtime/Engine/Private/Components/SplineMeshComponent.cpp
@@ -522,9 +522,6 @@ void USplineMeshComponent::UpdateRenderStateAndCollision_Internal(bool bConcurre
 	}
 
 	CachedMeshBodySetupGuid.Invalidate();
-	// spline visualizer -> tell the points its interactive -> spline comp -> to water body -> generator -> update mesh on spline mesh component
-	// release the mouse must call another post edit change so that recreate physics state at that point.
-	// make this optional - optional flag set based on whether it's interactive or not.
 	RecreatePhysicsState();
 
 	bMeshDirty = false;
