Undo 7248860 since it is modifying dormancy settings during cooking phase

#rnx

diff --git a/Engine/Source/Runtime/Engine/Private/World.cpp b/Engine/Source/Runtime/Engine/Private/World.cpp
--- a/Engine/Source/Runtime/Engine/Private/World.cpp
+++ b/Engine/Source/Runtime/Engine/Private/World.cpp
@@ -4429,14 +4429,6 @@ void UWorld::AddNetworkActor( AActor* Actor )
 		return;
 	}
 
-    // Remove DORM_Initial from dynamic actors
-	if (Actor->NetDormancy == DORM_Initial && Actor->IsNetStartupActor() == false)
-	{
-		UE_LOG(LogNet, VeryVerbose, TEXT("Dynamic actor %s dormancy was changed from DORM_Initial to DORM_DormantAll"), *Actor->GetFullName());
-        // Changing dormancy directly here since the actor was not yet registered to the network drivers
-		Actor->NetDormancy = DORM_DormantAll;
-	}
-
 	ForEachNetDriver(GEngine, this, [Actor](UNetDriver* const Driver)
 	{
 		if (Driver != nullptr)
