Added an explicit flag to allow batching imposters in HLOD and render them as instanced static meshes.

diff --git a/Engine/Source/Runtime/Engine/Classes/Components/PrimitiveComponent.h b/Engine/Source/Runtime/Engine/Classes/Components/PrimitiveComponent.h
--- a/Engine/Source/Runtime/Engine/Classes/Components/PrimitiveComponent.h
+++ b/Engine/Source/Runtime/Engine/Classes/Components/PrimitiveComponent.h
@@ -230,6 +230,10 @@ public:
 	UPROPERTY(EditAnywhere, AdvancedDisplay, BlueprintReadWrite, Category = HLOD)
 	uint8 bUseMaxLODAsImposter : 1;
 
+	/** If true, the proxy generation process will use instancing to render this imposter */
+	UPROPERTY(EditAnywhere, AdvancedDisplay, BlueprintReadWrite, Category = HLOD, meta = (EditCondition = "bUseMaxLODAsImposter"))
+	uint8 bBatchImpostersAsInstances : 1;
+
 	/**
 	 * When enabled this object will not be culled by distance. This is ignored if a child of a HLOD.
 	 */
diff --git a/Engine/Source/Runtime/Engine/Private/Components/PrimitiveComponent.cpp b/Engine/Source/Runtime/Engine/Private/Components/PrimitiveComponent.cpp
--- a/Engine/Source/Runtime/Engine/Private/Components/PrimitiveComponent.cpp
+++ b/Engine/Source/Runtime/Engine/Private/Components/PrimitiveComponent.cpp
@@ -340,6 +340,8 @@ UPrimitiveComponent::UPrimitiveComponent(const FObjectInitializer& ObjectInitial
 
 	LDMaxDrawDistance = 0.f;
 	CachedMaxDrawDistance = 0.f;
+	bUseMaxLODAsImposter = false;
+	bBatchImpostersAsInstances = false;
 	bNeverDistanceCull = false;
 
 	bUseEditorCompositing = false;
@@ -365,7 +367,6 @@ UPrimitiveComponent::UPrimitiveComponent(const FObjectInitializer& ObjectInitial
 	bReceiveMobileCSMShadows = true;
 #if WITH_EDITORONLY_DATA
 	bEnableAutoLODGeneration = true;
-	bUseMaxLODAsImposter = false;
 #endif // WITH_EDITORONLY_DATA
 }
 
diff --git a/Engine/Source/Runtime/Engine/Private/LODActor.cpp b/Engine/Source/Runtime/Engine/Private/LODActor.cpp
--- a/Engine/Source/Runtime/Engine/Private/LODActor.cpp
+++ b/Engine/Source/Runtime/Engine/Private/LODActor.cpp
@@ -1034,7 +1034,7 @@ bool IsImposter(const UStaticMeshComponent* InComponent)
 {
 	check(InComponent);
 
-	if (!InComponent->bUseMaxLODAsImposter)
+	if (!InComponent->bUseMaxLODAsImposter || !InComponent->bBatchImpostersAsInstances)
 	{
 		return false;
 	}
@@ -1058,14 +1058,7 @@ bool IsImposter(const UStaticMeshComponent* InComponent)
 		return false;
 	}
 
-	// Currently, the bUseMaxLODAsImposter flag serves multiple purposes. Until this is sorted out,
-	// only take into accounts quads & octagons meshes.
-	const int32 NumWedges = StaticMeshLOD.IndexBuffer.GetNumIndices();
-	const int32 NumVertexPositions = StaticMeshLOD.VertexBuffers.PositionVertexBuffer.GetNumVertices();
-	const int32 NumFaces = NumWedges / 3;
-
-	return (NumFaces == 2 && NumVertexPositions == 4) ||
-		   (NumFaces == 8 && NumVertexPositions == 9);
+	return true;
 }
 
 UMaterialInterface* ALODActor::GetImposterMaterial(const UStaticMeshComponent* InComponent) const
