[NO NON-PII COMMIT MESSAGE]

diff --git a/Engine/Source/Runtime/Engine/Private/AudioDecompress.cpp b/Engine/Source/Runtime/Engine/Private/AudioDecompress.cpp
--- a/Engine/Source/Runtime/Engine/Private/AudioDecompress.cpp
+++ b/Engine/Source/Runtime/Engine/Private/AudioDecompress.cpp
@@ -732,5 +732,4 @@ bool ShouldUseBackgroundPoolFor_FAsyncRealtimeAudioTask()
 	return !!CVarShouldUseBackgroundPoolFor_FAsyncRealtimeAudioTask.GetValueOnAnyThread();
 }
 
-PRAGMA_ENABLE_OPTIMIZATION
 // end
