Added FCoreDelegates::OnBeforeHandleSystemError, this is to move registry writing for editor analytics before main crash reporting, since killing the editor process once the "Fatal error!" dialog pops up (and then freezes for a while during crash reporting) would cause an abnormal shutdown to be reported.

This should help reduce the number of crashes being reported as abnormal shutdowns because the analytics are never written out.

diff --git a/Engine/Source/Runtime/Core/Private/Misc/CoreDelegates.cpp b/Engine/Source/Runtime/Core/Private/Misc/CoreDelegates.cpp
--- a/Engine/Source/Runtime/Core/Private/Misc/CoreDelegates.cpp
+++ b/Engine/Source/Runtime/Core/Private/Misc/CoreDelegates.cpp
@@ -42,6 +42,8 @@ FCoreDelegates::FOnUserControllerPairingChange FCoreDelegates::OnControllerPairi
 FCoreDelegates::FOnSafeFrameChangedEvent FCoreDelegates::OnSafeFrameChangedEvent;
 FCoreDelegates::FOnHandleSystemEnsure FCoreDelegates::OnHandleSystemEnsure;
 FCoreDelegates::FOnHandleSystemError FCoreDelegates::OnHandleSystemError;
+FCoreDelegates::FOnHandleSystemError FCoreDelegates::OnBeforeHandleSystemError;
+
 FCoreDelegates::FOnActorLabelChanged FCoreDelegates::OnActorLabelChanged;
 
 FCoreDelegates::FRegisterMovieStreamerDelegate FCoreDelegates::RegisterMovieStreamerDelegate;
diff --git a/Engine/Source/Runtime/Core/Private/Windows/WindowsPlatformCrashContext.cpp b/Engine/Source/Runtime/Core/Private/Windows/WindowsPlatformCrashContext.cpp
--- a/Engine/Source/Runtime/Core/Private/Windows/WindowsPlatformCrashContext.cpp
+++ b/Engine/Source/Runtime/Core/Private/Windows/WindowsPlatformCrashContext.cpp
@@ -822,6 +822,8 @@ private:
 
 		GLog->PanicFlushThreadedLogs();
 
+		FCoreDelegates::OnBeforeHandleSystemError.Broadcast();
+
 		// Get the default settings for the crash context
 		ECrashContextType Type = ECrashContextType::Crash;
 		const TCHAR* ErrorMessage = TEXT("Unhandled exception");
diff --git a/Engine/Source/Runtime/Core/Public/Misc/CoreDelegates.h b/Engine/Source/Runtime/Core/Public/Misc/CoreDelegates.h
--- a/Engine/Source/Runtime/Core/Public/Misc/CoreDelegates.h
+++ b/Engine/Source/Runtime/Core/Public/Misc/CoreDelegates.h
@@ -181,9 +181,14 @@ public:
 
 	// Callback when an ensure has occurred
 	static FOnHandleSystemEnsure OnHandleSystemEnsure;
+
 	// Callback when an error (crash) has occurred
 	static FOnHandleSystemError OnHandleSystemError;
 
+	// Callback when an error (crash) has occurred. Runs before the crash reporter.
+	// This should be used very carefully.
+	static FOnHandleSystemError OnBeforeHandleSystemError;
+
 	// Called when an actor label is changed
 	static FOnActorLabelChanged OnActorLabelChanged;
 
diff --git a/Engine/Source/Runtime/Engine/Private/EngineSessionManager.cpp b/Engine/Source/Runtime/Engine/Private/EngineSessionManager.cpp
--- a/Engine/Source/Runtime/Engine/Private/EngineSessionManager.cpp
+++ b/Engine/Source/Runtime/Engine/Private/EngineSessionManager.cpp
@@ -97,7 +97,7 @@ namespace
 void FEngineSessionManager::Initialize()
 {
 	// Register for crash and app state callbacks
-	FCoreDelegates::OnHandleSystemError.AddRaw(this, &FEngineSessionManager::OnCrashing);
+	FCoreDelegates::OnBeforeHandleSystemError.AddRaw(this, &FEngineSessionManager::OnCrashing);
 	FCoreDelegates::ApplicationHasReactivatedDelegate.AddRaw(this, &FEngineSessionManager::OnAppReactivate);
 	FCoreDelegates::ApplicationWillDeactivateDelegate.AddRaw(this, &FEngineSessionManager::OnAppDeactivate);
 	FCoreDelegates::ApplicationWillEnterBackgroundDelegate.AddRaw(this, &FEngineSessionManager::OnAppBackground);
