Use CreateOptionalDefaultSubobject to create the default static mesh component for AStaticMeshActor.
none
#rnx

diff --git a/Engine/Source/Runtime/Engine/Classes/Engine/StaticMeshActor.h b/Engine/Source/Runtime/Engine/Classes/Engine/StaticMeshActor.h
--- a/Engine/Source/Runtime/Engine/Classes/Engine/StaticMeshActor.h
+++ b/Engine/Source/Runtime/Engine/Classes/Engine/StaticMeshActor.h
@@ -61,10 +61,15 @@ protected:
 	virtual void PostEditChangeChainProperty(FPropertyChangedChainEvent& PropertyChangedEvent) override;
 #endif // WITH_EDITOR	
 	//~ End UObject Interface.
+	
+	void SetStaticMeshComponent(class UStaticMeshComponent* InStaticMeshComponent);
 
 public:
 	/** Returns StaticMeshComponent subobject **/
 	class UStaticMeshComponent* GetStaticMeshComponent() const { return StaticMeshComponent; }
+
+	/** Name of the StaticMeshComponent. Use this name if you want to prevent creation of the component (with ObjectInitializer.DoNotCreateDefaultSubobject). */
+	static FName StaticMeshComponentName;
 };
 
 
diff --git a/Engine/Source/Runtime/Engine/Private/StaticMeshActor.cpp b/Engine/Source/Runtime/Engine/Private/StaticMeshActor.cpp
--- a/Engine/Source/Runtime/Engine/Private/StaticMeshActor.cpp
+++ b/Engine/Source/Runtime/Engine/Private/StaticMeshActor.cpp
@@ -20,24 +20,35 @@
 
 
 #define LOCTEXT_NAMESPACE "StaticMeshActor"
+FName AStaticMeshActor::StaticMeshComponentName(TEXT("StaticMeshComponent0"));
 
 AStaticMeshActor::AStaticMeshActor(const FObjectInitializer& ObjectInitializer)
 	: Super(ObjectInitializer)
 {
 	bCanBeDamaged = false;
 
-	StaticMeshComponent = CreateDefaultSubobject<UStaticMeshComponent>(TEXT("StaticMeshComponent0"));
-	StaticMeshComponent->SetCollisionProfileName(UCollisionProfile::BlockAll_ProfileName);
-	StaticMeshComponent->Mobility = EComponentMobility::Static;
-	StaticMeshComponent->SetGenerateOverlapEvents(false);
-	StaticMeshComponent->bUseDefaultCollision = true;
-
-	RootComponent = StaticMeshComponent;
-
+	SetStaticMeshComponent(CreateOptionalDefaultSubobject<UStaticMeshComponent>(AStaticMeshActor::StaticMeshComponentName));
+	
 	// Only actors that are literally static mesh actors can be placed in clusters, native subclasses or BP subclasses are not safe by default
 	bCanBeInCluster = (GetClass() == AStaticMeshActor::StaticClass());
 }
 
+void AStaticMeshActor::SetStaticMeshComponent(UStaticMeshComponent* InStaticMeshComponent)
+{
+	check(StaticMeshComponent == nullptr); // Should not have been already initialized.
+
+	StaticMeshComponent = InStaticMeshComponent;
+	if(StaticMeshComponent)
+	{
+		StaticMeshComponent->SetCollisionProfileName(UCollisionProfile::BlockAll_ProfileName);
+		StaticMeshComponent->Mobility = EComponentMobility::Static;
+		StaticMeshComponent->SetGenerateOverlapEvents(false);
+		StaticMeshComponent->bUseDefaultCollision = true;
+
+		RootComponent = StaticMeshComponent;
+	}
+}
+
 void AStaticMeshActor::BeginPlay()
 {
 	// Since we allow AStaticMeshActor to specify whether it replicates via bStaticMeshReplicateMovement - per placed instance
