Fixing robomerge

========================================

- Adding a default param to GetAttachedActors to allow for not resetting the array every call.

diff --git a/Engine/Source/Runtime/Engine/Classes/GameFramework/Actor.h b/Engine/Source/Runtime/Engine/Classes/GameFramework/Actor.h
--- a/Engine/Source/Runtime/Engine/Classes/GameFramework/Actor.h
+++ b/Engine/Source/Runtime/Engine/Classes/GameFramework/Actor.h
@@ -2522,7 +2522,7 @@ public:
 	
 	/** Find all Actors which are attached directly to a component in this actor */
 	UFUNCTION(BlueprintPure, Category = "Utilities")
-	void GetAttachedActors(TArray<AActor*>& OutActors) const;
+	void GetAttachedActors(TArray<AActor*>& OutActors, bool bResetArray = true) const;
 
 	/**
 	 * Sets the ticking group for this actor.
diff --git a/Engine/Source/Runtime/Engine/Private/Actor.cpp b/Engine/Source/Runtime/Engine/Private/Actor.cpp
--- a/Engine/Source/Runtime/Engine/Private/Actor.cpp
+++ b/Engine/Source/Runtime/Engine/Private/Actor.cpp
@@ -1849,9 +1849,12 @@ void AActor::ForEachAttachedActors(TFunctionRef<bool(class AActor*)> Functor) co
 	}
 }
 
-void AActor::GetAttachedActors(TArray<class AActor*>& OutActors) const
+void AActor::GetAttachedActors(TArray<class AActor*>& OutActors, bool bResetArray) const
 {
-	OutActors.Reset();
+	if (bResetArray)
+	{
+		OutActors.Reset();
+	}
 	ForEachAttachedActors([&OutActors](AActor * Actor) { OutActors.AddUnique(Actor); return true; });
 }
 
