Consolidated uniform expression cache invalidations of a UMaterial to also invalidate child UMaterialInstance expression caches.

diff --git a/Engine/Source/Runtime/Engine/Private/Materials/Material.cpp b/Engine/Source/Runtime/Engine/Private/Materials/Material.cpp
--- a/Engine/Source/Runtime/Engine/Private/Materials/Material.cpp
+++ b/Engine/Source/Runtime/Engine/Private/Materials/Material.cpp
@@ -1155,7 +1155,6 @@ void UMaterial::OverrideTexture(const UTexture* InTextureToOverride, UTexture* O
 	if (bShouldRecacheMaterialExpressions)
 	{
 		RecacheUniformExpressions(false);
-		RecacheMaterialInstanceUniformExpressions(this);
 	}
 #endif // #if WITH_EDITOR
 }
@@ -1186,7 +1185,6 @@ void UMaterial::OverrideVectorParameterDefault(const FMaterialParameterInfo& Par
 	if (bShouldRecacheMaterialExpressions)
 	{
 		RecacheUniformExpressions(false);
-		RecacheMaterialInstanceUniformExpressions(this);
 	}
 #endif // #if WITH_EDITOR
 }
@@ -1218,7 +1216,6 @@ void UMaterial::OverrideScalarParameterDefault(const FMaterialParameterInfo& Par
 	if (bShouldRecacheMaterialExpressions)
 	{
 		RecacheUniformExpressions(false);
-		RecacheMaterialInstanceUniformExpressions(this);
 	}
 #endif // #if WITH_EDITOR
 }
@@ -1237,6 +1234,9 @@ void UMaterial::RecacheUniformExpressions(bool bRecreateUniformBuffer) const
 	{
 		DefaultMaterialInstance->CacheUniformExpressions_GameThread(bRecreateUniformBuffer);
 	}
+
+	// Need to invalidate all child material instances as well.
+	RecacheMaterialInstanceUniformExpressions(this, bRecreateUniformBuffer);
 }
 
 bool UMaterial::GetUsageByFlag(EMaterialUsage Usage) const
@@ -2803,7 +2803,6 @@ void UMaterial::UpdateMaterialShaderCacheAndTextureReferences()
 
 	//Force a recompute of the DDC key
 	CacheResourceShadersForRendering(true);
-	RecacheMaterialInstanceUniformExpressions(this);
 	
 	// Ensure that the ReferencedTextureGuids array is up to date.
 	if (GIsEditor)
@@ -4147,7 +4146,6 @@ void UMaterial::PostEditChangeProperty(FPropertyChangedEvent& PropertyChangedEve
 		// When redirecting an object pointer, we trust that the DDC hash will detect the change and that we don't need to force a recompile.
 		const bool bRegenerateId = PropertyChangedEvent.ChangeType != EPropertyChangeType::Redirected;
 		CacheResourceShadersForRendering(bRegenerateId);
-		RecacheMaterialInstanceUniformExpressions(this);
 
 		// Ensure that the ReferencedTextureGuids array is up to date.
 		if (GIsEditor)
diff --git a/Engine/Source/Runtime/Engine/Private/Materials/MaterialInstance.cpp b/Engine/Source/Runtime/Engine/Private/Materials/MaterialInstance.cpp
--- a/Engine/Source/Runtime/Engine/Private/Materials/MaterialInstance.cpp
+++ b/Engine/Source/Runtime/Engine/Private/Materials/MaterialInstance.cpp
@@ -59,7 +59,7 @@ void CacheMaterialInstanceUniformExpressions(const UMaterialInstance* MaterialIn
  * WARNING: This function is a noop outside of the Editor!
  * @param ParentMaterial - The parent material to look for.
  */
-void RecacheMaterialInstanceUniformExpressions(const UMaterialInterface* ParentMaterial)
+void RecacheMaterialInstanceUniformExpressions(const UMaterialInterface* ParentMaterial, bool bRecreateUniformBuffer)
 {
 	if (GIsEditor)
 	{
@@ -73,7 +73,7 @@ void RecacheMaterialInstanceUniformExpressions(const UMaterialInterface* ParentM
 				if (MaterialInstance->Parent == ParentMaterial)
 				{
 					UE_LOG(LogMaterial,Verbose,TEXT("--> %s"), *MaterialInstance->GetFullName());
-					CacheMaterialInstanceUniformExpressions(*It);
+					CacheMaterialInstanceUniformExpressions(*It, bRecreateUniformBuffer);
 					break;
 				}
 				new (ReentranceGuards) FMICReentranceGuard(MaterialInstance);
@@ -1259,7 +1259,6 @@ void UMaterialInstance::OverrideTexture(const UTexture* InTextureToOverride, UTe
 	if (bShouldRecacheMaterialExpressions)
 	{
 		RecacheUniformExpressions(false);
-		RecacheMaterialInstanceUniformExpressions(this);
 	}
 #endif // #if WITH_EDITOR
 }
@@ -1294,7 +1293,6 @@ void UMaterialInstance::OverrideVectorParameterDefault(const FMaterialParameterI
 	if (bShouldRecacheMaterialExpressions)
 	{
 		RecacheUniformExpressions(false);
-		RecacheMaterialInstanceUniformExpressions(this);
 	}
 #endif // #if WITH_EDITOR
 }
@@ -1329,7 +1327,6 @@ void UMaterialInstance::OverrideScalarParameterDefault(const FMaterialParameterI
 	if (bShouldRecacheMaterialExpressions)
 	{
 		RecacheUniformExpressions(false);
-		RecacheMaterialInstanceUniformExpressions(this);
 	}
 #endif // #if WITH_EDITOR
 }
@@ -3685,7 +3682,7 @@ void UMaterialInstance::PostEditChangeProperty(FPropertyChangedEvent& PropertyCh
 
 	if (PropertyChangedEvent.ChangeType == EPropertyChangeType::ValueSet || PropertyChangedEvent.ChangeType == EPropertyChangeType::ArrayClear || PropertyChangedEvent.ChangeType == EPropertyChangeType::ArrayRemove || PropertyChangedEvent.ChangeType == EPropertyChangeType::Unspecified || PropertyChangedEvent.ChangeType == EPropertyChangeType::Duplicate)
 	{
-		RecacheMaterialInstanceUniformExpressions(this);
+		RecacheMaterialInstanceUniformExpressions(this, false);
 	}
 }
 
diff --git a/Engine/Source/Runtime/Engine/Private/Materials/MaterialInstanceSupport.h b/Engine/Source/Runtime/Engine/Private/Materials/MaterialInstanceSupport.h
--- a/Engine/Source/Runtime/Engine/Private/Materials/MaterialInstanceSupport.h
+++ b/Engine/Source/Runtime/Engine/Private/Materials/MaterialInstanceSupport.h
@@ -25,7 +25,7 @@ void CacheMaterialInstanceUniformExpressions(const UMaterialInstance* MaterialIn
  * WARNING: This function is a noop outside of the Editor!
  * @param ParentMaterial - The parent material to look for.
  */
-void RecacheMaterialInstanceUniformExpressions(const UMaterialInterface* ParentMaterial);
+void RecacheMaterialInstanceUniformExpressions(const UMaterialInterface* ParentMaterial, bool bRecreateUniformBuffer);
 
 /** Protects the members of a UMaterialInstanceConstant from re-entrance. */
 class FMICReentranceGuard
diff --git a/Engine/Source/Runtime/Engine/Private/Materials/MaterialShared.cpp b/Engine/Source/Runtime/Engine/Private/Materials/MaterialShared.cpp
--- a/Engine/Source/Runtime/Engine/Private/Materials/MaterialShared.cpp
+++ b/Engine/Source/Runtime/Engine/Private/Materials/MaterialShared.cpp
@@ -2419,6 +2419,8 @@ void FMaterialRenderProxy::CacheUniformExpressions_GameThread(bool bRecreateUnif
 {
 	if (FApp::CanEverRender())
 	{
+		UE_LOG(LogMaterial, Verbose, TEXT("Caching uniform expressions for material: %s"), *GetFriendlyName());
+
 		FMaterialRenderProxy* RenderProxy = this;
 		ENQUEUE_RENDER_COMMAND(FCacheUniformExpressionsCommand)(
 			[RenderProxy, bRecreateUniformBuffer](FRHICommandListImmediate& RHICmdList)
