Fix crash when there is no default scene root node
#rnx

diff --git a/Engine/Source/Runtime/Engine/Private/SimpleConstructionScript.cpp b/Engine/Source/Runtime/Engine/Private/SimpleConstructionScript.cpp
--- a/Engine/Source/Runtime/Engine/Private/SimpleConstructionScript.cpp
+++ b/Engine/Source/Runtime/Engine/Private/SimpleConstructionScript.cpp
@@ -1053,7 +1053,7 @@ USceneComponent* USimpleConstructionScript::GetSceneRootComponentTemplate(USCS_N
 		{
 			const TArray<USCS_Node*>& SCSRootNodes = SCSStack[StackIndex]->GetRootNodes();
 
-			const bool bCanUseDefaultSceneRoot = DefaultSceneRootNode->ComponentTemplate && SCSRootNodes.Contains(DefaultSceneRootNode);
+			const bool bCanUseDefaultSceneRoot = DefaultSceneRootNode && DefaultSceneRootNode->ComponentTemplate && SCSRootNodes.Contains(DefaultSceneRootNode);
 			// Check for any scene component nodes in the root set that are not the default scene root
 			for (int32 RootNodeIndex = 0; RootNodeIndex < SCSRootNodes.Num() && RootComponentTemplate == nullptr; ++RootNodeIndex)
 			{
