Fixing build break.

diff --git a/Engine/Source/Runtime/Engine/Private/Materials/Material.cpp b/Engine/Source/Runtime/Engine/Private/Materials/Material.cpp
--- a/Engine/Source/Runtime/Engine/Private/Materials/Material.cpp
+++ b/Engine/Source/Runtime/Engine/Private/Materials/Material.cpp
@@ -1235,8 +1235,10 @@ void UMaterial::RecacheUniformExpressions(bool bRecreateUniformBuffer) const
 		DefaultMaterialInstance->CacheUniformExpressions_GameThread(bRecreateUniformBuffer);
 	}
 
+#if WITH_EDITOR
 	// Need to invalidate all child material instances as well.
 	RecacheMaterialInstanceUniformExpressions(this, bRecreateUniformBuffer);
+#endif
 }
 
 bool UMaterial::GetUsageByFlag(EMaterialUsage Usage) const
