flipped logic to only strip duplicate vertices when not supporting deferred, instead of when supporting mobile.
This fixes builds that supports both deferred&mobile, where you need the data to be not stripped.

diff --git a/Engine/Source/Runtime/Engine/Private/SkeletalMeshLODRenderData.cpp b/Engine/Source/Runtime/Engine/Private/SkeletalMeshLODRenderData.cpp
--- a/Engine/Source/Runtime/Engine/Private/SkeletalMeshLODRenderData.cpp
+++ b/Engine/Source/Runtime/Engine/Private/SkeletalMeshLODRenderData.cpp
@@ -71,7 +71,7 @@ FArchive& operator<<(FArchive& Ar, FSkelMeshRenderSection& S)
 	
 	// DuplicatedVerticesBuffer is used only for SkinCache and Editor features which is SM5 only
 	uint8 ClassDataStripFlags = 0;
-	ClassDataStripFlags |= (Ar.IsCooking() && Ar.CookingTarget()->SupportsFeature(ETargetPlatformFeatures::MobileRendering)) ? DuplicatedVertices : 0;
+	ClassDataStripFlags |= (Ar.IsCooking() && !Ar.CookingTarget()->SupportsFeature(ETargetPlatformFeatures::DeferredRendering)) ? DuplicatedVertices : 0;
 
 	// When data is cooked for server platform some of the
 	// variables are not serialized so that they're always
