Fix static analysis error

diff --git a/Engine/Source/Runtime/Engine/Classes/Sound/SoundCueTemplate.h b/Engine/Source/Runtime/Engine/Classes/Sound/SoundCueTemplate.h
--- a/Engine/Source/Runtime/Engine/Classes/Sound/SoundCueTemplate.h
+++ b/Engine/Source/Runtime/Engine/Classes/Sound/SoundCueTemplate.h
@@ -164,11 +164,14 @@ protected:
 			}
 		}
 
-		if (InputPinDisplayName)
+		if (InputPin)
 		{
-			InputPin->PinFriendlyName = *InputPinDisplayName;
+			if (InputPinDisplayName)
+			{
+				InputPin->PinFriendlyName = *InputPinDisplayName;
+			}
+			InputPin->MakeLinkTo(OutputPin);
 		}
-		InputPin->MakeLinkTo(OutputPin);
 	}
 #endif // WITH_EDITOR
 };
