some unused resources(structured buffers) are null on d3d10. Point to a dummy black resource, in order to not hit checks().

diff --git a/Engine/Source/Runtime/Engine/Private/SceneManagement.cpp b/Engine/Source/Runtime/Engine/Private/SceneManagement.cpp
--- a/Engine/Source/Runtime/Engine/Private/SceneManagement.cpp
+++ b/Engine/Source/Runtime/Engine/Private/SceneManagement.cpp
@@ -737,6 +737,16 @@ FViewUniformShaderParameters::FViewUniformShaderParameters()
 
 	PrimitiveSceneData = GIdentityPrimitiveBuffer.PrimitiveSceneDataBufferSRV;
 	LightmapSceneData = GIdentityPrimitiveBuffer.LightmapSceneDataBufferSRV;
+
+	//this can be deleted once sm4 support is removed.
+	if (!PrimitiveSceneData)
+	{
+		PrimitiveSceneData = GBlackTextureWithSRV->ShaderResourceViewRHI;
+	}
+	if (!LightmapSceneData)
+	{
+		LightmapSceneData = GBlackTextureWithSRV->ShaderResourceViewRHI;
+	}
 }
 
 FInstancedViewUniformShaderParameters::FInstancedViewUniformShaderParameters()
