Fixed unintended double -> int conversion in FDumpFPSChartToEndpoint::DumpChart

diff --git a/Engine/Source/Runtime/Engine/Private/ChartCreation.cpp b/Engine/Source/Runtime/Engine/Private/ChartCreation.cpp
--- a/Engine/Source/Runtime/Engine/Private/ChartCreation.cpp
+++ b/Engine/Source/Runtime/Engine/Private/ChartCreation.cpp
@@ -207,7 +207,7 @@ void FDumpFPSChartToEndpoint::DumpChart(double InWallClockTimeFromStartOfChartin
 {
 	FillOutMemberStats();
 
-	int32 TotalTime = Chart.GetTotalTime();
+	double TotalTime = Chart.GetTotalTime();
 	WallClockTimeFromStartOfCharting = InWallClockTimeFromStartOfCharting;
 	int32 NumFrames = (int32)Chart.GetNumFrames();
 	MapName = MoveTemp(InMapName);
