Fix issue where missing NetFieldExport's in replays could either assert or potentially break other replay state.
Now, they will properly cause property incompatibility.

diff --git a/Engine/Source/Runtime/Engine/Private/RepLayout.cpp b/Engine/Source/Runtime/Engine/Private/RepLayout.cpp
--- a/Engine/Source/Runtime/Engine/Private/RepLayout.cpp
+++ b/Engine/Source/Runtime/Engine/Private/RepLayout.cpp
@@ -6221,6 +6221,12 @@ bool FRepLayout::DeltaSerializeFastArrayProperty(FFastArrayDeltaSerializeParams&
 
 		if (!LocalNetFieldExportGroup.IsValid())
 		{
+			if (!bIsWriting)
+			{
+				UE_LOG(LogRep, Error, TEXT("DeltaSerializeFastArrayProperty: Unable to find NetFieldExportGroup during replay playback. Class=%s, Property=%s"), *Owner->GetName(), *Parent.CachedPropertyName.ToString());
+				return false;
+			}
+
 			UE_LOG(LogRepProperties, VeryVerbose, TEXT("DeltaSerializeFastArrayProperty: Create Netfield Export Group."))
 			LocalNetFieldExportGroup = CreateNetfieldExportGroup();
 			PackageMap->AddNetFieldExportGroup(OwnerPathName, LocalNetFieldExportGroup);
