SourceEffect references are not protected by GC as they are not referenced as properties.  Keep them alive as long as an FActiveSound is alive to avoid GC'ing out from underneath the audio engine.

diff --git a/Engine/Source/Runtime/Engine/Private/ActiveSound.cpp b/Engine/Source/Runtime/Engine/Private/ActiveSound.cpp
--- a/Engine/Source/Runtime/Engine/Private/ActiveSound.cpp
+++ b/Engine/Source/Runtime/Engine/Private/ActiveSound.cpp
@@ -166,6 +166,11 @@ void FActiveSound::AddReferencedObjects( FReferenceCollector& Collector)
 	Collector.AddReferencedObject(SourceEffectChain);
 	Collector.AddReferencedObject(SoundClassOverride);
 
+	for (FSourceEffectChainEntry& SourceEffect : SourceEffectChain->Chain)
+	{
+		Collector.AddReferencedObject(SourceEffect.Preset);
+	}
+
 	for (USoundConcurrency* Concurrency : ConcurrencySet)
 	{
 		if (Concurrency)
