Fix null deref

diff --git a/Engine/Source/Runtime/Engine/Private/ActiveSound.cpp b/Engine/Source/Runtime/Engine/Private/ActiveSound.cpp
--- a/Engine/Source/Runtime/Engine/Private/ActiveSound.cpp
+++ b/Engine/Source/Runtime/Engine/Private/ActiveSound.cpp
@@ -166,9 +166,12 @@ void FActiveSound::AddReferencedObjects( FReferenceCollector& Collector)
 	Collector.AddReferencedObject(SourceEffectChain);
 	Collector.AddReferencedObject(SoundClassOverride);
 
-	for (FSourceEffectChainEntry& SourceEffect : SourceEffectChain->Chain)
+	if (SourceEffectChain)
 	{
-		Collector.AddReferencedObject(SourceEffect.Preset);
+		for (FSourceEffectChainEntry& SourceEffect : SourceEffectChain->Chain)
+		{
+			Collector.AddReferencedObject(SourceEffect.Preset);
+		}
 	}
 
 	for (USoundConcurrency* Concurrency : ConcurrencySet)
