Submitting shader compiler changes from DK branch. Gives about a ~30% speedup from a cold DDC

#fyo wang.hao

diff --git a/Engine/Source/Runtime/Engine/Private/ShaderCompiler/ShaderCompiler.cpp b/Engine/Source/Runtime/Engine/Private/ShaderCompiler/ShaderCompiler.cpp
--- a/Engine/Source/Runtime/Engine/Private/ShaderCompiler/ShaderCompiler.cpp
+++ b/Engine/Source/Runtime/Engine/Private/ShaderCompiler/ShaderCompiler.cpp
@@ -1041,7 +1041,9 @@ int32 FShaderCompileThreadRunnable::PullTasksFromQueue()
 		FScopeLock Lock(&Manager->CompileQueueSection);
 
 		const int32 NumWorkersToFeed = Manager->bCompilingDuringGame ? Manager->NumShaderCompilingThreadsDuringGame : WorkerInfos.Num();
-
+		// Try to distribute the work evenly between the workers
+		const auto NumJobsPerWorker = (Manager->CompileQueue.Num() / NumWorkersToFeed) + 1;
+		
 		for (int32 WorkerIndex = 0; WorkerIndex < WorkerInfos.Num(); WorkerIndex++)
 		{
 			FShaderCompileWorkerInfo& CurrentWorkerInfo = *WorkerInfos[WorkerIndex];
@@ -1053,17 +1055,14 @@ int32 FShaderCompileThreadRunnable::PullTasksFromQueue()
 
 				if (Manager->CompileQueue.Num() > 0)
 				{
-					if (Manager->CompileQueue.Num() % 10 == 0)
-					{
-						UE_LOG(LogShaderCompilers, Display, TEXT("Shaders left to compile %i"), Manager->CompileQueue.Num());
-					}
+					UE_LOG(LogShaderCompilers, Display, TEXT("Worker (%d/%d): shaders left to compile %i"), WorkerIndex + 1, WorkerInfos.Num(), Manager->CompileQueue.Num());
 
 					bool bAddedLowLatencyTask = false;
+					const auto MaxNumJobs = FMath::Min3(NumJobsPerWorker, Manager->CompileQueue.Num(), Manager->MaxShaderJobBatchSize);
+					
 					int32 JobIndex = 0;
-
-					// Try to grab up to MaxShaderJobBatchSize jobs
 					// Don't put more than one low latency task into a batch
-					for (; JobIndex < Manager->MaxShaderJobBatchSize && JobIndex < Manager->CompileQueue.Num() && !bAddedLowLatencyTask; JobIndex++)
+					for (; JobIndex < MaxNumJobs && !bAddedLowLatencyTask; JobIndex++)
 					{
 						bAddedLowLatencyTask |= Manager->CompileQueue[JobIndex]->bOptimizeForLowLatency;
 						CurrentWorkerInfo.QueuedJobs.Add(Manager->CompileQueue[JobIndex]);
