Correctly propagate any curve atlas value changes for material instances to the scalar parameter value array
#rnx

diff --git a/Engine/Source/Runtime/Engine/Private/Materials/MaterialInstance.cpp b/Engine/Source/Runtime/Engine/Private/Materials/MaterialInstance.cpp
--- a/Engine/Source/Runtime/Engine/Private/Materials/MaterialInstance.cpp
+++ b/Engine/Source/Runtime/Engine/Private/Materials/MaterialInstance.cpp
@@ -466,23 +466,31 @@ void UMaterialInstance::InitResources()
 		Resource->GameThread_SetParent(SafeParent);
 	}
 
-	GameThread_InitMIParameters(this, ScalarParameterValues);
-	GameThread_InitMIParameters(this, VectorParameterValues);
-	GameThread_InitMIParameters(this, TextureParameterValues);
-	GameThread_InitMIParameters(this, FontParameterValues);
-
 #if WITH_EDITOR
 	//recalculate any scalar params based on a curve position in an atlas in case the atlas changed
-	for (FScalarParameterValue ScalarParam : ScalarParameterValues)
+	for (FScalarParameterValue& ScalarParam : ScalarParameterValues)
 	{
-		IsScalarParameterUsedAsAtlasPosition(ScalarParam.ParameterInfo, ScalarParam.AtlasData.bIsUsedAsAtlasPosition, ScalarParam.AtlasData.Curve, ScalarParam.AtlasData.Atlas);
 		if (ScalarParam.AtlasData.bIsUsedAsAtlasPosition)
 		{
-			SetScalarParameterAtlasInternal(ScalarParam.ParameterInfo, ScalarParam.AtlasData);
+			UCurveLinearColorAtlas* Atlas = Cast<UCurveLinearColorAtlas>(ScalarParam.AtlasData.Atlas.Get());
+			UCurveLinearColor* Curve = Cast<UCurveLinearColor>(ScalarParam.AtlasData.Curve.Get());
+			if (Curve && Atlas)
+			{
+				int32 Index = Atlas->GradientCurves.Find(Curve);
+				if (Index != INDEX_NONE)
+				{
+					ScalarParam.ParameterValue = (float)Index;
+				}
+			}
 		}
 	}
 #endif
 
+	GameThread_InitMIParameters(this, ScalarParameterValues);
+	GameThread_InitMIParameters(this, VectorParameterValues);
+	GameThread_InitMIParameters(this, TextureParameterValues);
+	GameThread_InitMIParameters(this, FontParameterValues);
+
 	PropagateDataToMaterialProxy();
 
 	CacheMaterialInstanceUniformExpressions(this);
