When bIsStatic is changed on a streaming level move the loaded level in to the correct collection

diff --git a/Engine/Source/Runtime/Engine/Private/LevelStreaming.cpp b/Engine/Source/Runtime/Engine/Private/LevelStreaming.cpp
--- a/Engine/Source/Runtime/Engine/Private/LevelStreaming.cpp
+++ b/Engine/Source/Runtime/Engine/Private/LevelStreaming.cpp
@@ -1548,8 +1548,7 @@ void ULevelStreaming::PostEditChangeProperty(FPropertyChangedEvent& PropertyChan
 		{
 			GetWorld()->UpdateLevelStreaming();
 		}
-
-		if (PropertyName == GET_MEMBER_NAME_CHECKED(ULevelStreaming, EditorStreamingVolumes))
+		else if (PropertyName == GET_MEMBER_NAME_CHECKED(ULevelStreaming, EditorStreamingVolumes))
 		{
 			RemoveStreamingVolumeDuplicates();
 
@@ -1574,6 +1573,23 @@ void ULevelStreaming::PostEditChangeProperty(FPropertyChangedEvent& PropertyChan
 			bHasCachedWorldAssetPackageFName = false;
 			bHasCachedLoadedLevelPackageName = false;
 		}
+		else if (PropertyName == GET_MEMBER_NAME_CHECKED(ULevelStreaming, bIsStatic))
+		{
+			if (LoadedLevel)
+			{
+				const ELevelCollectionType NewCollectionType = bIsStatic ? ELevelCollectionType::StaticLevels : ELevelCollectionType::DynamicSourceLevels;
+				FLevelCollection* PreviousCollection = LoadedLevel->GetCachedLevelCollection();
+
+				if (PreviousCollection && PreviousCollection->GetType() != NewCollectionType)
+				{
+					PreviousCollection->RemoveLevel(LoadedLevel);
+
+					UWorld* World = GetWorld();
+					FLevelCollection& LC = World->FindOrAddCollectionByType(NewCollectionType);
+					LC.AddLevel(LoadedLevel);
+				}
+			}
+		}
 	}
 
 	Super::PostEditChangeProperty(PropertyChangedEvent);
