Fixed 4.23 Chaos defaulting to TaskGraph which is currently unimplemented. Now defaults to dedicated thread.

diff --git a/Engine/Source/Runtime/Engine/Private/PhysicsEngine/PhysicsSettings.cpp b/Engine/Source/Runtime/Engine/Private/PhysicsEngine/PhysicsSettings.cpp
--- a/Engine/Source/Runtime/Engine/Private/PhysicsEngine/PhysicsSettings.cpp
+++ b/Engine/Source/Runtime/Engine/Private/PhysicsEngine/PhysicsSettings.cpp
@@ -163,7 +163,7 @@ void UPhysicsSettings::LoadSurfaceType()
 #endif	// WITH_EDITOR
 
 FChaosPhysicsSettings::FChaosPhysicsSettings() 
-	: DefaultThreadingModel(EChaosThreadingMode::TaskGraph)
+	: DefaultThreadingModel(EChaosThreadingMode::DedicatedThread)
 	, DedicatedThreadTickMode(EChaosSolverTickMode::VariableCappedWithTarget)
 	, DedicatedThreadBufferMode(EChaosBufferMode::Double)
 {
diff --git a/Engine/Source/Runtime/PhysicsCore/Public/PhysicsCoreTypes.h b/Engine/Source/Runtime/PhysicsCore/Public/PhysicsCoreTypes.h
--- a/Engine/Source/Runtime/PhysicsCore/Public/PhysicsCoreTypes.h
+++ b/Engine/Source/Runtime/PhysicsCore/Public/PhysicsCoreTypes.h
@@ -20,7 +20,7 @@ UENUM()
 enum class EChaosThreadingMode : uint8
 {
 	DedicatedThread,
-	TaskGraph,
+	TaskGraph UMETA(Hidden),
 	SingleThread,
 	Num UMETA(Hidden),
 	Invalid UMETA(Hidden)
