Change asserts in PhysScene_PhysX::UpdateKinematicsOnDefferedSkelMeshes() to ensure from check.

diff --git a/Engine/Source/Runtime/Engine/Private/PhysicsEngine/PhysScene_PhysX.cpp b/Engine/Source/Runtime/Engine/Private/PhysicsEngine/PhysScene_PhysX.cpp
--- a/Engine/Source/Runtime/Engine/Private/PhysicsEngine/PhysScene_PhysX.cpp
+++ b/Engine/Source/Runtime/Engine/Private/PhysicsEngine/PhysScene_PhysX.cpp
@@ -963,7 +963,14 @@ void FPhysScene_PhysX::MarkForPreSimKinematicUpdate(USkeletalMeshComponent* InSk
 		if (InSkelComp->bDeferredKinematicUpdate)
 		{
 			TPair<USkeletalMeshComponent*, FDeferredKinematicUpdateInfo>* FoundItem = DeferredKinematicUpdateSkelMeshes.FindByPredicate([InSkelComp](const TPair<USkeletalMeshComponent*, FDeferredKinematicUpdateInfo>& InItem) { return InSkelComp == InItem.Key; });
-			check(FoundItem != nullptr); // If the bool was set, we must be in the array!
+			if (!ensure(FoundItem != nullptr))// If the bool was set, we must be in the array!
+			{
+				FDeferredKinematicUpdateInfo Info;
+				Info.TeleportType = InTeleport;
+				Info.bNeedsSkinning = bNeedsSkinning;
+				DeferredKinematicUpdateSkelMeshes.Emplace(InSkelComp, Info);
+				return;
+			}
 
 			FDeferredKinematicUpdateInfo& Info = FoundItem->Value;
 
@@ -1002,7 +1009,7 @@ void FPhysScene_PhysX::ClearPreSimKinematicUpdate(USkeletalMeshComponent* InSkel
 		// Remove from map
 		int32 NumRemoved = DeferredKinematicUpdateSkelMeshes.RemoveAll([InSkelComp](const TPair<USkeletalMeshComponent*, FDeferredKinematicUpdateInfo>& InItem) { return InSkelComp == InItem.Key; });
 
-		check(NumRemoved == 1); // Should be in array if flag was set!
+		ensure(NumRemoved == 1); // Should be in array if flag was set!
 
 		// Clear flag
 		InSkelComp->bDeferredKinematicUpdate = false;
@@ -1019,7 +1026,7 @@ void FPhysScene_PhysX::UpdateKinematicsOnDeferredSkelMeshes()
 		USkeletalMeshComponent* SkelComp = DeferredKinematicUpdate.Key;
 		const FDeferredKinematicUpdateInfo& Info = DeferredKinematicUpdate.Value;
 
-		check(SkelComp->bDeferredKinematicUpdate); // Should be true if in map!
+		ensure(SkelComp->bDeferredKinematicUpdate); // Should be true if in map!
 
 		// Perform kinematic updates
 		SkelComp->UpdateKinematicBonesToAnim(SkelComp->GetComponentSpaceTransforms(), Info.TeleportType, Info.bNeedsSkinning, EAllowKinematicDeferral::DisallowDeferral);
