Exposed GetVectorPropertyAtSplinePoint to blueprints

diff --git a/Engine/Source/Runtime/Engine/Classes/Components/SplineComponent.h b/Engine/Source/Runtime/Engine/Classes/Components/SplineComponent.h
--- a/Engine/Source/Runtime/Engine/Classes/Components/SplineComponent.h
+++ b/Engine/Source/Runtime/Engine/Classes/Components/SplineComponent.h
@@ -521,6 +521,10 @@ public:
     /** Get a metadata property float value along the spline at spline point */
 	UFUNCTION(BlueprintCallable, Category = Spline)
 	float GetFloatPropertyAtSplinePoint(int32 Index, FName PropertyName) const;
+	
+ 	/** Get a metadata property vector value along the spline at spline point */
+	UFUNCTION(BlueprintCallable, Category = Spline)
+	FVector GetVectorPropertyAtSplinePoint(int32 Index, FName PropertyName) const;
 
 	/** Returns total length along this spline */
 	UFUNCTION(BlueprintCallable, Category=Spline) 
diff --git a/Engine/Source/Runtime/Engine/Private/Components/SplineComponent.cpp b/Engine/Source/Runtime/Engine/Private/Components/SplineComponent.cpp
--- a/Engine/Source/Runtime/Engine/Private/Components/SplineComponent.cpp
+++ b/Engine/Source/Runtime/Engine/Private/Components/SplineComponent.cpp
@@ -1455,6 +1455,11 @@ float USplineComponent::GetFloatPropertyAtSplinePoint(int32 Index, FName Propert
 	return GetPropertyValueAtSplinePoint<float>(GetSplinePointsMetadata(), Index, PropertyName);
 }
 
+FVector USplineComponent::GetVectorPropertyAtSplinePoint(int32 Index, FName PropertyName) const
+{
+	return GetPropertyValueAtSplinePoint<FVector>(GetSplinePointsMetadata(), Index, PropertyName);
+}
+
 #if !UE_BUILD_SHIPPING
 FPrimitiveSceneProxy* USplineComponent::CreateSceneProxy()
 {
