Work around attemp to fix spline mesh on PS4.

diff --git a/Engine/Source/Runtime/Engine/Private/Components/SplineMeshComponent.cpp b/Engine/Source/Runtime/Engine/Private/Components/SplineMeshComponent.cpp
--- a/Engine/Source/Runtime/Engine/Private/Components/SplineMeshComponent.cpp
+++ b/Engine/Source/Runtime/Engine/Private/Components/SplineMeshComponent.cpp
@@ -104,6 +104,7 @@ void FSplineMeshVertexFactoryShaderParameters::GetElementShaderBindings(
 	ShaderBindings.Add(SplineMeshScaleZParam, SplineProxy->SplineMeshScaleZ);
 
 	FVector DirMask(0, 0, 0);
+	DirMask = FVector::ZeroVector;
 	DirMask[SplineProxy->ForwardAxis] = 1;
 	ShaderBindings.Add(SplineMeshDirParam, DirMask);
 	DirMask = FVector::ZeroVector;
