Fix potential leak of AudioComponents when sounds fail to start and don't report back to the AudioComponent that playback finished (and failed)

diff --git a/Engine/Source/Runtime/Engine/Private/AudioDevice.cpp b/Engine/Source/Runtime/Engine/Private/AudioDevice.cpp
--- a/Engine/Source/Runtime/Engine/Private/AudioDevice.cpp
+++ b/Engine/Source/Runtime/Engine/Private/AudioDevice.cpp
@@ -4132,8 +4132,23 @@ void FAudioDevice::AddNewActiveSoundInternal(const FActiveSound& NewActiveSound,
 {
 	LLM_SCOPE(ELLMTag::AudioMisc);
 
+	if (!IsInAudioThread())
+	{
+		DECLARE_CYCLE_STAT(TEXT("FAudioThreadTask.AddNewActiveSound"), STAT_AudioAddNewActiveSound, STATGROUP_AudioThreadCommands);
+
+		FAudioDevice* AudioDevice = this;
+		FAudioThread::RunCommandOnAudioThread([AudioDevice, NewActiveSound, VirtualLoopToRetrigger]()
+		{
+			AudioDevice->AddNewActiveSoundInternal(NewActiveSound, VirtualLoopToRetrigger);
+		}, GET_STATID(STAT_AudioAddNewActiveSound));
+
+		return;
+	}
+
+
 	if (NewActiveSound.Sound == nullptr)
 	{
+		UAudioComponent::PlaybackCompleted(NewActiveSound.AudioComponentID, true);
 		return;
 	}
 
@@ -4143,6 +4158,7 @@ void FAudioDevice::AddNewActiveSoundInternal(const FActiveSound& NewActiveSound,
 		USoundSourceBus* Bus = Cast<USoundSourceBus>(NewActiveSound.Sound);
 		if (Bus)
 		{
+			UAudioComponent::PlaybackCompleted(NewActiveSound.AudioComponentID, true);
 			return;
 		}
 	}
@@ -4154,23 +4170,12 @@ void FAudioDevice::AddNewActiveSoundInternal(const FActiveSound& NewActiveSound,
 		if (!SoundIsAudible(NewActiveSound))
 		{
 			UE_LOG(LogAudio, Log, TEXT("New ActiveSound not created for out of range Sound %s"), *NewActiveSound.Sound->GetName());
+
+			UAudioComponent::PlaybackCompleted(NewActiveSound.AudioComponentID, true);
 			return;
 		}
 	}
 
-	if (!IsInAudioThread())
-	{
-		DECLARE_CYCLE_STAT(TEXT("FAudioThreadTask.AddNewActiveSound"), STAT_AudioAddNewActiveSound, STATGROUP_AudioThreadCommands);
-
-		FAudioDevice* AudioDevice = this;
-		FAudioThread::RunCommandOnAudioThread([AudioDevice, NewActiveSound, VirtualLoopToRetrigger]()
-		{
-			AudioDevice->AddNewActiveSoundInternal(NewActiveSound, VirtualLoopToRetrigger);
-		}, GET_STATID(STAT_AudioAddNewActiveSound));
-
-		return;
-	}
-
 	// Cull one-shot active sounds if we've reached our max limit of one shot active sounds before we attempt to evaluate concurrency
 	// Check for debug sound name
 #if !UE_BUILD_SHIPPING
@@ -4183,6 +4188,7 @@ void FAudioDevice::AddNewActiveSoundInternal(const FActiveSound& NewActiveSound,
 		NewActiveSound.Sound->GetName(SoundName);
 		if (!SoundName.Contains(DebugSound))
 		{
+			UAudioComponent::PlaybackCompleted(NewActiveSound.AudioComponentID, true);
 			return;
 		}
 	}
@@ -4196,7 +4202,10 @@ void FAudioDevice::AddNewActiveSoundInternal(const FActiveSound& NewActiveSound,
 	{
 		FString SoundWaveName;
 		SoundWave->GetName(SoundWaveName);
-		UE_LOG(LogAudio, Warning, TEXT("Replaying a procedural sound '%s' without stopping the previous instance. Only one sound instance per procedural sound wave is supported."), *SoundWaveName)
+
+		UE_LOG(LogAudio, Warning, TEXT("Replaying a procedural sound '%s' without stopping the previous instance. Only one sound instance per procedural sound wave is supported."), *SoundWaveName);
+
+		UAudioComponent::PlaybackCompleted(NewActiveSound.AudioComponentID, true);
 		return;
 	}
 
@@ -4234,8 +4243,11 @@ void FAudioDevice::AddNewActiveSoundInternal(const FActiveSound& NewActiveSound,
 			{
 				UE_LOG(LogAudio, Verbose, TEXT("New ActiveSound %s Virtualizing: Failed to pass concurrency"), *Sound->GetName());
 				AddVirtualLoop(VirtualLoop);
+				return;
 			}
 		}
+
+		UAudioComponent::PlaybackCompleted(NewActiveSound.AudioComponentID, true);
 		return;
 	}
 
@@ -4343,10 +4355,7 @@ bool FAudioDevice::RemoveVirtualLoop(FActiveSound& InActiveSound)
 		check(InActiveSound.bIsStopping);
 
 		const uint64 ComponentID = InActiveSound.GetAudioComponentID();
-		if (ComponentID > 0)
-		{
-			UAudioComponent::PlaybackCompleted(ComponentID, false);
-		}
+		UAudioComponent::PlaybackCompleted(ComponentID, false);
 
 #if !(UE_BUILD_SHIPPING || UE_BUILD_TEST)
 		if (InActiveSound.Sound)
@@ -4553,10 +4562,7 @@ void FAudioDevice::RemoveActiveSound(FActiveSound* ActiveSound)
 
 	// Perform the notification if not sound not set to re-trigger
 	const uint64 ComponentID = ActiveSound->GetAudioComponentID();
-	if (ComponentID > 0)
-	{
-		UAudioComponent::PlaybackCompleted(ComponentID, false);
-	}
+	UAudioComponent::PlaybackCompleted(ComponentID, false);
 
 	const int32 NumRemoved = ActiveSounds.RemoveSwap(ActiveSound);
 	if (!ensureMsgf(NumRemoved > 0, TEXT("Attempting to remove an already removed ActiveSound '%s'"), ActiveSound->Sound ? *ActiveSound->Sound->GetName() : TEXT("N/A")))
diff --git a/Engine/Source/Runtime/Engine/Private/Components/AudioComponent.cpp b/Engine/Source/Runtime/Engine/Private/Components/AudioComponent.cpp
--- a/Engine/Source/Runtime/Engine/Private/Components/AudioComponent.cpp
+++ b/Engine/Source/Runtime/Engine/Private/Components/AudioComponent.cpp
@@ -739,6 +739,11 @@ void UAudioComponent::PlaybackCompleted(uint64 AudioComponentID, bool bFailedToS
 {
 	check(IsInAudioThread());
 
+	if (AudioComponentID == 0)
+	{
+		return;
+	}
+
 	DECLARE_CYCLE_STAT(TEXT("FGameThreadAudioTask.PlaybackCompleted"), STAT_AudioPlaybackCompleted, STATGROUP_TaskGraphTasks);
 
 	FAudioThread::RunCommandOnGameThread([AudioComponentID, bFailedToStart]()
