Fix for concurrency not being observed properly when 'Limit To Owner' is set

diff --git a/Engine/Source/Runtime/Engine/Private/SoundConcurrency.cpp b/Engine/Source/Runtime/Engine/Private/SoundConcurrency.cpp
--- a/Engine/Source/Runtime/Engine/Private/SoundConcurrency.cpp
+++ b/Engine/Source/Runtime/Engine/Private/SoundConcurrency.cpp
@@ -291,22 +291,25 @@ void FSoundConcurrencyManager::CreateNewGroupsFromHandles(
 
 			case EConcurrencyMode::OwnerPerSound:
 			{
-				const uint32 OwnerObjectID = NewActiveSound.GetOwnerID();
+				USoundBase* Sound = NewActiveSound.GetSound();
+				check(Sound);
+
+				const FSoundObjectID SoundObjectID = static_cast<FSoundObjectID>(Sound->GetUniqueID());
+				const FSoundOwnerObjectID OwnerObjectID = NewActiveSound.GetOwnerID();
+
 				if (FSoundInstanceEntry* InstanceEntry = OwnerPerSoundConcurrencyMap.Find(OwnerObjectID))
 				{
-					USoundBase* Sound = NewActiveSound.GetSound();
-					check(Sound);
-					if (!InstanceEntry->SoundInstanceToConcurrencyGroup.Contains(Sound->GetUniqueID()))
+					if (!InstanceEntry->SoundInstanceToConcurrencyGroup.Contains(SoundObjectID))
 					{
 						FConcurrencyGroup& ConcurrencyGroup = CreateNewConcurrencyGroup(ConcurrencyHandle);
-						InstanceEntry->SoundInstanceToConcurrencyGroup.Add(ConcurrencyHandle.ObjectID, ConcurrencyGroup.GetGroupID());
+						InstanceEntry->SoundInstanceToConcurrencyGroup.Add(SoundObjectID, ConcurrencyGroup.GetGroupID());
 						OutGroupsToApply.Add(&ConcurrencyGroup);
 					}
 				}
 				else
 				{
 					FConcurrencyGroup& ConcurrencyGroup = CreateNewConcurrencyGroup(ConcurrencyHandle);
-					OwnerPerSoundConcurrencyMap.Emplace(OwnerObjectID, FSoundInstanceEntry(ConcurrencyHandle.ObjectID, ConcurrencyGroup.GetGroupID()));
+					OwnerPerSoundConcurrencyMap.Emplace(OwnerObjectID, FSoundInstanceEntry(SoundObjectID, ConcurrencyGroup.GetGroupID()));
 					OutGroupsToApply.Add(&ConcurrencyGroup);
 				}
 			}
