Didn't follow chain far enough to dangling UObject references in FSoundEffectBase::Instances.

diff --git a/Engine/Source/Runtime/Engine/Classes/Sound/SoundEffectBase.h b/Engine/Source/Runtime/Engine/Classes/Sound/SoundEffectBase.h
--- a/Engine/Source/Runtime/Engine/Classes/Sound/SoundEffectBase.h
+++ b/Engine/Source/Runtime/Engine/Classes/Sound/SoundEffectBase.h
@@ -25,7 +25,7 @@
 #define GET_EFFECT_SETTINGS(EFFECT_NAME) \
 		U##EFFECT_NAME##Preset* _Preset = CastChecked<U##EFFECT_NAME##Preset>(Preset); \
 		F##EFFECT_NAME##Settings Settings = _Preset->GetSettings(); \
-	
+
 #define EFFECT_PRESET_METHODS(EFFECT_NAME) \
 		virtual FText GetAssetActionName() const override { return FText::FromString(#EFFECT_NAME); } \
 		virtual UClass* GetSupportedClass() const override { return EFFECT_PRESET_NAME(EFFECT_NAME)::StaticClass(); } \
@@ -83,6 +83,8 @@ public:
 
 	void SetPreset(USoundEffectPreset* Inpreset);
 
+	USoundEffectPreset* GetPreset();
+
 	/** Removes the instance from the preset. */
 	void ClearPreset();
 
diff --git a/Engine/Source/Runtime/Engine/Classes/Sound/SoundEffectPreset.h b/Engine/Source/Runtime/Engine/Classes/Sound/SoundEffectPreset.h
--- a/Engine/Source/Runtime/Engine/Classes/Sound/SoundEffectPreset.h
+++ b/Engine/Source/Runtime/Engine/Classes/Sound/SoundEffectPreset.h
@@ -2,6 +2,7 @@
 #pragma once
 
 #include "CoreMinimal.h"
+#include "UObject/UObjectGlobals.h"
 #include "UObject/ObjectMacros.h"
 #include "UObject/Object.h"
 #include "SoundEffectPreset.generated.h"
@@ -19,13 +20,13 @@ class ENGINE_API USoundEffectPreset : public UObject
 public:
 	USoundEffectPreset(const FObjectInitializer& ObjectInitializer);
 	virtual ~USoundEffectPreset();
-	
+
 	virtual FText GetAssetActionName() const PURE_VIRTUAL(USoundEffectPreset::GetAssetActionName, return FText(););
 	virtual UClass* GetSupportedClass() const PURE_VIRTUAL(USoundEffectPreset::GetSupportedClass, return nullptr;);
 	virtual USoundEffectPreset* CreateNewPreset(UObject* InParent, FName Name, EObjectFlags Flags) const PURE_VIRTUAL(USoundEffectPreset::CreateNewPreset, return nullptr;);
 	virtual FSoundEffectBase* CreateNewEffect() const PURE_VIRTUAL(USoundEffectPreset::CreateNewEffect, return nullptr;);
 	virtual bool HasAssetActions() const { return true; }
-	virtual void Init() PURE_VIRTUAL(USoundEffectPreset::Init, ); 
+	virtual void Init() PURE_VIRTUAL(USoundEffectPreset::Init, );
 	virtual void OnInit() {};
 	virtual FColor GetPresetColor() const { return FColor(200.0f, 100.0f, 100.0f); }
 
@@ -34,6 +35,8 @@ public:
 	void AddEffectInstance(FSoundEffectBase* InSource);
 	void RemoveEffectInstance(FSoundEffectBase* InSource);
 
+	void AddReferencedEffects(FReferenceCollector& Collector);
+
 protected:
 
 #if WITH_EDITORONLY_DATA
diff --git a/Engine/Source/Runtime/Engine/Private/ActiveSound.cpp b/Engine/Source/Runtime/Engine/Private/ActiveSound.cpp
--- a/Engine/Source/Runtime/Engine/Private/ActiveSound.cpp
+++ b/Engine/Source/Runtime/Engine/Private/ActiveSound.cpp
@@ -162,16 +162,18 @@ void FActiveSound::AddReferencedObjects( FReferenceCollector& Collector)
 		}
 	}
 
-	Collector.AddReferencedObject(Sound);
-	Collector.AddReferencedObject(SourceEffectChain);
 	Collector.AddReferencedObject(SoundClassOverride);
 
+	Collector.AddReferencedObject(SourceEffectChain);
 	if (SourceEffectChain)
 	{
-		for (FSourceEffectChainEntry& SourceEffect : SourceEffectChain->Chain)
-		{
-			Collector.AddReferencedObject(SourceEffect.Preset);
-		}
+		SourceEffectChain->AddReferencedEffects(Collector);
+	}
+
+	Collector.AddReferencedObject(Sound);
+	if (Sound && Sound->SourceEffectChain)
+	{
+		Sound->SourceEffectChain->AddReferencedEffects(Collector);
 	}
 
 	for (USoundConcurrency* Concurrency : ConcurrencySet)
diff --git a/Engine/Source/Runtime/Engine/Private/SoundEffectBase.cpp b/Engine/Source/Runtime/Engine/Private/SoundEffectBase.cpp
--- a/Engine/Source/Runtime/Engine/Private/SoundEffectBase.cpp
+++ b/Engine/Source/Runtime/Engine/Private/SoundEffectBase.cpp
@@ -16,13 +16,13 @@ FSoundEffectBase::~FSoundEffectBase()
 }
 
 bool FSoundEffectBase::IsActive() const
-{ 
-	return bIsActive; 
+{
+	return bIsActive;
 }
 
 void FSoundEffectBase::SetEnabled(const bool bInIsEnabled)
-{ 
-	bIsActive = bInIsEnabled; 
+{
+	bIsActive = bInIsEnabled;
 }
 
 void FSoundEffectBase::SetPreset(USoundEffectPreset* Inpreset)
@@ -43,6 +43,11 @@ void FSoundEffectBase::SetPreset(USoundEffectPreset* Inpreset)
 	bChanged = true;
 }
 
+USoundEffectPreset* FSoundEffectBase::GetPreset()
+{
+	return Preset;
+}
+
 void FSoundEffectBase::ClearPreset()
 {
 	if (Preset)
@@ -72,6 +77,7 @@ void FSoundEffectBase::EffectCommand(TFunction<void()> Command)
 {
 	CommandQueue.Enqueue(MoveTemp(Command));
 }
+
 void FSoundEffectBase::PumpPendingMessages()
 {
 	// Pumps the command queue
diff --git a/Engine/Source/Runtime/Engine/Private/SoundEffectPreset.cpp b/Engine/Source/Runtime/Engine/Private/SoundEffectPreset.cpp
--- a/Engine/Source/Runtime/Engine/Private/SoundEffectPreset.cpp
+++ b/Engine/Source/Runtime/Engine/Private/SoundEffectPreset.cpp
@@ -10,7 +10,7 @@ USoundEffectPreset::USoundEffectPreset(const FObjectInitializer& ObjectInitializ
 	: Super(ObjectInitializer)
 	, bInitialized(false)
 {
-	
+
 }
 
 USoundEffectPreset::~USoundEffectPreset()
@@ -52,6 +52,18 @@ void USoundEffectPreset::AddEffectInstance(FSoundEffectBase* InSource)
 	Instances.AddUnique(InSource);
 }
 
+void USoundEffectPreset::AddReferencedEffects(FReferenceCollector& Collector)
+{
+	for (FSoundEffectBase* Effect : Instances)
+	{
+		if (Effect)
+		{
+			const USoundEffectPreset* EffectPreset = Effect->GetPreset();
+			Collector.AddReferencedObject(EffectPreset);
+		}
+	}
+}
+
 void USoundEffectPreset::RemoveEffectInstance(FSoundEffectBase* InSource)
 {
 	if (Instances.Contains(InSource))
@@ -67,9 +79,7 @@ void USoundEffectPreset::PostEditChangeProperty(FPropertyChangedEvent& PropertyC
 	Init();
 	Update();
 }
-#endif
 
-#if WITH_EDITORONLY_DATA
 void USoundEffectSourcePresetChain::PostEditChangeProperty(FPropertyChangedEvent& PropertyChangedEvent)
 {
 	if (GEngine)
@@ -78,5 +88,15 @@ void USoundEffectSourcePresetChain::PostEditChangeProperty(FPropertyChangedEvent
 		AudioDeviceManager->UpdateSourceEffectChain(GetUniqueID(), Chain, bPlayEffectChainTails);
 	}
 }
-#endif
+#endif // WITH_EDITORONLY_DATA
 
+void USoundEffectSourcePresetChain::AddReferencedEffects(FReferenceCollector& Collector)
+{
+	for (FSourceEffectChainEntry& SourceEffect : Chain)
+	{
+		if (SourceEffect.Preset)
+		{
+			SourceEffect.Preset->AddReferencedEffects(Collector);
+		}
+	}
+}
\ No newline at end of file
