Missing file to fix ec

diff --git a/Engine/Source/Runtime/Engine/Classes/Sound/SoundEffectSource.h b/Engine/Source/Runtime/Engine/Classes/Sound/SoundEffectSource.h
--- a/Engine/Source/Runtime/Engine/Classes/Sound/SoundEffectSource.h
+++ b/Engine/Source/Runtime/Engine/Classes/Sound/SoundEffectSource.h
@@ -47,6 +47,8 @@ public:
 	UPROPERTY(EditAnywhere, Category = Effects)
 	uint32 bPlayEffectChainTails : 1;
 
+	void AddReferencedEffects(FReferenceCollector& Collector);
+
 protected:
 
 #if WITH_EDITORONLY_DATA
