ALODActor
- Added GetLODDrawDistanceWithOverride() which takes into account the active distance override (r.HLOD.DistanceOverride)
- Added SetLODParent() which allow exclusion of imposters
#rnx

diff --git a/Engine/Source/Runtime/Engine/Classes/Engine/LODActor.h b/Engine/Source/Runtime/Engine/Classes/Engine/LODActor.h
--- a/Engine/Source/Runtime/Engine/Classes/Engine/LODActor.h
+++ b/Engine/Source/Runtime/Engine/Classes/Engine/LODActor.h
@@ -95,6 +95,16 @@ public:
 	/** Gets the LOD draw distance */
 	float GetDrawDistance() const { return LODDrawDistance; }
 	float GetLODDrawDistance() const { return LODDrawDistance; }
+	float GetLODDrawDistanceWithOverride() const;
+
+	/** 
+	 * Set LOD Parent component for all of our components, normally associated with an ALODActor. 
+	 * @param InLODParent			This component used to compute visibility when hierarchical LOD is enabled. 
+	 * @param InParentDrawDistance	Updates the MinDrawDistances of the LODParent
+	 * @param bInApplyToImposters	Whether this should be applied to imposters instanced static meshes components.
+	 */
+	void SetLODParent(UPrimitiveComponent* InLODParent, float InParentDrawDistance, bool bInApplyToImposters);
+
 
 	/** Get the key that we use to check if we need to (re)build */
 	const FName& GetKey() const { return Key; }
diff --git a/Engine/Source/Runtime/Engine/Private/LODActor.cpp b/Engine/Source/Runtime/Engine/Private/LODActor.cpp
--- a/Engine/Source/Runtime/Engine/Private/LODActor.cpp
+++ b/Engine/Source/Runtime/Engine/Private/LODActor.cpp
@@ -364,6 +364,34 @@ void ALODActor::ParseOverrideDistancesCVar()
 	}	
 }
 
+float ALODActor::GetLODDrawDistanceWithOverride() const
+{
+	const int32 NumDistances = ALODActor::HLODDistances.Num();
+	const int32 DistanceIndex = [&]()
+	{
+		if(CachedNumHLODLevels <= NumDistances)
+		{
+			return (LODLevel + (NumDistances - CachedNumHLODLevels)) - 1;
+		}
+		else
+		{
+			// We've reached the end of the array, change nothing
+			return (int32)INDEX_NONE;
+		}
+	}();
+
+	const float HLODDistanceOverride = (!ALODActor::HLODDistances.IsValidIndex(DistanceIndex)) ? 0.0f : ALODActor::HLODDistances[DistanceIndex];
+	// Determine desired HLOD state
+	float MinDrawDistance = LODDrawDistance;
+	const bool bIsOverridingHLODDistance = HLODDistanceOverride != 0.0f;
+	if(bIsOverridingHLODDistance)
+	{
+		MinDrawDistance = HLODDistanceOverride;
+	}
+
+	return MinDrawDistance;
+}
+
 void ALODActor::Tick(float DeltaSeconds)
 {
 	AActor::Tick(DeltaSeconds);
@@ -371,29 +399,8 @@ void ALODActor::Tick(float DeltaSeconds)
 	{		
 		if (ResetDrawDistanceTime > CVarHLODDitherPauseTime.GetValueOnAnyThread())
 		{
-			const int32 NumDistances = ALODActor::HLODDistances.Num();
-			const int32 DistanceIndex = [&]()
-	        {
-		        if (CachedNumHLODLevels <= NumDistances)
-		        {
-			    	return (LODLevel + (NumDistances - CachedNumHLODLevels)) - 1;
-		        }
-		        else
-		        {
-			        // We've reached the end of the array, change nothing
-			        return (int32)INDEX_NONE;
-		        }
-	        }();
-
-
-			const float HLODDistanceOverride = (!ALODActor::HLODDistances.IsValidIndex(DistanceIndex)) ? 0.0f : ALODActor::HLODDistances[DistanceIndex];
 			// Determine desired HLOD state
-			float MinDrawDistance = LODDrawDistance;
-			const bool bIsOverridingHLODDistance = HLODDistanceOverride != 0.0f;
-			if (bIsOverridingHLODDistance)
-			{
-				MinDrawDistance = HLODDistanceOverride;
-			}
+			float MinDrawDistance = GetLODDrawDistanceWithOverride();
 
 			SetComponentsMinDrawDistance(MinDrawDistance, true);
 			bNeedsDrawDistanceReset = false;
@@ -408,6 +415,24 @@ void ALODActor::Tick(float DeltaSeconds)
 	}
 }
 
+void ALODActor::SetLODParent(UPrimitiveComponent* InLODParent, float InParentDrawDistance, bool bApplyToImposters)
+{
+	if (bApplyToImposters)
+	{
+		AActor::SetLODParent(InLODParent, InParentDrawDistance);
+	}
+	else
+	{
+		if(InLODParent)
+		{
+			InLODParent->MinDrawDistance = InParentDrawDistance;
+			InLODParent->MarkRenderStateDirty();
+		}
+
+		StaticMeshComponent->SetLODParentPrimitive(InLODParent);
+	}
+}
+
 void ALODActor::PauseDitherTransition()
 {
 	SetComponentsMinDrawDistance(0.0f, true);
