Change to support skeletal mesh LOD streaming. The feature is disabled on all platforms

diff --git a/Engine/Plugins/VirtualProduction/Takes/Source/TakeRecorder/Private/SerializedRecorder.cpp b/Engine/Plugins/VirtualProduction/Takes/Source/TakeRecorder/Private/SerializedRecorder.cpp
--- a/Engine/Plugins/VirtualProduction/Takes/Source/TakeRecorder/Private/SerializedRecorder.cpp
+++ b/Engine/Plugins/VirtualProduction/Takes/Source/TakeRecorder/Private/SerializedRecorder.cpp
@@ -442,7 +442,7 @@ AActor* FSerializedRecorder::SetActorPossesableOrSpawnable(UMovieSceneSequence*
 				SkeletalMeshComponent->SetAnimationMode(EAnimationMode::AnimationSingleNode);
 				SkeletalMeshComponent->bEnableUpdateRateOptimizations = false;
 				SkeletalMeshComponent->VisibilityBasedAnimTickOption = EVisibilityBasedAnimTickOption::AlwaysTickPoseAndRefreshBones;
-				SkeletalMeshComponent->ForcedLodModel = 1;
+				SkeletalMeshComponent->SetForcedLOD(1);
 			}
 
 			// Disable auto-possession on recorded Pawns so that when the Spawnable is spawned it doesn't auto-possess the player
diff --git a/Engine/Plugins/VirtualProduction/Takes/Source/TakeRecorderSources/Private/TakeRecorderActorSource.cpp b/Engine/Plugins/VirtualProduction/Takes/Source/TakeRecorderSources/Private/TakeRecorderActorSource.cpp
--- a/Engine/Plugins/VirtualProduction/Takes/Source/TakeRecorderSources/Private/TakeRecorderActorSource.cpp
+++ b/Engine/Plugins/VirtualProduction/Takes/Source/TakeRecorderSources/Private/TakeRecorderActorSource.cpp
@@ -1173,7 +1173,7 @@ void UTakeRecorderActorSource::PostProcessCreatedObjectTemplateImpl(AActor* Obje
 	 	SkeletalMeshComponent->SetAnimationMode(EAnimationMode::AnimationSingleNode);
 	 	SkeletalMeshComponent->bEnableUpdateRateOptimizations = false;
 	 	SkeletalMeshComponent->VisibilityBasedAnimTickOption = EVisibilityBasedAnimTickOption::AlwaysTickPoseAndRefreshBones;
-	 	SkeletalMeshComponent->ForcedLodModel = 1;
+	 	SkeletalMeshComponent->SetForcedLOD(1);
 	 }
 
 	// Disable auto-possession on recorded Pawns so that when the Spawnable is spawned it doesn't auto-possess the player
diff --git a/Engine/Source/Editor/MeshPaint/Private/MeshPaintHelpers.cpp b/Engine/Source/Editor/MeshPaint/Private/MeshPaintHelpers.cpp
--- a/Engine/Source/Editor/MeshPaint/Private/MeshPaintHelpers.cpp
+++ b/Engine/Source/Editor/MeshPaint/Private/MeshPaintHelpers.cpp
@@ -1196,7 +1196,7 @@ void MeshPaintHelpers::ForceRenderMeshLOD(UMeshComponent* Component, int32 LODIn
 	}
 	else if (USkeletalMeshComponent* SkeletalMeshComponent = Cast<USkeletalMeshComponent>(Component))
 	{
-		SkeletalMeshComponent->ForcedLodModel = LODIndex + 1;
+		SkeletalMeshComponent->SetForcedLOD(LODIndex + 1);
 	}	
 }
 
diff --git a/Engine/Source/Editor/Persona/Private/AnimationEditorViewportClient.cpp b/Engine/Source/Editor/Persona/Private/AnimationEditorViewportClient.cpp
--- a/Engine/Source/Editor/Persona/Private/AnimationEditorViewportClient.cpp
+++ b/Engine/Source/Editor/Persona/Private/AnimationEditorViewportClient.cpp
@@ -487,7 +487,7 @@ void FAnimationViewportClient::DrawUVsForMesh(FViewport* InViewport, FCanvas* In
 	UDebugSkelMeshComponent* PreviewMeshComponent = GetAnimPreviewScene()->GetPreviewMeshComponent();
 
 	//use the overridden LOD level
-	const uint32 LODLevel = FMath::Clamp(PreviewMeshComponent->ForcedLodModel - 1, 0, PreviewMeshComponent->SkeletalMesh->GetLODNum() - 1);
+	const uint32 LODLevel = FMath::Clamp(PreviewMeshComponent->GetForcedLOD() - 1, 0, PreviewMeshComponent->SkeletalMesh->GetLODNum() - 1);
 
 	TArray<FVector2D> SelectedEdgeTexCoords; //No functionality in Persona for this (yet?)
 
@@ -1700,8 +1700,8 @@ void FAnimationViewportClient::ToggleCPUSkinning()
 	UDebugSkelMeshComponent* PreviewMeshComponent = GetAnimPreviewScene()->GetPreviewMeshComponent();
 	if (PreviewMeshComponent != nullptr)
 	{
-		PreviewMeshComponent->bCPUSkinning = !PreviewMeshComponent->bCPUSkinning;
-		PreviewMeshComponent->MarkRenderStateDirty();
+		const bool bCurVal = PreviewMeshComponent->GetCPUSkinningEnabled();
+		PreviewMeshComponent->SetCPUSkinningEnabled(!bCurVal);
 		Invalidate();
 	}
 }
@@ -1711,7 +1711,7 @@ bool FAnimationViewportClient::IsSetCPUSkinningChecked() const
 	UDebugSkelMeshComponent* PreviewMeshComponent = GetAnimPreviewScene()->GetPreviewMeshComponent();
 	if (PreviewMeshComponent != nullptr)
 	{
-		return PreviewMeshComponent->bCPUSkinning;
+		return PreviewMeshComponent->GetCPUSkinningEnabled();
 	}
 	return false;
 }
diff --git a/Engine/Source/Editor/Persona/Private/PersonaMeshDetails.cpp b/Engine/Source/Editor/Persona/Private/PersonaMeshDetails.cpp
--- a/Engine/Source/Editor/Persona/Private/PersonaMeshDetails.cpp
+++ b/Engine/Source/Editor/Persona/Private/PersonaMeshDetails.cpp
@@ -808,7 +808,7 @@ void FPersonaMeshDetails::AddLODLevelCategories(IDetailLayoutBuilder& DetailLayo
 		int32 CurrentLodIndex = 0;
 		if (GetPersonaToolkit()->GetPreviewMeshComponent() != nullptr)
 		{
-			CurrentLodIndex = GetPersonaToolkit()->GetPreviewMeshComponent()->ForcedLodModel;
+			CurrentLodIndex = GetPersonaToolkit()->GetPreviewMeshComponent()->GetForcedLOD();
 		}
 
 		FString LODControllerCategoryName = FString(TEXT("LODCustomMode"));
@@ -1111,7 +1111,7 @@ FText FPersonaMeshDetails::GetLODCustomModeNameContent(int32 LODIndex) const
 	int32 CurrentLodIndex = 0;
 	if (GetPersonaToolkit()->GetPreviewMeshComponent() != nullptr)
 	{
-		CurrentLodIndex = GetPersonaToolkit()->GetPreviewMeshComponent()->ForcedLodModel;
+		CurrentLodIndex = GetPersonaToolkit()->GetPreviewMeshComponent()->GetForcedLOD();
 	}
 	int32 RealCurrentLODIndex = (CurrentLodIndex == 0 ? 0 : CurrentLodIndex - 1);
 	if (LODIndex == INDEX_NONE)
@@ -1126,7 +1126,7 @@ ECheckBoxState FPersonaMeshDetails::IsLODCustomModeCheck(int32 LODIndex) const
 	int32 CurrentLodIndex = 0;
 	if (GetPersonaToolkit()->GetPreviewMeshComponent() != nullptr)
 	{
-		CurrentLodIndex = GetPersonaToolkit()->GetPreviewMeshComponent()->ForcedLodModel;
+		CurrentLodIndex = GetPersonaToolkit()->GetPreviewMeshComponent()->GetForcedLOD();
 	}
 	if (LODIndex == INDEX_NONE)
 	{
@@ -1140,7 +1140,7 @@ void FPersonaMeshDetails::SetLODCustomModeCheck(ECheckBoxState NewState, int32 L
 	int32 CurrentLodIndex = 0;
 	if (GetPersonaToolkit()->GetPreviewMeshComponent() != nullptr)
 	{
-		CurrentLodIndex = GetPersonaToolkit()->GetPreviewMeshComponent()->ForcedLodModel;
+		CurrentLodIndex = GetPersonaToolkit()->GetPreviewMeshComponent()->GetForcedLOD();
 	}
 	if (LODIndex == INDEX_NONE)
 	{
@@ -1446,6 +1446,21 @@ void FPersonaMeshDetails::CustomizeLODSettingsCategories(IDetailLayoutBuilder& D
 	IDetailPropertyRow& DisableBelowMinLodStrippingRow = LODSettingsCategory.AddProperty(DisableBelowMinLodStrippingPropertyHandle);
 	DisableBelowMinLodStrippingRow.IsEnabled(TAttribute<bool>::Create(TAttribute<bool>::FGetter::CreateSP(this, &FPersonaMeshDetails::IsLODInfoEditingEnabled, -1)));
 	DetailLayout.HideProperty(DisableBelowMinLodStrippingPropertyHandle);
+
+	TSharedPtr<IPropertyHandle> bSupportLODStreamingPropertyHandle = DetailLayout.GetProperty(GET_MEMBER_NAME_CHECKED(USkeletalMesh, bSupportLODStreaming), USkeletalMesh::StaticClass());
+	IDetailPropertyRow& bSupportLODStreamingRow = LODSettingsCategory.AddProperty(bSupportLODStreamingPropertyHandle);
+	bSupportLODStreamingRow.IsEnabled(TAttribute<bool>::Create(TAttribute<bool>::FGetter::CreateSP(this, &FPersonaMeshDetails::IsLODInfoEditingEnabled, -1)));
+	DetailLayout.HideProperty(bSupportLODStreamingPropertyHandle);
+
+	TSharedPtr<IPropertyHandle> MaxNumStreamedLODsPropertyHandle = DetailLayout.GetProperty(GET_MEMBER_NAME_CHECKED(USkeletalMesh, MaxNumStreamedLODs), USkeletalMesh::StaticClass());
+	IDetailPropertyRow& MaxNumStreamedLODsRow = LODSettingsCategory.AddProperty(MaxNumStreamedLODsPropertyHandle);
+	MaxNumStreamedLODsRow.IsEnabled(TAttribute<bool>::Create(TAttribute<bool>::FGetter::CreateSP(this, &FPersonaMeshDetails::IsLODInfoEditingEnabled, -1)));
+	DetailLayout.HideProperty(MaxNumStreamedLODsPropertyHandle);
+
+	TSharedPtr<IPropertyHandle> MaxNumOptionalLODsPropertyHandle = DetailLayout.GetProperty(GET_MEMBER_NAME_CHECKED(USkeletalMesh, MaxNumOptionalLODs), USkeletalMesh::StaticClass());
+	IDetailPropertyRow& MaxNumOptionalLODsRow = LODSettingsCategory.AddProperty(MaxNumOptionalLODsPropertyHandle);
+	MaxNumOptionalLODsRow.IsEnabled(TAttribute<bool>::Create(TAttribute<bool>::FGetter::CreateSP(this, &FPersonaMeshDetails::IsLODInfoEditingEnabled, -1)));
+	DetailLayout.HideProperty(MaxNumOptionalLODsPropertyHandle);
 }
 
 // save LOD settings
@@ -2962,7 +2977,7 @@ void FPersonaMeshDetails::SetCurrentLOD(int32 NewLodIndex)
 	{
 		return;
 	}
-	int32 CurrentDisplayLOD = GetPersonaToolkit()->GetPreviewMeshComponent()->ForcedLodModel;
+	int32 CurrentDisplayLOD = GetPersonaToolkit()->GetPreviewMeshComponent()->GetForcedLOD();
 	int32 RealCurrentDisplayLOD = CurrentDisplayLOD == 0 ? 0 : CurrentDisplayLOD - 1;
 	int32 RealNewLOD = NewLodIndex == 0 ? 0 : NewLodIndex - 1;
 	if (CurrentDisplayLOD == NewLodIndex || !LodCategories.IsValidIndex(RealCurrentDisplayLOD) || !LodCategories.IsValidIndex(RealNewLOD))
@@ -2988,9 +3003,9 @@ void FPersonaMeshDetails::UpdateLODCategoryVisibility() const
 	bool bAutoLod = false;
 	if (GetPersonaToolkit()->GetPreviewMeshComponent() != nullptr)
 	{
-		bAutoLod = GetPersonaToolkit()->GetPreviewMeshComponent()->ForcedLodModel == 0;
+		bAutoLod = GetPersonaToolkit()->GetPreviewMeshComponent()->GetForcedLOD() == 0;
 	}
-	int32 CurrentDisplayLOD = bAutoLod ? 0 : GetPersonaToolkit()->GetPreviewMeshComponent()->ForcedLodModel - 1;
+	int32 CurrentDisplayLOD = bAutoLod ? 0 : GetPersonaToolkit()->GetPreviewMeshComponent()->GetForcedLOD() - 1;
 	if (LodCategories.IsValidIndex(CurrentDisplayLOD) && GetPersonaToolkit()->GetMesh())
 	{
 		int32 SkeletalMeshLodNumber = GetPersonaToolkit()->GetMesh()->GetLODNum();
@@ -3010,15 +3025,15 @@ FText FPersonaMeshDetails::GetCurrentLodName() const
 	bool bAutoLod = false;
 	if (GetPersonaToolkit()->GetPreviewMeshComponent() != nullptr)
 	{
-		bAutoLod = GetPersonaToolkit()->GetPreviewMeshComponent()->ForcedLodModel == 0;
+		bAutoLod = GetPersonaToolkit()->GetPreviewMeshComponent()->GetForcedLOD() == 0;
 	}
-	int32 CurrentDisplayLOD = bAutoLod ? 0 : GetPersonaToolkit()->GetPreviewMeshComponent()->ForcedLodModel - 1;
+	int32 CurrentDisplayLOD = bAutoLod ? 0 : GetPersonaToolkit()->GetPreviewMeshComponent()->GetForcedLOD() - 1;
 	return FText::FromString(bAutoLod ? FString(TEXT("Auto (LOD0)")) : (FString(TEXT("LOD")) + FString::FromInt(CurrentDisplayLOD)));
 }
 
 FText FPersonaMeshDetails::GetCurrentLodTooltip() const
 {
-	if (GetPersonaToolkit()->GetPreviewMeshComponent() != nullptr && GetPersonaToolkit()->GetPreviewMeshComponent()->ForcedLodModel == 0)
+	if (GetPersonaToolkit()->GetPreviewMeshComponent() != nullptr && GetPersonaToolkit()->GetPreviewMeshComponent()->GetForcedLOD() == 0)
 	{
 		return LOCTEXT("PersonaLODPickerCurrentLODTooltip", "With Auto LOD selected, LOD0's properties are visible for editing");
 	}
@@ -3069,7 +3084,7 @@ TSharedRef<SWidget> FPersonaMeshDetails::OnGenerateLodMenuForLodPicker()
 	bool bAutoLod = false;
 	if (GetPersonaToolkit()->GetPreviewMeshComponent() != nullptr)
 	{
-		bAutoLod = GetPersonaToolkit()->GetPreviewMeshComponent()->ForcedLodModel == 0;
+		bAutoLod = GetPersonaToolkit()->GetPreviewMeshComponent()->GetForcedLOD() == 0;
 	}
 	const int32 SkelMeshLODCount = SkelMesh->GetLODNum();
 	if(SkelMeshLODCount < 2)
diff --git a/Engine/Source/Editor/Persona/Private/SAnimationEditorViewport.cpp b/Engine/Source/Editor/Persona/Private/SAnimationEditorViewport.cpp
--- a/Engine/Source/Editor/Persona/Private/SAnimationEditorViewport.cpp
+++ b/Engine/Source/Editor/Persona/Private/SAnimationEditorViewport.cpp
@@ -1526,7 +1526,7 @@ void SAnimationEditorViewportTabBody::PopulateUVChoices()
 	
 	if (UDebugSkelMeshComponent* PreviewComponent = GetPreviewScene()->GetPreviewMeshComponent())
 	{
-		int32 CurrentLOD = FMath::Clamp(PreviewComponent->ForcedLodModel - 1, 0, NumUVChannels.Num() - 1);
+		int32 CurrentLOD = FMath::Clamp(PreviewComponent->GetForcedLOD() - 1, 0, NumUVChannels.Num() - 1);
 
 		if (NumUVChannels.IsValidIndex(CurrentLOD))
 		{
@@ -1683,13 +1683,13 @@ int32 SAnimationEditorViewportTabBody::GetLODSelection() const
 	{
 		// If we are forcing a LOD level, report the actual LOD level we are displaying
 		// as the mesh can potentially change LOD count under the viewport.
-		if(PreviewComponent->ForcedLodModel > 0)
+		if(PreviewComponent->GetForcedLOD() > 0)
 		{
 			return PreviewComponent->PredictedLODLevel + 1;
 		}
 		else
 		{
-			return PreviewComponent->ForcedLodModel;
+			return PreviewComponent->GetForcedLOD();
 		}
 	}
 	return 0;
@@ -1707,7 +1707,7 @@ void SAnimationEditorViewportTabBody::OnSetLODModel(int32 LODSelectionType)
 	if( PreviewComponent )
 	{
 		LODSelection = LODSelectionType;
-		PreviewComponent->ForcedLodModel = LODSelectionType;
+		PreviewComponent->SetForcedLOD(LODSelectionType);
 		PopulateUVChoices();
 		GetPreviewScene()->BroadcastOnSelectedLODChanged();
 	}
@@ -1717,9 +1717,9 @@ void SAnimationEditorViewportTabBody::OnLODModelChanged()
 {
 	UDebugSkelMeshComponent* PreviewComponent = GetPreviewScene()->GetPreviewMeshComponent();
 
-	if (PreviewComponent && LODSelection != PreviewComponent->ForcedLodModel)
+	if (PreviewComponent && LODSelection != PreviewComponent->GetForcedLOD())
 	{
-		LODSelection = PreviewComponent->ForcedLodModel;
+		LODSelection = PreviewComponent->GetForcedLOD();
 		PopulateUVChoices();
 	}
 }
diff --git a/Engine/Source/Editor/PhysicsAssetEditor/Private/PhysicsAssetEditorSkeletalMeshComponent.cpp b/Engine/Source/Editor/PhysicsAssetEditor/Private/PhysicsAssetEditorSkeletalMeshComponent.cpp
--- a/Engine/Source/Editor/PhysicsAssetEditor/Private/PhysicsAssetEditorSkeletalMeshComponent.cpp
+++ b/Engine/Source/Editor/PhysicsAssetEditor/Private/PhysicsAssetEditorSkeletalMeshComponent.cpp
@@ -60,7 +60,7 @@ UPhysicsAssetEditorSkeletalMeshComponent::UPhysicsAssetEditorSkeletalMeshCompone
 	// this is because in phat editor, you'd like to see fixed bones to be fixed without animation force update
 	KinematicBonesUpdateType = EKinematicBonesUpdateToPhysics::SkipSimulatingBones;
 	bUpdateJointsFromAnimation = false;
-	ForcedLodModel = 1;
+	SetForcedLOD(1);
 
 	static FName CollisionProfileName(TEXT("PhysicsActor"));
 	SetCollisionProfileName(CollisionProfileName);
diff --git a/Engine/Source/Editor/SequenceRecorder/Private/ActorRecording.cpp b/Engine/Source/Editor/SequenceRecorder/Private/ActorRecording.cpp
--- a/Engine/Source/Editor/SequenceRecorder/Private/ActorRecording.cpp
+++ b/Engine/Source/Editor/SequenceRecorder/Private/ActorRecording.cpp
@@ -492,7 +492,7 @@ void UActorRecording::StartRecordingActorProperties(ULevelSequence* CurrentSeque
 							SkeletalMeshComponent->SetAnimationMode(EAnimationMode::AnimationSingleNode);
 							SkeletalMeshComponent->bEnableUpdateRateOptimizations = false;
 							SkeletalMeshComponent->VisibilityBasedAnimTickOption = EVisibilityBasedAnimTickOption::AlwaysTickPoseAndRefreshBones;
-							SkeletalMeshComponent->ForcedLodModel = 1;
+							SkeletalMeshComponent->SetForcedLOD(1);
 						}
 					}
 
diff --git a/Engine/Source/Editor/SequenceRecorder/Private/AnimationRecorder.cpp b/Engine/Source/Editor/SequenceRecorder/Private/AnimationRecorder.cpp
--- a/Engine/Source/Editor/SequenceRecorder/Private/AnimationRecorder.cpp
+++ b/Engine/Source/Editor/SequenceRecorder/Private/AnimationRecorder.cpp
@@ -822,8 +822,8 @@ void FAnimRecorderInstance::InitInternal(USkeletalMeshComponent* InComponent, co
 
 	if (InComponent)
 	{
-		CachedSkelCompForcedLodModel = InComponent->ForcedLodModel;
-		InComponent->ForcedLodModel = 1;
+		CachedSkelCompForcedLodModel = InComponent->GetForcedLOD();
+		InComponent->SetForcedLOD(1);
 
 		// turn off URO and make sure we always update even if out of view
 		bCachedEnableUpdateRateOptimizations = InComponent->bEnableUpdateRateOptimizations;
@@ -875,7 +875,7 @@ void FAnimRecorderInstance::FinishRecording(bool bShowMessage)
 	if (SkelComp.IsValid())
 	{
 		// restore force lod setting
-		SkelComp->ForcedLodModel = CachedSkelCompForcedLodModel;
+		SkelComp->SetForcedLOD(CachedSkelCompForcedLodModel);
 
 		// restore update flags
 		SkelComp->bEnableUpdateRateOptimizations = bCachedEnableUpdateRateOptimizations;
diff --git a/Engine/Source/Editor/UnrealEd/Private/Animation/DebugSkelMeshComponent.cpp b/Engine/Source/Editor/UnrealEd/Private/Animation/DebugSkelMeshComponent.cpp
--- a/Engine/Source/Editor/UnrealEd/Private/Animation/DebugSkelMeshComponent.cpp
+++ b/Engine/Source/Editor/UnrealEd/Private/Animation/DebugSkelMeshComponent.cpp
@@ -353,7 +353,7 @@ void UDebugSkelMeshComponent::EnablePreview(bool bEnable, UAnimationAsset* Previ
 
 bool UDebugSkelMeshComponent::ShouldCPUSkin()
 {
-	return 	bCPUSkinning || bDrawBoneInfluences || bDrawNormals || bDrawTangents || bDrawBinormals || bDrawMorphTargetVerts;
+	return 	GetCPUSkinningEnabled() || bDrawBoneInfluences || bDrawNormals || bDrawTangents || bDrawBinormals || bDrawMorphTargetVerts;
 }
 
 
diff --git a/Engine/Source/Editor/UnrealEd/Private/Factories/SkeletalMeshImport.cpp b/Engine/Source/Editor/UnrealEd/Private/Factories/SkeletalMeshImport.cpp
--- a/Engine/Source/Editor/UnrealEd/Private/Factories/SkeletalMeshImport.cpp
+++ b/Engine/Source/Editor/UnrealEd/Private/Factories/SkeletalMeshImport.cpp
@@ -402,6 +402,9 @@ ExistingSkelMeshData* SaveExistingSkelMeshData(USkeletalMesh* ExistingSkelMesh,
 	ExistingMeshDataPtr->UseMaterialNameSlotWorkflow = SkeletalMeshIsUsingMaterialSlotNameWorkflow(ExistingSkelMesh->AssetImportData);
 	ExistingMeshDataPtr->MinLOD = ExistingSkelMesh->MinLod;
 	ExistingMeshDataPtr->DisableBelowMinLodStripping = ExistingSkelMesh->DisableBelowMinLodStripping;
+	ExistingMeshDataPtr->bSupportLODStreaming = ExistingSkelMesh->bSupportLODStreaming;
+	ExistingMeshDataPtr->MaxNumStreamedLODs = ExistingSkelMesh->MaxNumStreamedLODs;
+	ExistingMeshDataPtr->MaxNumOptionalLODs = ExistingSkelMesh->MaxNumOptionalLODs;
 
 	FSkeletalMeshModel* ImportedResource = ExistingSkelMesh->GetImportedModel();
 
@@ -790,6 +793,9 @@ void RestoreExistingSkelMeshData(ExistingSkelMeshData* MeshData, USkeletalMesh*
 
 	SkeletalMesh->MinLod = MeshData->MinLOD;
 	SkeletalMesh->DisableBelowMinLodStripping = MeshData->DisableBelowMinLodStripping;
+	SkeletalMesh->bSupportLODStreaming = MeshData->bSupportLODStreaming;
+	SkeletalMesh->MaxNumStreamedLODs = MeshData->MaxNumStreamedLODs;
+	SkeletalMesh->MaxNumOptionalLODs = MeshData->MaxNumOptionalLODs;
 
 	FSkeletalMeshModel* SkeletalMeshImportedModel = SkeletalMesh->GetImportedModel();
 
diff --git a/Engine/Source/Editor/UnrealEd/Private/LODUtilities.cpp b/Engine/Source/Editor/UnrealEd/Private/LODUtilities.cpp
--- a/Engine/Source/Editor/UnrealEd/Private/LODUtilities.cpp
+++ b/Engine/Source/Editor/UnrealEd/Private/LODUtilities.cpp
@@ -169,7 +169,7 @@ void FLODUtilities::RemoveLOD(FSkeletalMeshUpdateContext& UpdateContext, int32 D
 			USkinnedMeshComponent* SkinnedComponent = Cast<USkinnedMeshComponent>(*Iter);
 			if(SkinnedComponent)
 			{
-				SkinnedComponent->ForcedLodModel = 0;
+				SkinnedComponent->SetForcedLOD(0);
 			}
 		}
 
diff --git a/Engine/Source/Editor/UnrealEd/Public/SkelImport.h b/Engine/Source/Editor/UnrealEd/Public/SkelImport.h
--- a/Engine/Source/Editor/UnrealEd/Public/SkelImport.h
+++ b/Engine/Source/Editor/UnrealEd/Public/SkelImport.h
@@ -83,6 +83,9 @@ struct ExistingSkelMeshData
 	FSkeletalMeshSamplingInfo				ExistingSamplingInfo;
 	FPerPlatformInt							MinLOD;
 	FPerPlatformBool						DisableBelowMinLodStripping;
+	FPerPlatformBool						bSupportLODStreaming;
+	FPerPlatformInt							MaxNumStreamedLODs;
+	FPerPlatformInt							MaxNumOptionalLODs;
 
 	TMap<UAssetUserData*, bool>				ExistingAssetUserData;
 };
diff --git a/Engine/Source/Runtime/Engine/Classes/Components/SkinnedMeshComponent.h b/Engine/Source/Runtime/Engine/Classes/Components/SkinnedMeshComponent.h
--- a/Engine/Source/Runtime/Engine/Classes/Components/SkinnedMeshComponent.h
+++ b/Engine/Source/Runtime/Engine/Classes/Components/SkinnedMeshComponent.h
@@ -375,6 +375,7 @@ public:
 	//
 	
 	/** If 0, auto-select LOD level. if >0, force to (ForcedLodModel-1). */
+	UE_DEPRECATED(4.24, "Direct access to ForcedLodModel is deprecated. Please use its getter and setter instead.")
 	UPROPERTY(EditAnywhere, AdvancedDisplay, BlueprintReadOnly, Category=LOD)
 	int32 ForcedLodModel;
 
@@ -523,6 +524,7 @@ public:
 	uint8 bCastCapsuleIndirectShadow:1;
 
 	/** Whether or not to CPU skin this component, requires render data refresh after changing */
+	UE_DEPRECATED(4.24, "Direct access to bCPUSkinning is deprecated. Please use its getter and setter instead.")
 	UPROPERTY(transient)
 	uint8 bCPUSkinning : 1;
 
@@ -645,13 +647,17 @@ public:
 	void SetMinLOD(int32 InNewMinLOD);
 
 	/**
-	 * Set MinLodModel of the mesh component
+	 * Set ForcedLodModel of the mesh component
 	 *
 	 * @param	InNewForcedLOD	Set new ForcedLODModel that forces to set the incoming LOD. Range from [1, Max Number of LOD]. This will affect in the next tick update. 
 	 */
 	UFUNCTION(BlueprintCallable, Category="Components|SkinnedMesh")
 	void SetForcedLOD(int32 InNewForcedLOD);
 
+	/** Get ForcedLodModel of the mesh component. Note that the actual forced LOD level is the return value minus one and zero means no forced LOD */
+	UFUNCTION(BlueprintCallable, Category="Components|SkinnedMesh")
+	int32 GetForcedLOD() const;
+
 #if WITH_EDITOR
 	/**
 	 * Get the LOD Bias of this component
@@ -842,6 +848,21 @@ public:
 	 */
 	virtual bool ShouldCPUSkin();
 
+	/**
+	 * Getter for bCPUSkinning member variable
+	 * May return a different value from ShouldCPUSkin()
+	 */
+	bool GetCPUSkinningEnabled() const;
+
+	/**
+	 * Set whether this component uses CPU skinning
+	 * Notes:
+	 * - If enabled, skeletal mesh referenced by this component will be removed from streaming manager
+	 * - Streaming cannot be (re)enabled as long as any component uses CPU skinning
+	 * - This function is expensive
+	 */
+	void SetCPUSkinningEnabled(bool bEnable, bool bRecreateRenderStateImmediately = false);
+
 	/** 
 	 * Function to operate on mesh object after its created, 
 	 * but before it's attached.
@@ -1451,6 +1472,7 @@ public:
 #endif
 
 	friend class FRenderStateRecreator;
+	friend class FSkeletalMeshStreamOut;
 };
 
 class FRenderStateRecreator
diff --git a/Engine/Source/Runtime/Engine/Classes/Engine/SkeletalMesh.h b/Engine/Source/Runtime/Engine/Classes/Engine/SkeletalMesh.h
--- a/Engine/Source/Runtime/Engine/Classes/Engine/SkeletalMesh.h
+++ b/Engine/Source/Runtime/Engine/Classes/Engine/SkeletalMesh.h
@@ -25,6 +25,8 @@
 #include "SkeletalMeshLODSettings.h"
 #include "Animation/NodeMappingProviderInterface.h"
 #include "Animation/SkinWeightProfile.h"
+#include "Engine/StreamableRenderAsset.h"
+#include "RenderAssetUpdate.h"
 
 #include "SkeletalMesh.generated.h"
 
@@ -44,6 +46,7 @@ class FSkeletalMeshLODModel;
 class FSkeletalMeshLODRenderData;
 class FSkinWeightVertexBuffer;
 struct FSkinWeightProfileInfo;
+class FSkeletalMeshUpdate;
 
 #if WITH_APEX_CLOTHING
 
@@ -450,7 +453,7 @@ namespace NSSkeletalMeshSourceFileLabels
  * @see https://docs.unrealengine.com/latest/INT/Engine/Content/Types/SkeletalMeshes/
  */
 UCLASS(hidecategories=Object, BlueprintType)
-class ENGINE_API USkeletalMesh : public UObject, public IInterface_CollisionDataProvider, public IInterface_AssetUserData, public INodeMappingProviderInterface
+class ENGINE_API USkeletalMesh : public UStreamableRenderAsset, public IInterface_CollisionDataProvider, public IInterface_AssetUserData, public INodeMappingProviderInterface
 {
 	GENERATED_UCLASS_BODY()
 
@@ -595,6 +598,18 @@ public:
 	FPerPlatformBool DisableBelowMinLodStripping;
 
 #if WITH_EDITORONLY_DATA
+	/** Whether we can stream the LODs of this mesh */
+	UPROPERTY(EditAnywhere, Category=LODSettings, meta=(DisplayName="Stream LODs"))
+	FPerPlatformBool bSupportLODStreaming;
+
+	/** Maximum number of LODs that can be streamed */
+	UPROPERTY(EditAnywhere, Category=LODSettings)
+	FPerPlatformInt MaxNumStreamedLODs;
+
+	/** Maximum number of LODs below min LOD level that can be saved to optional pak (currently, need to be either 0 or > num of LODs below MinLod) */
+	UPROPERTY(EditAnywhere, Category=LODSettings)
+	FPerPlatformInt MaxNumOptionalLODs;
+
 	UPROPERTY(EditAnywhere, BlueprintReadWrite, AssetRegistrySearchable, BlueprintSetter = SetLODSettings, Category = LODSettings)
 	USkeletalMeshLODSettings* LODSettings;
 
@@ -850,6 +865,11 @@ protected:
 	FOnMeshChanged OnMeshChanged;
 #endif
 
+	TRefCountPtr<FSkeletalMeshUpdate> PendingUpdate;
+
+	friend struct FSkeletalMeshUpdateContext;
+	friend class FSkeletalMeshUpdate;
+
 private:
 	/** 
 	 *	Array of named socket locations, set up in editor and used as a shortcut instead of specifying 
@@ -923,6 +943,34 @@ public:
 	virtual void GetPreloadDependencies(TArray<UObject*>& OutDeps) override;
 	//~ End UObject Interface.
 
+	//~ Begin UStreamableRenderAsset Interface.
+	virtual int32 GetLODGroupForStreaming() const final override;
+	virtual int32 GetNumMipsForStreaming() const final override;
+	virtual int32 GetNumNonStreamingMips() const final override;
+	virtual int32 CalcNumOptionalMips() const final override;
+	virtual int32 CalcCumulativeLODSize(int32 NumLODs) const final override;
+	virtual bool GetMipDataFilename(const int32 MipIndex, FString& OutBulkDataFilename) const final override;
+	virtual bool IsReadyForStreaming() const final override;
+	virtual int32 GetNumResidentMips() const final override;
+	virtual int32 GetNumRequestedMips() const final override;
+	virtual bool CancelPendingMipChangeRequest() final override;
+	virtual bool HasPendingUpdate() const final override;
+	virtual bool IsPendingUpdateLocked() const final override;
+	virtual bool StreamOut(int32 NewMipCount) final override;
+	virtual bool StreamIn(int32 NewMipCount, bool bHighPrio) final override;
+	virtual bool UpdateStreamingStatus(bool bWaitForMipFading = false) final override;
+	//~ End UStreamableRenderAsset Interface.
+
+	void LinkStreaming();
+	void UnlinkStreaming();
+
+	/**
+	* Cancels any pending static mesh streaming actions if possible.
+	* Returns when no more async loading requests are in flight.
+	*/
+	static void CancelAllPendingStreamingActions();
+
+
 	/** Setup-only routines - not concerned with the instance. */
 
 	void CalculateInvRefMatrices();
@@ -1243,7 +1291,7 @@ public:
 	/* 
 	 * Returns total number of LOD
 	 */
-	int32 GetLODNum() const { return LODInfo.Num();  }
+	int32 GetLODNum() const { check(LODInfo.Num() <= MAX_MESH_LOD_COUNT); return LODInfo.Num();  }
 
 public:
 	const TArray<FSkinWeightProfileInfo>& GetSkinWeightProfiles() const { return SkinWeightProfiles; }
diff --git a/Engine/Source/Runtime/Engine/Classes/Engine/SkeletalMeshLODSettings.h b/Engine/Source/Runtime/Engine/Classes/Engine/SkeletalMeshLODSettings.h
--- a/Engine/Source/Runtime/Engine/Classes/Engine/SkeletalMeshLODSettings.h
+++ b/Engine/Source/Runtime/Engine/Classes/Engine/SkeletalMeshLODSettings.h
@@ -123,6 +123,18 @@ protected:
 	UPROPERTY(globalconfig, EditAnywhere, Category = LODGroups)
 	FPerPlatformBool DisableBelowMinLodStripping;
 
+	/** Whether meshes in this group stream LODs by default */
+	UPROPERTY(globalconfig, EditAnywhere, Category=LODGroups, meta=(DisplayName="Stream LODs"))
+	FPerPlatformBool bSupportLODStreaming;
+
+	/** Default maximum number of streamed LODs for meshes in this group */
+	UPROPERTY(globalconfig, EditAnywhere, Category=LODGroups)
+	FPerPlatformInt MaxNumStreamedLODs;
+
+	/** Default maximum number of optional LODs for meshes in this group (currently, need to be either 0 or > num of LODs below MinLod) */
+	UPROPERTY(globalconfig, EditAnywhere, Category=LODGroups)
+	FPerPlatformInt MaxNumOptionalLODs;
+
 	UPROPERTY(globalconfig, EditAnywhere, Category=LODGroups)
 	TArray<FSkeletalMeshLODGroupSettings> LODGroups;
 
diff --git a/Engine/Source/Runtime/Engine/Private/Animation/SkinWeightProfile.cpp b/Engine/Source/Runtime/Engine/Private/Animation/SkinWeightProfile.cpp
--- a/Engine/Source/Runtime/Engine/Private/Animation/SkinWeightProfile.cpp
+++ b/Engine/Source/Runtime/Engine/Private/Animation/SkinWeightProfile.cpp
@@ -22,6 +22,12 @@ static void OnDefaultProfileCVarsChanged(IConsoleVariable* Variable)
 
 		if (bClearBuffer || bSetBuffer)
 		{
+			// Make sure no pending skeletal mesh LOD updates
+			if (IStreamingManager::Get_Concurrent())
+			{
+				IStreamingManager::Get().GetRenderAssetStreamingManager().BlockTillAllRequestsFinished();
+			}
+
 			for (TObjectIterator<USkeletalMesh> It; It; ++It)
 			{
 				if (*It)
@@ -376,3 +382,63 @@ SIZE_T FSkinWeightProfilesData::GetResourcesSize() const
 
 	return SummedSize;
 }
+
+void FSkinWeightProfilesData::SerializeMetaData(FArchive& Ar)
+{
+	TArray<FName, TInlineAllocator<8>> ProfileNames;
+	if (Ar.IsSaving())
+	{
+		OverrideData.GenerateKeyArray(ProfileNames);
+		Ar << ProfileNames;
+	}
+	else
+	{
+		Ar << ProfileNames;
+		OverrideData.Empty(ProfileNames.Num());
+		for (int32 Idx = 0; Idx < ProfileNames.Num(); ++Idx)
+		{
+			OverrideData.Add(ProfileNames[Idx]);
+		}
+	}
+}
+
+void FSkinWeightProfilesData::ReleaseCPUResources()
+{
+	for (TMap<FName, FRuntimeSkinWeightProfileData>::TIterator It(OverrideData); It; ++It)
+	{
+		It->Value = FRuntimeSkinWeightProfileData();
+	}
+}
+
+template <bool bRenderThread>
+void FSkinWeightProfilesData::CreateRHIBuffers_Internal(TArray<TPair<FName, FVertexBufferRHIRef>>& OutBuffers)
+{
+	const int32 NumActiveProfiles = ProfileNameToBuffer.Num();
+	check(BaseBuffer || !NumActiveProfiles);
+	OutBuffers.Empty(NumActiveProfiles);
+	for (TMap<FName, FSkinWeightVertexBuffer*>::TIterator It(ProfileNameToBuffer); It; ++It)
+	{
+		const FName& ProfileName = It->Key;
+		FSkinWeightVertexBuffer* OverrideBuffer = It->Value;
+		ApplyOverrideProfile(OverrideBuffer, ProfileName);
+		if (bRenderThread)
+		{
+			OutBuffers.Emplace(ProfileName, OverrideBuffer->CreateRHIBuffer_RenderThread());
+		}
+		else
+		{
+			OutBuffers.Emplace(ProfileName, OverrideBuffer->CreateRHIBuffer_Async());
+		}
+	}
+}
+
+void FSkinWeightProfilesData::CreateRHIBuffers_RenderThread(TArray<TPair<FName, FVertexBufferRHIRef>>& OutBuffers)
+{
+	CreateRHIBuffers_Internal<true>(OutBuffers);
+}
+
+void FSkinWeightProfilesData::CreateRHIBuffers_Async(TArray<TPair<FName, FVertexBufferRHIRef>>& OutBuffers)
+{
+	CreateRHIBuffers_Internal<false>(OutBuffers);
+}
+
diff --git a/Engine/Source/Runtime/Engine/Private/Components/SkinnedMeshComponent.cpp b/Engine/Source/Runtime/Engine/Private/Components/SkinnedMeshComponent.cpp
--- a/Engine/Source/Runtime/Engine/Private/Components/SkinnedMeshComponent.cpp
+++ b/Engine/Source/Runtime/Engine/Private/Components/SkinnedMeshComponent.cpp
@@ -667,7 +667,6 @@ void USkinnedMeshComponent::SendRenderDynamicData_Concurrent()
 		// scene proxy update of material usage based on active morphs
 		UpdateMorphMaterialUsageOnProxy();
 	}
-
 }
 
 void USkinnedMeshComponent::ClearMotionVector()
@@ -904,7 +903,64 @@ bool USkinnedMeshComponent::IsMaterialSlotNameValid(FName MaterialSlotName) cons
 
 bool USkinnedMeshComponent::ShouldCPUSkin()
 {
+	return GetCPUSkinningEnabled();
+}
+
+bool USkinnedMeshComponent::GetCPUSkinningEnabled() const
+{
+PRAGMA_DISABLE_DEPRECATION_WARNINGS
 	return bCPUSkinning;
+PRAGMA_ENABLE_DEPRECATION_WARNINGS
+}
+
+void USkinnedMeshComponent::SetCPUSkinningEnabled(bool bEnable, bool bRecreateRenderStateImmediately)
+{
+	checkSlow(IsInGameThread());
+	check(SkeletalMesh && SkeletalMesh->GetResourceForRendering());
+
+	if (GetCPUSkinningEnabled() == bEnable)
+	{
+		return;
+	}
+
+	if (bEnable && IStreamingManager::Get().IsRenderAssetStreamingEnabled())
+	{
+		UE_LOG(LogSkinnedMeshComp, Warning, TEXT("It is expensive to enable CPU skinning with LOD streaming on."));
+
+		IRenderAssetStreamingManager& Manager = IStreamingManager::Get().GetRenderAssetStreamingManager();
+		Manager.BlockTillAllRequestsFinished();
+
+		const bool bOriginalForcedFullyLoad = SkeletalMesh->bForceMiplevelsToBeResident;
+		SkeletalMesh->bForceMiplevelsToBeResident = true;
+		Manager.UpdateIndividualRenderAsset(SkeletalMesh);
+
+		while (SkeletalMesh->UpdateStreamingStatus())
+		{
+			FlushRenderingCommands();
+			FPlatformProcess::Sleep(RENDER_ASSET_STREAMING_SLEEP_DT);
+		}
+		check(SkeletalMesh->GetResourceForRendering()->CurrentFirstLODIdx <= SkeletalMesh->MinLod.Default);
+
+		SkeletalMesh->UnlinkStreaming();
+		SkeletalMesh->bForceMiplevelsToBeResident = bOriginalForcedFullyLoad;
+	}
+
+PRAGMA_DISABLE_DEPRECATION_WARNINGS
+	bCPUSkinning = bEnable;
+PRAGMA_ENABLE_DEPRECATION_WARNINGS
+	
+	if (IsRegistered())
+	{
+		if (bRecreateRenderStateImmediately)
+		{
+			RecreateRenderState_Concurrent();
+			FlushRenderingCommands();
+		}
+		else
+		{
+			MarkRenderStateDirty();
+		}
+	}
 }
 
 bool USkinnedMeshComponent::GetMaterialStreamingData(int32 MaterialIndex, FPrimitiveMaterialInfo& MaterialData) const
@@ -921,6 +977,13 @@ bool USkinnedMeshComponent::GetMaterialStreamingData(int32 MaterialIndex, FPrimi
 void USkinnedMeshComponent::GetStreamingRenderAssetInfo(FStreamingTextureLevelContext& LevelContext, TArray<FStreamingRenderAssetPrimitiveInfo>& OutStreamingRenderAssets) const
 {
 	GetStreamingTextureInfoInner(LevelContext, nullptr, GetComponentTransform().GetMaximumAxisScale() * StreamingDistanceMultiplier, OutStreamingRenderAssets);
+
+	if (IsStreamingRenderAsset(SkeletalMesh))
+	{
+		const int32 LocalForcedLodModel = GetForcedLOD();
+		const float TexelFactor = LocalForcedLodModel > 0 ? -(SkeletalMesh->GetLODNum() - LocalForcedLodModel + 1) : Bounds.SphereRadius * 2.f;
+		new (OutStreamingRenderAssets) FStreamingRenderAssetPrimitiveInfo(SkeletalMesh, Bounds, TexelFactor, PackedRelativeBox_Identity);
+	}
 }
 
 bool USkinnedMeshComponent::ShouldUpdateBoneVisibility() const
@@ -2742,7 +2805,21 @@ void USkinnedMeshComponent::UnHideBoneByName( FName BoneName )
 
 void USkinnedMeshComponent::SetForcedLOD(int32 InNewForcedLOD)
 {
+PRAGMA_DISABLE_DEPRECATION_WARNINGS
+	const int32 OldValue = ForcedLodModel;
 	ForcedLodModel = FMath::Clamp(InNewForcedLOD, 0, GetNumLODs());
+	if (OldValue != ForcedLodModel)
+	{
+		IStreamingManager::Get().NotifyPrimitiveUpdated(this);
+	}
+PRAGMA_ENABLE_DEPRECATION_WARNINGS
+}
+
+int32 USkinnedMeshComponent::GetForcedLOD() const
+{
+PRAGMA_DISABLE_DEPRECATION_WARNINGS
+	return ForcedLodModel;
+PRAGMA_ENABLE_DEPRECATION_WARNINGS
 }
 
 int32 USkinnedMeshComponent::GetNumLODs() const
@@ -2833,9 +2910,10 @@ bool USkinnedMeshComponent::UpdateLODStatus_Internal(int32 InMasterPoseComponent
 		}
 
 		// Support forcing to a particular LOD.
-		if (ForcedLodModel > 0)
+		const int32 LocalForcedLodModel = GetForcedLOD();
+		if (LocalForcedLodModel > 0)
 		{
-			NewPredictedLODLevel = FMath::Clamp(ForcedLodModel - 1, MinLodIndex, MaxLODIndex);
+			NewPredictedLODLevel = FMath::Clamp(LocalForcedLodModel - 1, MinLodIndex, MaxLODIndex);
 		}
 		else
 		{
@@ -2872,6 +2950,11 @@ bool USkinnedMeshComponent::UpdateLODStatus_Internal(int32 InMasterPoseComponent
 			}
 		}
 
+		if (MeshObject)
+		{
+			NewPredictedLODLevel = FMath::Max(NewPredictedLODLevel, MeshObject->MinDesiredLODLevel);
+		}
+
 #if !(UE_BUILD_SHIPPING || UE_BUILD_TEST)
 		if (CVarAnimVisualizeLODs.GetValueOnAnyThread() != 0)
 		{
@@ -2883,7 +2966,7 @@ bool USkinnedMeshComponent::UpdateLODStatus_Internal(int32 InMasterPoseComponent
 				{
 					const float ScreenSize = FMath::Sqrt(MeshObject->MaxDistanceFactor) * 2.f;
 					FString DebugString = FString::Printf(TEXT("PredictedLODLevel(%d)\nMinDesiredLODLevel(%d) ForcedLodModel(%d) MinLodIndex(%d) LODBias(%d)\nMaxDistanceFactor(%f) ScreenSize(%f)"),
-						PredictedLODLevel, MeshObject->MinDesiredLODLevel, ForcedLodModel, MinLodIndex, LODBias, MeshObject->MaxDistanceFactor, ScreenSize);
+						PredictedLODLevel, MeshObject->MinDesiredLODLevel, LocalForcedLodModel, MinLodIndex, LODBias, MeshObject->MaxDistanceFactor, ScreenSize);
 
 					// See if Child classes want to add something.
 					UpdateVisualizeLODString(DebugString);
@@ -2913,7 +2996,7 @@ bool USkinnedMeshComponent::UpdateLODStatus_Internal(int32 InMasterPoseComponent
 	// See if LOD has changed. 
 	bool bLODChanged = (NewPredictedLODLevel != OldPredictedLODLevel);
 	PredictedLODLevel = NewPredictedLODLevel;
-	
+
 	// also update slave component LOD status, as we may need to recalc required bones if this changes
 	// independently of our LOD
 	for (const TWeakObjectPtr<USkinnedMeshComponent>& SlaveComponent : SlavePoseComponents)
@@ -2998,26 +3081,23 @@ void USkinnedMeshComponent::SetComponentSpaceTransformsDoubleBuffering(bool bInD
 
 void USkinnedMeshComponent::GetCPUSkinnedVertices(TArray<FFinalSkinVertex>& OutVertices, int32 InLODIndex)
 {
-	// switch to CPU skinning
-	bool bCachedCPUSkinning = bCPUSkinning;
-	bCPUSkinning = true;
-
 	if (USkinnedMeshComponent* MasterPoseComponentPtr = MasterPoseComponent.Get())
 	{
-		MasterPoseComponentPtr->ForcedLodModel = InLODIndex + 1;
+		MasterPoseComponentPtr->SetForcedLOD(InLODIndex + 1);
 		MasterPoseComponentPtr->UpdateLODStatus();
 		MasterPoseComponentPtr->RefreshBoneTransforms(nullptr);
 	}
 	else
 	{
-		ForcedLodModel = InLODIndex + 1;
+		SetForcedLOD(InLODIndex + 1);
 		UpdateLODStatus();
 		RefreshBoneTransforms(nullptr);
 	}
 
-	// Recreate render state and flush the renderer
-	RecreateRenderState_Concurrent();
-	FlushRenderingCommands();
+	// switch to CPU skinning
+	const bool bCachedCPUSkinning = GetCPUSkinningEnabled();
+	constexpr bool bRecreateRenderStateImmediately = true;
+	SetCPUSkinningEnabled(true, bRecreateRenderStateImmediately);
 
 	check(MeshObject);
 	check(MeshObject->IsCPUSkinned());
@@ -3026,9 +3106,8 @@ void USkinnedMeshComponent::GetCPUSkinnedVertices(TArray<FFinalSkinVertex>& OutV
 	OutVertices = static_cast<FSkeletalMeshObjectCPUSkin*>(MeshObject)->GetCachedFinalVertices();
 	
 	// switch skinning mode, LOD etc. back
-	bCPUSkinning = bCachedCPUSkinning;
-	ForcedLodModel = 0;
-	RecreateRenderState_Concurrent();
+	SetForcedLOD(0);
+	SetCPUSkinningEnabled(bCachedCPUSkinning, bRecreateRenderStateImmediately);
 }
 
 void USkinnedMeshComponent::ReleaseResources()
diff --git a/Engine/Source/Runtime/Engine/Private/ContentStreaming.cpp b/Engine/Source/Runtime/Engine/Private/ContentStreaming.cpp
--- a/Engine/Source/Runtime/Engine/Private/ContentStreaming.cpp
+++ b/Engine/Source/Runtime/Engine/Private/ContentStreaming.cpp
@@ -7,6 +7,7 @@
 #include "ContentStreaming.h"
 #include "Engine/Texture2D.h"
 #include "Engine/StaticMesh.h"
+#include "Engine/SkeletalMesh.h"
 #include "Misc/CommandLine.h"
 #include "Misc/ConfigCacheIni.h"
 #include "UObject/UObjectHash.h"
@@ -951,11 +952,21 @@ bool FStreamingManagerCollection::IsStreamingEnabled() const
 }
 
 bool FStreamingManagerCollection::IsTextureStreamingEnabled() const
+{
+	return IsRenderAssetStreamingEnabled();
+}
+
+bool FStreamingManagerCollection::IsRenderAssetStreamingEnabled() const
 {
 	return TextureStreamingManager != 0;
 }
 
 IRenderAssetStreamingManager& FStreamingManagerCollection::GetTextureStreamingManager() const
+{
+	return GetRenderAssetStreamingManager();
+}
+
+IRenderAssetStreamingManager& FStreamingManagerCollection::GetRenderAssetStreamingManager() const
 {
 	check(TextureStreamingManager != 0);
 	return *TextureStreamingManager;
@@ -1245,6 +1256,13 @@ void FStreamingManagerCollection::AddOrRemoveTextureStreamingManagerIfNeeded(boo
 						It->LinkStreaming();
 					}
 				}
+				for (TObjectIterator<USkeletalMesh> It; It; ++It)
+				{
+					if (It->bIsStreamable)
+					{
+						It->LinkStreaming();
+					}
+				}
 			}
 		}
 	}
@@ -1274,6 +1292,13 @@ void FStreamingManagerCollection::AddOrRemoveTextureStreamingManagerIfNeeded(boo
 					It->LinkStreaming();
 				}
 			}
+			for (TObjectIterator<USkeletalMesh> It; It; ++It)
+			{
+				if (It->bIsStreamable)
+				{
+					It->LinkStreaming();
+				}
+			}
 		}
 	}
 }
diff --git a/Engine/Source/Runtime/Engine/Private/GameViewportClient.cpp b/Engine/Source/Runtime/Engine/Private/GameViewportClient.cpp
--- a/Engine/Source/Runtime/Engine/Private/GameViewportClient.cpp
+++ b/Engine/Source/Runtime/Engine/Private/GameViewportClient.cpp
@@ -3432,7 +3432,7 @@ bool UGameViewportClient::HandleForceSkelLODCommand( const TCHAR* Cmd, FOutputDe
 		USkeletalMeshComponent* SkelComp = *It;
 		if( SkelComp->GetScene() == InWorld->Scene && !SkelComp->IsTemplate())
 		{
-			SkelComp->ForcedLodModel = ForceLod;
+			SkelComp->SetForcedLOD(ForceLod);
 		}
 	}
 	return true;
diff --git a/Engine/Source/Runtime/Engine/Private/RawIndexBuffer.cpp b/Engine/Source/Runtime/Engine/Private/RawIndexBuffer.cpp
--- a/Engine/Source/Runtime/Engine/Private/RawIndexBuffer.cpp
+++ b/Engine/Source/Runtime/Engine/Private/RawIndexBuffer.cpp
@@ -373,6 +373,7 @@ void FRawStaticIndexBuffer16or32<INDEX_TYPE>::CacheOptimize()
 {
 #if WITH_EDITOR
 	CacheOptimizeIndexBuffer(Indices);
+	CachedNumIndices = Indices.Num();
 #endif
 }
 
diff --git a/Engine/Source/Runtime/Engine/Private/Rendering/MultiSizeIndexContainer.cpp b/Engine/Source/Runtime/Engine/Private/Rendering/MultiSizeIndexContainer.cpp
--- a/Engine/Source/Runtime/Engine/Private/Rendering/MultiSizeIndexContainer.cpp
+++ b/Engine/Source/Runtime/Engine/Private/Rendering/MultiSizeIndexContainer.cpp
@@ -163,6 +163,57 @@ void FMultiSizeIndexContainer::Serialize(FArchive& Ar, bool bNeedsCPUAccess)
 	IndexBuffer->Serialize(Ar);
 }
 
+void FMultiSizeIndexContainer::SerializeMetaData(FArchive& Ar, bool bNeedsCPUAccess)
+{
+	Ar << DataTypeSize;
+
+	if (!IndexBuffer)
+	{
+		if (DataTypeSize == sizeof(uint16))
+		{
+			IndexBuffer = new FRawStaticIndexBuffer16or32<uint16>(bNeedsCPUAccess);
+		}
+		else
+		{
+			IndexBuffer = new FRawStaticIndexBuffer16or32<uint32>(bNeedsCPUAccess);
+		}
+	}
+
+	IndexBuffer->SerializeMetaData(Ar);
+}
+
+FIndexBufferRHIRef FMultiSizeIndexContainer::CreateRHIBuffer_RenderThread()
+{
+	if (IndexBuffer)
+	{
+		if (DataTypeSize == sizeof(uint16))
+		{
+			return static_cast<FRawStaticIndexBuffer16or32<uint16>*>(IndexBuffer)->CreateRHIBuffer_RenderThread();
+		}
+		else
+		{
+			return static_cast<FRawStaticIndexBuffer16or32<uint32>*>(IndexBuffer)->CreateRHIBuffer_RenderThread();
+		}
+	}
+	return nullptr;
+}
+
+FIndexBufferRHIRef FMultiSizeIndexContainer::CreateRHIBuffer_Async()
+{
+	if (IndexBuffer)
+	{
+		if (DataTypeSize == sizeof(uint16))
+		{
+			return static_cast<FRawStaticIndexBuffer16or32<uint16>*>(IndexBuffer)->CreateRHIBuffer_Async();
+		}
+		else
+		{
+			return static_cast<FRawStaticIndexBuffer16or32<uint32>*>(IndexBuffer)->CreateRHIBuffer_Async();
+		}
+	}
+	return nullptr;
+}
+
 #if WITH_EDITOR
 /**
 * Retrieves index buffer related data
diff --git a/Engine/Source/Runtime/Engine/Private/Rendering/SkeletalMeshVertexClothBuffer.cpp b/Engine/Source/Runtime/Engine/Private/Rendering/SkeletalMeshVertexClothBuffer.cpp
--- a/Engine/Source/Runtime/Engine/Private/Rendering/SkeletalMeshVertexClothBuffer.cpp
+++ b/Engine/Source/Runtime/Engine/Private/Rendering/SkeletalMeshVertexClothBuffer.cpp
@@ -65,18 +65,55 @@ void FSkeletalMeshVertexClothBuffer::CleanUp()
 	VertexData = nullptr;
 }
 
+void FSkeletalMeshVertexClothBuffer::ClearMetaData()
+{
+	ClothIndexMapping.Empty();
+	NumVertices = 0;
+}
+
+template <bool bRenderThread>
+FVertexBufferRHIRef FSkeletalMeshVertexClothBuffer::CreateRHIBuffer_Internal()
+{
+	if (NumVertices)
+	{
+		FResourceArrayInterface* ResourceArray = VertexData ? VertexData->GetResourceArray() : nullptr;
+		const uint32 SizeInBytes = ResourceArray ? ResourceArray->GetResourceDataSize() : 0;
+		const uint32 BuffFlags = BUF_Static | BUF_ShaderResource;
+		FRHIResourceCreateInfo CreateInfo(ResourceArray);
+		CreateInfo.bWithoutNativeResource = !VertexData;
+
+		if (bRenderThread)
+		{
+			return RHICreateVertexBuffer(SizeInBytes, BuffFlags, CreateInfo);
+		}
+		else
+		{
+			return RHIAsyncCreateVertexBuffer(SizeInBytes, BuffFlags, CreateInfo);
+		}
+	}
+	return nullptr;
+}
+
+FVertexBufferRHIRef FSkeletalMeshVertexClothBuffer::CreateRHIBuffer_RenderThread()
+{
+	return CreateRHIBuffer_Internal<true>();
+}
+
+FVertexBufferRHIRef FSkeletalMeshVertexClothBuffer::CreateRHIBuffer_Async()
+{
+	return CreateRHIBuffer_Internal<false>();
+}
+
 /**
 * Initialize the RHI resource for this vertex buffer
 */
 void FSkeletalMeshVertexClothBuffer::InitRHI()
 {
-	check(VertexData);
-	FResourceArrayInterface* ResourceArray = VertexData->GetResourceArray();
-	if (ResourceArray->GetResourceDataSize() > 0)
+	VertexBufferRHI = CreateRHIBuffer_RenderThread();
+
+	if (VertexBufferRHI)
 	{
-		FRHIResourceCreateInfo CreateInfo(ResourceArray);
-		VertexBufferRHI = RHICreateVertexBuffer(ResourceArray->GetResourceDataSize(), BUF_Static | BUF_ShaderResource, CreateInfo);
-		VertexBufferSRV = RHICreateShaderResourceView(VertexBufferRHI, sizeof(FVector4), PF_A32B32G32R32F);
+		VertexBufferSRV = RHICreateShaderResourceView(VertexData ? VertexBufferRHI : nullptr, sizeof(FVector4), PF_A32B32G32R32F);
 	}
 }
 
@@ -122,6 +159,11 @@ FArchive& operator<<(FArchive& Ar, FSkeletalMeshVertexClothBuffer& VertexBuffer)
 	return Ar;
 }
 
+void FSkeletalMeshVertexClothBuffer::SerializeMetaData(FArchive& Ar)
+{
+	Ar << ClothIndexMapping << Stride << NumVertices;
+}
+
 
 /**
 * Initializes the buffer with the given vertices.
diff --git a/Engine/Source/Runtime/Engine/Private/Rendering/SkinWeightVertexBuffer.cpp b/Engine/Source/Runtime/Engine/Private/Rendering/SkinWeightVertexBuffer.cpp
--- a/Engine/Source/Runtime/Engine/Private/Rendering/SkinWeightVertexBuffer.cpp
+++ b/Engine/Source/Runtime/Engine/Private/Rendering/SkinWeightVertexBuffer.cpp
@@ -96,8 +96,7 @@ FArchive& operator<<(FArchive& Ar, FSkinWeightVertexBuffer& VertexBuffer)
 {
 	FStripDataFlags StripFlags(Ar);
 
-	Ar << VertexBuffer.bExtraBoneInfluences;
-	Ar << VertexBuffer.NumVertices;
+	VertexBuffer.SerializeMetaData(Ar);
 
 	if (Ar.IsLoading() || VertexBuffer.WeightData == NULL)
 	{
@@ -124,28 +123,67 @@ FArchive& operator<<(FArchive& Ar, FSkinWeightVertexBuffer& VertexBuffer)
 	return Ar;
 }
 
-void FSkinWeightVertexBuffer::InitRHI()
+void FSkinWeightVertexBuffer::SerializeMetaData(FArchive& Ar)
+{
+	Ar << bExtraBoneInfluences << Stride << NumVertices;
+}
+
+void FSkinWeightVertexBuffer::CopyMetaData(const FSkinWeightVertexBuffer& Other)
 {
-	check(WeightData);
-	FResourceArrayInterface* ResourceArray = WeightData->GetResourceArray();
-	if (ResourceArray->GetResourceDataSize() > 0)
+	bExtraBoneInfluences = Other.bExtraBoneInfluences;
+	Stride = Other.Stride;
+	NumVertices = Other.NumVertices;
+}
+
+template <bool bRenderThread>
+FVertexBufferRHIRef FSkinWeightVertexBuffer::CreateRHIBuffer_Internal()
+{
+	if (NumVertices)
 	{
 		// Create the vertex buffer.
+		FResourceArrayInterface* ResourceArray = WeightData ? WeightData->GetResourceArray() : nullptr;
+		const uint32 SizeInBytes = ResourceArray ? ResourceArray->GetResourceDataSize() : 0;
+		const uint32 BuffFlags = BUF_Static | BUF_ShaderResource;
 		FRHIResourceCreateInfo CreateInfo(ResourceArray);
+		CreateInfo.bWithoutNativeResource = !WeightData;
 
 		// BUF_ShaderResource is needed for support of the SkinCache (we could make is dependent on GEnableGPUSkinCacheShaders or are there other users?)
-		VertexBufferRHI = RHICreateVertexBuffer(ResourceArray->GetResourceDataSize(), BUF_Static | BUF_ShaderResource, CreateInfo);
-
-		bool bSRV = GSupportsResourceView && GPixelFormats[PF_R32_UINT].Supported;
-		// When bAllowCPUAccess is true, the meshes is likely going to be used for Niagara to spawn particles on mesh surface.
-		// And it can be the case for CPU *and* GPU access: no differenciation today. That is why we create a SRV in this case.
-		// This also avoid setting lots of states on all the members of all the different buffers used by meshes. Follow up: https://jira.it.epicgames.net/browse/UE-69376.
-		bSRV |= GetNeedsCPUAccess();
-
-		if (bSRV)
+		if (bRenderThread)
 		{
-			SRVValue = RHICreateShaderResourceView(VertexBufferRHI, 4, PF_R32_UINT);
+			return RHICreateVertexBuffer(SizeInBytes, BuffFlags, CreateInfo);
 		}
+		else
+		{
+			return RHIAsyncCreateVertexBuffer(SizeInBytes, BuffFlags, CreateInfo);
+		}
+	}
+	return nullptr;
+}
+
+FVertexBufferRHIRef FSkinWeightVertexBuffer::CreateRHIBuffer_RenderThread()
+{
+	return CreateRHIBuffer_Internal<true>();
+}
+
+FVertexBufferRHIRef FSkinWeightVertexBuffer::CreateRHIBuffer_Async()
+{
+	return CreateRHIBuffer_Internal<false>();
+}
+
+void FSkinWeightVertexBuffer::InitRHI()
+{
+	// BUF_ShaderResource is needed for support of the SkinCache (we could make is dependent on GEnableGPUSkinCacheShaders or are there other users?)
+	VertexBufferRHI = CreateRHIBuffer_RenderThread();
+
+	bool bSRV = VertexBufferRHI && GSupportsResourceView && GPixelFormats[PF_R32_UINT].Supported;
+	// When bAllowCPUAccess is true, the meshes is likely going to be used for Niagara to spawn particles on mesh surface.
+	// And it can be the case for CPU *and* GPU access: no differenciation today. That is why we create a SRV in this case.
+	// This also avoid setting lots of states on all the members of all the different buffers used by meshes. Follow up: https://jira.it.epicgames.net/browse/UE-69376.
+	bSRV |= GetNeedsCPUAccess();
+
+	if (bSRV)
+	{
+		SRVValue = RHICreateShaderResourceView(WeightData ? VertexBufferRHI : nullptr, 4, PF_R32_UINT);
 	}
 }
 
diff --git a/Engine/Source/Runtime/Engine/Private/SkeletalMesh.cpp b/Engine/Source/Runtime/Engine/Private/SkeletalMesh.cpp
--- a/Engine/Source/Runtime/Engine/Private/SkeletalMesh.cpp
+++ b/Engine/Source/Runtime/Engine/Private/SkeletalMesh.cpp
@@ -52,6 +52,9 @@
 #include "Animation/AnimSequence.h"
 #include "UObject/NiagaraObjectVersion.h"
 #include "Animation/SkinWeightProfile.h"
+#include "Streaming/SkeletalMeshUpdate.h"
+#include "UObject/CoreRedirects.h"
+#include "HAL/FileManager.h"
 
 #if WITH_EDITOR
 #include "Rendering/SkeletalMeshModel.h"
@@ -298,6 +301,10 @@ USkeletalMesh::USkeletalMesh(const FObjectInitializer& ObjectInitializer)
 #if WITH_EDITORONLY_DATA
 	ImportedModel = MakeShareable(new FSkeletalMeshModel());
 	VertexColorGuid = FGuid();
+	bSupportLODStreaming.Default = false;
+	MaxNumStreamedLODs.Default = 0;
+	// TODO: support saving some but not all optional LODs
+	MaxNumOptionalLODs.Default = 0;
 #endif
 	
 	MinLod.Default = 0;
@@ -621,7 +628,19 @@ void USkeletalMesh::InitResources()
 	FSkeletalMeshRenderData* SkelMeshRenderData = GetResourceForRendering();
 	if (SkelMeshRenderData)
 	{
-		SkelMeshRenderData->InitResources(bHasVertexColors, MorphTargets);
+		SkelMeshRenderData->InitResources(bHasVertexColors, MorphTargets, this);
+	}
+
+	// Determine whether or not this mesh can be streamed.
+	const int32 NumLODs = GetLODNum();
+	bIsStreamable = !NeverStream
+		&& NumLODs > 1
+		&& !SkelMeshRenderData->LODRenderData[0].bStreamedDataInlined;
+
+	UnlinkStreaming();
+	if (bIsStreamable)
+	{
+		LinkStreaming();
 	}
 }
 
@@ -758,6 +777,252 @@ void USkeletalMesh::GetResourceSizeEx(FResourceSizeEx& CumulativeResourceSize)
 	CumulativeResourceSize.AddDedicatedSystemMemoryBytes(RefSkeleton.GetDataSize());
 }
 
+int32 USkeletalMesh::GetLODGroupForStreaming() const
+{
+	// TODO: mesh LOD streaming may need to know LOD group settings
+	return 0;
+}
+
+int32 USkeletalMesh::GetNumMipsForStreaming() const
+{
+	return GetLODNum();
+}
+
+int32 USkeletalMesh::GetNumNonStreamingMips() const
+{
+	return SkeletalMeshRenderData->NumInlinedLODs;
+}
+
+int32 USkeletalMesh::CalcNumOptionalMips() const
+{
+	return SkeletalMeshRenderData->NumOptionalLODs;
+}
+
+int32 USkeletalMesh::CalcCumulativeLODSize(int32 NumLODs) const
+{
+	uint32 Accum = 0;
+	const int32 LODCount = GetLODNum();
+	const int32 LastLODIdx = LODCount - NumLODs;
+	for (int32 LODIdx = LODCount - 1; LODIdx >= LastLODIdx; --LODIdx)
+	{
+		Accum += SkeletalMeshRenderData->LODRenderData[LODIdx].BuffersSize;
+	}
+	check(Accum >= 0u);
+	return Accum;
+}
+
+bool USkeletalMesh::GetMipDataFilename(const int32 MipIndex, FString& OutBulkDataFilename) const
+{
+#if !WITH_EDITOR
+	// TODO: this is slow. Should cache the name once per mesh
+	FString PackageName = GetOutermost()->FileName.ToString();
+	// Handle name redirection and localization
+	const FCoreRedirectObjectName RedirectedName =
+		FCoreRedirects::GetRedirectedName(
+			ECoreRedirectFlags::Type_Package,
+			FCoreRedirectObjectName(NAME_None, NAME_None, *PackageName));
+	FString LocalizedName;
+	LocalizedName = FPackageName::GetDelegateResolvedPackagePath(RedirectedName.PackageName.ToString());
+	LocalizedName = FPackageName::GetLocalizedPackagePath(LocalizedName);
+	bool bSucceed = FPackageName::DoesPackageExist(LocalizedName, nullptr, &OutBulkDataFilename);
+	check(bSucceed);
+	OutBulkDataFilename = FPaths::ChangeExtension(OutBulkDataFilename, MipIndex < CalcNumOptionalMips() ? TEXT(".uptnl") : TEXT(".ubulk"));
+	check(MipIndex < CalcNumOptionalMips() || IFileManager::Get().FileExists(*OutBulkDataFilename));
+	return true;
+#else
+	return false;
+#endif
+}
+
+bool USkeletalMesh::IsReadyForStreaming() const
+{
+	return SkeletalMeshRenderData && SkeletalMeshRenderData->bReadyForStreaming;
+}
+
+int32 USkeletalMesh::GetNumResidentMips() const
+{
+	return GetLODNum() - SkeletalMeshRenderData->CurrentFirstLODIdx;
+}
+
+int32 USkeletalMesh::GetNumRequestedMips() const
+{
+	if (PendingUpdate && !PendingUpdate->IsCancelled())
+	{
+		return PendingUpdate->GetNumRequestedMips();
+	}
+	else
+	{
+		return GetCachedNumResidentLODs();
+	}
+}
+
+bool USkeletalMesh::CancelPendingMipChangeRequest()
+{
+	if (PendingUpdate)
+	{
+		if (!PendingUpdate->IsCancelled())
+		{
+			PendingUpdate->Abort();
+		}
+		return true;
+	}
+	return false;
+}
+
+bool USkeletalMesh::HasPendingUpdate() const
+{
+	return !!PendingUpdate;
+}
+
+bool USkeletalMesh::IsPendingUpdateLocked() const
+{
+	return PendingUpdate && PendingUpdate->IsLocked();
+}
+
+bool USkeletalMesh::StreamOut(int32 NewMipCount)
+{
+	check(IsInGameThread());
+	if (bIsStreamable
+		&& !PendingUpdate
+		&& SkeletalMeshRenderData.IsValid()
+		&& SkeletalMeshRenderData->bReadyForStreaming
+		&& NewMipCount < GetNumResidentMips())
+	{
+		PendingUpdate = new FSkeletalMeshStreamOut(this, NewMipCount);
+		return !PendingUpdate->IsCancelled();
+	}
+	return false;
+}
+
+bool USkeletalMesh::StreamIn(int32 NewMipCount, bool bHighPrio)
+{
+	check(IsInGameThread());
+	if (bIsStreamable
+		&& !PendingUpdate
+		&& SkeletalMeshRenderData.IsValid()
+		&& SkeletalMeshRenderData->bReadyForStreaming
+		&& NewMipCount > GetNumResidentMips())
+	{
+#if WITH_EDITOR
+		if (FPlatformProperties::HasEditorOnlyData())
+		{
+			if (GRHISupportsAsyncTextureCreation)
+			{
+				PendingUpdate = new FSkeletalMeshStreamIn_DDC_Async(this, NewMipCount);
+			}
+			else
+			{
+				PendingUpdate = new FSkeletalMeshStreamIn_DDC_RenderThread(this, NewMipCount);
+			}
+		}
+		else
+#endif
+		{
+			if (GRHISupportsAsyncTextureCreation)
+			{
+				PendingUpdate = new FSkeletalMeshStreamIn_IO_Async(this, NewMipCount, bHighPrio);
+			}
+			else
+			{
+				PendingUpdate = new FSkeletalMeshStreamIn_IO_RenderThread(this, NewMipCount, bHighPrio);
+			}
+		}
+		return !PendingUpdate->IsCancelled();
+	}
+	return false;
+}
+
+bool USkeletalMesh::UpdateStreamingStatus(bool bWaitForMipFading)
+{
+	// if resident and requested mip counts match then no pending request is in flight
+	if (PendingUpdate)
+	{
+		if (GIsRequestingExit || !SkeletalMeshRenderData)
+		{
+			PendingUpdate->Abort();
+		}
+
+		// When there is no renderthread, allow the gamethread to tick as the renderthread.
+		FRenderAssetUpdate::EThreadType TickThread = GIsThreadedRendering ? FRenderAssetUpdate::TT_None : FRenderAssetUpdate::TT_Render;
+		if (HasAnyFlags(RF_BeginDestroyed) && PendingUpdate->GetRelevantThread() == FRenderAssetUpdate::TT_Async)
+		{
+			// To avoid async tasks from timing out the GC, we tick as Async to force completion if this is relevant.
+			// This could lead the asset from releasing the PendingUpdate, which will be deleted once the async task completes.
+			TickThread = FRenderAssetUpdate::TT_Async;
+		}
+		PendingUpdate->Tick(TickThread);
+
+		if (!PendingUpdate->IsCompleted())
+		{
+			return true;
+		}
+
+#if WITH_EDITOR
+		const bool bRebuildPlatformData = PendingUpdate->DDCIsInvalid() && !IsPendingKillOrUnreachable();
+#endif
+
+		PendingUpdate.SafeRelease();
+
+#if WITH_EDITOR
+		if (GIsEditor)
+		{
+			// When all the requested mips are streamed in, generate an empty property changed event, to force the
+			// ResourceSize asset registry tag to be recalculated.
+			FPropertyChangedEvent EmptyPropertyChangedEvent(nullptr);
+			FCoreUObjectDelegates::OnObjectPropertyChanged.Broadcast(this, EmptyPropertyChangedEvent);
+
+			// We can't load the source art from a bulk data object if the mesh itself is pending kill because the linker will have been detached.
+			// In this case we don't rebuild the data and instead let the streaming request be cancelled. This will let the garbage collector finish
+			// destroying the object.
+			if (bRebuildPlatformData)
+			{
+				// TODO: force rebuild even if DDC keys match
+				SkeletalMeshRenderData->Cache(this);
+				// @TODO this can not be called from this callstack since the entry needs to be removed completely from the streamer.
+				// UpdateResource();
+			}
+		}
+#endif
+	}
+
+	// TODO: LOD fading?
+
+	return false;
+}
+
+void USkeletalMesh::LinkStreaming()
+{
+	if (!IsTemplate() && IStreamingManager::Get().IsTextureStreamingEnabled() && IsStreamingRenderAsset(this))
+	{
+		IStreamingManager::Get().GetTextureStreamingManager().AddStreamingRenderAsset(this);
+	}
+	else
+	{
+		StreamingIndex = INDEX_NONE;
+	}
+}
+
+void USkeletalMesh::UnlinkStreaming()
+{
+	if (!IsTemplate() && IStreamingManager::Get().IsTextureStreamingEnabled())
+	{
+		IStreamingManager::Get().GetTextureStreamingManager().RemoveStreamingRenderAsset(this);
+	}
+}
+
+void USkeletalMesh::CancelAllPendingStreamingActions()
+{
+	FlushRenderingCommands();
+
+	for (TObjectIterator<USkeletalMesh> It; It; ++It)
+	{
+		USkeletalMesh* StaticMesh = *It;
+		StaticMesh->CancelPendingMipChangeRequest();
+	}
+
+	FlushRenderingCommands();
+}
+
 /**
  * Operator for MemCount only, so it only serializes the arrays that needs to be counted.
 */
@@ -976,6 +1241,12 @@ void USkeletalMesh::BeginDestroy()
 {
 	Super::BeginDestroy();
 
+	// Cancel any in flight IO requests
+	CancelPendingMipChangeRequest();
+
+	// Safely unlink mesh from list of streamable ones.
+	UnlinkStreaming();
+
 	// remove the cache of link up
 	if (Skeleton)
 	{
@@ -3907,7 +4178,7 @@ void FSkeletalMeshSceneProxy::GetMeshElementsConditionallySelectable(const TArra
 		if (VisibilityMap & (1 << ViewIndex))
 		{
 			const FSceneView* View = Views[ViewIndex];
-			MeshObject->UpdateMinDesiredLODLevel(View, GetBounds(), ViewFamily.FrameNumber);
+			MeshObject->UpdateMinDesiredLODLevel(View, GetBounds(), ViewFamily.FrameNumber, SkeletalMeshRenderData->PendingFirstLODIdx);
 		}
 	}
 
@@ -3917,7 +4188,7 @@ void FSkeletalMeshSceneProxy::GetMeshElementsConditionallySelectable(const TArra
 	check(LODIndex < SkeletalMeshRenderData->LODRenderData.Num());
 	const FSkeletalMeshLODRenderData& LODData = SkeletalMeshRenderData->LODRenderData[LODIndex];
 
-	if( LODSections.Num() > 0 )
+	if( LODSections.Num() > 0 && LODIndex >= SkeletalMeshRenderData->CurrentFirstLODIdx )
 	{
 		const FLODSectionElements& LODSection = LODSections[LODIndex];
 
@@ -4014,6 +4285,11 @@ void FSkeletalMeshSceneProxy::CreateBaseMeshBatch(const FSceneView* View, const
 	BatchElement.NumPrimitives = LODData.RenderSections[SectionIndex].NumTriangles;
 }
 
+uint8 FSkeletalMeshSceneProxy::GetCurrentFirstLODIdx_Internal() const
+{
+	return SkeletalMeshRenderData->CurrentFirstLODIdx;
+}
+
 void FSkeletalMeshSceneProxy::GetDynamicElementsSection(const TArray<const FSceneView*>& Views, const FSceneViewFamily& ViewFamily, uint32 VisibilityMap, 
 	const FSkeletalMeshLODRenderData& LODData, const int32 LODIndex, const int32 SectionIndex, bool bSectionSelected,
 	const FSectionElementInfo& SectionElementInfo, bool bInSelectable, FMeshElementCollector& Collector ) const
diff --git a/Engine/Source/Runtime/Engine/Private/SkeletalMeshLODRenderData.cpp b/Engine/Source/Runtime/Engine/Private/SkeletalMeshLODRenderData.cpp
--- a/Engine/Source/Runtime/Engine/Private/SkeletalMeshLODRenderData.cpp
+++ b/Engine/Source/Runtime/Engine/Private/SkeletalMeshLODRenderData.cpp
@@ -8,6 +8,7 @@
 #include "EngineLogs.h"
 #include "EngineUtils.h"
 #include "Interfaces/ITargetPlatform.h"
+#include "Interfaces/ITargetPlatformManagerModule.h"
 #include "PlatformInfo.h"
 
 #if WITH_EDITOR
@@ -178,39 +179,38 @@
 		InstantiatedNodes[NodeIndex] = true;
 	}
 }
-void FSkeletalMeshLODRenderData::InitResources(bool bNeedsVertexColors, int32 LODIndex, TArray<UMorphTarget*>& InMorphTargets)
+
+void FSkeletalMeshLODRenderData::InitResources(bool bNeedsVertexColors, int32 LODIndex, TArray<UMorphTarget*>& InMorphTargets, USkeletalMesh* Owner)
 {
-	INC_DWORD_STAT_BY(STAT_SkeletalMeshIndexMemory, MultiSizeIndexContainer.IsIndexBufferValid() ? (MultiSizeIndexContainer.GetIndexBuffer()->Num() * MultiSizeIndexContainer.GetDataTypeSize()) : 0);
+	IncrementMemoryStats(bNeedsVertexColors);
 
 	MorphTargetVertexInfoBuffers.Reset();
 	MultiSizeIndexContainer.InitResources();
 
-	INC_DWORD_STAT_BY(STAT_SkeletalMeshVertexMemory, StaticVertexBuffers.PositionVertexBuffer.GetStride() * StaticVertexBuffers.PositionVertexBuffer.GetNumVertices());
-	INC_DWORD_STAT_BY(STAT_SkeletalMeshVertexMemory, StaticVertexBuffers.StaticMeshVertexBuffer.GetResourceSize());
 	BeginInitResource(&StaticVertexBuffers.PositionVertexBuffer);
 	BeginInitResource(&StaticVertexBuffers.StaticMeshVertexBuffer);
 
-	INC_DWORD_STAT_BY(STAT_SkeletalMeshVertexMemory, SkinWeightVertexBuffer.GetVertexDataSize());
+	if (GSkinWeightProfilesLoadByDefaultMode == 3)
+	{
+		SkinWeightProfilesData.SetDynamicDefaultSkinWeightProfile(Owner, LODIndex);
+	}
 	BeginInitResource(&SkinWeightVertexBuffer);
 
 	if (bNeedsVertexColors)
 	{
 		// Only init the color buffer if the mesh has vertex colors
-		INC_DWORD_STAT_BY(STAT_SkeletalMeshVertexMemory, StaticVertexBuffers.ColorVertexBuffer.GetAllocatedSize());
 		BeginInitResource(&StaticVertexBuffers.ColorVertexBuffer);
 	}
 
 	if (ClothVertexBuffer.GetNumVertices() > 0)
 	{
 		// Only init the clothing buffer if the mesh has clothing data
-		INC_DWORD_STAT_BY(STAT_SkeletalMeshVertexMemory, ClothVertexBuffer.GetVertexDataSize());
 		BeginInitResource(&ClothVertexBuffer);
 	}
 
 	if (RHISupportsTessellation(GMaxRHIShaderPlatform))
 	{
 		AdjacencyMultiSizeIndexContainer.InitResources();
-		INC_DWORD_STAT_BY(STAT_SkeletalMeshIndexMemory, AdjacencyMultiSizeIndexContainer.IsIndexBufferValid() ? (AdjacencyMultiSizeIndexContainer.GetIndexBuffer()->Num() * AdjacencyMultiSizeIndexContainer.GetDataTypeSize()) : 0);
 	}
 
 	// DuplicatedVerticesBuffer is used only for SkinCache and Editor features which is SM5 only
@@ -465,15 +465,7 @@
 
 void FSkeletalMeshLODRenderData::ReleaseResources()
 {
-	DEC_DWORD_STAT_BY(STAT_SkeletalMeshIndexMemory, MultiSizeIndexContainer.IsIndexBufferValid() ? (MultiSizeIndexContainer.GetIndexBuffer()->Num() * MultiSizeIndexContainer.GetDataTypeSize()) : 0);
-	DEC_DWORD_STAT_BY(STAT_SkeletalMeshIndexMemory, AdjacencyMultiSizeIndexContainer.IsIndexBufferValid() ? (AdjacencyMultiSizeIndexContainer.GetIndexBuffer()->Num() * AdjacencyMultiSizeIndexContainer.GetDataTypeSize()) : 0);
+	DecrementMemoryStats();
 
-	DEC_DWORD_STAT_BY(STAT_SkeletalMeshVertexMemory, StaticVertexBuffers.PositionVertexBuffer.GetStride() * StaticVertexBuffers.PositionVertexBuffer.GetNumVertices());
-	DEC_DWORD_STAT_BY(STAT_SkeletalMeshVertexMemory, StaticVertexBuffers.StaticMeshVertexBuffer.GetResourceSize());
-
-	DEC_DWORD_STAT_BY(STAT_SkeletalMeshVertexMemory, SkinWeightVertexBuffer.GetVertexDataSize());
-	DEC_DWORD_STAT_BY(STAT_SkeletalMeshVertexMemory, StaticVertexBuffers.ColorVertexBuffer.GetAllocatedSize());
-	DEC_DWORD_STAT_BY(STAT_SkeletalMeshVertexMemory, ClothVertexBuffer.GetVertexDataSize());
-
 	MultiSizeIndexContainer.ReleaseResources();
 	AdjacencyMultiSizeIndexContainer.ReleaseResources();
@@ -498,6 +490,42 @@
 	SkinWeightProfilesData.ReleaseResources();
 }
 
+void FSkeletalMeshLODRenderData::IncrementMemoryStats(bool bNeedsVertexColors)
+{
+	INC_DWORD_STAT_BY(STAT_SkeletalMeshIndexMemory, MultiSizeIndexContainer.IsIndexBufferValid() ? (MultiSizeIndexContainer.GetIndexBuffer()->Num() * MultiSizeIndexContainer.GetDataTypeSize()) : 0);
+	INC_DWORD_STAT_BY(STAT_SkeletalMeshVertexMemory, StaticVertexBuffers.PositionVertexBuffer.GetStride() * StaticVertexBuffers.PositionVertexBuffer.GetNumVertices());
+	INC_DWORD_STAT_BY(STAT_SkeletalMeshVertexMemory, StaticVertexBuffers.StaticMeshVertexBuffer.GetResourceSize());
+	INC_DWORD_STAT_BY(STAT_SkeletalMeshVertexMemory, SkinWeightVertexBuffer.GetVertexDataSize());
+
+	if (bNeedsVertexColors)
+	{
+		INC_DWORD_STAT_BY(STAT_SkeletalMeshVertexMemory, StaticVertexBuffers.ColorVertexBuffer.GetAllocatedSize());
+	}
+
+	if (ClothVertexBuffer.GetNumVertices() > 0)
+	{
+		INC_DWORD_STAT_BY(STAT_SkeletalMeshVertexMemory, ClothVertexBuffer.GetVertexDataSize());
+	}
+
+	if (RHISupportsTessellation(GMaxRHIShaderPlatform))
+	{
+		INC_DWORD_STAT_BY(STAT_SkeletalMeshIndexMemory, AdjacencyMultiSizeIndexContainer.IsIndexBufferValid() ? (AdjacencyMultiSizeIndexContainer.GetIndexBuffer()->Num() * AdjacencyMultiSizeIndexContainer.GetDataTypeSize()) : 0);
+	}
+}
+
+void FSkeletalMeshLODRenderData::DecrementMemoryStats()
+{
+	DEC_DWORD_STAT_BY(STAT_SkeletalMeshIndexMemory, MultiSizeIndexContainer.IsIndexBufferValid() ? (MultiSizeIndexContainer.GetIndexBuffer()->Num() * MultiSizeIndexContainer.GetDataTypeSize()) : 0);
+	DEC_DWORD_STAT_BY(STAT_SkeletalMeshIndexMemory, AdjacencyMultiSizeIndexContainer.IsIndexBufferValid() ? (AdjacencyMultiSizeIndexContainer.GetIndexBuffer()->Num() * AdjacencyMultiSizeIndexContainer.GetDataTypeSize()) : 0);
+
+	DEC_DWORD_STAT_BY(STAT_SkeletalMeshVertexMemory, StaticVertexBuffers.PositionVertexBuffer.GetStride() * StaticVertexBuffers.PositionVertexBuffer.GetNumVertices());
+	DEC_DWORD_STAT_BY(STAT_SkeletalMeshVertexMemory, StaticVertexBuffers.StaticMeshVertexBuffer.GetResourceSize());
+
+	DEC_DWORD_STAT_BY(STAT_SkeletalMeshVertexMemory, SkinWeightVertexBuffer.GetVertexDataSize());
+	DEC_DWORD_STAT_BY(STAT_SkeletalMeshVertexMemory, StaticVertexBuffers.ColorVertexBuffer.GetAllocatedSize());
+	DEC_DWORD_STAT_BY(STAT_SkeletalMeshVertexMemory, ClothVertexBuffer.GetVertexDataSize());
+}
+
 #if WITH_EDITOR
 
 void FSkeletalMeshLODRenderData::BuildFromLODModel(const FSkeletalMeshLODModel* ImportedModel, uint32 BuildFlags)
@@ -595,7 +623,7 @@
 
 #endif // WITH_EDITOR
 
-void FSkeletalMeshLODRenderData::ReleaseCPUResources()
+void FSkeletalMeshLODRenderData::ReleaseCPUResources(bool bForStreaming)
 {
 	if (!GIsEditor && !IsRunningCommandlet())
 	{
@@ -614,6 +642,13 @@
 
 		StaticVertexBuffers.PositionVertexBuffer.CleanUp();
 		StaticVertexBuffers.StaticMeshVertexBuffer.CleanUp();
+
+		if (bForStreaming)
+		{
+			ClothVertexBuffer.CleanUp();
+			StaticVertexBuffers.ColorVertexBuffer.CleanUp();
+			SkinWeightProfilesData.ReleaseCPUResources();
+		}
 	}
 }
 
@@ -646,135 +681,355 @@
 	CumulativeResourceSize.AddUnknownMemoryBytes(SkinWeightProfilesData.GetResourcesSize());	
 }
 
-void FSkeletalMeshLODRenderData::Serialize(FArchive& Ar, UObject* Owner, int32 Idx)
+int32 FSkeletalMeshLODRenderData::GetPlatformMinLODIdx(const ITargetPlatform* TargetPlatform, const USkeletalMesh* SkeletalMesh)
 {
-	DECLARE_SCOPE_CYCLE_COUNTER(TEXT("FSkeletalMeshLODRenderData::Serialize"), STAT_SkeletalMeshLODRenderData_Serialize, STATGROUP_LoadTime);
+#if WITH_EDITOR
+	check(TargetPlatform && SkeletalMesh);
+	const FName PlatformGroupName = TargetPlatform->GetPlatformInfo().PlatformGroupName;
+	const FName VanillaPlatformName = TargetPlatform->GetPlatformInfo().VanillaPlatformName;
+	return  SkeletalMesh->MinLod.GetValueForPlatformIdentifiers(PlatformGroupName, VanillaPlatformName);
+#else
+	return 0;
+#endif
+}
 
-	// Defined class flags for possible stripping
-	const uint8 LodAdjacencyStripFlag = 1;
-	const uint8 MinLodStripFlag = 2;
-
-	// Actual flags used during serialization
-	uint8 ClassDataStripFlags = 0;
-
+uint8 FSkeletalMeshLODRenderData::GenerateClassStripFlags(FArchive& Ar, const USkeletalMesh* OwnerMesh, int32 LODIdx)
+{
+#if WITH_EDITOR
 	const bool bIsCook = Ar.IsCooking();
 	const ITargetPlatform* CookTarget = Ar.CookingTarget();
 
 	extern int32 GForceStripMeshAdjacencyDataDuringCooking;
 	const bool bWantToStripTessellation = bIsCook && ((GForceStripMeshAdjacencyDataDuringCooking != 0) || !CookTarget->SupportsFeature(ETargetPlatformFeatures::Tessellation));
 
-	USkeletalMesh* OwnerMesh = CastChecked<USkeletalMesh>(Owner);
 	int32 MinMeshLod = 0;
 	bool bMeshDisablesMinLodStrip = false;
-#if WITH_EDITOR
-	if(bIsCook)
+	if (bIsCook)
 	{
 		MinMeshLod = OwnerMesh ? OwnerMesh->MinLod.GetValueForPlatformIdentifiers(CookTarget->GetPlatformInfo().PlatformGroupName, CookTarget->GetPlatformInfo().VanillaPlatformName) : 0;
 		bMeshDisablesMinLodStrip = OwnerMesh ? OwnerMesh->DisableBelowMinLodStripping.GetValueForPlatformIdentifiers(CookTarget->GetPlatformInfo().PlatformGroupName, CookTarget->GetPlatformInfo().VanillaPlatformName) : false;
 	}
+	const bool bWantToStripBelowMinLod = bIsCook && GStripSkeletalMeshLodsDuringCooking != 0 && MinMeshLod > LODIdx && !bMeshDisablesMinLodStrip;
+
+	uint8 ClassDataStripFlags = 0;
+	ClassDataStripFlags |= bWantToStripTessellation ? CDSF_AdjacencyData : 0;
+	ClassDataStripFlags |= bWantToStripBelowMinLod ? CDSF_MinLodData : 0;
+	return ClassDataStripFlags;
+#else
+	return 0;
 #endif
+}
 
-	const bool bWantToStripBelowMinLod = bIsCook && GStripSkeletalMeshLodsDuringCooking != 0 && MinMeshLod > Idx && !bMeshDisablesMinLodStrip;
+bool FSkeletalMeshLODRenderData::IsLODCookedOut(const ITargetPlatform* TargetPlatform, const USkeletalMesh* SkeletalMesh, bool bIsBelowMinLOD)
+{
+	check(SkeletalMesh);
+#if WITH_EDITOR
+	if (!bIsBelowMinLOD)
+	{
+		return false;
+	}
 
-	ClassDataStripFlags |= bWantToStripTessellation ? LodAdjacencyStripFlag : 0;
-	ClassDataStripFlags |= bWantToStripBelowMinLod ? MinLodStripFlag : 0;
+	if (!TargetPlatform)
+	{
+		TargetPlatform = GetTargetPlatformManagerRef().GetRunningTargetPlatform();
+	}
+	check(TargetPlatform);
 
-	FStripDataFlags StripFlags(Ar, ClassDataStripFlags);
+	const bool bSupportLODStreaming = SkeletalMesh->bSupportLODStreaming.GetValueForPlatformIdentifiers(
+		TargetPlatform->GetPlatformInfo().PlatformGroupName,
+		TargetPlatform->GetPlatformInfo().VanillaPlatformName);
+	
+	return !TargetPlatform->SupportsFeature(ETargetPlatformFeatures::MeshLODStreaming) || !bSupportLODStreaming;
+#else
+	return false;
+#endif
+}
 
-	// Skeletal mesh buffers are kept in CPU memory after initialization to support merging of skeletal meshes.
-	bool bKeepBuffersInCPUMemory = true;
-	bool bNeedsCPUAccess = true;
-	USkeletalMesh* SkelMeshOwner = nullptr;
+bool FSkeletalMeshLODRenderData::IsLODInlined(const ITargetPlatform* TargetPlatform, const USkeletalMesh* SkeletalMesh, int32 LODIdx, bool bIsBelowMinLOD)
+{
+	check(SkeletalMesh);
+#if WITH_EDITOR
+	if (!TargetPlatform)
+	{
+		TargetPlatform = GetTargetPlatformManagerRef().GetRunningTargetPlatform();
+	}
+	check(TargetPlatform);
 
+	const FName PlatformGroupName = TargetPlatform->GetPlatformInfo().PlatformGroupName;
+	const FName VanillaPlatformName = TargetPlatform->GetPlatformInfo().VanillaPlatformName;
+	const bool bSupportLODStreaming = SkeletalMesh->bSupportLODStreaming.GetValueForPlatformIdentifiers(PlatformGroupName, VanillaPlatformName);
+
+	if (!TargetPlatform->SupportsFeature(ETargetPlatformFeatures::MeshLODStreaming) || !bSupportLODStreaming)
+	{
+		return true;
+	}
+
+	if (bIsBelowMinLOD)
+	{
+		return false;
+	}
+
+	const int32 MaxNumStreamedLODs = SkeletalMesh->MaxNumStreamedLODs.GetValueForPlatformIdentifiers(PlatformGroupName, VanillaPlatformName);
+	const int32 NumLODs = SkeletalMesh->GetLODNum();
+	const int32 NumStreamedLODs = FMath::Min(MaxNumStreamedLODs, NumLODs - 1);
+	const int32 InlinedLODStartIdx = NumStreamedLODs;
+	return LODIdx >= InlinedLODStartIdx;
+#else
+	return false;
+#endif
+}
+
+int32 FSkeletalMeshLODRenderData::GetNumOptionalLODsAllowed(const ITargetPlatform* TargetPlatform, const USkeletalMesh* SkeletalMesh)
+{
+#if WITH_EDITOR
+	check(TargetPlatform && SkeletalMesh);
+	const FName PlatformGroupName = TargetPlatform->GetPlatformInfo().PlatformGroupName;
+	const FName VanillaPlatformName = TargetPlatform->GetPlatformInfo().VanillaPlatformName;
+	return  SkeletalMesh->MaxNumOptionalLODs.GetValueForPlatformIdentifiers(PlatformGroupName, VanillaPlatformName);
+#else
+	return 0;
+#endif
+}
+
+bool FSkeletalMeshLODRenderData::ShouldForceKeepCPUResources()
+{
 #if !WITH_EDITOR
 	static const auto CVar = IConsoleManager::Get().FindTConsoleVariableDataInt(TEXT("r.FreeSkeletalMeshBuffers"));
 	if (CVar)
 	{
-		bKeepBuffersInCPUMemory = !CVar->GetValueOnAnyThread();
-		bNeedsCPUAccess = bKeepBuffersInCPUMemory;
+		return !CVar->GetValueOnAnyThread();
 	}
 #endif
+	return true;
+}
+
+bool FSkeletalMeshLODRenderData::ShouldKeepCPUResources(const USkeletalMesh* SkeletalMesh, int32 LODIdx, bool bForceKeep)
+{
+	return bForceKeep
+		|| SkeletalMesh->GetResourceForRendering()->RequiresCPUSkinning(GMaxRHIFeatureLevel)
+		|| SkeletalMesh->NeedCPUData(LODIdx);
+}
+
+class FSkeletalMeshLODSizeCounter : public FArchive
+{
+public:
+	FSkeletalMeshLODSizeCounter()
+		: Size(0)
+	{
+		SetIsSaving(true);
+		SetIsPersistent(true);
+		ArIsCountingMemory = true;
+	}
+
+	virtual void Serialize(void*, int64 Length) final override
+	{
+		Size += Length;
+	}
+
+	virtual int64 TotalSize() final override
+	{
+		return Size;
+	}
+
+private:
+	int64 Size;
+};
+
+void FSkeletalMeshLODRenderData::SerializeStreamedData(FArchive& Ar, USkeletalMesh* Owner, int32 LODIdx, uint8 ClassDataStripFlags, bool bNeedsCPUAccess, bool bForceKeepCPUResources)
+{
+	FStripDataFlags StripFlags(Ar, ClassDataStripFlags);
+
+	// TODO: A lot of data in a render section is needed during initialization but maybe some can still be streamed
+	//Ar << RenderSections;
+
+	MultiSizeIndexContainer.Serialize(Ar, bNeedsCPUAccess);
+
+	if (Ar.IsLoading())
+	{
+		SkinWeightVertexBuffer.SetNeedsCPUAccess(bNeedsCPUAccess);
+	}
+
+	StaticVertexBuffers.PositionVertexBuffer.Serialize(Ar, bNeedsCPUAccess);
+	StaticVertexBuffers.StaticMeshVertexBuffer.Serialize(Ar, bNeedsCPUAccess);
+	Ar << SkinWeightVertexBuffer;
+
+	if (Owner && Owner->bHasVertexColors)
+	{
+		StaticVertexBuffers.ColorVertexBuffer.Serialize(Ar, bForceKeepCPUResources);
+	}
+
+	if (!StripFlags.IsClassDataStripped(CDSF_AdjacencyData))
+	{
+		AdjacencyMultiSizeIndexContainer.Serialize(Ar, bForceKeepCPUResources);
+	}
+
+	if (HasClothData())
+	{
+		Ar << ClothVertexBuffer;
+	}
+
+	Ar << SkinWeightProfilesData;
+	SkinWeightProfilesData.Init(&SkinWeightVertexBuffer);
+
+	if (Ar.IsLoading() && GSkinWeightProfilesLoadByDefaultMode == 1)
+	{
+#if !WITH_EDITOR
+		// Only allow overriding the base buffer in non-editor builds as it could otherwise be serialized into the asset
+		SkinWeightProfilesData.OverrideBaseBufferSkinWeightData(Owner, LODIdx);
+#endif
+	}
+}
+
+void FSkeletalMeshLODRenderData::SerializeAvailabilityInfo(FArchive& Ar, USkeletalMesh* Owner, int32 LODIdx, bool bAdjacencyDataStripped, bool bNeedsCPUAccess)
+{
+	MultiSizeIndexContainer.SerializeMetaData(Ar, bNeedsCPUAccess);
+	if (!bAdjacencyDataStripped)
+	{
+		AdjacencyMultiSizeIndexContainer.SerializeMetaData(Ar, bNeedsCPUAccess);
+	}
+	StaticVertexBuffers.StaticMeshVertexBuffer.SerializeMetaData(Ar);
+	StaticVertexBuffers.PositionVertexBuffer.SerializeMetaData(Ar);
+	StaticVertexBuffers.ColorVertexBuffer.SerializeMetaData(Ar);
+	if (Ar.IsLoading())
+	{
+		SkinWeightVertexBuffer.SetNeedsCPUAccess(bNeedsCPUAccess);
+	}
+	SkinWeightVertexBuffer.SerializeMetaData(Ar);
+	if (HasClothData())
+	{
+		ClothVertexBuffer.SerializeMetaData(Ar);
+	}
+	SkinWeightProfilesData.SerializeMetaData(Ar);
+	SkinWeightProfilesData.Init(&SkinWeightVertexBuffer);
+
+	if (Ar.IsLoading() && GSkinWeightProfilesLoadByDefaultMode == 1)
+	{
+#if !WITH_EDITOR
+		// Only allow overriding the base buffer in non-editor builds as it could otherwise be serialized into the asset
+		SkinWeightProfilesData.OverrideBaseBufferSkinWeightData(Owner, LODIdx);
+#endif
+	}
+}
+
+void FSkeletalMeshLODRenderData::Serialize(FArchive& Ar, UObject* Owner, int32 Idx)
+{
+	DECLARE_SCOPE_CYCLE_COUNTER(TEXT("FSkeletalMeshLODRenderData::Serialize"), STAT_SkeletalMeshLODRenderData_Serialize, STATGROUP_LoadTime);
+
+	USkeletalMesh* OwnerMesh = CastChecked<USkeletalMesh>(Owner);
+	// Actual flags used during serialization
+	const uint8 ClassDataStripFlags = GenerateClassStripFlags(Ar, OwnerMesh, Idx);
+	FStripDataFlags StripFlags(Ar, ClassDataStripFlags);
+
+	const bool bIsBelowMinLOD = StripFlags.IsClassDataStripped(CDSF_MinLodData);
+	bool bIsLODCookedOut = IsLODCookedOut(Ar.CookingTarget(), OwnerMesh, bIsBelowMinLOD);
+	Ar << bIsLODCookedOut;
+
+	bool bInlined = bIsLODCookedOut || IsLODInlined(Ar.CookingTarget(), OwnerMesh, Idx, bIsBelowMinLOD);
+	Ar << bInlined;
+	bStreamedDataInlined = bInlined;
+
+	// Skeletal mesh buffers are kept in CPU memory after initialization to support merging of skeletal meshes.
+	const bool bForceKeepCPUResources = ShouldForceKeepCPUResources();
+	bool bNeedsCPUAccess = bForceKeepCPUResources;
+
 	if (!StripFlags.IsDataStrippedForServer())
 	{
-		SkelMeshOwner = CastChecked<USkeletalMesh>(Owner);
-
 		// set cpu skinning flag on the vertex buffer so that the resource arrays know if they need to be CPU accessible
-		bNeedsCPUAccess = bKeepBuffersInCPUMemory || SkelMeshOwner->GetResourceForRendering()->RequiresCPUSkinning(GMaxRHIFeatureLevel) ||
-			SkelMeshOwner->NeedCPUData(Idx);
+		bNeedsCPUAccess = ShouldKeepCPUResources(OwnerMesh, Idx, bForceKeepCPUResources);
 	}
 
 	if (FPlatformProperties::RequiresCookedData())
 	{
-		if (bNeedsCPUAccess && SkelMeshOwner)
+		if (bNeedsCPUAccess && OwnerMesh)
 		{
-			UE_LOG(LogStaticMesh, Verbose, TEXT("[%s] Skeletal Mesh is marked for CPU read."), *SkelMeshOwner->GetName());
+			UE_LOG(LogStaticMesh, Verbose, TEXT("[%s] Skeletal Mesh is marked for CPU read."), *OwnerMesh->GetName());
 		}
 	}
 
-	if (StripFlags.IsDataStrippedForServer() || StripFlags.IsClassDataStripped(MinLodStripFlag))
+	Ar << RequiredBones;
+
+	if (!StripFlags.IsDataStrippedForServer() && !bIsLODCookedOut)
 	{
-		TArray<FSkelMeshRenderSection> DummySections;
-		Ar << DummySections;
-
-		FMultiSizeIndexContainer DummyIndexBuffer;
-		DummyIndexBuffer.Serialize(Ar, bKeepBuffersInCPUMemory);
-
-		TArray<FBoneIndexType> TempActiveBoneIndices;
-		Ar << TempActiveBoneIndices;
-	}
-	else
-	{
 		Ar << RenderSections;
-
-		MultiSizeIndexContainer.Serialize(Ar, bNeedsCPUAccess);
-
 		Ar << ActiveBoneIndices;
-	}
 
-	Ar << RequiredBones;
-	
-	if (!StripFlags.IsDataStrippedForServer() && !StripFlags.IsClassDataStripped(MinLodStripFlag))
-	{
-		if (Ar.IsLoading())
+#if WITH_EDITOR
+		if (Ar.IsSaving())
 		{
-			SkinWeightVertexBuffer.SetNeedsCPUAccess(bNeedsCPUAccess);
+			FSkeletalMeshLODSizeCounter LODSizeCounter;
+			LODSizeCounter.SetCookingTarget(Ar.CookingTarget());
+			LODSizeCounter.SetByteSwapping(Ar.IsByteSwapping());
+			SerializeStreamedData(LODSizeCounter, OwnerMesh, Idx, ClassDataStripFlags, bNeedsCPUAccess, bForceKeepCPUResources);
+			BuffersSize = LODSizeCounter.TotalSize();
 		}
+#endif
+		Ar << BuffersSize;
 
-		StaticVertexBuffers.PositionVertexBuffer.Serialize(Ar, bNeedsCPUAccess);
-		StaticVertexBuffers.StaticMeshVertexBuffer.Serialize(Ar, bNeedsCPUAccess);
-		Ar << SkinWeightVertexBuffer;
+		if (bInlined)
+		{
+			SerializeStreamedData(Ar, OwnerMesh, Idx, ClassDataStripFlags, bNeedsCPUAccess, bForceKeepCPUResources);
+			bIsLODOptional = false;
+		}
+		else if (Ar.IsCooking() || FPlatformProperties::RequiresCookedData())
+		{
+			bool bDiscardBulkData = false;
 
-		if (SkelMeshOwner && SkelMeshOwner->bHasVertexColors)
-		{
-			StaticVertexBuffers.ColorVertexBuffer.Serialize(Ar, bKeepBuffersInCPUMemory);
-		}
+#if WITH_EDITOR
+			if (Ar.IsSaving())
+			{
+				const int32 MaxNumOptionalLODs = GetNumOptionalLODsAllowed(Ar.CookingTarget(), OwnerMesh);
+				const int32 OptionalLODIdx = GetPlatformMinLODIdx(Ar.CookingTarget(), OwnerMesh) - Idx;
+				bDiscardBulkData = OptionalLODIdx > MaxNumOptionalLODs;
 
-		if (!StripFlags.IsClassDataStripped(LodAdjacencyStripFlag))
-		{
-			AdjacencyMultiSizeIndexContainer.Serialize(Ar, bKeepBuffersInCPUMemory);
-		}
+				TArray<uint8> TmpBuff;
+				if (!bDiscardBulkData)
+				{
+					FMemoryWriter MemWriter(TmpBuff, true);
+					MemWriter.SetCookingTarget(Ar.CookingTarget());
+					MemWriter.SetByteSwapping(Ar.IsByteSwapping());
+					SerializeStreamedData(MemWriter, OwnerMesh, Idx, ClassDataStripFlags, bNeedsCPUAccess, bForceKeepCPUResources);
+				}
 
-		if (HasClothData())
-		{
-			Ar << ClothVertexBuffer;
-		}
-	}	
+				bIsLODOptional = bIsBelowMinLOD;
+				const uint32 BulkDataFlags = (bDiscardBulkData ? 0 : BULKDATA_Force_NOT_InlinePayload)
+					| (bIsLODOptional ? BULKDATA_OptionalPayload : 0);
+				const uint32 OldBulkDataFlags = BulkData.GetBulkDataFlags();
+				BulkData.ClearBulkDataFlags(0xffffffffu);
+				BulkData.SetBulkDataFlags(BulkDataFlags);
+				if (TmpBuff.Num() > 0)
+				{
+					BulkData.Lock(LOCK_READ_WRITE);
+					void* BulkDataMem = BulkData.Realloc(TmpBuff.Num());
+					FMemory::Memcpy(BulkDataMem, TmpBuff.GetData(), TmpBuff.Num());
+					BulkData.Unlock();
+				}
+				BulkData.Serialize(Ar, Owner, Idx);
+				BulkData.ClearBulkDataFlags(0xffffffffu);
+				BulkData.SetBulkDataFlags(OldBulkDataFlags);
+			}
+			else
+#endif
+			{
+				FByteBulkData TmpBulkData;
+				TmpBulkData.Serialize(Ar, Owner, Idx);
 
-	Ar << SkinWeightProfilesData;
-	SkinWeightProfilesData.Init(&SkinWeightVertexBuffer);
+				bIsLODOptional = !!(TmpBulkData.GetBulkDataFlags() & BULKDATA_OptionalPayload);
 
-	if (Ar.IsLoading() && !!GSkinWeightProfilesLoadByDefaultMode)
-	{
-#if !WITH_EDITOR
-		if ( GSkinWeightProfilesLoadByDefaultMode == 1)
-		{
-			// Only allow overriding the base buffer in non-editor builds as it could otherwise be serialized into the asset
-			SkinWeightProfilesData.OverrideBaseBufferSkinWeightData(OwnerMesh, Idx);
-		}
-#endif 	
-		if (GSkinWeightProfilesLoadByDefaultMode == 3)
-		{
-			SkinWeightProfilesData.SetDynamicDefaultSkinWeightProfile(OwnerMesh, Idx);
+				int64 Tmp = TmpBulkData.GetBulkDataOffsetInFile();
+				check(Tmp >= 0 && Tmp <= 0xffffffffll);
+				OffsetInFile = static_cast<uint32>(Tmp);
+				Tmp = TmpBulkData.GetBulkDataSize();
+				check(Tmp >= 0 && Tmp <= 0xffffffffll);
+				BulkDataSize = static_cast<uint32>(Tmp);
+
+				bDiscardBulkData = !BulkDataSize;
+				BuffersSize = bDiscardBulkData ? 0 : BuffersSize;
+			}
+
+			if (!bDiscardBulkData)
+			{
+				SerializeAvailabilityInfo(Ar, OwnerMesh, Idx, StripFlags.IsClassDataStripped(CDSF_AdjacencyData), bNeedsCPUAccess);
+			}
 		}
 	}
 }
diff --git a/Engine/Source/Runtime/Engine/Private/SkeletalMeshLODSettings.cpp b/Engine/Source/Runtime/Engine/Private/SkeletalMeshLODSettings.cpp
--- a/Engine/Source/Runtime/Engine/Private/SkeletalMeshLODSettings.cpp
+++ b/Engine/Source/Runtime/Engine/Private/SkeletalMeshLODSettings.cpp
@@ -13,6 +13,10 @@ DEFINE_LOG_CATEGORY_STATIC(LogSkeletalMeshLODSettings, Warning, All)
 USkeletalMeshLODSettings::USkeletalMeshLODSettings(const FObjectInitializer& ObjectInitializer)
 	: Super(ObjectInitializer)
 {
+	bSupportLODStreaming.Default = false;
+	MaxNumStreamedLODs.Default = 0;
+	// TODO: support saving some but not all optional LODs
+	MaxNumOptionalLODs.Default = 0;
 }
 
 const FSkeletalMeshLODGroupSettings& USkeletalMeshLODSettings::GetSettingsForLODLevel(const int32 LODIndex) const
@@ -162,6 +166,11 @@ int32 USkeletalMeshLODSettings::SetLODSettingsToMesh(USkeletalMesh* InMesh) cons
 	{
 		InMesh->MinLod = MinLod;
 		InMesh->DisableBelowMinLodStripping = DisableBelowMinLodStripping;
+#if WITH_EDITORONLY_DATA
+		InMesh->bSupportLODStreaming = bSupportLODStreaming;
+		InMesh->MaxNumStreamedLODs = MaxNumStreamedLODs;
+		InMesh->MaxNumOptionalLODs = MaxNumOptionalLODs;
+#endif
 		// we only fill up until we have enough LODs
 		const int32 NumSettings = FMath::Min(LODGroups.Num(), InMesh->GetLODNum());
 		for (int32 Index = 0; Index < NumSettings; ++Index)
@@ -182,6 +191,11 @@ int32 USkeletalMeshLODSettings::SetLODSettingsFromMesh(USkeletalMesh* InMesh)
 	{
 		MinLod = InMesh->MinLod;
 		DisableBelowMinLodStripping = InMesh->DisableBelowMinLodStripping;
+#if WITH_EDITORONLY_DATA
+		bSupportLODStreaming = InMesh->bSupportLODStreaming;
+		MaxNumStreamedLODs = InMesh->MaxNumStreamedLODs;
+		MaxNumOptionalLODs = InMesh->MaxNumOptionalLODs;
+#endif
 		// we only fill up until we have enough LODs
 		const int32 NumSettings = InMesh->GetLODNum();
 		LODGroups.Reset(NumSettings);
diff --git a/Engine/Source/Runtime/Engine/Private/SkeletalMeshRenderData.cpp b/Engine/Source/Runtime/Engine/Private/SkeletalMeshRenderData.cpp
--- a/Engine/Source/Runtime/Engine/Private/SkeletalMeshRenderData.cpp
+++ b/Engine/Source/Runtime/Engine/Private/SkeletalMeshRenderData.cpp
@@ -12,6 +12,9 @@
 #include "DerivedDataCacheInterface.h"
 #include "Serialization/MemoryReader.h"
 #include "Serialization/MemoryWriter.h"
+#include "Interfaces/ITargetPlatformManagerModule.h"
+#include "Interfaces/ITargetPlatform.h"
+#include "PlatformInfo.h"
 
 
 #if ENABLE_COOK_STATS
@@ -29,7 +32,7 @@ namespace SkeletalMeshCookStats
 // differences, etc.) replace the version GUID below with a new one.
 // In case of merge conflicts with DDC versions, you MUST generate a new GUID
 // and set this new GUID as the version.                                       
-#define SKELETALMESH_DERIVEDDATA_VER TEXT("B79E93B9DF9E4D2D8F44F44A9D947994")
+#define SKELETALMESH_DERIVEDDATA_VER TEXT("70D1F1C048AF4DF2AF870B51CE803A66")
 
 static const FString& GetSkeletalMeshDerivedDataVersion()
 {
@@ -50,6 +53,23 @@ static FString BuildSkeletalMeshDerivedDataKey(USkeletalMesh* SkelMesh)
 	KeySuffix += SkelMesh->bHasVertexColors ? "1" : "0";
 	KeySuffix += SkelMesh->VertexColorGuid.ToString(EGuidFormats::Digits);
 
+	const ITargetPlatform* TargetPlatform = GetTargetPlatformManagerRef().GetRunningTargetPlatform();
+	check(TargetPlatform);
+	const FName PlatformGroupName = TargetPlatform->GetPlatformInfo().PlatformGroupName;
+	const FName VanillaPlatformName = TargetPlatform->GetPlatformInfo().VanillaPlatformName;
+
+	const bool bSupportLODStreaming = SkelMesh->bSupportLODStreaming.GetValueForPlatformIdentifiers(PlatformGroupName, VanillaPlatformName);
+	if (bSupportLODStreaming)
+	{
+		const int32 MaxNumStreamedLODs = SkelMesh->MaxNumStreamedLODs.GetValueForPlatformIdentifiers(PlatformGroupName, VanillaPlatformName);
+		const int32 MaxNumOptionalLODs = SkelMesh->MaxNumOptionalLODs.GetValueForPlatformIdentifiers(PlatformGroupName, VanillaPlatformName);
+		KeySuffix += *FString::Printf(TEXT("1%08x%08x"), MaxNumStreamedLODs, MaxNumOptionalLODs);
+	}
+	else
+	{
+		KeySuffix += TEXT("0zzzzzzzzzzzzzzzz");
+	}
+
 	return FDerivedDataCacheInterface::BuildCacheKey(
 		TEXT("SKELETALMESH"),
 		*GetSkeletalMeshDerivedDataVersion(),
@@ -81,6 +101,18 @@ void FSkeletalMeshRenderData::Cache(USkeletalMesh* Owner)
 			COOK_STAT(Timer.AddHit(DerivedData.Num()));
 			FMemoryReader Ar(DerivedData, /*bIsPersistent=*/ true);
 			Serialize(Ar, Owner);
+			for (int32 LODIndex = 0; LODIndex < LODRenderData.Num(); ++LODIndex)
+			{
+				FSkeletalMeshLODRenderData& LODData = LODRenderData[LODIndex];
+				if (LODData.bStreamedDataInlined)
+				{
+					break;
+				}
+				constexpr uint8 DummyStripFlags = 0;
+				const bool bForceKeepCPUResources = FSkeletalMeshLODRenderData::ShouldForceKeepCPUResources();
+				const bool bNeedsCPUAccess = FSkeletalMeshLODRenderData::ShouldKeepCPUResources(Owner, LODIndex, bForceKeepCPUResources);
+				LODData.SerializeStreamedData(Ar, Owner, LODIndex, DummyStripFlags, bNeedsCPUAccess, bForceKeepCPUResources);
+			}
 
 			int32 T1 = FPlatformTime::Cycles();
 			UE_LOG(LogSkeletalMesh, Verbose, TEXT("Skeletal Mesh found in DDC [%fms] %s"), FPlatformTime::ToMilliseconds(T1 - T0), *Owner->GetPathName());
@@ -106,6 +138,18 @@ void FSkeletalMeshRenderData::Cache(USkeletalMesh* Owner)
 
 			FMemoryWriter Ar(DerivedData, /*bIsPersistent=*/ true);
 			Serialize(Ar, Owner);
+			for (int32 LODIndex = 0; LODIndex < LODRenderData.Num(); ++LODIndex)
+			{
+				FSkeletalMeshLODRenderData& LODData = LODRenderData[LODIndex];
+				if (LODData.bStreamedDataInlined)
+				{
+					break;
+				}
+				const uint8 LODStripFlags = FSkeletalMeshLODRenderData::GenerateClassStripFlags(Ar, Owner, LODIndex);
+				const bool bForceKeepCPUResources = FSkeletalMeshLODRenderData::ShouldForceKeepCPUResources();
+				const bool bNeedsCPUAccess = FSkeletalMeshLODRenderData::ShouldKeepCPUResources(Owner, LODIndex, bForceKeepCPUResources);
+				LODData.SerializeStreamedData(Ar, Owner, LODIndex, LODStripFlags, bNeedsCPUAccess, bForceKeepCPUResources);
+			}
 			GetDerivedDataCacheRef().Put(*DerivedDataKey, DerivedData);
 
 			int32 T1 = FPlatformTime::Cycles();
@@ -133,14 +177,47 @@ void FSkeletalMeshRenderData::SyncUVChannelData(const TArray<FSkeletalMaterial>&
 
 #endif // WITH_EDITOR
 
+FSkeletalMeshRenderData::FSkeletalMeshRenderData()
+	: bReadyForStreaming(false)
+	, NumInlinedLODs(0)
+	, NumOptionalLODs(0)
+	, CurrentFirstLODIdx(0)
+	, PendingFirstLODIdx(0)
+	, bInitialized(false)
+{}
+
 void FSkeletalMeshRenderData::Serialize(FArchive& Ar, USkeletalMesh* Owner)
 {
 	DECLARE_SCOPE_CYCLE_COUNTER(TEXT("FSkeletalMeshRenderData::Serialize"), STAT_SkeletalMeshRenderData_Serialize, STATGROUP_LoadTime);
 
 	LODRenderData.Serialize(Ar, Owner);
+
+#if WITH_EDITOR
+	if (Ar.IsSaving())
+	{
+		NumInlinedLODs = 0;
+		NumOptionalLODs = 0;
+		for (int32 Idx = LODRenderData.Num() - 1; Idx >= 0; --Idx)
+		{
+			if (LODRenderData[Idx].bStreamedDataInlined)
+			{
+				++NumInlinedLODs;
+			}
+			if (LODRenderData[Idx].bIsLODOptional)
+			{
+				++NumOptionalLODs;
+			}
+		}
+	}
+#endif
+	Ar << NumInlinedLODs << NumOptionalLODs;
+	
+	CurrentFirstLODIdx = LODRenderData.Num() - NumInlinedLODs;
+	PendingFirstLODIdx = CurrentFirstLODIdx;
+	Owner->SetCachedNumResidentLODs(NumInlinedLODs);
 }
 
-void FSkeletalMeshRenderData::InitResources(bool bNeedsVertexColors, TArray<UMorphTarget*>& InMorphTargets)
+void FSkeletalMeshRenderData::InitResources(bool bNeedsVertexColors, TArray<UMorphTarget*>& InMorphTargets, USkeletalMesh* Owner)
 {
 	if (!bInitialized)
 	{
@@ -151,9 +228,17 @@ void FSkeletalMeshRenderData::InitResources(bool bNeedsVertexColors, TArray<UMor
 
 			if(RenderData.GetNumVertices() > 0)
 			{
-				RenderData.InitResources(bNeedsVertexColors, LODIndex, InMorphTargets);
+				RenderData.InitResources(bNeedsVertexColors, LODIndex, InMorphTargets, Owner);
 			}
 		}
+
+		ENQUEUE_RENDER_COMMAND(CmdSetSkeletalMeshReadyForStreaming)(
+			[this, Owner](FRHICommandListImmediate&)
+		{
+			bReadyForStreaming = true;
+			Owner->SetCachedReadyForStreaming(true);
+		});
+
 		bInitialized = true;
 	}
 }
diff --git a/Engine/Source/Runtime/Engine/Private/SkeletalRender.cpp b/Engine/Source/Runtime/Engine/Private/SkeletalRender.cpp
--- a/Engine/Source/Runtime/Engine/Private/SkeletalRender.cpp
+++ b/Engine/Source/Runtime/Engine/Private/SkeletalRender.cpp
@@ -24,9 +24,9 @@ FSkeletalMeshObject
 -----------------------------------------------------------------------------*/
 
 FSkeletalMeshObject::FSkeletalMeshObject(USkinnedMeshComponent* InMeshComponent, FSkeletalMeshRenderData* InSkelMeshRenderData, ERHIFeatureLevel::Type InFeatureLevel)
-:	MinDesiredLODLevel(0)
+:	MinDesiredLODLevel(InSkelMeshRenderData->CurrentFirstLODIdx)
 ,	MaxDistanceFactor(0.f)
-,	WorkingMinDesiredLODLevel(0)
+,	WorkingMinDesiredLODLevel(MinDesiredLODLevel)
 ,	WorkingMaxDistanceFactor(0.f)
 ,   bHasBeenUpdatedAtLeastOnce(false)
 #if RHI_RAYTRACING
@@ -70,7 +70,7 @@ FSkeletalMeshObject::~FSkeletalMeshObject()
 {
 }
 
-void FSkeletalMeshObject::UpdateMinDesiredLODLevel(const FSceneView* View, const FBoxSphereBounds& Bounds, int32 FrameNumber)
+void FSkeletalMeshObject::UpdateMinDesiredLODLevel(const FSceneView* View, const FBoxSphereBounds& Bounds, int32 FrameNumber, uint8 CurFirstLODIdx)
 {
 	static const auto* SkeletalMeshLODRadiusScale = IConsoleManager::Get().FindTConsoleVariableDataFloat(TEXT("r.SkeletalMeshLODRadiusScale"));
 	float LODScale = FMath::Clamp(SkeletalMeshLODRadiusScale->GetValueOnRenderThread(), 0.25f, 1.0f);
@@ -109,6 +109,14 @@ void FSkeletalMeshObject::UpdateMinDesiredLODLevel(const FSceneView* View, const
 		}
 	}
 
+	NewLODLevel = FMath::Max<int32>(NewLODLevel, CurFirstLODIdx);
+	if (!LastFrameNumber)
+	{
+		// We don't have last frame value on the first call to FSkeletalMeshObject::UpdateMinDesiredLODLevel so
+		// just reuse current frame value. Otherwise, MinDesiredLODLevel may get stale value in the update code below
+		WorkingMinDesiredLODLevel = NewLODLevel;
+	}
+
 	// Different path for first-time vs subsequent-times in this function (ie splitscreen)
 	if(FrameNumber != LastFrameNumber)
 	{
diff --git a/Engine/Source/Runtime/Engine/Private/Streaming/RenderAssetUpdate.cpp b/Engine/Source/Runtime/Engine/Private/Streaming/RenderAssetUpdate.cpp
--- a/Engine/Source/Runtime/Engine/Private/Streaming/RenderAssetUpdate.cpp
+++ b/Engine/Source/Runtime/Engine/Private/Streaming/RenderAssetUpdate.cpp
@@ -118,10 +118,12 @@ void ResumeRenderAssetStreaming()
 FRenderAssetUpdate::FRenderAssetUpdate(UStreamableRenderAsset* InAsset, int32 InRequestedMips)
 	: PendingFirstMip(INDEX_NONE)
 	, RequestedMips(INDEX_NONE)
+	, ScheduledGTTasks(0)
 	, ScheduledRenderTasks(0)
 	, ScheduledAsyncTasks(0)
 	, StreamableAsset(InAsset)
 	, bIsCancelled(false)
+	, bDeferExecution(false)
 	, TaskState(TS_Init)
 {
 	check(InAsset);
@@ -213,6 +215,56 @@ void FRenderAssetUpdate::Tick(EThreadType InCurrentThread)
 	}
 }
 
+class FRenderAssetUpdateTickGTTask
+{
+public:
+	FORCEINLINE FRenderAssetUpdateTickGTTask(FRenderAssetUpdate* InUpdate)
+		: PendingUpdate(InUpdate)
+	{}
+
+	static FORCEINLINE TStatId GetStatId()
+	{
+		RETURN_QUICK_DECLARE_CYCLE_STAT(FRenderAssetUpdateTickGTTask, STATGROUP_TaskGraphTasks);
+	}
+
+	static FORCEINLINE ENamedThreads::Type GetDesiredThread()
+	{
+		return ENamedThreads::GameThread;
+	}
+
+	static FORCEINLINE ESubsequentsMode::Type GetSubsequentsMode()
+	{
+		return ESubsequentsMode::FireAndForget;
+	}
+
+	void DoTask(ENamedThreads::Type CurThread, const FGraphEventRef& MyCompletionGraphEvent)
+	{
+		check(PendingUpdate);
+		PendingUpdate->Tick(FRenderAssetUpdate::TT_GameThread);
+		--PendingUpdate->ScheduledGTTasks;
+	}
+
+private:
+	TRefCountPtr<FRenderAssetUpdate> PendingUpdate;
+};
+
+void FRenderAssetUpdate::ScheduleGTTask()
+{
+	check(TaskState == TS_Locked);
+
+	if (IsInGameThread())
+	{
+		Tick(TT_GameThread);
+	}
+	else
+	{
+		// Notify that a tick is scheduled on the game thread
+		++ScheduledGTTasks;
+
+		TGraphTask<FRenderAssetUpdateTickGTTask>::CreateTask().ConstructAndDispatchWhenReady(this);
+	}
+}
+
 void FRenderAssetUpdate::ScheduleRenderTask()
 {
 	check(TaskState == TS_Locked);
diff --git a/Engine/Source/Runtime/Engine/Private/Streaming/RenderAssetUpdate.inl b/Engine/Source/Runtime/Engine/Private/Streaming/RenderAssetUpdate.inl
--- a/Engine/Source/Runtime/Engine/Private/Streaming/RenderAssetUpdate.inl
+++ b/Engine/Source/Runtime/Engine/Private/Streaming/RenderAssetUpdate.inl
@@ -63,12 +63,22 @@ FRenderAssetUpdate::ETaskState TRenderAssetUpdate<TContext>::TickInternal(EThrea
 	{
 		return TS_Suspended;
 	}
+	else if (bDeferExecution)
+	{
+		bDeferExecution = false;
+		return TS_Suspended;
+	}
 	else if (RelevantThread == InCurrentThread)
 	{
 		ClearCallbacks();
 		RelevantCallback(FContext(RelevantAsset, InCurrentThread));
 		return TS_Locked;
 	} 
+	else if (RelevantThread == FRenderAssetUpdate::TT_GameThread && !ScheduledGTTasks)
+	{
+		ScheduleGTTask();
+		return TS_Locked;
+	}
 	else if (RelevantThread == FRenderAssetUpdate::TT_Render && !ScheduledRenderTasks)
 	{
 		if (GIsThreadedRendering)
diff --git a/Engine/Source/Runtime/Engine/Private/Streaming/SkeletalMeshUpdate.cpp b/Engine/Source/Runtime/Engine/Private/Streaming/SkeletalMeshUpdate.cpp
new file mode 100644
--- /dev/null
+++ b/Engine/Source/Runtime/Engine/Private/Streaming/SkeletalMeshUpdate.cpp
@@ -0,0 +1,597 @@
+// Copyright 1998-2019 Epic Games, Inc. All Rights Reserved.
+
+/*=============================================================================
+SkeletalMeshUpdate.cpp: Helpers to stream in and out skeletal mesh LODs.
+=============================================================================*/
+
+#include "Streaming/SkeletalMeshUpdate.h"
+#include "RenderUtils.h"
+#include "Containers/ResourceArray.h"
+#include "Streaming/TextureStreamingHelpers.h"
+#include "HAL/PlatformFilemanager.h"
+#include "Serialization/MemoryReader.h"
+#include "Rendering/SkeletalMeshRenderData.h"
+#include "Streaming/RenderAssetUpdate.inl"
+
+template class TRenderAssetUpdate<FSkelMeshUpdateContext>;
+
+static constexpr uint32 GSkelMeshMaxNumResourceUpdatesPerLOD = 16;
+static constexpr uint32 GSkelMeshMaxNumResourceUpdatesPerBatch = (MAX_MESH_LOD_COUNT - 1) * GSkelMeshMaxNumResourceUpdatesPerLOD;
+
+FSkelMeshUpdateContext::FSkelMeshUpdateContext(USkeletalMesh* InMesh, EThreadType InCurrentThread)
+	: Mesh(InMesh)
+	, CurrentThread(InCurrentThread)
+{
+	check(InMesh);
+	checkSlow(InCurrentThread != FSkeletalMeshUpdate::TT_Render || IsInRenderingThread());
+	RenderData = Mesh->GetResourceForRendering();
+}
+
+FSkelMeshUpdateContext::FSkelMeshUpdateContext(UStreamableRenderAsset* InMesh, EThreadType InCurrentThread)
+#if UE_BUILD_SHIPPING
+	: FSkelMeshUpdateContext(static_cast<USkeletalMesh*>(InMesh), InCurrentThread)
+#else
+	: FSkelMeshUpdateContext(Cast<USkeletalMesh>(InMesh), InCurrentThread)
+#endif
+{}
+
+FSkeletalMeshUpdate::FSkeletalMeshUpdate(USkeletalMesh* InMesh, int32 InRequestedMips)
+	: TRenderAssetUpdate<FSkelMeshUpdateContext>(InMesh, InRequestedMips)
+{
+	FSkeletalMeshRenderData* RenderData = InMesh->GetResourceForRendering();
+	if (RenderData)
+	{
+		CurrentFirstLODIdx = RenderData->CurrentFirstLODIdx;
+		check(CurrentFirstLODIdx >= 0 && CurrentFirstLODIdx < MAX_MESH_LOD_COUNT);
+	}
+	else
+	{
+		RequestedMips = INDEX_NONE;
+		PendingFirstMip = INDEX_NONE;
+		bIsCancelled = true;
+		CurrentFirstLODIdx = INDEX_NONE;
+	}
+}
+
+void FSkeletalMeshStreamIn::FIntermediateBuffers::CreateFromCPUData_RenderThread(USkeletalMesh* Mesh, FSkeletalMeshLODRenderData& LODResource)
+{
+	FStaticMeshVertexBuffers& VBs = LODResource.StaticVertexBuffers;
+	TangentsVertexBuffer = VBs.StaticMeshVertexBuffer.CreateTangentsRHIBuffer_RenderThread();
+	TexCoordVertexBuffer = VBs.StaticMeshVertexBuffer.CreateTexCoordRHIBuffer_RenderThread();
+	PositionVertexBuffer = VBs.PositionVertexBuffer.CreateRHIBuffer_RenderThread();
+	ColorVertexBuffer = VBs.ColorVertexBuffer.CreateRHIBuffer_RenderThread();
+	LODResource.SkinWeightProfilesData.CreateRHIBuffers_RenderThread(AltSkinWeightVertexBuffers);
+	SkinWeightVertexBuffer = LODResource.SkinWeightVertexBuffer.CreateRHIBuffer_RenderThread();
+	ClothVertexBuffer = LODResource.ClothVertexBuffer.CreateRHIBuffer_RenderThread();
+	IndexBuffer = LODResource.MultiSizeIndexContainer.CreateRHIBuffer_RenderThread();
+	AdjacencyIndexBuffer = LODResource.AdjacencyMultiSizeIndexContainer.CreateRHIBuffer_RenderThread();
+}
+
+void FSkeletalMeshStreamIn::FIntermediateBuffers::CreateFromCPUData_Async(USkeletalMesh* Mesh, FSkeletalMeshLODRenderData& LODResource)
+{
+	FStaticMeshVertexBuffers& VBs = LODResource.StaticVertexBuffers;
+	TangentsVertexBuffer = VBs.StaticMeshVertexBuffer.CreateTangentsRHIBuffer_Async();
+	TexCoordVertexBuffer = VBs.StaticMeshVertexBuffer.CreateTexCoordRHIBuffer_Async();
+	PositionVertexBuffer = VBs.PositionVertexBuffer.CreateRHIBuffer_Async();
+	ColorVertexBuffer = VBs.ColorVertexBuffer.CreateRHIBuffer_Async();
+	LODResource.SkinWeightProfilesData.CreateRHIBuffers_Async(AltSkinWeightVertexBuffers);
+	SkinWeightVertexBuffer = LODResource.SkinWeightVertexBuffer.CreateRHIBuffer_Async();
+	ClothVertexBuffer = LODResource.ClothVertexBuffer.CreateRHIBuffer_Async();
+	IndexBuffer = LODResource.MultiSizeIndexContainer.CreateRHIBuffer_Async();
+	AdjacencyIndexBuffer = LODResource.AdjacencyMultiSizeIndexContainer.CreateRHIBuffer_Async();
+}
+
+void FSkeletalMeshStreamIn::FIntermediateBuffers::SafeRelease()
+{
+	TangentsVertexBuffer.SafeRelease();
+	TexCoordVertexBuffer.SafeRelease();
+	PositionVertexBuffer.SafeRelease();
+	ColorVertexBuffer.SafeRelease();
+	SkinWeightVertexBuffer.SafeRelease();
+	ClothVertexBuffer.SafeRelease();
+	IndexBuffer.SafeRelease();
+	AdjacencyIndexBuffer.SafeRelease();
+	AltSkinWeightVertexBuffers.Empty();
+}
+
+template <uint32 MaxNumUpdates>
+void FSkeletalMeshStreamIn::FIntermediateBuffers::TransferBuffers(FSkeletalMeshLODRenderData& LODResource, TRHIResourceUpdateBatcher<MaxNumUpdates>& Batcher)
+{
+	FStaticMeshVertexBuffers& VBs = LODResource.StaticVertexBuffers;
+	VBs.StaticMeshVertexBuffer.InitRHIForStreaming(TangentsVertexBuffer, TexCoordVertexBuffer, Batcher);
+	VBs.PositionVertexBuffer.InitRHIForStreaming(PositionVertexBuffer, Batcher);
+	VBs.ColorVertexBuffer.InitRHIForStreaming(ColorVertexBuffer, Batcher);
+	LODResource.SkinWeightVertexBuffer.InitRHIForStreaming(SkinWeightVertexBuffer, Batcher);
+	LODResource.ClothVertexBuffer.InitRHIForStreaming(ClothVertexBuffer, Batcher);
+	LODResource.MultiSizeIndexContainer.InitRHIForStreaming(IndexBuffer, Batcher);
+	LODResource.AdjacencyMultiSizeIndexContainer.InitRHIForStreaming(AdjacencyIndexBuffer, Batcher);
+	LODResource.SkinWeightProfilesData.InitRHIForStreaming(AltSkinWeightVertexBuffers, Batcher);
+	SafeRelease();
+}
+
+void FSkeletalMeshStreamIn::FIntermediateBuffers::CheckIsNull() const
+{
+	check(!TangentsVertexBuffer
+		&& !TexCoordVertexBuffer
+		&& !PositionVertexBuffer
+		&& !ColorVertexBuffer
+		&& !SkinWeightVertexBuffer
+		&& !ClothVertexBuffer
+		&& !IndexBuffer
+		&& !AdjacencyIndexBuffer
+		&& !AltSkinWeightVertexBuffers.Num());
+}
+
+FSkeletalMeshStreamIn::FSkeletalMeshStreamIn(USkeletalMesh* InMesh, int32 InRequestedMips)
+	: FSkeletalMeshUpdate(InMesh, InRequestedMips)
+{}
+
+FSkeletalMeshStreamIn::~FSkeletalMeshStreamIn()
+{
+#if DO_CHECK
+	for (int32 Idx = 0; Idx < MAX_MESH_LOD_COUNT; ++Idx)
+	{
+		IntermediateBuffersArray[Idx].CheckIsNull();
+	}
+#endif
+}
+
+template <bool bRenderThread>
+void FSkeletalMeshStreamIn::CreateBuffers_Internal(const FContext& Context)
+{
+	LLM_SCOPE(ELLMTag::SkeletalMesh);
+
+	USkeletalMesh* Mesh = Context.Mesh;
+	FSkeletalMeshRenderData* RenderData = Context.RenderData;
+	if (!IsCancelled() && Mesh && RenderData)
+	{
+		check(CurrentFirstLODIdx == RenderData->CurrentFirstLODIdx && PendingFirstMip < CurrentFirstLODIdx);
+
+		for (int32 LODIdx = PendingFirstMip; LODIdx < CurrentFirstLODIdx; ++LODIdx)
+		{
+			FSkeletalMeshLODRenderData& LODResource = RenderData->LODRenderData[LODIdx];
+			if (bRenderThread)
+			{
+				IntermediateBuffersArray[LODIdx].CreateFromCPUData_RenderThread(Mesh, LODResource);
+			}
+			else
+			{
+				IntermediateBuffersArray[LODIdx].CreateFromCPUData_Async(Mesh, LODResource);
+			}
+		}
+	}
+}
+
+void FSkeletalMeshStreamIn::CreateBuffers_RenderThread(const FContext& Context)
+{
+	check(Context.CurrentThread == TT_Render);
+	CreateBuffers_Internal<true>(Context);
+}
+
+void FSkeletalMeshStreamIn::CreateBuffers_Async(const FContext& Context)
+{
+	check(Context.CurrentThread == TT_Async);
+	CreateBuffers_Internal<false>(Context);
+}
+
+void FSkeletalMeshStreamIn::DiscardNewLODs(const FContext& Context)
+{
+	FSkeletalMeshRenderData* RenderData = Context.RenderData;
+	if (RenderData)
+	{
+		check(CurrentFirstLODIdx == RenderData->CurrentFirstLODIdx && PendingFirstMip < CurrentFirstLODIdx);
+
+		for (int32 LODIdx = PendingFirstMip; LODIdx < CurrentFirstLODIdx; ++LODIdx)
+		{
+			FSkeletalMeshLODRenderData& LODResource = RenderData->LODRenderData[LODIdx];
+			LODResource.ReleaseCPUResources(true);
+		}
+	}
+}
+
+void FSkeletalMeshStreamIn::DoFinishUpdate(const FContext& Context)
+{
+	USkeletalMesh* Mesh = Context.Mesh;
+	FSkeletalMeshRenderData* RenderData = Context.RenderData;
+	if (!IsCancelled() && Mesh && RenderData)
+	{
+		check(Context.CurrentThread == TT_Render
+			&& CurrentFirstLODIdx == RenderData->CurrentFirstLODIdx
+			&& PendingFirstMip < CurrentFirstLODIdx);
+		// Use a scope to flush the batcher before updating CurrentFirstLODIdx
+		{
+			TRHIResourceUpdateBatcher<GSkelMeshMaxNumResourceUpdatesPerBatch> Batcher;
+
+			for (int32 LODIdx = PendingFirstMip; LODIdx < CurrentFirstLODIdx; ++LODIdx)
+			{
+				FSkeletalMeshLODRenderData& LODResource = RenderData->LODRenderData[LODIdx];
+				LODResource.IncrementMemoryStats(Mesh->bHasVertexColors);
+				IntermediateBuffersArray[LODIdx].TransferBuffers(LODResource, Batcher);
+			}
+		}
+		check(Mesh->GetCachedNumResidentLODs() == RenderData->LODRenderData.Num() - RenderData->CurrentFirstLODIdx);
+		RenderData->CurrentFirstLODIdx = RenderData->PendingFirstLODIdx = PendingFirstMip;
+		Mesh->SetCachedNumResidentLODs(static_cast<uint8>(RenderData->LODRenderData.Num() - PendingFirstMip));
+	}
+	else
+	{
+		for (int32 LODIdx = PendingFirstMip; LODIdx < CurrentFirstLODIdx; ++LODIdx)
+		{
+			IntermediateBuffersArray[LODIdx].SafeRelease();
+		}
+	}
+}
+
+void FSkeletalMeshStreamIn::DoCancel(const FContext& Context)
+{
+	// TODO: support streaming CPU data for editor builds
+	if (!GIsEditor)
+	{
+		DiscardNewLODs(Context);
+	}
+	DoFinishUpdate(Context);
+}
+
+FSkeletalMeshStreamOut::FSkeletalMeshStreamOut(USkeletalMesh* InMesh, int32 InRequestedMips)
+	: FSkeletalMeshUpdate(InMesh, InRequestedMips)
+{
+	PushTask(FContext(InMesh, TT_None), TT_GameThread, SRA_UPDATE_CALLBACK(DoConditionalMarkComponentsDirty), TT_None, nullptr);
+}
+
+void FSkeletalMeshStreamOut::DoConditionalMarkComponentsDirty(const FContext& Context)
+{
+	QUICK_SCOPE_CYCLE_COUNTER(STAT_FSkeletalMeshStreamOut_DoConditionalMarkComponentsDirty);
+	CSV_SCOPED_TIMING_STAT_GLOBAL(SkStreamingMarkDirtyTime);
+	check(Context.CurrentThread == TT_GameThread);
+
+	USkeletalMesh* Mesh = Context.Mesh;
+	FSkeletalMeshRenderData* RenderData = Context.RenderData;
+	if (!IsCancelled() && Mesh && RenderData)
+	{
+		RenderData->PendingFirstLODIdx = PendingFirstMip;
+
+		TArray<const UPrimitiveComponent*> Comps;
+		IStreamingManager::Get().GetTextureStreamingManager().GetAssetComponents(Mesh, Comps, [](const UPrimitiveComponent* Comp)
+		{
+			return !Comp->IsComponentTickEnabled();
+		});
+		for (int32 Idx = 0; Idx < Comps.Num(); ++Idx)
+		{
+			check(Comps[Idx]->IsA<USkinnedMeshComponent>());
+			USkinnedMeshComponent* Comp = (USkinnedMeshComponent*)Comps[Idx];
+			if (Comp->PredictedLODLevel < PendingFirstMip)
+			{
+				Comp->PredictedLODLevel = PendingFirstMip;
+				Comp->bForceMeshObjectUpdate = true;
+				Comp->MarkRenderDynamicDataDirty();
+			}
+		}
+	}
+	else
+	{
+		Abort();
+	}
+	PushTask(Context, TT_Render, SRA_UPDATE_CALLBACK(DoReleaseBuffers), (EThreadType)Context.CurrentThread, SRA_UPDATE_CALLBACK(DoCancel));
+}
+
+void FSkeletalMeshStreamOut::DoReleaseBuffers(const FContext& Context)
+{
+	check(Context.CurrentThread == TT_Render);
+	USkeletalMesh* Mesh = Context.Mesh;
+	FSkeletalMeshRenderData* RenderData = Context.RenderData;
+	if (!IsCancelled() && Mesh && RenderData)
+	{
+		check(CurrentFirstLODIdx == RenderData->CurrentFirstLODIdx && PendingFirstMip > CurrentFirstLODIdx);
+		check(Mesh->GetCachedNumResidentLODs() == RenderData->LODRenderData.Num() - CurrentFirstLODIdx);
+		RenderData->CurrentFirstLODIdx = PendingFirstMip;
+		Mesh->SetCachedNumResidentLODs(static_cast<uint8>(RenderData->LODRenderData.Num() - PendingFirstMip));
+
+		TRHIResourceUpdateBatcher<GSkelMeshMaxNumResourceUpdatesPerBatch> Batcher;
+
+		for (int32 LODIdx = CurrentFirstLODIdx; LODIdx < PendingFirstMip; ++LODIdx)
+		{
+			FSkeletalMeshLODRenderData& LODResource = RenderData->LODRenderData[LODIdx];
+			FStaticMeshVertexBuffers& VBs = LODResource.StaticVertexBuffers;
+			LODResource.DecrementMemoryStats();
+			VBs.StaticMeshVertexBuffer.ReleaseRHIForStreaming(Batcher);
+			VBs.PositionVertexBuffer.ReleaseRHIForStreaming(Batcher);
+			VBs.ColorVertexBuffer.ReleaseRHIForStreaming(Batcher);
+			LODResource.SkinWeightVertexBuffer.ReleaseRHIForStreaming(Batcher);
+			LODResource.ClothVertexBuffer.ReleaseRHIForStreaming(Batcher);
+			LODResource.MultiSizeIndexContainer.ReleaseRHIForStreaming(Batcher);
+			LODResource.AdjacencyMultiSizeIndexContainer.ReleaseRHIForStreaming(Batcher);
+			LODResource.SkinWeightProfilesData.ReleaseRHIForStreaming(Batcher);
+		}
+	}
+}
+
+void FSkeletalMeshStreamOut::DoCancel(const FContext& Context)
+{
+	FSkeletalMeshRenderData* RenderData = Context.RenderData;
+	if (RenderData)
+	{
+		RenderData->PendingFirstLODIdx = CurrentFirstLODIdx;
+	}
+}
+
+void FSkeletalMeshStreamIn_IO::FCancelIORequestsTask::DoWork()
+{
+	check(PendingUpdate);
+	// Acquire the lock of this object in order to cancel any pending IO.
+	// If the object is currently being ticked, wait.
+	ETaskState OldState = PendingUpdate->DoLock();
+	PendingUpdate->CancelIORequest();
+	PendingUpdate->DoUnlock(OldState);
+}
+
+FSkeletalMeshStreamIn_IO::FSkeletalMeshStreamIn_IO(USkeletalMesh* InMesh, int32 InRequestedMips, bool bHighPrio)
+	: FSkeletalMeshStreamIn(InMesh, InRequestedMips)
+	, IOFileHandle(nullptr)
+	, IORequest(nullptr)
+	, bHighPrioIORequest(bHighPrio)
+{}
+
+void FSkeletalMeshStreamIn_IO::Abort()
+{
+	if (!IsCancelled() && !IsCompleted())
+	{
+		FSkeletalMeshStreamIn::Abort();
+
+		// IO requests can only exist in the lifetime of IOFileHandle.
+		if (IOFileHandle)
+		{
+			// Prevent the update from being considered done before this is finished.
+			// By checking that it was not already cancelled, we make sure this doesn't get called twice.
+			(new FAsyncCancelIORequestsTask(this))->StartBackgroundTask();
+		}
+	}
+}
+
+FString FSkeletalMeshStreamIn_IO::GetIOFilename(const FContext& Context)
+{
+	USkeletalMesh* Mesh = Context.Mesh;
+	if (!IsCancelled() && Mesh)
+	{
+		FString Filename;
+		verify(Mesh->GetMipDataFilename(PendingFirstMip, Filename));
+		return Filename;
+	}
+	MarkAsCancelled();
+	return FString();
+}
+
+void FSkeletalMeshStreamIn_IO::SetAsyncFileCallback(const FContext& Context)
+{
+	AsyncFileCallback = [this, Context](bool bWasCancelled, IAsyncReadRequest* Req)
+	{
+		// At this point task synchronization would hold the number of pending requests.
+		TaskSynchronization.Decrement();
+
+		if (bWasCancelled)
+		{
+			MarkAsCancelled();
+		}
+
+#if !UE_BUILD_SHIPPING
+		// On some platforms the IO is too fast to test cancelation requests timing issues.
+		if (FRenderAssetStreamingSettings::ExtraIOLatency > 0 && TaskSynchronization.GetValue() == 0)
+		{
+			FPlatformProcess::Sleep(FRenderAssetStreamingSettings::ExtraIOLatency * .001f); // Slow down the streaming.
+		}
+#endif
+
+		// The tick here is intended to schedule the success or cancel callback.
+		// Using TT_None ensure gets which could create a dead lock.
+		Tick(FSkeletalMeshUpdate::TT_None);
+	};
+}
+
+void FSkeletalMeshStreamIn_IO::SetIORequest(const FContext& Context, const FString& IOFilename)
+{
+	if (IsCancelled())
+	{
+		return;
+	}
+	check(!IOFileHandle && PendingFirstMip < CurrentFirstLODIdx);
+	IOFileHandle = FPlatformFileManager::Get().GetPlatformFile().OpenAsyncRead(*IOFilename);
+	FSkeletalMeshRenderData* RenderData = Context.RenderData;
+	FString debugName;
+	if (Context.Mesh)
+	{
+		debugName = Context.Mesh->GetName();
+	}
+	if (IOFileHandle && RenderData)
+	{
+		SetAsyncFileCallback(Context);
+
+		int64 ReadOffset;
+		int64 ReadSize;
+		{
+			FSkeletalMeshLODRenderData& FirstLOD = RenderData->LODRenderData[PendingFirstMip];
+			FSkeletalMeshLODRenderData& LastLOD = RenderData->LODRenderData[CurrentFirstLODIdx - 1];
+			ReadOffset = FirstLOD.OffsetInFile;
+			ReadSize = LastLOD.OffsetInFile + LastLOD.BulkDataSize - ReadOffset;
+#if DO_CHECK
+			int64 TestOffset = ReadOffset;
+			for (int32 LODIdx = PendingFirstMip; LODIdx < CurrentFirstLODIdx; ++LODIdx)
+			{
+				FSkeletalMeshLODRenderData& LODResource = RenderData->LODRenderData[LODIdx];
+				const int64 CurOffset = (int64)LODResource.OffsetInFile;
+				check(ReadOffset <= CurOffset);
+				check(TestOffset == CurOffset);
+				TestOffset += LODResource.BulkDataSize;
+			}
+			check(TestOffset == ReadOffset + ReadSize);
+#endif
+		}
+
+		// Increment as we push the request. If a request complete immediately, then it will call the callback
+		// but that won't do anything because the tick would not try to acquire the lock since it is already locked.
+		TaskSynchronization.Increment();
+		LODData.Empty(ReadSize);
+		LODData.AddUninitialized(ReadSize);
+		IORequest = IOFileHandle->ReadRequest(
+			ReadOffset,
+			ReadSize,
+			bHighPrioIORequest ? AIOP_BelowNormal : AIOP_Low,
+			&AsyncFileCallback,
+			LODData.GetData());
+	}
+	else
+	{
+		MarkAsCancelled();
+	}
+}
+
+void FSkeletalMeshStreamIn_IO::ClearIORequest(const FContext& Context)
+{
+	if (IOFileHandle)
+	{
+		if (IORequest)
+		{
+			// If clearing requests not yet completed, cancel and wait.
+			if (!IORequest->PollCompletion())
+			{
+				IORequest->Cancel();
+				IORequest->WaitCompletion();
+			}
+			delete IORequest;
+			IORequest = nullptr;
+		}
+
+		delete IOFileHandle;
+		IOFileHandle = nullptr;
+	}
+}
+
+void FSkeletalMeshStreamIn_IO::SerializeLODData(const FContext& Context)
+{
+	LLM_SCOPE(ELLMTag::SkeletalMesh);
+
+	check(!TaskSynchronization.GetValue());
+	USkeletalMesh* Mesh = Context.Mesh;
+	FSkeletalMeshRenderData* RenderData = Context.RenderData;
+	if (!IsCancelled() && Mesh && RenderData)
+	{
+		check(PendingFirstMip < CurrentFirstLODIdx && CurrentFirstLODIdx == RenderData->CurrentFirstLODIdx);
+		FMemoryReader Ar(LODData, true);
+		for (int32 LODIdx = PendingFirstMip; LODIdx < CurrentFirstLODIdx; ++LODIdx)
+		{
+			FSkeletalMeshLODRenderData& LODResource = RenderData->LODRenderData[LODIdx];
+			const bool bForceKeepCPUResources = FSkeletalMeshLODRenderData::ShouldForceKeepCPUResources();
+			const bool bNeedsCPUAccess = FSkeletalMeshLODRenderData::ShouldKeepCPUResources(Mesh, LODIdx, bForceKeepCPUResources);
+			constexpr uint8 DummyStripFlags = 0;
+			LODResource.SerializeStreamedData(Ar, Mesh, LODIdx, DummyStripFlags, bNeedsCPUAccess, bForceKeepCPUResources);
+		}
+		LODData.Empty();
+	}
+}
+
+void FSkeletalMeshStreamIn_IO::CancelIORequest()
+{
+	if (IORequest)
+	{
+		// Calling cancel will trigger the SetAsyncFileCallback() which will also try a tick but will fail.
+		IORequest->Cancel();
+	}
+}
+
+template <bool bRenderThread>
+TSkeletalMeshStreamIn_IO<bRenderThread>::TSkeletalMeshStreamIn_IO(USkeletalMesh* InMesh, int32 InRequestedMips, bool bHighPrio)
+	: FSkeletalMeshStreamIn_IO(InMesh, InRequestedMips, bHighPrio)
+{
+	PushTask(FContext(InMesh, TT_None), TT_Async, SRA_UPDATE_CALLBACK(DoInitiateIO), TT_None, nullptr);
+}
+
+template <bool bRenderThread>
+void TSkeletalMeshStreamIn_IO<bRenderThread>::DoInitiateIO(const FContext& Context)
+{
+	check(Context.CurrentThread == TT_Async);
+	const FString IOFilename = GetIOFilename(Context);
+	SetIORequest(Context, IOFilename);
+	PushTask(Context, TT_Async, SRA_UPDATE_CALLBACK(DoSerializeLODData), TT_Async, SRA_UPDATE_CALLBACK(DoCancelIO));
+}
+
+template <bool bRenderThread>
+void TSkeletalMeshStreamIn_IO<bRenderThread>::DoSerializeLODData(const FContext& Context)
+{
+	check(Context.CurrentThread == TT_Async);
+	SerializeLODData(Context);
+	ClearIORequest(Context);
+	const EThreadType TThread = bRenderThread ? TT_Render : TT_Async;
+	const EThreadType CThread = (EThreadType)Context.CurrentThread;
+	PushTask(Context, TThread, SRA_UPDATE_CALLBACK(DoCreateBuffers), CThread, SRA_UPDATE_CALLBACK(DoCancel));
+}
+
+template <bool bRenderThread>
+void TSkeletalMeshStreamIn_IO<bRenderThread>::DoCreateBuffers(const FContext& Context)
+{
+	if (bRenderThread)
+	{
+		CreateBuffers_RenderThread(Context);
+	}
+	else
+	{
+		CreateBuffers_Async(Context);
+	}
+	check(!TaskSynchronization.GetValue());
+	PushTask(Context, TT_Render, SRA_UPDATE_CALLBACK(DoFinishUpdate), (EThreadType)Context.CurrentThread, SRA_UPDATE_CALLBACK(DoCancel));
+}
+
+template <bool bRenderThread>
+void TSkeletalMeshStreamIn_IO<bRenderThread>::DoCancelIO(const FContext& Context)
+{
+	ClearIORequest(Context);
+	PushTask(Context, TT_None, nullptr, (EThreadType)Context.CurrentThread, SRA_UPDATE_CALLBACK(DoCancel));
+}
+
+template class TSkeletalMeshStreamIn_IO<true>;
+template class TSkeletalMeshStreamIn_IO<false>;
+
+#if WITH_EDITOR
+FSkeletalMeshStreamIn_DDC::FSkeletalMeshStreamIn_DDC(USkeletalMesh* InMesh, int32 InRequestedMips)
+	: FSkeletalMeshStreamIn(InMesh, InRequestedMips)
+	, bDerivedDataInvalid(false)
+{}
+
+void FSkeletalMeshStreamIn_DDC::LoadNewLODsFromDDC(const FContext& Context)
+{
+	check(Context.CurrentThread == TT_Async);
+	// TODO: support streaming CPU data for editor builds
+}
+
+template <bool bRenderThread>
+TSkeletalMeshStreamIn_DDC<bRenderThread>::TSkeletalMeshStreamIn_DDC(USkeletalMesh* InMesh, int32 InRequestedMips)
+	: FSkeletalMeshStreamIn_DDC(InMesh, InRequestedMips)
+{
+	PushTask(FContext(InMesh, TT_None), TT_Async, SRA_UPDATE_CALLBACK(DoLoadNewLODsFromDDC), TT_None, nullptr);
+}
+
+template <bool bRenderThread>
+void TSkeletalMeshStreamIn_DDC<bRenderThread>::DoLoadNewLODsFromDDC(const FContext& Context)
+{
+	LoadNewLODsFromDDC(Context);
+	check(!TaskSynchronization.GetValue());
+	const EThreadType TThread = bRenderThread ? TT_Render : TT_Async;
+	const EThreadType CThread = (EThreadType)Context.CurrentThread;
+	PushTask(Context, TThread, SRA_UPDATE_CALLBACK(DoCreateBuffers), CThread, SRA_UPDATE_CALLBACK(DoCancel));
+}
+
+template <bool bRenderThread>
+void TSkeletalMeshStreamIn_DDC<bRenderThread>::DoCreateBuffers(const FContext& Context)
+{
+	if (bRenderThread)
+	{
+		CreateBuffers_RenderThread(Context);
+	}
+	else
+	{
+		CreateBuffers_Async(Context);
+	}
+	check(!TaskSynchronization.GetValue());
+	PushTask(Context, TT_Render, SRA_UPDATE_CALLBACK(DoFinishUpdate), (EThreadType)Context.CurrentThread, SRA_UPDATE_CALLBACK(DoCancel));
+}
+
+template class TSkeletalMeshStreamIn_DDC<true>;
+template class TSkeletalMeshStreamIn_DDC<false>;
+#endif
diff --git a/Engine/Source/Runtime/Engine/Private/Streaming/SkeletalMeshUpdate.h b/Engine/Source/Runtime/Engine/Private/Streaming/SkeletalMeshUpdate.h
new file mode 100644
--- /dev/null
+++ b/Engine/Source/Runtime/Engine/Private/Streaming/SkeletalMeshUpdate.h
@@ -0,0 +1,257 @@
+// Copyright 1998-2019 Epic Games, Inc. All Rights Reserved.
+
+/*=============================================================================
+SkeletalMeshUpdate.h: Helpers to stream in and out skeletal mesh LODs.
+=============================================================================*/
+
+#pragma once
+
+#include "CoreMinimal.h"
+#include "Engine/SkeletalMesh.h"
+#include "Async/AsyncFileHandle.h"
+
+/**
+* A context used to update or proceed with the next update step.
+* The mesh and render data references could be stored in the update object
+* but are currently kept outside to avoid lifetime management within the object.
+*/
+struct FSkelMeshUpdateContext
+{
+	typedef int32 EThreadType;
+
+	FSkelMeshUpdateContext(USkeletalMesh* InMesh, EThreadType InCurrentThread);
+
+	FSkelMeshUpdateContext(UStreamableRenderAsset* InMesh, EThreadType InCurrentThread);
+
+	UStreamableRenderAsset* GetRenderAsset() const
+	{
+		return Mesh;
+	}
+
+	EThreadType GetCurrentThread() const
+	{
+		return CurrentThread;
+	}
+
+	/** The mesh to update, this must be the same one as the one used when creating the FSkeletalMeshUpdate object. */
+	USkeletalMesh* Mesh;
+	/** The current render data of this mesh. */
+	FSkeletalMeshRenderData* RenderData;
+	/** The thread on which the context was created. */
+	EThreadType CurrentThread;
+};
+
+extern template class TRenderAssetUpdate<FSkelMeshUpdateContext>;
+
+/**
+* This class provides a framework for loading and unloading the LODs of skeletal meshes.
+* Each thread essentially calls Tick() until the job is done.
+* The object can be safely deleted when IsCompleted() returns true.
+*/
+class FSkeletalMeshUpdate : public TRenderAssetUpdate<FSkelMeshUpdateContext>
+{
+public:
+	FSkeletalMeshUpdate(USkeletalMesh* InMesh, int32 InRequestedMips);
+
+	virtual ~FSkeletalMeshUpdate() {}
+
+	virtual void Abort()
+	{
+		TRenderAssetUpdate<FSkelMeshUpdateContext>::Abort();
+	}
+
+#if WITH_EDITOR
+	/** Returns whether DDC of this mesh needs to be regenerated.  */
+	virtual bool DDCIsInvalid() const { return false; }
+#endif
+
+protected:
+	/** Cached index of current first LOD that will be replaced by PendingFirstMip */
+	int32 CurrentFirstLODIdx;
+};
+
+class FSkeletalMeshStreamIn : public FSkeletalMeshUpdate
+{
+public:
+	FSkeletalMeshStreamIn(USkeletalMesh* InMesh, int32 InRequestedMips);
+
+	virtual ~FSkeletalMeshStreamIn();
+
+protected:
+	/** Correspond to the buffers in FSkeletalMeshLODRenderData */
+	struct FIntermediateBuffers
+	{
+		FVertexBufferRHIRef TangentsVertexBuffer;
+		FVertexBufferRHIRef TexCoordVertexBuffer;
+		FVertexBufferRHIRef PositionVertexBuffer;
+		FVertexBufferRHIRef ColorVertexBuffer;
+		FVertexBufferRHIRef SkinWeightVertexBuffer;
+		FVertexBufferRHIRef ClothVertexBuffer;
+		FIndexBufferRHIRef IndexBuffer;
+		FIndexBufferRHIRef AdjacencyIndexBuffer;
+		TArray<TPair<FName, FVertexBufferRHIRef>> AltSkinWeightVertexBuffers;
+
+		void CreateFromCPUData_RenderThread(USkeletalMesh* Mesh, FSkeletalMeshLODRenderData& LODResource);
+		void CreateFromCPUData_Async(USkeletalMesh* Mesh, FSkeletalMeshLODRenderData& LODResource);
+
+		void SafeRelease();
+
+		/** Transfer ownership of buffers to a LOD resource */
+		template <uint32 MaxNumUpdates>
+		void TransferBuffers(FSkeletalMeshLODRenderData& LODResource, TRHIResourceUpdateBatcher<MaxNumUpdates>& Batcher);
+
+		void CheckIsNull() const;
+	};
+
+	/** Create buffers with new LOD data on render or pooled thread */
+	void CreateBuffers_RenderThread(const FContext& Context);
+	void CreateBuffers_Async(const FContext& Context);
+
+	/** Discard newly streamed-in CPU data */
+	void DiscardNewLODs(const FContext& Context);
+
+	/** Apply the new buffers (if not cancelled) and finish the update process. When cancelled, the intermediate buffers simply gets discarded. */
+	void DoFinishUpdate(const FContext& Context);
+
+	/** Discard streamed-in CPU data and intermediate RHI buffers */
+	void DoCancel(const FContext& Context);
+
+	/** The intermediate buffers created in the update process. */
+	FIntermediateBuffers IntermediateBuffersArray[MAX_MESH_LOD_COUNT];
+
+private:
+	template <bool bRenderThread>
+	void CreateBuffers_Internal(const FContext& Context);
+};
+
+class FSkeletalMeshStreamOut : public FSkeletalMeshUpdate
+{
+public:
+	FSkeletalMeshStreamOut(USkeletalMesh* InMesh, int32 InRequestedMips);
+
+	virtual ~FSkeletalMeshStreamOut() {}
+
+private:
+	void DoConditionalMarkComponentsDirty(const FContext& Context);
+
+	/** Release RHI buffers and update SRVs */
+	void DoReleaseBuffers(const FContext& Context);
+
+	/** */
+	void DoCancel(const FContext& Context);
+
+	uint32 StartFrameNumber;
+};
+
+class FSkeletalMeshStreamIn_IO : public FSkeletalMeshStreamIn
+{
+public:
+	FSkeletalMeshStreamIn_IO(USkeletalMesh* InMesh, int32 InRequestedMips, bool bHighPrio);
+
+	virtual ~FSkeletalMeshStreamIn_IO() {}
+
+	virtual void Abort() override;
+
+protected:
+	class FCancelIORequestsTask : public FNonAbandonableTask
+	{
+	public:
+		FCancelIORequestsTask(FSkeletalMeshStreamIn_IO* InPendingUpdate)
+			: PendingUpdate(InPendingUpdate)
+		{}
+
+		void DoWork();
+
+		FORCEINLINE TStatId GetStatId() const
+		{
+			RETURN_QUICK_DECLARE_CYCLE_STAT(FCancelIORequestsTask_SkeletalMesh, STATGROUP_ThreadPoolAsyncTasks);
+		}
+
+	private:
+		FSkeletalMeshStreamIn_IO* PendingUpdate;
+	};
+
+	typedef FAutoDeleteAsyncTask<FCancelIORequestsTask> FAsyncCancelIORequestsTask;
+	friend class FCancelIORequestsTask;
+
+	/** Figure out the full name of the .bulk file */
+	FString GetIOFilename(const FContext& Context);
+
+	/** Set a callback called when IORequest is completed or cancelled */
+	void SetAsyncFileCallback(const FContext& Context);
+
+	/** Create a new async IO request to read in LOD data */
+	void SetIORequest(const FContext& Context, const FString& IOFilename);
+
+	/** Release IORequest and IOFileHandle. IORequest will be cancelled if still inflight */
+	void ClearIORequest(const FContext& Context);
+
+	/** Serialize data of new LODs to corresponding FStaticMeshLODResources */
+	void SerializeLODData(const FContext& Context);
+
+	/** Called by FAsyncCancelIORequestsTask to cancel inflight IO request if any */
+	void CancelIORequest();
+
+	TArray<uint8> LODData;
+	IAsyncReadFileHandle* IOFileHandle;
+	IAsyncReadRequest* IORequest;
+	FAsyncFileCallBack AsyncFileCallback;
+	bool bHighPrioIORequest;
+};
+
+template <bool bRenderThread>
+class TSkeletalMeshStreamIn_IO : public FSkeletalMeshStreamIn_IO
+{
+public:
+	TSkeletalMeshStreamIn_IO(USkeletalMesh* InMesh, int32 InRequestedMips, bool bHighPrio);
+
+	virtual ~TSkeletalMeshStreamIn_IO() {}
+
+protected:
+	void DoInitiateIO(const FContext& Context);
+
+	void DoSerializeLODData(const FContext& Context);
+
+	void DoCreateBuffers(const FContext& Context);
+
+	void DoCancelIO(const FContext& Context);
+};
+
+typedef TSkeletalMeshStreamIn_IO<true> FSkeletalMeshStreamIn_IO_RenderThread;
+typedef TSkeletalMeshStreamIn_IO<false> FSkeletalMeshStreamIn_IO_Async;
+
+#if WITH_EDITOR
+class FSkeletalMeshStreamIn_DDC : public FSkeletalMeshStreamIn
+{
+public:
+	FSkeletalMeshStreamIn_DDC(USkeletalMesh* InMesh, int32 InRequestedMips);
+
+	virtual ~FSkeletalMeshStreamIn_DDC() {}
+
+	virtual bool DDCIsInvalid() const override { return bDerivedDataInvalid; }
+
+protected:
+	void LoadNewLODsFromDDC(const FContext& Context);
+
+	bool bDerivedDataInvalid;
+};
+
+template <bool bRenderThread>
+class TSkeletalMeshStreamIn_DDC : public FSkeletalMeshStreamIn_DDC
+{
+public:
+	TSkeletalMeshStreamIn_DDC(USkeletalMesh* InMesh, int32 InRequestedMips);
+
+	virtual ~TSkeletalMeshStreamIn_DDC() {}
+
+private:
+	/** Load new LOD buffers from DDC and queue a task to create RHI buffers on RT */
+	void DoLoadNewLODsFromDDC(const FContext& Context);
+
+	/** Create RHI buffers for newly streamed-in LODs and queue a task to rename references on RT */
+	void DoCreateBuffers(const FContext& Context);
+};
+
+typedef TSkeletalMeshStreamIn_DDC<true> FSkeletalMeshStreamIn_DDC_RenderThread;
+typedef TSkeletalMeshStreamIn_DDC<false> FSkeletalMeshStreamIn_DDC_Async;
+#endif
\ No newline at end of file
diff --git a/Engine/Source/Runtime/Engine/Private/Streaming/StaticMeshUpdate.cpp b/Engine/Source/Runtime/Engine/Private/Streaming/StaticMeshUpdate.cpp
--- a/Engine/Source/Runtime/Engine/Private/Streaming/StaticMeshUpdate.cpp
+++ b/Engine/Source/Runtime/Engine/Private/Streaming/StaticMeshUpdate.cpp
@@ -15,8 +15,8 @@ StaticMeshUpdate.cpp: Helpers to stream in and out static mesh LODs.
 // Instantiate TRenderAssetUpdate for FStaticMeshUpdateContext
 template class TRenderAssetUpdate<FStaticMeshUpdateContext>;
 
-static constexpr uint32 GMaxNumResourceUpdatesPerLOD = 14;
-static constexpr uint32 GMaxNumResourceUpdatesPerBatch = (MAX_STATIC_MESH_LODS - 1) * GMaxNumResourceUpdatesPerLOD;
+static constexpr uint32 GStaticMeshMaxNumResourceUpdatesPerLOD = 14;
+static constexpr uint32 GStaticMeshMaxNumResourceUpdatesPerBatch = (MAX_STATIC_MESH_LODS - 1) * GStaticMeshMaxNumResourceUpdatesPerLOD;
 
 FStaticMeshUpdateContext::FStaticMeshUpdateContext(UStaticMesh* InMesh, EThreadType InCurrentThread)
 	: Mesh(InMesh)
@@ -219,7 +219,7 @@ void FStaticMeshStreamIn::DoFinishUpdate(const FContext& Context)
 			&& PendingFirstMip < CurrentFirstLODIdx);
 		// Use a scope to flush the batcher before updating CurrentFirstLODIdx
 		{
-			TRHIResourceUpdateBatcher<GMaxNumResourceUpdatesPerBatch> Batcher;
+			TRHIResourceUpdateBatcher<GStaticMeshMaxNumResourceUpdatesPerBatch> Batcher;
 
 			for (int32 LODIdx = PendingFirstMip; LODIdx < CurrentFirstLODIdx; ++LODIdx)
 			{
@@ -270,7 +270,7 @@ void FStaticMeshStreamOut::DoReleaseBuffers(const FContext& Context)
 		RenderData->CurrentFirstLODIdx = PendingFirstMip;
 		Mesh->SetCachedNumResidentLODs(static_cast<uint8>(RenderData->LODResources.Num() - PendingFirstMip));
 
-		TRHIResourceUpdateBatcher<GMaxNumResourceUpdatesPerBatch> Batcher;
+		TRHIResourceUpdateBatcher<GStaticMeshMaxNumResourceUpdatesPerBatch> Batcher;
 
 		for (int32 LODIdx = CurrentFirstLODIdx; LODIdx < PendingFirstMip; ++LODIdx)
 		{
@@ -443,6 +443,8 @@ void FStaticMeshStreamIn_IO::ClearIORequest(const FContext& Context)
 
 void FStaticMeshStreamIn_IO::SerializeLODData(const FContext& Context)
 {
+	LLM_SCOPE(ELLMTag::StaticMesh);
+
 	check(!TaskSynchronization.GetValue());
 	UStaticMesh* Mesh = Context.Mesh;
 	FStaticMeshRenderData* RenderData = Context.RenderData;
diff --git a/Engine/Source/Runtime/Engine/Private/Streaming/StreamingManagerTexture.cpp b/Engine/Source/Runtime/Engine/Private/Streaming/StreamingManagerTexture.cpp
--- a/Engine/Source/Runtime/Engine/Private/Streaming/StreamingManagerTexture.cpp
+++ b/Engine/Source/Runtime/Engine/Private/Streaming/StreamingManagerTexture.cpp
@@ -9,6 +9,7 @@
 #include "Engine/World.h"
 #include "Engine/TextureStreamingTypes.h"
 #include "Engine/StaticMesh.h"
+#include "Engine/SkeletalMesh.h"
 #include "Materials/MaterialInterface.h"
 #include "Misc/CommandLine.h"
 #include "Misc/ConfigCacheIni.h"
@@ -118,6 +119,8 @@ FRenderAssetStreamingManager::FRenderAssetStreamingManager()
 	// TODO: NumStreamedMips_StaticMesh, NumStreamedMips_SkeletalMesh
 	NumStreamedMips_StaticMesh.Empty(1);
 	NumStreamedMips_StaticMesh.Add(INT32_MAX);
+	NumStreamedMips_SkeletalMesh.Empty(1);
+	NumStreamedMips_SkeletalMesh.Add(INT32_MAX);
 
 	// setup the streaming resource flush function pointer
 	GFlushStreamingFunc = &FlushResourceStreaming;
@@ -701,8 +704,7 @@ void FRenderAssetStreamingManager::AddStreamingRenderAsset(UStaticMesh* StaticMe
 
 void FRenderAssetStreamingManager::AddStreamingRenderAsset(USkeletalMesh* SkeletalMesh)
 {
-	// TODO
-	LowLevelFatalError(TEXT("FRenderAssetStreamingManager::AddStreamingRenderAsset(USkeletalMesh* SkeletalMesh) is not implemented"));
+	AddStreamingRenderAsset_Internal(SkeletalMesh, FStreamingRenderAsset::AT_SkeletalMesh);
 }
 
 /**
@@ -1571,6 +1573,61 @@ void FRenderAssetStreamingManager::GetObjectReferenceBounds(const UObject* RefOb
 	}
 }
 
+static bool IsLevelManagerValid(const FLevelRenderAssetManager* LevelManager)
+{
+	return LevelManager && LevelManager->IsInitialized() && LevelManager->HasRenderAssetReferences() && LevelManager->GetLevel()->bIsVisible;
+}
+
+static void GatherComponentsFromView(
+	const UStreamableRenderAsset* RenderAsset,
+	const FRenderAssetInstanceView* View,
+	const TFunction<bool(const UPrimitiveComponent*)>& ShouldChoose,
+	TArray<const UPrimitiveComponent*>& OutComps)
+{
+	if (View)
+	{
+		for (FRenderAssetInstanceView::FRenderAssetLinkConstIterator It = View->GetElementIterator(RenderAsset); It; ++It)
+		{
+			const UPrimitiveComponent* Comp = It.GetComponent();
+			if (Comp && ShouldChoose(Comp))
+			{
+				OutComps.Add(Comp);
+			}
+		}
+	}
+}
+
+void FRenderAssetStreamingManager::GetAssetComponents(const UStreamableRenderAsset* RenderAsset, TArray<const UPrimitiveComponent*>& OutComps, TFunction<bool(const UPrimitiveComponent*)> ShouldChoose)
+{
+	checkSlow(IsInGameThread());
+
+	if (RenderAsset && RenderAsset->StreamingIndex >= 0)
+	{
+		FScopeLock Lock(&CriticalSection);
+		QUICK_SCOPE_CYCLE_COUNTER(STAT_FRenderAssetStreamingManager_GetAssetComponents);
+
+		const int32 Idx = RenderAsset->StreamingIndex;
+		check(StreamingRenderAssets.IsValidIndex(Idx) && StreamingRenderAssets[Idx].RenderAsset == RenderAsset);
+		const FStreamingRenderAsset& StreamingRenderAsset = StreamingRenderAssets[Idx];
+
+		for (TConstSetBitIterator<> It(StreamingRenderAsset.LevelIndexUsage); It; ++It)
+		{
+			const int32 LevelIdx = It.GetIndex();
+			check(LevelRenderAssetManagers.IsValidIndex(LevelIdx));
+			FLevelRenderAssetManager* LevelManager = LevelRenderAssetManagers[LevelIdx];
+
+			if (IsLevelManagerValid(LevelManager))
+			{
+				const FRenderAssetInstanceView* StaticInstancesView = LevelManager->GetRawAsyncView();
+				GatherComponentsFromView(RenderAsset, StaticInstancesView, ShouldChoose, OutComps);
+			}
+		}
+
+		const FRenderAssetInstanceView* DynamicInstancesView = DynamicComponentManager.GetAsyncView(false);
+		GatherComponentsFromView(RenderAsset, DynamicInstancesView, ShouldChoose, OutComps);
+	}
+}
+
 void FRenderAssetStreamingManager::PropagateLightingScenarioChange()
 {
 	FScopeLock ScopeLock(&CriticalSection);
@@ -1625,6 +1682,25 @@ bool FRenderAssetStreamingManager::HandleListStreamingRenderAssetsCommand( const
 	SyncStates(true);
 
 	const bool bShouldOnlyListUnkownRef = FParse::Command(&Cmd, TEXT("UNKOWNREF"));
+	FStreamingRenderAsset::EAssetType ListAssetType = FStreamingRenderAsset::AT_Num;
+	{
+		FString AssetTypeStr;
+		if (FParse::Value(Cmd, TEXT("AssetType="), AssetTypeStr))
+		{
+			if (AssetTypeStr == TEXT("Texture"))
+			{
+				ListAssetType = FStreamingRenderAsset::AT_Texture;
+			}
+			else if (AssetTypeStr == TEXT("StaticMesh"))
+			{
+				ListAssetType = FStreamingRenderAsset::AT_StaticMesh;
+			}
+			else if (AssetTypeStr == TEXT("SkeletalMesh"))
+			{
+				ListAssetType = FStreamingRenderAsset::AT_SkeletalMesh;
+			}
+		}
+	}
 
 	// Sort texture/mesh by names so that the state can be compared between runs.
 	TMap<FString, int32> SortedRenderAssets;
@@ -1645,14 +1721,19 @@ bool FRenderAssetStreamingManager::HandleListStreamingRenderAssetsCommand( const
 		const UStreamableRenderAsset* RenderAsset = StreamingRenderAsset.RenderAsset;
 		const typename FStreamingRenderAsset::EAssetType AssetType = StreamingRenderAsset.RenderAssetType;
 		
+		if (ListAssetType != FStreamingRenderAsset::AT_Num && ListAssetType != AssetType)
+		{
+			continue;
+		}
+
 		UE_LOG(LogContentStreaming, Log,  TEXT("%s [%d] : %s"),
 			FStreamingRenderAsset::GetStreamingAssetTypeStr(AssetType),
 			It.Value(),
 			*RenderAsset->GetFullName() );
 
-		int32 CurrentMipIndex = FMath::Max(RenderAsset->GetNumMipsForStreaming() - StreamingRenderAsset.ResidentMips, 0);
-		int32 WantedMipIndex = FMath::Max(RenderAsset->GetNumMipsForStreaming() - StreamingRenderAsset.GetPerfectWantedMips(), 0);
-		int32 MaxAllowedMipIndex = FMath::Max(RenderAsset->GetNumMipsForStreaming() - StreamingRenderAsset.MaxAllowedMips, 0);
+		const int32 CurrentMipIndex = FMath::Max(RenderAsset->GetNumMipsForStreaming() - StreamingRenderAsset.ResidentMips, 0);
+		const int32 WantedMipIndex = FMath::Max(RenderAsset->GetNumMipsForStreaming() - StreamingRenderAsset.GetPerfectWantedMips(), 0);
+		const int32 MaxAllowedMipIndex = FMath::Max(RenderAsset->GetNumMipsForStreaming() - StreamingRenderAsset.MaxAllowedMips, 0);
 
 		if (AssetType == FStreamingRenderAsset::AT_Texture)
 		{
@@ -1690,10 +1771,12 @@ bool FRenderAssetStreamingManager::HandleListStreamingRenderAssetsCommand( const
 				LODGroupName = StaticMesh->LODGroup.ToString();
 			}
 #endif
-			UE_LOG(LogContentStreaming, Log, TEXT("    CurrentLOD=%d WantedLOD=%d MaxAllowedLOD=%d LastRenderTime=%s BudgetBias=%d Group=%s"),
+			UE_LOG(LogContentStreaming, Log, TEXT("    CurrentLOD=%d WantedLOD=%d MaxAllowedLOD=%d NumLODs=%d NumForcedLODs=%d LastRenderTime=%s BudgetBias=%d Group=%s"),
 				CurrentMipIndex,
 				WantedMipIndex,
 				MaxAllowedMipIndex,
+				StreamingRenderAsset.MipCount,
+				StreamingRenderAsset.NumForcedMips,
 				LastRenderTime == MAX_FLT ? TEXT("NotTracked") : *FString::Printf(TEXT("%.3f"), LastRenderTime),
 				StreamingRenderAsset.BudgetMipBias,
 				*LODGroupName);
@@ -1731,7 +1814,7 @@ bool FRenderAssetStreamingManager::HandleCancelRenderAssetStreamingCommand( cons
 
 	UTexture2D::CancelPendingTextureStreaming();
 	UStaticMesh::CancelAllPendingStreamingActions();
-	// TODO: USkeletalMesh
+	USkeletalMesh::CancelAllPendingStreamingActions();
 	return true;
 }
 
diff --git a/Engine/Source/Runtime/Engine/Private/Streaming/StreamingManagerTexture.h b/Engine/Source/Runtime/Engine/Private/Streaming/StreamingManagerTexture.h
--- a/Engine/Source/Runtime/Engine/Private/Streaming/StreamingManagerTexture.h
+++ b/Engine/Source/Runtime/Engine/Private/Streaming/StreamingManagerTexture.h
@@ -165,6 +165,8 @@ struct FRenderAssetStreamingManager final : public IRenderAssetStreamingManager
 	/** Return all bounds related to the ref object */
 	virtual void GetObjectReferenceBounds(const UObject* RefObject, TArray<FBox>& AssetBoxes) override;
 
+	virtual void GetAssetComponents(const UStreamableRenderAsset* RenderAsset, TArray<const UPrimitiveComponent*>& OutComps, TFunction<bool(const UPrimitiveComponent*)> ShouldChoose) override;
+
 	/** Propagates a change to the active lighting scenario. */
 	void PropagateLightingScenarioChange() override;
 
diff --git a/Engine/Source/Runtime/Engine/Private/Streaming/TextureInstanceView.cpp b/Engine/Source/Runtime/Engine/Private/Streaming/TextureInstanceView.cpp
--- a/Engine/Source/Runtime/Engine/Private/Streaming/TextureInstanceView.cpp
+++ b/Engine/Source/Runtime/Engine/Private/Streaming/TextureInstanceView.cpp
@@ -578,7 +578,7 @@ void FRenderAssetInstanceAsyncView::GetRenderAssetScreenSize(
 		else
 		{
 			int32 IterationCount_DebuggingOnly = 0;
-			for (auto It = View->GetElementIterator(InAsset); It && (MaxSize_VisibleOnly < MAX_TEXTURE_SIZE || LogPrefix); ++It, ++IterationCount_DebuggingOnly)
+			for (auto It = View->GetElementIterator(InAsset); It && (AssetType != FStreamingRenderAsset::AT_Texture || MaxSize_VisibleOnly < MAX_TEXTURE_SIZE || LogPrefix); ++It, ++IterationCount_DebuggingOnly)
 			{
 				View->VerifyElementIdx_DebuggingOnly(It.GetCurElementIdx_ForDebuggingOnly(), IterationCount_DebuggingOnly);
 				// Only handle elements that are in bounds.
diff --git a/Engine/Source/Runtime/Engine/Private/Streaming/TextureInstanceView.h b/Engine/Source/Runtime/Engine/Private/Streaming/TextureInstanceView.h
--- a/Engine/Source/Runtime/Engine/Private/Streaming/TextureInstanceView.h
+++ b/Engine/Source/Runtime/Engine/Private/Streaming/TextureInstanceView.h
@@ -145,8 +145,6 @@ public:
 
 		int32 GetCurElementIdx_ForDebuggingOnly() const { return CurrElementIndex; }
 
-	protected:
-
 		FORCEINLINE const UPrimitiveComponent* GetComponent() const { return State.Elements[CurrElementIndex].Component; }
 
 		const FRenderAssetInstanceView& State;
diff --git a/Engine/Source/Runtime/Engine/Private/UnrealEngine.cpp b/Engine/Source/Runtime/Engine/Private/UnrealEngine.cpp
--- a/Engine/Source/Runtime/Engine/Private/UnrealEngine.cpp
+++ b/Engine/Source/Runtime/Engine/Private/UnrealEngine.cpp
@@ -3384,6 +3384,9 @@ struct FSortedSkeletalMesh
 	int32		LodCount;
 	int32		MobileMinLOD;
 	int32		MeshDisablesMinLODStripping;
+	int32		bSupportLODStreaming;
+	int32		MaxNumStreamedLODs;
+	int32		MaxNumOptionalLODs;
 	int32		VertexCountLod0;
 	int32		VertexCountLod1;
 	int32		VertexCountLod2;
@@ -3393,7 +3396,8 @@ struct FSortedSkeletalMesh
 	FString		Name;
 
 	/** Constructor, initializing every member variable with passed in values. */
-	FSortedSkeletalMesh(int32 InNumKB, int32 InMaxKB, int32 InResKBExc, int32 InResKBInc, int32 InResKBIncMobile, int32 InLodCount, int32 InMobileMinLOD, int32 InMeshDisablesMinLODStripping, int32 InVertexCountLod0,
+	FSortedSkeletalMesh(int32 InNumKB, int32 InMaxKB, int32 InResKBExc, int32 InResKBInc, int32 InResKBIncMobile, int32 InLodCount, int32 InMobileMinLOD,
+		int32 InMeshDisablesMinLODStripping, int32 bInSupportLODStreaming, int32 InMaxNumStreamedLODs, int32 InMaxNumOptionalLODs, int32 InVertexCountLod0,
 		int32 InVertexCountLod1, int32 InVertexCountLod2, int32 InVertexCountTotal, int32 InVertexCountTotalMobile, int32 InVertexCountCollision, FString InName)
 		: NumKB(InNumKB)
 		, MaxKB(InMaxKB)
@@ -3403,6 +3407,9 @@ struct FSortedSkeletalMesh
 		, LodCount(InLodCount)
 		, MobileMinLOD(InMobileMinLOD)
 		, MeshDisablesMinLODStripping(InMeshDisablesMinLODStripping)
+		, bSupportLODStreaming(bInSupportLODStreaming)
+		, MaxNumStreamedLODs(InMaxNumStreamedLODs)
+		, MaxNumOptionalLODs(InMaxNumOptionalLODs)
 		, VertexCountLod0(InVertexCountLod0)
 		, VertexCountLod1(InVertexCountLod1)
 		, VertexCountLod2(InVertexCountLod2)
@@ -5325,6 +5332,9 @@ bool UEngine::HandleListSkeletalMeshesCommand(const TCHAR* Cmd, FOutputDevice& A
 		int32 LodCount = 0;
 		int32 MobileMinLOD = -1;
 		int32 MeshDisablesMinLODStripping = 0;
+		int32 bSupportLODStreaming = -1;
+		int32 MaxNumStreamedLODs = -1;
+		int32 MaxNumOptionalLODs = -1;
 #if WITH_EDITORONLY_DATA
 		if (Mesh->MinLod.PerPlatform.Find(("Mobile")) != nullptr)
 		{
@@ -5336,6 +5346,24 @@ bool UEngine::HandleListSkeletalMeshesCommand(const TCHAR* Cmd, FOutputDevice& A
 		{
 			MeshDisablesMinLODStripping = *Mesh->DisableBelowMinLodStripping.PerPlatform.Find(("Mobile")) ? 1 : 0;
 		}
+
+		bSupportLODStreaming = Mesh->bSupportLODStreaming.Default;
+		if (const bool* Found = Mesh->bSupportLODStreaming.PerPlatform.Find("Mobile"))
+		{
+			bSupportLODStreaming = *Found;
+		}
+
+		MaxNumStreamedLODs = Mesh->MaxNumStreamedLODs.Default;
+		if (const int32* Found = Mesh->MaxNumStreamedLODs.PerPlatform.Find("Mobile"))
+		{
+			MaxNumStreamedLODs = *Found;
+		}
+
+		MaxNumOptionalLODs = Mesh->MaxNumOptionalLODs.Default;
+		if (const int32* Found = Mesh->MaxNumOptionalLODs.PerPlatform.Find("Mobile"))
+		{
+			MaxNumOptionalLODs = *Found;
+		}
 #endif
 
 		int32 VertexCountLod0 = 0;
@@ -5383,6 +5411,9 @@ bool UEngine::HandleListSkeletalMeshesCommand(const TCHAR* Cmd, FOutputDevice& A
 			LodCount,
 			MobileMinLOD,
 			MeshDisablesMinLODStripping,
+			bSupportLODStreaming,
+			MaxNumStreamedLODs,
+			MaxNumOptionalLODs,
 			VertexCountLod0,
 			VertexCountLod1,
 			VertexCountLod2,
@@ -5406,11 +5437,11 @@ bool UEngine::HandleListSkeletalMeshesCommand(const TCHAR* Cmd, FOutputDevice& A
 
 	if (bCSV)
 	{
-		Ar.Logf(TEXT(",NumKB, MaxKB, ResKBExc, ResKBInc, ResKBIncMobile, LOD Count, MobileMinLOD, DisableMinLODStripping, Verts LOD0, Verts LOD1, Verts LOD2, Verts Total, Verts Total Mobile, Verts Collision, Name"));
+		Ar.Logf(TEXT(",NumKB, MaxKB, ResKBExc, ResKBInc, ResKBIncMobile, LOD Count, MobileMinLOD, DisableMinLODStripping, bSupportLODStreaming, MaxNumStreamedLODs, MaxNumOptionalLODs, Verts LOD0, Verts LOD1, Verts LOD2, Verts Total, Verts Total Mobile, Verts Collision, Name"));
 	}
 	else
 	{
-		Ar.Logf(TEXT("       NumKB       MaxKB    ResKBExc    ResKBInc  ResKBIncMobile    NumLODs   MobileMinLOD  DisableMinLODStripping  VertsLOD0   VertsLOD1   VertsLOD2   VertsTotal  VertsTotalMobile  VertsColl  Name"));
+		Ar.Logf(TEXT("       NumKB       MaxKB    ResKBExc    ResKBInc  ResKBIncMobile    NumLODs   MobileMinLOD  DisableMinLODStripping  bSupportLODStreaming  MaxNumStreamedLODs  MaxNumOptionalLODs  VertsLOD0   VertsLOD1   VertsLOD2   VertsTotal  VertsTotalMobile  VertsColl  Name"));
 	}
 
 	for (int32 MeshIndex = 0; MeshIndex<SortedMeshes.Num(); MeshIndex++)
@@ -5419,15 +5450,17 @@ bool UEngine::HandleListSkeletalMeshesCommand(const TCHAR* Cmd, FOutputDevice& A
 
 		if (bCSV)
 		{
-			Ar.Logf(TEXT(",%i, %i, %i, %i, %i, %i, %i, %i, %i, %i, %i, %i, %i, %i, %s"),
-				SortedMesh.NumKB, SortedMesh.MaxKB, SortedMesh.ResKBExc, SortedMesh.ResKBInc, SortedMesh.ResKBIncMobile, SortedMesh.LodCount, SortedMesh.MobileMinLOD, SortedMesh.MeshDisablesMinLODStripping,
+			Ar.Logf(TEXT(",%i, %i, %i, %i, %i, %i, %i, %i, %i, %i, %i, %i, %i, %i, %i, %i, %i, %s"),
+				SortedMesh.NumKB, SortedMesh.MaxKB, SortedMesh.ResKBExc, SortedMesh.ResKBInc, SortedMesh.ResKBIncMobile, SortedMesh.LodCount, SortedMesh.MobileMinLOD,
+				SortedMesh.MeshDisablesMinLODStripping, SortedMesh.bSupportLODStreaming, SortedMesh.MaxNumStreamedLODs, SortedMesh.MaxNumOptionalLODs,
 				SortedMesh.VertexCountLod0, SortedMesh.VertexCountLod1, SortedMesh.VertexCountLod2, SortedMesh.VertexCountTotal, SortedMesh.VertexCountTotalMobile, SortedMesh.VertexCountCollision,
 				*SortedMesh.Name);
 		}
 		else
 		{
-			Ar.Logf(TEXT("%9i KB %8i KB %8i KB %8i KB  %11i KB %10i %14i %23i %10i %11i %11i %12i %17i %10i  %s"),
-				SortedMesh.NumKB, SortedMesh.MaxKB, SortedMesh.ResKBExc, SortedMesh.ResKBInc, SortedMesh.ResKBIncMobile, SortedMesh.LodCount, SortedMesh.MobileMinLOD, SortedMesh.MeshDisablesMinLODStripping,
+			Ar.Logf(TEXT("%9i KB %8i KB %8i KB %8i KB  %11i KB %10i %14i %23i %21i %19i %19i %10i %11i %11i %12i %17i %10i  %s"),
+				SortedMesh.NumKB, SortedMesh.MaxKB, SortedMesh.ResKBExc, SortedMesh.ResKBInc, SortedMesh.ResKBIncMobile, SortedMesh.LodCount, SortedMesh.MobileMinLOD,
+				SortedMesh.MeshDisablesMinLODStripping, SortedMesh.bSupportLODStreaming, SortedMesh.MaxNumStreamedLODs, SortedMesh.MaxNumOptionalLODs,
 				SortedMesh.VertexCountLod0, SortedMesh.VertexCountLod1, SortedMesh.VertexCountLod2, SortedMesh.VertexCountTotal, SortedMesh.VertexCountTotalMobile, SortedMesh.VertexCountCollision,
 				*SortedMesh.Name);
 		}
diff --git a/Engine/Source/Runtime/Engine/Public/Animation/SkinWeightProfile.h b/Engine/Source/Runtime/Engine/Public/Animation/SkinWeightProfile.h
--- a/Engine/Source/Runtime/Engine/Public/Animation/SkinWeightProfile.h
+++ b/Engine/Source/Runtime/Engine/Public/Animation/SkinWeightProfile.h
@@ -91,36 +91,43 @@ struct FRuntimeSkinWeightProfileData
 	void ApplyOverrides(FSkinWeightVertexBuffer* OverrideBuffer, const FSkinWeightVertexBuffer* BaseBuffer) const
 	{
 		const TSkinWeightInfo<bExtraBoneInfluences>* SkinWeightInfoPtr = BaseBuffer->GetSkinWeightPtr<bExtraBoneInfluences>(0);
-		TArray<TSkinWeightInfo<bExtraBoneInfluences>> OverrideArray;
-
-		const int32 ExpectedNumVerts = BaseBuffer->GetNumVertices();
-		OverrideArray.SetNumUninitialized(ExpectedNumVerts);
-		FMemory::Memcpy(OverrideArray.GetData(), SkinWeightInfoPtr, sizeof(TSkinWeightInfo<bExtraBoneInfluences>) * ExpectedNumVerts);
-
-		// Apply overrides
+		
+		if (SkinWeightInfoPtr)
 		{
-			for (auto VertexIndexOverridePair : VertexIndexOverrideIndex)
+			TArray<TSkinWeightInfo<bExtraBoneInfluences>> OverrideArray;
+			const int32 ExpectedNumVerts = BaseBuffer->GetNumVertices();
+			OverrideArray.SetNumUninitialized(ExpectedNumVerts);
+			FMemory::Memcpy(OverrideArray.GetData(), SkinWeightInfoPtr, sizeof(TSkinWeightInfo<bExtraBoneInfluences>) * ExpectedNumVerts);
+
+			// Apply overrides
 			{
-				const uint32 VertexIndex = VertexIndexOverridePair.Key;
-				TSkinWeightInfo<bExtraBoneInfluences>& Entry = OverrideArray[VertexIndex];
+				for (auto VertexIndexOverridePair : VertexIndexOverrideIndex)
+				{
+					const uint32 VertexIndex = VertexIndexOverridePair.Key;
+					TSkinWeightInfo<bExtraBoneInfluences>& Entry = OverrideArray[VertexIndex];
 
-				const uint32 OverrideIndex = VertexIndexOverridePair.Value;
-				const FSkinWeightOverrideInfo& OverrideInfo = OverridesInfo[OverrideIndex];
+					const uint32 OverrideIndex = VertexIndexOverridePair.Value;
+					const FSkinWeightOverrideInfo& OverrideInfo = OverridesInfo[OverrideIndex];
 
-				FMemory::Memzero(Entry.InfluenceBones);
-				FMemory::Memzero(Entry.InfluenceWeights);
+					FMemory::Memzero(Entry.InfluenceBones);
+					FMemory::Memzero(Entry.InfluenceWeights);
 
-				for (int32 Index = 0; Index < OverrideInfo.NumInfluences; ++Index)
-				{
-					const uint16 WeightData = Weights[OverrideInfo.InfluencesOffset + Index];
+					for (int32 Index = 0; Index < OverrideInfo.NumInfluences; ++Index)
+					{
+						const uint16 WeightData = Weights[OverrideInfo.InfluencesOffset + Index];
 
-					Entry.InfluenceBones[Index] = (WeightData) >> 8;
-					Entry.InfluenceWeights[Index] = (WeightData & 0xFF);
+						Entry.InfluenceBones[Index] = (WeightData) >> 8;
+						Entry.InfluenceWeights[Index] = (WeightData & 0xFF);
+					}
 				}
 			}
-		}
 
-		(*OverrideBuffer) = OverrideArray;
+			(*OverrideBuffer) = OverrideArray;
+		}
+		else
+		{
+			OverrideBuffer->CopyMetaData(*BaseBuffer);
+		}
 	}
 
 	/** Per skin weight offset into Weights array and number of weights stored */
@@ -162,7 +169,37 @@ struct ENGINE_API FSkinWeightProfilesData
 
 	friend FArchive& operator<<(FArchive& Ar, FSkinWeightProfilesData& OverrideData);
 
+	void SerializeMetaData(FArchive& Ar);
+
+	void ReleaseCPUResources();
+
+	void CreateRHIBuffers_RenderThread(TArray<TPair<FName, FVertexBufferRHIRef>>& OutBuffers);
+	void CreateRHIBuffers_Async(TArray<TPair<FName, FVertexBufferRHIRef>>& OutBuffers);
+
+	template <uint32 MaxNumUpdates>
+	void InitRHIForStreaming(const TArray<TPair<FName, FVertexBufferRHIRef>>& IntermediateBuffers, TRHIResourceUpdateBatcher<MaxNumUpdates>& Batcher)
+	{
+		for (int32 Idx = 0; Idx < IntermediateBuffers.Num(); ++Idx)
+		{
+			const FName& ProfileName = IntermediateBuffers[Idx].Key;
+			FVertexBufferRHIParamRef IntermediateBuffer = IntermediateBuffers[Idx].Value;
+			ProfileNameToBuffer.FindChecked(ProfileName)->InitRHIForStreaming(IntermediateBuffer, Batcher);
+		}
+	}
+	
+	template <uint32 MaxNumUpdates>
+	void ReleaseRHIForStreaming(TRHIResourceUpdateBatcher<MaxNumUpdates>& Batcher)
+	{
+		for (TMap<FName, FSkinWeightVertexBuffer*>::TIterator It(ProfileNameToBuffer); It; ++It)
+		{
+			It->Value->ReleaseRHIForStreaming(Batcher);
+		}
+	}
+
 protected:
+	template <bool bRenderThread>
+	void CreateRHIBuffers_Internal(TArray<TPair<FName, FVertexBufferRHIRef>>& OutBuffers);
+
 	FSkinWeightVertexBuffer* BaseBuffer;
 	FSkinWeightVertexBuffer* DefaultOverrideSkinWeightBuffer;
 
diff --git a/Engine/Source/Runtime/Engine/Public/ContentStreaming.h b/Engine/Source/Runtime/Engine/Public/ContentStreaming.h
--- a/Engine/Source/Runtime/Engine/Public/ContentStreaming.h
+++ b/Engine/Source/Runtime/Engine/Public/ContentStreaming.h
@@ -436,6 +436,12 @@ struct IRenderAssetStreamingManager : public IStreamingManager
 	/** Return all bounds related to the ref object */
 	ENGINE_API virtual void GetObjectReferenceBounds(const UObject* RefObject, TArray<FBox>& AssetBoxes) = 0;
 
+	/** Return all components referencing the asset */
+	ENGINE_API virtual void GetAssetComponents(
+		const UStreamableRenderAsset* RenderAsset,
+		TArray<const UPrimitiveComponent*>& OutComps,
+		TFunction<bool(const UPrimitiveComponent*)> ShouldChoose = [](const UPrimitiveComponent*) { return true; }) = 0;
+
 	//BEGIN: APIs for backward compatibility
 	ENGINE_API void UpdateIndividualTexture(UTexture2D* Texture);
 	ENGINE_API bool StreamOutTextureData(int64 RequiredMemorySize);
@@ -625,13 +631,23 @@ struct ENGINE_VTABLE FStreamingManagerCollection : public IStreamingManager
 	/**
 	 * Checks whether texture streaming is active
 	 */
-	virtual bool IsTextureStreamingEnabled() const;
+	ENGINE_API bool IsTextureStreamingEnabled() const;
+
+	/**
+	 * Checks whether texture/mesh streaming is active
+	 */
+	virtual bool IsRenderAssetStreamingEnabled() const;
 
 	/**
 	 * Gets a reference to the Texture Streaming Manager interface
 	 */
 	ENGINE_API IRenderAssetStreamingManager& GetTextureStreamingManager() const;
 
+	/**
+	 * Get the streaming manager for textures and meshes
+	 */
+	ENGINE_API IRenderAssetStreamingManager& GetRenderAssetStreamingManager() const;
+
 	/**
 	 * Gets a reference to the Audio Streaming Manager interface
 	 */
diff --git a/Engine/Source/Runtime/Engine/Public/RawIndexBuffer.h b/Engine/Source/Runtime/Engine/Public/RawIndexBuffer.h
--- a/Engine/Source/Runtime/Engine/Public/RawIndexBuffer.h
+++ b/Engine/Source/Runtime/Engine/Public/RawIndexBuffer.h
@@ -307,6 +307,8 @@ class FRawStaticIndexBuffer16or32Interface : public FIndexBuffer
 public:
 	virtual void Serialize( FArchive& Ar ) = 0;
 
+	virtual void SerializeMetaData(FArchive& Ar) = 0;
+
 	/**
 	 * The following methods are basically just accessors that allow us
 	 * to hide the implementation of FRawStaticIndexBuffer16or32 by making
@@ -343,7 +345,8 @@ public:
 	* @param InNeedsCPUAccess - true if resource array data should be CPU accessible
 	*/
 	FRawStaticIndexBuffer16or32(bool InNeedsCPUAccess=false)
-	:	Indices(InNeedsCPUAccess)
+		: Indices(InNeedsCPUAccess)
+		, CachedNumIndices(0)
 	{
 	}
 
@@ -352,33 +355,11 @@ public:
 	*/
 	virtual void InitRHI() override
 	{
-		uint32 Size = Indices.Num() * sizeof(INDEX_TYPE);
-		if (Indices.Num())
-		{
-			// Create the index buffer.
-			FRHIResourceCreateInfo CreateInfo(&Indices);
-
-			extern ENGINE_API bool DoSkeletalMeshIndexBuffersNeedSRV();
-			bool bSRV = DoSkeletalMeshIndexBuffersNeedSRV();
-			// When bAllowCPUAccess is true, the meshes is likely going to be used for Niagara to spawn particles on mesh surface.
-			// And it can be the case for CPU *and* GPU access: no differenciation today. That is why we create a SRV in this case.
-			// This also avoid setting lots of states on all the members of all the different buffers used by meshes. Follow up: https://jira.it.epicgames.net/browse/UE-69376.
-			bSRV |= Indices.GetAllowCPUAccess();
+		IndexBufferRHI = CreateRHIBuffer_RenderThread();
 
-			EBufferUsageFlags Flags = BUF_Static;
-
-			if (bSRV)
-			{
-				// BUF_ShaderResource is needed for SkinCache RecomputeSkinTangents
-				Flags = (EBufferUsageFlags)(Flags | BUF_ShaderResource);
-			}
-
-			IndexBufferRHI = RHICreateIndexBuffer(sizeof(INDEX_TYPE), Size, Flags, CreateInfo);
-
-			if (bSRV)
-			{
-				SRVValue = RHICreateShaderResourceView(IndexBufferRHI);
-			}
+		if (IndexBufferRHI && IsSRVNeeded())
+		{
+			SRVValue = RHICreateShaderResourceView(Indices.Num() ? IndexBufferRHI : nullptr);
 		}
 	}
 	
@@ -396,7 +377,14 @@ public:
 	*/
 	virtual void Serialize( FArchive& Ar ) override
 	{
+		SerializeMetaData(Ar);
 		Indices.BulkSerialize(Ar);
+		check(CachedNumIndices == Indices.Num());
+	}
+
+	virtual void SerializeMetaData(FArchive& Ar) override
+	{
+		Ar << CachedNumIndices;
 	}
 
 	/**
@@ -413,11 +401,12 @@ public:
 
 	virtual int32 Num() const override
 	{
-		return Indices.Num();
+		return CachedNumIndices;
 	}
 
 	virtual int32 AddItem(uint32 Val) override
 	{
+		++CachedNumIndices;
 		return Indices.Add(Val);
 	}
 
@@ -433,17 +422,22 @@ public:
 
 	virtual void Insert(int32 Idx, int32 Num) override
 	{
+		CachedNumIndices += Num;
 		Indices.InsertUninitialized(Idx, Num);
+		check(CachedNumIndices == Indices.Num());
 	}
 
 	virtual void Remove(int32 Idx, int32 Num) override
 	{
+		CachedNumIndices -= Num;
 		Indices.RemoveAt(Idx, Num);
+		check(CachedNumIndices == Indices.Num());
 	}
 
 	virtual void Empty(int32 Slack) override
 	{
 		Indices.Empty(Slack);
+		CachedNumIndices = 0;
 	}
 
 	virtual int32 GetResourceDataSize() const override
@@ -454,8 +448,86 @@ public:
 	virtual void AssignNewBuffer(const TArray<INDEX_TYPE>& Buffer)
 	{
 		Indices = TArray<INDEX_TYPE,TAlignedHeapAllocator<INDEXBUFFER_ALIGNMENT> >(Buffer);
+		CachedNumIndices = Indices.Num();
+	}
+
+	/** Create an RHI index buffer with CPU data. CPU data may be discarded after creation (see TResourceArray::Discard) */
+	FIndexBufferRHIRef CreateRHIBuffer_RenderThread() { return CreateRHIBuffer_Internal<true>(); }
+	FIndexBufferRHIRef CreateRHIBuffer_Async() { return CreateRHIBuffer_Internal<false>(); }
+
+	/** Similar to Init/ReleaseRHI but only update existing SRV so references to the SRV stays valid */
+	template <uint32 MaxNumUpdates>
+	void InitRHIForStreaming(FIndexBufferRHIParamRef IntermediateBuffer, TRHIResourceUpdateBatcher<MaxNumUpdates>& Batcher)
+	{
+		if (IndexBufferRHI && IntermediateBuffer)
+		{
+			Batcher.QueueUpdateRequest(IndexBufferRHI, IntermediateBuffer);
+			if (SRVValue)
+			{
+				Batcher.QueueUpdateRequest(SRVValue, IndexBufferRHI);
+			}
+		}
+	}
+
+	template <uint32 MaxNumUpdates>
+	void ReleaseRHIForStreaming(TRHIResourceUpdateBatcher<MaxNumUpdates>& Batcher)
+	{
+		if (IndexBufferRHI)
+		{
+			Batcher.QueueUpdateRequest(IndexBufferRHI, nullptr);
+		}
+		if (SRVValue)
+		{
+			Batcher.QueueUpdateRequest(SRVValue, nullptr);
+		}
 	}
 
 private:
 	TResourceArray<INDEX_TYPE,INDEXBUFFER_ALIGNMENT> Indices;
+
+	int32 CachedNumIndices;
+
+	bool IsSRVNeeded() const
+	{
+		extern ENGINE_API bool DoSkeletalMeshIndexBuffersNeedSRV();
+		bool bSRV = DoSkeletalMeshIndexBuffersNeedSRV();
+		// When bAllowCPUAccess is true, the meshes is likely going to be used for Niagara to spawn particles on mesh surface.
+		// And it can be the case for CPU *and* GPU access: no differenciation today. That is why we create a SRV in this case.
+		// This also avoid setting lots of states on all the members of all the different buffers used by meshes. Follow up: https://jira.it.epicgames.net/browse/UE-69376.
+		bSRV |= Indices.GetAllowCPUAccess();
+		return bSRV;
+	}
+
+	template <bool bRenderThread>
+	FIndexBufferRHIRef CreateRHIBuffer_Internal()
+	{
+		if (CachedNumIndices)
+		{
+			// Create the index buffer.
+			FRHIResourceCreateInfo CreateInfo(&Indices);
+			EBufferUsageFlags Flags = BUF_Static;
+
+			if (IsSRVNeeded())
+			{
+				// BUF_ShaderResource is needed for SkinCache RecomputeSkinTangents
+				Flags = (EBufferUsageFlags)(Flags | BUF_ShaderResource);
+			}
+
+			FIndexBufferRHIRef Ret;
+			const uint32 Size = Indices.Num() * sizeof(INDEX_TYPE);
+			CreateInfo.bWithoutNativeResource = !Size;
+			if (bRenderThread)
+			{
+				Ret = RHICreateIndexBuffer(sizeof(INDEX_TYPE), Size, Flags, CreateInfo);
+			}
+			else
+			{
+				Ret = RHIAsyncCreateIndexBuffer(sizeof(INDEX_TYPE), Size, Flags, CreateInfo);
+			}
+
+			CachedNumIndices = Indices.Num();
+			return Ret;
+		}
+		return nullptr;
+	}
 };
diff --git a/Engine/Source/Runtime/Engine/Public/RenderAssetUpdate.h b/Engine/Source/Runtime/Engine/Public/RenderAssetUpdate.h
--- a/Engine/Source/Runtime/Engine/Public/RenderAssetUpdate.h
+++ b/Engine/Source/Runtime/Engine/Public/RenderAssetUpdate.h
@@ -35,6 +35,7 @@ public:
 		TT_None,	// No thread.
 		TT_Render,	// The render thread.
 		TT_Async,	// An async work thread.
+		TT_GameThread, // The game thread
 		TT_GameRunningAsync	// The game thread running async work.
 	};
 
@@ -131,9 +132,11 @@ protected:
 		bIsCancelled = true;
 	}
 
+	void ScheduleGTTask();
 	void ScheduleRenderTask();
 	void ScheduleAsyncTask();
 
+	friend class FRenderAssetUpdateTickGTTask;
 
 	/** An async task used to call tick on the pending update. */
 	class FMipUpdateTask : public FNonAbandonableTask
@@ -163,6 +166,9 @@ protected:
 	/** Critical Section. */
 	FCriticalSection CS;
 
+	/** Number of Ticks scheduled on the game thread. */
+	int32 ScheduledGTTasks;
+
 	/** Number of Ticks scheduled on the render thread. */
 	int32 ScheduledRenderTasks;
 
@@ -178,6 +184,9 @@ protected:
 	/** Whether the task has been cancelled because the update could not proceed or because the user called Abort(). */
 	bool bIsCancelled;
 
+	/** Defer execution even if a task pushes a new task on the same thread. */
+	bool bDeferExecution;
+
 	/** The state of the work yet to be performed to complete the update or cancelation. */
 	volatile ETaskState TaskState;
 
diff --git a/Engine/Source/Runtime/Engine/Public/Rendering/MultiSizeIndexContainer.h b/Engine/Source/Runtime/Engine/Public/Rendering/MultiSizeIndexContainer.h
--- a/Engine/Source/Runtime/Engine/Public/Rendering/MultiSizeIndexContainer.h
+++ b/Engine/Source/Runtime/Engine/Public/Rendering/MultiSizeIndexContainer.h
@@ -43,6 +43,8 @@ public:
 	*/
 	void Serialize(FArchive& Ar, bool bNeedsCPUAccess);
 
+	void SerializeMetaData(FArchive& Ar, bool bNeedsCPUAccess);
+
 	/**
 	* Creates a new index buffer
 	*/
@@ -92,6 +94,42 @@ public:
 
 	friend FArchive& operator<<(FArchive& Ar, FMultiSizeIndexContainer& Buffer);
 
+	FIndexBufferRHIRef CreateRHIBuffer_RenderThread();
+	FIndexBufferRHIRef CreateRHIBuffer_Async();
+
+	template <uint32 MaxNumUpdates>
+	void InitRHIForStreaming(FIndexBufferRHIParamRef IntermediateBuffer, TRHIResourceUpdateBatcher<MaxNumUpdates>& Batcher)
+	{
+		check(!((uint32)!!IntermediateBuffer ^ (uint32)!!IndexBuffer));
+		if (IntermediateBuffer)
+		{
+			if (DataTypeSize == sizeof(uint16))
+			{
+				static_cast<FRawStaticIndexBuffer16or32<uint16>*>(IndexBuffer)->InitRHIForStreaming(IntermediateBuffer, Batcher);
+			}
+			else
+			{
+				static_cast<FRawStaticIndexBuffer16or32<uint32>*>(IndexBuffer)->InitRHIForStreaming(IntermediateBuffer, Batcher);
+			}
+		}
+	}
+
+	template <uint32 MaxNumUpdates>
+	void ReleaseRHIForStreaming(TRHIResourceUpdateBatcher<MaxNumUpdates>& Batcher)
+	{
+		if (IndexBuffer)
+		{
+			if (DataTypeSize == sizeof(uint16))
+			{
+				static_cast<FRawStaticIndexBuffer16or32<uint16>*>(IndexBuffer)->ReleaseRHIForStreaming(Batcher);
+			}
+			else
+			{
+				static_cast<FRawStaticIndexBuffer16or32<uint32>*>(IndexBuffer)->ReleaseRHIForStreaming(Batcher);
+			}
+		}
+	}
+
 private:
 	/** Size of the index buffer's index type (should be 2 or 4 bytes) */
 	uint8 DataTypeSize;
diff --git a/Engine/Source/Runtime/Engine/Public/Rendering/SkeletalMeshLODRenderData.h b/Engine/Source/Runtime/Engine/Public/Rendering/SkeletalMeshLODRenderData.h
--- a/Engine/Source/Runtime/Engine/Public/Rendering/SkeletalMeshLODRenderData.h
+++ b/Engine/Source/Runtime/Engine/Public/Rendering/SkeletalMeshLODRenderData.h
@@ -128,12 +128,28 @@ public:
 
 	TArray<FBoneIndexType> RequiredBones;
 
+	uint32 BuffersSize;
+
+	/** Offset in the .bulk file if this LOD is streamed */
+	uint32 OffsetInFile;
+	/** Size of serialized resource data in bytes */
+	uint32 BulkDataSize;
+
+	/** Whether buffers of this LOD is inlined (i.e. stored in .uexp instead of .ubulk) */
+	uint32 bStreamedDataInlined : 1;
+	/** Whether this LOD is below MinLod */
+	uint32 bIsLODOptional : 1;
+
+#if WITH_EDITOR
+	FByteBulkData BulkData;
+#endif
+
 	/**
 	* Initialize the LOD's render resources.
 	*
 	* @param Parent Parent mesh
 	*/
-	void InitResources(bool bNeedsVertexColors, int32 LODIndex, TArray<class UMorphTarget*>& InMorphTargets);
+	void InitResources(bool bNeedsVertexColors, int32 LODIndex, TArray<class UMorphTarget*>& InMorphTargets, USkeletalMesh* Owner);
 
 	/**
 	* Releases the LOD's render resources.
@@ -143,10 +159,15 @@ public:
 	/**
 	* Releases the LOD's CPU render resources.
 	*/
-	void ReleaseCPUResources();
+	void ReleaseCPUResources(bool bForStreaming = false);
 
 	/** Constructor (default) */
 	FSkeletalMeshLODRenderData()
+		: BuffersSize(0)
+		, OffsetInFile(0xffffffff)
+		, BulkDataSize(0)
+		, bStreamedDataInlined(false)
+		, bIsLODOptional(false)
 	{
 	}
 
@@ -160,6 +181,15 @@ public:
 	*/
 	void Serialize(FArchive& Ar, UObject* Owner, int32 Idx);
 
+	/**
+	* Serialize the portion of data that might be streamed
+	* @param bNeedsCPUAcces - Whether something requires keeping CPU copy of resources (see ShouldKeepCPUResources)
+	* @param bForceKeepCPUResources - Whether we want to force keeping CPU resources
+	* @return Size of streamed LOD data in bytes
+	*/
+	void SerializeStreamedData(FArchive& Ar, USkeletalMesh* Owner, int32 LODIdx, uint8 ClassStripFlags, bool bNeedsCPUAccess, bool bForceKeepCPUResources);
+
+	void SerializeAvailabilityInfo(FArchive& Ar, USkeletalMesh* Owner, int32 LODIdx, bool bHasAdjacencyData, bool bNeedsCPUAccess);
 
 #if WITH_EDITOR
 	/**
@@ -219,4 +249,29 @@ public:
 	uint32 FindSectionIndex(const FSkelMeshRenderSection& Section) const;
 
 	ENGINE_API int32 NumNonClothingSections() const;
+
+	void IncrementMemoryStats(bool bNeedsVertexColors);
+	void DecrementMemoryStats();
+
+	static bool ShouldForceKeepCPUResources();
+	static bool ShouldKeepCPUResources(const USkeletalMesh* SkeletalMesh, int32 LODIdx, bool bForceKeep);
+
+private:
+	enum EClassDataStripFlag : uint8
+	{
+		CDSF_AdjacencyData = 1,
+		CDSF_MinLodData = 2
+	};
+
+	static int32 GetPlatformMinLODIdx(const ITargetPlatform* TargetPlatform, const USkeletalMesh* SkeletalMesh);
+
+	static uint8 GenerateClassStripFlags(FArchive& Ar, const USkeletalMesh* OwnerMesh, int32 LODIdx);
+
+	static bool IsLODCookedOut(const ITargetPlatform* TargetPlatform, const USkeletalMesh* SkeletalMesh, bool bIsBelowMinLOD);
+
+	static bool IsLODInlined(const ITargetPlatform* TargetPlatform, const USkeletalMesh* SkeletalMesh, int32 LODIdx, bool bIsBelowMinLOD);
+
+	static int32 GetNumOptionalLODsAllowed(const ITargetPlatform* TargetPlatform, const USkeletalMesh* SkeletalMesh);
+
+	friend class FSkeletalMeshRenderData;
 };
\ No newline at end of file
diff --git a/Engine/Source/Runtime/Engine/Public/Rendering/SkeletalMeshRenderData.h b/Engine/Source/Runtime/Engine/Public/Rendering/SkeletalMeshRenderData.h
--- a/Engine/Source/Runtime/Engine/Public/Rendering/SkeletalMeshRenderData.h
+++ b/Engine/Source/Runtime/Engine/Public/Rendering/SkeletalMeshRenderData.h
@@ -18,11 +18,28 @@ public:
 	/** Per-LOD render data. */
 	TIndirectArray<FSkeletalMeshLODRenderData> LODRenderData;
 
+	/** True if rhi resources are initialized */
+	bool bReadyForStreaming;
+
+	/** Const after serialization. */
+	uint8 NumInlinedLODs;
+
+	/** Const after serialization. */
+	uint8 NumOptionalLODs;
+
+	/** [RenderThread] Index of the most detailed valid LOD. */
+	uint8 CurrentFirstLODIdx;
+
+	/** [GameThread/RenderThread] Future value of CurrentFirstLODIdx. */
+	uint8 PendingFirstLODIdx;
+
 #if WITH_EDITORONLY_DATA
 	/** UV data used for streaming accuracy debug view modes. In sync for rendering thread */
 	TArray<FMeshUVChannelInfo> UVChannelDataPerMaterial;
 #endif
 
+	FSkeletalMeshRenderData();
+
 #if WITH_EDITOR
 	void Cache(USkeletalMesh* Owner);
 
@@ -33,7 +50,7 @@ public:
 	void Serialize(FArchive& Ar, USkeletalMesh* Owner);
 
 	/** Initializes rendering resources. */
-	void InitResources(bool bNeedsVertexColors, TArray<UMorphTarget*>& InMorphTargets);
+	void InitResources(bool bNeedsVertexColors, TArray<UMorphTarget*>& InMorphTargets, USkeletalMesh* Owner);
 
 	/** Releases rendering resources. */
 	ENGINE_API void ReleaseResources();
diff --git a/Engine/Source/Runtime/Engine/Public/Rendering/SkeletalMeshVertexClothBuffer.h b/Engine/Source/Runtime/Engine/Public/Rendering/SkeletalMeshVertexClothBuffer.h
--- a/Engine/Source/Runtime/Engine/Public/Rendering/SkeletalMeshVertexClothBuffer.h
+++ b/Engine/Source/Runtime/Engine/Public/Rendering/SkeletalMeshVertexClothBuffer.h
@@ -39,6 +39,8 @@ public:
 	*/
 	void CleanUp();
 
+	void ClearMetaData();
+
 	/**
 	* Initializes the buffer with the given vertices.
 	* @param InVertices - The vertices to initialize the buffer with.
@@ -54,6 +56,8 @@ public:
 	*/
 	friend FArchive& operator<<(FArchive& Ar, FSkeletalMeshVertexClothBuffer& VertexBuffer);
 
+	void SerializeMetaData(FArchive& Ar);
+
 	//~ Begin FRenderResource interface.
 
 	/**
@@ -119,6 +123,35 @@ public:
 		return ClothIndexMapping;
 	}
 
+	/** Create an RHI vertex buffer with CPU data. CPU data may be discarded after creation (see TResourceArray::Discard) */
+	FVertexBufferRHIRef CreateRHIBuffer_RenderThread();
+	FVertexBufferRHIRef CreateRHIBuffer_Async();
+
+	/** Similar to Init/ReleaseRHI but only update existing SRV so references to the SRV stays valid */
+	template <uint32 MaxNumUpdates>
+	void InitRHIForStreaming(FVertexBufferRHIParamRef IntermediateBuffer, TRHIResourceUpdateBatcher<MaxNumUpdates>& Batcher)
+	{
+		if (VertexBufferRHI && IntermediateBuffer)
+		{
+			check(VertexBufferSRV);
+			Batcher.QueueUpdateRequest(VertexBufferRHI, IntermediateBuffer);
+			Batcher.QueueUpdateRequest(VertexBufferSRV, VertexBufferRHI, sizeof(FVector4), PF_A32B32G32R32F);
+		}
+	}
+
+	template <uint32 MaxNumUpdates>
+	void ReleaseRHIForStreaming(TRHIResourceUpdateBatcher<MaxNumUpdates>& Batcher)
+	{
+		if (VertexBufferRHI)
+		{
+			Batcher.QueueUpdateRequest(VertexBufferRHI, nullptr);
+		}
+		if (VertexBufferSRV)
+		{
+			Batcher.QueueUpdateRequest(VertexBufferSRV, nullptr, 0, 0);
+		}
+	}
+
 private:
 	FShaderResourceViewRHIRef VertexBufferSRV;
 
@@ -138,4 +171,7 @@ private:
 	* Allocates the vertex data storage type
 	*/
 	void AllocateData();
+
+	template <bool bRenderThread>
+	FVertexBufferRHIRef CreateRHIBuffer_Internal();
 };
\ No newline at end of file
diff --git a/Engine/Source/Runtime/Engine/Public/Rendering/SkinWeightVertexBuffer.h b/Engine/Source/Runtime/Engine/Public/Rendering/SkinWeightVertexBuffer.h
--- a/Engine/Source/Runtime/Engine/Public/Rendering/SkinWeightVertexBuffer.h
+++ b/Engine/Source/Runtime/Engine/Public/Rendering/SkinWeightVertexBuffer.h
@@ -95,6 +95,10 @@ public:
 
 	friend FArchive& operator<<(FArchive& Ar, FSkinWeightVertexBuffer& VertexBuffer);
 
+	void SerializeMetaData(FArchive& Ar);
+
+	void CopyMetaData(const FSkinWeightVertexBuffer& Other);
+
 	// FRenderResource interface.
 	virtual void InitRHI() override;
 	virtual void ReleaseRHI() override;
@@ -182,6 +186,35 @@ public:
 		return *this;
 	}
 
+	/** Create an RHI vertex buffer with CPU data. CPU data may be discarded after creation (see TResourceArray::Discard) */
+	FVertexBufferRHIRef CreateRHIBuffer_RenderThread();
+	FVertexBufferRHIRef CreateRHIBuffer_Async();
+
+	/** Similar to Init/ReleaseRHI but only update existing SRV so references to the SRV stays valid */
+	template <uint32 MaxNumUpdates>
+	void InitRHIForStreaming(FVertexBufferRHIParamRef IntermediateBuffer, TRHIResourceUpdateBatcher<MaxNumUpdates>& Batcher)
+	{
+		if (VertexBufferRHI && IntermediateBuffer)
+		{
+			check(SRVValue);
+			Batcher.QueueUpdateRequest(VertexBufferRHI, IntermediateBuffer);
+			Batcher.QueueUpdateRequest(SRVValue, VertexBufferRHI, 4, PF_R32_UINT);
+		}
+	}
+
+	template <uint32 MaxNumUpdates>
+	void ReleaseRHIForStreaming(TRHIResourceUpdateBatcher<MaxNumUpdates>& Batcher)
+	{
+		if (VertexBufferRHI)
+		{
+			Batcher.QueueUpdateRequest(VertexBufferRHI, nullptr);
+		}
+		if (SRVValue)
+		{
+			Batcher.QueueUpdateRequest(SRVValue, nullptr, 0, 0);
+		}
+	}
+
 protected:
 	// guaranteed only to be valid if the vertex buffer is valid
 	FShaderResourceViewRHIRef SRVValue;
@@ -214,4 +247,7 @@ private:
 	template <bool bUsesExtraBoneInfluences>
 	void SetWeightsForVertex(uint32 VertexIndex, const FSoftSkinVertex& SrcVertex);
 #endif //WITH_EDITOR
+
+	template <bool bRenderThread>
+	FVertexBufferRHIRef CreateRHIBuffer_Internal();
 };
\ No newline at end of file
diff --git a/Engine/Source/Runtime/Engine/Public/SkeletalMeshTypes.h b/Engine/Source/Runtime/Engine/Public/SkeletalMeshTypes.h
--- a/Engine/Source/Runtime/Engine/Public/SkeletalMeshTypes.h
+++ b/Engine/Source/Runtime/Engine/Public/SkeletalMeshTypes.h
@@ -293,6 +293,11 @@ public:
 
 	virtual void OnTransformChanged() override;
 
+	virtual uint8 GetCurrentFirstLODIdx_RenderThread() const final override
+	{
+		return GetCurrentFirstLODIdx_Internal();
+	}
+
 	const TArray<FMatrix>& GetMeshObjectReferenceToLocalMatrices() const;
 	const TIndirectArray<FSkeletalMeshLODRenderData>& GetSkeletalMeshRenderDataLOD() const;
 
@@ -381,6 +386,8 @@ protected:
 
 	void GetMeshElementsConditionallySelectable(const TArray<const FSceneView*>& Views, const FSceneViewFamily& ViewFamily, bool bInSelectable, uint32 VisibilityMap, FMeshElementCollector& Collector) const;
 
+	/** Only call on render thread timeline */
+	uint8 GetCurrentFirstLODIdx_Internal() const;
 private:
 	void CreateBaseMeshBatch(const FSceneView* View, const FSkeletalMeshLODRenderData& LODData, const int32 LODIndex, const int32 SectionIndex, const FSectionElementInfo& SectionElementInfo, FMeshBatch& Mesh) const;
 };
diff --git a/Engine/Source/Runtime/Engine/Public/SkeletalRenderPublic.h b/Engine/Source/Runtime/Engine/Public/SkeletalRenderPublic.h
--- a/Engine/Source/Runtime/Engine/Public/SkeletalRenderPublic.h
+++ b/Engine/Source/Runtime/Engine/Public/SkeletalRenderPublic.h
@@ -143,7 +143,7 @@ public:
 	 *	This is called from the rendering thread (PreRender) so be very careful what you read/write to.
 	 * @param FrameNumber from ViewFamily.FrameNumber
 	 */
-	void UpdateMinDesiredLODLevel(const FSceneView* View, const FBoxSphereBounds& Bounds, int32 FrameNumber);
+	void UpdateMinDesiredLODLevel(const FSceneView* View, const FBoxSphereBounds& Bounds, int32 FrameNumber, uint8 CurFirstLODIdx);
 
 	/**
 	 *	Return true if this does have valid dynamic data to render
diff --git a/Engine/Source/Runtime/Engine/Public/StaticMeshResources.h b/Engine/Source/Runtime/Engine/Public/StaticMeshResources.h
--- a/Engine/Source/Runtime/Engine/Public/StaticMeshResources.h
+++ b/Engine/Source/Runtime/Engine/Public/StaticMeshResources.h
@@ -115,7 +115,7 @@ private:
 	int32 DefaultNumLODs;
 	/** Maximum number of streamed LODs */
 	int32 DefaultMaxNumStreamedLODs;
-	/** Maximum number of optional LODs */
+	/** Maximum number of optional LODs (currently, need to be either 0 or > max number of LODs below MinLOD) */
 	int32 DefaultMaxNumOptionalLODs;
 	/** Default lightmap resolution. */
 	int32 DefaultLightMapResolution;
@@ -778,7 +778,7 @@ public:
 		const FMaterialRenderProxy* RenderProxy,
 		FMeshBatch& OutMeshBatch) const;
 
-	virtual uint8 GetCurrentFirstLODIdx_RenderThread() const override
+	virtual uint8 GetCurrentFirstLODIdx_RenderThread() const final override
 	{
 		return GetCurrentFirstLODIdx_Internal();
 	}
diff --git a/Engine/Source/Runtime/RHI/Public/RHICommandList.h b/Engine/Source/Runtime/RHI/Public/RHICommandList.h
--- a/Engine/Source/Runtime/RHI/Public/RHICommandList.h
+++ b/Engine/Source/Runtime/RHI/Public/RHICommandList.h
@@ -1763,6 +1763,12 @@ struct FRHIShaderResourceViewUpdateInfo_VB
 	uint8 Format;
 };
 
+struct FRHIShaderResourceViewUpdateInfo_IB
+{
+	FShaderResourceViewRHIParamRef SRV;
+	FIndexBufferRHIParamRef IndexBuffer;
+};
+
 struct FRHIVertexBufferUpdateInfo
 {
 	FRHIVertexBuffer* DestBuffer;
@@ -1797,6 +1803,7 @@ struct FRHIResourceUpdateInfo
 		FRHIVertexBufferUpdateInfo VertexBuffer;
 		FRHIIndexBufferUpdateInfo IndexBuffer;
 		FRHIShaderResourceViewUpdateInfo_VB VertexBufferSRV;
+		FRHIShaderResourceViewUpdateInfo_IB IndexBufferSRV;
 	};
 
 	void ReleaseRefs();
@@ -4991,7 +4998,14 @@ struct TRHIResourceUpdateBatcher
 
 	void QueueUpdateRequest(FRHIShaderResourceView* SRV, FRHIIndexBuffer* IndexBuffer)
 	{
-		// TODO
+		FRHIResourceUpdateInfo& UpdateInfo = GetNextUpdateInfo();
+		UpdateInfo.Type = FRHIResourceUpdateInfo::UT_IndexBufferSRV;
+		UpdateInfo.IndexBufferSRV = { SRV, IndexBuffer };
+		SRV->AddRef();
+		if (IndexBuffer)
+		{
+			IndexBuffer->AddRef();
+		}
 	}
 
 private:
diff --git a/Engine/Source/Runtime/RHI/Public/RHICommandListCommandExecutes.inl b/Engine/Source/Runtime/RHI/Public/RHICommandListCommandExecutes.inl
--- a/Engine/Source/Runtime/RHI/Public/RHICommandListCommandExecutes.inl
+++ b/Engine/Source/Runtime/RHI/Public/RHICommandListCommandExecutes.inl
@@ -633,6 +633,13 @@ void FRHIResourceUpdateInfo::ReleaseRefs()
 			VertexBufferSRV.VertexBuffer->Release();
 		}
 		break;
+	case UT_IndexBufferSRV:
+		IndexBufferSRV.SRV->Release();
+		if (IndexBufferSRV.IndexBuffer)
+		{
+			IndexBufferSRV.IndexBuffer->Release();
+		}
+		break;
 	default:
 		// Unrecognized type, do nothing
 		break;
@@ -675,6 +682,11 @@ void FRHICommandUpdateRHIResources::Execute(FRHICommandListBase& CmdList)
 				Info.VertexBufferSRV.Stride,
 				Info.VertexBufferSRV.Format);
 			break;
+		case FRHIResourceUpdateInfo::UT_IndexBufferSRV:
+			GDynamicRHI->RHIUpdateShaderResourceView(
+				Info.IndexBufferSRV.SRV,
+				Info.IndexBufferSRV.IndexBuffer);
+			break;
 		default:
 			// Unrecognized type, do nothing
 			break;
