Fix non-unity build error

diff --git a/Engine/Source/Runtime/Engine/Private/Animation/SkinWeightProfile.cpp b/Engine/Source/Runtime/Engine/Private/Animation/SkinWeightProfile.cpp
--- a/Engine/Source/Runtime/Engine/Private/Animation/SkinWeightProfile.cpp
+++ b/Engine/Source/Runtime/Engine/Private/Animation/SkinWeightProfile.cpp
@@ -3,6 +3,7 @@
 #include "Animation/SkinWeightProfile.h"
 #include "Engine/SkeletalMesh.h"
 #include "UObject/UObjectIterator.h"
+#include "ContentStreaming.h"
 
 #include "Components/SkinnedMeshComponent.h"
 #include "Rendering/SkeletalMeshRenderData.h"
diff --git a/Engine/Source/Runtime/Engine/Private/Streaming/SkeletalMeshUpdate.cpp b/Engine/Source/Runtime/Engine/Private/Streaming/SkeletalMeshUpdate.cpp
--- a/Engine/Source/Runtime/Engine/Private/Streaming/SkeletalMeshUpdate.cpp
+++ b/Engine/Source/Runtime/Engine/Private/Streaming/SkeletalMeshUpdate.cpp
@@ -7,10 +7,12 @@ SkeletalMeshUpdate.cpp: Helpers to stream in and out skeletal mesh LODs.
 #include "Streaming/SkeletalMeshUpdate.h"
 #include "RenderUtils.h"
 #include "Containers/ResourceArray.h"
+#include "ContentStreaming.h"
 #include "Streaming/TextureStreamingHelpers.h"
 #include "HAL/PlatformFilemanager.h"
 #include "Serialization/MemoryReader.h"
 #include "Rendering/SkeletalMeshRenderData.h"
+#include "Components/SkinnedMeshComponent.h"
 #include "Streaming/RenderAssetUpdate.inl"
 
 template class TRenderAssetUpdate<FSkelMeshUpdateContext>;
