UE4 - Replaced USystemTimeTimecodeProvider implementation. It was very slow on switch and did not work anyway.

diff --git a/Engine/Source/Runtime/Engine/Private/Engine/SystemTimeTimecodeProvider.cpp b/Engine/Source/Runtime/Engine/Private/Engine/SystemTimeTimecodeProvider.cpp
--- a/Engine/Source/Runtime/Engine/Private/Engine/SystemTimeTimecodeProvider.cpp
+++ b/Engine/Source/Runtime/Engine/Private/Engine/SystemTimeTimecodeProvider.cpp
@@ -5,11 +5,23 @@
 #include "Misc/CoreMisc.h"
 #include "Misc/DateTime.h"
 
-
-FTimecode USystemTimeTimecodeProvider::GetTimecode() const
+static double ComputeTimeCodeOffset()
 {
 	const FDateTime DateTime = FDateTime::Now();
+	double HighPerformanceClock = FPlatformTime::Seconds();
 	const FTimespan Timespan = DateTime.GetTimeOfDay();
+	double Delta = Timespan.GetTotalSeconds() - HighPerformanceClock;
+	return Delta;
+}
+
+
+FTimecode USystemTimeTimecodeProvider::GetTimecode() const
+{
+	static double HighPerformanceClockDelta = ComputeTimeCodeOffset();
+	static double SecondsPerDay = 24.0 * 60.0 * 60.0;
+	FTimespan Timespan;
+	Timespan = FTimespan::FromSeconds(fmod(HighPerformanceClockDelta + FPlatformTime::Seconds(), SecondsPerDay));
 
-	return FTimecode::FromTimespan(Timespan, FrameRate, FTimecode::IsDropFormatTimecodeSupported(FrameRate), false);
+	FTimecode Result =  FTimecode::FromTimespan(Timespan, FrameRate, FTimecode::IsDropFormatTimecodeSupported(FrameRate), false);
+	return Result;
 }
