More fixes on deprecation warnings

diff --git a/Engine/Source/Runtime/Engine/Public/Animation/SkinWeightProfile.h b/Engine/Source/Runtime/Engine/Public/Animation/SkinWeightProfile.h
--- a/Engine/Source/Runtime/Engine/Public/Animation/SkinWeightProfile.h
+++ b/Engine/Source/Runtime/Engine/Public/Animation/SkinWeightProfile.h
@@ -182,7 +182,7 @@ struct ENGINE_API FSkinWeightProfilesData
 		for (int32 Idx = 0; Idx < IntermediateBuffers.Num(); ++Idx)
 		{
 			const FName& ProfileName = IntermediateBuffers[Idx].Key;
-			FVertexBufferRHIParamRef IntermediateBuffer = IntermediateBuffers[Idx].Value;
+			FRHIVertexBuffer* IntermediateBuffer = IntermediateBuffers[Idx].Value;
 			ProfileNameToBuffer.FindChecked(ProfileName)->InitRHIForStreaming(IntermediateBuffer, Batcher);
 		}
 	}
diff --git a/Engine/Source/Runtime/Engine/Public/RawIndexBuffer.h b/Engine/Source/Runtime/Engine/Public/RawIndexBuffer.h
--- a/Engine/Source/Runtime/Engine/Public/RawIndexBuffer.h
+++ b/Engine/Source/Runtime/Engine/Public/RawIndexBuffer.h
@@ -457,7 +457,7 @@ public:
 
 	/** Similar to Init/ReleaseRHI but only update existing SRV so references to the SRV stays valid */
 	template <uint32 MaxNumUpdates>
-	void InitRHIForStreaming(FIndexBufferRHIParamRef IntermediateBuffer, TRHIResourceUpdateBatcher<MaxNumUpdates>& Batcher)
+	void InitRHIForStreaming(FRHIIndexBuffer* IntermediateBuffer, TRHIResourceUpdateBatcher<MaxNumUpdates>& Batcher)
 	{
 		if (IndexBufferRHI && IntermediateBuffer)
 		{
diff --git a/Engine/Source/Runtime/Engine/Public/Rendering/MultiSizeIndexContainer.h b/Engine/Source/Runtime/Engine/Public/Rendering/MultiSizeIndexContainer.h
--- a/Engine/Source/Runtime/Engine/Public/Rendering/MultiSizeIndexContainer.h
+++ b/Engine/Source/Runtime/Engine/Public/Rendering/MultiSizeIndexContainer.h
@@ -98,7 +98,7 @@ public:
 	FIndexBufferRHIRef CreateRHIBuffer_Async();
 
 	template <uint32 MaxNumUpdates>
-	void InitRHIForStreaming(FIndexBufferRHIParamRef IntermediateBuffer, TRHIResourceUpdateBatcher<MaxNumUpdates>& Batcher)
+	void InitRHIForStreaming(FRHIIndexBuffer* IntermediateBuffer, TRHIResourceUpdateBatcher<MaxNumUpdates>& Batcher)
 	{
 		check(!((uint32)!!IntermediateBuffer ^ (uint32)!!IndexBuffer));
 		if (IntermediateBuffer)
diff --git a/Engine/Source/Runtime/Engine/Public/Rendering/SkeletalMeshVertexClothBuffer.h b/Engine/Source/Runtime/Engine/Public/Rendering/SkeletalMeshVertexClothBuffer.h
--- a/Engine/Source/Runtime/Engine/Public/Rendering/SkeletalMeshVertexClothBuffer.h
+++ b/Engine/Source/Runtime/Engine/Public/Rendering/SkeletalMeshVertexClothBuffer.h
@@ -129,7 +129,7 @@ public:
 
 	/** Similar to Init/ReleaseRHI but only update existing SRV so references to the SRV stays valid */
 	template <uint32 MaxNumUpdates>
-	void InitRHIForStreaming(FVertexBufferRHIParamRef IntermediateBuffer, TRHIResourceUpdateBatcher<MaxNumUpdates>& Batcher)
+	void InitRHIForStreaming(FRHIVertexBuffer* IntermediateBuffer, TRHIResourceUpdateBatcher<MaxNumUpdates>& Batcher)
 	{
 		if (VertexBufferRHI && IntermediateBuffer)
 		{
diff --git a/Engine/Source/Runtime/Engine/Public/Rendering/SkinWeightVertexBuffer.h b/Engine/Source/Runtime/Engine/Public/Rendering/SkinWeightVertexBuffer.h
--- a/Engine/Source/Runtime/Engine/Public/Rendering/SkinWeightVertexBuffer.h
+++ b/Engine/Source/Runtime/Engine/Public/Rendering/SkinWeightVertexBuffer.h
@@ -192,7 +192,7 @@ public:
 
 	/** Similar to Init/ReleaseRHI but only update existing SRV so references to the SRV stays valid */
 	template <uint32 MaxNumUpdates>
-	void InitRHIForStreaming(FVertexBufferRHIParamRef IntermediateBuffer, TRHIResourceUpdateBatcher<MaxNumUpdates>& Batcher)
+	void InitRHIForStreaming(FRHIVertexBuffer* IntermediateBuffer, TRHIResourceUpdateBatcher<MaxNumUpdates>& Batcher)
 	{
 		if (VertexBufferRHI && IntermediateBuffer)
 		{
