Make sure FastArrayDelta serialization uses proper offsets for Unmapped GUIDs when playing back replays.

diff --git a/Engine/Source/Runtime/Engine/Private/RepLayout.cpp b/Engine/Source/Runtime/Engine/Private/RepLayout.cpp
--- a/Engine/Source/Runtime/Engine/Private/RepLayout.cpp
+++ b/Engine/Source/Runtime/Engine/Private/RepLayout.cpp
@@ -6695,7 +6695,7 @@ bool FRepLayout::DeltaSerializeFastArrayProperty(FFastArrayDeltaSerializeParams&
 					ItemLayoutStart,
 					ItemLayoutEnd,
 					nullptr,
-					FastArrayHelper.GetRawPtr(0),
+					ThisElement,
 					ThisElement,
 					&GuidReferences,
 					bOutHasUnmapped,
