Don't partially unposses the spectator pawn in DestroySpectatorPawn, allow the DestroyActor call to properly route through to the full unpossess process.

diff --git a/Engine/Source/Runtime/Engine/Private/PlayerController.cpp b/Engine/Source/Runtime/Engine/Private/PlayerController.cpp
--- a/Engine/Source/Runtime/Engine/Private/PlayerController.cpp
+++ b/Engine/Source/Runtime/Engine/Private/PlayerController.cpp
@@ -4867,7 +4867,6 @@ void APlayerController::DestroySpectatorPawn()
 			SetViewTarget(this);
 		}
 
-		GetSpectatorPawn()->UnPossessed();
 		GetWorld()->DestroyActor(GetSpectatorPawn());
 		SetSpectatorPawn(NULL);
 	}
