Fix for crash in PSC Manager when seamless travelling.

diff --git a/Engine/Source/Runtime/Engine/Classes/Engine/World.h b/Engine/Source/Runtime/Engine/Classes/Engine/World.h
--- a/Engine/Source/Runtime/Engine/Classes/Engine/World.h
+++ b/Engine/Source/Runtime/Engine/Classes/Engine/World.h
@@ -3618,6 +3618,7 @@ public:
 	// Global Callback after actors have been initialized (on any world)
 	static UWorld::FOnWorldInitializedActors OnWorldInitializedActors;
 
+	static FWorldEvent OnWorldBeginTearDown;
 private:
 	FWorldDelegates() {}
 };
diff --git a/Engine/Source/Runtime/Engine/Classes/Particles/ParticleSystemManager.h b/Engine/Source/Runtime/Engine/Classes/Particles/ParticleSystemManager.h
--- a/Engine/Source/Runtime/Engine/Classes/Particles/ParticleSystemManager.h
+++ b/Engine/Source/Runtime/Engine/Classes/Particles/ParticleSystemManager.h
@@ -85,6 +85,7 @@ private:
 	static void OnWorldInit(UWorld* World, const UWorld::InitializationValues IVS);
 	static void OnWorldCleanup(UWorld* World, bool bSessionEnded, bool bCleanupResources);
 	static void OnPreWorldFinishDestroy(UWorld* World);
+	static void OnWorldBeginTearDown(UWorld* World);
 
 	typedef TFunction<void(UParticleSystemComponent*, FPSCTickData&, int32)> FPSCTickListFunction;
 
diff --git a/Engine/Source/Runtime/Engine/Private/Particles/ParticleComponents.cpp b/Engine/Source/Runtime/Engine/Private/Particles/ParticleComponents.cpp
--- a/Engine/Source/Runtime/Engine/Private/Particles/ParticleComponents.cpp
+++ b/Engine/Source/Runtime/Engine/Private/Particles/ParticleComponents.cpp
@@ -3860,8 +3860,8 @@ void UParticleSystemComponent::OnRegister()
 	}
 
 	UE_LOG(LogParticles,Verbose,
-		TEXT("OnRegister %s Component=0x%p Scene=0x%p FXSystem=0x%p"),
-		Template != NULL ? *Template->GetName() : TEXT("NULL"), this, World->Scene, FXSystem);
+		TEXT("OnRegister %s Component=0x%p World=0x%p Scene=0x%p FXSystem=0x%p"),
+		Template != NULL ? *Template->GetName() : TEXT("NULL"), this, GetWorld(), World->Scene, FXSystem);
 
 	if (LODLevel == -1)
 	{
diff --git a/Engine/Source/Runtime/Engine/Private/Particles/ParticleSystemManager.cpp b/Engine/Source/Runtime/Engine/Private/Particles/ParticleSystemManager.cpp
--- a/Engine/Source/Runtime/Engine/Private/Particles/ParticleSystemManager.cpp
+++ b/Engine/Source/Runtime/Engine/Private/Particles/ParticleSystemManager.cpp
@@ -145,6 +145,7 @@ void FParticleSystemWorldManager::OnStartup()
 	FWorldDelegates::OnPreWorldInitialization.AddStatic(&OnWorldInit);
 	FWorldDelegates::OnPostWorldCleanup.AddStatic(&OnWorldCleanup);
 	FWorldDelegates::OnPreWorldFinishDestroy.AddStatic(&OnPreWorldFinishDestroy);
+	FWorldDelegates::OnWorldBeginTearDown.AddStatic(&OnWorldBeginTearDown);
 }
 
 void FParticleSystemWorldManager::OnShutdown()
@@ -169,16 +170,16 @@ void FParticleSystemWorldManager::OnWorldInit(UWorld* World, const UWorld::Initi
 		TickGroupEnum = FindObjectChecked<UEnum>(ANY_PACKAGE, TEXT("ETickingGroup"));
 	}
 #endif
-
-	WorldManagers.Add(World) = new FParticleSystemWorldManager(World);
-	//UE_LOG(LogParticles, Warning, TEXT("| OnWorldInit | World: %p | %s |"), World, *World->GetFullName());
+	FParticleSystemWorldManager* NewWorldMan = new FParticleSystemWorldManager(World);
+	WorldManagers.Add(World) = NewWorldMan;
+	UE_LOG(LogParticles, Verbose, TEXT("| OnWorldInit | WorldMan: %p | World: %p | %s |"), NewWorldMan, World, *World->GetFullName());
 }
 
 void FParticleSystemWorldManager::OnWorldCleanup(UWorld* World, bool bSessionEnded, bool bCleanupResources)
 {
-	//UE_LOG(LogParticles, Warning, TEXT("| OnWorldCleanup | World: %p | %s |"), World, *World->GetFullName());
 	if (FParticleSystemWorldManager** WorldMan = WorldManagers.Find(World))
 	{
+		UE_LOG(LogParticles, Verbose, TEXT("| OnWorldCleanup | WorldMan: %p | World: %p | %s |"), *WorldMan, World, *World->GetFullName());
 		delete (*WorldMan);
 		WorldManagers.Remove(World);
 	}
@@ -186,9 +187,19 @@ void FParticleSystemWorldManager::OnWorldCleanup(UWorld* World, bool bSessionEnd
 
 void FParticleSystemWorldManager::OnPreWorldFinishDestroy(UWorld* World)
 {
-	//UE_LOG(LogParticles, Warning, TEXT("| OnPreWorldFinishDestroy | World: %p | %s |"), World, *World->GetFullName());
 	if (FParticleSystemWorldManager** WorldMan = WorldManagers.Find(World))
 	{
+		UE_LOG(LogParticles, Verbose, TEXT("| OnPreWorldFinishDestroy | WorldMan: %p | World: %p | %s |"), *WorldMan, World, *World->GetFullName());
+		delete (*WorldMan);
+		WorldManagers.Remove(World);
+	}
+}
+
+void FParticleSystemWorldManager::OnWorldBeginTearDown(UWorld* World)
+{
+	if (FParticleSystemWorldManager** WorldMan = WorldManagers.Find(World))
+	{
+		UE_LOG(LogParticles, Verbose, TEXT("| OnWorldBeginTearDown | WorldMan: %p | World: %p | %s |"), *WorldMan, World, *World->GetFullName());
 		delete (*WorldMan);
 		WorldManagers.Remove(World);
 	}
@@ -224,6 +235,7 @@ FParticleSystemWorldManager::FParticleSystemWorldManager(UWorld* InWorld)
 
 FParticleSystemWorldManager::~FParticleSystemWorldManager()
 {
+	UE_LOG(LogParticles, Verbose, TEXT("| FParticleSystemWorldManager::~FParticleSystemWorldManager() | %p |"), this);
 	Cleanup();
 }
 
@@ -263,7 +275,7 @@ FString FParticleSystemWorldManager::GetReferencerName() const
 
 void FParticleSystemWorldManager::HandlePostGarbageCollect()
 {
-	//UE_LOG(LogParticles, Warning, TEXT("| HandlePostGarbageCollect |"));
+	UE_LOG(LogParticles, Verbose, TEXT("| HandlePostGarbageCollect | WorldMan: %p |"), this);
 	for (int32 PSCIndex = ManagedPSCs.Num() - 1; PSCIndex >= 0; --PSCIndex)
 	{
 		if (ManagedPSCs[PSCIndex] == nullptr)
@@ -275,6 +287,8 @@ void FParticleSystemWorldManager::HandlePostGarbageCollect()
 
 void FParticleSystemWorldManager::Cleanup()
 {
+	UE_LOG(LogParticles, Verbose, TEXT("| FParticleSystemWorldManager::Cleanup() | %p |"), this);
+
 	FCoreUObjectDelegates::GetPostGarbageCollect().Remove(PostGarbageCollectHandle);
 
 	// Clear out pending particle system components.
@@ -308,12 +322,12 @@ bool FParticleSystemWorldManager::RegisterComponent(UParticleSystemComponent* PS
 			PSC->SetManagerHandle(Handle);
 			PSC->SetPendingManagerAdd(true);
 
-			UE_LOG(LogParticles, Verbose, TEXT("| Register New: %p | %d | %s"), PSC, Handle, *PSC->Template->GetName());
+			UE_LOG(LogParticles, Verbose, TEXT("| Register New: %p | Man: %p | %d | %s"), PSC, this, Handle, *PSC->Template->GetName());
 			return true;
 		}
 		else
 		{
-			UE_LOG(LogParticles, Verbose, TEXT("| Register Existing Pending PSC: %p | %d | %s"), PSC, Handle, *PSC->Template->GetName());
+			UE_LOG(LogParticles, Verbose, TEXT("| Register Existing Pending PSC: %p | Man: %p | %d | %s"), PSC, this, Handle, *PSC->Template->GetName());
 			return false;
 		}
 	}
@@ -323,7 +337,8 @@ bool FParticleSystemWorldManager::RegisterComponent(UParticleSystemComponent* PS
 		//Ensure we're not set to unregister if we were.
 		TickData.bPendingUnregister = false;
 		PSC->SetPendingManagerRemove(false);
-		UE_LOG(LogParticles, Verbose, TEXT("| Register Existing PSC: %p | %d | %s"), PSC, Handle, *PSC->Template->GetName());
+		UE_LOG(LogParticles, Verbose, TEXT("| Register Existing PSC: %p | Man: %p | %d | %s"), PSC, this, Handle, *PSC->Template->GetName());
+
 		return true;
 	}
 
@@ -338,7 +353,7 @@ void FParticleSystemWorldManager::UnregisterComponent(UParticleSystemComponent*
 	{
 		if (PSC->IsPendingManagerAdd())
 		{
-			UE_LOG(LogParticles, Verbose, TEXT("| UnRegister Pending PSC: %p | %d | %s"), PSC, Handle, *PSC->Template->GetName());
+			UE_LOG(LogParticles, Verbose, TEXT("| UnRegister Pending PSC: %p | Man: %p | %d | %s"), PSC, this, Handle, *PSC->Template->GetName());
 
 			//Clear existing handle
 			check(PendingRegisterPSCs[Handle]);
@@ -361,7 +376,7 @@ void FParticleSystemWorldManager::UnregisterComponent(UParticleSystemComponent*
 			//Don't remove us immediately from the arrays as this can occur mid tick. Just mark us for removal next frame.
 			TickData.bPendingUnregister = true;
 			PSC->SetPendingManagerRemove(true);
-			UE_LOG(LogParticles, Verbose, TEXT("| UnRegister PSC: %p | %d | %s"), PSC, Handle, *PSC->Template->GetName());
+			UE_LOG(LogParticles, Verbose, TEXT("| UnRegister PSC: %p | Man: %p | %d | %s"), PSC, this, Handle, *PSC->Template->GetName());
 		}
 	}
 	//handled dependencies
@@ -404,7 +419,7 @@ void FParticleSystemWorldManager::AddPSC(UParticleSystemComponent* PSC)
 		TickList.Add(Handle);
 #endif
 
-		//UE_LOG(LogParticles, Warning, TEXT("| Add PSC - PSC: %p | %d | %d |Num: %d |"), ManagedPSCs[Handle], Handle, TickList[TickData.TickListHandle], ManagedPSCs.Num());
+		UE_LOG(LogParticles, Verbose, TEXT("| Add PSC - PSC: %p | Man: %p | %d | %d |Num: %d |"), ManagedPSCs[Handle], this, Handle, TickList[TickData.TickListHandle], ManagedPSCs.Num());
 	}
 }
 
@@ -419,12 +434,12 @@ void FParticleSystemWorldManager::RemovePSC(int32 PSCIndex)
 
 	FPSCTickData& TickData = PSCTickData[PSCIndex];
 
-	//UE_LOG(LogParticles, Warning, TEXT("| Remove PSC - PSC: %p | %d | %d |Num: %d |"), ManagedPSCs[PSCIndex], PSCIndex, TickList[TickData.TickListHandle], ManagedPSCs.Num());
+	UE_LOG(LogParticles, Verbose, TEXT("| Remove PSC - PSC: %p | Man: %p | %d |Num: %d |"), ManagedPSCs[PSCIndex], this, PSCIndex, ManagedPSCs.Num());
 
 #if PSC_MAN_USE_STATIC_TICK_LISTS
 	FTickList& TickList = TickData.bCanTickConcurrent ? TickLists_Concurrent[(int32)TickData.TickGroup] : TickLists_GT[(int32)TickData.TickGroup];
 
-	//UE_LOG(LogParticles, Warning, TEXT("| Remove PSC - PSC: %p | %d | %d |Num: %d |"), ManagedPSCs[PSCIndex], PSCIndex, TickList[TickData.TickListHandle], ManagedPSCs.Num());
+	UE_LOG(LogParticles, Verbose, TEXT("| Remove PSC - PSC: %p | Man: %p | %d | %d |Num: %d |"), ManagedPSCs[PSCIndex], this, PSCIndex, TickList[TickData.TickListHandle], ManagedPSCs.Num());
 
 	TickList.Remove(PSCIndex);
 
@@ -609,7 +624,7 @@ void FParticleSystemWorldManager::ClearPendingUnregister()
 
 void FParticleSystemWorldManager::Tick(ETickingGroup TickGroup, float DeltaTime, ELevelTick TickType, ENamedThreads::Type CurrentThread, const FGraphEventRef& MyCompletionGraphEvent)
 {
-	//UE_LOG(LogParticles, Warning, TEXT("| ---- PSC World Manager Tick ----- | TG %s | World: %p - %s |"), *TickGroupEnum->GetNameByValue(TickGroup).ToString(), World, *World->GetFullName());
+	//UE_LOG(LogParticles, Verbose, TEXT("| ---- PSC World Manager Tick ----- | TG %s | World: %p - %s |"), *TickGroupEnum->GetNameByValue(TickGroup).ToString(), World, *World->GetFullName());
 	
 	SCOPE_CYCLE_COUNTER(STAT_PSCMan_Tick);
 	CSV_SCOPED_TIMING_STAT_EXCLUSIVE(Effects);
diff --git a/Engine/Source/Runtime/Engine/Private/World.cpp b/Engine/Source/Runtime/Engine/Private/World.cpp
--- a/Engine/Source/Runtime/Engine/Private/World.cpp
+++ b/Engine/Source/Runtime/Engine/Private/World.cpp
@@ -329,6 +329,7 @@ FWorldDelegates::FWorldGetAssetTags FWorldDelegates::GetAssetTags;
 FWorldDelegates::FOnWorldTickStart FWorldDelegates::OnWorldTickStart;
 FWorldDelegates::FOnWorldPreActorTick FWorldDelegates::OnWorldPreActorTick;
 FWorldDelegates::FOnWorldPostActorTick FWorldDelegates::OnWorldPostActorTick;
+FWorldDelegates::FWorldEvent FWorldDelegates::OnWorldBeginTearDown;
 #if WITH_EDITOR
 FWorldDelegates::FRefreshLevelScriptActionsEvent FWorldDelegates::RefreshLevelScriptActions;
 #endif // WITH_EDITOR
@@ -2513,6 +2514,10 @@ void UWorld::BeginTearingDown()
 	bIsTearingDown = true;
 	UE_LOG(LogWorld, Log, TEXT("BeginTearingDown for %s"), *GetOutermost()->GetName());
 	BeginTearingDownEvent.Broadcast();
+
+	//Simultaneous similar edits that caused merge conflict. Taking both for now to unblock.
+	//Can likely be unified.
+	FWorldDelegates::OnWorldBeginTearDown.Broadcast(this);
 }
 
 void UWorld::RemoveFromWorld( ULevel* Level, bool bAllowIncrementalRemoval )
