Fix a crash when loading old content on a newer engine

diff --git a/Engine/Source/Runtime/Engine/Private/SkeletalMeshRenderData.cpp b/Engine/Source/Runtime/Engine/Private/SkeletalMeshRenderData.cpp
--- a/Engine/Source/Runtime/Engine/Private/SkeletalMeshRenderData.cpp
+++ b/Engine/Source/Runtime/Engine/Private/SkeletalMeshRenderData.cpp
@@ -32,7 +32,7 @@ namespace SkeletalMeshCookStats
 // differences, etc.) replace the version GUID below with a new one.
 // In case of merge conflicts with DDC versions, you MUST generate a new GUID
 // and set this new GUID as the version.                                       
-#define SKELETALMESH_DERIVEDDATA_VER TEXT("70D1F1C048AF4DF2AF870B51CE803A66")
+#define SKELETALMESH_DERIVEDDATA_VER TEXT("51BDAF351B6941548818C2F8B33397D7")
 
 static const FString& GetSkeletalMeshDerivedDataVersion()
 {
diff --git a/Engine/Source/Runtime/Engine/Public/RawIndexBuffer.h b/Engine/Source/Runtime/Engine/Public/RawIndexBuffer.h
--- a/Engine/Source/Runtime/Engine/Public/RawIndexBuffer.h
+++ b/Engine/Source/Runtime/Engine/Public/RawIndexBuffer.h
@@ -377,9 +377,8 @@ public:
 	*/
 	virtual void Serialize( FArchive& Ar ) override
 	{
-		SerializeMetaData(Ar);
 		Indices.BulkSerialize(Ar);
-		check(CachedNumIndices == Indices.Num());
+		CachedNumIndices = Indices.Num();
 	}
 
 	virtual void SerializeMetaData(FArchive& Ar) override
