Allow cooked sound waves to be loaded in the editor, particularly for the package info commandlet
#rnx

diff --git a/Engine/Source/Runtime/Engine/Private/AudioDerivedData.cpp b/Engine/Source/Runtime/Engine/Private/AudioDerivedData.cpp
--- a/Engine/Source/Runtime/Engine/Private/AudioDerivedData.cpp
+++ b/Engine/Source/Runtime/Engine/Private/AudioDerivedData.cpp
@@ -1583,12 +1583,16 @@ void USoundWave::FinishCachePlatformData()
 	}
 
 #if DO_CHECK
-	FString DerivedDataKey;
-	FName AudioFormat = GetWaveFormatForRunningPlatform(*this);
-	const FPlatformAudioCookOverrides* CompressionOverrides = GetPlatformCompressionOverridesForCurrentPlatform();
-	GetStreamedAudioDerivedDataKey(*this, AudioFormat, CompressionOverrides, DerivedDataKey);
+	// If we're allowing cooked data to be loaded then the derived data key will not have been serialized, so won't match and that's fine
+	if (!GAllowCookedDataInEditorBuilds)
+	{
+		FString DerivedDataKey;
+		FName AudioFormat = GetWaveFormatForRunningPlatform(*this);
+		const FPlatformAudioCookOverrides* CompressionOverrides = GetPlatformCompressionOverridesForCurrentPlatform();
+		GetStreamedAudioDerivedDataKey(*this, AudioFormat, CompressionOverrides, DerivedDataKey);
 
-	check(RunningPlatformData->DerivedDataKey == DerivedDataKey);
+		check(RunningPlatformData->DerivedDataKey == DerivedDataKey);
+	}
 #endif
 }
 
