Support duplicating non-optional mips in uptnl file, which will reduce IO seek time when loading textures with optional mips
- Since non-optional mips will be duplicated, this will increase build size
- Non-optional mips to be duplicated in uptnl files will be serialized twice when cooking, once in ubulk and once in uptnl
- Duplication can be controlled per LODGROUP in TextureLODGroups with attribute DuplicateNonOptionalMips
- When loading a texture with optional mips, if uptnl is available, all non-inline mips will be loaded from uptnl instead of from ubulk to reduce IO seek time and will improve performance

diff --git a/Engine/Source/Runtime/CoreUObject/Private/Serialization/BulkData.cpp b/Engine/Source/Runtime/CoreUObject/Private/Serialization/BulkData.cpp
--- a/Engine/Source/Runtime/CoreUObject/Private/Serialization/BulkData.cpp
+++ b/Engine/Source/Runtime/CoreUObject/Private/Serialization/BulkData.cpp
@@ -974,10 +974,12 @@ void FUntypedBulkData::Serialize( FArchive& Ar, UObject* Owner, int32 Idx, bool
 			const bool bPayloadInline = !(BulkDataFlags&BULKDATA_PayloadAtEndOfFile);
 
 			// fix up the file offset, but only if not stored inline
+			int64 OffsetInFileFixup = 0;
 			if (Owner != NULL && Owner->GetLinker() && !bPayloadInline)
 			{
-				BulkDataOffsetInFile += Owner->GetLinker()->Summary.BulkDataStartOffset;
+				OffsetInFileFixup = Owner->GetLinker()->Summary.BulkDataStartOffset;
 			}
+			BulkDataOffsetInFile += OffsetInFileFixup;
 
 			// We're allowing defered serialization.
 			FArchive* CacheableArchive = Ar.GetCacheableArchive();
@@ -1034,6 +1036,53 @@ void FUntypedBulkData::Serialize( FArchive& Ar, UObject* Owner, int32 Idx, bool
 						}
 					}
 				}
+				else if (BulkDataFlags & BULKDATA_DuplicateNonOptionalPayload)
+				{
+					// Load from optional payload instead if optional payload is available
+					FString OptionalFilename = FPaths::ChangeExtension(Filename, TEXT(".uptnl"));
+					if (IFileManager::Get().FileExists(*OptionalFilename))
+					{
+						Filename = OptionalFilename;
+						Ar << BulkDataFlags;
+						BulkDataFlags |= BULKDATA_OptionalPayload;
+
+						if (BulkDataFlags & BULKDATA_Size64Bit)
+						{
+							Ar << BulkDataSizeOnDisk;
+						}
+						else
+						{
+							int32 BulkDataSizeOnDiskAsInt32;
+							Ar << BulkDataSizeOnDiskAsInt32;
+							BulkDataSizeOnDisk = BulkDataSizeOnDiskAsInt32;
+						}
+
+						Ar << BulkDataOffsetInFile;
+						BulkDataOffsetInFile += OffsetInFileFixup;
+					}
+					else
+					{
+						// Skip all bulkdata info of the copy in uptnl.
+						uint32 DummyBulkDataFlags;
+						Ar << DummyBulkDataFlags;
+
+						if (BulkDataFlags & BULKDATA_Size64Bit)
+						{
+							int64 DummyBulkDataSizeOnDisk;
+							Ar << DummyBulkDataSizeOnDisk;
+						}
+						else
+						{
+							int32 DummyBulkDataSizeOnDiskAsInt32;
+							Ar << DummyBulkDataSizeOnDiskAsInt32;
+						}
+
+						int64 DummyBulkDataOffsetInFile;
+						Ar << DummyBulkDataOffsetInFile;
+
+						Filename = FPaths::ChangeExtension(Filename, TEXT(".ubulk"));
+					}
+				}
 				else if (BulkDataFlags & BULKDATA_OptionalPayload)
 				{
 					Filename = FPaths::ChangeExtension(Filename, TEXT(".uptnl"));
@@ -1119,7 +1168,7 @@ void FUntypedBulkData::Serialize( FArchive& Ar, UObject* Owner, int32 Idx, bool
 							Ar.Seek(CurOffset);
 						}
 					}
-  			}
+  				}
 			}
 		}
 		// We're saving to the persistent archive.
@@ -1188,6 +1237,50 @@ void FUntypedBulkData::Serialize( FArchive& Ar, UObject* Owner, int32 Idx, bool
 				BulkStore.BulkDataFlagsPos = SavedBulkDataFlagsPos;
 				BulkStore.BulkDataFlags = BulkDataFlags;
 				BulkStore.BulkData = this;
+
+				// If having flag BULKDATA_DuplicateNonOptionalPayload, duplicate bulk data in optional storage (.uptnl)
+				if (BulkDataFlags & BULKDATA_DuplicateNonOptionalPayload)
+				{
+					int64 SavedDupeBulkDataFlagsPos = INDEX_NONE;
+					int64 SavedDupeBulkDataSizeOnDiskPos = INDEX_NONE;
+					int64 SavedDupeBulkDataOffsetInFilePos = INDEX_NONE;
+
+					int32 SavedDupeBulkDataFlags = (BulkDataFlags & ~BULKDATA_DuplicateNonOptionalPayload) | BULKDATA_OptionalPayload;
+					{
+						FArchive::FScopeSetDebugSerializationFlags S(Ar, DSF_IgnoreDiff);
+
+						SavedDupeBulkDataFlagsPos = Ar.Tell();
+						Ar << SavedDupeBulkDataFlags;
+
+						// And serialize the placeholder which is going to be overwritten later.
+						SavedDupeBulkDataSizeOnDiskPos = Ar.Tell();
+						BulkDataSizeOnDisk = INDEX_NONE;
+						if (BulkDataFlags & BULKDATA_Size64Bit)
+						{
+							Ar << BulkDataSizeOnDisk;
+						}
+						else
+						{
+							int32 BulkDataSizeOnDiskAsInt32 = BulkDataSizeOnDisk;
+							Ar << BulkDataSizeOnDiskAsInt32;
+						}
+						// Keep track of position we are going to serialize placeholder BulkDataOffsetInFile.
+						SavedDupeBulkDataOffsetInFilePos = Ar.Tell();
+						BulkDataOffsetInFile = INDEX_NONE;
+						// And serialize the placeholder which is going to be overwritten later.
+						Ar << BulkDataOffsetInFile;
+					}
+
+					// add duplicate bulkdata with different flag
+					Index = LinkerSave->BulkDataToAppend.AddZeroed(1);
+					FLinkerSave::FBulkDataStorageInfo& DupeBulkStore = LinkerSave->BulkDataToAppend[Index];
+
+					DupeBulkStore.BulkDataOffsetInFilePos = SavedDupeBulkDataOffsetInFilePos;
+					DupeBulkStore.BulkDataSizeOnDiskPos = SavedDupeBulkDataSizeOnDiskPos;
+					DupeBulkStore.BulkDataFlagsPos = SavedDupeBulkDataFlagsPos;
+					DupeBulkStore.BulkDataFlags = SavedDupeBulkDataFlags;
+					DupeBulkStore.BulkData = this;
+				}
 				
 				// Serialize bulk data into the storage info
 				BulkDataSizeOnDisk = -1;
diff --git a/Engine/Source/Runtime/CoreUObject/Public/Serialization/BulkData.h b/Engine/Source/Runtime/CoreUObject/Public/Serialization/BulkData.h
--- a/Engine/Source/Runtime/CoreUObject/Public/Serialization/BulkData.h
+++ b/Engine/Source/Runtime/CoreUObject/Public/Serialization/BulkData.h
@@ -41,7 +41,9 @@ enum EBulkDataFlags
 	/** This payload will be memory mapped, this requires alignment, no compression etc. */
 	BULKDATA_MemoryMappedPayload = 1 << 12,
 	/** Bulk data size is 64 bits long */
-	BULKDATA_Size64Bit							= 1 << 13
+	BULKDATA_Size64Bit							= 1 << 13,
+	/** Duplicate non-optional payload in optional bulk data. */
+	BULKDATA_DuplicateNonOptionalPayload		= 1 << 14
 };
 
 /**
diff --git a/Engine/Source/Runtime/Engine/Classes/Engine/TextureLODSettings.h b/Engine/Source/Runtime/Engine/Classes/Engine/TextureLODSettings.h
--- a/Engine/Source/Runtime/Engine/Classes/Engine/TextureLODSettings.h
+++ b/Engine/Source/Runtime/Engine/Classes/Engine/TextureLODSettings.h
@@ -35,6 +35,7 @@ struct ENGINE_API FTextureLODGroup
 		, MinMagFilter(NAME_Aniso)
 		, MipFilter(NAME_Point)
 		, MipLoadOptions(ETextureMipLoadOptions::AllMips)
+		, DuplicateNonOptionalMips(false)
 	{
 		SetupGroup();
 	}
@@ -101,6 +102,9 @@ struct ENGINE_API FTextureLODGroup
 	UPROPERTY()
 	ETextureMipLoadOptions MipLoadOptions;
 
+	UPROPERTY()
+	bool DuplicateNonOptionalMips;
+
 	void SetupGroup();
 };
 
diff --git a/Engine/Source/Runtime/Engine/Private/Texture2D.cpp b/Engine/Source/Runtime/Engine/Private/Texture2D.cpp
--- a/Engine/Source/Runtime/Engine/Private/Texture2D.cpp
+++ b/Engine/Source/Runtime/Engine/Private/Texture2D.cpp
@@ -483,8 +483,8 @@ bool UTexture2D::GetMipDataFilename(const int32 MipIndex, FString& OutBulkDataFi
 			OutBulkDataFilename = PlatformData->Mips[MipIndex].BulkData.GetFilename();
 #else
 			OutBulkDataFilename = PlatformData->CachedPackageFileName;
-			const int32 NumOptionalMips = CalcNumOptionalMips();
-			OutBulkDataFilename = FPaths::ChangeExtension(OutBulkDataFilename, MipIndex < NumOptionalMips ? TEXT(".uptnl") : TEXT(".ubulk"));
+			bool UseOptionalBulkDataFileName = PlatformData->Mips[MipIndex].BulkData.GetBulkDataFlags() & BULKDATA_OptionalPayload;
+			OutBulkDataFilename = FPaths::ChangeExtension(OutBulkDataFilename, UseOptionalBulkDataFileName ? TEXT(".uptnl") : TEXT(".ubulk"));
 #endif
 			return true;
 		}
diff --git a/Engine/Source/Runtime/Engine/Private/TextureDerivedData.cpp b/Engine/Source/Runtime/Engine/Private/TextureDerivedData.cpp
--- a/Engine/Source/Runtime/Engine/Private/TextureDerivedData.cpp
+++ b/Engine/Source/Runtime/Engine/Private/TextureDerivedData.cpp
@@ -1286,8 +1286,8 @@ static void SerializePlatformData(
 			if (!bIsVirtual)
 			{
 				FirstMipToSerialize = FMath::Clamp(FirstMipToSerialize, 0, FMath::Max(NumMips - 1, 0));
-				NumMips -= FirstMipToSerialize;
-			}
+			NumMips -= FirstMipToSerialize;
+		}
 			else
 			{
 				FirstMipToSerialize = FMath::Clamp(FirstMipToSerialize, 0, FMath::Max((int32)PlatformData->VTData->GetNumMips() - 1, 0));
@@ -1309,51 +1309,54 @@ static void SerializePlatformData(
 	{
 		if (bIsVirtual == false)
 		{
-			BulkDataMipFlags.AddZeroed(PlatformData->Mips.Num());
-			for (int32 MipIndex = 0; MipIndex < PlatformData->Mips.Num(); ++MipIndex)
-			{
-				BulkDataMipFlags[MipIndex] = PlatformData->Mips[MipIndex].BulkData.GetBulkDataFlags();
-			}
-
-			int32 MinMipToInline = 0;
-			int32 OptionalMips = 0; // TODO: do we need to consider platforms saving texture assets as cooked files? all the info to calculate the optional is part of the editor only data
+		BulkDataMipFlags.AddZeroed(PlatformData->Mips.Num());
+		for (int32 MipIndex = 0; MipIndex < PlatformData->Mips.Num(); ++MipIndex)
+		{
+			BulkDataMipFlags[MipIndex] = PlatformData->Mips[MipIndex].BulkData.GetBulkDataFlags();
+		}
 
+		int32 MinMipToInline = 0;
+		int32 OptionalMips = 0; // TODO: do we need to consider platforms saving texture assets as cooked files? all the info to calculate the optional is part of the editor only data
+		bool DuplicateNonOptionalMips = false;
+		
 #if WITH_EDITORONLY_DATA
-			check(Ar.CookingTarget());
-			// This also needs to check whether the project enables texture streaming.
-			// Currently, there is no reliable way to implement this because there is no difference
-			// between the project settings (CVar) and the command line setting (from -NoTextureStreaming)
-			if (bStreamable && Ar.CookingTarget()->SupportsFeature(ETargetPlatformFeatures::TextureStreaming))
+		check(Ar.CookingTarget());
+		// This also needs to check whether the project enables texture streaming.
+		// Currently, there is no reliable way to implement this because there is no difference
+		// between the project settings (CVar) and the command line setting (from -NoTextureStreaming)
+		if (bStreamable && Ar.CookingTarget()->SupportsFeature(ETargetPlatformFeatures::TextureStreaming))
 #else
-			if (bStreamable)
+		if (bStreamable)
 #endif
-			{
-				MinMipToInline = FMath::Max(0, NumMips - PlatformData->GetNumNonStreamingMips());
+		{
+			MinMipToInline = FMath::Max(0, NumMips - PlatformData->GetNumNonStreamingMips());
 #if WITH_EDITORONLY_DATA
-				const int32 Width = PlatformData->SizeX;
-				const int32 Height = PlatformData->SizeY;
-				const int32 LODGroup = Texture->LODGroup;
-				const int32 LODBias = Texture->LODBias;
-				const int32 NumCinematicMipLevels = Texture->NumCinematicMipLevels;
+			const int32 Width = PlatformData->SizeX;
+			const int32 Height = PlatformData->SizeY;
+			const int32 LODGroup = Texture->LODGroup;
+			const int32 LODBias = Texture->LODBias;
+			const int32 NumCinematicMipLevels = Texture->NumCinematicMipLevels;
 
-				OptionalMips = Ar.CookingTarget()->GetTextureLODSettings().CalculateNumOptionalMips(LODGroup, Width, Height, NumMips, MinMipToInline, Texture->MipGenSettings);
+			OptionalMips = Ar.CookingTarget()->GetTextureLODSettings().CalculateNumOptionalMips(LODGroup, Width, Height, NumMips, MinMipToInline, Texture->MipGenSettings);
+			DuplicateNonOptionalMips = Ar.CookingTarget()->GetTextureLODSettings().TextureLODGroups[LODGroup].DuplicateNonOptionalMips;
 #endif
-			}
+		}
 
 			for (int32 MipIndex = 0; MipIndex < NumMips && MipIndex < OptionalMips; ++MipIndex) //-V654
-			{
+		{
 				PlatformData->Mips[MipIndex + FirstMipToSerialize].BulkData.SetBulkDataFlags(BULKDATA_Force_NOT_InlinePayload | BULKDATA_OptionalPayload);
-			}
+		}
 
-			for (int32 MipIndex = OptionalMips; MipIndex < NumMips && MipIndex < MinMipToInline; ++MipIndex)
-			{
-				PlatformData->Mips[MipIndex + FirstMipToSerialize].BulkData.SetBulkDataFlags(BULKDATA_Force_NOT_InlinePayload);
-			}
-			for (int32 MipIndex = MinMipToInline; MipIndex < NumMips; ++MipIndex)
-			{
-				PlatformData->Mips[MipIndex + FirstMipToSerialize].BulkData.SetBulkDataFlags(BULKDATA_ForceInlinePayload | BULKDATA_SingleUse);
-			}
+		const uint32 AdditionalNonOptionalBulkDataFlags = DuplicateNonOptionalMips ? BULKDATA_DuplicateNonOptionalPayload : 0;
+		for (int32 MipIndex = OptionalMips; MipIndex < NumMips && MipIndex < MinMipToInline; ++MipIndex)
+		{
+			PlatformData->Mips[MipIndex + FirstMipToSerialize].BulkData.SetBulkDataFlags(BULKDATA_Force_NOT_InlinePayload | AdditionalNonOptionalBulkDataFlags);
 		}
+		for (int32 MipIndex = MinMipToInline; MipIndex < NumMips; ++MipIndex)
+		{
+			PlatformData->Mips[MipIndex + FirstMipToSerialize].BulkData.SetBulkDataFlags(BULKDATA_ForceInlinePayload | BULKDATA_SingleUse);
+		}
+	}
 		else // bVirtual == false
 		{
 			const int32 NumChunks = PlatformData->VTData->Chunks.Num();
@@ -1398,12 +1401,12 @@ static void SerializePlatformData(
 
 	if (bIsVirtual == false)
 	{
-		for (int32 MipIndex = 0; MipIndex < BulkDataMipFlags.Num(); ++MipIndex)
-		{
+	for (int32 MipIndex = 0; MipIndex < BulkDataMipFlags.Num(); ++MipIndex)
+	{
 			check(Ar.IsSaving());
-			PlatformData->Mips[MipIndex].BulkData.ClearBulkDataFlags(~BulkDataMipFlags[MipIndex]);
-			PlatformData->Mips[MipIndex].BulkData.SetBulkDataFlags(BulkDataMipFlags[MipIndex]);
-		}
+		PlatformData->Mips[MipIndex].BulkData.ClearBulkDataFlags(~BulkDataMipFlags[MipIndex]);
+		PlatformData->Mips[MipIndex].BulkData.SetBulkDataFlags(BulkDataMipFlags[MipIndex]);
+	}
 	}
 	else
 	{
@@ -1430,16 +1433,16 @@ void FTexturePlatformData::SerializeCooked(FArchive& Ar, UTexture* Owner, bool b
 	{
 		// Patch up Size as due to mips being stripped out during cooking it could be wrong.
 		if (Mips.Num() > 0)
-		{
-			SizeX = Mips[0].SizeX;
-			SizeY = Mips[0].SizeY;
+	{
+		SizeX = Mips[0].SizeX;
+		SizeY = Mips[0].SizeY;
 			
-			// SizeZ is not the same as NumSlices for texture arrays and cubemaps.
-			if (Owner && Owner->IsA(UVolumeTexture::StaticClass()))
-			{
-				 NumSlices = Mips[0].SizeZ;
-			}
+		// SizeZ is not the same as NumSlices for texture arrays and cubemaps.
+		if (Owner && Owner->IsA(UVolumeTexture::StaticClass()))
+		{
+			 NumSlices = Mips[0].SizeZ;
 		}
+	}
 		else if ( VTData )
 		{
 			SizeX = VTData->Width;
