[NO NON-PII COMMIT MESSAGE]

diff --git a/Engine/Source/Runtime/Engine/Private/VorbisAudioInfo.cpp b/Engine/Source/Runtime/Engine/Private/VorbisAudioInfo.cpp
--- a/Engine/Source/Runtime/Engine/Private/VorbisAudioInfo.cpp
+++ b/Engine/Source/Runtime/Engine/Private/VorbisAudioInfo.cpp
@@ -197,7 +197,19 @@ size_t FVorbisAudioInfo::ReadStreaming(void *Ptr, uint32 Size )
 		if (!CurrentStreamingChunkData || CurrentStreamingChunkIndex != NextStreamingChunkIndex)
 		{
 			CurrentStreamingChunkIndex = NextStreamingChunkIndex;
-			CurrentStreamingChunkData = GetLoadedChunk(StreamingSoundWave, CurrentStreamingChunkIndex, CurrentStreamingChunksSize);
+
+			check(CurrentStreamingChunkIndex >= 0);
+
+			if (static_cast<uint32>(CurrentStreamingChunkIndex) >= StreamingSoundWave->GetNumChunks())
+			{
+				CurrentStreamingChunkData = nullptr;
+				CurrentStreamingChunksSize = 0;
+			}
+			else
+			{
+				CurrentStreamingChunkData = GetLoadedChunk(StreamingSoundWave, CurrentStreamingChunkIndex, CurrentStreamingChunksSize);
+			}
+
 			if (CurrentStreamingChunkData)
 			{
 				check(CurrentStreamingChunkIndex < StreamingSoundWave->RunningPlatformData->Chunks.Num());
