Repack USoundWave. 720 -> 704 bytes.
#rnx

diff --git a/Engine/Source/Runtime/Engine/Classes/Sound/SoundWave.h b/Engine/Source/Runtime/Engine/Classes/Sound/SoundWave.h
--- a/Engine/Source/Runtime/Engine/Classes/Sound/SoundWave.h
+++ b/Engine/Source/Runtime/Engine/Classes/Sound/SoundWave.h
@@ -363,9 +363,26 @@ class ENGINE_API USoundWave : public USoundBase
 	/** Whether this SoundWave was decompressed from OGG. */
 	uint8 bDecompressedFromOgg : 1;
 
-	/** Use this to override how much audio data is loaded when this USoundWave is loaded. */
-	UPROPERTY(EditAnywhere, AdvancedDisplay, Category = "Loading")
-	int32 InitialChunkSize;
+private:
+
+#if !WITH_EDITOR
+	// This is set to false on initialization, then set to true on non-editor platforms when we cache appropriate sample rate.
+	uint8 bCachedSampleRateFromPlatformSettings:1;
+
+	// This is set when SetSampleRate is called to invalidate our cached sample rate while not re-parsing project settings.
+	uint8 bSampleRateManuallyReset:1;
+#endif
+
+	enum class ESoundWaveResourceState : uint8
+	{
+		NeedsFree,
+		Freeing,
+		Freed
+	};
+
+	volatile ESoundWaveResourceState ResourceState;
+
+public:
 
 #if WITH_EDITORONLY_DATA
 	/** Specify a sound to use for the baked analysis. Will default to this USoundWave if not sete. */
@@ -432,29 +449,16 @@ class ENGINE_API USoundWave : public USoundBase
 	bool GetInterpolatedCookedFFTDataForTime(float InTime, uint32& InOutLastIndex, TArray<FSoundWaveSpectralData>& OutData, bool bLoop);
 	bool GetInterpolatedCookedEnvelopeDataForTime(float InTime, uint32& InOutLastIndex, float& OutAmplitude, bool bLoop);
 
+	/** Use this to override how much audio data is loaded when this USoundWave is loaded. */
+	UPROPERTY(EditAnywhere, AdvancedDisplay, Category = "Loading")
+	int32 InitialChunkSize;
+
 private:
 
 	/** Helper functions to search analysis data. Takes starting index to start query. Returns which data index the result was found at. Returns INDEX_NONE if not found. */
 	uint32 GetInterpolatedCookedFFTDataForTimeInternal(float InTime, uint32 StartingIndex, TArray<FSoundWaveSpectralData>& OutData, bool bLoop);
 	uint32 GetInterpolatedCookedEnvelopeDataForTimeInternal(float InTime, uint32 StartingIndex, float& OutAmplitude, bool bLoop);
 
-#if !WITH_EDITOR
-	// This is set to false on initialization, then set to true on non-editor platforms when we cache appropriate sample rate.
-	uint8 bCachedSampleRateFromPlatformSettings : 1;
-
-	// This is set when SetSampleRate is called to invalidate our cached sample rate while not re-parsing project settings.
-	uint8 bSampleRateManuallyReset : 1;
-#endif
-
-	enum class ESoundWaveResourceState : uint8
-	{
-		NeedsFree,
-		Freeing,
-		Freed
-	};
-
-	volatile ESoundWaveResourceState ResourceState;
-
 	/** What state the precache decompressor is in. */
 	FThreadSafeCounter PrecacheState;
 
@@ -506,6 +510,16 @@ protected:
 	int32 SampleRate;
 
 public:
+
+	/** Resource index to cross reference with buffers */
+	int32 ResourceID;
+
+	/** Size of resource copied from the bulk data */
+	int32 ResourceSize;
+
+	/** Cache the total used memory recorded for this SoundWave to keep INC/DEC consistent */
+	int32 TrackedMemoryUsage;
+
 	/**
 	 * Subtitle cues.  If empty, use SpokenText as the subtitle.  Will often be empty,
 	 * as the contents of the subtitle is commonly identical to what is spoken.
@@ -559,12 +573,12 @@ public:
 	/** The number of frames which have been precached for this sound wave. */
 	int32 NumPrecacheFrames;
 
-	/** Pointer to 16 bit PCM data - used to decompress data to and preview sounds */
-	uint8* RawPCMData;
-
 	/** Size of RawPCMData, or what RawPCMData would be if the sound was fully decompressed */
 	int32 RawPCMDataSize;
 
+	/** Pointer to 16 bit PCM data - used to decompress data to and preview sounds */
+	uint8* RawPCMData;
+
 	/** Memory containing the data copied from the compressed bulk data */
 	FOwnedBulkDataPtr* OwnedBulkDataPtr;
 	const uint8* ResourceData;
@@ -588,15 +602,6 @@ public:
 
 #endif
 
-	/** Resource index to cross reference with buffers */
-	int32 ResourceID;
-
-	/** Size of resource copied from the bulk data */
-	int32 ResourceSize;
-
-	/** Cache the total used memory recorded for this SoundWave to keep INC/DEC consistent */
-	int32 TrackedMemoryUsage;
-
 	/** The streaming derived data for this sound on this platform. */
 	FStreamedAudioPlatformData* RunningPlatformData;
 
