Moving the atlas generation to PostLoad so it will be called in cook commandlets
#rnx

diff --git a/Engine/Source/Runtime/Engine/Private/Materials/MaterialInstance.cpp b/Engine/Source/Runtime/Engine/Private/Materials/MaterialInstance.cpp
--- a/Engine/Source/Runtime/Engine/Private/Materials/MaterialInstance.cpp
+++ b/Engine/Source/Runtime/Engine/Private/Materials/MaterialInstance.cpp
@@ -466,26 +466,6 @@ void UMaterialInstance::InitResources()
 		Resource->GameThread_SetParent(SafeParent);
 	}
 
-#if WITH_EDITOR
-	//recalculate any scalar params based on a curve position in an atlas in case the atlas changed
-	for (FScalarParameterValue& ScalarParam : ScalarParameterValues)
-	{
-		if (ScalarParam.AtlasData.bIsUsedAsAtlasPosition)
-		{
-			UCurveLinearColorAtlas* Atlas = Cast<UCurveLinearColorAtlas>(ScalarParam.AtlasData.Atlas.Get());
-			UCurveLinearColor* Curve = Cast<UCurveLinearColor>(ScalarParam.AtlasData.Curve.Get());
-			if (Curve && Atlas)
-			{
-				int32 Index = Atlas->GradientCurves.Find(Curve);
-				if (Index != INDEX_NONE)
-				{
-					ScalarParam.ParameterValue = (float)Index;
-				}
-			}
-		}
-	}
-#endif
-
 	GameThread_InitMIParameters(this, ScalarParameterValues);
 	GameThread_InitMIParameters(this, VectorParameterValues);
 	GameThread_InitMIParameters(this, TextureParameterValues);
@@ -3156,6 +3136,28 @@ void UMaterialInstance::PostLoad()
 
 	Super::PostLoad();
 
+#if WITH_EDITOR
+	//recalculate any scalar params based on a curve position in an atlas in case the atlas changed
+	for (FScalarParameterValue& ScalarParam : ScalarParameterValues)
+	{
+		if (ScalarParam.AtlasData.bIsUsedAsAtlasPosition)
+		{
+			UCurveLinearColorAtlas* Atlas = Cast<UCurveLinearColorAtlas>(ScalarParam.AtlasData.Atlas.Get());
+			UCurveLinearColor* Curve = Cast<UCurveLinearColor>(ScalarParam.AtlasData.Curve.Get());
+			if (Curve && Atlas)
+			{
+				Curve->ConditionalPostLoad();
+				Atlas->ConditionalPostLoad();
+				int32 Index = Atlas->GradientCurves.Find(Curve);
+				if (Index != INDEX_NONE)
+				{
+					ScalarParam.ParameterValue = (float)Index;
+				}
+			}
+		}
+	}
+#endif
+
 	if (FApp::CanEverRender())
 	{
 		// Resources can be processed / registered now that we're back on the main thread
