Re-instate fix of potential leak of AudioComponents when sounds fail to start and don't report back to the AudioComponent that playback finished (and failed)

diff --git a/Engine/Source/Runtime/Engine/Private/AudioDevice.cpp b/Engine/Source/Runtime/Engine/Private/AudioDevice.cpp
--- a/Engine/Source/Runtime/Engine/Private/AudioDevice.cpp
+++ b/Engine/Source/Runtime/Engine/Private/AudioDevice.cpp
@@ -4132,8 +4132,23 @@ void FAudioDevice::AddNewActiveSoundInternal(const FActiveSound& NewActiveSound,
 {
 	LLM_SCOPE(ELLMTag::AudioMisc);
 
-	if (NewActiveSound.Sound == nullptr)
+	if (!IsInAudioThread())
+	{
+		DECLARE_CYCLE_STAT(TEXT("FAudioThreadTask.AddNewActiveSound"), STAT_AudioAddNewActiveSound, STATGROUP_AudioThreadCommands);
+
+		FAudioDevice* AudioDevice = this;
+		FAudioThread::RunCommandOnAudioThread([AudioDevice, NewActiveSound, VirtualLoopToRetrigger]()
+		{
+			AudioDevice->AddNewActiveSoundInternal(NewActiveSound, VirtualLoopToRetrigger);
+		}, GET_STATID(STAT_AudioAddNewActiveSound));
+
+		return;
+	}
+
+	USoundBase* Sound = NewActiveSound.GetSound();
+	if (Sound == nullptr)
 	{
+		ReportSoundFailedToStart(NewActiveSound.AudioComponentID, VirtualLoopToRetrigger);
 		return;
 	}
 
@@ -4143,34 +4158,24 @@ void FAudioDevice::AddNewActiveSoundInternal(const FActiveSound& NewActiveSound,
 		USoundSourceBus* Bus = Cast<USoundSourceBus>(NewActiveSound.Sound);
 		if (Bus)
 		{
+			ReportSoundFailedToStart(NewActiveSound.AudioComponentID, VirtualLoopToRetrigger);
 			return;
 		}
 	}
 
-	if (NewActiveSound.Sound->GetDuration() <= FMath::Max(0.0f, SoundDistanceOptimizationLengthCVar))
+	if (Sound->GetDuration() <= FMath::Max(0.0f, SoundDistanceOptimizationLengthCVar))
 	{
 		// TODO: Determine if this check has already been completed at AudioComponent level and skip if so. Also,
 		// unify code paths determining if sound is audible.
 		if (!SoundIsAudible(NewActiveSound))
 		{
 			UE_LOG(LogAudio, Log, TEXT("New ActiveSound not created for out of range Sound %s"), *NewActiveSound.Sound->GetName());
+
+			ReportSoundFailedToStart(NewActiveSound.AudioComponentID, VirtualLoopToRetrigger);
 			return;
 		}
 	}
 
-	if (!IsInAudioThread())
-	{
-		DECLARE_CYCLE_STAT(TEXT("FAudioThreadTask.AddNewActiveSound"), STAT_AudioAddNewActiveSound, STATGROUP_AudioThreadCommands);
-
-		FAudioDevice* AudioDevice = this;
-		FAudioThread::RunCommandOnAudioThread([AudioDevice, NewActiveSound, VirtualLoopToRetrigger]()
-		{
-			AudioDevice->AddNewActiveSoundInternal(NewActiveSound, VirtualLoopToRetrigger);
-		}, GET_STATID(STAT_AudioAddNewActiveSound));
-
-		return;
-	}
-
 	// Cull one-shot active sounds if we've reached our max limit of one shot active sounds before we attempt to evaluate concurrency
 	// Check for debug sound name
 #if !UE_BUILD_SHIPPING
@@ -4183,20 +4188,21 @@ void FAudioDevice::AddNewActiveSoundInternal(const FActiveSound& NewActiveSound,
 		NewActiveSound.Sound->GetName(SoundName);
 		if (!SoundName.Contains(DebugSound))
 		{
+			ReportSoundFailedToStart(NewActiveSound.AudioComponentID, VirtualLoopToRetrigger);
 			return;
 		}
 	}
-#endif
-
-	const USoundBase* Sound = NewActiveSound.GetSound();
-	check(Sound);
+#endif // !UE_BUILD_SHIPPING
 
 	const USoundWave* SoundWave = Cast<USoundWave>(Sound);
 	if (SoundWave && SoundWave->bProcedural && SoundWave->IsGeneratingAudio())
 	{
 		FString SoundWaveName;
 		SoundWave->GetName(SoundWaveName);
-		UE_LOG(LogAudio, Warning, TEXT("Replaying a procedural sound '%s' without stopping the previous instance. Only one sound instance per procedural sound wave is supported."), *SoundWaveName)
+
+		UE_LOG(LogAudio, Warning, TEXT("Replaying a procedural sound '%s' without stopping the previous instance. Only one sound instance per procedural sound wave is supported."), *SoundWaveName);
+
+		ReportSoundFailedToStart(NewActiveSound.AudioComponentID, VirtualLoopToRetrigger);
 		return;
 	}
 
@@ -4297,8 +4303,26 @@ void FAudioDevice::AddNewActiveSoundInternal(const FActiveSound& NewActiveSound,
 	}
 }
 
+void FAudioDevice::ReportSoundFailedToStart(const uint64 AudioComponentID, FAudioVirtualLoop* VirtualLoop)
+{
+	check(IsInAudioThread());
+
+	if (VirtualLoop)
+	{
+		FActiveSound& VirtualActiveSound = VirtualLoop->GetActiveSound();
+		AddSoundToStop(&VirtualActiveSound);
+	}
+	else
+	{
+		const bool bFailedToStart = true;
+		UAudioComponent::PlaybackCompleted(AudioComponentID, bFailedToStart);
+	}
+}
+
 void FAudioDevice::RetriggerVirtualLoop(FAudioVirtualLoop& VirtualLoopToRetrigger)
 {
+	check(IsInAudioThread());
+
 	AddNewActiveSoundInternal(VirtualLoopToRetrigger.GetActiveSound(), &VirtualLoopToRetrigger);
 }
 
@@ -4343,10 +4367,7 @@ bool FAudioDevice::RemoveVirtualLoop(FActiveSound& InActiveSound)
 		check(InActiveSound.bIsStopping);
 
 		const uint64 ComponentID = InActiveSound.GetAudioComponentID();
-		if (ComponentID > 0)
-		{
-			UAudioComponent::PlaybackCompleted(ComponentID, false);
-		}
+		UAudioComponent::PlaybackCompleted(ComponentID, false);
 
 #if !(UE_BUILD_SHIPPING || UE_BUILD_TEST)
 		if (InActiveSound.Sound)
@@ -4553,10 +4574,7 @@ void FAudioDevice::RemoveActiveSound(FActiveSound* ActiveSound)
 
 	// Perform the notification if not sound not set to re-trigger
 	const uint64 ComponentID = ActiveSound->GetAudioComponentID();
-	if (ComponentID > 0)
-	{
-		UAudioComponent::PlaybackCompleted(ComponentID, false);
-	}
+	UAudioComponent::PlaybackCompleted(ComponentID, false);
 
 	const int32 NumRemoved = ActiveSounds.RemoveSwap(ActiveSound);
 	if (!ensureMsgf(NumRemoved > 0, TEXT("Attempting to remove an already removed ActiveSound '%s'"), ActiveSound->Sound ? *ActiveSound->Sound->GetName() : TEXT("N/A")))
@@ -5729,6 +5747,8 @@ float FAudioDevice::GetGameDeltaTime() const
 
 void FAudioDevice::UpdateVirtualLoops(bool bForceUpdate)
 {
+	check(IsInAudioThread());
+
 	if (FAudioVirtualLoop::IsEnabled())
 	{
 		TArray<FAudioVirtualLoop> VirtualLoopsToRetrigger;
diff --git a/Engine/Source/Runtime/Engine/Private/Components/AudioComponent.cpp b/Engine/Source/Runtime/Engine/Private/Components/AudioComponent.cpp
--- a/Engine/Source/Runtime/Engine/Private/Components/AudioComponent.cpp
+++ b/Engine/Source/Runtime/Engine/Private/Components/AudioComponent.cpp
@@ -739,6 +739,11 @@ void UAudioComponent::PlaybackCompleted(uint64 AudioComponentID, bool bFailedToS
 {
 	check(IsInAudioThread());
 
+	if (AudioComponentID == 0)
+	{
+		return;
+	}
+
 	DECLARE_CYCLE_STAT(TEXT("FGameThreadAudioTask.PlaybackCompleted"), STAT_AudioPlaybackCompleted, STATGROUP_TaskGraphTasks);
 
 	FAudioThread::RunCommandOnGameThread([AudioComponentID, bFailedToStart]()
diff --git a/Engine/Source/Runtime/Engine/Public/AudioDevice.h b/Engine/Source/Runtime/Engine/Public/AudioDevice.h
--- a/Engine/Source/Runtime/Engine/Public/AudioDevice.h
+++ b/Engine/Source/Runtime/Engine/Public/AudioDevice.h
@@ -1221,6 +1221,11 @@ private:
 	 */
 	void AddNewActiveSoundInternal(const FActiveSound& ActiveSound, FAudioVirtualLoop* VirtualLoop);
 
+	/**
+	 * Reports if a sound fails to start when attempting to create a new active sound.
+	 */
+	void ReportSoundFailedToStart(const uint64 AudioComponentID, FAudioVirtualLoop* VirtualLoop);
+
 	/**
 	* Initializes all plugin listeners belonging to this audio device.
 	* Called in the game thread.
