When sound fails to play due to concurrency and is a new sound (non-virtual loop) request, report that it failed to start.

diff --git a/Engine/Source/Runtime/Engine/Private/AudioDevice.cpp b/Engine/Source/Runtime/Engine/Private/AudioDevice.cpp
--- a/Engine/Source/Runtime/Engine/Private/AudioDevice.cpp
+++ b/Engine/Source/Runtime/Engine/Private/AudioDevice.cpp
@@ -4241,6 +4241,10 @@ void FAudioDevice::AddNewActiveSoundInternal(const FActiveSound& NewActiveSound,
 				UE_LOG(LogAudio, Verbose, TEXT("New ActiveSound %s Virtualizing: Failed to pass concurrency"), *Sound->GetName());
 				AddVirtualLoop(VirtualLoop);
 			}
+			else
+			{
+				ReportSoundFailedToStart(NewActiveSound.GetAudioComponentID(), VirtualLoopToRetrigger);
+			}
 		}
 		return;
 	}
