- Fix crash due to VirtualLoop attempting to realize, only to be shot down by concurrency and then prematurally reporting the sound as complete (as further realization attempts can be made after early outing).
- Clean-up code to be safer when bailing on AddNewActiveSoundInternal by adding ReportSoundFailedToStart, which queues virtual loop to be stopped/removed or reporting completion and failing on new sound request.

diff --git a/Engine/Source/Runtime/Engine/Private/Components/AudioComponent.cpp b/Engine/Source/Runtime/Engine/Private/Components/AudioComponent.cpp
--- a/Engine/Source/Runtime/Engine/Private/Components/AudioComponent.cpp
+++ b/Engine/Source/Runtime/Engine/Private/Components/AudioComponent.cpp
@@ -739,11 +739,6 @@ void UAudioComponent::PlaybackCompleted(uint64 AudioComponentID, bool bFailedToS
 {
 	check(IsInAudioThread());
 
-	if (AudioComponentID == 0)
-	{
-		return;
-	}
-
 	DECLARE_CYCLE_STAT(TEXT("FGameThreadAudioTask.PlaybackCompleted"), STAT_AudioPlaybackCompleted, STATGROUP_TaskGraphTasks);
 
 	FAudioThread::RunCommandOnGameThread([AudioComponentID, bFailedToStart]()
