Fixed compilation error when TEXTURE2DMIPMAP_USE_COMPACT_BULKDATA is true

diff --git a/Engine/Source/Runtime/Engine/Public/TextureResource.h b/Engine/Source/Runtime/Engine/Public/TextureResource.h
--- a/Engine/Source/Runtime/Engine/Public/TextureResource.h
+++ b/Engine/Source/Runtime/Engine/Public/TextureResource.h
@@ -64,6 +64,7 @@ struct FTexture2DMipMap
 		bool CanLoadFromDisk() const { return !IsInlined(); }
 		bool IsAvailableForUse() const { return !(BulkDataFlags & BULKDATA_Unused); }
 		bool IsBulkDataLoaded() const { return IsInlined(); }
+		bool IsAsyncLoadingComplete() const { return true; }
 		bool IsStoredCompressedOnDisk() const { return !!(BulkDataFlags & BULKDATA_SerializeCompressed); }
 		const void* LockReadOnly() const;
 		void* Lock(uint32 LockFlags);
