Replicated world settings properties need to be reset in RewindForReplay

diff --git a/Engine/Source/Runtime/Engine/Classes/GameFramework/WorldSettings.h b/Engine/Source/Runtime/Engine/Classes/GameFramework/WorldSettings.h
--- a/Engine/Source/Runtime/Engine/Classes/GameFramework/WorldSettings.h
+++ b/Engine/Source/Runtime/Engine/Classes/GameFramework/WorldSettings.h
@@ -796,6 +796,8 @@ public:
 	FORCEINLINE class APlayerState* GetPauserPlayerState() const { return PauserPlayerState; }
 	FORCEINLINE virtual void SetPauserPlayerState(class APlayerState* PlayerState) { PauserPlayerState = PlayerState; }
 
+	virtual void RewindForReplay() override;
+
 private:
 
 	// Hidden functions that don't make sense to use on this class.
diff --git a/Engine/Source/Runtime/Engine/Private/WorldSettings.cpp b/Engine/Source/Runtime/Engine/Private/WorldSettings.cpp
--- a/Engine/Source/Runtime/Engine/Private/WorldSettings.cpp
+++ b/Engine/Source/Runtime/Engine/Private/WorldSettings.cpp
@@ -887,4 +887,15 @@ FSoftClassPath AWorldSettings::GetAISystemClassName() const
 	return bEnableAISystem ? UAISystemBase::GetAISystemClassName() : FSoftClassPath();
 }
 
+void AWorldSettings::RewindForReplay()
+{
+	Super::RewindForReplay();
+
+	PauserPlayerState = nullptr;
+	TimeDilation = 1.0;
+	MatineeTimeDilation = 1.0;
+	bWorldGravitySet = false;
+	bHighPriorityLoading = false;
+}
+
 #undef LOCTEXT_NAMESPACE
