Cherry-pick 7630987: Added cast ray tracing shadows per material flag

diff --git a/Engine/Source/Runtime/Engine/Classes/Materials/Material.h b/Engine/Source/Runtime/Engine/Classes/Materials/Material.h
--- a/Engine/Source/Runtime/Engine/Classes/Materials/Material.h
+++ b/Engine/Source/Runtime/Engine/Classes/Materials/Material.h
@@ -849,6 +849,10 @@ public:
 	UPROPERTY(EditAnywhere, Category=Material)
 	uint32 bUseMaterialAttributes:1;
 
+	/** when true, the material casts ray tracing shadows. */
+	UPROPERTY(EditAnywhere, Category = Material)
+	uint32 bCastsRayTracedShadows : 1;
+
 	/** When true, translucent materials are fogged. Defaults to true. */
 	UPROPERTY(EditAnywhere, Category=Translucency, meta=(DisplayName = "Apply Fogging"))
 	uint32 bUseTranslucencyVertexFog:1;
@@ -1024,6 +1028,7 @@ public:
 	ENGINE_API virtual bool IsUIMaterial() const { return MaterialDomain == MD_UI; }
 	ENGINE_API virtual bool IsPostProcessMaterial() const { return MaterialDomain == MD_PostProcess; }
 	ENGINE_API virtual USubsurfaceProfile* GetSubsurfaceProfile_Internal() const override;
+	ENGINE_API virtual bool CastsRayTracedShadows() const override;
 
 	ENGINE_API void SetShadingModel(EMaterialShadingModel NewModel) { ensure(ShadingModel < MSM_NUM); ShadingModel = NewModel; ShadingModels = FMaterialShadingModelField(ShadingModel); }
 
diff --git a/Engine/Source/Runtime/Engine/Classes/Materials/MaterialInstance.h b/Engine/Source/Runtime/Engine/Classes/Materials/MaterialInstance.h
--- a/Engine/Source/Runtime/Engine/Classes/Materials/MaterialInstance.h
+++ b/Engine/Source/Runtime/Engine/Classes/Materials/MaterialInstance.h
@@ -467,6 +467,7 @@ public:
 	ENGINE_API virtual bool IsMasked() const override;
 	
 	ENGINE_API virtual USubsurfaceProfile* GetSubsurfaceProfile_Internal() const override;
+	ENGINE_API virtual bool CastsRayTracedShadows() const override;
 
 	/** Checks to see if an input property should be active, based on the state of the material */
 	ENGINE_API virtual bool IsPropertyActive(EMaterialProperty InProperty) const override;
diff --git a/Engine/Source/Runtime/Engine/Classes/Materials/MaterialInterface.h b/Engine/Source/Runtime/Engine/Classes/Materials/MaterialInterface.h
--- a/Engine/Source/Runtime/Engine/Classes/Materials/MaterialInterface.h
+++ b/Engine/Source/Runtime/Engine/Classes/Materials/MaterialInterface.h
@@ -723,6 +723,7 @@ public:
 	ENGINE_API virtual bool IsDeferredDecal() const;
 
 	ENGINE_API virtual USubsurfaceProfile* GetSubsurfaceProfile_Internal() const;
+	ENGINE_API virtual bool CastsRayTracedShadows() const;
 
 	/**
 	 * Force the streaming system to disregard the normal logic for the specified duration and
diff --git a/Engine/Source/Runtime/Engine/Private/Materials/Material.cpp b/Engine/Source/Runtime/Engine/Private/Materials/Material.cpp
--- a/Engine/Source/Runtime/Engine/Private/Materials/Material.cpp
+++ b/Engine/Source/Runtime/Engine/Private/Materials/Material.cpp
@@ -872,6 +872,7 @@ UMaterial::UMaterial(const FObjectInitializer& ObjectInitializer)
 	bAutomaticallySetUsageInEditor = true;
 
 	bUseMaterialAttributes = false;
+	bCastsRayTracedShadows = true;
 	bUseTranslucencyVertexFog = true;
 	BlendableLocation = BL_AfterTonemapping;
 	BlendablePriority = 0;
@@ -5687,6 +5688,11 @@ USubsurfaceProfile* UMaterial::GetSubsurfaceProfile_Internal() const
 	return SubsurfaceProfile; 
 }
 
+bool UMaterial::CastsRayTracedShadows() const
+{
+	return bCastsRayTracedShadows;
+}
+
 static bool IsPropertyActive_Internal(EMaterialProperty InProperty,
 	EMaterialDomain Domain,
 	EBlendMode BlendMode,
diff --git a/Engine/Source/Runtime/Engine/Private/Materials/MaterialInstance.cpp b/Engine/Source/Runtime/Engine/Private/Materials/MaterialInstance.cpp
--- a/Engine/Source/Runtime/Engine/Private/Materials/MaterialInstance.cpp
+++ b/Engine/Source/Runtime/Engine/Private/Materials/MaterialInstance.cpp
@@ -4121,6 +4121,12 @@ USubsurfaceProfile* UMaterialInstance::GetSubsurfaceProfile_Internal() const
 	return Parent ? Parent->GetSubsurfaceProfile_Internal() : 0;
 }
 
+bool UMaterialInstance::CastsRayTracedShadows() const
+{
+	//#dxr_todo: do per material instance override?
+	return Parent ? Parent->CastsRayTracedShadows() : true;
+}
+
 /** Checks to see if an input property should be active, based on the state of the material */
 bool UMaterialInstance::IsPropertyActive(EMaterialProperty InProperty) const
 {
diff --git a/Engine/Source/Runtime/Engine/Private/Materials/MaterialInterface.cpp b/Engine/Source/Runtime/Engine/Private/Materials/MaterialInterface.cpp
--- a/Engine/Source/Runtime/Engine/Private/Materials/MaterialInterface.cpp
+++ b/Engine/Source/Runtime/Engine/Private/Materials/MaterialInterface.cpp
@@ -445,6 +445,11 @@ USubsurfaceProfile* UMaterialInterface::GetSubsurfaceProfile_Internal() const
 	return NULL;
 }
 
+bool UMaterialInterface::CastsRayTracedShadows() const
+{
+	return true;
+}
+
 void UMaterialInterface::SetFeatureLevelToCompile(ERHIFeatureLevel::Type FeatureLevel, bool bShouldCompile)
 {
 	uint32 FeatureLevelBit = (1 << FeatureLevel);
diff --git a/Engine/Source/Runtime/Engine/Private/Materials/MaterialShared.cpp b/Engine/Source/Runtime/Engine/Private/Materials/MaterialShared.cpp
--- a/Engine/Source/Runtime/Engine/Private/Materials/MaterialShared.cpp
+++ b/Engine/Source/Runtime/Engine/Private/Materials/MaterialShared.cpp
@@ -1412,6 +1412,11 @@ bool FMaterialResource::HasRuntimeVirtualTextureOutput() const
 	return false;
 }
 
+bool FMaterialResource::CastsRayTracedShadows() const
+{
+	return Material->bCastsRayTracedShadows;
+}
+
 UMaterialInterface* FMaterialResource::GetMaterialInterface() const 
 { 
 	return MaterialInstance ? (UMaterialInterface*)MaterialInstance : (UMaterialInterface*)Material;
diff --git a/Engine/Source/Runtime/Engine/Private/RayTracingInstance.cpp b/Engine/Source/Runtime/Engine/Private/RayTracingInstance.cpp
--- a/Engine/Source/Runtime/Engine/Private/RayTracingInstance.cpp
+++ b/Engine/Source/Runtime/Engine/Private/RayTracingInstance.cpp
@@ -16,6 +16,7 @@ void FRayTracingInstance::BuildInstanceMaskAndFlags()
 
 	bool bAllSegmentsOpaque = true;
 	bool bAnySegmentsCastShadow = false;
+	bool bAllSegmentsCastShadow = true;
 
 	for (int32 SegmentIndex = 0; SegmentIndex < Materials.Num(); SegmentIndex++)
 	{
@@ -26,10 +27,11 @@ void FRayTracingInstance::BuildInstanceMaskAndFlags()
 		const EBlendMode BlendMode = Material.GetBlendMode();
 		Mask |= ComputeBlendModeMask(BlendMode);
 		bAllSegmentsOpaque &= BlendMode == BLEND_Opaque;
-		bAnySegmentsCastShadow |= MeshBatch.CastRayTracedShadow;
+		bAnySegmentsCastShadow |= MeshBatch.CastRayTracedShadow && Material.CastsRayTracedShadows();
+		bAllSegmentsCastShadow &= MeshBatch.CastRayTracedShadow && Material.CastsRayTracedShadows();
 	}
 
-	bForceOpaque = bAllSegmentsOpaque;
+	bForceOpaque = bAllSegmentsOpaque && bAllSegmentsCastShadow;
 	Mask |= bAnySegmentsCastShadow ? RAY_TRACING_MASK_SHADOW : 0;
 }
 
diff --git a/Engine/Source/Runtime/Engine/Public/MaterialShared.h b/Engine/Source/Runtime/Engine/Public/MaterialShared.h
--- a/Engine/Source/Runtime/Engine/Public/MaterialShared.h
+++ b/Engine/Source/Runtime/Engine/Public/MaterialShared.h
@@ -1509,7 +1509,8 @@ public:
 	virtual uint32 GetStencilRefValue() const { return 0; }
 	virtual uint32 GetStencilCompare() const { return 0; }
 	virtual bool HasRuntimeVirtualTextureOutput() const { return false; }
-	
+	virtual bool CastsRayTracedShadows() const { return true; }
+
 	/**
 	 * Should shaders compiled for this material be saved to disk?
 	 */
@@ -2252,6 +2253,7 @@ public:
 	ENGINE_API virtual bool ShouldApplyFogging() const override;
 	ENGINE_API virtual bool ComputeFogPerPixel() const override;
 	ENGINE_API virtual bool HasRuntimeVirtualTextureOutput() const override;
+	ENGINE_API virtual bool CastsRayTracedShadows() const override;
 	ENGINE_API virtual UMaterialInterface* GetMaterialInterface() const override;
 	/**
 	 * Should shaders compiled for this material be saved to disk?
diff --git a/Engine/Source/Runtime/Renderer/Private/RayTracing/RayTracingMaterialHitShaders.cpp b/Engine/Source/Runtime/Renderer/Private/RayTracing/RayTracingMaterialHitShaders.cpp
--- a/Engine/Source/Runtime/Renderer/Private/RayTracing/RayTracingMaterialHitShaders.cpp
+++ b/Engine/Source/Runtime/Renderer/Private/RayTracing/RayTracingMaterialHitShaders.cpp
@@ -204,7 +204,7 @@ void FRayTracingMeshProcessor::BuildRayTracingMeshCommands(
 
 	SharedCommand.SetShaders(PassShaders.GetUntypedShaders());
 	SharedCommand.InstanceMask = ComputeBlendModeMask(MaterialResource.GetBlendMode());
-	SharedCommand.bCastRayTracedShadows = MeshBatch.CastRayTracedShadow;
+	SharedCommand.bCastRayTracedShadows = MeshBatch.CastRayTracedShadow && MaterialResource.CastsRayTracedShadows();
 	SharedCommand.bOpaque = MaterialResource.GetBlendMode() == EBlendMode::BLEND_Opaque;
 
 	FVertexInputStreamArray VertexStreams;
