Fixed typo in 7631092

diff --git a/Engine/Source/Runtime/Engine/Classes/Materials/Material.h b/Engine/Source/Runtime/Engine/Classes/Materials/Material.h
--- a/Engine/Source/Runtime/Engine/Classes/Materials/Material.h
+++ b/Engine/Source/Runtime/Engine/Classes/Materials/Material.h
@@ -851,7 +851,7 @@ public:
 
 	/** when true, the material casts ray tracing shadows. */
 	UPROPERTY(EditAnywhere, Category = Material)
-	uint32 bCastsRayTracedShadows : 1;
+	uint32 bCastRayTracedShadows : 1;
 
 	/** When true, translucent materials are fogged. Defaults to true. */
 	UPROPERTY(EditAnywhere, Category=Translucency, meta=(DisplayName = "Apply Fogging"))
diff --git a/Engine/Source/Runtime/Engine/Private/Materials/Material.cpp b/Engine/Source/Runtime/Engine/Private/Materials/Material.cpp
--- a/Engine/Source/Runtime/Engine/Private/Materials/Material.cpp
+++ b/Engine/Source/Runtime/Engine/Private/Materials/Material.cpp
@@ -872,7 +872,7 @@ UMaterial::UMaterial(const FObjectInitializer& ObjectInitializer)
 	bAutomaticallySetUsageInEditor = true;
 
 	bUseMaterialAttributes = false;
-	bCastsRayTracedShadows = true;
+	bCastRayTracedShadows = true;
 	bUseTranslucencyVertexFog = true;
 	BlendableLocation = BL_AfterTonemapping;
 	BlendablePriority = 0;
@@ -5690,7 +5690,7 @@ USubsurfaceProfile* UMaterial::GetSubsurfaceProfile_Internal() const
 
 bool UMaterial::CastsRayTracedShadows() const
 {
-	return bCastsRayTracedShadows;
+	return bCastRayTracedShadows;
 }
 
 static bool IsPropertyActive_Internal(EMaterialProperty InProperty,
diff --git a/Engine/Source/Runtime/Engine/Private/Materials/MaterialShared.cpp b/Engine/Source/Runtime/Engine/Private/Materials/MaterialShared.cpp
--- a/Engine/Source/Runtime/Engine/Private/Materials/MaterialShared.cpp
+++ b/Engine/Source/Runtime/Engine/Private/Materials/MaterialShared.cpp
@@ -1414,7 +1414,7 @@ bool FMaterialResource::HasRuntimeVirtualTextureOutput() const
 
 bool FMaterialResource::CastsRayTracedShadows() const
 {
-	return Material->bCastsRayTracedShadows;
+	return Material->bCastRayTracedShadows;
 }
 
 UMaterialInterface* FMaterialResource::GetMaterialInterface() const 
