Hooked up water bodies to the world post processing volume system.  If water has enabled post process settings post processing automatically changes to those settingsa when underwater

diff --git a/Engine/Source/Runtime/Engine/Classes/Engine/World.h b/Engine/Source/Runtime/Engine/Classes/Engine/World.h
--- a/Engine/Source/Runtime/Engine/Classes/Engine/World.h
+++ b/Engine/Source/Runtime/Engine/Classes/Engine/World.h
@@ -1461,6 +1461,12 @@ public:
 	}
 #endif
 
+	/** Inserts a post process volume into the world in priority order */
+	void InsertPostProcessVolume(IInterface_PostProcessVolume* InVolume);
+
+	/** Removes a post process volume from the world */
+	void RemovePostProcessVolume(IInterface_PostProcessVolume* InVolume);
+
 	/** An array of post processing volumes, sorted in ascending order of priority.					*/
 	TArray< IInterface_PostProcessVolume * > PostProcessVolumes;
 
diff --git a/Engine/Source/Runtime/Engine/Private/UnrealClient.cpp b/Engine/Source/Runtime/Engine/Private/UnrealClient.cpp
--- a/Engine/Source/Runtime/Engine/Private/UnrealClient.cpp
+++ b/Engine/Source/Runtime/Engine/Private/UnrealClient.cpp
@@ -1395,29 +1395,6 @@ FRHIGPUMask FViewport::GetGPUMask(FRHICommandListImmediate& RHICmdList) const
 	return FRHIGPUMask::FromIndex(RHICmdList.GetViewportNextPresentGPUIndex(GetViewportRHI()));
 }
 
-void InsertVolume(IInterface_PostProcessVolume* Volume, TArray< IInterface_PostProcessVolume* >& VolumeArray)
-{
-	const int32 NumVolumes = VolumeArray.Num();
-	float TargetPriority = Volume->GetProperties().Priority;
-	int32 InsertIndex = 0;
-	// TODO: replace with binary search.
-	for (; InsertIndex < NumVolumes ; InsertIndex++)
-	{
-		IInterface_PostProcessVolume* CurrentVolume = VolumeArray[InsertIndex];
-		float CurrentPriority = CurrentVolume->GetProperties().Priority;
-
-		if (TargetPriority < CurrentPriority)
-		{
-			break;
-		}
-		if (CurrentVolume == Volume)
-		{
-			return;
-		}
-	}
-	VolumeArray.Insert(Volume, InsertIndex);
-}
-
 void APostProcessVolume::PostUnregisterAllComponents()
 {
 	// Route clear to super first.
@@ -1425,6 +1402,7 @@ void APostProcessVolume::PostUnregisterAllComponents()
 	// World will be NULL during exit purge.
 	if (GetWorld())
 	{
+		GetWorld()->RemovePostProcessVolume(this);
 		GetWorld()->PostProcessVolumes.RemoveSingle(this);
 	}
 }
@@ -1433,19 +1411,19 @@ void APostProcessVolume::PostRegisterAllComponents()
 {
 	// Route update to super first.
 	Super::PostRegisterAllComponents();
-	InsertVolume(this, GetWorld()->PostProcessVolumes);
+	GetWorld()->InsertPostProcessVolume(this);
 }
 
 void UPostProcessComponent::OnRegister()
 {
 	Super::OnRegister();
-	InsertVolume(this, GetWorld()->PostProcessVolumes);
+	GetWorld()->InsertPostProcessVolume(this);
 }
 
 void UPostProcessComponent::OnUnregister()
 {
 	Super::OnUnregister();
-	GetWorld()->PostProcessVolumes.RemoveSingle(this);
+	GetWorld()->RemovePostProcessVolume(this);
 }
 
 void UPostProcessComponent::Serialize(FArchive& Ar)
diff --git a/Engine/Source/Runtime/Engine/Private/World.cpp b/Engine/Source/Runtime/Engine/Private/World.cpp
--- a/Engine/Source/Runtime/Engine/Private/World.cpp
+++ b/Engine/Source/Runtime/Engine/Private/World.cpp
@@ -124,6 +124,7 @@
 #include "Engine/AssetManager.h"
 #include "Engine/HLODProxy.h"
 #include "ProfilingDebugging/CsvProfiler.h"
+#include "Interfaces/Interface_PostProcessVolume.h"
 
 #if INCLUDE_CHAOS
 #include "ChaosSolversModule.h"
@@ -7460,6 +7461,34 @@ FPrimaryAssetId UWorld::GetPrimaryAssetId() const
 	return FPrimaryAssetId();
 }
 
+void UWorld::InsertPostProcessVolume(IInterface_PostProcessVolume* InVolume)
+{
+	const int32 NumVolumes = PostProcessVolumes.Num();
+	float TargetPriority = InVolume->GetProperties().Priority;
+	int32 InsertIndex = 0;
+	// TODO: replace with binary search.
+	for (; InsertIndex < NumVolumes; InsertIndex++)
+	{
+		IInterface_PostProcessVolume* CurrentVolume = PostProcessVolumes[InsertIndex];
+		float CurrentPriority = CurrentVolume->GetProperties().Priority;
+
+		if (TargetPriority < CurrentPriority)
+		{
+			break;
+		}
+		if (CurrentVolume == InVolume)
+		{
+			return;
+		}
+	}
+	PostProcessVolumes.Insert(InVolume, InsertIndex);
+}
+
+void UWorld::RemovePostProcessVolume(IInterface_PostProcessVolume* InVolume)
+{
+	PostProcessVolumes.RemoveSingle(InVolume);
+}
+
 /**
 * Dump visible actors in current world.
 */
