[NO NON-PII COMMIT MESSAGE]

diff --git a/Engine/Source/Runtime/Engine/Private/AudioDerivedData.cpp b/Engine/Source/Runtime/Engine/Private/AudioDerivedData.cpp
--- a/Engine/Source/Runtime/Engine/Private/AudioDerivedData.cpp
+++ b/Engine/Source/Runtime/Engine/Private/AudioDerivedData.cpp
@@ -332,9 +332,8 @@ class FStreamedAudioCacheDerivedDataWorker : public FNonAbandonableTask
 					MaxChunkSize = FPlatformCompressionUtilities::GetChunkSizeForCookOverrides(CompressionOverrides);
 					UE_LOG(LogAudio, Display, TEXT("Chunk size for %s: %d"), *SoundWave.GetFullName(), MaxChunkSize);
 				}
-
+				
 				check(FirstChunkSize != 0 && MaxChunkSize != 0);
-				check(FirstChunkSize <= (256 * 1024) && MaxChunkSize <= (256 * 1024));
 
 				if (AudioFormat->SplitDataForStreaming(CompressedBuffer, ChunkBuffers, FirstChunkSize, MaxChunkSize))
 				{
