- Clean up quality settings setter/initialization code

diff --git a/Engine/Source/Runtime/Engine/Private/AudioDevice.cpp b/Engine/Source/Runtime/Engine/Private/AudioDevice.cpp
--- a/Engine/Source/Runtime/Engine/Private/AudioDevice.cpp
+++ b/Engine/Source/Runtime/Engine/Private/AudioDevice.cpp
@@ -270,10 +270,11 @@ FAudioEffectsManager* FAudioDevice::CreateEffectsManager()
 	return new FAudioEffectsManager(this);
 }
 
-FAudioQualitySettings FAudioDevice::GetQualityLevelSettings()
+const FAudioQualitySettings& FAudioDevice::GetQualityLevelSettings()
 {
 	const UAudioSettings* AudioSettings = GetDefault<UAudioSettings>();
-	return AudioSettings->GetQualityLevelSettings(GEngine->GetGameUserSettings()->GetAudioQualityLevel());
+	const int32 QualityLevel = GEngine ? GEngine->GetGameUserSettings()->GetAudioQualityLevel() : 0;
+	return AudioSettings->GetQualityLevelSettings(QualityLevel);
 }
 
 bool FAudioDevice::Init(int32 InMaxChannels)
diff --git a/Engine/Source/Runtime/Engine/Private/AudioDeviceManager.cpp b/Engine/Source/Runtime/Engine/Private/AudioDeviceManager.cpp
--- a/Engine/Source/Runtime/Engine/Private/AudioDeviceManager.cpp
+++ b/Engine/Source/Runtime/Engine/Private/AudioDeviceManager.cpp
@@ -461,10 +461,15 @@ bool FAudioDeviceManager::CreateAudioDevice(bool bCreateNewDevice, FCreateAudioD
 
 	if (bRequiresInit)
 	{
+		// Set to highest max channels initially provided by any quality setting, so that
+		// setting to lower quality but potentially returning to higher quality later at
+		// runtime is supported.
 		const UAudioSettings* AudioSettings = GetDefault<UAudioSettings>();
-		if (OutResults.AudioDevice->Init(AudioSettings->GetHighestMaxChannels())) //-V595
+		const int32 HighestMaxChannels = AudioSettings ? AudioSettings->GetHighestMaxChannels() : 0;
+		if (OutResults.AudioDevice && OutResults.AudioDevice->Init(HighestMaxChannels))
 		{
-			OutResults.AudioDevice->SetMaxChannels(AudioSettings->GetQualityLevelSettings(GEngine->GetGameUserSettings()->GetAudioQualityLevel()).MaxChannels); //-V595
+			const FAudioQualitySettings& QualitySettings = OutResults.AudioDevice->GetQualityLevelSettings();
+			OutResults.AudioDevice->SetMaxChannels(QualitySettings.MaxChannels);
 		}
 		else
 		{
diff --git a/Engine/Source/Runtime/Engine/Public/AudioDevice.h b/Engine/Source/Runtime/Engine/Public/AudioDevice.h
--- a/Engine/Source/Runtime/Engine/Public/AudioDevice.h
+++ b/Engine/Source/Runtime/Engine/Public/AudioDevice.h
@@ -445,7 +445,7 @@ public:
 	}
 
 	/** Returns the quality settings used by the default audio settings. */
-	static FAudioQualitySettings GetQualityLevelSettings();
+	static const FAudioQualitySettings& GetQualityLevelSettings();
 
 	/**
 	 * Basic initialization of the platform agnostic layer of the audio system
