Making the cascade spawn at location default to autoactivate, which was the old behavior.

diff --git a/Engine/Source/Runtime/Engine/Classes/Kismet/GameplayStatics.h b/Engine/Source/Runtime/Engine/Classes/Kismet/GameplayStatics.h
--- a/Engine/Source/Runtime/Engine/Classes/Kismet/GameplayStatics.h
+++ b/Engine/Source/Runtime/Engine/Classes/Kismet/GameplayStatics.h
@@ -375,10 +375,10 @@ class ENGINE_API UGameplayStatics : public UBlueprintFunctionLibrary
 	 * @param bAutoActivate - Whether the component will be automatically activated on creation.
 	 */
 	UFUNCTION(BlueprintCallable, Category="Effects|Components|ParticleSystem", meta=(Keywords = "particle system", WorldContext="WorldContextObject", UnsafeDuringActorConstruction = "true"))
-	static UParticleSystemComponent* SpawnEmitterAtLocation(const UObject* WorldContextObject, UParticleSystem* EmitterTemplate, FVector Location, FRotator Rotation = FRotator::ZeroRotator, FVector Scale = FVector(1.f), bool bAutoDestroy = true, EPSCPoolMethod PoolingMethod = EPSCPoolMethod::None, bool bAutoActivate = false);
+	static UParticleSystemComponent* SpawnEmitterAtLocation(const UObject* WorldContextObject, UParticleSystem* EmitterTemplate, FVector Location, FRotator Rotation = FRotator::ZeroRotator, FVector Scale = FVector(1.f), bool bAutoDestroy = true, EPSCPoolMethod PoolingMethod = EPSCPoolMethod::None, bool bAutoActivate = true);
 
 	// Backwards compatible version of SpawnEmitterAttached for C++ without Scale
-	static UParticleSystemComponent* SpawnEmitterAtLocation(const UObject* WorldContextObject, UParticleSystem* EmitterTemplate, FVector Location, FRotator Rotation, bool bAutoDestroy, EPSCPoolMethod PoolingMethod = EPSCPoolMethod::None, bool bAutoActivate=false);
+	static UParticleSystemComponent* SpawnEmitterAtLocation(const UObject* WorldContextObject, UParticleSystem* EmitterTemplate, FVector Location, FRotator Rotation, bool bAutoDestroy, EPSCPoolMethod PoolingMethod = EPSCPoolMethod::None, bool bAutoActivate=true);
 
 	/** Plays the specified effect at the given location and rotation, fire and forget. The system will go away when the effect is complete. Does not replicate.
 	 * @param World - The World to spawn in
