Fix a false alarm generated by PVS-Studio

diff --git a/Engine/Source/Runtime/Engine/Private/SkeletalMeshLODRenderData.cpp b/Engine/Source/Runtime/Engine/Private/SkeletalMeshLODRenderData.cpp
--- a/Engine/Source/Runtime/Engine/Private/SkeletalMeshLODRenderData.cpp
+++ b/Engine/Source/Runtime/Engine/Private/SkeletalMeshLODRenderData.cpp
@@ -935,7 +935,7 @@ void FSkeletalMeshLODRenderData::Serialize(FArchive& Ar, UObject* Owner, int32 I
 	if (!StripFlags.IsDataStrippedForServer())
 	{
 		// set cpu skinning flag on the vertex buffer so that the resource arrays know if they need to be CPU accessible
-		bNeedsCPUAccess = ShouldKeepCPUResources(OwnerMesh, Idx, bForceKeepCPUResources);
+		bNeedsCPUAccess = ShouldKeepCPUResources(OwnerMesh, Idx, bForceKeepCPUResources); // -V595
 	}
 
 	if (FPlatformProperties::RequiresCookedData())
