Disable LOD clamp for skeletal meshes that are not streaming LODs

diff --git a/Engine/Source/Runtime/Engine/Private/Components/SkinnedMeshComponent.cpp b/Engine/Source/Runtime/Engine/Private/Components/SkinnedMeshComponent.cpp
--- a/Engine/Source/Runtime/Engine/Private/Components/SkinnedMeshComponent.cpp
+++ b/Engine/Source/Runtime/Engine/Private/Components/SkinnedMeshComponent.cpp
@@ -2950,7 +2950,9 @@ bool USkinnedMeshComponent::UpdateLODStatus_Internal(int32 InMasterPoseComponent
 			}
 		}
 
-		if (MeshObject)
+		// This clamp is needed for meshes with LODs streamed but doesn't work well
+		// with those that have forced LOD level. Need to think of a better solution.
+		if (SkeletalMesh->bIsStreamable && MeshObject)
 		{
 			NewPredictedLODLevel = FMath::Max(NewPredictedLODLevel, MeshObject->MinDesiredLODLevel);
 		}
