Changing the field bAutoActivate name so that we don't get the warning of having been saved with the bad default value. The name is slightly more reflective of what it is doing anyway.

diff --git a/Engine/Source/Runtime/Engine/Classes/Kismet/GameplayStatics.h b/Engine/Source/Runtime/Engine/Classes/Kismet/GameplayStatics.h
--- a/Engine/Source/Runtime/Engine/Classes/Kismet/GameplayStatics.h
+++ b/Engine/Source/Runtime/Engine/Classes/Kismet/GameplayStatics.h
@@ -375,10 +375,10 @@ class ENGINE_API UGameplayStatics : public UBlueprintFunctionLibrary
 	 * @param bAutoActivate - Whether the component will be automatically activated on creation.
 	 */
 	UFUNCTION(BlueprintCallable, Category="Effects|Components|ParticleSystem", meta=(Keywords = "particle system", WorldContext="WorldContextObject", UnsafeDuringActorConstruction = "true"))
-	static UParticleSystemComponent* SpawnEmitterAtLocation(const UObject* WorldContextObject, UParticleSystem* EmitterTemplate, FVector Location, FRotator Rotation = FRotator::ZeroRotator, FVector Scale = FVector(1.f), bool bAutoDestroy = true, EPSCPoolMethod PoolingMethod = EPSCPoolMethod::None, bool bAutoActivate = true);
+	static UParticleSystemComponent* SpawnEmitterAtLocation(const UObject* WorldContextObject, UParticleSystem* EmitterTemplate, FVector Location, FRotator Rotation = FRotator::ZeroRotator, FVector Scale = FVector(1.f), bool bAutoDestroy = true, EPSCPoolMethod PoolingMethod = EPSCPoolMethod::None, bool bAutoActivateSystem = true);
 
 	// Backwards compatible version of SpawnEmitterAttached for C++ without Scale
-	static UParticleSystemComponent* SpawnEmitterAtLocation(const UObject* WorldContextObject, UParticleSystem* EmitterTemplate, FVector Location, FRotator Rotation, bool bAutoDestroy, EPSCPoolMethod PoolingMethod = EPSCPoolMethod::None, bool bAutoActivate=true);
+	static UParticleSystemComponent* SpawnEmitterAtLocation(const UObject* WorldContextObject, UParticleSystem* EmitterTemplate, FVector Location, FRotator Rotation, bool bAutoDestroy, EPSCPoolMethod PoolingMethod = EPSCPoolMethod::None, bool bAutoActivateSystem =true);
 
 	/** Plays the specified effect at the given location and rotation, fire and forget. The system will go away when the effect is complete. Does not replicate.
 	 * @param World - The World to spawn in
diff --git a/Engine/Source/Runtime/Engine/Private/GameplayStatics.cpp b/Engine/Source/Runtime/Engine/Private/GameplayStatics.cpp
--- a/Engine/Source/Runtime/Engine/Private/GameplayStatics.cpp
+++ b/Engine/Source/Runtime/Engine/Private/GameplayStatics.cpp
@@ -949,12 +949,12 @@ UParticleSystemComponent* CreateParticleSystem(UParticleSystem* EmitterTemplate,
 	return PSC;
 }
 
-UParticleSystemComponent* UGameplayStatics::SpawnEmitterAtLocation(const UObject* WorldContextObject, UParticleSystem* EmitterTemplate, FVector SpawnLocation, FRotator SpawnRotation, bool bAutoDestroy, EPSCPoolMethod PoolingMethod, bool bAutoActivate)
+UParticleSystemComponent* UGameplayStatics::SpawnEmitterAtLocation(const UObject* WorldContextObject, UParticleSystem* EmitterTemplate, FVector SpawnLocation, FRotator SpawnRotation, bool bAutoDestroy, EPSCPoolMethod PoolingMethod, bool bAutoActivateSystem)
 {
-	return SpawnEmitterAtLocation(WorldContextObject, EmitterTemplate, SpawnLocation, SpawnRotation, FVector(1.f), bAutoDestroy, PoolingMethod, bAutoActivate);
+	return SpawnEmitterAtLocation(WorldContextObject, EmitterTemplate, SpawnLocation, SpawnRotation, FVector(1.f), bAutoDestroy, PoolingMethod, bAutoActivateSystem);
 }
 
-UParticleSystemComponent* UGameplayStatics::InternalSpawnEmitterAtLocation(UWorld* World, UParticleSystem* EmitterTemplate, FVector SpawnLocation, FRotator SpawnRotation, FVector SpawnScale, bool bAutoDestroy, EPSCPoolMethod PoolingMethod, bool bAutoActivate)
+UParticleSystemComponent* UGameplayStatics::InternalSpawnEmitterAtLocation(UWorld* World, UParticleSystem* EmitterTemplate, FVector SpawnLocation, FRotator SpawnRotation, FVector SpawnScale, bool bAutoDestroy, EPSCPoolMethod PoolingMethod, bool bAutoActivateSystem)
 {
 	check(World && EmitterTemplate);
 
@@ -968,7 +968,7 @@ UParticleSystemComponent* UGameplayStatics::InternalSpawnEmitterAtLocation(UWorl
 		PSC->RelativeScale3D = SpawnScale;
 
 		PSC->RegisterComponentWithWorld(World);
-		if (bAutoActivate)
+		if (bAutoActivateSystem)
 		{
 			PSC->ActivateSystem(true);
 		}
@@ -992,35 +992,35 @@ UParticleSystemComponent* UGameplayStatics::InternalSpawnEmitterAtLocation(UWorl
 	return nullptr;
 }
 
-UParticleSystemComponent* UGameplayStatics::SpawnEmitterAtLocation(const UObject* WorldContextObject, UParticleSystem* EmitterTemplate, FVector SpawnLocation, FRotator SpawnRotation, FVector SpawnScale, bool bAutoDestroy, EPSCPoolMethod PoolingMethod, bool bAutoActivate)
+UParticleSystemComponent* UGameplayStatics::SpawnEmitterAtLocation(const UObject* WorldContextObject, UParticleSystem* EmitterTemplate, FVector SpawnLocation, FRotator SpawnRotation, FVector SpawnScale, bool bAutoDestroy, EPSCPoolMethod PoolingMethod, bool bAutoActivateSystem)
 {
 	UParticleSystemComponent* PSC = nullptr;
 	if (EmitterTemplate)
 	{
 		if (UWorld* World = GEngine->GetWorldFromContextObject(WorldContextObject, EGetWorldErrorMode::LogAndReturnNull))
 		{
-			PSC = InternalSpawnEmitterAtLocation(World, EmitterTemplate, SpawnLocation, SpawnRotation, SpawnScale, bAutoDestroy, PoolingMethod, bAutoActivate);
+			PSC = InternalSpawnEmitterAtLocation(World, EmitterTemplate, SpawnLocation, SpawnRotation, SpawnScale, bAutoDestroy, PoolingMethod, bAutoActivateSystem);
 		}
 	}
 	return PSC;
 }
 
-UParticleSystemComponent* UGameplayStatics::SpawnEmitterAtLocation(UWorld* World, UParticleSystem* EmitterTemplate, const FTransform& SpawnTransform, bool bAutoDestroy, EPSCPoolMethod PoolingMethod, bool bAutoActivate)
+UParticleSystemComponent* UGameplayStatics::SpawnEmitterAtLocation(UWorld* World, UParticleSystem* EmitterTemplate, const FTransform& SpawnTransform, bool bAutoDestroy, EPSCPoolMethod PoolingMethod, bool bAutoActivateSystem)
 {
 	UParticleSystemComponent* PSC = nullptr;
 	if (World && EmitterTemplate)
 	{
-		PSC = InternalSpawnEmitterAtLocation(World, EmitterTemplate, SpawnTransform.GetLocation(), SpawnTransform.GetRotation().Rotator(), SpawnTransform.GetScale3D(), bAutoDestroy, PoolingMethod, bAutoActivate);
+		PSC = InternalSpawnEmitterAtLocation(World, EmitterTemplate, SpawnTransform.GetLocation(), SpawnTransform.GetRotation().Rotator(), SpawnTransform.GetScale3D(), bAutoDestroy, PoolingMethod, bAutoActivateSystem);
 	}
 	return PSC;
 }
 
-UParticleSystemComponent* UGameplayStatics::SpawnEmitterAttached(UParticleSystem* EmitterTemplate, USceneComponent* AttachToComponent, FName AttachPointName, FVector Location, FRotator Rotation, EAttachLocation::Type LocationType, bool bAutoDestroy, EPSCPoolMethod PoolingMethod, bool bAutoActivate)
+UParticleSystemComponent* UGameplayStatics::SpawnEmitterAttached(UParticleSystem* EmitterTemplate, USceneComponent* AttachToComponent, FName AttachPointName, FVector Location, FRotator Rotation, EAttachLocation::Type LocationType, bool bAutoDestroy, EPSCPoolMethod PoolingMethod, bool bAutoActivateSystem)
 {
-	return SpawnEmitterAttached(EmitterTemplate, AttachToComponent, AttachPointName, Location, Rotation, FVector(1.f), LocationType, bAutoDestroy, PoolingMethod, bAutoActivate);
+	return SpawnEmitterAttached(EmitterTemplate, AttachToComponent, AttachPointName, Location, Rotation, FVector(1.f), LocationType, bAutoDestroy, PoolingMethod, bAutoActivateSystem);
 }
 
-UParticleSystemComponent* UGameplayStatics::SpawnEmitterAttached(UParticleSystem* EmitterTemplate, USceneComponent* AttachToComponent, FName AttachPointName, FVector Location, FRotator Rotation, FVector Scale, EAttachLocation::Type LocationType, bool bAutoDestroy, EPSCPoolMethod PoolingMethod, bool bAutoActivate)
+UParticleSystemComponent* UGameplayStatics::SpawnEmitterAttached(UParticleSystem* EmitterTemplate, USceneComponent* AttachToComponent, FName AttachPointName, FVector Location, FRotator Rotation, FVector Scale, EAttachLocation::Type LocationType, bool bAutoDestroy, EPSCPoolMethod PoolingMethod, bool bAutoActivateSystem)
 {
 	UParticleSystemComponent* PSC = nullptr;
 	if (EmitterTemplate)
@@ -1068,7 +1068,7 @@ UParticleSystemComponent* UGameplayStatics::SpawnEmitterAttached(UParticleSystem
 					}
 
 					PSC->RegisterComponentWithWorld(World);
-					if(bAutoActivate)
+					if(bAutoActivateSystem)
 					{
 						PSC->ActivateSystem(true);
 					}
