Moved water spline metadata to the water actor to avoid blueprints recompiling and sometimes wiping out the metadata (due to it not being outered to the actor).

diff --git a/Engine/Source/Runtime/Core/Public/UObject/FortniteMainBranchObjectVersion.h b/Engine/Source/Runtime/Core/Public/UObject/FortniteMainBranchObjectVersion.h
--- a/Engine/Source/Runtime/Core/Public/UObject/FortniteMainBranchObjectVersion.h
+++ b/Engine/Source/Runtime/Core/Public/UObject/FortniteMainBranchObjectVersion.h
@@ -93,6 +93,9 @@ struct CORE_API FFortniteMainBranchObjectVersion
 		// Fixup bad defaults in water metadata
 		FixUpWaterMetadata,
 
+		// Move the location of water metadata
+		MoveWaterMetadataToActor,
+
 		// -----<new versions can be added above this line>-------------------------------------------------
 		VersionPlusOne,
 		LatestVersion = VersionPlusOne - 1
diff --git a/Engine/Source/Runtime/Engine/Classes/Components/SplineComponent.h b/Engine/Source/Runtime/Engine/Classes/Components/SplineComponent.h
--- a/Engine/Source/Runtime/Engine/Classes/Components/SplineComponent.h
+++ b/Engine/Source/Runtime/Engine/Classes/Components/SplineComponent.h
@@ -73,8 +73,8 @@ struct ENGINE_API FSplineCurves
 	UPROPERTY()
 	FInterpCurveFloat ReparamTable;
 
-	UPROPERTY(Instanced)
-	USplineMetadata* Metadata;
+	UPROPERTY()
+	USplineMetadata* Metadata_DEPRECATED;
 
 	bool operator==(const FSplineCurves& Other) const
 	{
@@ -188,13 +188,13 @@ UCLASS(ClassGroup=Utility, ShowCategories = (Mobility), HideCategories = (Physic
 class ENGINE_API USplineComponent : public UPrimitiveComponent
 {
 	GENERATED_UCLASS_BODY()
-    
-    /** Child class can optionally provide their metadata object through this constructor */
-	USplineComponent(const FObjectInitializer& ObjectInitializer, USplineMetadata* Metadata);
 
 	UPROPERTY(EditAnywhere, Category=Points)
 	FSplineCurves SplineCurves;
 
+	UPROPERTY()
+	USplineMetadata* SplineMetadata;
+
 	/** Deprecated - please use GetSplinePointsPosition() to fetch this FInterpCurve */
 	UPROPERTY()
 	FInterpCurveVector SplineInfo_DEPRECATED;
@@ -316,8 +316,10 @@ public:
 	const FInterpCurveQuat& GetSplinePointsRotation() const { return SplineCurves.Rotation; }
 	FInterpCurveVector& GetSplinePointsScale() { return SplineCurves.Scale; }
 	const FInterpCurveVector& GetSplinePointsScale() const { return SplineCurves.Scale; }
-	USplineMetadata* GetSplinePointsMetadata() { return SplineCurves.Metadata; }
-	const USplineMetadata* GetSplinePointsMetadata() const { return SplineCurves.Metadata; }
+	USplineMetadata* GetSplinePointsMetadata() { return SplineMetadata; }
+	const USplineMetadata* GetSplinePointsMetadata() const { return SplineMetadata; }
+
+	void SetSplineMetadata(USplineMetadata* InMetadata) { SplineMetadata = InMetadata; }
 
 	void ApplyComponentInstanceData(struct FSplineInstanceData* ComponentInstanceData, const bool bPostUCS);
 
diff --git a/Engine/Source/Runtime/Engine/Private/Components/SplineComponent.cpp b/Engine/Source/Runtime/Engine/Private/Components/SplineComponent.cpp
--- a/Engine/Source/Runtime/Engine/Private/Components/SplineComponent.cpp
+++ b/Engine/Source/Runtime/Engine/Private/Components/SplineComponent.cpp
@@ -24,12 +24,6 @@ USplineMetadata::USplineMetadata(const FObjectInitializer& ObjectInitializer)
 }
 
 USplineComponent::USplineComponent(const FObjectInitializer& ObjectInitializer)
- : USplineComponent(ObjectInitializer, nullptr)
-{
-
-}
-
-USplineComponent::USplineComponent(const FObjectInitializer& ObjectInitializer, USplineMetadata* Metadata)
 	: Super(ObjectInitializer)
 	, bAllowSplineEditingPerInstance_DEPRECATED(true)
 	, ReparamStepsPerSegment(10)
@@ -49,8 +43,6 @@ USplineComponent::USplineComponent(const FObjectInitializer& ObjectInitializer,
 	, ScaleVisualizationWidth(30.0f)
 #endif
 {
-	SplineCurves.Metadata = Metadata;
-
 	SplineCurves.Position.Points.Reset(10);
 	SplineCurves.Rotation.Points.Reset(10);
 	SplineCurves.Scale.Points.Reset(10);
@@ -70,12 +62,6 @@ USplineComponent::USplineComponent(const FObjectInitializer& ObjectInitializer,
 	}
 #endif
 
-	if (SplineCurves.Metadata)
-	{
-		const int32 NumPoints = GetNumberOfSplinePoints();
-		SplineCurves.Metadata->Fixup(NumPoints);
-	}
-
 	UpdateSpline();
 
 	// Set these deprecated values up so that old assets with default values load correctly (and are subsequently upgraded during Serialize)
@@ -162,12 +148,6 @@ void USplineComponent::Serialize(FArchive& Ar)
 void USplineComponent::PostLoad()
 {
 	Super::PostLoad();
-
-	if (USplineMetadata* MetaData = GetSplinePointsMetadata())
-	{
-		const int32 NumPoints = GetNumberOfSplinePoints();
-		MetaData->Fixup(NumPoints);
-	}
 }
 
 void FSplineCurves::UpdateSpline(bool bClosedLoop, bool bStationaryEndpoints, int32 ReparamStepsPerSegment, bool bLoopPositionOverride, float LoopPosition, const FVector& Scale3D)
