Fix for dsync when using root motion on relevancy boundary. From Dan Youhon

#rxn

diff --git a/Engine/Source/Runtime/Engine/Private/Components/CharacterMovementComponent.cpp b/Engine/Source/Runtime/Engine/Private/Components/CharacterMovementComponent.cpp
--- a/Engine/Source/Runtime/Engine/Private/Components/CharacterMovementComponent.cpp
+++ b/Engine/Source/Runtime/Engine/Private/Components/CharacterMovementComponent.cpp
@@ -2153,12 +2153,18 @@ void UCharacterMovementComponent::PerformMovement(float DeltaSeconds)
 	// no movement if we can't move, or if currently doing physical simulation on UpdatedComponent
 	if (MovementMode == MOVE_None || UpdatedComponent->Mobility != EComponentMobility::Movable || UpdatedComponent->IsSimulatingPhysics())
 	{
-		if (!CharacterOwner->bClientUpdating && !CharacterOwner->bServerMoveIgnoreRootMotion && CharacterOwner->IsPlayingRootMotion() && CharacterOwner->GetMesh())
+		if (!CharacterOwner->bClientUpdating && !CharacterOwner->bServerMoveIgnoreRootMotion)
 		{
 			// Consume root motion
-			TickCharacterPose(DeltaSeconds);
-			RootMotionParams.Clear();
-			CurrentRootMotion.Clear();
+			if (CharacterOwner->IsPlayingRootMotion() && CharacterOwner->GetMesh())
+			{
+				TickCharacterPose(DeltaSeconds);
+				RootMotionParams.Clear();
+			}
+			if (CurrentRootMotion.HasActiveRootMotionSources())
+			{
+				CurrentRootMotion.Clear();
+			}
 		}
 		// Clear pending physics forces
 		ClearAccumulatedForces();
