Unmapped all components owned by Actors when the ActorChannel is deleted.
This prevents issues where component references may point to the PendingKill version of the Actor / Component.

diff --git a/Engine/Source/Runtime/Engine/Private/DataChannel.cpp b/Engine/Source/Runtime/Engine/Private/DataChannel.cpp
--- a/Engine/Source/Runtime/Engine/Private/DataChannel.cpp
+++ b/Engine/Source/Runtime/Engine/Private/DataChannel.cpp
@@ -1881,6 +1881,13 @@ void UActorChannel::DestroyActorAndComponents()
 	// Destroy the actor
 	if ( Actor != NULL )
 	{
+		// Unmap any components in this actor. This will make sure that once the Actor is remapped
+		// any references to components will be remapped as well.
+		for (UActorComponent* Component : Actor->GetComponents())
+		{
+			MoveMappedObjectToUnmapped(Component);
+		}
+
 		// Unmap this object so we can remap it if it becomes relevant again in the future
 		MoveMappedObjectToUnmapped( Actor );
 
