EditorUI: remove ClampUVs function that incorrectly translates UVs outside 0-1 range when drawing in StaticMeshEditor
#rnx

diff --git a/Engine/Source/Runtime/Engine/Private/PrimitiveDrawingUtils.cpp b/Engine/Source/Runtime/Engine/Private/PrimitiveDrawingUtils.cpp
--- a/Engine/Source/Runtime/Engine/Private/PrimitiveDrawingUtils.cpp
+++ b/Engine/Source/Runtime/Engine/Private/PrimitiveDrawingUtils.cpp
@@ -1618,33 +1618,6 @@ void ApplyViewModeOverrides(
 #endif
 }
 
-void ClampUVs(FVector2D* UVs, int32 NumUVs)
-{
-	const float FudgeFactor = 0.1f;
-	FVector2D Bias(0.0f,0.0f);
-
-	float MinU = UVs[0].X;
-	float MinV = UVs[0].Y;
-	for (int32 i = 1; i < NumUVs; ++i)
-	{
-		MinU = FMath::Min(MinU,UVs[i].X);
-		MinV = FMath::Min(MinU,UVs[i].Y);
-	}
-
-	if (MinU < -FudgeFactor || MinU > (1.0f+FudgeFactor))
-	{
-		Bias.X = FMath::FloorToFloat(MinU);
-	}
-	if (MinV < -FudgeFactor || MinV > (1.0f+FudgeFactor))
-	{
-		Bias.Y = FMath::FloorToFloat(MinV);
-	}
-
-	for (int32 i = 0; i < NumUVs; i++)
-	{
-		UVs[i] += Bias;
-	}
-}
 
 bool IsUVOutOfBounds(FVector2D UV)
 {
@@ -1698,9 +1671,6 @@ void DrawUVsInternal(FViewport* InViewport, FCanvas* InCanvas, int32 InTextYPos,
 					bOutOfBounds[Corner] = IsUVOutOfBounds(UVs[Corner]);
 				}
 
-				// Clamp the UV triangle to the [0,1] range (with some fudge).
-				ClampUVs(UVs, 3);
-
 				for (int32 Edge = 0; Edge < 3; Edge++)
 				{
 					int32 Corner1 = Edge;
@@ -1724,7 +1694,6 @@ void DrawUVsInternal(FViewport* InViewport, FCanvas* InCanvas, int32 InTextYPos,
 					FVector2D UVs[2];
 					UVs[0] = (SelectedEdgeTexCoords[UVIndex]);
 					UVs[1] = (SelectedEdgeTexCoords[UVIndex + 1]);
-					ClampUVs(UVs, 2);
 
 					LineItem.Draw(InCanvas, UVs[0] * UVBoxScale + UVBoxOrigin, UVs[1] * UVBoxScale + UVBoxOrigin);
 				}
