Fix a crash when running automation tests in EngineTest

diff --git a/Engine/Source/Runtime/Engine/Private/Rendering/SkinWeightVertexBuffer.cpp b/Engine/Source/Runtime/Engine/Private/Rendering/SkinWeightVertexBuffer.cpp
--- a/Engine/Source/Runtime/Engine/Private/Rendering/SkinWeightVertexBuffer.cpp
+++ b/Engine/Source/Runtime/Engine/Private/Rendering/SkinWeightVertexBuffer.cpp
@@ -4,6 +4,7 @@
 #include "EngineUtils.h"
 #include "Rendering/SkeletalMeshLODModel.h"
 #include "RenderUtils.h"
+#include "SkeletalMeshTypes.h"
 
 ///
 
@@ -125,7 +126,16 @@ FArchive& operator<<(FArchive& Ar, FSkinWeightVertexBuffer& VertexBuffer)
 
 void FSkinWeightVertexBuffer::SerializeMetaData(FArchive& Ar)
 {
-	Ar << bExtraBoneInfluences << Stride << NumVertices;
+	Ar.UsingCustomVersion(FSkeletalMeshCustomVersion::GUID);
+	if (Ar.CustomVer(FSkeletalMeshCustomVersion::GUID) < FSkeletalMeshCustomVersion::SplitModelAndRenderData)
+	{
+		check(Ar.IsLoading());
+		Ar << bExtraBoneInfluences << NumVertices;
+	}
+	else
+	{
+		Ar << bExtraBoneInfluences << Stride << NumVertices;
+	}
 }
 
 void FSkinWeightVertexBuffer::CopyMetaData(const FSkinWeightVertexBuffer& Other)
