Fix MSVC static analysis warning

diff --git a/Engine/Source/Runtime/Engine/Private/SkeletalMesh.cpp b/Engine/Source/Runtime/Engine/Private/SkeletalMesh.cpp
--- a/Engine/Source/Runtime/Engine/Private/SkeletalMesh.cpp
+++ b/Engine/Source/Runtime/Engine/Private/SkeletalMesh.cpp
@@ -635,6 +635,7 @@ void USkeletalMesh::InitResources()
 	const int32 NumLODs = GetLODNum();
 	bIsStreamable = !NeverStream
 		&& NumLODs > 1
+		&& SkelMeshRenderData
 		&& !SkelMeshRenderData->LODRenderData[0].bStreamedDataInlined;
 
 	UnlinkStreaming();
