[NO NON-PII COMMIT MESSAGE]

diff --git a/Engine/Source/Runtime/Engine/Private/ADPCMAudioInfo.cpp b/Engine/Source/Runtime/Engine/Private/ADPCMAudioInfo.cpp
--- a/Engine/Source/Runtime/Engine/Private/ADPCMAudioInfo.cpp
+++ b/Engine/Source/Runtime/Engine/Private/ADPCMAudioInfo.cpp
@@ -16,6 +16,13 @@ FAutoConsoleVariableRef CVarDisableADPCMSeekLock(
 	TEXT("Disables ADPCM seek crit section fix for multiple seek requests per frame.\n"),
 	ECVF_Default);
 
+static int32 bDisableADPCMSeekingCVar = 0;
+FAutoConsoleVariableRef CVarDisableADPCMSeeking(
+	TEXT("au.adpcm.DisableSeeking"),
+	bDisableADPCMSeekingCVar,
+	TEXT("Disables seeking with ADPCM.\n"),
+	ECVF_Default);
+
 #define WAVE_FORMAT_LPCM  1
 #define WAVE_FORMAT_ADPCM 2
 
@@ -52,6 +59,11 @@ FADPCMAudioInfo::~FADPCMAudioInfo(void)
 
 void FADPCMAudioInfo::SeekToTime(const float InSeekTime)
 {
+	if (bDisableADPCMSeekingCVar)
+	{
+		return;
+	}
+
 	if (bDisableADPCMSeekLockCVar)
 	{
 		SeekToTimeInternal(InSeekTime);
@@ -114,6 +126,14 @@ void FADPCMAudioInfo::SeekToTimeInternal(const float InSeekTime)
 	}
 	else
 	{
+		if (StreamingSoundWave->GetNumChunks() == 0)
+		{
+			UE_LOG(LogAudio, Error, TEXT("Entered streaming seek path with a non-streaming sound!"));
+			return;
+		}
+
+		const uint32 TotalStreamingChunks = StreamingSoundWave->GetNumChunks();
+
 		if (Format == WAVE_FORMAT_ADPCM)
 		{
 			CurrentCompressedBlockIndex = TotalSamplesStreamed / SamplesPerBlock; // Compute the block index that where SeekTime resides.
@@ -131,6 +151,13 @@ void FADPCMAudioInfo::SeekToTimeInternal(const float InSeekTime)
 
 				// Always add chunks in NumChannels pairs
 				CurrentChunkBufferOffset += ChannelBlockSize;
+				
+				if (CurrentChunkIndex >= TotalStreamingChunks)
+				{
+					CurrentChunkIndex = 0;
+					CurrentChunkBufferOffset = 0;
+					break;
+				}
 			}
 		}
 		else if (Format == WAVE_FORMAT_LPCM)
@@ -145,6 +172,13 @@ void FADPCMAudioInfo::SeekToTimeInternal(const float InSeekTime)
 			{
 				CurrentChunkBufferOffset -= StreamingSoundWave->GetSizeOfChunk(CurrentChunkIndex);
 				CurrentChunkIndex++;
+				
+				if (CurrentChunkIndex >= TotalStreamingChunks)
+				{
+					CurrentChunkIndex = 0;
+					CurrentChunkBufferOffset = 0;
+					break;
+				}
 			}
 
 			// 3. Trim remainder of block size, effectively aligning the block to a channel pair boundary
