Fix for async loading assert in timer manager.

diff --git a/Engine/Source/Runtime/Engine/Private/GameInstance.cpp b/Engine/Source/Runtime/Engine/Private/GameInstance.cpp
--- a/Engine/Source/Runtime/Engine/Private/GameInstance.cpp
+++ b/Engine/Source/Runtime/Engine/Private/GameInstance.cpp
@@ -37,11 +37,10 @@
 
 UGameInstance::UGameInstance(const FObjectInitializer& ObjectInitializer)
 	: Super(ObjectInitializer)
-	, TimerManager(new FTimerManager())
+	, TimerManager(new FTimerManager(this))
 	, LatentActionManager(new FLatentActionManager())
 	, SubsystemCollection(this)
 {
-	TimerManager->SetGameInstance(this);
 }
 
 void UGameInstance::FinishDestroy()
diff --git a/Engine/Source/Runtime/Engine/Private/TimerManager.cpp b/Engine/Source/Runtime/Engine/Private/TimerManager.cpp
--- a/Engine/Source/Runtime/Engine/Private/TimerManager.cpp
+++ b/Engine/Source/Runtime/Engine/Private/TimerManager.cpp
@@ -246,7 +246,7 @@ struct FTimerHeapOrder
 	int32 NumTimers;
 };
 
-FTimerManager::FTimerManager()
+FTimerManager::FTimerManager(UGameInstance* GameInstance)
 	: InternalTime(0.0)
 	, LastTickedFrame(static_cast<uint64>(-1))
 	, OwningGameInstance(nullptr)
@@ -256,6 +256,11 @@ FTimerManager::FTimerManager()
 		// Off by default, renable if needed
 		//FCoreDelegates::OnHandleSystemError.AddRaw(this, &FTimerManager::OnCrash);
 	}
+
+	if (GameInstance)
+	{
+		SetGameInstance(GameInstance);
+	}
 }
 
 FTimerManager::~FTimerManager()
@@ -954,9 +959,6 @@ TStatId FTimerManager::GetStatId() const
 
 void FTimerManager::SetGameInstance(UGameInstance* InGameInstance)
 {
-	// not currently threadsafe
-	check(IsInGameThread());
-
 	OwningGameInstance = InGameInstance;
 
 #if UE_ENABLE_TRACKING_TIMER_SOURCES
diff --git a/Engine/Source/Runtime/Engine/Public/TimerManager.h b/Engine/Source/Runtime/Engine/Public/TimerManager.h
--- a/Engine/Source/Runtime/Engine/Public/TimerManager.h
+++ b/Engine/Source/Runtime/Engine/Public/TimerManager.h
@@ -168,7 +168,7 @@ public:
 	// ----------------------------------
 	// Timer API
 
-	FTimerManager();
+	explicit FTimerManager(UGameInstance* GameInstance = nullptr);
 	virtual ~FTimerManager();
 
 	/**
@@ -391,7 +391,7 @@ public:
 	/** Debug command to output info on all timers currently set to the log. */
 	void ListTimers() const;
 
-	/** Used by the UGameInstance constructor to set this manager's owning game instance. */
+private:
 	void SetGameInstance(UGameInstance* InGameInstance);
 
 // This should be private, but needs to be public for testing.
