Have blueprints call IsDataValid on their simple construction script inside their IsDataValidFunction. Simple construction script now calls IsDataValid on its SCS nodes. This ensures that components added to the blueprint are checked when we're validating the blueprint.

diff --git a/Engine/Source/Runtime/CoreUObject/Private/UObject/UObjectGlobals.cpp b/Engine/Source/Runtime/CoreUObject/Private/UObject/UObjectGlobals.cpp
--- a/Engine/Source/Runtime/CoreUObject/Private/UObject/UObjectGlobals.cpp
+++ b/Engine/Source/Runtime/CoreUObject/Private/UObject/UObjectGlobals.cpp
@@ -4527,3 +4527,31 @@ namespace UE4CodeGen_Private
 		NewClass->StaticLink();
 	}
 }
+
+EDataValidationResult CombineDataValidationResults(EDataValidationResult Result1, EDataValidationResult Result2)
+{
+	/**
+	 * Anything combined with an Invalid result is Invalid. Any result combined with a NotValidated result is the same result
+	 *
+	 * The combined results should match the following matrix
+	 *
+	 *				|	NotValidated	|	Valid	|	Invalid
+	 * -------------+-------------------+-----------+----------
+	 * NotValidated	|	NotValidated	|	Valid	|	Invalid
+	 * Valid		|	Valid			|	Valid	|	Invalid
+	 * Invalid		|	Invalid			|	Invalid	|	Invalid
+	 *
+	 */
+
+	if (Result1 == EDataValidationResult::Invalid || Result2 == EDataValidationResult::Invalid)
+	{
+		return EDataValidationResult::Invalid;
+	}
+
+	if (Result1 == EDataValidationResult::Valid || Result2 == EDataValidationResult::Valid)
+	{
+		return EDataValidationResult::Valid;
+	}
+
+	return EDataValidationResult::NotValidated;
+}
diff --git a/Engine/Source/Runtime/CoreUObject/Public/UObject/UObjectGlobals.h b/Engine/Source/Runtime/CoreUObject/Public/UObject/UObjectGlobals.h
--- a/Engine/Source/Runtime/CoreUObject/Public/UObject/UObjectGlobals.h
+++ b/Engine/Source/Runtime/CoreUObject/Public/UObject/UObjectGlobals.h
@@ -2686,3 +2686,13 @@ enum class EDataValidationResult : uint8
 	/** Asset has not yet been validated */
 	NotValidated
 };
+
+/**
+ * Combines two different data validation results and returns the combined result.
+ *
+ * @param	Result1			One of the data validation results to be combined
+ * @param	Result2			One of the data validation results to be combined
+ *
+ * @return	Returns the combined data validation result
+ */
+COREUOBJECT_API EDataValidationResult CombineDataValidationResults(EDataValidationResult Result1, EDataValidationResult Result2);
diff --git a/Engine/Source/Runtime/Engine/Classes/Engine/SCS_Node.h b/Engine/Source/Runtime/Engine/Classes/Engine/SCS_Node.h
--- a/Engine/Source/Runtime/Engine/Classes/Engine/SCS_Node.h
+++ b/Engine/Source/Runtime/Engine/Classes/Engine/SCS_Node.h
@@ -93,7 +93,10 @@ class USCS_Node : public UObject
 
 	/** Make sure, the guid is proper - backward compatibility */
 	ENGINE_API void ValidateGuid();
-#endif
+
+	/** Returns Valid if this object has data validation rules set up for it and the data for this object is valid. Returns Invalid if it does not pass the rules. Returns NotValidated if no rules are set for this object. */
+	virtual EDataValidationResult IsDataValid(TArray<FText>& ValidationErrors) override;
+#endif // WITH_EDITOR
 
 	/**
 	 * Create the specified component on the actor, then call action on children
diff --git a/Engine/Source/Runtime/Engine/Classes/Engine/SimpleConstructionScript.h b/Engine/Source/Runtime/Engine/Classes/Engine/SimpleConstructionScript.h
--- a/Engine/Source/Runtime/Engine/Classes/Engine/SimpleConstructionScript.h
+++ b/Engine/Source/Runtime/Engine/Classes/Engine/SimpleConstructionScript.h
@@ -103,6 +103,11 @@ class USimpleConstructionScript : public UObject
 	/** Checks the root node set for scene components and ensures that it is valid (e.g. after a removal) */
 	ENGINE_API void ValidateSceneRootNodes();
 
+#if WITH_EDITOR
+	/** Returns Valid if this object has data validation rules set up for it and the data for this object is valid. Returns Invalid if it does not pass the rules. Returns NotValidated if no rules are set for this object. */
+	virtual EDataValidationResult IsDataValid(TArray<FText>& ValidationErrors) override;
+#endif // WITH_EDITOR
+
 private:
 	/** Root nodes of the construction script */
 	UPROPERTY()
diff --git a/Engine/Source/Runtime/Engine/Private/Blueprint.cpp b/Engine/Source/Runtime/Engine/Private/Blueprint.cpp
--- a/Engine/Source/Runtime/Engine/Private/Blueprint.cpp
+++ b/Engine/Source/Runtime/Engine/Private/Blueprint.cpp
@@ -1929,7 +1929,31 @@ UInheritableComponentHandler* UBlueprint::GetInheritableComponentHandler(bool bC
 
 EDataValidationResult UBlueprint::IsDataValid(TArray<FText>& ValidationErrors)
 {
-	return GeneratedClass ? GeneratedClass->GetDefaultObject()->IsDataValid(ValidationErrors) : EDataValidationResult::Invalid;
+	EDataValidationResult IsValid = GeneratedClass ? GeneratedClass->GetDefaultObject()->IsDataValid(ValidationErrors) : EDataValidationResult::Invalid;
+	IsValid = (IsValid == EDataValidationResult::NotValidated) ? EDataValidationResult::Valid : IsValid;
+
+	EDataValidationResult IsSCSValid = SimpleConstructionScript->IsDataValid(ValidationErrors);
+	IsValid = CombineDataValidationResults(IsValid, IsSCSValid);
+
+	for (UActorComponent* Component : ComponentTemplates)
+	{
+		if (Component)
+		{
+			EDataValidationResult IsComponentValid = Component->IsDataValid(ValidationErrors);
+			IsValid = CombineDataValidationResults(IsValid, IsComponentValid);
+		}
+	}
+
+	for (UTimelineTemplate* Timeline : Timelines)
+	{
+		if (Timeline)
+		{
+			EDataValidationResult IsTimelineValid = Timeline->IsDataValid(ValidationErrors);
+			IsValid = CombineDataValidationResults(IsValid, IsTimelineValid);
+		}
+	}
+
+	return IsValid;
 }
 
 bool UBlueprint::FindDiffs(const UBlueprint* OtherBlueprint, FDiffResults& Results) const
diff --git a/Engine/Source/Runtime/Engine/Private/SCS_Node.cpp b/Engine/Source/Runtime/Engine/Private/SCS_Node.cpp
--- a/Engine/Source/Runtime/Engine/Private/SCS_Node.cpp
+++ b/Engine/Source/Runtime/Engine/Private/SCS_Node.cpp
@@ -600,4 +600,29 @@ void USCS_Node::ValidateGuid()
 	}
 }
 
-#endif
+EDataValidationResult USCS_Node::IsDataValid(TArray<FText>& ValidationErrors)
+{
+	EDataValidationResult Result = Super::IsDataValid(ValidationErrors);
+	Result = (Result == EDataValidationResult::NotValidated) ? EDataValidationResult::Valid : Result;
+
+	// check the component that this node represents
+	if (ComponentTemplate)
+	{
+		EDataValidationResult ComponentResult = ComponentTemplate->IsDataValid(ValidationErrors);
+		Result = CombineDataValidationResults(Result, ComponentResult);
+	}
+
+	// check children
+	for (USCS_Node* Child : ChildNodes)
+	{
+		if (Child)
+		{
+			EDataValidationResult ChildResult = Child->IsDataValid(ValidationErrors);
+			Result = CombineDataValidationResults(Result, ChildResult);
+		}
+	}
+
+	return Result;
+}
+
+#endif // WITH_EDITOR
diff --git a/Engine/Source/Runtime/Engine/Private/SimpleConstructionScript.cpp b/Engine/Source/Runtime/Engine/Private/SimpleConstructionScript.cpp
--- a/Engine/Source/Runtime/Engine/Private/SimpleConstructionScript.cpp
+++ b/Engine/Source/Runtime/Engine/Private/SimpleConstructionScript.cpp
@@ -1190,6 +1190,22 @@ void USimpleConstructionScript::ValidateSceneRootNodes()
 }
 
 #if WITH_EDITOR
+EDataValidationResult USimpleConstructionScript::IsDataValid(TArray<FText>& ValidationErrors)
+{
+	EDataValidationResult Result = Super::IsDataValid(ValidationErrors);
+	Result = (Result == EDataValidationResult::NotValidated) ? EDataValidationResult::Valid : Result;
+
+	for (USCS_Node* Node : RootNodes)
+	{
+		if (Node)
+		{
+			EDataValidationResult NodeResult = Node->IsDataValid(ValidationErrors);
+			Result = CombineDataValidationResults(Result, NodeResult);
+		}
+	}
+	return Result;
+}
+
 void USimpleConstructionScript::GenerateListOfExistingNames(TSet<FName>& CurrentNames) const
 {
 	TArray<const USCS_Node*> ChildrenNodes = GetAllNodesConst();
