Fix crash in UBlueprint::IsDataValid when there's no SimpleConstructionScript.

diff --git a/Engine/Source/Runtime/Engine/Private/Blueprint.cpp b/Engine/Source/Runtime/Engine/Private/Blueprint.cpp
--- a/Engine/Source/Runtime/Engine/Private/Blueprint.cpp
+++ b/Engine/Source/Runtime/Engine/Private/Blueprint.cpp
@@ -1932,8 +1932,11 @@ EDataValidationResult UBlueprint::IsDataValid(TArray<FText>& ValidationErrors)
 	EDataValidationResult IsValid = GeneratedClass ? GeneratedClass->GetDefaultObject()->IsDataValid(ValidationErrors) : EDataValidationResult::Invalid;
 	IsValid = (IsValid == EDataValidationResult::NotValidated) ? EDataValidationResult::Valid : IsValid;
 
-	EDataValidationResult IsSCSValid = SimpleConstructionScript->IsDataValid(ValidationErrors);
-	IsValid = CombineDataValidationResults(IsValid, IsSCSValid);
+	if (SimpleConstructionScript)
+	{
+		EDataValidationResult IsSCSValid = SimpleConstructionScript->IsDataValid(ValidationErrors);
+		IsValid = CombineDataValidationResults(IsValid, IsSCSValid);
+	}
 
 	for (UActorComponent* Component : ComponentTemplates)
 	{
