[NO NON-PII COMMIT MESSAGE]

diff --git a/Engine/Source/Runtime/Engine/Private/Actor.cpp b/Engine/Source/Runtime/Engine/Private/Actor.cpp
--- a/Engine/Source/Runtime/Engine/Private/Actor.cpp
+++ b/Engine/Source/Runtime/Engine/Private/Actor.cpp
@@ -378,7 +378,8 @@ UWorld* AActor::GetWorld() const
 {
 	// CDO objects do not belong to a world
 	// If the actors outer is destroyed or unreachable we are shutting down and the world should be nullptr
-	if (!HasAnyFlags(RF_ClassDefaultObject) && !GetOuter()->HasAnyFlags(RF_BeginDestroyed) && !GetOuter()->IsUnreachable())
+	if (!HasAnyFlags(RF_ClassDefaultObject) && ensureMsgf(GetOuter(), TEXT("Actor: %s has a null OuterPrivate in AActor::GetWorld()"), *GetFullName())
+		&& !GetOuter()->HasAnyFlags(RF_BeginDestroyed) && !GetOuter()->IsUnreachable())
 	{
 		if (ULevel* Level = GetLevel())
 		{
