#UE4 Slightly less aggressive cleanup of old native default subobject components since they need to exist to be considered for reinstancing during serialize.

diff --git a/Engine/Source/Runtime/Engine/Classes/Components/ActorComponent.h b/Engine/Source/Runtime/Engine/Classes/Components/ActorComponent.h
--- a/Engine/Source/Runtime/Engine/Classes/Components/ActorComponent.h
+++ b/Engine/Source/Runtime/Engine/Classes/Components/ActorComponent.h
@@ -239,8 +239,8 @@ private:
 	/** During undo/redo it isn't safe to cache owner */
 	uint8 bCanUseCachedOwner:1;
 
-	/** Invalidates this component as an export once the linker is set. Used to clean up old native default subobjects that were removed from code */
-	uint8 bInvalidateExportWhenLinkerIsSet : 1;
+	/** Marks this component pending kill once PostLoad occurs. Used to clean up old native default subobjects that were removed from code */
+	uint8 bMarkPendingKillOnPostLoad : 1;
 #endif
 
 	/** True if this component was owned by a net startup actor during level load. */
@@ -795,7 +795,6 @@ public:
 	virtual void PostRename(UObject* OldOuter, const FName OldName) override;
 	virtual void Serialize(FArchive& Ar) override;
 #if WITH_EDITOR
-	virtual void PostLinkerChange() override;
 	virtual bool Modify( bool bAlwaysMarkDirty = true ) override;
 	virtual void PreEditChange(UProperty* PropertyThatWillChange) override;
 	virtual void PostEditChangeProperty(FPropertyChangedEvent& PropertyChangedEvent) override;
diff --git a/Engine/Source/Runtime/Engine/Private/Components/ActorComponent.cpp b/Engine/Source/Runtime/Engine/Private/Components/ActorComponent.cpp
--- a/Engine/Source/Runtime/Engine/Private/Components/ActorComponent.cpp
+++ b/Engine/Source/Runtime/Engine/Private/Components/ActorComponent.cpp
@@ -197,8 +197,8 @@ void UActorComponent::PostInitProperties()
 				if (HasAnyInternalFlags(EInternalObjectFlags::AsyncLoading))
 				{
 					// Async loading components cannot be pending kill, or the async loading code will assert when trying to postload them.
-					// Instead, wait until the linker is set and invalidate the export so it does not proceed to load
-					bInvalidateExportWhenLinkerIsSet = true;
+					// Instead, wait until the postload and mark pending kill at that time
+					bMarkPendingKillOnPostLoad = true;
 				}
 				else
 #endif // WITH_EDITOR
@@ -277,6 +277,14 @@ void UActorComponent::PostLoad()
 		// For a brief period of time we were inadvertently storing these for all components, need to clear it out
 		UCSModifiedProperties.Empty();
 	}
+
+#if WITH_EDITOR
+	if (bMarkPendingKillOnPostLoad)
+	{
+		MarkPendingKill();
+		bMarkPendingKillOnPostLoad = false;
+	}
+#endif // WITH_EDITOR
 }
 
 bool UActorComponent::Rename( const TCHAR* InName, UObject* NewOuter, ERenameFlags Flags )
@@ -552,17 +560,6 @@ bool UActorComponent::CallRemoteFunction( UFunction* Function, void* Parameters,
 /** FComponentReregisterContexts for components which have had PreEditChange called but not PostEditChange. */
 static TMap<TWeakObjectPtr<UActorComponent>,FComponentReregisterContext*> EditReregisterContexts;
 
-void UActorComponent::PostLinkerChange()
-{
-	Super::PostLinkerChange();
-
-	if (bInvalidateExportWhenLinkerIsSet)
-	{
-		FLinkerLoad::InvalidateExport(this);
-		bInvalidateExportWhenLinkerIsSet = false;
-	}
-}
-
 bool UActorComponent::Modify( bool bAlwaysMarkDirty/*=true*/ )
 {
 	// If this is a construction script component we don't store them in the transaction buffer.  Instead, mark
