Added SessionSummary analytics event that triggers on every shutdown, which will make it easier to write queries against analytics data.

diff --git a/Engine/Source/Runtime/Engine/Private/EngineSessionManager.cpp b/Engine/Source/Runtime/Engine/Private/EngineSessionManager.cpp
--- a/Engine/Source/Runtime/Engine/Private/EngineSessionManager.cpp
+++ b/Engine/Source/Runtime/Engine/Private/EngineSessionManager.cpp
@@ -22,6 +22,12 @@
 #include "Misc/App.h"
 #include "HAL/PlatformProcess.h"
 #include "HAL/PlatformOutputDevices.h"
+#include "Interfaces/IPluginManager.h"
+#include "IVREditorModule.h"
+#include "Kismet2/DebuggerCommands.h"
+#include "Interfaces/IProjectManager.h"
+#include "Misc/EngineVersion.h"
+#include "RHI.h"
 
 #define LOCTEXT_NAMESPACE "SessionManager"
 
@@ -59,7 +65,10 @@ namespace SessionManagerDefs
 	static const FString PlatformProcessIDKey(TEXT("PlatformProcessID"));
 	static const FString EngineVersionStoreKey(TEXT("EngineVersion"));
 	static const FString TimestampStoreKey(TEXT("Timestamp"));
-	static const FString StartupTimeStoreKey(TEXT("StartupTimestamp"));
+	static const FString StartupTimestampStoreKey(TEXT("StartupTimestamp"));
+	static const FString Idle1MinStoreKey(TEXT("Idle1Min"));
+	static const FString Idle5MinStoreKey(TEXT("Idle5Min"));
+	static const FString Idle30MinStoreKey(TEXT("Idle30Min"));
 	static const FString SessionIdStoreKey(TEXT("SessionId"));
 	static const FString StatusStoreKey(TEXT("LastExecutionState"));
 	static const FString DebuggerStoreKey(TEXT("IsDebugger"));
@@ -72,6 +81,12 @@ namespace SessionManagerDefs
 	static const FString EditorValueString(TEXT("Editor"));
 	static const FString GameValueString(TEXT("Game"));
 	static const FString UnknownProjectValueString(TEXT("UnknownProject"));
+	static const FString PluginsStoreKey(TEXT("Plugins"));
+	static const FString WasShutdownStoreKey(TEXT("WasShutdown"));
+	static const FString AverageFPSStoreKey(TEXT("AverageFPS"));
+	static const FString IsInVRModeStoreKey(TEXT("IsInVRMode"));
+	static const FString IsInEnterpriseStoreKey(TEXT("IsInEnterprise"));
+	static const FString IsInPIEStoreKey(TEXT("IsInPIE"));
 }
 
 namespace
@@ -90,6 +105,19 @@ namespace
 		}
 		return FDateTime::MinValue();
 	}
+
+	FString BoolToStoredString(bool bInValue)
+	{
+		return bInValue ? SessionManagerDefs::TrueValueString : SessionManagerDefs::FalseValueString;
+	}
+
+	bool GetStoredBool(const FString& SectionName, const FString& StoredKey)
+	{
+		FString StoredString = SessionManagerDefs::FalseValueString;
+		FPlatformMisc::GetStoredValue(SessionManagerDefs::StoreId, SectionName, StoredKey, StoredString);
+
+		return StoredString == SessionManagerDefs::TrueValueString;
+	}
 }
 
 /* FEngineSessionManager */
@@ -136,7 +164,7 @@ void FEngineSessionManager::InitializeRecords(bool bFirstAttempt)
 			{
 				FTimespan RecordAge = FDateTime::UtcNow() - Record.Timestamp;
 
-				if (Record.bCrashed || Record.bIsTerminating)
+				if (Record.bCrashed || Record.bIsTerminating || Record.bWasShutdown)
 				{
 					// Crashed / terminated sessions
 					SessionRecordsToReport.Add(Record);
@@ -172,10 +200,15 @@ void FEngineSessionManager::InitializeRecords(bool bFirstAttempt)
 		}
 	}
 
-	for (FSessionRecord& ReportingSession : SessionRecordsToReport)
+	for (FSessionRecord& SessionRecord : SessionRecordsToReport)
 	{
-		// Send error report for session that timed out or crashed
-		SendAbnormalShutdownReport(ReportingSession);
+		SendSessionRecordEvent(TEXT("SessionSummary"), SessionRecord, true);
+
+		if (!SessionRecord.bWasShutdown)
+		{
+			// Send error report for session that timed out or crashed
+			SendAbnormalShutdownReport(SessionRecord);
+		}
 	}
 }
 
@@ -202,7 +235,7 @@ void FEngineSessionManager::Tick(float DeltaTime)
 			{
 				CurrentSession.bIsDebugger = bIsDebuggerPresent;
 			
-				FString IsDebuggerString = CurrentSession.bIsDebugger ? SessionManagerDefs::TrueValueString : SessionManagerDefs::FalseValueString;
+				FString IsDebuggerString = BoolToStoredString(CurrentSession.bIsDebugger);
 				FPlatformMisc::SetStoredValue(SessionManagerDefs::StoreId, CurrentSessionSectionName, SessionManagerDefs::DebuggerStoreKey, IsDebuggerString);
 
 				if (!CurrentSession.bWasEverDebugger && CurrentSession.bIsDebugger)
@@ -223,12 +256,54 @@ void FEngineSessionManager::Tick(float DeltaTime)
 			const FString TimestampString = TimestampToString(FDateTime::UtcNow());
 			FPlatformMisc::SetStoredValue(SessionManagerDefs::StoreId, CurrentSessionSectionName, SessionManagerDefs::TimestampStoreKey, TimestampString);
 
+			const float IdleSeconds = FPlatformTime::Seconds() - FSlateApplication::Get().GetLastUserInteractionTime();
+
+			// 1 + 1 minutes
+			if (IdleSeconds > (60 + 60))
+			{
+				CurrentSession.Idle1Min += 1;
+				FPlatformMisc::SetStoredValue(SessionManagerDefs::StoreId, CurrentSessionSectionName, SessionManagerDefs::Idle1MinStoreKey, FString::FromInt(CurrentSession.Idle1Min));
+			}
+
+			// 5 + 1 minutes
+			if (IdleSeconds > (5 * 60 + 60))
+			{
+				CurrentSession.Idle5Min += 1;
+				FPlatformMisc::SetStoredValue(SessionManagerDefs::StoreId, CurrentSessionSectionName, SessionManagerDefs::Idle5MinStoreKey, FString::FromInt(CurrentSession.Idle5Min));
+			}
+
+			// 30 + 1 minutes
+			if (IdleSeconds > (30 * 60 + 60))
+			{
+				CurrentSession.Idle30Min += 1;
+				FPlatformMisc::SetStoredValue(SessionManagerDefs::StoreId, CurrentSessionSectionName, SessionManagerDefs::Idle30MinStoreKey, FString::FromInt(CurrentSession.Idle30Min));
+			}
+
+
 #if PLATFORM_SUPPORTS_WATCHDOG
 			if (!WatchdogSectionName.IsEmpty())
 			{
 				FPlatformMisc::SetStoredValue(SessionManagerDefs::StoreId, WatchdogSectionName, SessionManagerDefs::TimestampStoreKey, TimestampString);
 			}
 #endif
+
+			extern ENGINE_API float GAverageFPS;
+
+			CurrentSession.AverageFPS = GAverageFPS;
+			const FString AverageFPSString = FString::SanitizeFloat(GAverageFPS);
+			FPlatformMisc::SetStoredValue(SessionManagerDefs::StoreId, CurrentSessionSectionName, SessionManagerDefs::AverageFPSStoreKey, AverageFPSString);
+
+			CurrentSession.bIsInVRMode = IVREditorModule::Get().IsVREditorModeActive();
+			const FString IsInVRModeString = BoolToStoredString(CurrentSession.bIsInVRMode);
+			FPlatformMisc::SetStoredValue(SessionManagerDefs::StoreId, CurrentSessionSectionName, SessionManagerDefs::IsInVRModeStoreKey, IsInVRModeString);
+
+			CurrentSession.bIsInEnterprise = IProjectManager::Get().IsEnterpriseProject();
+			const FString IsInEnterpriseString = BoolToStoredString(CurrentSession.bIsInEnterprise);
+			FPlatformMisc::SetStoredValue(SessionManagerDefs::StoreId, CurrentSessionSectionName, SessionManagerDefs::IsInEnterpriseStoreKey, IsInEnterpriseString);
+
+			CurrentSession.bIsInPIE = FPlayWorldCommandCallbacks::IsInPIE();
+			const FString IsInPIEString = BoolToStoredString(CurrentSession.bIsInPIE);
+			FPlatformMisc::SetStoredValue(SessionManagerDefs::StoreId, CurrentSessionSectionName, SessionManagerDefs::IsInPIEStoreKey, IsInPIEString);
 		}
 	}
 }
@@ -246,6 +321,8 @@ void FEngineSessionManager::Shutdown()
 	if (!CurrentSession.bIsTerminating) // Skip Slate if terminating, since we can't guarantee which thread called us.
 	{
 		FSlateApplication::Get().GetOnModalLoopTickEvent().RemoveAll(this);
+
+		FPlatformMisc::SetStoredValue(SessionManagerDefs::StoreId, CurrentSessionSectionName, SessionManagerDefs::WasShutdownStoreKey, SessionManagerDefs::TrueValueString);
 	}
 
 	// Clear the session record for this session
@@ -253,8 +330,6 @@ void FEngineSessionManager::Shutdown()
 	{
 		if (!CurrentSession.bCrashed)
 		{
-			DeleteStoredRecordValues(CurrentSessionSectionName);
-
 #if PLATFORM_SUPPORTS_WATCHDOG
 			if (!WatchdogSectionName.IsEmpty())
 			{
@@ -287,7 +362,7 @@ bool FEngineSessionManager::BeginReadWriteRecords()
 	SessionListString.ParseIntoArray(SessionIds, TEXT(","));
 
 	// Retrieve all the sessions in the list from storage
-	for (FString& SessionId : SessionIds)
+	for (const FString& SessionId : SessionIds)
 	{
 		FString SectionName = GetStoreSectionString(SessionId);
 
@@ -320,24 +395,24 @@ bool FEngineSessionManager::BeginReadWriteRecords()
 			}
 
 			// Read optional values
-			FString WasDebuggerString = IsDebuggerString;
-			FPlatformMisc::GetStoredValue(SessionManagerDefs::StoreId, SectionName, SessionManagerDefs::WasDebuggerStoreKey, WasDebuggerString);
 			FString ModeString = SessionManagerDefs::EditorValueString;
 			FPlatformMisc::GetStoredValue(SessionManagerDefs::StoreId, SectionName, SessionManagerDefs::ModeStoreKey, ModeString);
 			FString ProjectName = SessionManagerDefs::UnknownProjectValueString;
 			FPlatformMisc::GetStoredValue(SessionManagerDefs::StoreId, SectionName, SessionManagerDefs::ProjectNameStoreKey, ProjectName);
-			FString IsDeactivatedString = SessionManagerDefs::FalseValueString;
-			FPlatformMisc::GetStoredValue(SessionManagerDefs::StoreId, SectionName, SessionManagerDefs::DeactivatedStoreKey, IsDeactivatedString);
-			FString IsInBackgroundString = SessionManagerDefs::FalseValueString;
-			FPlatformMisc::GetStoredValue(SessionManagerDefs::StoreId, SectionName, SessionManagerDefs::BackgroundStoreKey, IsInBackgroundString);
 			FString UserActivityString = SessionManagerDefs::DefaultUserActivity;
 			FPlatformMisc::GetStoredValue(SessionManagerDefs::StoreId, SectionName, SessionManagerDefs::UserActivityStoreKey, UserActivityString);
-			FString IsVanillaString = SessionManagerDefs::FalseValueString;
-			FPlatformMisc::GetStoredValue(SessionManagerDefs::StoreId, SectionName, SessionManagerDefs::VanillaStoreKey, IsVanillaString);
-			FString IsGPUCrashString = SessionManagerDefs::FalseValueString;
-			FPlatformMisc::GetStoredValue(SessionManagerDefs::StoreId, SectionName, SessionManagerDefs::GPUCrashStoreKey, IsGPUCrashString);
-			FString IsTerminatingString = SessionManagerDefs::FalseValueString;
-			FPlatformMisc::GetStoredValue(SessionManagerDefs::StoreId, SectionName, SessionManagerDefs::TerminatingKey, IsTerminatingString);
+			FString StartupTimestampString;
+			FPlatformMisc::GetStoredValue(SessionManagerDefs::StoreId, SectionName, SessionManagerDefs::StartupTimestampStoreKey, StartupTimestampString);
+			FString Idle1MinString;
+			FPlatformMisc::GetStoredValue(SessionManagerDefs::StoreId, SectionName, SessionManagerDefs::Idle1MinStoreKey, Idle1MinString);
+			FString Idle5MinString;
+			FPlatformMisc::GetStoredValue(SessionManagerDefs::StoreId, SectionName, SessionManagerDefs::Idle5MinStoreKey, Idle5MinString);
+			FString Idle30MinString;
+			FPlatformMisc::GetStoredValue(SessionManagerDefs::StoreId, SectionName, SessionManagerDefs::Idle30MinStoreKey, Idle30MinString);
+			FString PluginsString;
+			FPlatformMisc::GetStoredValue(SessionManagerDefs::StoreId, SectionName, SessionManagerDefs::PluginsStoreKey, PluginsString);
+			FString AverageFPSString;
+			FPlatformMisc::GetStoredValue(SessionManagerDefs::StoreId, SectionName, SessionManagerDefs::AverageFPSStoreKey, AverageFPSString);
 
 			// Create new record from the read values
 			FSessionRecord NewRecord;
@@ -345,16 +420,27 @@ bool FEngineSessionManager::BeginReadWriteRecords()
 			NewRecord.Mode = ModeString == SessionManagerDefs::EditorValueString ? EEngineSessionManagerMode::Editor : EEngineSessionManagerMode::Game;
 			NewRecord.ProjectName = ProjectName;
 			NewRecord.EngineVersion = EngineVersionString;
+			NewRecord.StartupTimestamp = StringToTimestamp(StartupTimestampString);
 			NewRecord.Timestamp = StringToTimestamp(TimestampString);
+			NewRecord.Idle1Min = FCString::Atod(*Idle1MinString);
+			NewRecord.Idle5Min = FCString::Atod(*Idle5MinString);
+			NewRecord.Idle30Min = FCString::Atod(*Idle30MinString);
+			NewRecord.AverageFPS = FCString::Atof(*AverageFPSString);
+			NewRecord.CurrentUserActivity = UserActivityString;
 			NewRecord.bCrashed = IsCrashString == SessionManagerDefs::TrueValueString;
-			NewRecord.bGPUCrashed = IsGPUCrashString == SessionManagerDefs::TrueValueString;
+			NewRecord.bGPUCrashed = GetStoredBool(SectionName, SessionManagerDefs::GPUCrashStoreKey);
 			NewRecord.bIsDebugger = IsDebuggerString == SessionManagerDefs::TrueValueString;
-			NewRecord.bWasEverDebugger = WasDebuggerString == SessionManagerDefs::TrueValueString;
-			NewRecord.bIsDeactivated = IsDeactivatedString == SessionManagerDefs::TrueValueString;
-			NewRecord.bIsInBackground = IsInBackgroundString == SessionManagerDefs::TrueValueString;
-			NewRecord.CurrentUserActivity = UserActivityString;
-			NewRecord.bIsVanilla = IsVanillaString == SessionManagerDefs::TrueValueString;
-			NewRecord.bIsTerminating = IsTerminatingString == SessionManagerDefs::TrueValueString;
+			NewRecord.bWasEverDebugger = GetStoredBool(SectionName, SessionManagerDefs::WasDebuggerStoreKey);
+			NewRecord.bIsDeactivated = GetStoredBool(SectionName, SessionManagerDefs::DeactivatedStoreKey);
+			NewRecord.bIsInBackground = GetStoredBool(SectionName, SessionManagerDefs::BackgroundStoreKey);
+			NewRecord.bIsVanilla = GetStoredBool(SectionName, SessionManagerDefs::VanillaStoreKey);
+			NewRecord.bIsTerminating = GetStoredBool(SectionName, SessionManagerDefs::TerminatingKey);
+			NewRecord.bWasShutdown = GetStoredBool(SectionName, SessionManagerDefs::WasShutdownStoreKey);
+			NewRecord.bIsInPIE = GetStoredBool(SectionName, SessionManagerDefs::IsInPIEStoreKey);
+			NewRecord.bIsInVRMode = GetStoredBool(SectionName, SessionManagerDefs::IsInVRModeStoreKey);
+			NewRecord.bIsInEnterprise = GetStoredBool(SectionName, SessionManagerDefs::IsInEnterpriseStoreKey);
+
+			PluginsString.ParseIntoArray(NewRecord.Plugins, TEXT(","));
 
 			SessionRecords.Add(NewRecord);
 		}
@@ -408,15 +494,106 @@ void FEngineSessionManager::DeleteStoredRecordValues(const FString& SectionName)
 	FPlatformMisc::DeleteStoredValue(SessionManagerDefs::StoreId, SectionName, SessionManagerDefs::CrashStoreKey);
 	FPlatformMisc::DeleteStoredValue(SessionManagerDefs::StoreId, SectionName, SessionManagerDefs::GPUCrashStoreKey);
 	FPlatformMisc::DeleteStoredValue(SessionManagerDefs::StoreId, SectionName, SessionManagerDefs::EngineVersionStoreKey);
+	FPlatformMisc::DeleteStoredValue(SessionManagerDefs::StoreId, SectionName, SessionManagerDefs::StartupTimestampStoreKey);
 	FPlatformMisc::DeleteStoredValue(SessionManagerDefs::StoreId, SectionName, SessionManagerDefs::TimestampStoreKey);
+	FPlatformMisc::DeleteStoredValue(SessionManagerDefs::StoreId, SectionName, SessionManagerDefs::Idle1MinStoreKey);
+	FPlatformMisc::DeleteStoredValue(SessionManagerDefs::StoreId, SectionName, SessionManagerDefs::Idle5MinStoreKey);
+	FPlatformMisc::DeleteStoredValue(SessionManagerDefs::StoreId, SectionName, SessionManagerDefs::Idle30MinStoreKey);
 	FPlatformMisc::DeleteStoredValue(SessionManagerDefs::StoreId, SectionName, SessionManagerDefs::DebuggerStoreKey);
 	FPlatformMisc::DeleteStoredValue(SessionManagerDefs::StoreId, SectionName, SessionManagerDefs::WasDebuggerStoreKey);
+	FPlatformMisc::DeleteStoredValue(SessionManagerDefs::StoreId, SectionName, SessionManagerDefs::WasShutdownStoreKey);
 	FPlatformMisc::DeleteStoredValue(SessionManagerDefs::StoreId, SectionName, SessionManagerDefs::DeactivatedStoreKey);
 	FPlatformMisc::DeleteStoredValue(SessionManagerDefs::StoreId, SectionName, SessionManagerDefs::BackgroundStoreKey);
 	FPlatformMisc::DeleteStoredValue(SessionManagerDefs::StoreId, SectionName, SessionManagerDefs::UserActivityStoreKey);
 	FPlatformMisc::DeleteStoredValue(SessionManagerDefs::StoreId, SectionName, SessionManagerDefs::VanillaStoreKey);
 	FPlatformMisc::DeleteStoredValue(SessionManagerDefs::StoreId, SectionName, SessionManagerDefs::TerminatingKey);
 	FPlatformMisc::DeleteStoredValue(SessionManagerDefs::StoreId, SectionName, SessionManagerDefs::PlatformProcessIDKey);
+	FPlatformMisc::DeleteStoredValue(SessionManagerDefs::StoreId, SectionName, SessionManagerDefs::PluginsStoreKey);
+	FPlatformMisc::DeleteStoredValue(SessionManagerDefs::StoreId, SectionName, SessionManagerDefs::AverageFPSStoreKey);
+	FPlatformMisc::DeleteStoredValue(SessionManagerDefs::StoreId, SectionName, SessionManagerDefs::IsInPIEStoreKey);
+	FPlatformMisc::DeleteStoredValue(SessionManagerDefs::StoreId, SectionName, SessionManagerDefs::IsInEnterpriseStoreKey);
+	FPlatformMisc::DeleteStoredValue(SessionManagerDefs::StoreId, SectionName, SessionManagerDefs::IsInVRModeStoreKey);
+}
+
+void FEngineSessionManager::SendSessionRecordEvent(const FString& EventName, const FSessionRecord& Record, bool bSendHardwareDetails)
+{
+	FGuid SessionId;
+	FString SessionIdString = Record.SessionId;
+	if (FGuid::Parse(SessionIdString, SessionId))
+	{
+		// convert session guid to one with braces for sending to analytics
+		SessionIdString = SessionId.ToString(EGuidFormats::DigitsWithHyphensInBraces);
+	}
+
+#if !PLATFORM_PS4
+	FString ShutdownTypeString = Record.bCrashed ? SessionManagerDefs::CrashSessionToken :
+		(Record.bWasEverDebugger ? SessionManagerDefs::DebuggerSessionToken :
+		(Record.bIsTerminating ? SessionManagerDefs::TerminatedSessionToken : SessionManagerDefs::AbnormalSessionToken));
+#else
+	// PS4 cannot set the crash flag so report abnormal shutdowns with a specific token meaning "crash or abnormal shutdown".
+	FString ShutdownTypeString = Record.bWasEverDebugger ? SessionManagerDefs::DebuggerSessionToken : SessionManagerDefs::PS4SessionToken;
+#endif
+
+	const FString& RunTypeString = Record.Mode == EEngineSessionManagerMode::Editor ? SessionManagerDefs::EditorValueString : SessionManagerDefs::GameValueString;
+	FString PluginsString = FString::Join(Record.Plugins, TEXT(","));
+
+	TArray< FAnalyticsEventAttribute > AnalyticsAttributes;
+	AnalyticsAttributes.Emplace(TEXT("RunType"), RunTypeString);
+	AnalyticsAttributes.Emplace(TEXT("ProjectName"), Record.ProjectName);
+	AnalyticsAttributes.Emplace(TEXT("Platform"), FPlatformProperties::PlatformName());
+	AnalyticsAttributes.Emplace(TEXT("SessionId"), SessionIdString);
+	AnalyticsAttributes.Emplace(TEXT("EngineVersion"), Record.EngineVersion);
+	AnalyticsAttributes.Emplace(TEXT("ShutdownType"), ShutdownTypeString);
+	AnalyticsAttributes.Emplace(TEXT("Timestamp"), Record.Timestamp.ToIso8601());
+	AnalyticsAttributes.Emplace(TEXT("CurrentUserActivity"), Record.CurrentUserActivity);
+	AnalyticsAttributes.Emplace(TEXT("IsVanilla"), Record.bIsVanilla);
+	AnalyticsAttributes.Emplace(TEXT("WasDebugged"), Record.bWasEverDebugger);
+	AnalyticsAttributes.Emplace(TEXT("GPUCrash"), Record.bGPUCrashed);
+	AnalyticsAttributes.Emplace(SessionManagerDefs::WasShutdownStoreKey, Record.bWasShutdown);
+	AnalyticsAttributes.Emplace(SessionManagerDefs::StartupTimestampStoreKey, Record.StartupTimestamp.ToIso8601());
+	AnalyticsAttributes.Emplace(SessionManagerDefs::AverageFPSStoreKey, Record.AverageFPS);
+	AnalyticsAttributes.Emplace(SessionManagerDefs::IsInPIEStoreKey, Record.bIsInPIE);
+	AnalyticsAttributes.Emplace(SessionManagerDefs::IsInEnterpriseStoreKey, Record.bIsInEnterprise);
+	AnalyticsAttributes.Emplace(SessionManagerDefs::IsInVRModeStoreKey, Record.bIsInVRMode);
+
+	AnalyticsAttributes.Emplace(TEXT("1MinIdle"), Record.Idle1Min);
+	AnalyticsAttributes.Emplace(TEXT("5MinIdle"), Record.Idle5Min);
+	AnalyticsAttributes.Emplace(TEXT("30MinIdle"), Record.Idle30Min);
+
+	if (bSendHardwareDetails)
+	{
+		// Add project info whether we are in editor or game.
+		const UGeneralProjectSettings& ProjectSettings = *GetDefault<UGeneralProjectSettings>();
+		FString OSMajor;
+		FString OSMinor;
+		FPlatformMisc::GetOSVersions(/*out*/ OSMajor, /*out*/ OSMinor);
+
+		FPlatformMemoryStats Stats = FPlatformMemory::GetStats();
+		AnalyticsAttributes.Emplace(TEXT("ProjectName"), ProjectSettings.ProjectName);
+		AnalyticsAttributes.Emplace(TEXT("ProjectID"), ProjectSettings.ProjectID);
+		AnalyticsAttributes.Emplace(TEXT("ProjectDescription"), ProjectSettings.Description);
+		AnalyticsAttributes.Emplace(TEXT("ProjectVersion"), ProjectSettings.ProjectVersion);
+		AnalyticsAttributes.Emplace(TEXT("GPUVendorID"), GRHIVendorId);
+		AnalyticsAttributes.Emplace(TEXT("GPUDeviceID"), GRHIDeviceId);
+		AnalyticsAttributes.Emplace(TEXT("GRHIDeviceRevision"), GRHIDeviceRevision);
+		AnalyticsAttributes.Emplace(TEXT("GRHIAdapterInternalDriverVersion"), GRHIAdapterInternalDriverVersion);
+		AnalyticsAttributes.Emplace(TEXT("GRHIAdapterUserDriverVersion"), GRHIAdapterUserDriverVersion);
+		AnalyticsAttributes.Emplace(TEXT("TotalPhysicalRAM"), Stats.TotalPhysical);
+		AnalyticsAttributes.Emplace(TEXT("CPUPhysicalCores"), FPlatformMisc::NumberOfCores());
+		AnalyticsAttributes.Emplace(TEXT("CPULogicalCores"), FPlatformMisc::NumberOfCoresIncludingHyperthreads());
+		AnalyticsAttributes.Emplace(TEXT("DesktopGPUAdapter"), FPlatformMisc::GetPrimaryGPUBrand());
+		AnalyticsAttributes.Emplace(TEXT("RenderingGPUAdapter"), GRHIAdapterName);
+		AnalyticsAttributes.Emplace(TEXT("CPUVendor"), FPlatformMisc::GetCPUVendor());
+		AnalyticsAttributes.Emplace(TEXT("CPUBrand"), FPlatformMisc::GetCPUBrand());
+		AnalyticsAttributes.Emplace(TEXT("OSMajor"), OSMajor);
+		AnalyticsAttributes.Emplace(TEXT("OSMinor"), OSMinor);
+		AnalyticsAttributes.Emplace(TEXT("OSVersion"), FPlatformMisc::GetOSVersion());
+		AnalyticsAttributes.Emplace(TEXT("Is64BitOS"), FPlatformMisc::Is64bitOperatingSystem());
+	}
+
+	FEngineAnalytics::GetProvider().RecordEvent(EventName, AnalyticsAttributes);
+
+	UE_LOG(LogEngineSessionManager, Log, TEXT("EngineSessionManager sent report. Event=%s, Type=%s, SessionId=%s"), *EventName, *ShutdownTypeString, *SessionIdString);
 }
 
 /**
@@ -461,8 +638,6 @@ void FEngineSessionManager::DeleteStoredRecordValues(const FString& SectionName)
  */
 void FEngineSessionManager::SendAbnormalShutdownReport(const FSessionRecord& Record)
 {
-	FString PlatformName(FPlatformProperties::PlatformName());
-
 #if PLATFORM_WINDOWS | PLATFORM_MAC | PLATFORM_UNIX
 	// do nothing
 #elif PLATFORM_PS4
@@ -481,41 +656,7 @@ void FEngineSessionManager::SendAbnormalShutdownReport(const FSessionRecord& Rec
 	return; // TODO: CWood: disabled on other platforms
 #endif
 
-	FGuid SessionId;
-	FString SessionIdString = Record.SessionId;
-	if (FGuid::Parse(SessionIdString, SessionId))
-	{
-		// convert session guid to one with braces for sending to analytics
-		SessionIdString = SessionId.ToString(EGuidFormats::DigitsWithHyphensInBraces);
-	}
-
-#if !PLATFORM_PS4
-	FString ShutdownTypeString = Record.bCrashed ? SessionManagerDefs::CrashSessionToken :
-		(Record.bWasEverDebugger ? SessionManagerDefs::DebuggerSessionToken :
-		(Record.bIsTerminating ? SessionManagerDefs::TerminatedSessionToken : SessionManagerDefs::AbnormalSessionToken));
-#else
-	// PS4 cannot set the crash flag so report abnormal shutdowns with a specific token meaning "crash or abnormal shutdown".
-	FString ShutdownTypeString = Record.bWasEverDebugger ? SessionManagerDefs::DebuggerSessionToken : SessionManagerDefs::PS4SessionToken;
-#endif
-
-	const FString& RunTypeString = Record.Mode == EEngineSessionManagerMode::Editor ? SessionManagerDefs::EditorValueString : SessionManagerDefs::GameValueString;
-
-	TArray< FAnalyticsEventAttribute > AbnormalShutdownAttributes;
-	AbnormalShutdownAttributes.Add(FAnalyticsEventAttribute(TEXT("RunType"), RunTypeString));
-	AbnormalShutdownAttributes.Add(FAnalyticsEventAttribute(TEXT("ProjectName"), Record.ProjectName));
-	AbnormalShutdownAttributes.Add(FAnalyticsEventAttribute(TEXT("Platform"), PlatformName));
-	AbnormalShutdownAttributes.Add(FAnalyticsEventAttribute(TEXT("SessionId"), SessionIdString));
-	AbnormalShutdownAttributes.Add(FAnalyticsEventAttribute(TEXT("EngineVersion"), Record.EngineVersion));
-	AbnormalShutdownAttributes.Add(FAnalyticsEventAttribute(TEXT("ShutdownType"), ShutdownTypeString));
-	AbnormalShutdownAttributes.Add(FAnalyticsEventAttribute(TEXT("Timestamp"), Record.Timestamp.ToIso8601()));
-	AbnormalShutdownAttributes.Add(FAnalyticsEventAttribute(TEXT("CurrentUserActivity"), Record.CurrentUserActivity));
-	AbnormalShutdownAttributes.Add(FAnalyticsEventAttribute(TEXT("IsVanilla"), Record.bIsVanilla));
-	AbnormalShutdownAttributes.Add(FAnalyticsEventAttribute(TEXT("WasDebugged"), Record.bWasEverDebugger));
-	AbnormalShutdownAttributes.Add(FAnalyticsEventAttribute(TEXT("GPUCrash"), Record.bGPUCrashed));
-
-	FEngineAnalytics::GetProvider().RecordEvent(TEXT("Engine.AbnormalShutdown"), AbnormalShutdownAttributes);
-
-	UE_LOG(LogEngineSessionManager, Log, TEXT("EngineSessionManager sent abnormal shutdown report. Type=%s, SessionId=%s"), *ShutdownTypeString, *SessionIdString);
+	SendSessionRecordEvent(TEXT("Engine.AbnormalShutdown"), Record, false);
 }
 
 void FEngineSessionManager::CreateAndWriteRecordForSession()
@@ -543,21 +684,41 @@ void FEngineSessionManager::CreateAndWriteRecordForSession()
 	CurrentSession.bIsDebugger = CurrentSession.bWasEverDebugger = FPlatformMisc::IsDebuggerPresent();
 	CurrentSession.CurrentUserActivity = GetUserActivityString();
 	CurrentSession.bIsVanilla = GEngine && GEngine->IsVanillaProduct();
+
 	CurrentSessionSectionName = GetStoreSectionString(CurrentSession.SessionId);
 
+	TArray<TSharedRef<IPlugin>> Plugins = IPluginManager::Get().GetEnabledPlugins();
+	for (const TSharedRef<IPlugin>& Plugin : Plugins)
+	{
+		CurrentSession.Plugins.Add(Plugin->GetName());
+	}
+
+	CurrentSession.Plugins.Sort();
+
 	FString ModeString = CurrentSession.Mode == EEngineSessionManagerMode::Editor ? SessionManagerDefs::EditorValueString : SessionManagerDefs::GameValueString;
-	FString IsDebuggerString = CurrentSession.bIsDebugger ? SessionManagerDefs::TrueValueString : SessionManagerDefs::FalseValueString;
-	FString WasDebuggerString = CurrentSession.bWasEverDebugger ? SessionManagerDefs::TrueValueString : SessionManagerDefs::FalseValueString;
-	FString IsDeactivatedString = CurrentSession.bIsDeactivated ? SessionManagerDefs::TrueValueString : SessionManagerDefs::FalseValueString;
-	FString IsInBackgroundString = CurrentSession.bIsInBackground ? SessionManagerDefs::TrueValueString : SessionManagerDefs::FalseValueString;
-	FString IsVanillaString = CurrentSession.bIsVanilla ? SessionManagerDefs::TrueValueString : SessionManagerDefs::FalseValueString;
-	FString IsTerminatingString = CurrentSession.bIsTerminating ? SessionManagerDefs::TrueValueString : SessionManagerDefs::FalseValueString;
+	FString IsDebuggerString = BoolToStoredString(CurrentSession.bIsDebugger);
+	FString WasDebuggerString = BoolToStoredString(CurrentSession.bWasEverDebugger);
+	FString IsDeactivatedString = BoolToStoredString(CurrentSession.bIsDeactivated);
+	FString IsInBackgroundString = BoolToStoredString(CurrentSession.bIsInBackground);
+	FString IsVanillaString = BoolToStoredString(CurrentSession.bIsVanilla);
+	FString IsTerminatingString = BoolToStoredString(CurrentSession.bIsTerminating);
+	FString CurrentTimestampString = TimestampToString(CurrentSession.Timestamp);
+	FString Idle1MinString = FString::FromInt(CurrentSession.Idle1Min);
+	FString Idle5MinString = FString::FromInt(CurrentSession.Idle5Min);
+	FString Idle30MinString = FString::FromInt(CurrentSession.Idle30Min);
+	FString AverageFPSString = FString::SanitizeFloat(CurrentSession.AverageFPS);
+	FString PluginsString = FString::Join(CurrentSession.Plugins, TEXT(","));
+	FString WasShutdownString = BoolToStoredString(CurrentSession.bWasShutdown);
 
 	FPlatformMisc::SetStoredValue(SessionManagerDefs::StoreId, CurrentSessionSectionName, SessionManagerDefs::ModeStoreKey, ModeString);
 	FPlatformMisc::SetStoredValue(SessionManagerDefs::StoreId, CurrentSessionSectionName, SessionManagerDefs::ProjectNameStoreKey, CurrentSession.ProjectName);
 	FPlatformMisc::SetStoredValue(SessionManagerDefs::StoreId, CurrentSessionSectionName, SessionManagerDefs::CrashStoreKey, SessionManagerDefs::FalseValueString);
 	FPlatformMisc::SetStoredValue(SessionManagerDefs::StoreId, CurrentSessionSectionName, SessionManagerDefs::EngineVersionStoreKey, CurrentSession.EngineVersion);
-	FPlatformMisc::SetStoredValue(SessionManagerDefs::StoreId, CurrentSessionSectionName, SessionManagerDefs::TimestampStoreKey, TimestampToString(CurrentSession.Timestamp));
+	FPlatformMisc::SetStoredValue(SessionManagerDefs::StoreId, CurrentSessionSectionName, SessionManagerDefs::StartupTimestampStoreKey, CurrentTimestampString);
+	FPlatformMisc::SetStoredValue(SessionManagerDefs::StoreId, CurrentSessionSectionName, SessionManagerDefs::TimestampStoreKey, CurrentTimestampString);
+	FPlatformMisc::SetStoredValue(SessionManagerDefs::StoreId, CurrentSessionSectionName, SessionManagerDefs::Idle1MinStoreKey, Idle1MinString);
+	FPlatformMisc::SetStoredValue(SessionManagerDefs::StoreId, CurrentSessionSectionName, SessionManagerDefs::Idle5MinStoreKey, Idle5MinString);
+	FPlatformMisc::SetStoredValue(SessionManagerDefs::StoreId, CurrentSessionSectionName, SessionManagerDefs::Idle30MinStoreKey, Idle30MinString);
 	FPlatformMisc::SetStoredValue(SessionManagerDefs::StoreId, CurrentSessionSectionName, SessionManagerDefs::DebuggerStoreKey, IsDebuggerString);
 	FPlatformMisc::SetStoredValue(SessionManagerDefs::StoreId, CurrentSessionSectionName, SessionManagerDefs::WasDebuggerStoreKey, WasDebuggerString);
 	FPlatformMisc::SetStoredValue(SessionManagerDefs::StoreId, CurrentSessionSectionName, SessionManagerDefs::DeactivatedStoreKey, IsDeactivatedString);
@@ -565,7 +726,13 @@ void FEngineSessionManager::CreateAndWriteRecordForSession()
 	FPlatformMisc::SetStoredValue(SessionManagerDefs::StoreId, CurrentSessionSectionName, SessionManagerDefs::UserActivityStoreKey, CurrentSession.CurrentUserActivity);
 	FPlatformMisc::SetStoredValue(SessionManagerDefs::StoreId, CurrentSessionSectionName, SessionManagerDefs::VanillaStoreKey, IsVanillaString);
 	FPlatformMisc::SetStoredValue(SessionManagerDefs::StoreId, CurrentSessionSectionName, SessionManagerDefs::TerminatingKey, IsTerminatingString);
+	FPlatformMisc::SetStoredValue(SessionManagerDefs::StoreId, CurrentSessionSectionName, SessionManagerDefs::WasShutdownStoreKey, WasShutdownString);
 	FPlatformMisc::SetStoredValue(SessionManagerDefs::StoreId, CurrentSessionSectionName, SessionManagerDefs::PlatformProcessIDKey, CurrentProcessIDString);
+	FPlatformMisc::SetStoredValue(SessionManagerDefs::StoreId, CurrentSessionSectionName, SessionManagerDefs::PluginsStoreKey, PluginsString);
+	FPlatformMisc::SetStoredValue(SessionManagerDefs::StoreId, CurrentSessionSectionName, SessionManagerDefs::AverageFPSStoreKey, AverageFPSString);
+	FPlatformMisc::SetStoredValue(SessionManagerDefs::StoreId, CurrentSessionSectionName, SessionManagerDefs::IsInPIEStoreKey, BoolToStoredString(CurrentSession.bIsInPIE));
+	FPlatformMisc::SetStoredValue(SessionManagerDefs::StoreId, CurrentSessionSectionName, SessionManagerDefs::IsInEnterpriseStoreKey, BoolToStoredString(CurrentSession.bIsInEnterprise));
+	FPlatformMisc::SetStoredValue(SessionManagerDefs::StoreId, CurrentSessionSectionName, SessionManagerDefs::IsInVRModeStoreKey, BoolToStoredString(CurrentSession.bIsInVRMode));
 
 	SessionRecords.Add(CurrentSession);
 
@@ -587,7 +754,7 @@ void FEngineSessionManager::OnCrashing()
 		CurrentSession.bCrashed = true;
 		CurrentSession.bGPUCrashed = GIsGPUCrashed;
 		FPlatformMisc::SetStoredValue(SessionManagerDefs::StoreId, CurrentSessionSectionName, SessionManagerDefs::CrashStoreKey, SessionManagerDefs::TrueValueString);
-		FPlatformMisc::SetStoredValue(SessionManagerDefs::StoreId, CurrentSessionSectionName, SessionManagerDefs::GPUCrashStoreKey, CurrentSession.bGPUCrashed ? SessionManagerDefs::TrueValueString : SessionManagerDefs::FalseValueString);
+		FPlatformMisc::SetStoredValue(SessionManagerDefs::StoreId, CurrentSessionSectionName, SessionManagerDefs::GPUCrashStoreKey, BoolToStoredString(CurrentSession.bGPUCrashed));
 
 #if PLATFORM_SUPPORTS_WATCHDOG
 		if (!WatchdogSectionName.IsEmpty())
@@ -678,7 +845,7 @@ void FEngineSessionManager::OnVanillaStateChanged(bool bIsVanilla)
 	if (CurrentSession.bIsVanilla != bIsVanilla && bInitializedRecords)
 	{
 		CurrentSession.bIsVanilla = bIsVanilla;
-		FString IsVanillaString = CurrentSession.bIsVanilla ? SessionManagerDefs::TrueValueString : SessionManagerDefs::FalseValueString;
+		FString IsVanillaString = BoolToStoredString(CurrentSession.bIsVanilla);
 		FPlatformMisc::SetStoredValue(SessionManagerDefs::StoreId, CurrentSessionSectionName, SessionManagerDefs::VanillaStoreKey, IsVanillaString);
 	}
 }
@@ -787,7 +954,7 @@ void FEngineSessionManager::StartWatchdog(const FString& RunType, const FString&
 		if (WatchdogProcessHandle.IsValid())
 		{
 			FString WatchdogStartTimeString = TimestampToString(FDateTime::UtcNow());
-			FString WasDebuggerString = CurrentSession.bWasEverDebugger ? SessionManagerDefs::TrueValueString : SessionManagerDefs::FalseValueString;
+			FString WasDebuggerString = BoolToStoredString(CurrentSession.bWasEverDebugger);
 
 			WatchdogStartedAttributes.Add(FAnalyticsEventAttribute(TEXT("Outcome"), TEXT("Succeeded")));
 			UE_LOG(LogEngineSessionManager, Log, TEXT("Started UnrealWatchdog for process id %u"), ProcessId);
@@ -795,7 +962,7 @@ void FEngineSessionManager::StartWatchdog(const FString& RunType, const FString&
 			WatchdogSectionName = GetWatchdogStoreSectionString(ProcessId);
 
 			FPlatformMisc::SetStoredValue(SessionManagerDefs::StoreId, WatchdogSectionName, SessionManagerDefs::CommandLineStoreKey, FCommandLine::GetOriginalForLogging());
-			FPlatformMisc::SetStoredValue(SessionManagerDefs::StoreId, WatchdogSectionName, SessionManagerDefs::StartupTimeStoreKey, WatchdogStartTimeString);
+			FPlatformMisc::SetStoredValue(SessionManagerDefs::StoreId, WatchdogSectionName, SessionManagerDefs::StartupTimestampStoreKey, WatchdogStartTimeString);
 			FPlatformMisc::SetStoredValue(SessionManagerDefs::StoreId, WatchdogSectionName, SessionManagerDefs::TimestampStoreKey, WatchdogStartTimeString);
 			FPlatformMisc::SetStoredValue(SessionManagerDefs::StoreId, WatchdogSectionName, SessionManagerDefs::StatusStoreKey, SessionManagerDefs::RunningSessionToken);
 			FPlatformMisc::SetStoredValue(SessionManagerDefs::StoreId, WatchdogSectionName, SessionManagerDefs::UserActivityStoreKey, CurrentSession.CurrentUserActivity);
diff --git a/Engine/Source/Runtime/Engine/Private/EngineSessionManager.h b/Engine/Source/Runtime/Engine/Private/EngineSessionManager.h
--- a/Engine/Source/Runtime/Engine/Private/EngineSessionManager.h
+++ b/Engine/Source/Runtime/Engine/Private/EngineSessionManager.h
@@ -27,7 +27,11 @@ public:
 		, HeartbeatTimeElapsed(0.0f)
 	{
 		CurrentSession.Mode = EEngineSessionManagerMode::Editor;
-		CurrentSession.Timestamp = FDateTime::MinValue();		
+		CurrentSession.Timestamp = FDateTime::MinValue();
+		CurrentSession.StartupTimestamp = FDateTime::MinValue();
+		CurrentSession.Idle1Min = 0;
+		CurrentSession.Idle5Min = 0;
+		CurrentSession.Idle30Min = 0;
 		CurrentSession.bCrashed = false;
 		CurrentSession.bGPUCrashed = false;
 		CurrentSession.bIsDebugger = false;
@@ -36,6 +40,10 @@ public:
 		CurrentSession.bIsInBackground = false;
 		CurrentSession.bIsVanilla = false;
 		CurrentSession.bIsTerminating = false;
+		CurrentSession.bWasShutdown = false;
+		CurrentSession.bIsInPIE = false;
+		CurrentSession.bIsInEnterprise = false;
+		CurrentSession.bIsInVRMode = false;
 	}
 
 	void Initialize();
@@ -47,20 +55,31 @@ public:
 private:
 	struct FSessionRecord
 	{
-		FString SessionId;
 		EEngineSessionManagerMode Mode;
+		FString SessionId;
 		FString ProjectName;
 		FString EngineVersion;
+		FDateTime StartupTimestamp;
 		FDateTime Timestamp;
-		bool bCrashed;
-		bool bGPUCrashed;
-		bool bIsDebugger;
-		bool bWasEverDebugger;
-		bool bIsDeactivated;
-		bool bIsInBackground;
+		int32 Idle1Min;
+		int32 Idle5Min;
+		int32 Idle30Min;
 		FString CurrentUserActivity;
-		bool bIsVanilla;
-		bool bIsTerminating;
+		TArray<FString> Plugins;
+		float AverageFPS;
+
+		bool bCrashed : 1;
+		bool bGPUCrashed : 1;
+		bool bIsDebugger : 1;
+		bool bWasEverDebugger : 1;
+		bool bIsDeactivated : 1;
+		bool bIsInBackground : 1;
+		bool bIsVanilla : 1;
+		bool bIsTerminating : 1;
+		bool bWasShutdown : 1;
+		bool bIsInPIE : 1;
+		bool bIsInEnterprise : 1;
+		bool bIsInVRMode : 1;
 	};
 
 private:
@@ -70,6 +89,7 @@ private:
 	void DeleteStoredRecord(const FSessionRecord& Record);
 	void DeleteStoredRecordValues(const FString& SectionName) const;
 	void SendAbnormalShutdownReport(const FSessionRecord& Record);
+	void SendSessionRecordEvent(const FString& EventName, const FSessionRecord& Record, bool bSendHardwareDetails);
 	void CreateAndWriteRecordForSession();
 	void OnCrashing();
 	void OnAppReactivate();
