Add some stats useful for troubleshooting mesh LOD streaming

diff --git a/Engine/Source/Runtime/Engine/Private/Streaming/AsyncTextureStreaming.cpp b/Engine/Source/Runtime/Engine/Private/Streaming/AsyncTextureStreaming.cpp
--- a/Engine/Source/Runtime/Engine/Private/Streaming/AsyncTextureStreaming.cpp
+++ b/Engine/Source/Runtime/Engine/Private/Streaming/AsyncTextureStreaming.cpp
@@ -816,6 +816,17 @@ void FRenderAssetStreamingMipCalcTask::UpdateStats_Async()
 
 	Stats.OverBudget = 0;
 
+	Stats.NumStreamedMeshes = 0;
+	Stats.AvgNumStreamedLODs = 0.f;
+	Stats.AvgNumResidentLODs = 0.f;
+	Stats.AvgNumEvictedLODs = 0.f;
+	Stats.StreamedMeshMem = 0;
+	Stats.ResidentMeshMem = 0;
+	Stats.EvictedMeshMem = 0;
+	int32 TotalNumStreamedLODs = 0;
+	int32 TotalNumResidentLODs = 0;
+	int32 TotalNumEvictedLODs = 0;
+
 	for (FStreamingRenderAsset& StreamingRenderAsset : StreamingRenderAssets)
 	{
 		if (IsAborted()) break;
@@ -882,6 +893,33 @@ void FRenderAssetStreamingMipCalcTask::UpdateStats_Async()
 		{
 			Stats.PendingRequests += StreamingRenderAsset.GetSize(StreamingRenderAsset.RequestedMips) - ResidentSize;
 		}
+
+		if (StreamingRenderAsset.RenderAssetType != FStreamingRenderAsset::AT_Texture)
+		{
+			const bool bOptLODsExist = StreamingRenderAsset.OptionalMipsState == FStreamingRenderAsset::OMS_HasOptionalMips;
+			const int32 NumLODs = bOptLODsExist ? StreamingRenderAsset.MipCount : StreamingRenderAsset.NumNonOptionalMips;
+			const int32 NumStreamedLODs = NumLODs - StreamingRenderAsset.NumNonStreamingMips;
+			const int32 NumResidentLODs = StreamingRenderAsset.ResidentMips;
+			const int32 NumEvictedLODs = NumLODs - NumResidentLODs;
+			const int64 TotalSize = StreamingRenderAsset.GetSize(NumLODs);
+			const int64 StreamedSize = TotalSize - StreamingRenderAsset.GetSize(StreamingRenderAsset.NumNonStreamingMips);
+			const int64 EvictedSize = TotalSize - ResidentSize;
+
+			++Stats.NumStreamedMeshes;
+			TotalNumStreamedLODs += NumStreamedLODs;
+			TotalNumResidentLODs += NumResidentLODs;
+			TotalNumEvictedLODs += NumEvictedLODs;
+			Stats.StreamedMeshMem += StreamedSize;
+			Stats.ResidentMeshMem += ResidentSize;
+			Stats.EvictedMeshMem += EvictedSize;
+		}
+	}
+
+	if (Stats.NumStreamedMeshes > 0)
+	{
+		Stats.AvgNumStreamedLODs = (float)TotalNumStreamedLODs / Stats.NumStreamedMeshes;
+		Stats.AvgNumResidentLODs = (float)TotalNumResidentLODs / Stats.NumStreamedMeshes;
+		Stats.AvgNumEvictedLODs = (float)TotalNumEvictedLODs / Stats.NumStreamedMeshes;
 	}
 
 	Stats.OverBudget += FMath::Max<int64>(Stats.RequiredPool - Stats.StreamingPool, 0);
diff --git a/Engine/Source/Runtime/Engine/Private/Streaming/StreamingManagerTexture.cpp b/Engine/Source/Runtime/Engine/Private/Streaming/StreamingManagerTexture.cpp
--- a/Engine/Source/Runtime/Engine/Private/Streaming/StreamingManagerTexture.cpp
+++ b/Engine/Source/Runtime/Engine/Private/Streaming/StreamingManagerTexture.cpp
@@ -1386,7 +1386,7 @@ void FRenderAssetStreamingManager::UpdateResourceStreaming( float DeltaTime, boo
 
 		STAT(GatheredStats.SetupAsyncTaskCycles = 0);
 		STAT(GatheredStats.UpdateStreamingDataCycles = 0);
-		STAT(GatheredStats.StreamTexturesCycles = 0);
+		STAT(GatheredStats.StreamRenderAssetsCycles = 0);
 		STAT(GatheredStats.CallbacksCycles = 0);
 #if STATS
 		UpdateStats();
@@ -1453,7 +1453,7 @@ void FRenderAssetStreamingManager::UpdateResourceStreaming( float DeltaTime, boo
 	}
 	else if (AsyncWork->IsDone())
 	{
-		STAT(GatheredStats.StreamTexturesCycles = -(int32)FPlatformTime::Cycles();)
+		STAT(GatheredStats.StreamRenderAssetsCycles = -(int32)FPlatformTime::Cycles();)
 
 		// Since this step is lightweight, tick each texture inflight here, to accelerate the state changes.
 		for (int32 TextureIndex : InflightRenderAssets)
@@ -1469,7 +1469,7 @@ void FRenderAssetStreamingManager::UpdateResourceStreaming( float DeltaTime, boo
 
 		ProcessingStage = 0;
 
-		STAT(GatheredStats.StreamTexturesCycles += FPlatformTime::Cycles();)
+		STAT(GatheredStats.StreamRenderAssetsCycles += FPlatformTime::Cycles();)
 #if STATS
 			UpdateStats();
 #elif UE_BUILD_TEST
diff --git a/Engine/Source/Runtime/Engine/Private/Streaming/StreamingTexture.cpp b/Engine/Source/Runtime/Engine/Private/Streaming/StreamingTexture.cpp
--- a/Engine/Source/Runtime/Engine/Private/Streaming/StreamingTexture.cpp
+++ b/Engine/Source/Runtime/Engine/Private/Streaming/StreamingTexture.cpp
@@ -92,7 +92,7 @@ void FStreamingRenderAsset::UpdateStaticData(const FRenderAssetStreamingSettings
 		}
 
 		NumNonOptionalMips = MipCount - RenderAsset->CalcNumOptionalMips();
-		OptionalMipsState = (NumNonOptionalMips == MipCount) ? EOptionalMipsState::NoOptionalMips : EOptionalMipsState::NotCached;
+		OptionalMipsState = (NumNonOptionalMips == MipCount) ? EOptionalMipsState::OMS_NoOptionalMips : EOptionalMipsState::OMS_NotCached;
 
 		const int32 MaxNumMips = RenderAssetType == AT_Texture ? MAX_TEXTURE_MIP_COUNT : MaxNumMeshLODs;
 		for (int32 MipIndex = 0; MipIndex < MaxNumMips; ++MipIndex)
@@ -116,7 +116,7 @@ void FStreamingRenderAsset::UpdateStaticData(const FRenderAssetStreamingSettings
 		BudgetMipBias = 0;
 		BoostFactor = 1.f;
 		NumNonOptionalMips = MipCount;
-		OptionalMipsState = EOptionalMipsState::NoOptionalMips;
+		OptionalMipsState = EOptionalMipsState::OMS_NoOptionalMips;
 
 		bIsCharacterTexture = false;
 		bIsTerrainTexture = false;
@@ -132,9 +132,9 @@ void FStreamingRenderAsset::UpdateOptionalMipsState_Async()
 {
 	// Here we do a lazy update where we check if the highres mip file exists only if it could be useful to do so.
 	// This requires texture to be at max resolution before the optional mips .
-	if (OptionalMipsState == EOptionalMipsState::NotCached && !OptionalBulkDataFilename.IsEmpty())
+	if (OptionalMipsState == EOptionalMipsState::OMS_NotCached && !OptionalBulkDataFilename.IsEmpty())
 	{
-		OptionalMipsState = IFileManager::Get().FileExists(*OptionalBulkDataFilename) ? EOptionalMipsState::HasOptionalMips : EOptionalMipsState::NoOptionalMips;
+		OptionalMipsState = IFileManager::Get().FileExists(*OptionalBulkDataFilename) ? EOptionalMipsState::OMS_HasOptionalMips : EOptionalMipsState::OMS_NoOptionalMips;
 	}
 
 }
@@ -180,7 +180,7 @@ void FStreamingRenderAsset::UpdateDynamicData(const int32* NumStreamedMips, int3
 				// If the optional mips are not available, or if we shouldn't load them now, clamp the possible mips requested. 
 				// (when the non-optional mips are not yet loaded, loading optional mips generates cross files requests).
 				// This is not bullet proof though since the texture/mesh could have a pending stream-out request.
-				if (OptionalMipsState != EOptionalMipsState::HasOptionalMips || ResidentMips < NumNonOptionalMips)
+				if (OptionalMipsState != EOptionalMipsState::OMS_HasOptionalMips || ResidentMips < NumNonOptionalMips)
 				{
 					TempMaxAllowedMips = FMath::Min(TempMaxAllowedMips, NumNonOptionalMips);
 				}
@@ -209,7 +209,7 @@ void FStreamingRenderAsset::UpdateDynamicData(const int32* NumStreamedMips, int3
 		MinAllowedMips = 0;
 		MaxAllowedMips = 0;
 		NumNonOptionalMips = 0;
-		OptionalMipsState = EOptionalMipsState::NotCached;
+		OptionalMipsState = EOptionalMipsState::OMS_NotCached;
 		LastRenderTime = FLT_MAX;	
 	}
 }
diff --git a/Engine/Source/Runtime/Engine/Private/Streaming/StreamingTexture.h b/Engine/Source/Runtime/Engine/Private/Streaming/StreamingTexture.h
--- a/Engine/Source/Runtime/Engine/Private/Streaming/StreamingTexture.h
+++ b/Engine/Source/Runtime/Engine/Private/Streaming/StreamingTexture.h
@@ -19,18 +19,17 @@ struct FRenderAssetStreamingSettings;
 /** Self-contained structure to manage a streaming texture/mesh, possibly on a separate thread. */
 struct FStreamingRenderAsset
 {
-private:
+	static constexpr int32 MaxNumMeshLODs = MAX_MESH_LOD_COUNT;
+	static_assert(2 * MaxNumMeshLODs >= MAX_TEXTURE_MIP_COUNT, "Failed mip count assumption");
+
 	enum EOptionalMipsState : uint8
 	{
-		NotCached,
-		NoOptionalMips,
-		HasOptionalMips,
+		OMS_NotCached,
+		OMS_NoOptionalMips,
+		OMS_HasOptionalMips,
+		OMS_Num
 	};
 
-	static constexpr int32 MaxNumMeshLODs = MAX_MESH_LOD_COUNT;
-	static_assert(2 * MaxNumMeshLODs >= MAX_TEXTURE_MIP_COUNT, "Failed mip count assumption");
-
-public:
 	enum EAssetType : uint8
 	{
 		AT_Texture,
@@ -174,9 +173,9 @@ public:
 	FORCEINLINE void ClearCachedOptionalMipsState_Async()
 	{
 		// If we already have our optional mips there is no need to recache, pak files can't go away!
-		if (OptionalMipsState == EOptionalMipsState::NoOptionalMips && NumNonOptionalMips != MipCount)
+		if (OptionalMipsState == EOptionalMipsState::OMS_NoOptionalMips && NumNonOptionalMips != MipCount)
 		{
-			OptionalMipsState = EOptionalMipsState::NotCached;
+			OptionalMipsState = EOptionalMipsState::OMS_NotCached;
 		}
 	}
 
diff --git a/Engine/Source/Runtime/Engine/Private/Streaming/TextureStreamingHelpers.cpp b/Engine/Source/Runtime/Engine/Private/Streaming/TextureStreamingHelpers.cpp
--- a/Engine/Source/Runtime/Engine/Private/Streaming/TextureStreamingHelpers.cpp
+++ b/Engine/Source/Runtime/Engine/Private/Streaming/TextureStreamingHelpers.cpp
@@ -29,21 +29,29 @@ DECLARE_MEMORY_STAT_POOL(TEXT("IO Bandwidth"), STAT_Streaming14_MipIOBandwidth,
 
 DECLARE_CYCLE_STAT(TEXT("Setup Async Task"), STAT_Streaming01_SetupAsyncTask, STATGROUP_Streaming);
 DECLARE_CYCLE_STAT(TEXT("Update Streaming Data"), STAT_Streaming02_UpdateStreamingData, STATGROUP_Streaming);
-DECLARE_CYCLE_STAT(TEXT("Streaming Texture"), STAT_Streaming03_StreamTextures, STATGROUP_Streaming);
+DECLARE_CYCLE_STAT(TEXT("Streaming Render Assets"), STAT_Streaming03_StreamRenderAssets, STATGROUP_Streaming);
 DECLARE_CYCLE_STAT(TEXT("Notifications"), STAT_Streaming04_Notifications, STATGROUP_Streaming);
 DECLARE_DWORD_COUNTER_STAT(TEXT("Pending 2D Update"), STAT_Streaming05_Pending2DUpdate, STATGROUP_Streaming);
 
 /** Streaming Overview stats */
 
-DECLARE_MEMORY_STAT_POOL(TEXT("Streamable Textures"),	STAT_StreamingOverview01_StreamableTextures, STATGROUP_StreamingOverview, FPlatformMemory::MCR_StreamingPool);
-DECLARE_MEMORY_STAT_POOL(TEXT("     Required"),			STAT_StreamingOverview02_Required, STATGROUP_StreamingOverview, FPlatformMemory::MCR_StreamingPool);
-DECLARE_MEMORY_STAT_POOL(TEXT("     Cached"),			STAT_StreamingOverview03_Cached, STATGROUP_StreamingOverview, FPlatformMemory::MCR_StreamingPool);
+DECLARE_MEMORY_STAT_POOL(TEXT("Streamable Render Assets"), STAT_StreamingOverview01_StreamableRenderAssets, STATGROUP_StreamingOverview, FPlatformMemory::MCR_StreamingPool);
+DECLARE_MEMORY_STAT_POOL(TEXT("     Required"), STAT_StreamingOverview02_Required, STATGROUP_StreamingOverview, FPlatformMemory::MCR_StreamingPool);
+DECLARE_MEMORY_STAT_POOL(TEXT("     Cached"), STAT_StreamingOverview03_Cached, STATGROUP_StreamingOverview, FPlatformMemory::MCR_StreamingPool);
 DECLARE_MEMORY_STAT_POOL(TEXT("Streaming Overbudget"),	STAT_StreamingOverview04_StreamingOverbudget, STATGROUP_StreamingOverview, FPlatformMemory::MCR_StreamingPool);
-DECLARE_MEMORY_STAT_POOL(TEXT("Unstreambale Textures"), STAT_StreamingOverview05_UnstreamableTextures, STATGROUP_StreamingOverview, FPlatformMemory::MCR_StreamingPool);
-DECLARE_MEMORY_STAT_POOL(TEXT("     NeverStream"),		STAT_StreamingOverview06_NeverStream, STATGROUP_StreamingOverview, FPlatformMemory::MCR_StreamingPool);
-DECLARE_MEMORY_STAT_POOL(TEXT("     UI Group"),			STAT_StreamingOverview07_UIGroup, STATGROUP_StreamingOverview, FPlatformMemory::MCR_StreamingPool);
+DECLARE_MEMORY_STAT_POOL(TEXT("Unstreambale Render Assets"), STAT_StreamingOverview05_UnstreamableRenderAssets, STATGROUP_StreamingOverview, FPlatformMemory::MCR_StreamingPool);
+DECLARE_MEMORY_STAT_POOL(TEXT("     NeverStream"), STAT_StreamingOverview06_NeverStream, STATGROUP_StreamingOverview, FPlatformMemory::MCR_StreamingPool);
+DECLARE_MEMORY_STAT_POOL(TEXT("     UI Group"), STAT_StreamingOverview07_UIGroup, STATGROUP_StreamingOverview, FPlatformMemory::MCR_StreamingPool);
 DECLARE_MEMORY_STAT_POOL(TEXT("Average Required PoolSize"),	STAT_StreamingOverview08_AverageRequiredPool, STATGROUP_StreamingOverview, FPlatformMemory::MCR_StreamingPool);
 
+DECLARE_DWORD_COUNTER_STAT(TEXT("Number of Streamed Meshes"), STAT_StreamingOverview09_NumStreamedMeshes, STATGROUP_StreamingOverview);
+DECLARE_FLOAT_COUNTER_STAT(TEXT("Average Number of Streamed Mesh LODs"), STAT_StreamingOverview10_AvgNumStreamedLODs, STATGROUP_StreamingOverview);
+DECLARE_FLOAT_COUNTER_STAT(TEXT("Average Number of Resident Mesh LODs"), STAT_StreamingOverview11_AvgNumResidentLODs, STATGROUP_StreamingOverview);
+DECLARE_FLOAT_COUNTER_STAT(TEXT("Average Number of Evicted Mesh LODs"), STAT_StreamingOverview12_AvgNumEvictedLODs, STATGROUP_StreamingOverview);
+DECLARE_MEMORY_STAT_POOL(TEXT("Mesh LOD Bytes Streamable"), STAT_StreamingOverview13_StreamedMeshMem, STATGROUP_StreamingOverview, FPlatformMemory::MCR_StreamingPool);
+DECLARE_MEMORY_STAT_POOL(TEXT("Mesh LOD Bytes Resident"), STAT_StreamingOverview14_ResidentMeshMem, STATGROUP_StreamingOverview, FPlatformMemory::MCR_StreamingPool);
+DECLARE_MEMORY_STAT_POOL(TEXT("Mesh LOD Bytes Evicted"), STAT_StreamingOverview15_EvictedMeshMem, STATGROUP_StreamingOverview, FPlatformMemory::MCR_StreamingPool);
+
 DEFINE_STAT(STAT_RenderAssetStreaming_GameThreadUpdateTime);
 
 DEFINE_LOG_CATEGORY(LogContentStreaming);
@@ -383,7 +391,7 @@ void FRenderAssetStreamingStats::Apply()
 
 	SET_CYCLE_COUNTER(STAT_Streaming01_SetupAsyncTask, SetupAsyncTaskCycles);
 	SET_CYCLE_COUNTER(STAT_Streaming02_UpdateStreamingData, UpdateStreamingDataCycles);
-	SET_CYCLE_COUNTER(STAT_Streaming03_StreamTextures, StreamTexturesCycles);
+	SET_CYCLE_COUNTER(STAT_Streaming03_StreamRenderAssets, StreamRenderAssetsCycles);
 	SET_CYCLE_COUNTER(STAT_Streaming04_Notifications, CallbacksCycles);
 	INC_DWORD_STAT_BY(STAT_Streaming05_Pending2DUpdate, GPending2DUpdateCount);
 
@@ -395,14 +403,22 @@ void FRenderAssetStreamingStats::Apply()
 	GAverageRequiredPool = (int64)((double)GRequiredPoolSizeSum / (double)FMath::Max<int64>(1, GRequiredPoolSizeCount));
 #endif
 
-	SET_MEMORY_STAT(STAT_StreamingOverview01_StreamableTextures, RequiredPool + CachedMips); 
+	SET_MEMORY_STAT(STAT_StreamingOverview01_StreamableRenderAssets, RequiredPool + CachedMips); 
 	SET_MEMORY_STAT(STAT_StreamingOverview02_Required, RequiredPool);
 	SET_MEMORY_STAT(STAT_StreamingOverview03_Cached, CachedMips); 
 	SET_MEMORY_STAT(STAT_StreamingOverview04_StreamingOverbudget, FMath::Max<int64>(RequiredPool - StreamingPool, 0)); 
-	SET_MEMORY_STAT(STAT_StreamingOverview05_UnstreamableTextures, NonStreamingMips);
+	SET_MEMORY_STAT(STAT_StreamingOverview05_UnstreamableRenderAssets, NonStreamingMips);
 	SET_MEMORY_STAT(STAT_StreamingOverview06_NeverStream, GNeverStreamTextureMemory);
 	SET_MEMORY_STAT(STAT_StreamingOverview07_UIGroup, GUITextureMemory);
 	SET_MEMORY_STAT(STAT_StreamingOverview08_AverageRequiredPool, GAverageRequiredPool); 
+
+	SET_DWORD_STAT(STAT_StreamingOverview09_NumStreamedMeshes, NumStreamedMeshes);
+	SET_FLOAT_STAT(STAT_StreamingOverview10_AvgNumStreamedLODs, AvgNumStreamedLODs);
+	SET_FLOAT_STAT(STAT_StreamingOverview11_AvgNumResidentLODs, AvgNumResidentLODs);
+	SET_FLOAT_STAT(STAT_StreamingOverview12_AvgNumEvictedLODs, AvgNumEvictedLODs);
+	SET_MEMORY_STAT(STAT_StreamingOverview13_StreamedMeshMem, StreamedMeshMem);
+	SET_MEMORY_STAT(STAT_StreamingOverview14_ResidentMeshMem, ResidentMeshMem);
+	SET_MEMORY_STAT(STAT_StreamingOverview15_EvictedMeshMem, EvictedMeshMem);
 }
 
 void ResetAverageRequiredTexturePoolSize()
diff --git a/Engine/Source/Runtime/Engine/Private/Streaming/TextureStreamingHelpers.h b/Engine/Source/Runtime/Engine/Private/Streaming/TextureStreamingHelpers.h
--- a/Engine/Source/Runtime/Engine/Private/Streaming/TextureStreamingHelpers.h
+++ b/Engine/Source/Runtime/Engine/Private/Streaming/TextureStreamingHelpers.h
@@ -166,5 +166,13 @@ struct FRenderAssetStreamingStats
 	volatile int32 CallbacksCycles;
 	int32 SetupAsyncTaskCycles;
 	int32 UpdateStreamingDataCycles;
-	int32 StreamTexturesCycles;
+	int32 StreamRenderAssetsCycles; // CPU cycles used to process Mip/LOD load and unload requests
+
+	int32 NumStreamedMeshes;	// Number of meshes managed by the streamer
+	float AvgNumStreamedLODs;	// Average number of mesh LODs that can be streamed
+	float AvgNumResidentLODs;	// Average number of mesh LODs resident
+	float AvgNumEvictedLODs;	// Average number of mesh LODs evicted
+	int64 StreamedMeshMem;		// Total memory in bytes of mesh LODs that can be streamed
+	int64 ResidentMeshMem;		// Total memory in bytes of resident mesh LODs
+	int64 EvictedMeshMem;		// Total memory in bytes of evicted mesh LODs
 };
