Fix logging for deprecated warning

diff --git a/Engine/Source/Runtime/Engine/Private/ShaderCompiler/ShaderCompiler.cpp b/Engine/Source/Runtime/Engine/Private/ShaderCompiler/ShaderCompiler.cpp
--- a/Engine/Source/Runtime/Engine/Private/ShaderCompiler/ShaderCompiler.cpp
+++ b/Engine/Source/Runtime/Engine/Private/ShaderCompiler/ShaderCompiler.cpp
@@ -4449,7 +4449,7 @@ void RecompileShadersForRemote(
 				// These platforms are deprecated and we should warn about their use
 				if (ShaderPlatform == SP_OPENGL_SM5 || ShaderPlatform == SP_PCD3D_SM4)
 				{
-					UE_LOG(LogShaders, Warning, TEXT("You are compiling shaders for a deprecated platform '%s'"), *LegacyShaderPlatformToShaderFormat(ShaderPlatform).ToString());
+					UE_LOG(LogShaderCompilers, Warning, TEXT("You are compiling shaders for a deprecated platform '%s'"), *LegacyShaderPlatformToShaderFormat(ShaderPlatform).ToString());
 				}
 
 				if (bCompileChangedShaders)
