Fix-up MaxChannel count code so it
1. is safer (MaxChannel-related members are now privatized within AudioDevice implementation).
2. allows for testing with au.AudioMaxChannels cvar properly.

diff --git a/Engine/Source/Runtime/Engine/Private/AudioDevice.cpp b/Engine/Source/Runtime/Engine/Private/AudioDevice.cpp
--- a/Engine/Source/Runtime/Engine/Private/AudioDevice.cpp
+++ b/Engine/Source/Runtime/Engine/Private/AudioDevice.cpp
@@ -199,11 +199,7 @@ FAttenuationListenerData FAttenuationListenerData::Create(const FAudioDevice& Au
 -----------------------------------------------------------------------------*/
 
 FAudioDevice::FAudioDevice()
-	: MaxChannels(0)
-	, MaxChannels_GameThread(0)
-	, MaxChannelsScale(1.0f)
-	, MaxChannelsScale_GameThread(1.0f)
-	, NumStoppingVoices(32)
+	: NumStoppingVoices(32)
 	, SampleRate(0)
 	, NumPrecacheFrames(MONO_PCM_BUFFER_SAMPLES)
 	, CommonAudioPoolSize(0)
@@ -213,7 +209,12 @@ FAudioDevice::FAudioDevice()
 	, SpatializationPluginInterface(nullptr)
 	, ReverbPluginInterface(nullptr)
 	, OcclusionInterface(nullptr)
-	, PluginListeners()
+	, AmbisonicsMixer(nullptr)
+	, MaxSources(0)
+	, MaxChannels(0)
+	, MaxChannels_GameThread(0)
+	, MaxChannelsScale(1.0f)
+	, MaxChannelsScale_GameThread(1.0f)
 	, CurrentTick(0)
 	, TestAudioComponent(nullptr)
 	, DebugState(DEBUGSTATE_None)
@@ -277,7 +278,7 @@ const FAudioQualitySettings& FAudioDevice::GetQualityLevelSettings()
 	return AudioSettings->GetQualityLevelSettings(QualityLevel);
 }
 
-bool FAudioDevice::Init(int32 InMaxChannels)
+bool FAudioDevice::Init(int32 InMaxSources)
 {
 	SCOPED_BOOT_TIMING("FAudioDevice::Init");
 
@@ -296,14 +297,16 @@ bool FAudioDevice::Init(int32 InMaxChannels)
 	// Get a copy of the platform-specific settings (overridden by platforms)
 	PlatformSettings = GetPlatformSettings();
 
-	// Max channels is the max value supplied to Init call (quality settings), unless overwritten by the platform settings.
+	// MaxSources is the max value supplied to Init call (quality settings), unless overwritten by the platform settings.
 	// This does not have to be the minimum value in this case (nor is it desired, so platforms can potentially scale up)
-	// as the Sources array has yet to be initialized.
-	MaxChannels = PlatformSettings.MaxChannels > 0 ? PlatformSettings.MaxChannels : InMaxChannels;
-	MaxChannels_GameThread = MaxChannels;
+	// as the Sources array has yet to be initialized. If the cvar is largest, take that value to allow for testing
+	const int32 PlatformMaxSources = PlatformSettings.MaxChannels > 0 ? PlatformSettings.MaxChannels : InMaxSources;
+	MaxSources = FMath::Max(PlatformMaxSources, AudioChannelCountCVar);
+	MaxChannels = MaxSources;
+	MaxChannels_GameThread = MaxSources;
 
 	// Ensure and not assert so if in editor, user can change quality setting and re-serialize if so desired.
-	ensureMsgf(MaxChannels > 0, TEXT("Neither passed nor platform MaxChannel setting was positive value"));
+	ensureMsgf(MaxSources > 0, TEXT("Neither passed MaxSources nor platform MaxChannel setting was positive value"));
 
 	// Mixed sample rate is set by the platform
 	SampleRate = PlatformSettings.SampleRate;
@@ -378,7 +381,7 @@ bool FAudioDevice::Init(int32 InMaxChannels)
 			//Set up initialization parameters for system level effect plugins:
 			FAudioPluginInitializationParams PluginInitializationParams;
 			PluginInitializationParams.SampleRate = SampleRate;
-			PluginInitializationParams.NumSources = MaxChannels;
+			PluginInitializationParams.NumSources = MaxSources;
 			PluginInitializationParams.BufferLength = PlatformSettings.CallbackBufferFrameSize;
 			PluginInitializationParams.AudioDevicePtr = this;
 
@@ -505,10 +508,13 @@ void FAudioDevice::SetMaxChannels(int32 InMaxChannels)
 		return;
 	}
 
-	if (!IsInAudioThread())
+	if (IsInGameThread())
 	{
 		MaxChannels_GameThread = InMaxChannels;
+	}
 
+	if (!IsInAudioThread())
+	{
 		DECLARE_CYCLE_STAT(TEXT("FAudioThreadTask.SetMaxChannels"), STAT_AudioSetMaxChannels, STATGROUP_AudioThreadCommands);
 
 		FAudioThread::RunCommandOnAudioThread([this, InMaxChannels]()
@@ -516,13 +522,11 @@ void FAudioDevice::SetMaxChannels(int32 InMaxChannels)
 			this->SetMaxChannels(InMaxChannels);
 
 		}, GET_STATID(STAT_AudioSetMaxChannels));
-
-		return;
 	}
 
-	if (InMaxChannels > MaxChannels)
+	if (InMaxChannels > MaxSources)
 	{
-		UE_LOG(LogAudio, Warning, TEXT("Can't increase MaxChannels past existing setting!"));
+		UE_LOG(LogAudio, Warning, TEXT("Can't increase MaxChannels past MaxSources"));
 		return;
 	}
 
@@ -561,24 +565,19 @@ void FAudioDevice::SetMaxChannelsScaled(float InScaledChannelCount)
 
 int32 FAudioDevice::GetMaxChannels() const
 {
-	if (IsInAudioThread())
-	{
-		if (AudioChannelCountCVar > 0 && AudioChannelCountCVar < MaxChannels)
-		{
-			return FMath::Max(int32(AudioChannelCountCVar * MaxChannelsScale * AudioChannelCountScaleCVar), 1);
-		}
+	// Get thread-context version of channel scalar & scale by cvar scalar
+	float MaxChannelScalarToApply = IsInAudioThread() ? MaxChannelsScale : MaxChannelsScale_GameThread;
+	MaxChannelScalarToApply *= AudioChannelCountScaleCVar;
 
-		return FMath::Max(int32(MaxChannels * MaxChannelsScale * AudioChannelCountScaleCVar), 1);
-	}
-	else
+	// Get thread-context version of channel max.  Override by cvar if cvar is valid.
+	int32 OutMaxChannels = IsInAudioThread() ? MaxChannels : MaxChannels_GameThread;
+	if (AudioChannelCountCVar > 0)
 	{
-		if (AudioChannelCountCVar > 0 && AudioChannelCountCVar < MaxChannels)
-		{
-			return FMath::Max(int32(AudioChannelCountCVar * MaxChannelsScale_GameThread * AudioChannelCountScaleCVar), 1);
-		}
-
-		return FMath::Max(int32(MaxChannels_GameThread * MaxChannelsScale_GameThread * AudioChannelCountScaleCVar), 1);
+		OutMaxChannels = AudioChannelCountCVar;
 	}
+
+	// Find product of max channels and final scalar, and clamp between 1 and MaxSources.
+	return FMath::Clamp(static_cast<int32>(OutMaxChannels * MaxChannelScalarToApply), 1, MaxSources);
 }
 
 void FAudioDevice::Teardown()
diff --git a/Engine/Source/Runtime/Engine/Public/AudioDevice.h b/Engine/Source/Runtime/Engine/Public/AudioDevice.h
--- a/Engine/Source/Runtime/Engine/Public/AudioDevice.h
+++ b/Engine/Source/Runtime/Engine/Public/AudioDevice.h
@@ -450,7 +450,7 @@ public:
 	/**
 	 * Basic initialization of the platform agnostic layer of the audio system
 	 */
-	bool Init(int32 InMaxChannels);
+	bool Init(int32 InMaxSources);
 
 	/**
 	 * Tears down the audio device
@@ -1583,14 +1583,6 @@ public:
 
 public:
 
-	/** The maximum number of concurrent audible sounds */
-	int32 MaxChannels;
-	int32 MaxChannels_GameThread;
-
-	/** A scaler on the max channels. */
-	float MaxChannelsScale;
-	float MaxChannelsScale_GameThread;
-
 	/** The number of sources to reserve for stopping sounds. */
 	int32 NumStoppingVoices;
 
@@ -1637,6 +1629,17 @@ public:
 	TArray<FTransform> ListenerTransforms;
 
 private:
+	/** The maximum number of sources.  Value cannot change after initialization. */
+	int32 MaxSources;
+
+	/** The maximum number of concurrent audible sounds. Value cannot exceed MaxSources. */
+	int32 MaxChannels;
+	int32 MaxChannels_GameThread;
+
+	/** Normalized (0.0f - 1.0f) scalar multiplier on max channels. */
+	float MaxChannelsScale;
+	float MaxChannelsScale_GameThread;
+
 	uint64 CurrentTick;
 
 	/** An AudioComponent to play test sounds on */
