Fix compile error in UE4 main. #jira none

diff --git a/Engine/Source/Runtime/Engine/Private/AudioDebug.cpp b/Engine/Source/Runtime/Engine/Private/AudioDebug.cpp
--- a/Engine/Source/Runtime/Engine/Private/AudioDebug.cpp
+++ b/Engine/Source/Runtime/Engine/Private/AudioDebug.cpp
@@ -924,7 +924,7 @@ int32 FAudioDebugger::RenderStatWaves(UWorld* World, FViewport* Viewport, FCanva
 
 	const int32 ActiveInstances = WaveInstances.Num();
 
-	const int32 Max = AudioDevice->MaxChannels / 2;
+	const int32 Max = AudioDevice->GetMaxChannels() / 2;
 	float f = FMath::Clamp<float>((float)(ActiveInstances - Max) / (float)Max, 0.f, 1.f);
 	const int32 R = FMath::TruncToInt(f * 255);
 
