Fix for compile issues from an editor dependancy in EngineSessionManager.
This shoud be fixed with a refactor that spilts out editor only code, that won't happen for 4.23 however.

diff --git a/Engine/Source/Runtime/Engine/Private/EngineSessionManager.cpp b/Engine/Source/Runtime/Engine/Private/EngineSessionManager.cpp
--- a/Engine/Source/Runtime/Engine/Private/EngineSessionManager.cpp
+++ b/Engine/Source/Runtime/Engine/Private/EngineSessionManager.cpp
@@ -23,12 +23,15 @@
 #include "HAL/PlatformProcess.h"
 #include "HAL/PlatformOutputDevices.h"
 #include "Interfaces/IPluginManager.h"
-#include "IVREditorModule.h"
-#include "Kismet2/DebuggerCommands.h"
 #include "Interfaces/IProjectManager.h"
 #include "Misc/EngineVersion.h"
 #include "RHI.h"
 
+#if WITH_EDITOR
+#include "IVREditorModule.h"
+#include "Kismet2/DebuggerCommands.h"
+#endif // WITH_EDITOR
+
 #define LOCTEXT_NAMESPACE "SessionManager"
 
 DEFINE_LOG_CATEGORY(LogEngineSessionManager);
@@ -287,6 +290,7 @@ void FEngineSessionManager::Tick(float DeltaTime)
 			}
 #endif
 
+#if WITH_EDITOR
 			extern ENGINE_API float GAverageFPS;
 
 			CurrentSession.AverageFPS = GAverageFPS;
@@ -304,6 +308,7 @@ void FEngineSessionManager::Tick(float DeltaTime)
 			CurrentSession.bIsInPIE = FPlayWorldCommandCallbacks::IsInPIE();
 			const FString IsInPIEString = BoolToStoredString(CurrentSession.bIsInPIE);
 			FPlatformMisc::SetStoredValue(SessionManagerDefs::StoreId, CurrentSessionSectionName, SessionManagerDefs::IsInPIEStoreKey, IsInPIEString);
+#endif
 		}
 	}
 }
