PR #5433: Remove the incorrect use of 'absolution' from comment. (Contributed by [REDACTED_PERSON])

diff --git a/Engine/Source/Runtime/Engine/Classes/Sound/SoundAttenuation.h b/Engine/Source/Runtime/Engine/Classes/Sound/SoundAttenuation.h
--- a/Engine/Source/Runtime/Engine/Classes/Sound/SoundAttenuation.h
+++ b/Engine/Source/Runtime/Engine/Classes/Sound/SoundAttenuation.h
@@ -143,7 +143,7 @@ struct ENGINE_API FSoundAttenuationSettings : public FBaseAttenuationSettings
 	UPROPERTY(EditAnywhere, BlueprintReadWrite, Category = AttenuationSpatialization, meta=(ClampMin = "0", EditCondition="bSpatialize", DisplayName="Non-Spatialized Radius"))
 	float OmniRadius;
 
-	/** The world-space absolution distance between left and right stereo channels when stereo assets are 3D spatialized. */
+	/** The world-space distance between left and right stereo channels when stereo assets are 3D spatialized. */
 	UPROPERTY(EditAnywhere, BlueprintReadWrite, Category = AttenuationSpatialization, meta = (ClampMin = "0", EditCondition = "bSpatialize", DisplayName = "3D Stereo Spread"))
 	float StereoSpread;
 
