Also broadcast LevelRemovedFromWorld delegate when level gets removed (was in World::RemoveFromWorld but not in UWorld::RemoveLevel)

diff --git a/Engine/Source/Runtime/Engine/Private/World.cpp b/Engine/Source/Runtime/Engine/Private/World.cpp
--- a/Engine/Source/Runtime/Engine/Private/World.cpp
+++ b/Engine/Source/Runtime/Engine/Private/World.cpp
@@ -6565,6 +6565,7 @@ bool UWorld::RemoveLevel( ULevel* InLevel )
 		}
 #endif //WITH_EDITOR
 		Levels.Remove( InLevel );
+		FWorldDelegates::LevelRemovedFromWorld.Broadcast(InLevel, this);
 		BroadcastLevelsChanged();
 	}
 	return bRemovedLevel;
