Despite this fix I am setting Switch to not precache anyways, to make audiomixer memory characteristics mirror switch's characteristics when audio mixer is not on.

diff --git a/Engine/Source/Runtime/Engine/Private/ADPCMAudioInfo.cpp b/Engine/Source/Runtime/Engine/Private/ADPCMAudioInfo.cpp
--- a/Engine/Source/Runtime/Engine/Private/ADPCMAudioInfo.cpp
+++ b/Engine/Source/Runtime/Engine/Private/ADPCMAudioInfo.cpp
@@ -278,6 +278,13 @@ bool FADPCMAudioInfo::ReadCompressedInfo(const uint8* InSrcBufferData, uint32 In
 
 bool FADPCMAudioInfo::ReadCompressedData(uint8* Destination, bool bLooping, uint32 BufferSize)
 {
+	// If we've already read through this asset and we are not looping, memzero and early out.
+	if (TotalSamplesStreamed >= TotalSamplesPerChannel && !bLooping)
+	{
+		FMemory::Memzero(Destination, BufferSize);
+		return true;
+	}
+
 	const uint32 ChannelSampleSize = sizeof(uint16) * NumChannels;
 
 	// This correctly handles any BufferSize as long as its a multiple of sample size * number of channels
@@ -340,16 +347,19 @@ bool FADPCMAudioInfo::ReadCompressedData(uint8* Destination, bool bLooping, uint
 			if(TotalSamplesStreamed >= TotalSamplesPerChannel)
 			{
 				ReachedEndOfSamples = true;
-				// This is set to the max value to trigger the decompression of the first audio block
-				CurrentUncompressedBlockSampleIndex = UncompressedBlockSize / sizeof(uint16);
-				CurrentCompressedBlockIndex = 0;
-				TotalSamplesStreamed = 0;
 				if(!bLooping)
 				{
 					// Zero remaining buffer
 					FMemory::Memzero(OutData, BufferSize);
 					return true;
 				}
+				else
+				{
+					// This is set to the max value to trigger the decompression of the first audio block
+					CurrentUncompressedBlockSampleIndex = UncompressedBlockSize / sizeof(uint16);
+					CurrentCompressedBlockIndex = 0;
+					TotalSamplesStreamed = 0;
+				}
 			}
 		}
 	}
