#UE4 Direct references from a primary asset via asset bundles no longer imply unconditional ownership. Higher priority assets will still be owners in these cases

diff --git a/Engine/Source/Runtime/Engine/Private/AssetManager.cpp b/Engine/Source/Runtime/Engine/Private/AssetManager.cpp
--- a/Engine/Source/Runtime/Engine/Private/AssetManager.cpp
+++ b/Engine/Source/Runtime/Engine/Private/AssetManager.cpp
@@ -2790,27 +2790,7 @@ EAssetSetManagerResult::Type UAssetManager::ShouldSetManager(const FAssetIdentif
 			return EAssetSetManagerResult::DoNotSet;
 		}
 	}
-	else if (Flags & EAssetSetManagerFlags::TargetHasDirectManager)
-	{
-		// If target has another direct manager being set in this run, never recurse and set manager if we think this is an "owner" reference and not a back reference
-
-		bool bIsOwnershipReference = Flags & EAssetSetManagerFlags::IsDirectSet;
-
-		if (ManagerPrimaryAssetId.PrimaryAssetType == MapType)
-		{
-			// References made by maps are ownership references, because there is no way to distinguish between sublevels and top level maps we "include" sublevels in parent maps via reference
-			bIsOwnershipReference = true;
-		}
-
-		if (bIsOwnershipReference)
-		{
-			return EAssetSetManagerResult::SetButDoNotRecurse;
-		}
-		else
-		{
-			return EAssetSetManagerResult::DoNotSet;
-		}
-	}
+	
 	return EAssetSetManagerResult::SetAndRecurse;
 }
 
