Fixed state managament bug causing PrecacheState mismatch assertion failure when toggling audio mixer in packaged builds

diff --git a/Engine/Source/Runtime/Engine/Private/AudioDeviceManager.cpp b/Engine/Source/Runtime/Engine/Private/AudioDeviceManager.cpp
--- a/Engine/Source/Runtime/Engine/Private/AudioDeviceManager.cpp
+++ b/Engine/Source/Runtime/Engine/Private/AudioDeviceManager.cpp
@@ -222,9 +222,6 @@ void FAudioDeviceManager::ToggleAudioMixer()
 			{
 				USoundWave* SoundWave = *SoundWaveIt;
 				FreeResource(SoundWave);
-
-				// Flag that the sound wave needs to do a full decompress again
-				SoundWave->DecompressionType = DTYPE_Setup;
 			}
 
 			// Unload the previous audio device module
@@ -861,6 +858,10 @@ void FAudioDeviceManager::FreeResource(USoundWave* SoundWave)
 		FSoundBuffer* SoundBuffer = WaveBufferMap.FindRef(SoundWave->ResourceID);
 		FreeBufferResource(SoundBuffer);
 
+		// Flag that the sound wave needs to do a full decompress again
+		SoundWave->DecompressionType = DTYPE_Setup;
+		SoundWave->SetPrecacheState(ESoundWavePrecacheState::NotStarted);
+
 		SoundWave->ResourceID = 0;
 	}
 }
diff --git a/Engine/Source/Runtime/Engine/Private/SoundWave.cpp b/Engine/Source/Runtime/Engine/Private/SoundWave.cpp
--- a/Engine/Source/Runtime/Engine/Private/SoundWave.cpp
+++ b/Engine/Source/Runtime/Engine/Private/SoundWave.cpp
@@ -1438,6 +1438,7 @@ void USoundWave::FreeResources()
 	ResourceID = 0;
 	bDynamicResource = false;
 	DecompressionType = DTYPE_Setup;
+	SetPrecacheState(ESoundWavePrecacheState::NotStarted);
 	bDecompressedFromOgg = false;
 
 	if (ResourceState == ESoundWaveResourceState::Freeing)
