Export the vtable for InterpTrack and InterpTrackInst, so licensees can derive from them on Mac/Linux.

#rnx

diff --git a/Engine/Source/Runtime/Engine/Classes/Matinee/InterpTrack.h b/Engine/Source/Runtime/Engine/Classes/Matinee/InterpTrack.h
--- a/Engine/Source/Runtime/Engine/Classes/Matinee/InterpTrack.h
+++ b/Engine/Source/Runtime/Engine/Classes/Matinee/InterpTrack.h
@@ -100,7 +100,7 @@ struct FSubTrackGroup
 };
 
 UCLASS(collapsecategories, hidecategories=Object, abstract, MinimalAPI)
-class UInterpTrack : public UObject, public FCurveEdInterface, public FInterpEdInputInterface
+class ENGINE_VTABLE UInterpTrack : public UObject, public FCurveEdInterface, public FInterpEdInputInterface
 {
 	GENERATED_UCLASS_BODY()
 
diff --git a/Engine/Source/Runtime/Engine/Classes/Matinee/InterpTrackInst.h b/Engine/Source/Runtime/Engine/Classes/Matinee/InterpTrackInst.h
--- a/Engine/Source/Runtime/Engine/Classes/Matinee/InterpTrackInst.h
+++ b/Engine/Source/Runtime/Engine/Classes/Matinee/InterpTrackInst.h
@@ -17,7 +17,7 @@ class AActor;
 class UInterpTrack;
 
 UCLASS(MinimalAPI)
-class UInterpTrackInst : public UObject
+class ENGINE_VTABLE UInterpTrackInst : public UObject
 {
 	GENERATED_UCLASS_BODY()
 
