Fix crash in FindTeleportSpot when actor's root is not primitive component.

diff --git a/Engine/Source/Runtime/Engine/Private/LevelActor.cpp b/Engine/Source/Runtime/Engine/Private/LevelActor.cpp
--- a/Engine/Source/Runtime/Engine/Private/LevelActor.cpp
+++ b/Engine/Source/Runtime/Engine/Private/LevelActor.cpp
@@ -824,11 +824,15 @@ bool UWorld::FindTeleportSpot(const AActor* TestActor, FVector& TestLocation, FR
 		else
 		{
 			USceneComponent* const RootComponent = TestActor->GetRootComponent();
-			FCollisionShape shape =  Cast<UPrimitiveComponent>(RootComponent)->GetCollisionShape();
-			if (shape.IsBox() && (Cast<UBoxComponent>(RootComponent) == nullptr))
+			UPrimitiveComponent* PrimComponent = Cast<UPrimitiveComponent>(RootComponent);
+			if(PrimComponent != nullptr)
 			{
-				UE_LOG(LogPhysics, Warning, TEXT("UWorld::FindTeleportSpot called with an actor that is intersecting geometry. Failed to find new location likely due to "
-					" actor's root component not being a collider component."));
+				FCollisionShape shape = PrimComponent->GetCollisionShape();
+				if (shape.IsBox() && (Cast<UBoxComponent>(PrimComponent) == nullptr))
+				{
+					UE_LOG(LogPhysics, Warning, TEXT("UWorld::FindTeleportSpot called with an actor that is intersecting geometry. Failed to find new location likely due to "
+						" actor's root component not being a collider component."));
+				}
 			}
 		}
 
