Copying //UE4/Dev-Physics to Dev-Main (//UE4/Dev-Main) [at] 7703071

diff --git a/Engine/Source/Runtime/AnimGraphRuntime/Public/BoneControllers/AnimNode_RigidBody_Chaos.h b/Engine/Source/Runtime/AnimGraphRuntime/Public/BoneControllers/AnimNode_RigidBody_Chaos.h
--- a/Engine/Source/Runtime/AnimGraphRuntime/Public/BoneControllers/AnimNode_RigidBody_Chaos.h
+++ b/Engine/Source/Runtime/AnimGraphRuntime/Public/BoneControllers/AnimNode_RigidBody_Chaos.h
@@ -5,26 +5,27 @@
 #include "BoneControllers/AnimNode_RigidBody.h"
 #include "BoneControllers/AnimNode_SkeletalControlBase.h"
 #include "Components/SkeletalMeshComponent.h"
-#include "GeometryCollection/GeometryCollectionSimulationTypes.h"
-#include "GeometryCollection/PhysicsAssetSimulation.h"
-#include "PhysicalMaterials/Experimental/ChaosPhysicalMaterial.h"
+#include "Physics/ImmediatePhysics/ImmediatePhysicsDeclares.h"
 
 #include "AnimNode_RigidBody_Chaos.generated.h"
 
-
-#if INCLUDE_CHAOS
-namespace Chaos
-{
-	class FPBDRigidsSolver;
-}
-class FBoneHierarchy;
-class FSkeletalMeshPhysicsObject;
-struct FSkeletalMeshPhysicsObjectInputs;
-struct FSkeletalMeshPhysicsObjectOutputs;
-#endif
+struct FBodyInstance;
+struct FConstraintInstance;
 
 extern ANIMGRAPHRUNTIME_API TAutoConsoleVariable<int32> CVarEnableChaosRigidBodyNode;
 
+///** Determines in what space the simulation should run */
+//UENUM()
+//enum class ESimulationSpace : uint8
+//{
+//	/** Simulate in component space. Moving the entire skeletal mesh will have no affect on velocities */
+//	ComponentSpace,
+//	/** Simulate in world space. Moving the skeletal mesh will generate velocity changes */
+//	WorldSpace,
+//	/** Simulate in another bone space. Moving the entire skeletal mesh and individually modifying the base bone will have no affect on velocities */
+//	BaseBoneSpace,
+//};
+
 /**
  *	Controller that simulates physics based on the physics asset of the skeletal mesh component
  */
@@ -33,6 +34,7 @@ struct ANIMGRAPHRUNTIME_API FAnimNode_RigidBody_Chaos : public FAnimNode_Skeleta
 {
 	GENERATED_USTRUCT_BODY()
 
+public:
 	FAnimNode_RigidBody_Chaos();
 	~FAnimNode_RigidBody_Chaos();
 
@@ -48,122 +50,18 @@ struct ANIMGRAPHRUNTIME_API FAnimNode_RigidBody_Chaos : public FAnimNode_Skeleta
 	virtual void PreUpdate(const UAnimInstance* InAnimInstance) override;
 	virtual void UpdateInternal(const FAnimationUpdateContext& Context) override;
 	virtual bool HasPreUpdate() const override { return true; }
-	virtual bool IsValidToEvaluate(const USkeleton* Skeleton, const FBoneContainer& RequiredBones) override { return true; }
+	virtual bool IsValidToEvaluate(const USkeleton* Skeleton, const FBoneContainer& RequiredBones) override;
 	virtual bool NeedsDynamicReset() const override;
 	virtual void ResetDynamics(ETeleportType InTeleportType) override;
 	virtual int32 GetLODThreshold() const override;
 	// End of FAnimNode_SkeletalControlBase interface
 
-public:
-	//
-	// ChaosPhysics
-	//
+	void PostSerialize(const FArchive& Ar);
 
 	/** Physics asset to use. If empty use the skeletal mesh's default physics asset */
 	UPROPERTY(EditAnywhere, Category = Settings)
-		UPhysicsAsset* OverridePhysicsAsset;
-
-	/** Physical Properties */
-	UPROPERTY(EditAnywhere, BlueprintReadOnly, Category = "ChaosPhysics")
-		const UChaosPhysicalMaterial* PhysicalMaterial;
-
-	//
-	// ChaosPhysics | General
-	//
-
-	/** When Simulating is enabled the Component will initialize its rigid bodies within the solver. */
-	UPROPERTY(EditAnywhere, BlueprintReadWrite, Category = "ChaosPhysics|General")
-		bool bSimulating;
-
-	/** Number of solver interations. */
-	UPROPERTY(EditAnywhere, BlueprintReadWrite, Category = "ChaosPhysics|General")
-		int32 NumIterations;
-
-	/** If true, this component will get collision notification events (@see IChaosNotifyHandlerInterface) */
-	UPROPERTY(EditAnywhere, BlueprintReadWrite, Category = "ChaosPhysics|General")
-		bool bNotifyCollisions;
-
-	/** ObjectType defines how to initialize the rigid collision structures. */
-	UPROPERTY(EditAnywhere, BlueprintReadWrite, Category = "ChaosPhysics|General")
-		EObjectStateTypeEnum ObjectType;
-
-	/** Density / mass.
-	 *
-	 * Common densities in g/cm^3:
-	 *     gold: 19.3
-	 *     lead: 11.3
-	 *     copper: 8.3 - 9.0
-	 *     steel: 8.03
-	 *     iron: 7.8
-	 *     aluminum: 2.7
-	 *     glass: 2.4 - 2.8
-	 *     brick: 1.4 - 2.4
-	 *     concrete: 0.45 - 2.4
-	 *     bone: 1.7 - 2.0
-	 *     muscle: 1.06
-	 *     water: 1.0
-	 *     fat: 0.9196
-	 *     gasoline: 0.7
-	 *     wood: 0.67
-	 *     tree bark: 0.24
-	 *     air: 0.001293
-	 */
-	UPROPERTY(EditAnywhere, BlueprintReadWrite, Category = "ChaosPhysics|General")
-		float Density;
-	UPROPERTY(EditAnywhere, BlueprintReadWrite, Category = "ChaosPhysics|General")
-		float MinMass;
-	UPROPERTY(EditAnywhere, BlueprintReadWrite, Category = "ChaosPhysics|General")
-		float MaxMass;
-
-	//
-	// ChaosPhysics | Collisions
-	//
-
-	/** CollisionType defines how to initialize the rigid collision structures.  */
-	UPROPERTY(EditAnywhere, BlueprintReadWrite, Category = "ChaosPhysics|Collisions")
-		ECollisionTypeEnum CollisionType;
-
-	/** Number of particles to generate per unit area (square cm). 0.1 would generate 1 collision particle per 10 cm^2. */
-	UPROPERTY(EditAnywhere, BlueprintReadWrite, Category = "ChaosPhysics|Collisions")
-		float ImplicitShapeParticlesPerUnitArea;
-	/** Minimum number of particles for each implicit shape. */
-	UPROPERTY(EditAnywhere, BlueprintReadWrite, Category = "ChaosPhysics|Collisions")
-		int ImplicitShapeMinNumParticles;
-	/** Maximum number of particles for each implicit shape. */
-	UPROPERTY(EditAnywhere, BlueprintReadWrite, Category = "ChaosPhysics|Collisions")
-		int ImplicitShapeMaxNumParticles;
-
-	/** Resolution on the smallest axes for the level set. (def: 5) */
-	UPROPERTY(EditAnywhere, BlueprintReadOnly, Category = "ChaosPhysics|Collisions")
-		int32 MinLevelSetResolution;
-	/** Resolution on the smallest axes for the level set. (def: 10) */
-	UPROPERTY(EditAnywhere, BlueprintReadOnly, Category = "ChaosPhysics|Collisions")
-		int32 MaxLevelSetResolution;
-
-	/** Collision group - 0 = collides with all, -1 = none */
-	UPROPERTY(EditAnywhere, BlueprintReadWrite, Category = "ChaosPhysics|Collisions")
-		int32 CollisionGroup;
-
-	//
-	// ChaosPhysics | Initial Velocity
-	//
-
-	/** Where to pull initial velocity from - user defined or animation. */
-	UPROPERTY(EditAnywhere, BlueprintReadWrite, Category = "ChaosPhysics|Initial Velocity")
-		EInitialVelocityTypeEnum InitialVelocityType;
-	/** Initial linear velocity. */
-	UPROPERTY(EditAnywhere, BlueprintReadWrite, Category = "ChaosPhysics|Initial Velocity")
-		FVector InitialLinearVelocity;
-	/** Initial angular velocity. */
-	UPROPERTY(EditAnywhere, BlueprintReadWrite, Category = "ChaosPhysics|Initial Velocity")
-		FVector InitialAngularVelocity;
-
-private:
-	FTransform PreviousCompWorldSpaceTM;
-	FTransform CurrentTransform;
-	FTransform PreviousTransform;
+	UPhysicsAsset* OverridePhysicsAsset;
 
-public:
 	/** Override gravity*/
 	UPROPERTY(EditAnywhere, Category = Settings, meta = (PinHiddenByDefault, editcondition = "bOverrideWorldGravity"))
 	FVector OverrideWorldGravity;
@@ -208,12 +106,9 @@ public:
 	UPROPERTY(EditAnywhere, Category = Settings)
 	bool bForceDisableCollisionBetweenConstraintBodies;
 
-
-private:
-	ETeleportType ResetSimulatedTeleportType;
-
-public:
-	UPROPERTY(EditAnywhere, Category = Settings, meta=(InlineEditConditionToggle))
+	// @todo(ccaulfield): Inline toggles don't seem to be working in our branch
+	//UPROPERTY(EditAnywhere, Category = Settings, meta = (InlineEditConditionToggle))
+	UPROPERTY(EditAnywhere, Category = Settings)
 	uint8 bEnableWorldGeometry : 1;
 
 	UPROPERTY(EditAnywhere, Category = Settings, meta = (InlineEditConditionToggle))
@@ -242,13 +137,12 @@ public:
 	uint8 bClampLinearTranslationLimitToRefPose : 1;
 
 private:
-	uint8 bSimulationStarted : 1;
-	uint8 bCheckForBodyTransformInit : 1;
-
-public:
-	void PostSerialize(const FArchive& Ar);
+#if WITH_EDITORONLY_DATA
+	UPROPERTY()
+	bool bComponentSpaceSimulation_DEPRECATED;	//use SimulationSpace
+#endif
 
-private:
+#if INCLUDE_CHAOS
 
 	// FAnimNode_SkeletalControlBase interface
 	virtual void InitializeBoneReferences(const FBoneContainer& RequiredBones) override;
@@ -260,13 +154,25 @@ private:
 
 	void InitializeNewBodyTransformsDuringSimulation(FComponentSpacePoseContext& Output, const FTransform& ComponentTransform, const FTransform& BaseBoneTM);
 
-private:
+	ETeleportType ResetSimulatedTeleportType;
+
+	uint8 bSimulationStarted : 1;
+	uint8 bCheckForBodyTransformInit : 1;
+
+	FTransform PreviousCompWorldSpaceTM;
+	FTransform CurrentTransform;
+	FTransform PreviousTransform;
+
 
 	float AccumulatedDeltaTime;
+	float AnimPhysicsMinDeltaTime;
+	bool bSimulateAnimPhysicsAfterReset;
 
 	/** This should only be used for removing the delegate during termination. Do NOT use this for any per frame work */
 	TWeakObjectPtr<USkeletalMeshComponent> SkelMeshCompWeakPtr;
 
+	ImmediatePhysics_Chaos::FSimulation* PhysicsSimulation;
+
 	struct FOutputBoneData
 	{
 		FOutputBoneData()
@@ -311,29 +217,12 @@ private:
 	};
 
 	TArray<FOutputBoneData> OutputBoneData;
-	//TArray<int32> SkeletonBoneIndexToBodyIndex;
+	TArray<ImmediatePhysics_Chaos::FActorHandle*> Bodies;
+	TArray<int32> SkeletonBoneIndexToBodyIndex;
 	TArray<FBodyAnimData> BodyAnimData;
 
-#if INCLUDE_CHAOS
-	/** Called by the Physics Object to get its set up parameters. */
-	void PhysicsObjectInitCallback(const USkeletalMeshComponent* InSkelMeshComponent, const UAnimInstance* InAnimInstance, FSkeletalMeshPhysicsObjectParams& OutPhysicsParams);
-
-	/** Called by the Physics Object to get the lastest pose from the animation */
-	bool UpdatePhysicsInputs(FComponentSpacePoseContext& PoseContext, const float Dt, FBoneHierarchy& OutBoneHierarchy);
-
-	void UpdateAnimNodeOutputs(const FBoneHierarchy& InBoneHierarchy, const FSkeletalMeshPhysicsObjectOutputs& InPhysicsOutputs, FComponentSpacePoseContext& PoseContext, TArray<FBoneTransform>& OutBoneTransforms);
-
-	Chaos::FPBDRigidsSolver* Solver;
-	FSkeletalMeshPhysicsObject* PhysicsObject;
-	//TUniquePtr<Chaos::TChaosPhysicsMaterial<float>> PhysicsMaterial;	//@todo(mlentine): Don't have one per static mesh
-
-	//ImmediatePhysics::FSimulation* PhysicsSimulation;
-	//TArray<ImmediatePhysics::FActorHandle*> Bodies;
-	//TArray<FPhysicsConstraintHandle*> Constraints;
-	//TArray<USkeletalMeshComponent::FPendingRadialForces> PendingRadialForces;
-	//FPhysScene* PhysScene;
-	//FCollisionQueryParams QueryParams;
-#endif
+	TArray<FPhysicsConstraintHandle*> Constraints;
+	TArray<USkeletalMeshComponent::FPendingRadialForces> PendingRadialForces;
 
 	TSet<UPrimitiveComponent*> ComponentsInSim;
 
@@ -345,6 +234,9 @@ private:
 
 	FSphere CachedBounds;
 
+	FCollisionQueryParams QueryParams;
+
+	FPhysScene* PhysScene;
 
 	// Evaluation counter, to detect when we haven't be evaluated in a while.
 	FGraphTraversalCounter EvalCounter;
@@ -359,8 +251,46 @@ private:
 	FBlendedHeapCurve CapturedFrozenCurves;
 
 	FVector PreviousComponentLinearVelocity;
+#endif
 };
 
+#if !INCLUDE_CHAOS
+// Stub out public API and safe-initialize all public properties if Chaos is not compiled.
+inline FAnimNode_RigidBody_Chaos::FAnimNode_RigidBody_Chaos()
+	: OverridePhysicsAsset(nullptr)
+	, OverrideWorldGravity(FVector::ZeroVector)
+	, ExternalForce(FVector::ZeroVector)
+	, ComponentLinearAccScale(FVector::ZeroVector)
+	, ComponentLinearVelScale(FVector::ZeroVector)
+	, ComponentAppliedLinearAccClamp(FVector::ZeroVector)
+	, CachedBoundsScale(false)
+	, BaseBoneRef()
+	, OverlapChannel(ECollisionChannel::ECC_WorldStatic)
+	, SimulationSpace(ESimulationSpace::ComponentSpace)
+	, bForceDisableCollisionBetweenConstraintBodies(false)
+	, bEnableWorldGeometry(false)
+	, bOverrideWorldGravity(false)
+	, bTransferBoneVelocities(0)
+	, bFreezeIncomingPoseOnStart(0)
+	, bClampLinearTranslationLimitToRefPose(0)
+#if WITH_EDITORONLY_DATA
+	, bComponentSpaceSimulation_DEPRECATED(false)
+#endif
+{}
+inline FAnimNode_RigidBody_Chaos::~FAnimNode_RigidBody_Chaos() {}
+inline void FAnimNode_RigidBody_Chaos::FAnimNode_RigidBody_Chaos::GatherDebugData(FNodeDebugData& DebugData) { Super::GatherDebugData(DebugData); }
+inline void FAnimNode_RigidBody_Chaos::UpdateComponentPose_AnyThread(const FAnimationUpdateContext& Context) { Super::UpdateComponentPose_AnyThread(Context); }
+inline void FAnimNode_RigidBody_Chaos::EvaluateComponentPose_AnyThread(FComponentSpacePoseContext& Output) { Super::EvaluateComponentPose_AnyThread(Output); }
+inline void FAnimNode_RigidBody_Chaos::EvaluateSkeletalControl_AnyThread(FComponentSpacePoseContext& Output, TArray<FBoneTransform>& OutBoneTransforms) { Super::EvaluateSkeletalControl_AnyThread(Output, OutBoneTransforms); }
+inline void FAnimNode_RigidBody_Chaos::OnInitializeAnimInstance(const FAnimInstanceProxy* InProxy, const UAnimInstance* InAnimInstance) { Super::OnInitializeAnimInstance(InProxy, InAnimInstance); }
+inline void FAnimNode_RigidBody_Chaos::PreUpdate(const UAnimInstance* InAnimInstance) { Super::PreUpdate(InAnimInstance); }
+inline void FAnimNode_RigidBody_Chaos::UpdateInternal(const FAnimationUpdateContext& Context) { Super::UpdateInternal(Context); }
+inline bool FAnimNode_RigidBody_Chaos::IsValidToEvaluate(const USkeleton* Skeleton, const FBoneContainer& RequiredBones) { return true; }
+inline bool FAnimNode_RigidBody_Chaos::NeedsDynamicReset() const { return false; }
+inline void FAnimNode_RigidBody_Chaos::ResetDynamics(ETeleportType InTeleportType) {}
+inline int32 FAnimNode_RigidBody_Chaos::GetLODThreshold() const { return 0; }
+#endif
+
 #if WITH_EDITORONLY_DATA
 template<>
 struct TStructOpsTypeTraits<FAnimNode_RigidBody_Chaos> : public TStructOpsTypeTraitsBase2<FAnimNode_RigidBody_Chaos>
diff --git a/Engine/Source/Runtime/Engine/Classes/Components/SkeletalMeshComponent.h b/Engine/Source/Runtime/Engine/Classes/Components/SkeletalMeshComponent.h
--- a/Engine/Source/Runtime/Engine/Classes/Components/SkeletalMeshComponent.h
+++ b/Engine/Source/Runtime/Engine/Classes/Components/SkeletalMeshComponent.h
@@ -64,6 +64,15 @@ enum class EAnimCurveType : uint8
 	MaxAnimCurveType UMETA(Hidden)
 };
 
+UENUM()
+enum class EClothMassMode : uint8
+{
+	UniformMass,
+	TotalMass,
+	Density,
+	MaxClothMassMode UMETA(Hidden)
+};
+
 struct FAnimationEvaluationContext
 {
 	// The anim instance we are evaluating
@@ -614,6 +623,28 @@ public:
 	UPROPERTY(transient)
 	uint16 CachedAnimCurveUidVersion;
 
+	UPROPERTY(EditAnywhere, BlueprintReadWrite, Category = Clothing)
+	EClothMassMode MassMode;
+
+	UPROPERTY(EditAnywhere, BlueprintReadWrite, Category = Clothing)
+	float UniformMass;
+
+	UPROPERTY(EditAnywhere, BlueprintReadWrite, Category = Clothing)
+	float TotalMass;
+	
+	/**
+	 * Water: 1.0
+	 * Cotton: 0.155
+	 * Wool: 0.13
+	 * Silk: 0.133
+	 */
+	UPROPERTY(EditAnywhere, BlueprintReadWrite, Category = Clothing)
+	float Density;
+
+	UPROPERTY(EditAnywhere, BlueprintReadWrite, Category = Clothing)
+	float MinPerParticleMass;
+
+
 	/**
 	 * weight to blend between simulated results and key-framed positions
 	 * if weight is 1.0, shows only cloth simulation results and 0.0 will show only skinned results
diff --git a/Engine/Source/Runtime/Engine/Classes/PhysicsEngine/AggregateGeom.h b/Engine/Source/Runtime/Engine/Classes/PhysicsEngine/AggregateGeom.h
--- a/Engine/Source/Runtime/Engine/Classes/PhysicsEngine/AggregateGeom.h
+++ b/Engine/Source/Runtime/Engine/Classes/PhysicsEngine/AggregateGeom.h
@@ -95,6 +95,8 @@ struct ENGINE_API FKAggregateGeom
 	/** Returns the volume of this element */
 	float GetVolume(const FVector& Scale3D) const;
 
+	FGuid MakeDDCKey() const;
+
 private:
 
 	/** Helper function for safely copying instances */
diff --git a/Engine/Source/Runtime/Engine/Classes/PhysicsEngine/BodyInstance.h b/Engine/Source/Runtime/Engine/Classes/PhysicsEngine/BodyInstance.h
--- a/Engine/Source/Runtime/Engine/Classes/PhysicsEngine/BodyInstance.h
+++ b/Engine/Source/Runtime/Engine/Classes/PhysicsEngine/BodyInstance.h
@@ -1006,7 +1006,7 @@ public:
 
 public:
 	// #PHYS2 Rename, not just for physx now.
-	FPhysxUserData PhysxUserData;
+	FPhysicsUserData PhysicsUserData;
 
 	struct FWeldInfo
 	{
diff --git a/Engine/Source/Runtime/Engine/Classes/PhysicsEngine/BodySetup.h b/Engine/Source/Runtime/Engine/Classes/PhysicsEngine/BodySetup.h
--- a/Engine/Source/Runtime/Engine/Classes/PhysicsEngine/BodySetup.h
+++ b/Engine/Source/Runtime/Engine/Classes/PhysicsEngine/BodySetup.h
@@ -39,6 +39,21 @@ namespace physx
 	class PxTriangleMeshGeometry;
 }
 
+#if WITH_CHAOS
+namespace Chaos
+{
+	template<typename T, int d>
+	class TImplicitObject;
+
+	template <typename T>
+	class TTriangleMeshImplicitObject;
+}
+
+template<typename T, int d>
+class FChaosDerivedDataReader;
+
+#endif
+
 DECLARE_CYCLE_STAT_EXTERN(TEXT("PhysX Cooking"), STAT_PhysXCooking, STATGROUP_Physics, );
 
 
@@ -127,6 +142,10 @@ class UBodySetup : public UObject
 {
 	GENERATED_UCLASS_BODY()
 
+	/** Needs implementation in BodySetup.cpp to compile UniquePtr for forward declared class */
+	UBodySetup(FVTableHelper& Helper);
+	virtual ~UBodySetup();
+
 	/** Simplified collision representation of this  */
 	UPROPERTY(EditAnywhere, Category = BodySetup, meta=(DisplayName = "Primitives", NoResetToDefault))
 	struct FKAggregateGeom AggGeom;
@@ -243,6 +262,11 @@ public:
 	TArray<physx::PxTriangleMesh*> TriMeshes;
 #endif
 
+#if WITH_CHAOS
+	//FBodySetupTriMeshes* TriMeshWrapper;
+	TArray<TUniquePtr<Chaos::TTriangleMeshImplicitObject<float>>> ChaosTriMeshes;
+#endif
+
 	/** Additional UV info, if available. Used for determining UV for a line trace impact. */
 	FBodySetupUVInfo UVInfo;
 
@@ -298,10 +322,21 @@ public:
 	ENGINE_API void AbortPhysicsMeshAsyncCreation();
 
 private:
-#if WITH_PHYSX
+#if WITH_PHYSX && PHYSICS_INTERFACE_PHYSX
+	bool ProcessFormatData_PhysX(FByteBulkData* FormatData);
+	bool RuntimeCookPhysics_PhysX();
+
+	// #TODO MRMesh for some reason needs to be able to call this - that case needs fixed to correctly use the create meshes flow
+	friend class UMRMeshComponent;
+	/** Finish creating the physics meshes and update the body setup data with cooked data */
+	ENGINE_API void FinishCreatingPhysicsMeshes_PhysX(const TArray<physx::PxConvexMesh*>& ConvexMeshes, const TArray<physx::PxConvexMesh*>& ConvexMeshesNegX, const TArray<physx::PxTriangleMesh*>& TriMeshes);
 	/** Finalize game thread data before calling back user's delegate */
 	void FinishCreatePhysicsMeshesAsync(FPhysXCookHelper* AsyncPhysicsCookHelper, FOnAsyncPhysicsCookFinished OnAsyncPhysicsCookFinished);
-#endif // WITH_PHYSX
+#elif WITH_CHAOS
+	bool ProcessFormatData_Chaos(FByteBulkData* FormatData);
+	bool RuntimeCookPhysics_Chaos();
+	void FinishCreatingPhysicsMeshes_Chaos(FChaosDerivedDataReader<float, 3>& InReader);
+#endif
 
 	/**
 	* Given a format name returns its cooked data.
@@ -314,10 +349,12 @@ private:
 
 public:
 
-#if WITH_PHYSX
-	/** Finish creating the physics meshes and update the body setup data with cooked data */
-	ENGINE_API void FinishCreatingPhysicsMeshes(const TArray<physx::PxConvexMesh*>& ConvexMeshes, const TArray<physx::PxConvexMesh*>& ConvexMeshesNegX, const TArray<physx::PxTriangleMesh*>& TriMeshes);
-#endif // WITH_PHYSX
+	/**
+	 * Generate a string to uniquely describe the state of the geometry in this setup to populate the DDC
+	 *
+	 * @param OutString The generated string will be place in this FString
+	 */
+	void GetGeometryDDCKey(FString& OutString) const;
 
 	/** Returns the volume of this element */
 	ENGINE_API virtual float GetVolume(const FVector& Scale) const;
diff --git a/Engine/Source/Runtime/Engine/Classes/PhysicsEngine/ConvexElem.h b/Engine/Source/Runtime/Engine/Classes/PhysicsEngine/ConvexElem.h
--- a/Engine/Source/Runtime/Engine/Classes/PhysicsEngine/ConvexElem.h
+++ b/Engine/Source/Runtime/Engine/Classes/PhysicsEngine/ConvexElem.h
@@ -5,6 +5,11 @@
 #include "CoreMinimal.h"
 #include "UObject/ObjectMacros.h"
 #include "PhysicsEngine/ShapeElem.h"
+
+#if WITH_CHAOS
+#include "Chaos/Serializable.h"
+#endif
+
 #include "ConvexElem.generated.h"
 
 struct FDynamicMeshVertex;
@@ -15,6 +20,12 @@ namespace physx
 	class PxConvexMesh;
 }
 
+namespace Chaos
+{
+	template <typename T, int d>
+	class TImplicitObject;
+}
+
 /** One convex hull, used for simplified collision. */
 USTRUCT()
 struct FKConvexElem : public FKShapeElem
@@ -40,33 +51,17 @@ private:
 	/** Convex mesh for this body, flipped across X, created from cooked data in CreatePhysicsMeshes */
 	physx::PxConvexMesh*   ConvexMeshNegX;
 
-public:
+#if WITH_CHAOS
+	TUniquePtr<Chaos::TImplicitObject<float, 3>> ChaosConvex;
+#endif
 
+public:
 
-	FKConvexElem()
-		: FKShapeElem(EAggCollisionShape::Convex)
-		, ElemBox(ForceInit)
-		, Transform(FTransform::Identity)
-		, ConvexMesh(NULL)
-		, ConvexMeshNegX(NULL)
-	{}
+	ENGINE_API FKConvexElem();
+	ENGINE_API FKConvexElem(const FKConvexElem& Other);
+	ENGINE_API ~FKConvexElem();
 
-	FKConvexElem(const FKConvexElem& Other)
-		: ConvexMesh(nullptr)
-		, ConvexMeshNegX(nullptr)
-	{
-		CloneElem(Other);
-	}
-
-	const FKConvexElem& operator=(const FKConvexElem& Other)
-	{
-		ensureMsgf(ConvexMesh == nullptr, TEXT("We are leaking memory. Why are we calling the assignment operator on an element that has already allocated resources?"));
-		ensureMsgf(ConvexMeshNegX == nullptr, TEXT("We are leaking memory. Why are we calling the assignment operator on an element that has already allocated resources?"));
-		ConvexMesh = nullptr;
-		ConvexMeshNegX = nullptr;
-		CloneElem(Other);
-		return *this;
-	}
+	ENGINE_API const FKConvexElem& operator=(const FKConvexElem& Other);
 
 	ENGINE_API void	DrawElemWire(class FPrimitiveDrawInterface* PDI, const FTransform& ElemTM, const float Scale, const FColor Color) const;
 
@@ -108,6 +103,14 @@ public:
 	/** Set the PhysX convex mesh to use for this element */
 	ENGINE_API void SetMirroredConvexMesh(physx::PxConvexMesh* InMesh);
 
+#if WITH_CHAOS
+	ENGINE_API const TUniquePtr<Chaos::TImplicitObject<float, 3>>& GetChaosConvexMesh() const;
+
+	ENGINE_API void SetChaosConvexMesh(TUniquePtr<Chaos::TImplicitObject<float, 3>>&& InChaosConvex);
+
+	ENGINE_API void ResetChaosConvexMesh();
+#endif
+
 	/** Get current transform applied to convex mesh vertices */
 	FTransform GetTransform() const
 	{
@@ -132,11 +135,5 @@ public:
 
 private:
 	/** Helper function to safely copy instances of this shape*/
-	void CloneElem(const FKConvexElem& Other)
-	{
-		Super::CloneElem(Other);
-		VertexData = Other.VertexData;
-		ElemBox = Other.ElemBox;
-		Transform = Other.Transform;
-	}
+	void CloneElem(const FKConvexElem& Other);
 };
diff --git a/Engine/Source/Runtime/Engine/Classes/PhysicsEngine/PhysicsSettings.h b/Engine/Source/Runtime/Engine/Classes/PhysicsEngine/PhysicsSettings.h
--- a/Engine/Source/Runtime/Engine/Classes/PhysicsEngine/PhysicsSettings.h
+++ b/Engine/Source/Runtime/Engine/Classes/PhysicsEngine/PhysicsSettings.h
@@ -16,8 +16,6 @@
 #include "GameFramework/WorldSettings.h"
 #include "PhysicsCoreTypes.h"
 
-#include "Framework/Threading.h"
-
 #include "PhysicsSettings.generated.h"
 
 /**
diff --git a/Engine/Source/Runtime/Engine/Engine.Build.cs b/Engine/Source/Runtime/Engine/Engine.Build.cs
--- a/Engine/Source/Runtime/Engine/Engine.Build.cs
+++ b/Engine/Source/Runtime/Engine/Engine.Build.cs
@@ -85,8 +85,6 @@ public class Engine : ModuleRules
 				"PakFile",
                 "NetworkReplayStreaming",
 				"PhysicsCore",
-				"PhysicsSQ",
-				"ChaosSolvers"
 			}
 		);
 
@@ -376,15 +374,25 @@ public class Engine : ModuleRules
 			DynamicallyLoadedModuleNames.Add("PhysXCooking");
 		}
 
-			// Engine public headers need to know about some types (enums etc.)
-			PublicIncludePathModuleNames.Add("ClothingSystemRuntimeInterface");
-			PublicDependencyModuleNames.Add("ClothingSystemRuntimeInterface");
+		if (Target.bCompileChaos || Target.bUseChaos)
+        {
+            PublicDependencyModuleNames.AddRange(
+				new string[] {
+					"PhysicsSQ",
+					"ChaosSolvers"
+				}
+			);
+        }
 
-			if (Target.bBuildEditor)
-			{
-				PrivateDependencyModuleNames.Add("ClothingSystemEditorInterface");
-				PrivateIncludePathModuleNames.Add("ClothingSystemEditorInterface");
-			}
+        // Engine public headers need to know about some types (enums etc.)
+        PublicIncludePathModuleNames.Add("ClothingSystemRuntimeInterface");
+		PublicDependencyModuleNames.Add("ClothingSystemRuntimeInterface");
+
+		if (Target.bBuildEditor)
+		{
+			PrivateDependencyModuleNames.Add("ClothingSystemEditorInterface");
+			PrivateIncludePathModuleNames.Add("ClothingSystemEditorInterface");
+		}
 
 		if ((Target.Platform == UnrealTargetPlatform.Win64) ||
 			(Target.Platform == UnrealTargetPlatform.Win32))
diff --git a/Engine/Source/Runtime/Engine/Private/Animation/AnimationRuntime.cpp b/Engine/Source/Runtime/Engine/Private/Animation/AnimationRuntime.cpp
--- a/Engine/Source/Runtime/Engine/Private/Animation/AnimationRuntime.cpp
+++ b/Engine/Source/Runtime/Engine/Private/Animation/AnimationRuntime.cpp
@@ -1616,6 +1616,30 @@ FTransform FAnimationRuntime::GetComponentSpaceTransformRefPose(const FReference
 	return GetComponentSpaceTransform(RefSkeleton, RefSkeleton.GetRefBonePose(), BoneIndex);
 }
 
+const FTransform& FAnimationRuntime::GetComponentSpaceTransformWithCache(const FReferenceSkeleton& RefSkeleton, const TArray<FTransform> &BoneSpaceTransforms, int32 BoneIndex, TArray<FTransform>& CachedTransforms, TArray<bool>& CachedTransformReady)
+{
+	if (!CachedTransformReady[BoneIndex])
+	{
+		CachedTransformReady[BoneIndex] = true;
+		CachedTransforms[BoneIndex] = FTransform::Identity;
+
+		if (RefSkeleton.IsValidIndex(BoneIndex))
+		{
+			int32 ParentBoneIndex = RefSkeleton.GetParentIndex(BoneIndex);
+			if (ParentBoneIndex != INDEX_NONE)
+			{
+				const FTransform& ParentComponentSpaceTransform = GetComponentSpaceTransformWithCache(RefSkeleton, BoneSpaceTransforms, ParentBoneIndex, CachedTransforms, CachedTransformReady);
+				CachedTransforms[BoneIndex] = BoneSpaceTransforms[BoneIndex] * ParentComponentSpaceTransform;
+			}
+			else
+			{
+				CachedTransforms[BoneIndex] = BoneSpaceTransforms[BoneIndex];
+			}
+		}
+	}
+	return CachedTransforms[BoneIndex];
+}
+
 void FAnimationRuntime::FillUpComponentSpaceTransforms(const FReferenceSkeleton& RefSkeleton, const TArray<FTransform> &BoneSpaceTransforms, TArray<FTransform> &ComponentSpaceTransforms)
 {
 	ComponentSpaceTransforms.Empty(BoneSpaceTransforms.Num());
diff --git a/Engine/Source/Runtime/Engine/Private/Collision/CollisionConversions.cpp b/Engine/Source/Runtime/Engine/Private/Collision/CollisionConversions.cpp
--- a/Engine/Source/Runtime/Engine/Private/Collision/CollisionConversions.cpp
+++ b/Engine/Source/Runtime/Engine/Private/Collision/CollisionConversions.cpp
@@ -14,9 +14,6 @@
 #if PHYSICS_INTERFACE_PHYSX
 #include "PhysXInterfaceWrapper.h"
 #include "Collision/CollisionConversionsPhysx.h"
-#elif PHYSICS_INTERFACE_LLIMMEDIATE
-#include "Physics/Experimental/LLImmediateInterfaceWrapper.h"
-#include "Collision/Experimental/CollisionConversionsLLImmediate.h"
 #elif WITH_CHAOS
 #include "Physics/Experimental/ChaosInterfaceWrapper.h"
 #include "Physics/Experimental/PhysInterface_Chaos.h"
@@ -102,7 +99,7 @@ static FVector FindGeomOpposingNormal(ECollisionShapeType QueryGeomType, const F
 	// TODO: can we support other shapes here as well?
 	if (QueryGeomType == ECollisionShapeType::Capsule || QueryGeomType == ECollisionShapeType::Sphere)
 	{
-		FPhysicsShape* Shape = GetShape(Hit);
+		const FPhysicsShape* Shape = GetShape(Hit);
 		if (Shape)
 		{
 			ECollisionShapeType GeomType = GetGeometryType(*Shape);
@@ -114,7 +111,7 @@ static FVector FindGeomOpposingNormal(ECollisionShapeType QueryGeomType, const F
 			case ECollisionShapeType::Convex:		return FindConvexMeshOpposingNormal(Hit, TraceDirectionDenorm, InNormal);
 			case ECollisionShapeType::Heightfield:		return FindHeightFieldOpposingNormal(Hit, TraceDirectionDenorm, InNormal);
 			case ECollisionShapeType::Trimesh:	return FindTriMeshOpposingNormal(Hit, TraceDirectionDenorm, InNormal);
-			default: check(false);	//unsupported geom type
+			default: break; //TODO_SQ_IMPLEMENTATION check(false);	//unsupported geom type
 			}
 		}
 	}
@@ -130,10 +127,7 @@ static void SetHitResultFromShapeAndFaceIndex(const FPhysicsShape& Shape,  const
 	UPrimitiveComponent* OwningComponent = nullptr;
 	if(const FBodyInstance* BodyInst = GetUserData(Actor))
 	{
-#if WITH_CHAOS || WITH_IMMEDIATE_PHYSX || PHYSICS_INTERFACE_LLIMMEDIATE
-        ensure(false);
-#else
-		BodyInst = FPhysicsInterface_PhysX::ShapeToOriginalBodyInstance(BodyInst, &Shape);
+		BodyInst = FPhysicsInterface::ShapeToOriginalBodyInstance(BodyInst, &Shape);
 
 		//Normal case where we hit a body
 		OutResult.Item = BodyInst->InstanceBodyIndex;
@@ -144,7 +138,6 @@ static void SetHitResultFromShapeAndFaceIndex(const FPhysicsShape& Shape,  const
 		}
 
 		OwningComponent = BodyInst->OwnerComponent.Get();
-#endif
 	}
 #if PHYSICS_INTERFACE_PHYSX
 	else if(const FCustomPhysXPayload* CustomPayload = GetUserData<FCustomPhysXPayload>(Shape))	//todo(ocohen): wrap with PHYSX
@@ -191,8 +184,6 @@ static void SetHitResultFromShapeAndFaceIndex(const FPhysicsShape& Shape,  const
 	OutResult.FaceIndex = INDEX_NONE;
 }
 
-#if WITH_PHYSX
-
 EConvertQueryResult ConvertQueryImpactHit(const UWorld* World, const FHitLocation& Hit, FHitResult& OutResult, float CheckLength, const FCollisionFilterData& QueryFilter, const FVector& StartLoc, const FVector& EndLoc, const FPhysicsGeometry* Geom, const FTransform& QueryTM, bool bReturnFaceIndex, bool bReturnPhysMat)
 {
 	SCOPE_CYCLE_COUNTER(STAT_ConvertQueryImpactHit);
@@ -348,7 +339,13 @@ EConvertQueryResult ConvertTraceResults(bool& OutHasValidBlockingHit, const UWor
 		{
 			if (TIsSame<HitType, FHitSweep>::Value)
 			{
+#if WITH_PHYSX
 				SetInternalFaceIndex(Hit, FindFaceIndex(Hit, Dir));
+#else
+				//#TODO Chaos implementation
+				ensure(false);
+				SetInternalFaceIndex(Hit, 0);
+#endif
 			}
 
 			FHitResult& NewResult = OutHits[OutHits.AddDefaulted()];
@@ -390,8 +387,6 @@ template EConvertQueryResult ConvertTraceResults<FHitSweep>(bool& OutHasValidBlo
 template EConvertQueryResult ConvertTraceResults<FHitRaycast>(bool& OutHasValidBlockingHit, const UWorld* World, int32 NumHits, FHitRaycast* Hits, float CheckLength, const FCollisionFilterData& QueryFilter, TArray<FHitResult>& OutHits, const FVector& StartLoc, const FVector& EndLoc, const FPhysicsGeometry& Geom, const FTransform& QueryTM, float MaxDistance, bool bReturnFaceIndex, bool bReturnPhysMat);
 template EConvertQueryResult ConvertTraceResults<FHitRaycast>(bool& OutHasValidBlockingHit, const UWorld* World, int32 NumHits, FHitRaycast* Hits, float CheckLength, const FCollisionFilterData& QueryFilter, FHitResult& OutHit, const FVector& StartLoc, const FVector& EndLoc, const FPhysicsGeometry& Geom, const FTransform& QueryTM, float MaxDistance, bool bReturnFaceIndex, bool bReturnPhysMat);
 
-#endif
-
 /** Util to convert an overlapped shape into a sweep hit result, returns whether it was a blocking hit. */
 static bool ConvertOverlappedShapeToImpactHit(const UWorld* World, const FHitLocation& Hit, const FVector& StartLoc, const FVector& EndLoc, FHitResult& OutResult, const FPhysicsGeometry& Geom, const FTransform& QueryTM, const FCollisionFilterData& QueryFilter, bool bReturnPhysMat)
 {
@@ -508,7 +503,7 @@ void ConvertQueryOverlap(const FPhysicsShape& Shape, const FPhysicsActor& Actor,
 	// Try body instance
 	if (const FBodyInstance* BodyInst = GetUserData(Actor))
 	{
-#if WITH_CHAOS || WITH_IMMEDIATE_PHYSX || PHYSICS_INTERFACE_LLIMMEDIATE
+#if WITH_CHAOS || WITH_IMMEDIATE_PHYSX
         ensure(false);
 #else
         BodyInst = FPhysicsInterface_PhysX::ShapeToOriginalBodyInstance(BodyInst, &Shape);
diff --git a/Engine/Source/Runtime/Engine/Private/Collision/CollisionConversions.h b/Engine/Source/Runtime/Engine/Private/Collision/CollisionConversions.h
--- a/Engine/Source/Runtime/Engine/Private/Collision/CollisionConversions.h
+++ b/Engine/Source/Runtime/Engine/Private/Collision/CollisionConversions.h
@@ -10,9 +10,6 @@
 #include "PhysicsPublic.h"
 #include "PhysXIncludes.h"
 
-#if WITH_PHYSX
-
-
 enum class EConvertQueryResult
 {
 	Valid,
@@ -33,7 +30,7 @@ enum class EConvertQueryResult
  * @param	bReturnPhysMat		True if we want to lookup the physical material
  * @return	Whether result passed NaN/Inf checks.
  */
-EConvertQueryResult ConvertQueryImpactHit(const UWorld* World, const PxLocationHit& PHit, FHitResult& OutResult, float CheckLength, const FCollisionFilterData& QueryFilter, const FVector& StartLoc, const FVector& EndLoc, const PxGeometry* Geom, const FTransform& QueryTM, bool bReturnFaceIndex, bool bReturnPhysMat);
+EConvertQueryResult ConvertQueryImpactHit(const UWorld* World, const FHitLocation& PHit, FHitResult& OutResult, float CheckLength, const FCollisionFilterData& QueryFilter, const FVector& StartLoc, const FVector& EndLoc, const FPhysicsGeometry* Geom, const FTransform& QueryTM, bool bReturnFaceIndex, bool bReturnPhysMat);
 
 /** 
  * Util to convert physX sweep results to unreal hit results and add to array
@@ -51,10 +48,10 @@ EConvertQueryResult ConvertQueryImpactHit(const UWorld* World, const PxLocationH
  * @return	Whether all results passed NaN/Inf checks.
  */
 template <typename Hit>
-EConvertQueryResult ConvertTraceResults(bool& OutHasValidBlockingHit, const UWorld* World, int32 NumHits, Hit* Hits, float CheckLength, const FCollisionFilterData& QueryFilter, TArray<FHitResult>& OutHits, const FVector& StartLoc, const FVector& EndLoc, const PxGeometry& Geom, const FTransform& QueryTM, float MaxDistance, bool bReturnFaceIndex, bool bReturnPhysMat);
+EConvertQueryResult ConvertTraceResults(bool& OutHasValidBlockingHit, const UWorld* World, int32 NumHits, Hit* Hits, float CheckLength, const FCollisionFilterData& QueryFilter, TArray<FHitResult>& OutHits, const FVector& StartLoc, const FVector& EndLoc, const FPhysicsGeometry& Geom, const FTransform& QueryTM, float MaxDistance, bool bReturnFaceIndex, bool bReturnPhysMat);
 
 template <typename Hit>
-EConvertQueryResult ConvertTraceResults(bool& OutHasValidBlockingHit, const UWorld* World, int32 NumHits, Hit* Hits, float CheckLength, const FCollisionFilterData& QueryFilter, FHitResult& OutHits, const FVector& StartLoc, const FVector& EndLoc, const PxGeometry& Geom, const FTransform& QueryTM, float MaxDistance, bool bReturnFaceIndex, bool bReturnPhysMat);
+EConvertQueryResult ConvertTraceResults(bool& OutHasValidBlockingHit, const UWorld* World, int32 NumHits, Hit* Hits, float CheckLength, const FCollisionFilterData& QueryFilter, FHitResult& OutHits, const FVector& StartLoc, const FVector& EndLoc, const FPhysicsGeometry& Geom, const FTransform& QueryTM, float MaxDistance, bool bReturnFaceIndex, bool bReturnPhysMat);
 
 /** 
  * Util to convert physX overlap query to our overlap result
@@ -64,7 +61,7 @@ EConvertQueryResult ConvertTraceResults(bool& OutHasValidBlockingHit, const UWor
  * @param	OutOverlap	(out) Result converted
  * @param	QueryFilter	Query Filter 
  */
-void ConvertQueryOverlap(const UWorld* World, const PxShape& PShape, const PxRigidActor& PActor, FOverlapResult& OutOverlap, const FCollisionFilterData& QueryFilter);
+void ConvertQueryOverlap(const UWorld* World, const FPhysicsShape& PShape, const FPhysicsActorHandle& PActor, FOverlapResult& OutOverlap, const FCollisionFilterData& QueryFilter);
 
 /**
  * Util to determine if a shape is deemed blocking based on the query filter
@@ -87,9 +84,6 @@ bool IsBlocking(const FPhysicsShape& PShape, const FCollisionFilterData& QueryFi
 bool ConvertOverlapResults(int32 NumOverlaps, FHitOverlap* POverlapResults, const FCollisionFilterData& QueryFilter, TArray<FOverlapResult>& OutOverlaps);
 
 
-#endif // WITH_PHYX
-
-
 struct FCompareFHitResultTime
 {
 	FORCEINLINE bool operator()(const FHitResult& A, const FHitResult& B) const
diff --git a/Engine/Source/Runtime/Engine/Private/Collision/CollisionDebugDrawing.cpp b/Engine/Source/Runtime/Engine/Private/Collision/CollisionDebugDrawing.cpp
--- a/Engine/Source/Runtime/Engine/Private/Collision/CollisionDebugDrawing.cpp
+++ b/Engine/Source/Runtime/Engine/Private/Collision/CollisionDebugDrawing.cpp
@@ -6,6 +6,12 @@
 #include "DrawDebugHelpers.h"
 #include "PhysXPublic.h"
 
+#if INCLUDE_CHAOS
+#include "Chaos/Sphere.h"
+#include "Chaos/Box.h"
+#include "Chaos/Capsule.h"
+#endif
+
 
 #if ENABLE_DRAW_DEBUG
 
@@ -246,4 +252,45 @@ void DrawGeomSweeps(const UWorld* InWorld, const FVector& Start, const FVector&
 
 #endif // WITH_PHYSX
 
+#if INCLUDE_CHAOS
+void DrawGeomOverlaps(const UWorld* InWorld, const Chaos::TImplicitObject<float, 3>& Geom, const FTransform& GeomPose, TArray<struct FOverlapResult>& Overlaps, float Lifetime)
+{
+	using namespace Chaos;
+	FVector Pos = GeomPose.GetLocation();
+	FQuat Rot = GeomPose.GetRotation();
+
+	if (const auto Box = Geom.template GetObject<TBox<float,3>>())
+	{
+		DrawBoxOverlap(InWorld, Pos, Box->Extents() * 0.5f, Rot, Overlaps, Lifetime);
+	}
+	else if (const auto Sphere = Geom.template GetObject<TSphere<float, 3>>())
+	{
+		DrawSphereOverlap(InWorld, Pos, Sphere->GetRadius(), Overlaps, Lifetime);
+	}
+	else if (const auto Capsule = Geom.template GetObject<TCapsule<float>>())
+	{
+		// Convert here from Chaos to unreal definition of capsule height
+		DrawCapsuleOverlap(InWorld, Pos, Capsule->GetHeight() * 0.5f + Capsule->GetRadius(), Capsule->GetRadius(), Rot, Overlaps, Lifetime);
+	}
+}
+
+void DrawGeomSweeps(const UWorld* InWorld, const FVector& Start, const FVector& End, const Chaos::TImplicitObject<float, 3>& Geom, const FQuat& Rotation, const TArray<FHitResult>& Hits, float Lifetime)
+{
+	using namespace Chaos;
+	if (const auto Box = Geom.template GetObject<TBox<float, 3>>())
+	{
+		DrawBoxSweeps(InWorld, Start, End, Box->Extents() * 0.5f, Rotation, Hits, Lifetime);
+	}
+	else if (const auto Sphere = Geom.template GetObject<TSphere<float, 3>>())
+	{
+		DrawSphereSweeps(InWorld, Start, End, Sphere->GetRadius(), Hits, Lifetime);
+	}
+	else if (const auto Capsule = Geom.template GetObject<TCapsule<float>>())
+	{
+		// Convert here from Chaos to unreal definition of capsule height
+		DrawCapsuleSweeps(InWorld, Start, End, Capsule->GetHeight() * 0.5f + Capsule->GetRadius(), Capsule->GetRadius(), Rotation, Hits, Lifetime);
+	}
+}
+#endif
+
 #endif // ENABLE_DRAW_DEBUG
diff --git a/Engine/Source/Runtime/Engine/Private/Collision/CollisionDebugDrawing.h b/Engine/Source/Runtime/Engine/Private/Collision/CollisionDebugDrawing.h
--- a/Engine/Source/Runtime/Engine/Private/Collision/CollisionDebugDrawing.h
+++ b/Engine/Source/Runtime/Engine/Private/Collision/CollisionDebugDrawing.h
@@ -16,3 +16,8 @@ void DrawGeomOverlaps(const UWorld* InWorld, const PxGeometry& PGeom, const PxTr
 /** Draw PhysX geom being swept with hits */
 void DrawGeomSweeps(const UWorld* InWorld, const FVector& Start, const FVector& End, const PxGeometry& PGeom, const PxQuat& Rotation, const TArray<FHitResult>& Hits, float Lifetime);
 #endif
+
+#if INCLUDE_CHAOS
+void DrawGeomOverlaps(const UWorld* InWorld, const Chaos::TImplicitObject<float, 3>& Geom, const FTransform& GeomPose, TArray<struct FOverlapResult>& Overlaps, float Lifetime);
+void DrawGeomSweeps(const UWorld* InWorld, const FVector& Start, const FVector& End, const Chaos::TImplicitObject<float,3>& Geom, const FQuat& Rotation, const TArray<FHitResult>& Hits, float Lifetime);
+#endif
diff --git a/Engine/Source/Runtime/Engine/Private/Collision/KAggregateGeom.cpp b/Engine/Source/Runtime/Engine/Private/Collision/KAggregateGeom.cpp
--- a/Engine/Source/Runtime/Engine/Private/Collision/KAggregateGeom.cpp
+++ b/Engine/Source/Runtime/Engine/Private/Collision/KAggregateGeom.cpp
@@ -15,6 +15,11 @@
 #include "Engine/Polys.h"
 #include "PhysXIncludes.h"
 
+#if WITH_CHAOS
+#include "Chaos/ImplicitObject.h"
+#endif
+
+
 #define MIN_HULL_VERT_DISTANCE		(0.1f)
 #define MIN_HULL_VALID_DIMENSION	(0.5f)
 
@@ -697,6 +702,9 @@ FArchive& operator<<(FArchive& Ar,FKConvexElem& Elem)
 		FMemory::Memzero(&Elem.ElemBox, sizeof(Elem.ElemBox));
 		Elem.ConvexMesh = NULL;
 		Elem.ConvexMeshNegX = NULL;
+#if WITH_CHAOS
+		Elem.ChaosConvex.Reset();
+#endif
 	}
 	return Ar;
 }
diff --git a/Engine/Source/Runtime/Engine/Private/Collision/SceneQuery.cpp b/Engine/Source/Runtime/Engine/Private/Collision/SceneQuery.cpp
--- a/Engine/Source/Runtime/Engine/Private/Collision/SceneQuery.cpp
+++ b/Engine/Source/Runtime/Engine/Private/Collision/SceneQuery.cpp
@@ -1,9 +1,5 @@
 // Copyright 1998-2019 Epic Games, Inc. All Rights Reserved.
 
-#if WITH_CHAOS || PHYSICS_INTERFACE_LLIMMEDIATE
-#define STARTQUERYTIMER()
-#endif
-
 #include "Engine/World.h"
 #include "Collision.h"
 #include "PhysicsEngine/PhysicsSettings.h"
@@ -16,24 +12,22 @@
 #include "Physics/PhysicsInterfaceUtils.h"
 #include "Collision/CollisionConversions.h"
 #include "PhysicsEngine/ScopedSQHitchRepeater.h"
+#include "PhysicsInterfaceDeclaresCore.h"
 
 #if PHYSICS_INTERFACE_PHYSX
 #include "PhysXInterfaceWrapper.h"
-#elif PHYSICS_INTERFACE_LLIMMEDIATE
-#include "Physics/Experimental/LLImmediateInterfaceWrapper.h"
 #endif
 
 #include "Collision/CollisionDebugDrawing.h"
 
 float DebugLineLifetime = 2.f;
 
-#if !WITH_CHAOS
-
 #include "PhysicsEngine/PhysXSupport.h"
 #include "PhysicsEngine/CollisionAnalyzerCapture.h"
 
 #if INCLUDE_CHAOS
 #include "ChaosSolversModule.h"
+#include "Physics/Experimental/ChaosInterfaceWrapper.h"
 #endif
 
 enum class ESingleMultiOrTest : uint8
@@ -52,43 +46,28 @@ enum class ESweepOrRay : uint8
 struct FGeomSQAdditionalInputs
 {
 	FGeomSQAdditionalInputs(const FCollisionShape& InCollisionShape, const FQuat& InGeomRot)
-#if WITH_PHYSX
-		: ShapeAdaptor(InGeomRot, InCollisionShape)
+		: ShapeAdapter(InGeomRot, InCollisionShape)
 		, CollisionShape(InCollisionShape)
 	{
 	}
 
-	const PxGeometry* GetGeometry() const
-	{
-		return &ShapeAdaptor.GetGeometry();
-	}
-
-	const FQuat* GetGeometryOrientation() const
-	{
-		return &ShapeAdaptor.GetGeomOrientation();
-	}
-
-	FPhysXShapeAdaptor ShapeAdaptor;
-#else
-	{}
-
 	const FPhysicsGeometry* GetGeometry() const
 	{
-		return nullptr;
+		return &ShapeAdapter.GetGeometry();
 	}
 
 	const FQuat* GetGeometryOrientation() const
 	{
-		return nullptr;
+		return &ShapeAdapter.GetGeomOrientation();
 	}
 
-#endif
 
 	const FCollisionShape* GetCollisionShape() const
 	{
 		return &CollisionShape;
 	}
 
+	FPhysicsShapeAdapter ShapeAdapter;
 	const FCollisionShape& CollisionShape;
 };
 
@@ -148,7 +127,7 @@ struct TSQTraits
 	static const ESweepOrRay GeometryQuery = InGeometryQuery;
 	using THitType = InHitType;
 	using TOutHits = typename TChooseClass<InSingleMultiOrTest == ESingleMultiOrTest::Multi, TArray<FHitResult>, FHitResult>::Result;
-	using THitBuffer = typename TChooseClass<InSingleMultiOrTest == ESingleMultiOrTest::Multi, FDynamicHitBuffer<InHitType>, typename TChooseClass<InGeometryQuery == ESweepOrRay::Sweep, FPhysicsSweepBuffer, FPhysicsRaycastBuffer>::Result >::Result;
+	using THitBuffer = typename TChooseClass<InSingleMultiOrTest == ESingleMultiOrTest::Multi, FDynamicHitBuffer<InHitType>, typename TChooseClass<InGeometryQuery == ESweepOrRay::Sweep, FSingleHitBuffer<FHitSweep>, FSingleHitBuffer<FHitRaycast>>::Result >::Result;
 
 	// GetNumHits - multi
 	template <ESingleMultiOrTest T = SingleMultiOrTest>
@@ -214,7 +193,11 @@ struct TSQTraits
 		}
 		else
 		{
+#if PHYSICS_INTERFACE_PHYSX
 			DrawGeomSweeps(World, Start, End, *PGeom, U2PQuat(*PGeomRot), Hits, DebugLineLifetime);
+#else
+			DrawGeomSweeps(World, Start, End, *PGeom, *PGeomRot, Hits, DebugLineLifetime);
+#endif
 		}
 	}
 
@@ -524,7 +507,7 @@ bool FGenericPhysicsInterface::GeomSweepMulti(const UWorld* World, const FCollis
 	using TCastTraits = TSQTraits<FHitSweep, ESweepOrRay::Sweep, ESingleMultiOrTest::Multi>;
 	return TSceneCastCommon<TCastTraits>(World, OutHits, FGeomSQAdditionalInputs(InGeom, InGeomRot), Start, End, TraceChannel, Params, ResponseParams, ObjectParams);
 }
-#endif
+
 
 //////////////////////////////////////////////////////////////////////////
 // GEOM OVERLAP
@@ -631,7 +614,7 @@ bool FGenericPhysicsInterface::GeomOverlapBlockingTest(const UWorld* World, cons
 	TArray<FOverlapResult> Overlaps;	//needed only for template shared code
 	FTransform GeomTransform(Rot, Pos);
 #if WITH_PHYSX && !WITH_CHAOS_NEEDS_TO_BE_FIXED
-	FPhysXShapeAdaptor Adaptor(GeomTransform.GetRotation(), CollisionShape);
+	FPhysXShapeAdapter Adaptor(GeomTransform.GetRotation(), CollisionShape);
 	return GeomOverlapMultiImp<EQueryInfo::IsBlocking>(World, Adaptor.GetGeometry(), CollisionShape, Adaptor.GetGeomPose(GeomTransform.GetTranslation()), Overlaps, TraceChannel, Params, ResponseParams, ObjectParams);
 #else
 	return false;
@@ -646,7 +629,7 @@ bool FGenericPhysicsInterface::GeomOverlapAnyTest(const UWorld* World, const str
 	TArray<FOverlapResult> Overlaps;	//needed only for template shared code
 	FTransform GeomTransform(Rot, Pos);
 #if WITH_PHYSX && !WITH_CHAOS_NEEDS_TO_BE_FIXED
-	FPhysXShapeAdaptor Adaptor(GeomTransform.GetRotation(), CollisionShape);
+	FPhysXShapeAdapter Adaptor(GeomTransform.GetRotation(), CollisionShape);
 	return GeomOverlapMultiImp<EQueryInfo::GatherAll>(World, Adaptor.GetGeometry(), CollisionShape, Adaptor.GetGeomPose(GeomTransform.GetTranslation()), Overlaps, TraceChannel, Params, ResponseParams, ObjectParams);
 #else
 	return false;
@@ -675,7 +658,7 @@ bool FGenericPhysicsInterface::GeomOverlapMulti(const UWorld* World, const FColl
 
 	FTransform GeomTransform(InRotation, InPosition);
 #if WITH_PHYSX && !WITH_CHAOS_NEEDS_TO_BE_FIXED
-	FPhysXShapeAdaptor Adaptor(GeomTransform.GetRotation(), InGeom);
+	FPhysXShapeAdapter Adaptor(GeomTransform.GetRotation(), InGeom);
 	return GeomOverlapMultiImp<EQueryInfo::GatherAll>(World, Adaptor.GetGeometry(), InGeom, Adaptor.GetGeomPose(GeomTransform.GetTranslation()), OutOverlaps, TraceChannel, Params, ResponseParams, ObjectParams);
 #else
 	return false;
diff --git a/Engine/Source/Runtime/Engine/Private/Collision/SceneQueryLowLevel.cpp b/Engine/Source/Runtime/Engine/Private/Collision/SceneQueryLowLevel.cpp
--- a/Engine/Source/Runtime/Engine/Private/Collision/SceneQueryLowLevel.cpp
+++ b/Engine/Source/Runtime/Engine/Private/Collision/SceneQueryLowLevel.cpp
@@ -3,23 +3,25 @@
 #if WITH_PHYSX
 
 #include "PhysXPublic.h"
+#endif
 #include "Physics/PhysicsInterfaceDeclares.h"
 #include "Physics/PhysicsInterfaceCore.h"
 #include "PhysicsInterfaceDeclaresCore.h"
 
-#if !WITH_CHAOS_NEEDS_TO_BE_FIXED
-#include "SQAccelerator.h"
 #include "PhysicsEngine/CollisionQueryFilterCallback.h"
 #include "PhysicsCore.h"
 #if PHYSICS_INTERFACE_PHYSX
 #include "PhysXInterfaceWrapper.h"
-#elif PHYSICS_INTERFACE_LLIMMEDIATE
-#include "Physics/Experimental//LLImmediateInterfaceWrapper.h"
-#include "Experimental/SceneQueryLLImmediateImp.h"
 #endif
 
 #include "PhysTestSerializer.h"
+
+#if INCLUDE_CHAOS
+#include "SQAccelerator.h"
 #include "SQVerifier.h"
+#include "PBDRigidsSolver.h"
+#include "Chaos/PBDRigidsEvolutionGBF.h"
+#endif
 
 int32 ForceStandardSQ = 0;
 FAutoConsoleVariableRef CVarForceStandardSQ(TEXT("p.ForceStandardSQ"), ForceStandardSQ, TEXT("If enabled, we force the standard scene query even if custom SQ structure is enabled"));
@@ -40,11 +42,11 @@ void FinalizeCapture(FPhysTestSerializer& Serializer)
 	{
 		Serializer.Serialize(TEXT("SQCapture"));
 	}
-#if INCLUDE_CHAOS
+#if INCLUDE_CHAOS && WITH_PHYSX
 	if (ReplaySQs)
 	{
 		const bool bReplaySuccess = SQComparisonHelper(Serializer);
-		if (!bReplaySuccess)
+		if (!bReplaySuccess)'ve just'
 		{
 			UE_LOG(LogPhysicsCore, Warning, TEXT("Chaos SQ does not match physx"));
 			if (SerializeBadSQs && !SerializeSQs)
@@ -64,99 +66,83 @@ void FinalizeCapture(FPhysTestSerializer& Serializer) {}
 
 void LowLevelRaycast(FPhysScene& Scene, const FVector& Start, const FVector& Dir, float DeltaMag, FPhysicsHitCallback<FHitRaycast>& HitBuffer, EHitFlags OutputFlags, FQueryFlags QueryFlags, const FCollisionFilterData& Filter, const FQueryFilterData& QueryFilterData, ICollisionQueryFilterCallbackBase* QueryCallback)
 {
-#if WITH_CUSTOM_SQ_STRUCTURE
-	if (ForceStandardSQ == 0)
+#if !defined(PHYSICS_INTERFACE_PHYSX) || !PHYSICS_INTERFACE_PHYSX
+	Chaos::FPhysicsSolver* Solver = Scene.GetSolver();
+	Chaos::FPhysicsSolver::FPBDRigidsEvolution* Evolution = Solver->GetEvolution();
+	FChaosSQAccelerator SQAccelerator(*Evolution);
+	//ISQAccelerator* SQAccelerator = Scene.GetSQAccelerator();
+	SQAccelerator.Raycast(Start, Dir, DeltaMag, HitBuffer, OutputFlags, QueryFilterData, *QueryCallback);
+#else
+	if (SerializeSQs | ReplaySQs)
 	{
-		ISQAccelerator* SQAccelerator = Scene.GetSQAccelerator();
-		SQAccelerator->Raycast(Start, Dir, DeltaMag, HitBuffer, OutputFlags, QueryFlags, Filter, QueryFilterData, *QueryCallback);
-		FinalizeQuery(HitBuffer);
+		FPhysTestSerializer Serializer;
+		Serializer.SetPhysicsData(*Scene.GetPxScene());
+		FSQCapture& SweepCapture = Serializer.CaptureSQ();
+		SweepCapture.StartCapturePhysXRaycast(*Scene.GetPxScene(), Start, Dir, DeltaMag, OutputFlags, QueryFilterData, Filter, *QueryCallback);
+		Scene.GetPxScene()->raycast(U2PVector(Start), U2PVector(Dir), DeltaMag, HitBuffer, U2PHitFlags(OutputFlags), QueryFilterData, QueryCallback);
+		SweepCapture.EndCapturePhysXRaycast(HitBuffer);
+
+		FinalizeCapture(Serializer);
 	}
 	else
-#endif
 	{
-#if PHYSICS_INTERFACE_PHYSX
-		if (SerializeSQs | ReplaySQs)
-		{
-			FPhysTestSerializer Serializer;
-			Serializer.SetPhysicsData(*Scene.GetPxScene());
-			FSQCapture& SweepCapture = Serializer.CaptureSQ();
-			SweepCapture.StartCapturePhysXRaycast(*Scene.GetPxScene(), Start, Dir, DeltaMag, OutputFlags, QueryFilterData, Filter, *QueryCallback);
-			Scene.GetPxScene()->raycast(U2PVector(Start), U2PVector(Dir), DeltaMag, HitBuffer, U2PHitFlags(OutputFlags), QueryFilterData, QueryCallback);
-			SweepCapture.EndCapturePhysXRaycast(HitBuffer);
-
-			FinalizeCapture(Serializer);
-		}
-		else
-		{
-			Scene.GetPxScene()->raycast(U2PVector(Start), U2PVector(Dir), DeltaMag, HitBuffer, U2PHitFlags(OutputFlags), QueryFilterData, QueryCallback);
-		}
-#endif
+		Scene.GetPxScene()->raycast(U2PVector(Start), U2PVector(Dir), DeltaMag, HitBuffer, U2PHitFlags(OutputFlags), QueryFilterData, QueryCallback);
 	}
+#endif
 }
 
 void LowLevelSweep(FPhysScene& Scene, const FPhysicsGeometry& QueryGeom, const FTransform& StartTM, const FVector& Dir, float DeltaMag, FPhysicsHitCallback<FHitSweep>& HitBuffer, EHitFlags OutputFlags, FQueryFlags QueryFlags, const FCollisionFilterData& Filter, const FQueryFilterData& QueryFilterData, ICollisionQueryFilterCallbackBase* QueryCallback)
 {
-#if WITH_CUSTOM_SQ_STRUCTURE
-	if (ForceStandardSQ == 0)
+#if !defined(PHYSICS_INTERFACE_PHYSX) || !PHYSICS_INTERFACE_PHYSX
+	Chaos::FPhysicsSolver* Solver = Scene.GetSolver();
+	Chaos::FPhysicsSolver::FPBDRigidsEvolution* Evolution = Solver->GetEvolution();
+	FChaosSQAccelerator SQAccelerator(*Evolution);
+	//ISQAccelerator* SQAccelerator = Scene.GetSQAccelerator();
+	SQAccelerator.Sweep(QueryGeom, StartTM, Dir, DeltaMag, HitBuffer, OutputFlags, QueryFilterData, *QueryCallback);
+#else
+	if (SerializeSQs | ReplaySQs)
 	{
-		ISQAccelerator* SQAccelerator = Scene.GetSQAccelerator();
-		SQAccelerator->Sweep(QueryGeom, StartTM, Dir, DeltaMag, HitBuffer, OutputFlags, QueryFlags, Filter, QueryFilterData, *QueryCallback);
-		FinalizeQuery(HitBuffer);
+		FPhysTestSerializer Serializer;
+		Serializer.SetPhysicsData(*Scene.GetPxScene());
+		FSQCapture& SweepCapture = Serializer.CaptureSQ();
+		SweepCapture.StartCapturePhysXSweep(*Scene.GetPxScene(), QueryGeom, StartTM, Dir, DeltaMag, OutputFlags, QueryFilterData, Filter, *QueryCallback);
+		Scene.GetPxScene()->sweep(QueryGeom, U2PTransform(StartTM), U2PVector(Dir), DeltaMag, HitBuffer, U2PHitFlags(OutputFlags), QueryFilterData, QueryCallback);
+		SweepCapture.EndCapturePhysXSweep(HitBuffer);
+
+		FinalizeCapture(Serializer);
 	}
 	else
-#endif
 	{
-#if PHYSICS_INTERFACE_PHYSX
-		if (SerializeSQs | ReplaySQs)
-		{
-			FPhysTestSerializer Serializer;
-			Serializer.SetPhysicsData(*Scene.GetPxScene());
-			FSQCapture& SweepCapture = Serializer.CaptureSQ();
-			SweepCapture.StartCapturePhysXSweep(*Scene.GetPxScene(), QueryGeom, StartTM, Dir, DeltaMag, OutputFlags, QueryFilterData, Filter, *QueryCallback);
-			Scene.GetPxScene()->sweep(QueryGeom, U2PTransform(StartTM), U2PVector(Dir), DeltaMag, HitBuffer, U2PHitFlags(OutputFlags), QueryFilterData, QueryCallback);
-			SweepCapture.EndCapturePhysXSweep(HitBuffer);
-
-			FinalizeCapture(Serializer);
-		}
-		else
-		{
-			Scene.GetPxScene()->sweep(QueryGeom, U2PTransform(StartTM), U2PVector(Dir), DeltaMag, HitBuffer, U2PHitFlags(OutputFlags), QueryFilterData, QueryCallback);
-		}
-#endif
+		Scene.GetPxScene()->sweep(QueryGeom, U2PTransform(StartTM), U2PVector(Dir), DeltaMag, HitBuffer, U2PHitFlags(OutputFlags), QueryFilterData, QueryCallback);
 	}
+#endif
 }
 
-void LowLevelOverlap(FPhysScene& Scene, const PxGeometry& QueryGeom, const FTransform& GeomPose, FPhysicsHitCallback<FHitOverlap>& HitBuffer, FQueryFlags QueryFlags, const FCollisionFilterData& Filter, const FQueryFilterData& QueryFilterData, ICollisionQueryFilterCallbackBase* QueryCallback)
+void LowLevelOverlap(FPhysScene& Scene, const FPhysicsGeometry& QueryGeom, const FTransform& GeomPose, FPhysicsHitCallback<FHitOverlap>& HitBuffer, FQueryFlags QueryFlags, const FCollisionFilterData& Filter, const FQueryFilterData& QueryFilterData, ICollisionQueryFilterCallbackBase* QueryCallback)
 {
-#if WITH_CUSTOM_SQ_STRUCTURE
-	if (ForceStandardSQ == 0)
+#if !defined(PHYSICS_INTERFACE_PHYSX) || !PHYSICS_INTERFACE_PHYSX
+	Chaos::FPhysicsSolver* Solver = Scene.GetSolver();
+	Chaos::FPhysicsSolver::FPBDRigidsEvolution* Evolution = Solver->GetEvolution();
+	FChaosSQAccelerator SQAccelerator(*Evolution);
+	//ISQAccelerator* SQAccelerator = Scene.GetSQAccelerator();
+	SQAccelerator.Overlap(QueryGeom, GeomPose, HitBuffer, QueryFilterData, *QueryCallback);
+#else
+	if (SerializeSQs | ReplaySQs)
 	{
-		ISQAccelerator* SQAccelerator = Scene.GetSQAccelerator();
-		SQAccelerator->Overlap(QueryGeom, GeomPose, HitBuffer, QueryFlags, Filter, QueryFilterData, *QueryCallback);
-		FinalizeQuery(HitBuffer);
+		FPhysTestSerializer Serializer;
+		Serializer.SetPhysicsData(*Scene.GetPxScene());
+		FSQCapture& SweepCapture = Serializer.CaptureSQ();
+		SweepCapture.StartCapturePhysXOverlap(*Scene.GetPxScene(), QueryGeom, GeomPose, QueryFilterData, Filter, *QueryCallback);
+		Scene.GetPxScene()->overlap(QueryGeom, U2PTransform(GeomPose), HitBuffer, QueryFilterData, QueryCallback);
+		SweepCapture.EndCapturePhysXOverlap(HitBuffer);
+
+		FinalizeCapture(Serializer);
 	}
 	else
-#endif
 	{
-#if PHYSICS_INTERFACE_PHYSX
-		if (SerializeSQs | ReplaySQs)
-		{
-			FPhysTestSerializer Serializer;
-			Serializer.SetPhysicsData(*Scene.GetPxScene());
-			FSQCapture& SweepCapture = Serializer.CaptureSQ();
-			SweepCapture.StartCapturePhysXOverlap(*Scene.GetPxScene(), QueryGeom, GeomPose, QueryFilterData, Filter, *QueryCallback);
-			Scene.GetPxScene()->overlap(QueryGeom, U2PTransform(GeomPose), HitBuffer, QueryFilterData, QueryCallback);
-			SweepCapture.EndCapturePhysXOverlap(HitBuffer);
-
-			FinalizeCapture(Serializer);
-		}
-		else
-		{
-			Scene.GetPxScene()->overlap(QueryGeom, U2PTransform(GeomPose), HitBuffer, QueryFilterData, QueryCallback);
-		}		
+		Scene.GetPxScene()->overlap(QueryGeom, U2PTransform(GeomPose), HitBuffer, QueryFilterData, QueryCallback);
+	}		
 #endif
-	}
 }
 
-#endif // !WITH_CHAOS_NEEDS_TO_BE_FIXED
-#endif // WITH_PHYSX 
\ No newline at end of file
+//#endif // WITH_PHYSX 
\ No newline at end of file
diff --git a/Engine/Source/Runtime/Engine/Private/Collision/WorldCollision.cpp b/Engine/Source/Runtime/Engine/Private/Collision/WorldCollision.cpp
--- a/Engine/Source/Runtime/Engine/Private/Collision/WorldCollision.cpp
+++ b/Engine/Source/Runtime/Engine/Private/Collision/WorldCollision.cpp
@@ -391,7 +391,7 @@ bool UWorld::ComponentSweepMulti(TArray<struct FHitResult>& OutHits, class UPrim
 #if WITH_PHYSX
 	FPhysicsCommand::ExecuteRead(BodyInstance->ActorHandle, [&](const FPhysicsActorHandle& Actor)
 	{
-		if(!Actor.IsValid())
+		if(!FPhysicsInterface::IsValid(Actor))
 		{
 			return;
 		}
@@ -423,9 +423,6 @@ bool UWorld::ComponentSweepMulti(TArray<struct FHitResult>& OutHits, class UPrim
 			// consider localshape rotation for shape rotation
 			const FQuat ShapeQuat = Quat * ShapeLocalTransform.GetRotation();
 
-#if WITH_CHAOS
-            check(false);
-#else
 			FPhysicsGeometryCollection GeomCollection = FPhysicsInterface::GetGeometryCollection(Shape);
 			TArray<FHitResult> TmpHits;
 			if(FPhysicsInterface::GeomSweepMulti(this, GeomCollection, ShapeQuat, TmpHits, GlobalStartTransform_Shape.GetTranslation(), GlobalEndTransform_Shape.GetTranslation(), TraceChannel, Params, FCollisionResponseParams(PrimComp->GetCollisionResponseToChannels())))
@@ -433,7 +430,6 @@ bool UWorld::ComponentSweepMulti(TArray<struct FHitResult>& OutHits, class UPrim
 				bHaveBlockingHit = true;
 			}
 			OutHits.Append(TmpHits);	//todo: should these be made unique?
-#endif
 		}
 	});
 #endif //WITH_PHYSX
diff --git a/Engine/Source/Runtime/Engine/Private/Components/SkeletalMeshComponent.cpp b/Engine/Source/Runtime/Engine/Private/Components/SkeletalMeshComponent.cpp
--- a/Engine/Source/Runtime/Engine/Private/Components/SkeletalMeshComponent.cpp
+++ b/Engine/Source/Runtime/Engine/Private/Components/SkeletalMeshComponent.cpp
@@ -198,6 +198,11 @@ USkeletalMeshComponent::USkeletalMeshComponent(const FObjectInitializer& ObjectI
 
 #endif//#if WITH_APEX_CLOTHING
 
+	MassMode = EClothMassMode::Density;
+	UniformMass = 1.f;
+	TotalMass = 100.0f;
+	Density = 0.1f;
+	MinPerParticleMass = 0.0001f;
 	EdgeStiffness = 1.f;
 	BendingStiffness = 1.f;
 	AreaStiffness = 1.f;
diff --git a/Engine/Source/Runtime/Engine/Private/PhysicsEngine/BodyInstance.cpp b/Engine/Source/Runtime/Engine/Private/PhysicsEngine/BodyInstance.cpp
--- a/Engine/Source/Runtime/Engine/Private/PhysicsEngine/BodyInstance.cpp
+++ b/Engine/Source/Runtime/Engine/Private/PhysicsEngine/BodyInstance.cpp
@@ -344,6 +344,7 @@ FBodyInstance::FBodyInstance()
 	, CustomSleepThresholdMultiplier(1.f)
 	, StabilizationThresholdMultiplier(1.f)
 	, PhysicsBlendWeight(0.f)
+	, ActorHandle(DefaultPhysicsActorHandle())
 {
 	MaxAngularVelocity = UPhysicsSettings::Get()->MaxAngularVelocity;
 }
@@ -437,14 +438,12 @@ void FBodyInstance::UpdateTriMeshVertices(const TArray<FVector> & NewPositions)
 			TArray<FPhysicsShapeHandle> Shapes;
 			const int32 SyncShapeCount = GetAllShapes_AssumesLocked(Shapes);
 
-			PxTriangleMeshGeometry PTriangleMeshGeometry;
 			for(FPhysicsShapeHandle& Shape : Shapes)
 			{
 				if(FPhysicsInterface::GetShapeType(Shape) == ECollisionShapeType::Trimesh)
 				{
 					FPhysicsGeometryCollection GeoCollection = FPhysicsInterface::GetGeometryCollection(Shape);
-					GeoCollection.GetTriMeshGeometry(PTriangleMeshGeometry);
-					FPhysicsInterface::SetGeometry(Shape, PTriangleMeshGeometry);
+					FPhysicsInterface::SetGeometry(Shape, GeoCollection.GetTriMeshGeometry());
 				}
 			}
 		});
@@ -1007,12 +1006,13 @@ struct FInitBodiesHelper
 	// Return to actor ref
 	void CreateActor_AssumesLocked(FBodyInstance* Instance, const FTransform& Transform) const
 	{
-		checkSlow(!Instance->ActorHandle.IsValid());
+		checkSlow(!FPhysicsInterface::IsValid(Instance->ActorHandle));
 		const ECollisionEnabled::Type CollisionType = Instance->GetCollisionEnabled();
 		const bool bDisableSim = !CollisionEnabledHasPhysics(CollisionType) && CDisableQueryOnlyActors.GetValueOnGameThread();
 
 		FActorCreationParams ActorParams;
 		ActorParams.InitialTM = Transform;
+		ActorParams.BodyInstance = Instance;
 #if USE_BODYINSTANCE_DEBUG_NAMES
 		ActorParams.DebugName = Instance->CharDebugName.IsValid() ? Instance->CharDebugName->GetData() : nullptr;
 #endif
@@ -1111,7 +1111,7 @@ struct FInitBodiesHelper
 
 			// Don't process if we've already got a body
 			// Just ask actorref
-			if(Instance->GetPhysicsActorHandle().IsValid())
+			if (FPhysicsInterface::IsValid(Instance->GetPhysicsActorHandle()))
 			{
 				Instance->OwnerComponent = nullptr;
 				Instance->BodySetup      = nullptr;
@@ -1131,7 +1131,7 @@ struct FInitBodiesHelper
 			}
 
 			// Init user data structure to point back at this instance
-			Instance->PhysxUserData = FPhysxUserData(Instance);
+			Instance->PhysicsUserData = FPhysicsUserData(Instance);
 
 			CreateActor_AssumesLocked(Instance, Transform);
 			const bool bInitFail = CreateShapes_AssumesLocked(Instance);
@@ -1154,7 +1154,7 @@ struct FInitBodiesHelper
 				continue;
 			}
 
-			FPhysicsInterface::SetActorUserData_AssumesLocked(Instance->ActorHandle, &Instance->PhysxUserData);
+			FPhysicsInterface::SetActorUserData_AssumesLocked(Instance->ActorHandle, &Instance->PhysicsUserData);
 			}
 
 		return true;
@@ -1177,7 +1177,7 @@ struct FInitBodiesHelper
 					for(FBodyInstance* BI : Bodies)
 					{
 						const FPhysicsActorHandle& ActorHandle = BI->GetPhysicsActorHandle();
-						if(ActorHandle.IsValid())
+						if(FPhysicsInterface::IsValid(ActorHandle))
 						{
 							FPhysicsInterface::AddActorToAggregate_AssumesLocked(Aggregate, ActorHandle);
 						}
@@ -1190,10 +1190,21 @@ struct FInitBodiesHelper
 
 					for(FBodyInstance* BI : Bodies)
 					{
-						const FPhysicsActorHandle& ActorHandle = BI->GetPhysicsActorHandle();
-						if(ActorHandle.IsValid())
+						FPhysicsActorHandle& ActorHandle = BI->GetPhysicsActorHandle();
+						if(FPhysicsInterface::IsValid(ActorHandle))
 						{
 							ActorHandles.Add(ActorHandle);
+
+/*
+							With the implementation of AddActorsToScene_AssumesLocked,
+							this call duplicates the AddToSolver operation and is not necessary
+#ifdef WITH_CHAOS
+							if (BI->bSimulatePhysics)
+							{
+								PhysScene->AddActorToSolver(ActorHandle);
+							}
+#endif
+*/
 						}
 					}
 
@@ -1620,14 +1631,7 @@ bool FBodyInstance::UpdateBodyScale(const FVector& InScale3D, bool bForceUpdate)
 
 		// Apply scaling
 		ComputeScalingVectors(ScaleMode, InScale3D, AdjustedScale3D, AdjustedScale3DAbs);
-
-		//we need to allocate all of these here because PhysX insists on using the stack. This is wasteful, but reduces a lot of code duplication
-		PxSphereGeometry PSphereGeom;
-		PxBoxGeometry PBoxGeom;
-		PxCapsuleGeometry PCapsuleGeom;
-		PxConvexMeshGeometry PConvexGeom;
-		PxTriangleMeshGeometry PTriMeshGeom;
-
+		
 		for(FPhysicsShapeHandle& Shape : Shapes)
 		{
 			bool bInvalid = false;	//we only mark invalid if actually found geom and it's invalid scale
@@ -1650,7 +1654,7 @@ bool FBodyInstance::UpdateBodyScale(const FVector& InScale3D, bool bForceUpdate)
 					FKSphereElem* SphereElem = ShapeElem->GetShapeCheck<FKSphereElem>();
 					ensure(ScaleMode == EScaleMode::LockedXYZ);
 
-					GeoCollection.GetSphereGeometry(PSphereGeom);
+					PxSphereGeometry& PSphereGeom = GeoCollection.GetSphereGeometry();
 					 
 					PSphereGeom.radius = FMath::Max(SphereElem->Radius * AdjustedScale3DAbs.X, FCollisionShape::MinSphereRadius());
 
@@ -1671,7 +1675,7 @@ bool FBodyInstance::UpdateBodyScale(const FVector& InScale3D, bool bForceUpdate)
 				case ECollisionShapeType::Box:
 				{
 					FKBoxElem* BoxElem = ShapeElem->GetShapeCheck<FKBoxElem>();
-					GeoCollection.GetBoxGeometry(PBoxGeom);
+					PxBoxGeometry& PBoxGeom = GeoCollection.GetBoxGeometry();
 
 					PBoxGeom.halfExtents.x = FMath::Max((0.5f * BoxElem->X * AdjustedScale3DAbs.X), FCollisionShape::MinBoxExtent());
 					PBoxGeom.halfExtents.y = FMath::Max((0.5f * BoxElem->Y * AdjustedScale3DAbs.Y), FCollisionShape::MinBoxExtent());
@@ -1699,7 +1703,7 @@ bool FBodyInstance::UpdateBodyScale(const FVector& InScale3D, bool bForceUpdate)
 					float ScaleRadius = FMath::Max(AdjustedScale3DAbs.X, AdjustedScale3DAbs.Y);
 					float ScaleLength = AdjustedScale3DAbs.Z;
 
-					GeoCollection.GetCapsuleGeometry(PCapsuleGeom);
+					PxCapsuleGeometry& PCapsuleGeom = GeoCollection.GetCapsuleGeometry();
 
 					// this is a bit confusing since radius and height is scaled
 					// first apply the scale first 
@@ -1733,7 +1737,7 @@ bool FBodyInstance::UpdateBodyScale(const FVector& InScale3D, bool bForceUpdate)
 				case ECollisionShapeType::Convex:
 				{
 					FKConvexElem* ConvexElem = ShapeElem->GetShapeCheck<FKConvexElem>();
-					GeoCollection.GetConvexGeometry(PConvexGeom);
+					PxConvexMeshGeometry& PConvexGeom = GeoCollection.GetConvexGeometry();
 
 					bool bUseNegX = CalcMeshNegScaleCompensation(AdjustedScale3D, LocalTransform);
 
@@ -1759,7 +1763,7 @@ bool FBodyInstance::UpdateBodyScale(const FVector& InScale3D, bool bForceUpdate)
 				{
 					check(ShapeElem == nullptr);	//trimesh shape doesn't have userData
 
-					GeoCollection.GetTriMeshGeometry(PTriMeshGeom);
+					PxTriangleMeshGeometry& PTriMeshGeom = GeoCollection.GetTriMeshGeometry();
 
 					// find which trimesh elems it is
 					// it would be nice to know if the order of PShapes array index is in the order of createShape
@@ -1982,7 +1986,7 @@ void FBodyInstance::SetInstanceSimulatePhysics(bool bSimulate, bool bMaintainPhy
 
 bool FBodyInstance::IsValidBodyInstance() const
 {
-	return ActorHandle.IsValid();
+	return FPhysicsInterface::IsValid(ActorHandle);
 }
 
 FTransform GetUnrealWorldTransformImp_AssumesLocked(const FBodyInstance* BodyInstance, bool bWithProjection, bool bGlobalPose)
@@ -2039,7 +2043,7 @@ void FBodyInstance::SetBodyTransform(const FTransform& NewTransform, ETeleportTy
 		return;
 	}
 
-	if(ActorHandle.IsValid())
+	if(FPhysicsInterface::IsValid(ActorHandle))
 	{
 		if(!NewTransform.IsValid())
 		{
@@ -2115,10 +2119,10 @@ FVector FBodyInstance::GetUnrealWorldVelocity() const
 FVector FBodyInstance::GetUnrealWorldVelocity_AssumesLocked() const
 {
 	FVector LinVel(EForceInit::ForceInitToZero);
-	if(ActorHandle.IsValid())
-{
+	if (FPhysicsInterface::IsValid(ActorHandle))
+	{
 		LinVel = FPhysicsInterface::GetLinearVelocity_AssumesLocked(ActorHandle);
-}
+	}
 
 	return LinVel;
 }
@@ -2244,8 +2248,8 @@ void FBodyInstance::CopyBodyInstancePropertiesFrom(const FBodyInstance* FromInst
 	check(FromInst);
 	check(FromInst->OwnerComponent.Get() == NULL);
 	check(FromInst->BodySetup.Get() == NULL);
-	check(!FromInst->ActorHandle.IsValid());
-	check(!ActorHandle.IsValid());
+	check(!FPhysicsInterface::IsValid(FromInst->ActorHandle));
+	check(!FPhysicsInterface::IsValid(ActorHandle));
 
 	*this = *FromInst;
 }
@@ -2269,25 +2273,14 @@ void FBodyInstance::CopyRuntimeBodyInstancePropertiesFrom(const FBodyInstance* F
 
 const FPhysScene* FBodyInstance::GetPhysicsScene() const
 {
-	if(ActorHandle.IsValid())
-	{
-		return FPhysicsInterface::GetCurrentScene(ActorHandle);
-	}
-	
-	return nullptr;
+	return FPhysicsInterface::GetCurrentScene(ActorHandle);
 }
 
 FPhysScene* FBodyInstance::GetPhysicsScene()
 {
-	if (ActorHandle.IsValid())
-	{
-		return FPhysicsInterface::GetCurrentScene(ActorHandle);
-	}
-
-	return nullptr;
+	return FPhysicsInterface::GetCurrentScene(ActorHandle);
 }
 
-
 FPhysicsActorHandle& FBodyInstance::GetPhysicsActorHandle()
 {
 	return ActorHandle;
@@ -2538,11 +2531,11 @@ void FBodyInstance::UpdateMassProperties()
 	UPhysicalMaterial* PhysMat = GetSimplePhysicalMaterial();
 
 #if WITH_PHYSX
-	if(ActorHandle.IsValid() && FPhysicsInterface::IsRigidBody(ActorHandle))
+	if(FPhysicsInterface::IsValid(ActorHandle) && FPhysicsInterface::IsRigidBody(ActorHandle))
 	{
 		FPhysicsCommand::ExecuteWrite(ActorHandle, [&](const FPhysicsActorHandle& Actor)
 		{
-			check(Actor.IsValid());
+			check(FPhysicsInterface::IsValid(Actor));
 
 			if(GetNumSimShapes_AssumesLocked(Actor) > 0)
 			{
@@ -3083,14 +3076,11 @@ bool FBodyInstance::OverlapTestForBodiesImpl(const FVector& Pos, const FQuat& Ro
 
 		// Get all the shapes from the actor
 		FInlineShapeArray TargetShapes;
+
 		const int32 NumTargetShapes = FillInlineShapeArray_AssumesLocked(TargetShapes, Actor);
 
 		for(const FPhysicsShapeHandle& Shape : TargetShapes)
 		{
-#if WITH_CHAOS || WITH_IMMEDIATE_PHYSX
-			check(false);
-#else
-
 			ECollisionShapeType ShapeType = FPhysicsInterface::GetShapeType(Shape);
 			if(ShapeType == ECollisionShapeType::Heightfield || ShapeType == ECollisionShapeType::Trimesh)
 			{
@@ -3108,7 +3098,6 @@ bool FBodyInstance::OverlapTestForBodiesImpl(const FVector& Pos, const FQuat& Ro
 					return;
 				}
 			}
-#endif
 		}
 	});
 	return bHaveOverlap;
@@ -3198,7 +3187,7 @@ bool FBodyInstance::OverlapMulti(TArray<struct FOverlapResult>& InOutOverlaps, c
 
 	FPhysicsCommand::ExecuteRead(TargetInstance->ActorHandle, [&](const FPhysicsActorHandle& Actor)
 	{
-		if(Actor.IsValid())
+		if(FPhysicsInterface::IsValid(Actor))
 		{
 		// Get all the shapes from the actor
 			FInlineShapeArray PShapes;
@@ -3208,9 +3197,6 @@ bool FBodyInstance::OverlapMulti(TArray<struct FOverlapResult>& InOutOverlaps, c
 			TArray<struct FOverlapResult> TempOverlaps;
 			for (int32 ShapeIdx = 0; ShapeIdx < NumShapes; ShapeIdx++)
 			{
-#if WITH_CHAOS || WITH_IMMEDIATE_PHYSX
-				check(false);
-#else
 				FPhysicsShapeHandle& ShapeRef = PShapes[ShapeIdx];
 
 				FPhysicsGeometryCollection GeomCollection = FPhysicsInterface::GetGeometryCollection(ShapeRef);
@@ -3236,7 +3222,6 @@ bool FBodyInstance::OverlapMulti(TArray<struct FOverlapResult>& InOutOverlaps, c
 					bHaveBlockingHit = true;
 				}
 				InOutOverlaps.Append(TempOverlaps);
-#endif
 			}
 			}
 		});
diff --git a/Engine/Source/Runtime/Engine/Private/PhysicsEngine/BodySetup.cpp b/Engine/Source/Runtime/Engine/Private/PhysicsEngine/BodySetup.cpp
--- a/Engine/Source/Runtime/Engine/Private/PhysicsEngine/BodySetup.cpp
+++ b/Engine/Source/Runtime/Engine/Private/PhysicsEngine/BodySetup.cpp
@@ -30,17 +30,27 @@
 #endif // WITH_PHYSX
 
 #include "Modules/ModuleManager.h"
+
 #if WITH_PHYSX
 	#include "IPhysXCookingModule.h"
 	#include "IPhysXCooking.h"
+	#include "PhysicsEngine/PhysDerivedData.h"
 #endif
 
 #include "Physics/PhysicsInterfaceUtils.h"
-#include "PhysicsEngine/PhysDerivedData.h"
 #include "PhysicalMaterials/PhysicalMaterial.h"
 #include "ProfilingDebugging/CookStats.h"
 #include "UObject/AnimPhysObjectVersion.h"
 
+#if INCLUDE_CHAOS
+#include "Chaos/TriangleMeshImplicitObject.h"
+#endif
+
+#if WITH_CHAOS
+	#include "Experimental/ChaosDerivedData.h"
+	#include "Physics/Experimental/ChaosDerivedDataReader.h"
+#endif
+
 /** Helper for enum output... */
 #ifndef CASE_ENUM_TO_TEXT
 #define CASE_ENUM_TO_TEXT(txt) case txt: return TEXT(#txt);
@@ -91,6 +101,13 @@ FCookBodySetupInfo::FCookBodySetupInfo() :
 {
 }
 
+UBodySetup::UBodySetup(FVTableHelper& Helper)
+	: Super(Helper)
+{
+}
+
+UBodySetup::~UBodySetup() = default;
+
 
 #if ENABLE_COOK_STATS
 namespace PhysXBodySetupCookStats
@@ -183,6 +200,9 @@ void UBodySetup::CopyBodyPropertiesFrom(const UBodySetup* FromSetup)
 		FKConvexElem& ConvexElem = AggGeom.ConvexElems[i];
 		ConvexElem.SetConvexMesh(nullptr);
 		ConvexElem.SetMirroredConvexMesh(nullptr);
+#if WITH_CHAOS
+		ConvexElem.ResetChaosConvexMesh();
+#endif
 	}
 
 	DefaultInstance.CopyBodyInstancePropertiesFrom(&FromSetup->DefaultInstance);
@@ -209,6 +229,9 @@ void UBodySetup::AddCollisionFrom(const FKAggregateGeom& FromAggGeom)
 		FKConvexElem& ConvexElem = AggGeom.ConvexElems[i];
 		ConvexElem.SetConvexMesh(nullptr);
 		ConvexElem.SetMirroredConvexMesh(nullptr);
+#if WITH_CHAOS
+		ConvexElem.ResetChaosConvexMesh();
+#endif
 	}
 }
 
@@ -395,7 +418,6 @@ void UBodySetup::CreatePhysicsMeshes()
 {
 	SCOPE_CYCLE_COUNTER(STAT_CreatePhysicsMeshes);
 
-#if WITH_PHYSX
 	// Create meshes from cooked data if not already done
 	if(bCreatedPhysicsMeshes)
 	{
@@ -423,70 +445,89 @@ void UBodySetup::CreatePhysicsMeshes()
 
 	if (FormatData)
 	{
-		if (FormatData->IsLocked())
+#if WITH_PHYSX  && PHYSICS_INTERFACE_PHYSX
+		bClearMeshes = !ProcessFormatData_PhysX(FormatData);
+#elif WITH_CHAOS
+		bClearMeshes = !ProcessFormatData_Chaos(FormatData);
+#endif
+	}
+	else
+	{
+		if (IsRuntime(this))
 		{
-			// seems it's being already processed
-			return;
+#if WITH_PHYSX  && PHYSICS_INTERFACE_PHYSX
+			bClearMeshes = !RuntimeCookPhysics_PhysX();
+#elif WITH_CHAOS
+			bClearMeshes = !RuntimeCookPhysics_Chaos();
+#endif
 		}
+	}
 
-		FPhysXCookingDataReader CookedDataReader(*FormatData, &UVInfo);
+	if(bClearMeshes)
+	{
+		ClearPhysicsMeshes();
+	}
+	
+	bCreatedPhysicsMeshes = true;
+
+}
+
+#if WITH_PHYSX && PHYSICS_INTERFACE_PHYSX
+bool UBodySetup::RuntimeCookPhysics_PhysX()
+{
+	FPhysXCookHelper CookHelper(GetPhysXCookingModule());
 
-		if (GetCollisionTraceFlag() != CTF_UseComplexAsSimple)
+	GetCookInfo(CookHelper.CookInfo, GetRuntimeOnlyCookOptimizationFlags());
+	if(CookHelper.HasSomethingToCook(CookHelper.CookInfo))
+	{
+		if(!IsRuntimeCookingEnabled())
 		{
-			bool bNeedsCooking = bGenerateNonMirroredCollision && CookedDataReader.ConvexMeshes.Num() != AggGeom.ConvexElems.Num();
-			bNeedsCooking = bNeedsCooking || (bGenerateMirroredCollision && CookedDataReader.ConvexMeshesNegX.Num() != AggGeom.ConvexElems.Num());
-			if (bNeedsCooking)	//Because of bugs it's possible to save with out of sync cooked data. In editor we want to fixup this data
+			UE_LOG(LogPhysics, Error, TEXT("Attempting to build physics data for %s at runtime, but runtime cooking is disabled (see the RuntimePhysXCooking plugin)."), *GetPathName());
+		}
+		else
+		{
+			if(CookHelper.CreatePhysicsMeshes_Concurrent())
+			{
+				FinishCreatingPhysicsMeshes_PhysX(CookHelper.OutNonMirroredConvexMeshes, CookHelper.OutMirroredConvexMeshes, CookHelper.OutTriangleMeshes);
+				bFailedToCreatePhysicsMeshes = false;
+				return true;
+			}
+			else
 			{
-				InvalidatePhysicsData();
-				CreatePhysicsMeshes();
-				return;
+				bFailedToCreatePhysicsMeshes = true;
 			}
 		}
+	}			
+	return false;
+}
 
-		FinishCreatingPhysicsMeshes(CookedDataReader.ConvexMeshes, CookedDataReader.ConvexMeshesNegX, CookedDataReader.TriMeshes);
-		bClearMeshes = false;
+bool UBodySetup::ProcessFormatData_PhysX(FByteBulkData* FormatData)
+{
+	if(FormatData->IsLocked())
+	{
+		// seems it's being already processed
+		return false;
 	}
-	else
+
+	FPhysXCookingDataReader CookedDataReader(*FormatData, &UVInfo);
+
+	if(GetCollisionTraceFlag() != CTF_UseComplexAsSimple)
 	{
-		if (IsRuntime(this))
+		bool bNeedsCooking = bGenerateNonMirroredCollision && CookedDataReader.ConvexMeshes.Num() != AggGeom.ConvexElems.Num();
+		bNeedsCooking = bNeedsCooking || (bGenerateMirroredCollision && CookedDataReader.ConvexMeshesNegX.Num() != AggGeom.ConvexElems.Num());
+		if(bNeedsCooking)	//Because of bugs it's possible to save with out of sync cooked data. In editor we want to fixup this data
 		{
-			FPhysXCookHelper CookHelper(GetPhysXCookingModule());
-					
-			GetCookInfo(CookHelper.CookInfo, GetRuntimeOnlyCookOptimizationFlags());
-			if(CookHelper.HasSomethingToCook(CookHelper.CookInfo))
-			{
-				if (!IsRuntimeCookingEnabled())
-				{
-					UE_LOG(LogPhysics, Error, TEXT("Attempting to build physics data for %s at runtime, but runtime cooking is disabled (see the RuntimePhysXCooking plugin)."), *GetPathName());
-				}
-				else
-				{
-                    if (CookHelper.CreatePhysicsMeshes_Concurrent())
-                    {
-					    FinishCreatingPhysicsMeshes(CookHelper.OutNonMirroredConvexMeshes, CookHelper.OutMirroredConvexMeshes, CookHelper.OutTriangleMeshes);
-					    bClearMeshes = false;
-                        bFailedToCreatePhysicsMeshes = false;
-                    }
-                    else
-                    {
-                        bFailedToCreatePhysicsMeshes = true;
-                    }
-				}
-			}
+			InvalidatePhysicsData();
+			CreatePhysicsMeshes();
+			return false;
 		}
 	}
 
-	if(bClearMeshes)
-	{
-		ClearPhysicsMeshes();
-	}
-	
-	bCreatedPhysicsMeshes = true;
-#endif //WITH_PHYSX
+	FinishCreatingPhysicsMeshes_PhysX(CookedDataReader.ConvexMeshes, CookedDataReader.ConvexMeshesNegX, CookedDataReader.TriMeshes);
+	return true;
 }
 
-#if WITH_PHYSX
-void UBodySetup::FinishCreatingPhysicsMeshes(const TArray<PxConvexMesh*>& ConvexMeshes, const TArray<PxConvexMesh*>& ConvexMeshesNegX, const TArray<PxTriangleMesh*>& CookedTriMeshes)
+void UBodySetup::FinishCreatingPhysicsMeshes_PhysX(const TArray<PxConvexMesh*>& ConvexMeshes, const TArray<PxConvexMesh*>& ConvexMeshesNegX, const TArray<PxTriangleMesh*>& CookedTriMeshes)
 {
 	check(IsInGameThread());
 	ClearPhysicsMeshes();
@@ -544,12 +585,12 @@ void UBodySetup::CreatePhysicsMeshesAsync(FOnAsyncPhysicsCookFinished OnAsyncPhy
 {
 	check(IsInGameThread());
 
-#if WITH_PHYSX
+#if WITH_PHYSX && PHYSICS_INTERFACE_PHYSX
 	// Don't start another cook cycle if one's already in progress
 	check(CurrentCookHelper == nullptr);
 #endif
 
-#if WITH_PHYSX_COOKING
+#if WITH_PHYSX_COOKING && PHYSICS_INTERFACE_PHYSX
 	if (IsRuntime(this) && !IsRuntimeCookingEnabled())
 	{
 		UE_LOG(LogPhysics, Error, TEXT("Attempting to build physics data for %s at runtime, but runtime cooking is disabled (see the RuntimePhysXCooking plugin)."), *GetPathName());
@@ -558,7 +599,7 @@ void UBodySetup::CreatePhysicsMeshesAsync(FOnAsyncPhysicsCookFinished OnAsyncPhy
 	}
 #endif
 
-#if WITH_PHYSX
+#if WITH_PHYSX && PHYSICS_INTERFACE_PHYSX
 	if(IPhysXCookingModule* PhysXCookingModule = GetPhysXCookingModule())
 	{
 		FPhysXCookHelper* AsyncPhysicsCookHelper = new FPhysXCookHelper(PhysXCookingModule);
@@ -595,7 +636,7 @@ void UBodySetup::AbortPhysicsMeshAsyncCreation()
 #endif
 }
 
-#if WITH_PHYSX
+#if WITH_PHYSX && PHYSICS_INTERFACE_PHYSX
 void UBodySetup::FinishCreatePhysicsMeshesAsync(FPhysXCookHelper* AsyncPhysicsCookHelper, FOnAsyncPhysicsCookFinished OnAsyncPhysicsCookFinished)
 {
 	// Ensure we haven't gotten multiple cooks going
@@ -607,7 +648,7 @@ void UBodySetup::FinishCreatePhysicsMeshesAsync(FPhysXCookHelper* AsyncPhysicsCo
 
 	if(AsyncPhysicsCookHelper)
 	{
-		FinishCreatingPhysicsMeshes(AsyncPhysicsCookHelper->OutNonMirroredConvexMeshes, AsyncPhysicsCookHelper->OutMirroredConvexMeshes, AsyncPhysicsCookHelper->OutTriangleMeshes);
+		FinishCreatingPhysicsMeshes_PhysX(AsyncPhysicsCookHelper->OutNonMirroredConvexMeshes, AsyncPhysicsCookHelper->OutMirroredConvexMeshes, AsyncPhysicsCookHelper->OutTriangleMeshes);
 		UVInfo = AsyncPhysicsCookHelper->OutUVInfo;
 		delete AsyncPhysicsCookHelper;
 
@@ -622,9 +663,58 @@ void UBodySetup::FinishCreatePhysicsMeshesAsync(FPhysXCookHelper* AsyncPhysicsCo
 }
 #endif // WITH_PHYSX
 
+#if WITH_CHAOS
+bool UBodySetup::ProcessFormatData_Chaos(FByteBulkData* FormatData)
+{
+	if(FormatData->IsLocked())
+	{
+		// seems it's being already processed
+		return false;
+	}
+
+	FChaosDerivedDataReader<float, 3> Reader(FormatData);
+	FinishCreatingPhysicsMeshes_Chaos(Reader);
+	
+	return true;
+}
+
+bool UBodySetup::RuntimeCookPhysics_Chaos()
+{
+	return true;
+}
+
+void UBodySetup::FinishCreatingPhysicsMeshes_Chaos(FChaosDerivedDataReader<float, 3>& InReader)
+{
+	check(IsInGameThread());
+	ClearPhysicsMeshes();
+
+	const FString FullName = GetFullName();
+	if (GetCollisionTraceFlag() != CTF_UseComplexAsSimple)
+	{
+		for (int32 ElementIndex = 0; ElementIndex < AggGeom.ConvexElems.Num(); ElementIndex++)
+		{
+			FKConvexElem& ConvexElem = AggGeom.ConvexElems[ElementIndex];
+			ConvexElem.SetChaosConvexMesh(MoveTemp(InReader.ConvexImplicitObjects[ElementIndex]));
+		}
+		InReader.ConvexImplicitObjects.Reset();
+	}
+
+	ChaosTriMeshes = MoveTemp(InReader.TrimeshImplicitObjects);
+
+	// Clear the cooked data
+	if (!GIsEditor && !bSharedCookedData)
+	{
+		CookedFormatData.FlushData();
+	}
+
+	bCreatedPhysicsMeshes = true;
+}
+
+#endif
+
 void UBodySetup::ClearPhysicsMeshes()
 {
-#if WITH_PHYSX
+#if WITH_PHYSX && PHYSICS_INTERFACE_PHYSX
 	for(int32 i=0; i<AggGeom.ConvexElems.Num(); i++)
 	{
 		FKConvexElem* ConvexElem = &(AggGeom.ConvexElems[i]);
@@ -654,9 +744,18 @@ void UBodySetup::ClearPhysicsMeshes()
 	}
 	TriMeshes.Empty();
 
-	bCreatedPhysicsMeshes = false;
+#elif WITH_CHAOS
+	for (int32 i = 0; i < AggGeom.ConvexElems.Num(); i++)
+	{
+		FKConvexElem* ConvexElem = &(AggGeom.ConvexElems[i]);
+		ConvexElem->ResetChaosConvexMesh();
+	}
+	ChaosTriMeshes.Reset();
 #endif // WITH_PHYSX
 
+	bCreatedPhysicsMeshes = false;
+
+
 	// Also clear render info
 	AggGeom.FreeRenderInfo();
 }
@@ -697,6 +796,9 @@ void UBodySetup::AddShapesToRigidActor_AssumesLocked(
 	AddParams.TriMeshes = TArrayView<PxTriangleMesh*>(TriMeshes);
 #endif
 
+#if WITH_CHAOS
+	AddParams.ChaosTriMeshes = MakeArrayView(ChaosTriMeshes);
+#endif
 	FPhysicsInterface::AddGeometry(OwningInstance->ActorHandle, AddParams, NewShapes);
 }
 
@@ -1110,6 +1212,46 @@ bool UBodySetup::CalcUVAtLocation(const FVector& BodySpaceLocation, int32 FaceIn
 	return bSuccess;
 }
 
+template<typename DDCBuilderType>
+void GetDDCBuiltData(FByteBulkData* OutResult, DDCBuilderType& InBuilder, UBodySetup* InSetup, bool bInIsRuntime)
+{
+	TArray<uint8> OutData;
+
+	if(InBuilder.CanBuild())
+	{
+		COOK_STAT(FCookStats::FScopedStatsCounter Timer = PhysXBodySetupCookStats::UsageStats.TimeSyncWork());
+
+		// Debugging switch, force builder to always run
+		bool bSkipDDC = false;
+
+		bool bDataWasBuilt = false;
+		bool bDDCHit = false;
+
+		if(!bSkipDDC)
+		{
+			bDDCHit = GetDerivedDataCacheRef().GetSynchronous(&InBuilder, OutData, &bDataWasBuilt);
+		}
+		else
+		{
+			bDataWasBuilt = true;
+			InBuilder.Build(OutData);
+		}
+
+		COOK_STAT(Timer.AddHitOrMiss(!bDDCHit || bDataWasBuilt ? FCookStats::CallStats::EHitOrMiss::Miss : FCookStats::CallStats::EHitOrMiss::Hit, OutData.Num()));
+	}
+
+	if(OutData.Num())
+	{
+		OutResult->Lock(LOCK_READ_WRITE);
+		FMemory::Memcpy(OutResult->Realloc(OutData.Num()), OutData.GetData(), OutData.Num());
+		OutResult->Unlock();
+	}
+	else if(!bInIsRuntime)	//only want to warn if DDC cooking failed - if it's really trying to use runtime and we can't, the runtime cooker code will catch it
+	{
+		UE_LOG(LogPhysics, Warning, TEXT("Attempt to build physics data for %s when we are unable to."), *InSetup->GetPathName());
+	}
+}
+
 FByteBulkData* UBodySetup::GetCookedData(FName Format, bool bRuntimeOnlyOptimizedVersion)
 {
 	if (IsTemplate())
@@ -1137,7 +1279,7 @@ FByteBulkData* UBodySetup::GetCookedData(FName Format, bool bRuntimeOnlyOptimize
 	FByteBulkData* Result = &UseCookedData->GetFormat(Format);
 	bool bIsRuntime = IsRuntime(this);
 
-#if WITH_PHYSX && WITH_EDITOR
+#if /*WITH_PHYSX &&*/ WITH_EDITOR
 	if (!bContainedData)
 	{
 		SCOPE_CYCLE_COUNTER(STAT_PhysXCooking);
@@ -1147,37 +1289,48 @@ FByteBulkData* UBodySetup::GetCookedData(FName Format, bool bRuntimeOnlyOptimize
 			return nullptr;
 		}
 
+#if PHYSICS_INTERFACE_PHYSX
 		const bool bEligibleForRuntimeOptimization = UseCookedData == &CookedFormatDataRuntimeOnlyOptimization;
-
 		const EPhysXMeshCookFlags CookingFlags = bEligibleForRuntimeOptimization ? GetRuntimeOnlyCookOptimizationFlags() : EPhysXMeshCookFlags::Default;
-
-		TArray<uint8> OutData;
-		FDerivedDataPhysXCooker* DerivedPhysXData = new FDerivedDataPhysXCooker(Format, CookingFlags, this, bIsRuntime);
+		FDerivedDataPhysXCooker* PhysicsDerivedCooker = new FDerivedDataPhysXCooker(Format, CookingFlags, this, bIsRuntime);
+#elif WITH_CHAOS 
+		FChaosDerivedDataCooker* PhysicsDerivedCooker = new FChaosDerivedDataCooker(this, Format);
+#else
+		static_assert(false, "No cooker defined for this physics interface");
+#endif
 			
-		if (DerivedPhysXData->CanBuild())
-		{
-			COOK_STAT(auto Timer = PhysXBodySetupCookStats::UsageStats.TimeSyncWork());
-			bool bDataWasBuilt = false;
-			bool DDCHit = GetDerivedDataCacheRef().GetSynchronous(DerivedPhysXData, OutData, &bDataWasBuilt);
-			COOK_STAT(Timer.AddHitOrMiss(!DDCHit || bDataWasBuilt ? FCookStats::CallStats::EHitOrMiss::Miss : FCookStats::CallStats::EHitOrMiss::Hit, OutData.Num()));
-		}
-
-		if (OutData.Num())
-		{
-			Result->Lock(LOCK_READ_WRITE);
-			FMemory::Memcpy(Result->Realloc(OutData.Num()), OutData.GetData(), OutData.Num());
-			Result->Unlock();
-		}
-		else if(!bIsRuntime)	//only want to warn if DDC cooking failed - if it's really trying to use runtime and we can't, the runtime cooker code will catch it
-		{
-			UE_LOG(LogPhysics, Warning, TEXT("Attempt to build physics data for %s when we are unable to."), *GetPathName());
-		}
+		GetDDCBuiltData(Result, *PhysicsDerivedCooker, this, bIsRuntime);
 	}
 #endif // WITH_PHYSX && WITH_EDITOR
+
 	check(Result);
 	return Result->GetBulkDataSize() > 0 ? Result : NULL; // we don't return empty bulk data...but we save it to avoid thrashing the DDC
 }
 
+void UBodySetup::GetGeometryDDCKey(FString& OutString) const
+{
+	// If the geometry changes in a way not controlled below, increment the key version to invalidate DDC entries
+	enum { BODY_SETUP_GEOMETRY_KEY_VER = 0 };
+
+	FString MeshIdString(TEXT("CDP_BODYSETUP"));
+
+	IInterface_CollisionDataProvider* CDP = Cast<IInterface_CollisionDataProvider>(GetOuter());
+	if(CDP)
+	{
+		CDP->GetMeshId(MeshIdString);
+	}
+
+	OutString = FString::Printf(TEXT("%s_%s_%s_%d_%d_%d_%d_%d"),
+		*BodySetupGuid.ToString(),
+		*MeshIdString,
+		*AggGeom.MakeDDCKey().ToString(),
+		(int32)bGenerateNonMirroredCollision,
+		(int32)bGenerateMirroredCollision,
+		(int32)UPhysicsSettings::Get()->bSupportUVFromHitResults,
+		(int32)GetCollisionTraceFlag(),
+		(int32)BODY_SETUP_GEOMETRY_KEY_VER);
+}
+
 void UBodySetup::PostInitProperties()
 {
 	Super::PostInitProperties();
@@ -1330,6 +1483,26 @@ float FKAggregateGeom::GetVolume(const FVector& Scale) const
 	return Volume;
 }
 
+FGuid FKAggregateGeom::MakeDDCKey() const
+{
+	UScriptStruct* StructType = FKAggregateGeom::StaticStruct();
+	TArray<uint8> Bytes;
+	Bytes.Reserve(64);
+	FMemoryWriter MemAr(Bytes);
+
+	StructType->SerializeTaggedProperties(MemAr, (uint8*)this, StructType, nullptr);
+
+	FSHA1 Sha;
+	Sha.Update(Bytes.GetData(), Bytes.Num());
+	Sha.Final();
+
+	uint32 Hash[5];
+	Sha.GetHash((uint8*)Hash);
+	FGuid OutGuid(Hash[0] ^ Hash[4], Hash[1], Hash[2], Hash[3]);
+
+	return OutGuid;
+}
+
 int32 FKAggregateGeom::GetElementCount(EAggCollisionShape::Type Type) const
 {
 	switch (Type)
@@ -1349,6 +1522,52 @@ int32 FKAggregateGeom::GetElementCount(EAggCollisionShape::Type Type) const
 	}
 }
 
+
+FKConvexElem::FKConvexElem()
+	: FKShapeElem(EAggCollisionShape::Convex)
+	, ElemBox(ForceInit)
+	, Transform(FTransform::Identity)
+	, ConvexMesh(NULL)
+	, ConvexMeshNegX(NULL)
+{}
+
+FKConvexElem::FKConvexElem(const FKConvexElem& Other)
+	: ConvexMesh(nullptr)
+	, ConvexMeshNegX(nullptr)
+{
+	CloneElem(Other);
+}
+
+FKConvexElem::~FKConvexElem()
+{
+
+}
+
+const FKConvexElem& FKConvexElem::operator=(const FKConvexElem& Other)
+{
+	ensureMsgf(ConvexMesh == nullptr, TEXT("We are leaking memory. Why are we calling the assignment operator on an element that has already allocated resources?"));
+	ensureMsgf(ConvexMeshNegX == nullptr, TEXT("We are leaking memory. Why are we calling the assignment operator on an element that has already allocated resources?"));
+	ConvexMesh = nullptr;
+	ConvexMeshNegX = nullptr;
+#if WITH_CHAOS
+	ensureMsgf(!ChaosConvex, TEXT("We are leaking memory. Why are we calling the assignment operator on an element that has already allocated resources?"));
+	ResetChaosConvexMesh();
+#endif
+	CloneElem(Other);
+	return *this;
+}
+
+/** Helper function to safely copy instances of this shape*/
+void FKConvexElem::CloneElem(const FKConvexElem& Other)
+{
+	Super::CloneElem(Other);
+	VertexData = Other.VertexData;
+	ElemBox = Other.ElemBox;
+	Transform = Other.Transform;
+
+	// TODO: Should this also copy the ChaosConvexMesh?
+}
+
 void FKConvexElem::ScaleElem(FVector DeltaSize, float MinSize)
 {
 	FTransform ScaledTransform = GetTransform();
@@ -1424,6 +1643,23 @@ float FKConvexElem::GetVolume(const FVector& Scale) const
 	return Volume;
 }
 
+#if WITH_CHAOS
+const TUniquePtr<Chaos::TImplicitObject<float, 3>>& FKConvexElem::GetChaosConvexMesh() const
+{
+	return ChaosConvex;
+}
+
+void FKConvexElem::SetChaosConvexMesh(TUniquePtr<Chaos::TImplicitObject<float, 3>>&& InChaosConvex)
+{
+	ChaosConvex = MoveTemp(InChaosConvex);
+}
+
+void FKConvexElem::ResetChaosConvexMesh()
+{
+	ChaosConvex.Reset();
+}
+#endif
+
 #if WITH_EDITORONLY_DATA
 void FKSphereElem::FixupDeprecated( FArchive& Ar )
 {
diff --git a/Engine/Source/Runtime/Engine/Private/PhysicsEngine/CollisionAnalyzerCapture.h b/Engine/Source/Runtime/Engine/Private/PhysicsEngine/CollisionAnalyzerCapture.h
--- a/Engine/Source/Runtime/Engine/Private/PhysicsEngine/CollisionAnalyzerCapture.h
+++ b/Engine/Source/Runtime/Engine/Private/PhysicsEngine/CollisionAnalyzerCapture.h
@@ -13,41 +13,47 @@ bool bSkipCapture = false;
 bool GCollisionAnalyzerIsRecording = false;
 
 /** Util to convert from PhysX shape and rotation to unreal shape enum, dimension vector and rotation */
-static void ConvertGeometryCollection(const FPhysicsGeometryCollection& GeomCollection, FQuat& InOutRot, ECAQueryShape::Type& OutQueryShape, FVector& OutDims)
+inline static void ConvertGeometryCollection(const FPhysicsGeometryCollection& GeomCollection, FQuat& InOutRot, ECAQueryShape::Type& OutQueryShape, FVector& OutDims)
 {
 	OutQueryShape = ECAQueryShape::Capsule;
 	OutDims = FVector(0, 0, 0);
 	
-#if WITH_CHAOS || PHYSICS_INTERFACE_LLIMMEDIATE
-    ensure(false);
-#else
 	switch (GeomCollection.GetType())
 	{
 		case ECollisionShapeType::Capsule:
 		{
 			OutQueryShape = ECAQueryShape::Capsule;
-			PxCapsuleGeometry PCapsuleGeom;
-			GeomCollection.GetCapsuleGeometry(PCapsuleGeom);
-			OutDims = FVector(PCapsuleGeom.radius, PCapsuleGeom.radius, PCapsuleGeom.halfHeight + PCapsuleGeom.radius);
+			const auto& CapsuleGeom = GeomCollection.GetCapsuleGeometry();
+#if PHYSICS_INTERFACE_PHYSX
+			OutDims = FVector(CapsuleGeom.radius, CapsuleGeom.radius, CapsuleGeom.halfHeight + CapsuleGeom.radius);
 			InOutRot = ConvertToUECapsuleRot(InOutRot);
+#else
+			OutDims = FVector(CapsuleGeom.GetRadius(), CapsuleGeom.GetRadius(), CapsuleGeom.GetHeight() * 0.5f + CapsuleGeom.GetRadius());
+#endif
 			break;
 		}
 
 		case ECollisionShapeType::Sphere:
 		{
 			OutQueryShape = ECAQueryShape::Sphere;
-			PxSphereGeometry PSphereGeom;
-			GeomCollection.GetSphereGeometry(PSphereGeom);
-			OutDims = FVector(PSphereGeom.radius);
+			const auto& SphereGeom = GeomCollection.GetSphereGeometry();
+#if PHYSICS_INTERFACE_PHYSX
+			OutDims = FVector(SphereGeom.radius);
+#else
+			OutDims = FVector(SphereGeom.GetRadius());
+#endif
 			break;
 		}
 
 		case ECollisionShapeType::Box:
 		{
 			OutQueryShape = ECAQueryShape::Box;
-			PxBoxGeometry PBoxGeom;
-			GeomCollection.GetBoxGeometry(PBoxGeom);
-			OutDims = P2UVector(PBoxGeom.halfExtents);
+			const auto& BoxGeom = GeomCollection.GetBoxGeometry();
+#if PHYSICS_INTERFACE_PHYSX
+			OutDims = P2UVector(BoxGeom.halfExtents);
+#else
+			OutDims = FVector(BoxGeom.Extents() * 0.5f);
+#endif
 			break;
 		}
 
@@ -60,10 +66,9 @@ static void ConvertGeometryCollection(const FPhysicsGeometryCollection& GeomColl
 		default:
 			UE_LOG(LogCollision, Warning, TEXT("CaptureGeomSweep: Unknown geom type."));
 	}
-#endif
 }
 
-bool CollisionShapeToAnalyzerType(const FCollisionShape& InShape, ECAQueryShape::Type& OutType, FVector& OutDims)
+inline bool CollisionShapeToAnalyzerType(const FCollisionShape& InShape, ECAQueryShape::Type& OutType, FVector& OutDims)
 {
 	switch(InShape.ShapeType)
 	{
@@ -95,7 +100,7 @@ bool CollisionShapeToAnalyzerType(const FCollisionShape& InShape, ECAQueryShape:
 }
 
 /** Util to extract type and dimensions from physx geom being swept, and pass info to CollisionAnalyzer, if its recording */
-void CaptureGeomSweep(const UWorld* World, const FVector& Start, const FVector& End, const FQuat& Rot, ECAQueryMode::Type QueryMode, const FCollisionShape& PGeom, ECollisionChannel TraceChannel, const struct FCollisionQueryParams& Params, const struct FCollisionResponseParams& ResponseParams, const struct FCollisionObjectQueryParams& ObjectParams, const TArray<FHitResult>& Results, double CPUTime)
+inline void CaptureGeomSweep(const UWorld* World, const FVector& Start, const FVector& End, const FQuat& Rot, ECAQueryMode::Type QueryMode, const FCollisionShape& PGeom, ECollisionChannel TraceChannel, const struct FCollisionQueryParams& Params, const struct FCollisionResponseParams& ResponseParams, const struct FCollisionObjectQueryParams& ObjectParams, const TArray<FHitResult>& Results, double CPUTime)
 {
 	if (bSkipCapture || !GCollisionAnalyzerIsRecording || !IsInGameThread())
 	{
@@ -119,7 +124,7 @@ void CaptureGeomSweep(const UWorld* World, const FVector& Start, const FVector&
 }
 
 /** Util to extract type and dimensions from physx geom being swept, and pass info to CollisionAnalyzer, if its recording */
-void CaptureGeomSweep(const UWorld* World, const FVector& Start, const FVector& End, const FQuat& GeomRot, ECAQueryMode::Type QueryMode, const FPhysicsGeometryCollection& GeomCollection, ECollisionChannel TraceChannel, const struct FCollisionQueryParams& Params, const struct FCollisionResponseParams& ResponseParams, const struct FCollisionObjectQueryParams& ObjectParams, const TArray<FHitResult>& Results, double CPUTime)
+inline void CaptureGeomSweep(const UWorld* World, const FVector& Start, const FVector& End, const FQuat& GeomRot, ECAQueryMode::Type QueryMode, const FPhysicsGeometryCollection& GeomCollection, ECollisionChannel TraceChannel, const struct FCollisionQueryParams& Params, const struct FCollisionResponseParams& ResponseParams, const struct FCollisionObjectQueryParams& ObjectParams, const TArray<FHitResult>& Results, double CPUTime)
 {
 	if(bSkipCapture || !GCollisionAnalyzerIsRecording || !IsInGameThread())
 	{
@@ -144,7 +149,7 @@ void CaptureGeomSweep(const UWorld* World, const FVector& Start, const FVector&
 }
 
 /** Util to capture a raycast with the CollisionAnalyzer if recording */
-void CaptureRaycast(const UWorld* World, const FVector& Start, const FVector& End, ECAQueryMode::Type QueryMode, ECollisionChannel TraceChannel, const struct FCollisionQueryParams& Params, const struct FCollisionResponseParams& ResponseParams, const struct FCollisionObjectQueryParams& ObjectParams, const TArray<FHitResult>& Results, double CPUTime)
+inline void CaptureRaycast(const UWorld* World, const FVector& Start, const FVector& End, ECAQueryMode::Type QueryMode, ECollisionChannel TraceChannel, const struct FCollisionQueryParams& Params, const struct FCollisionResponseParams& ResponseParams, const struct FCollisionObjectQueryParams& ObjectParams, const TArray<FHitResult>& Results, double CPUTime)
 {
 	if (bSkipCapture || !GCollisionAnalyzerIsRecording || !IsInGameThread())
 	{
@@ -160,7 +165,7 @@ void CaptureRaycast(const UWorld* World, const FVector& Start, const FVector& En
 	FCollisionAnalyzerModule::Get()->CaptureQuery(Start, End, FQuat::Identity, ECAQueryType::Raycast, ECAQueryShape::Sphere, QueryMode, FVector(0, 0, 0), TraceChannel, Params, ResponseParams, ObjectParams, Results, TouchAllResults, CPUTime);
 }
 
-void CaptureOverlap(const UWorld* World, const FPhysicsGeometryCollection& PGeom, const FTransform& InGeomTransform, ECAQueryMode::Type QueryMode, ECollisionChannel TraceChannel, const struct FCollisionQueryParams& Params, const struct FCollisionResponseParams& ResponseParams, const struct FCollisionObjectQueryParams& ObjectParams, TArray<FOverlapResult>& Results, double CPUTime)
+inline void CaptureOverlap(const UWorld* World, const FPhysicsGeometryCollection& PGeom, const FTransform& InGeomTransform, ECAQueryMode::Type QueryMode, ECollisionChannel TraceChannel, const struct FCollisionQueryParams& Params, const struct FCollisionResponseParams& ResponseParams, const struct FCollisionObjectQueryParams& ObjectParams, TArray<FOverlapResult>& Results, double CPUTime)
 {
 	if (bSkipCapture || !GCollisionAnalyzerIsRecording || !IsInGameThread())
 	{
@@ -189,7 +194,7 @@ void CaptureOverlap(const UWorld* World, const FPhysicsGeometryCollection& PGeom
 	FCollisionAnalyzerModule::Get()->CaptureQuery(InGeomTransform.GetTranslation(), FVector(0, 0, 0), UseRot, ECAQueryType::GeomOverlap, QueryShape, QueryMode, Dims, TraceChannel, Params, ResponseParams, ObjectParams, HitResults, TouchAllResults, CPUTime);
 }
 
-void CaptureOverlap(const UWorld* World, const FCollisionShape& PGeom, const FTransform& InGeomTransform, ECAQueryMode::Type QueryMode, ECollisionChannel TraceChannel, const struct FCollisionQueryParams& Params, const struct FCollisionResponseParams& ResponseParams, const struct FCollisionObjectQueryParams& ObjectParams, TArray<FOverlapResult>& Results, double CPUTime)
+inline void CaptureOverlap(const UWorld* World, const FCollisionShape& PGeom, const FTransform& InGeomTransform, ECAQueryMode::Type QueryMode, ECollisionChannel TraceChannel, const struct FCollisionQueryParams& Params, const struct FCollisionResponseParams& ResponseParams, const struct FCollisionObjectQueryParams& ObjectParams, TArray<FOverlapResult>& Results, double CPUTime)
 {
 	if(bSkipCapture || !GCollisionAnalyzerIsRecording || !IsInGameThread())
 	{
