Copying //UE4/Dev-Physics to Dev-Main (//UE4/Dev-Main) [at] 7703071

diff --git a/Engine/Source/Runtime/Engine/Private/PhysicsEngine/CollisionQueryFilterCallback.cpp b/Engine/Source/Runtime/Engine/Private/PhysicsEngine/CollisionQueryFilterCallback.cpp
--- a/Engine/Source/Runtime/Engine/Private/PhysicsEngine/CollisionQueryFilterCallback.cpp
+++ b/Engine/Source/Runtime/Engine/Private/PhysicsEngine/CollisionQueryFilterCallback.cpp
@@ -8,10 +8,13 @@
 #include "Components/PrimitiveComponent.h"
 #include "Collision.h"
 
-#if PHYSICS_INTERFACE_PHYSX
+#if WITH_PHYSX
 #include "PhysXInterfaceWrapper.h"
-#elif PHYSICS_INTERFACE_LLIMMEDIATE
-#include "Physics/Experimental/LLImmediateInterfaceWrapper.h"
+#endif
+
+#if INCLUDE_CHAOS
+#include "ChaosInterfaceWrapperCore.h"
+#include "Physics/Experimental/ChaosInterfaceWrapper.h"
 #endif
 
 ECollisionQueryHitType FCollisionQueryFilterCallback::CalcQueryHitType(const FCollisionFilterData& QueryFilter, const FCollisionFilterData& ShapeFilter, bool bPreFilter)
@@ -85,25 +88,49 @@ ECollisionQueryHitType FCollisionQueryFilterCallback::CalcQueryHitType(const FCo
 	return ECollisionQueryHitType::None;
 }
 
-ECollisionQueryHitType FCollisionQueryFilterCallback::PreFilterImp(const FCollisionFilterData& FilterData, const FPhysicsShape& Shape, const FPhysicsActor& Actor)
+#if WITH_PHYSX
+ECollisionQueryHitType FCollisionQueryFilterCallback::PreFilterImp(const FCollisionFilterData& FilterData, const physx::PxShape& Shape, const physx::PxActor& Actor)
 {
 	SCOPE_CYCLE_COUNTER(STAT_Collision_PreFilter);
-	FCollisionFilterData ShapeFilter = GetQueryFilterData(Shape);
+	FCollisionFilterData ShapeFilter = PhysXInterface::GetQueryFilterData(Shape);
 
 	// We usually don't have ignore components so we try to avoid the virtual getSimulationFilterData() call below. 'word2' of shape sim filter data is componentID.
 	uint32 ComponentID = 0;
 	if (IgnoreComponents.Num() > 0)
 	{
-		ComponentID = GetSimulationFilterData(Shape).Word2;
+		ComponentID = PhysXInterface::GetSimulationFilterData(Shape).Word2;
 	}
 
 	FBodyInstance* BodyInstance = nullptr;
 #if ENABLE_PREFILTER_LOGGING || DETECT_SQ_HITCHES
-	BodyInstance = GetUserData(Actor);
+	BodyInstance = PhysXInterface::GetUserData(Actor);
 #endif // ENABLE_PREFILTER_LOGGING || DETECT_SQ_HITCHES
 
 	return PreFilterImp(FilterData, ShapeFilter, ComponentID, BodyInstance);
 }
+#endif
+
+#if INCLUDE_CHAOS
+ECollisionQueryHitType FCollisionQueryFilterCallback::PreFilterImp(const FCollisionFilterData& FilterData, const Chaos::TPerShapeData<float,3>& Shape, const Chaos::TGeometryParticle<float,3>& Actor)
+{
+	SCOPE_CYCLE_COUNTER(STAT_Collision_PreFilter);
+	FCollisionFilterData ShapeFilter = ChaosInterface::GetQueryFilterData(Shape);
+
+	// We usually don't have ignore components so we try to avoid the virtual getSimulationFilterData() call below. 'word2' of shape sim filter data is componentID.
+	uint32 ComponentID = 0;
+	if (IgnoreComponents.Num() > 0)
+	{
+		ComponentID = ChaosInterface::GetSimulationFilterData(Shape).Word2;
+	}
+
+	FBodyInstance* BodyInstance = nullptr;
+#if ENABLE_PREFILTER_LOGGING || DETECT_SQ_HITCHES
+	BodyInstance = ChaosInterface::GetUserData(Actor);
+#endif // ENABLE_PREFILTER_LOGGING || DETECT_SQ_HITCHES
+
+	return PreFilterImp(FilterData, ShapeFilter, ComponentID, BodyInstance);
+}
+#endif
 
 ECollisionQueryHitType FCollisionQueryFilterCallback::PreFilterImp(const FCollisionFilterData& FilterData, const FCollisionFilterData& ShapeFilter, uint32 ComponentID, const FBodyInstance* BodyInstance)
 {
@@ -228,7 +255,8 @@ ECollisionQueryHitType FCollisionQueryFilterCallback::PostFilterImp(const FColli
 	}
 }
 
-ECollisionQueryHitType FCollisionQueryFilterCallback::PostFilterImp(const FCollisionFilterData& FilterData, const FPhysicsQueryHit& Hit)
+#if WITH_PHYSX
+ECollisionQueryHitType FCollisionQueryFilterCallback::PostFilterImp(const FCollisionFilterData& FilterData, const physx::PxQueryHit& Hit)
 {
 	// Unused in non-sweeps
 	if (!bIsSweep)
@@ -241,6 +269,24 @@ ECollisionQueryHitType FCollisionQueryFilterCallback::PostFilterImp(const FColli
 
 	return PostFilterImp(FilterData, bIsOverlap);
 }
+#endif
+
+#if INCLUDE_CHAOS
+ECollisionQueryHitType FCollisionQueryFilterCallback::PostFilterImp(const FCollisionFilterData& FilterData, const ChaosInterface::FQueryHit& Hit)
+{
+	// Unused in non-sweeps
+	if (!bIsSweep)
+	{
+		return ECollisionQueryHitType::None;
+	}
+
+	const ChaosInterface::FLocationHit& SweepHit = (const ChaosInterface::FLocationHit&)Hit;
+	const bool bIsOverlap = ChaosInterface::HadInitialOverlap(SweepHit);
+
+	return PostFilterImp(FilterData, bIsOverlap);
+}
+#endif
+
 
 #if WITH_PHYSX
 
diff --git a/Engine/Source/Runtime/Engine/Private/PhysicsEngine/CollisionQueryFilterCallback.h b/Engine/Source/Runtime/Engine/Private/PhysicsEngine/CollisionQueryFilterCallback.h
--- a/Engine/Source/Runtime/Engine/Private/PhysicsEngine/CollisionQueryFilterCallback.h
+++ b/Engine/Source/Runtime/Engine/Private/PhysicsEngine/CollisionQueryFilterCallback.h
@@ -64,24 +64,30 @@ public:
 
 	ECollisionQueryHitType PostFilterImp(const FCollisionFilterData& FilterData, bool bIsOverlap);
 
-	ECollisionQueryHitType PostFilterImp(const FCollisionFilterData& FilterData, const FPhysicsQueryHit& Hit);
-	ECollisionQueryHitType PreFilterImp(const FCollisionFilterData& FilterData, const FPhysicsShape& Shape, const FPhysicsActor& Actor);
+#if WITH_PHYSX
+	ECollisionQueryHitType PostFilterImp(const FCollisionFilterData& FilterData, const physx::PxQueryHit& Hit);
+	ECollisionQueryHitType PreFilterImp(const FCollisionFilterData& FilterData, const physx::PxShape& Shape, const physx::PxActor& Actor);
+#endif
 
 #if INCLUDE_CHAOS
-	virtual ECollisionQueryHitType PreFilterChaos(const FCollisionFilterData& FilterData, const Chaos::TImplicitObject<float, 3>& Shape, int32 ActorIdx, int32 ShapeIdx) { check(false);  return ECollisionQueryHitType::Block; }
+	ECollisionQueryHitType PreFilterImp(const FCollisionFilterData& FilterData, const Chaos::TPerShapeData<float,3>& Shape, const Chaos::TGeometryParticle<float,3>& Actor);
+	ECollisionQueryHitType PostFilterImp(const FCollisionFilterData& FilterData, const ChaosInterface::FQueryHit& Hit);
 #endif
 
-	virtual ECollisionQueryHitType PostFilter(const FCollisionFilterData& FilterData, const FPhysicsQueryHit& Hit) override
+#if INCLUDE_CHAOS
+	virtual ECollisionQueryHitType PostFilter(const FCollisionFilterData& FilterData, const ChaosInterface::FQueryHit& Hit) override
 	{
 		return PostFilterImp(FilterData, Hit);
 	}
-	
-	virtual ECollisionQueryHitType PreFilter(const FCollisionFilterData& FilterData, const FPhysicsShape& Shape, const FPhysicsActor& Actor) override
+	virtual ECollisionQueryHitType PreFilter(const FCollisionFilterData& FilterData, const Chaos::TPerShapeData<float,3>& Shape, const Chaos::TGeometryParticle<float,3>& Actor) override
 	{
 		return PreFilterImp(FilterData, Shape, Actor);
 	}
+#endif
 
 #if WITH_PHYSX
+	virtual ECollisionQueryHitType PostFilter(const FCollisionFilterData& FilterData, const physx::PxQueryHit& Hit) override { return PostFilterImp(FilterData, Hit); }
+	virtual ECollisionQueryHitType PreFilter(const FCollisionFilterData& FilterData, const physx::PxShape& Shape, physx::PxRigidActor& Actor) override { return PreFilterImp(FilterData, Shape, Actor); }
 	virtual PxQueryHitType::Enum preFilter(const PxFilterData& filterData, const PxShape* shape, const PxRigidActor* actor, PxHitFlags& queryFlags) override;
 	virtual PxQueryHitType::Enum postFilter(const PxFilterData& filterData, const PxQueryHit& hit) override;
 #endif
diff --git a/Engine/Source/Runtime/Engine/Private/PhysicsEngine/ConstraintInstance.cpp b/Engine/Source/Runtime/Engine/Private/PhysicsEngine/ConstraintInstance.cpp
--- a/Engine/Source/Runtime/Engine/Private/PhysicsEngine/ConstraintInstance.cpp
+++ b/Engine/Source/Runtime/Engine/Private/PhysicsEngine/ConstraintInstance.cpp
@@ -270,13 +270,13 @@ float ComputeAverageMass_AssumesLocked(const FPhysicsActorHandle& InActor1, cons
 	float TotalMass = 0;
 	int NumDynamic = 0;
 
-	if (InActor1.IsValid() && FPhysicsInterface::IsRigidBody(InActor1))
+	if (FPhysicsInterface::IsValid(InActor1) && FPhysicsInterface::IsRigidBody(InActor1))
 	{
 		TotalMass += FPhysicsInterface::GetMass_AssumesLocked(InActor1);
 		++NumDynamic;
 	}
 
-	if(InActor2.IsValid() && FPhysicsInterface::IsRigidBody(InActor2))
+	if(FPhysicsInterface::IsValid(InActor2) && FPhysicsInterface::IsRigidBody(InActor2))
 	{
 		TotalMass += FPhysicsInterface::GetMass_AssumesLocked(InActor2);
 		++NumDynamic;
@@ -299,7 +299,7 @@ bool GetActorRefs(FBodyInstance* Body1, FBodyInstance* Body2, FPhysicsActorHandl
 
 	// Do not create joint unless you have two actors
 	// Do not create joint unless one of the actors is dynamic
-	if((!ActorRef1.IsValid() || !FPhysicsInterface::IsRigidBody(ActorRef1)) && (!ActorRef2.IsValid() || !FPhysicsInterface::IsRigidBody(ActorRef2)))
+	if((!FPhysicsInterface::IsValid(ActorRef1) || !FPhysicsInterface::IsRigidBody(ActorRef1)) && (!FPhysicsInterface::IsValid(ActorRef2) || !FPhysicsInterface::IsRigidBody(ActorRef2)))
 	{
 #if !(UE_BUILD_SHIPPING || UE_BUILD_TEST)
 		FMessageLog("PIE").Warning()
@@ -310,7 +310,7 @@ bool GetActorRefs(FBodyInstance* Body1, FBodyInstance* Body2, FPhysicsActorHandl
 		return false;
 	}
 
-	if(ActorRef1.IsValid() && ActorRef2.IsValid() && ActorRef1.Equals(ActorRef2))
+	if(FPhysicsInterface::IsValid(ActorRef1) && FPhysicsInterface::IsValid(ActorRef2) && ActorRef1 == ActorRef2)
 	{
 #if !(UE_BUILD_SHIPPING || UE_BUILD_TEST)
 		const UPrimitiveComponent* PrimComp = Body1 ? Body1->OwnerComponent.Get() : nullptr;
@@ -328,8 +328,8 @@ bool GetActorRefs(FBodyInstance* Body1, FBodyInstance* Body2, FPhysicsActorHandl
 	bool bActor2ValidToSim = false;
 	FPhysicsCommand::ExecuteRead(ActorRef1, ActorRef2, [&](const FPhysicsActorHandle& ActorA, const FPhysicsActorHandle& ActorB)
 	{
-		bActor1ValidToSim = !ActorRef1.IsValid() || FPhysicsInterface::CanSimulate_AssumesLocked(ActorRef1);
-		bActor2ValidToSim = !ActorRef2.IsValid() || FPhysicsInterface::CanSimulate_AssumesLocked(ActorRef2);
+		bActor1ValidToSim = !FPhysicsInterface::IsValid(ActorRef1) || FPhysicsInterface::CanSimulate_AssumesLocked(ActorRef1);
+		bActor2ValidToSim = !FPhysicsInterface::IsValid(ActorRef2) || FPhysicsInterface::CanSimulate_AssumesLocked(ActorRef2);
 	});
 
 	if(!bActor1ValidToSim || !bActor2ValidToSim)
@@ -351,7 +351,7 @@ bool FConstraintInstance::CreateJoint_AssumesLocked(const FPhysicsActorHandle& I
 	LLM_SCOPE(ELLMTag::PhysX);
 
 	FTransform Local1 = GetRefFrame(EConstraintFrame::Frame1);
-	if(InActorRef1.IsValid())
+	if(FPhysicsInterface::IsValid(InActorRef1))
 	{
 		Local1.ScaleTranslation(FVector(LastKnownScale));
 	}
@@ -359,7 +359,7 @@ bool FConstraintInstance::CreateJoint_AssumesLocked(const FPhysicsActorHandle& I
 	checkf(Local1.IsValid() && !Local1.ContainsNaN(), TEXT("%s"), *Local1.ToString());
 
 	FTransform Local2 = GetRefFrame(EConstraintFrame::Frame2);
-	if(InActorRef2.IsValid())
+	if(FPhysicsInterface::IsValid(InActorRef2))
 	{
 		Local2.ScaleTranslation(FVector(LastKnownScale));
 	}
@@ -404,12 +404,12 @@ void EnsureSleepingActorsStaySleeping_AssumesLocked(const FPhysicsActorHandle& I
 	// creation of joints wakes up rigid bodies, so we put them to sleep again if both were initially asleep
 	if (bActor1Asleep && bActor2Asleep)
 	{
-		if(InActorRef1.IsValid() && !FPhysicsInterface::IsKinematic_AssumesLocked(InActorRef1))
+		if(FPhysicsInterface::IsValid(InActorRef1) && !FPhysicsInterface::IsKinematic_AssumesLocked(InActorRef1))
 		{
 			FPhysicsInterface::PutToSleep_AssumesLocked(InActorRef1);
 		}
 
-		if(InActorRef2.IsValid() && !FPhysicsInterface::IsKinematic_AssumesLocked(InActorRef2))
+		if(FPhysicsInterface::IsValid(InActorRef2) && !FPhysicsInterface::IsKinematic_AssumesLocked(InActorRef2))
 		{
 			FPhysicsInterface::PutToSleep_AssumesLocked(InActorRef2);
 		}
diff --git a/Engine/Source/Runtime/Engine/Private/PhysicsEngine/Experimental/ChaosDerivedData.cpp b/Engine/Source/Runtime/Engine/Private/PhysicsEngine/Experimental/ChaosDerivedData.cpp
new file mode 100644
--- /dev/null
+++ b/Engine/Source/Runtime/Engine/Private/PhysicsEngine/Experimental/ChaosDerivedData.cpp
@@ -0,0 +1,236 @@
+// Copyright 1998-2019 Epic Games, Inc. All Rights Reserved.
+
+#include "ChaosDerivedData.h"
+
+#if WITH_CHAOS
+
+#include "Chaos/CollisionConvexMesh.h"
+#include "Chaos/ChaosArchive.h"
+#include "ChaosDerivedDataUtil.h"
+#include "Chaos/Particles.h"
+#include "Chaos/TriangleMeshImplicitObject.h"
+#include "Chaos/Convex.h"
+#include "Serialization/Archive.h"
+#include "PhysicsEngine/BodySetup.h"
+#include "Serialization/MemoryWriter.h"
+
+const TCHAR* FChaosDerivedDataCooker::GetPluginName() const
+{
+	return TEXT("ChaosGeometryData");
+}
+
+const TCHAR* FChaosDerivedDataCooker::GetVersionString() const
+{
+	return TEXT("19C69FC43DDA4F058B28C21F08D623F2");
+}
+
+FString FChaosDerivedDataCooker::GetPluginSpecificCacheKeySuffix() const
+{
+	FString SetupGeometryKey(TEXT("INVALID"));
+
+	if(Setup)
+	{
+		Setup->GetGeometryDDCKey(SetupGeometryKey);
+	}
+
+	return FString::Printf(TEXT("%s_%s"),
+		*RequestedFormat.ToString(),
+		*SetupGeometryKey);
+}
+
+bool FChaosDerivedDataCooker::IsBuildThreadsafe() const
+{
+	// #BG Investigate Parallel Build
+	return false;
+}
+
+bool FChaosDerivedDataCooker::Build(TArray<uint8>& OutData)
+{
+	bool bSucceeded = false;
+
+	if(Setup)
+	{
+		FCookBodySetupInfo CookInfo;
+
+		EPhysXMeshCookFlags TempFlags = static_cast<EPhysXMeshCookFlags>(0); // #BGTODO Remove need for PhysX specific flags
+		Setup->GetCookInfo(CookInfo, TempFlags);
+
+		FMemoryWriter MemWriterAr(OutData);
+		Chaos::FChaosArchive Ar(MemWriterAr);
+
+		int32 PrecisionSize = (int32)sizeof(BuildPrecision);
+
+		Ar << PrecisionSize;
+		BuildInternal<BuildPrecision>(Ar, CookInfo);
+
+		bSucceeded = true;
+	}
+
+	return bSucceeded;
+}
+
+void FChaosDerivedDataCooker::AddReferencedObjects(FReferenceCollector& Collector)
+{
+	if(Setup)
+	{
+		Collector.AddReferencedObject(Setup);
+	}
+}
+
+FChaosDerivedDataCooker::FChaosDerivedDataCooker(UBodySetup* InSetup, FName InFormat)
+	: Setup(InSetup)
+	, RequestedFormat(InFormat)
+{
+
+}
+
+template<typename Precision>
+void FChaosDerivedDataCooker::BuildTriangleMeshes(TArray<TUniquePtr<Chaos::TTriangleMeshImplicitObject<Precision>>>& OutTriangleMeshes, const FCookBodySetupInfo& InParams)
+{
+	if(!InParams.bCookTriMesh)
+	{
+		return;
+	}
+
+	TArray<FVector> FinalVerts = InParams.TriangleMeshDesc.Vertices;
+
+	// Push indices into one flat array
+	TArray<int32> FinalIndices;
+	FinalIndices.Reserve(InParams.TriangleMeshDesc.Indices.Num() * 3);
+	for(const FTriIndices& Tri : InParams.TriangleMeshDesc.Indices)
+	{
+		FinalIndices.Add(Tri.v0);
+		FinalIndices.Add(Tri.v1);
+		FinalIndices.Add(Tri.v2);
+	}
+
+	Chaos::CleanTrimesh(FinalVerts, FinalIndices, nullptr);
+
+	// Build particle list #BG Maybe allow TParticles to copy vectors?
+	Chaos::TParticles<Precision, 3> TriMeshParticles;
+	TriMeshParticles.AddParticles(FinalVerts.Num());
+
+	const int32 NumVerts = FinalVerts.Num();
+	for(int32 VertIndex = 0; VertIndex < NumVerts; ++VertIndex)
+	{
+		TriMeshParticles.X(VertIndex) = FinalVerts[VertIndex];
+	}
+
+	// Build chaos triangle list. #BGTODO Just make the clean function take these types instead of double copying
+	const int32 NumTriangles = FinalIndices.Num() / 3;
+	TArray<Chaos::TVector<int32, 3>> Triangles;
+	Triangles.Reserve(NumTriangles);
+
+	for(int32 TriangleIndex = 0; TriangleIndex < NumTriangles; ++TriangleIndex)
+	{
+		const int32 BaseIndex = TriangleIndex * 3;
+		Triangles.Add(Chaos::TVector<int32, 3>(FinalIndices[BaseIndex], FinalIndices[BaseIndex + 1], FinalIndices[BaseIndex + 2]));
+	}
+
+	OutTriangleMeshes.Emplace(new Chaos::TTriangleMeshImplicitObject<Precision>(MoveTemp(TriMeshParticles), MoveTemp(Triangles)));
+}
+
+template void FChaosDerivedDataCooker::BuildTriangleMeshes(TArray<TUniquePtr<Chaos::TTriangleMeshImplicitObject<float>>>& OutTriangleMeshes, const FCookBodySetupInfo& InParams);
+//#BGTODO When it's possible to build with doubles, re-enable this. (Currently at least TRigidTransform cannot build with double precision because we don't have a base transform implementation using them)
+//template void FChaosDerivedDataCooker::BuildTriangleMeshes(TArray<TUniquePtr<Chaos::TImplicitObject<float, 3>>>& OutTriangleMeshes, const FCookBodySetupInfo& InParams);
+
+template<typename Precision>
+void FChaosDerivedDataCooker::BuildConvexMeshes(TArray<TUniquePtr<Chaos::TImplicitObject<Precision, 3>>>& OutConvexMeshes, const FCookBodySetupInfo& InParams)
+{
+	auto BuildConvexFromVerts = [](TArray<TUniquePtr<Chaos::TImplicitObject<Precision, 3>>>& OutConvexes, const TArray<TArray<FVector>>& InMeshVerts)
+	{
+		for(const TArray<FVector>& HullVerts : InMeshVerts)
+		{
+			Chaos::TParticles<Precision, 3> ConvexParticles;
+
+			const int32 NumHullVerts = HullVerts.Num();
+
+			ConvexParticles.AddParticles(NumHullVerts);
+
+			for(int32 VertIndex = 0; VertIndex < NumHullVerts; ++VertIndex)
+			{
+				ConvexParticles.X(VertIndex) = HullVerts[VertIndex];
+			}
+
+			OutConvexes.Emplace(new Chaos::TConvex<Precision, 3>(ConvexParticles));
+		}
+	};
+
+	if(InParams.bCookNonMirroredConvex)
+	{
+		BuildConvexFromVerts(OutConvexMeshes, InParams.NonMirroredConvexVertices);
+	}
+
+	if(InParams.bCookMirroredConvex)
+	{
+		BuildConvexFromVerts(OutConvexMeshes, InParams.MirroredConvexVertices);
+	}
+
+}
+
+template void FChaosDerivedDataCooker::BuildConvexMeshes(TArray<TUniquePtr<Chaos::TImplicitObject<float, 3>>>& OutConvexMeshes, const FCookBodySetupInfo& InParams);
+//#BGTODO When it's possible to build with doubles, re-enable this. (Currently at least TRigidTransform cannot build with double precision because we don't have a base transform implementation using them)
+//template void FChaosDerivedDataCooker::BuildConvexMeshes(TArray<Chaos::TCollisionConvexMesh<double>>& OutTriangleMeshes, const FCookBodySetupInfo& InParams);
+
+template<typename Precision>
+void BuildSimpleShapes(TArray<TUniquePtr<Chaos::TImplicitObject<Precision, 3>>>& OutImplicits, UBodySetup* InSetup)
+{
+	check(InSetup);
+
+	FKAggregateGeom& AggGeom = InSetup->AggGeom;
+
+	for(FKBoxElem& Box : AggGeom.BoxElems)
+	{
+		Chaos::TVector<Precision, 3> HalfBoxExtent = Chaos::TVector<Precision, 3>(Box.X, Box.Y, Box.Z) / 2.0f;
+		TUniquePtr<Chaos::TBox<Precision, 3>> NonTransformed = MakeUnique<Chaos::TBox<Precision, 3>>(-HalfBoxExtent, HalfBoxExtent);
+		Chaos::TRigidTransform<Precision, 3> ShapeLocalTransform(Box.Center, Box.Rotation.Quaternion());
+		OutImplicits.Add(MakeUnique<Chaos::TImplicitObjectTransformed<Precision, 3>>(MakeSerializable(NonTransformed), MoveTemp(NonTransformed), ShapeLocalTransform));
+	}
+
+	for(FKSphereElem& Sphere : AggGeom.SphereElems)
+	{
+		OutImplicits.Add(MakeUnique<Chaos::TSphere<Precision, 3>>(Sphere.Center, Sphere.Radius));
+	}
+
+	for(FKSphylElem& Sphyl : AggGeom.SphylElems)
+	{
+		const float HalfLength = Sphyl.Length / 2.0f;
+		const Chaos::TVector<Precision, 3> TopPoint(0.0f, 0.0f, HalfLength);
+		const Chaos::TVector<Precision, 3> BottomPoint(0.0f, 0.0f, -HalfLength);
+		TUniquePtr<Chaos::TCapsule<Precision>> NonTransformed = MakeUnique<Chaos::TCapsule<Precision>>(TopPoint, BottomPoint, Sphyl.Radius);
+		Chaos::TRigidTransform<Precision, 3> ShapeLocalTransform(Sphyl.Center, Sphyl.Rotation.Quaternion());
+		OutImplicits.Add(MakeUnique<Chaos::TImplicitObjectTransformed<Precision, 3>>(MakeSerializable(NonTransformed), MoveTemp(NonTransformed), ShapeLocalTransform));
+	}
+
+	const int32 NumTaperedCapsules = AggGeom.TaperedCapsuleElems.Num();
+	UE_CLOG(NumTaperedCapsules > 0, LogChaos, Warning, TEXT("Ignoring %d tapered spheres when building collision data for body setup %s"), NumTaperedCapsules, *InSetup->GetName());
+}
+
+template<typename Precision>
+void FChaosDerivedDataCooker::BuildInternal(Chaos::FChaosArchive& Ar, FCookBodySetupInfo& InInfo)
+{
+	TArray<TUniquePtr<Chaos::TImplicitObject<Precision, 3>>> SimpleImplicits;
+	TArray<TUniquePtr<Chaos::TTriangleMeshImplicitObject<Precision>>> ComplexImplicits;
+
+	//BuildSimpleShapes(SimpleImplicits, Setup);
+	BuildConvexMeshes(SimpleImplicits, InInfo);
+	BuildTriangleMeshes(ComplexImplicits, InInfo);
+
+	//TUniquePtr<Chaos::TImplicitObjectUnion<Precision, 3>> SimpleUnion = MakeUnique<Chaos::TImplicitObjectUnion<Precision, 3>>(MoveTemp(SimpleImplicits));
+	//TUniquePtr<Chaos::TImplicitObjectUnion<Precision, 3>> ComplexUnion = MakeUnique<Chaos::TImplicitObjectUnion<Precision, 3>>(MoveTemp(ComplexImplicits));
+
+	FBodySetupUVInfo UVInfo;
+	if(InInfo.bSupportUVFromHitResults)
+	{
+		UVInfo.FillFromTriMesh(InInfo.TriangleMeshDesc);
+	}
+
+	Ar << SimpleImplicits << ComplexImplicits << UVInfo;
+}
+
+template void FChaosDerivedDataCooker::BuildInternal<float>(Chaos::FChaosArchive& Ar, FCookBodySetupInfo& InInfo);
+//#BGTODO When it's possible to build with doubles, re-enable this. (Currently at least TRigidTransform cannot build with double precision because we don't have a base transform implementation using them)
+//template void FChaosDerivedDataCooker::BuildInternal<float>(FArchive& Ar, FCookBodySetupInfo& InInfo);
+
+#endif
+ 
diff --git a/Engine/Source/Runtime/Engine/Private/PhysicsEngine/Experimental/ChaosDerivedData.h b/Engine/Source/Runtime/Engine/Private/PhysicsEngine/Experimental/ChaosDerivedData.h
new file mode 100644
--- /dev/null
+++ b/Engine/Source/Runtime/Engine/Private/PhysicsEngine/Experimental/ChaosDerivedData.h
@@ -0,0 +1,94 @@
+// Copyright 1998-2019 Epic Games, Inc. All Rights Reserved.
+
+#pragma once
+
+#if WITH_CHAOS
+
+#include "CoreMinimal.h"
+#include "DerivedDataPluginInterface.h"
+#include "UObject/GCObject.h"
+#include "Interface_CollisionDataProviderCore.h"
+
+struct FBodySetupUVInfo;
+struct FCookBodySetupInfo;
+
+class UBodySetup;
+
+namespace Chaos
+{
+	class FChaosArchive;
+
+	template<typename T, int d>
+	class TImplicitObject;
+
+	template<typename T>
+	class TTriangleMeshImplicitObject;
+}
+
+struct FChaosTriMeshCollisionBuildParameters
+{
+	bool bCollapseVerts;
+
+	FTriMeshCollisionData MeshDesc;
+	FBodySetupUVInfo* UVInfo;
+};
+
+struct FChaosConvexMeshCollisionBuildParameters
+{
+	bool bCollapseVerts;
+	bool bMirror;
+};
+
+class FChaosDerivedDataCooker : public FDerivedDataPluginInterface, public FGCObject
+{
+public:
+
+	using BuildPrecision = float;
+
+	// FDerivedDataPluginInterface Interface
+	virtual const TCHAR* GetPluginName() const override;
+	virtual const TCHAR* GetVersionString() const override;
+	virtual FString GetPluginSpecificCacheKeySuffix() const override;
+	virtual bool IsBuildThreadsafe() const override;
+	virtual bool Build(TArray<uint8>& OutData) override;
+	//End FDerivedDataPluginInterface Interface
+
+	// FGCObject Interface
+	virtual void AddReferencedObjects(FReferenceCollector& Collector) override;
+	// End FGCObject Interface
+
+	FChaosDerivedDataCooker(UBodySetup* InSetup, FName InFormat);
+
+	bool CanBuild()
+	{
+		return Setup != nullptr;
+	}
+
+private:
+
+	template<typename Precision>
+	void BuildInternal(Chaos::FChaosArchive& Ar, FCookBodySetupInfo& InInfo);
+
+	template<typename Precision>
+	void BuildTriangleMeshes(TArray<TUniquePtr<Chaos::TTriangleMeshImplicitObject<Precision>>>& OutTriangleMeshes, const FCookBodySetupInfo& InParams);
+
+	template<typename Precision>
+	void BuildConvexMeshes(TArray<TUniquePtr<Chaos::TImplicitObject<Precision, 3>>>& OutTriangleMeshes, const FCookBodySetupInfo& InParams);
+
+	UBodySetup* Setup;
+	FName RequestedFormat;
+};
+
+extern template void FChaosDerivedDataCooker::BuildTriangleMeshes(TArray<TUniquePtr<Chaos::TTriangleMeshImplicitObject<float>>>& OutTriangleMeshes, const FCookBodySetupInfo& InParams);
+// #BGTODO When it's possible to build with doubles, re-enable this. (Currently at least TRigidTransform cannot build with double precision because we don't have a base transform implementation using them)
+//extern template void FChaosDerivedDataCooker::BuildTriangleMeshes(TArray<Chaos::TCollisionTriangleMesh<double>>& OutTriangleMeshes, const FCookBodySetupInfo& InParams);
+
+extern template void FChaosDerivedDataCooker::BuildConvexMeshes(TArray<TUniquePtr<Chaos::TImplicitObject<float, 3>>>& OutTriangleMeshes, const FCookBodySetupInfo& InParams);
+// #BGTODO When it's possible to build with doubles, re-enable this. (Currently at least TRigidTransform cannot build with double precision because we don't have a base transform implementation using them)
+//extern template void FChaosDerivedDataCooker::BuildConvexMeshes(TArray<Chaos::TCollisionConvexMesh<double>>& OutTriangleMeshes, const FCookBodySetupInfo& InParams);
+
+extern template void FChaosDerivedDataCooker::BuildInternal<float>(Chaos::FChaosArchive& Ar, FCookBodySetupInfo& InInfo);
+// #BGTODO When it's possible to build with doubles, re-enable this. (Currently at least TRigidTransform cannot build with double precision because we don't have a base transform implementation using them)
+//extern template void FChaosDerivedDataCooker::BuildInternal<float>(FArchive& Ar, FCookBodySetupInfo& InInfo);
+
+#endif
\ No newline at end of file
diff --git a/Engine/Source/Runtime/Engine/Private/PhysicsEngine/Experimental/ChaosDerivedDataUtil.cpp b/Engine/Source/Runtime/Engine/Private/PhysicsEngine/Experimental/ChaosDerivedDataUtil.cpp
new file mode 100644
--- /dev/null
+++ b/Engine/Source/Runtime/Engine/Private/PhysicsEngine/Experimental/ChaosDerivedDataUtil.cpp
@@ -0,0 +1,168 @@
+// Copyright 1998-2019 Epic Games, Inc. All Rights Reserved.
+
+#include "ChaosDerivedDataUtil.h"
+#include "CoreMinimal.h"
+
+#if INCLUDE_CHAOS
+#include "ChaosLog.h"
+
+namespace Chaos
+{
+
+	void CleanTrimesh(TArray<FVector>& InOutVertices, TArray<int32>& InOutIndices, TArray<int32>* OutOptFaceRemap)
+	{
+		TArray<FVector> LocalSourceVerts = InOutVertices;
+		TArray<int32> LocalSourceIndices = InOutIndices;
+		
+		const int32 NumSourceVerts = LocalSourceVerts.Num();
+		const int32 NumSourceTriangles = LocalSourceIndices.Num() / 3;
+
+		if(NumSourceVerts == 0 || (LocalSourceIndices.Num() % 3) != 0)
+		{
+			// No valid geometry passed in
+			return;
+		}
+
+		// New condensed list of unique verts from the trimesh
+		TArray<FVector> LocalUniqueVerts;
+		// New condensed list of indices after cleaning
+		TArray<int32> LocalUniqueIndices;
+		// Array mapping unique vertex index back to source data index
+		TArray<int32> LocalUniqueToSourceIndices;
+		// Remapping table from source index to unique index
+		TArray<int32> LocalVertexRemap;
+		// Remapping table from source triangle to unique triangle
+		TArray<int32> LocalTriangleRemap;
+
+		LocalUniqueVerts.Reserve(NumSourceVerts);
+		LocalVertexRemap.AddUninitialized(NumSourceVerts);
+		LocalTriangleRemap.AddUninitialized(NumSourceTriangles);
+
+		auto ValidateTrianglesPre = [&InOutVertices](int32 A, int32 B, int32 C) -> bool
+		{
+			const FVector v0 = InOutVertices[A];
+			const FVector v1 = InOutVertices[B];
+			const FVector v2 = InOutVertices[C];
+			return v0 != v1 && v0 != v2 && v1 != v2;
+		};
+
+		int32 NumBadTris = 0;
+		for(int32 SrcTriIndex = 0; SrcTriIndex < NumSourceTriangles; ++SrcTriIndex)
+		{
+			const int32 A = InOutIndices[SrcTriIndex * 3];
+			const int32 B = InOutIndices[SrcTriIndex * 3 + 1];
+			const int32 C = InOutIndices[SrcTriIndex * 3 + 2];
+
+			if(!ValidateTrianglesPre(A, B, C))
+			{
+				++NumBadTris;
+			}
+		}
+		UE_CLOG(NumBadTris > 0, LogChaos, Warning, TEXT("Input trimesh contains %d bad triangles."), NumBadTris);
+
+		float WeldThresholdSq = 0.0f;// SMALL_NUMBER * SMALL_NUMBER;
+
+		for(int32 SourceVertIndex = 0; SourceVertIndex < NumSourceVerts; ++SourceVertIndex)
+		{
+			const FVector& SourceVert = LocalSourceVerts[SourceVertIndex];
+
+			bool bUnique = true; // assume the vertex is unique until we find otherwise
+			int32 RemapIndex = INDEX_NONE; // if the vertex isn't unique this will be set
+
+			const int32 NumUniqueVerts = LocalUniqueVerts.Num();
+			for(int32 UniqueVertIndex = 0; UniqueVertIndex < NumUniqueVerts; ++UniqueVertIndex)
+			{
+				const FVector& UniqueVert = LocalUniqueVerts[UniqueVertIndex];
+
+				if((UniqueVert - SourceVert).SizeSquared() <= WeldThresholdSq)
+				{
+					//This vertex isn't unique and needs to be merged with the unique one we found
+					bUnique = false;
+					RemapIndex = UniqueVertIndex;
+
+					// Done searching
+					break;
+				}
+			}
+
+			if(bUnique)
+			{
+				LocalUniqueVerts.Add(SourceVert);
+				LocalUniqueToSourceIndices.Add(SourceVertIndex);
+				LocalVertexRemap[SourceVertIndex] = LocalUniqueVerts.Num() - 1;
+			}
+			else
+			{
+				LocalVertexRemap[SourceVertIndex] = RemapIndex;
+			}
+		}
+
+		// Build the new index buffer, removing now invalid merged triangles
+		auto ValidateTriangleIndices = [](int32 A, int32 B, int32 C) -> bool
+		{
+			return A != B && A != C && B != C;
+		};
+
+		auto ValidateTriangleArea = [](const FVector& A, const FVector& B, const FVector& C)
+		{
+			const float AreaSq = FVector::CrossProduct(A - B, A - C).SizeSquared();
+
+			return AreaSq > SMALL_NUMBER;
+		};
+
+		int32 NumDiscardedTriangles_Welded = 0;
+		int32 NumDiscardedTriangles_Area = 0;
+
+		for(int32 OriginalTriIndex = 0; OriginalTriIndex < NumSourceTriangles; ++OriginalTriIndex)
+		{
+			const int32 OrigAIndex = LocalSourceIndices[OriginalTriIndex * 3];
+			const int32 OrigBIndex = LocalSourceIndices[OriginalTriIndex * 3 + 1];
+			const int32 OrigCIndex = LocalSourceIndices[OriginalTriIndex * 3 + 2];
+
+			const int32 RemappedAIndex = LocalVertexRemap[OrigAIndex];
+			const int32 RemappedBIndex = LocalVertexRemap[OrigBIndex];
+			const int32 RemappedCIndex = LocalVertexRemap[OrigCIndex];
+
+			const FVector& A = LocalUniqueVerts[RemappedAIndex];
+			const FVector& B = LocalUniqueVerts[RemappedBIndex];
+			const FVector& C = LocalUniqueVerts[RemappedCIndex];
+
+			// Only consider triangles that are actually valid for collision
+			// #BG Consider being able to fix small triangles by collapsing them if we hit this a lot
+			const bool bValidIndices = ValidateTriangleIndices(RemappedAIndex, RemappedBIndex, RemappedCIndex);
+			const bool bValidArea = ValidateTriangleArea(A, B, C);
+			if(bValidIndices && bValidArea)
+			{
+				LocalUniqueIndices.Add(RemappedAIndex);
+				LocalUniqueIndices.Add(RemappedBIndex);
+				LocalUniqueIndices.Add(RemappedCIndex);
+				LocalTriangleRemap.Add(OriginalTriIndex);
+			}
+			else
+			{
+				if(!bValidIndices)
+				{
+					++NumDiscardedTriangles_Welded;
+				}
+				else if(!bValidArea)
+				{
+					++NumDiscardedTriangles_Area;
+				}
+			}
+		}
+
+		UE_CLOG(NumDiscardedTriangles_Welded > 0, LogChaos, Warning, TEXT("Discarded %d welded triangles when cooking trimesh."), NumDiscardedTriangles_Welded);
+		UE_CLOG(NumDiscardedTriangles_Area > 0, LogChaos, Warning, TEXT("Discarded %d small triangles when cooking trimesh."), NumDiscardedTriangles_Area);
+
+		InOutVertices = LocalUniqueVerts;
+		InOutIndices = LocalUniqueIndices;
+
+		if(OutOptFaceRemap)
+		{
+			*OutOptFaceRemap = LocalTriangleRemap;
+		}
+	}
+
+}
+
+#endif // INCLUDE_CHAOS
\ No newline at end of file
diff --git a/Engine/Source/Runtime/Engine/Private/PhysicsEngine/Experimental/ChaosDerivedDataUtil.h b/Engine/Source/Runtime/Engine/Private/PhysicsEngine/Experimental/ChaosDerivedDataUtil.h
new file mode 100644
--- /dev/null
+++ b/Engine/Source/Runtime/Engine/Private/PhysicsEngine/Experimental/ChaosDerivedDataUtil.h
@@ -0,0 +1,12 @@
+// Copyright 1998-2019 Epic Games, Inc. All Rights Reserved.
+
+#include "CoreMinimal.h"
+
+#if INCLUDE_CHAOS
+
+namespace Chaos
+{
+	void CleanTrimesh(TArray<FVector>& InOutVertices, TArray<int32>& InOutIndices, TArray<int32>* OutOptFaceRemap);
+}
+
+#endif
diff --git a/Engine/Source/Runtime/Engine/Private/PhysicsEngine/Experimental/PhysInterface_Chaos.cpp b/Engine/Source/Runtime/Engine/Private/PhysicsEngine/Experimental/PhysInterface_Chaos.cpp
--- a/Engine/Source/Runtime/Engine/Private/PhysicsEngine/Experimental/PhysInterface_Chaos.cpp
+++ b/Engine/Source/Runtime/Engine/Private/PhysicsEngine/Experimental/PhysInterface_Chaos.cpp
@@ -6,21 +6,32 @@
 #include "Physics/Experimental/PhysScene_Chaos.h"
 #include "Physics/PhysicsInterfaceTypes.h"
 #include "PhysicsEngine/AggregateGeom.h"
-#include "SolverObjects/BodyInstancePhysicsObject.h"
+#include "Chaos/Framework/SingleParticlePhysicsProxy.h"
 
 #include "Chaos/Box.h"
 #include "Chaos/Cylinder.h"
+#include "Chaos/Capsule.h"
 #include "Chaos/ImplicitObjectTransformed.h"
 #include "Chaos/ImplicitObjectUnion.h"
 #include "Chaos/Levelset.h"
 #include "Chaos/PBDRigidParticles.h"
 #include "Chaos/Sphere.h"
-#include "PBDRigidsSolver.h"
+#include "PhysicsSolver.h"
 #include "Templates/UniquePtr.h"
 #include "ChaosSolversModule.h"
+#include "Chaos/ErrorReporter.h"
+#include "Chaos/ImplicitObjectScaled.h"
+#include "Chaos/TriangleMeshImplicitObject.h"
+#include "Chaos/Convex.h"
+#include "Chaos/GeometryQueries.h"
+
 
 #include "Async/ParallelFor.h"
 #include "Components/PrimitiveComponent.h"
+#include "Physics/PhysicsFiltering.h"
+#include "Collision/CollisionConversions.h"
+#include "PhysicsInterfaceUtilsCore.h"
+
 
 #if WITH_PHYSX
 #include "geometry/PxConvexMesh.h"
@@ -29,6 +40,8 @@
 #include "extensions/PxMassProperties.h"
 #endif
 
+
+
 #define FORCE_ANALYTICS 0
 
 DEFINE_STAT(STAT_TotalPhysicsTime);
@@ -78,98 +91,50 @@ FPhysInterface_Chaos::~FPhysInterface_Chaos()
 // Interface functions
 void FPhysInterface_Chaos::CreateActor(const FActorCreationParams& InParams, FPhysicsActorHandle& Handle)
 {
-	if (FChaosSolversModule* ChaosModule = FModuleManager::Get().GetModulePtr<FChaosSolversModule>("ChaosSolvers"))
+	// Create the new particle
+	if (InParams.bStatic)
 	{
-		if (Chaos::IDispatcher* Dispatcher = ChaosModule->GetDispatcher())
-		{
-			if (FPhysScene_ChaosInterface* Scene = InParams.Scene)
-			{
-				// Create the new physics proxy
-				if (!Handle.GetPhysicsObject())
-				{
-					FBodyInstancePhysicsObject * PhysicsObject = new FBodyInstancePhysicsObject(nullptr, { 1.0, 1.0, FVector(1.f) });
-					PhysicsObject->CreationParameters = InParams;
-					Handle.SetScene(Scene);
-					Handle.SetPhysicsObject(PhysicsObject);
-					Scene->Scene.AddObject(nullptr, PhysicsObject);
-				}
-
-				// Pass the proxy off to the physics thread
-				Dispatcher->EnqueueCommand([Scene, InParams, Handle](Chaos::FPersistentPhysicsTask* PhysThread)
-				{
-					if (Chaos::FPBDRigidsSolver* Solver = Scene->GetSolver())
-					{
-						// Collect references
-						if (FBodyInstancePhysicsObject* Object = Handle.GetPhysicsObject())
-						{
-							FParticlesType& Particles = Solver->GetRigidParticles();
-
-							int32 Index = Particles.Size();
-							Object->SetRigidBodyId(Index);
-							Object->InitializedIndices.Add(Index);
-
-							// Add the particle
-							Particles.AddParticles(1);
-							Particles.X(Index) = InParams.InitialTM.GetTranslation();
-							Particles.R(Index) = InParams.InitialTM.GetRotation();
-							Particles.M(Index) = 1;					// TODO: This should be based on InParams?
-							Particles.I(Index) = FMatrix::Identity; // TODO: This should be based on InParams?
-							Particles.V(Index) = Chaos::TVector<float, 3>(0); // TODO: Do we want to support initial velocity?
-							Particles.W(Index) = Chaos::TVector<float, 3>(0); // TODO: Do we want to support initial velocity?
-
-							// Set the initial object state
-							// TODO: Right now we only support static or dynamic.
-							//       Params will probably also need to be able to support Kinematic and initial sleep states.
-							const Chaos::EObjectStateType InitialState = InParams.bStatic ? Chaos::EObjectStateType::Static : Chaos::EObjectStateType::Dynamic;
-							Particles.SetObjectState(Index, InitialState);
-
-							// Initially disable simulation if this is a query only object
-							Particles.Disabled(Index) = InParams.bQueryOnly;
+		Handle = new Chaos::TGeometryParticle<float, 3>();
+	}
+	else if (InParams.BodyInstance && InParams.BodyInstance->ShouldInstanceSimulatingPhysics())
+	{
+		Handle = new Chaos::TPBDRigidParticle<float, 3>();
+	}
+	else
+	{
+		Handle = new Chaos::TKinematicGeometryParticle<float, 3>();
+	}
 
-							// TODO: At the moment, we don't actually have support for objects with no gravity...
-							check(InParams.bEnableGravity);
+	// Set up the new particle's game-thread data. This will be sent to physics-thread when
+	// the particle is added to the scene later.
+	Handle->SetX(InParams.InitialTM.GetLocation());
+	Handle->SetR(InParams.InitialTM.GetRotation());
+}
 
 
-							Solver->InitializeFromParticleData(Index);
-						}
-					}
-				});
-			}
-		}
-	}
+void FPhysInterface_Chaos::AddActorToSolver(FPhysicsActorHandle& Handle, Chaos::FPhysicsSolver* Solver, Chaos::IDispatcher* Dispatcher)
+{
+	Solver->RegisterObject(Handle);
 }
 
 void FPhysInterface_Chaos::ReleaseActor(FPhysicsActorHandle& Handle, FPhysScene* InScene, bool bNeverDerferRelease)
 {
-	if (FPhysScene_ChaosInterface* Scene = Handle.GetScene())
+	check(Handle);
+	if (InScene)
 	{
-		check(Scene);
-		check(InScene == Scene);
+		RemoveActorFromSolver(Handle, InScene->GetSolver(), FChaosSolversModule::GetModule()->GetDispatcher());
+	}
 
-		if (FBodyInstancePhysicsObject* Object = Handle.GetPhysicsObject())
-		{
-			// Invalidate the reference
-			Handle.SetPhysicsObject(nullptr);
+	delete Handle;
 
-			// Queue a task to remove the object from the scene
-			if (FChaosSolversModule* ChaosModule = FModuleManager::Get().GetModulePtr<FChaosSolversModule>("ChaosSolvers"))
-			{
-				if (Chaos::IDispatcher* Dispatcher = ChaosModule->GetDispatcher())
-				{
-					if (Chaos::FPBDRigidsSolver* Solver = Scene->GetSolver())
-					{
-						Dispatcher->EnqueueCommand([Solver, Object](Chaos::FPersistentPhysicsTask* PhysThread)
-						{
-							// Remove the object from the solver, & delete the object
-							Solver->UnregisterObject(Object);
-							Object->OnRemoveFromScene();
-							Object->SyncBeforeDestroy();
-							delete Object;
-						});
-					}
-				}
-			}
-		}
+	Handle = nullptr;
+}
+
+void FPhysInterface_Chaos::RemoveActorFromSolver(FPhysicsActorHandle& Handle, Chaos::FPhysicsSolver* Solver, Chaos::IDispatcher* Dispatcher)
+{
+	if (Solver)
+	{
+		Solver->UnregisterObject(Handle);
 	}
 }
 
@@ -183,7 +148,7 @@ FPhysicsAggregateReference_Chaos FPhysInterface_Chaos::CreateAggregate(int32 Max
 
 void FPhysInterface_Chaos::ReleaseAggregate(FPhysicsAggregateReference_Chaos& InAggregate) {}
 int32 FPhysInterface_Chaos::GetNumActorsInAggregate(const FPhysicsAggregateReference_Chaos& InAggregate) { return 0; }
-void FPhysInterface_Chaos::AddActorToAggregate_AssumesLocked(const FPhysicsAggregateReference_Chaos& InAggregate, const FPhysicsActorReference_Chaos& InActor) {}
+void FPhysInterface_Chaos::AddActorToAggregate_AssumesLocked(const FPhysicsAggregateReference_Chaos& InAggregate, const FPhysicsActorHandle& InActor) {}
 
 
 int32 FPhysInterface_Chaos::GetNumShapes(const FPhysicsActorHandle& InHandle)
@@ -194,8 +159,9 @@ int32 FPhysInterface_Chaos::GetNumShapes(const FPhysicsActorHandle& InHandle)
 
 void FPhysInterface_Chaos::ReleaseShape(const FPhysicsShapeHandle& InShape)
 {
-    check(!InShape.ActorRef.IsValid());
-	delete InShape.Object;
+    check(!IsValid(InShape.ActorRef));
+	//no need to delete because ownership is on actor. Is this an invalid assumption with the current API?
+	//delete InShape.Shape;
 }
 
 void FPhysInterface_Chaos::AttachShape(const FPhysicsActorHandle& InActor, const FPhysicsShapeHandle& InNewShape)
@@ -208,98 +174,106 @@ void FPhysInterface_Chaos::DetachShape(const FPhysicsActorHandle& InActor, FPhys
 	// #todo : Implement
 }
 
-void FPhysInterface_Chaos::SetActorUserData_AssumesLocked(FPhysicsActorReference_Chaos& InActorReference, FPhysxUserData* InUserData)
+void FPhysInterface_Chaos::SetActorUserData_AssumesLocked(FPhysicsActorHandle& InActorReference, FPhysicsUserData* InUserData)
 {
-    FBodyInstance* BodyInstance = InUserData->Get<FBodyInstance>(InUserData);
+	InActorReference->SetUserData(InUserData);
+}
 
-    if (BodyInstance)
-    {
-        InActorReference.SetBodyInstance(BodyInstance);
-    }
+bool FPhysInterface_Chaos::IsRigidBody(const FPhysicsActorHandle& InActorReference)
+{
+	return InActorReference->ObjectState() == Chaos::EObjectStateType::Dynamic;
 }
 
-bool FPhysInterface_Chaos::IsRigidBody(const FPhysicsActorReference_Chaos& InActorReference)
+bool FPhysInterface_Chaos::IsStatic(const FPhysicsActorHandle& InActorReference)
 {
-    return InActorReference.IsValid();
+	return InActorReference->ObjectState() == Chaos::EObjectStateType::Static;
 }
 
-bool FPhysInterface_Chaos::IsStatic(const FPhysicsActorReference_Chaos& InActorReference)
+bool FPhysInterface_Chaos::IsKinematic(const FPhysicsActorHandle& InActorReference)
 {
-	if (InActorReference.GetPhysicsObject())
-	{
-		return InActorReference.GetPhysicsObject()->GetInitialState().GetInverseMass()==0;
-	}
-	check(false);
-	return 0.f;
+	return InActorReference->ObjectState() == Chaos::EObjectStateType::Kinematic;
 }
 
-bool FPhysInterface_Chaos::IsKinematic_AssumesLocked(const FPhysicsActorReference_Chaos& InActorReference)
+bool FPhysInterface_Chaos::IsKinematic_AssumesLocked(const FPhysicsActorHandle& InActorReference)
 {
-    return IsStatic(InActorReference);
+	return IsKinematic(InActorReference);
 }
 
-bool FPhysInterface_Chaos::IsSleeping(const FPhysicsActorReference_Chaos& InActorReference)
+bool FPhysInterface_Chaos::IsSleeping(const FPhysicsActorHandle& InActorReference)
 {
 	// #todo : Implement
 	return false;
 }
 
-bool FPhysInterface_Chaos::IsCcdEnabled(const FPhysicsActorReference_Chaos& InActorReference)
+bool FPhysInterface_Chaos::IsCcdEnabled(const FPhysicsActorHandle& InActorReference)
 {
     return false;
 }
 
-bool FPhysInterface_Chaos::IsInScene(const FPhysicsActorReference_Chaos& InActorReference)
+bool FPhysInterface_Chaos::IsInScene(const FPhysicsActorHandle& InActorReference)
+{
+	// TODO: Implement
+	return false;
+}
+
+FPhysScene* FPhysInterface_Chaos::GetCurrentScene(const FPhysicsActorHandle& InHandle)
 {
-    return InActorReference.GetScene() != nullptr;
+	if (IPhysicsProxyBase* Proxy = InHandle->Proxy)
+	{
+		Chaos::FPhysicsSolver* Solver = Proxy->GetSolver();
+		return static_cast<FPhysScene*>(Solver ? Solver->PhysSceneHack : nullptr);
+	}
+	return nullptr;
 }
 
-bool FPhysInterface_Chaos::CanSimulate_AssumesLocked(const FPhysicsActorReference_Chaos& InActorReference)
+bool FPhysInterface_Chaos::CanSimulate_AssumesLocked(const FPhysicsActorHandle& InActorReference)
 {
 	// #todo : Implement
 	return true;
 }
 
-float FPhysInterface_Chaos::GetMass_AssumesLocked(const FPhysicsActorReference_Chaos& InActorReference)
+float FPhysInterface_Chaos::GetMass_AssumesLocked(const FPhysicsActorHandle& InActorReference)
 {
 	// #todo : Implement
 	return 1.f;
 }
 
-void FPhysInterface_Chaos::SetSendsSleepNotifies_AssumesLocked(const FPhysicsActorReference_Chaos& InActorReference, bool bSendSleepNotifies)
+void FPhysInterface_Chaos::SetSendsSleepNotifies_AssumesLocked(const FPhysicsActorHandle& InActorReference, bool bSendSleepNotifies)
 {
-    check(bSendSleepNotifies == false);
+	// # todo: Implement
+    //check(bSendSleepNotifies == false);
 }
 
-void FPhysInterface_Chaos::PutToSleep_AssumesLocked(const FPhysicsActorReference_Chaos& InActorReference)
+void FPhysInterface_Chaos::PutToSleep_AssumesLocked(const FPhysicsActorHandle& InActorReference)
 {
 	// #todo : Implement
 }
 
-void FPhysInterface_Chaos::WakeUp_AssumesLocked(const FPhysicsActorReference_Chaos& InActorReference)
+void FPhysInterface_Chaos::WakeUp_AssumesLocked(const FPhysicsActorHandle& InActorReference)
 {
 	// #todo : Implement
 }
 
-void FPhysInterface_Chaos::SetIsKinematic_AssumesLocked(const FPhysicsActorReference_Chaos& InActorReference, bool bIsKinematic)
+void FPhysInterface_Chaos::SetIsKinematic_AssumesLocked(const FPhysicsActorHandle& InActorReference, bool bIsKinematic)
 {
 	// #todo : Implement
 }
 
-void FPhysInterface_Chaos::SetCcdEnabled_AssumesLocked(const FPhysicsActorReference_Chaos& InActorReference, bool bIsCcdEnabled)
+void FPhysInterface_Chaos::SetCcdEnabled_AssumesLocked(const FPhysicsActorHandle& InActorReference, bool bIsCcdEnabled)
 {
-    check(bIsCcdEnabled == false);
+	// #todo: Implement
+    //check(bIsCcdEnabled == false);
 }
 
-FTransform FPhysInterface_Chaos::GetGlobalPose_AssumesLocked(const FPhysicsActorReference_Chaos& InActorReference)
+FTransform FPhysInterface_Chaos::GetGlobalPose_AssumesLocked(const FPhysicsActorHandle& InActorReference)
 {
-	// #todo : Implement
-	return Chaos::TRigidTransform<float, 3>();
+	return Chaos::TRigidTransform<float, 3>(InActorReference->X(), InActorReference->R());
 }
 
-void FPhysInterface_Chaos::SetGlobalPose_AssumesLocked(const FPhysicsActorReference_Chaos& InActorReference, const FTransform& InNewPose, bool bAutoWake)
+void FPhysInterface_Chaos::SetGlobalPose_AssumesLocked(const FPhysicsActorHandle& InActorReference, const FTransform& InNewPose, bool bAutoWake)
 {
-	// #todo : Implement
+	InActorReference->SetX(InNewPose.GetLocation());
+	InActorReference->SetR(InNewPose.GetRotation());
 }
 
 FTransform FPhysInterface_Chaos::GetTransform_AssumesLocked(const FPhysicsActorHandle& InRef, bool bForceGlobalPose /*= false*/)
@@ -318,163 +292,165 @@ FTransform FPhysInterface_Chaos::GetTransform_AssumesLocked(const FPhysicsActorH
 	return GetGlobalPose_AssumesLocked(InRef);
 }
 
-bool FPhysInterface_Chaos::HasKinematicTarget_AssumesLocked(const FPhysicsActorReference_Chaos& InActorReference)
+bool FPhysInterface_Chaos::HasKinematicTarget_AssumesLocked(const FPhysicsActorHandle& InActorReference)
 {
     return IsStatic(InActorReference);
 }
 
-FTransform FPhysInterface_Chaos::GetKinematicTarget_AssumesLocked(const FPhysicsActorReference_Chaos& InActorReference)
+FTransform FPhysInterface_Chaos::GetKinematicTarget_AssumesLocked(const FPhysicsActorHandle& InActorReference)
 {
 	// #todo : Implement
-	return FTransform();
+	//for now just use global pose
+	return FPhysInterface_Chaos::GetGlobalPose_AssumesLocked(InActorReference);
 }
 
-void FPhysInterface_Chaos::SetKinematicTarget_AssumesLocked(const FPhysicsActorReference_Chaos& InActorReference, const FTransform& InNewTarget)
+void FPhysInterface_Chaos::SetKinematicTarget_AssumesLocked(const FPhysicsActorHandle& InActorReference, const FTransform& InNewTarget)
 {
 	// #todo : Implement
-    //InActorReference.GetScene()->SetKinematicTransform(InActorReference, InNewTarget);
+	//for now just use global pose
+	FPhysInterface_Chaos::SetGlobalPose_AssumesLocked(InActorReference, InNewTarget);
 }
 
-FVector FPhysInterface_Chaos::GetLinearVelocity_AssumesLocked(const FPhysicsActorReference_Chaos& InActorReference)
+FVector FPhysInterface_Chaos::GetLinearVelocity_AssumesLocked(const FPhysicsActorHandle& InActorReference)
 {
 	// #todo : Implement
 	return FVector(0);
 }
 
-void FPhysInterface_Chaos::SetLinearVelocity_AssumesLocked(const FPhysicsActorReference_Chaos& InActorReference, const FVector& InNewVelocity, bool bAutoWake)
+void FPhysInterface_Chaos::SetLinearVelocity_AssumesLocked(const FPhysicsActorHandle& InActorReference, const FVector& InNewVelocity, bool bAutoWake)
 {
 	// #todo : Implement
 }
 
-FVector FPhysInterface_Chaos::GetAngularVelocity_AssumesLocked(const FPhysicsActorReference_Chaos& InActorReference)
+FVector FPhysInterface_Chaos::GetAngularVelocity_AssumesLocked(const FPhysicsActorHandle& InActorReference)
 {
 	// #todo : Implement
 	return FVector(0);
 }
 
-void FPhysInterface_Chaos::SetAngularVelocity_AssumesLocked(const FPhysicsActorReference_Chaos& InActorReference, const FVector& InNewVelocity, bool bAutoWake)
+void FPhysInterface_Chaos::SetAngularVelocity_AssumesLocked(const FPhysicsActorHandle& InActorReference, const FVector& InNewVelocity, bool bAutoWake)
 {
 	// #todo : Implement
 }
 
-float FPhysInterface_Chaos::GetMaxAngularVelocity_AssumesLocked(const FPhysicsActorReference_Chaos& InActorReference)
+float FPhysInterface_Chaos::GetMaxAngularVelocity_AssumesLocked(const FPhysicsActorHandle& InActorReference)
 {
     return FLT_MAX;
 }
 
-void FPhysInterface_Chaos::SetMaxAngularVelocity_AssumesLocked(const FPhysicsActorReference_Chaos& InActorReference, float InMaxAngularVelocity)
+void FPhysInterface_Chaos::SetMaxAngularVelocity_AssumesLocked(const FPhysicsActorHandle& InActorReference, float InMaxAngularVelocity)
 {
 }
 
-float FPhysInterface_Chaos::GetMaxDepenetrationVelocity_AssumesLocked(const FPhysicsActorReference_Chaos& InActorReference)
+float FPhysInterface_Chaos::GetMaxDepenetrationVelocity_AssumesLocked(const FPhysicsActorHandle& InActorReference)
 {
     return FLT_MAX;
 }
 
-void FPhysInterface_Chaos::SetMaxDepenetrationVelocity_AssumesLocked(const FPhysicsActorReference_Chaos& InActorReference, float InMaxDepenetrationVelocity)
+void FPhysInterface_Chaos::SetMaxDepenetrationVelocity_AssumesLocked(const FPhysicsActorHandle& InActorReference, float InMaxDepenetrationVelocity)
 {
 }
 
-FVector FPhysInterface_Chaos::GetWorldVelocityAtPoint_AssumesLocked(const FPhysicsActorReference_Chaos& InActorReference, const FVector& InPoint)
+FVector FPhysInterface_Chaos::GetWorldVelocityAtPoint_AssumesLocked(const FPhysicsActorHandle& InActorReference, const FVector& InPoint)
 {
 	// #todo : Implement
 	return FVector(0);
 }
 
-FTransform FPhysInterface_Chaos::GetComTransform_AssumesLocked(const FPhysicsActorReference_Chaos& InActorReference)
+FTransform FPhysInterface_Chaos::GetComTransform_AssumesLocked(const FPhysicsActorHandle& InActorReference)
 {
 	// #todo : Implement
 	return FTransform();
 }
 
-FTransform FPhysInterface_Chaos::GetComTransformLocal_AssumesLocked(const FPhysicsActorReference_Chaos& InActorReference)
+FTransform FPhysInterface_Chaos::GetComTransformLocal_AssumesLocked(const FPhysicsActorHandle& InActorReference)
 {
 	// #todo : Implement
 	return FTransform();
 }
 
-FVector FPhysInterface_Chaos::GetLocalInertiaTensor_AssumesLocked(const FPhysicsActorReference_Chaos& InActorReference)
+FVector FPhysInterface_Chaos::GetLocalInertiaTensor_AssumesLocked(const FPhysicsActorHandle& InActorReference)
 {
 	// #todo : Implement
 	return FVector(1);
 }
 
-FBox FPhysInterface_Chaos::GetBounds_AssumesLocked(const FPhysicsActorReference_Chaos& InActorReference)
+FBox FPhysInterface_Chaos::GetBounds_AssumesLocked(const FPhysicsActorHandle& InActorReference)
 {
 	// #todo : Implement
 	//const auto& Box = InActorReference.GetScene()->Scene.GetSolver()->GetRigidParticles().Geometry(InActorReference.GetScene()->GetIndexFromId(InActorReference.GetId()))->BoundingBox();
     return FBox(FVector(-0.5), FVector(0.5));
 }
 
-void FPhysInterface_Chaos::SetLinearDamping_AssumesLocked(const FPhysicsActorReference_Chaos& InActorReference, float InDamping)
+void FPhysInterface_Chaos::SetLinearDamping_AssumesLocked(const FPhysicsActorHandle& InActorReference, float InDamping)
 {
 
 }
 
-void FPhysInterface_Chaos::SetAngularDamping_AssumesLocked(const FPhysicsActorReference_Chaos& InActorReference, float InDamping)
+void FPhysInterface_Chaos::SetAngularDamping_AssumesLocked(const FPhysicsActorHandle& InActorReference, float InDamping)
 {
 
 }
 
-void FPhysInterface_Chaos::AddImpulse_AssumesLocked(const FPhysicsActorReference_Chaos& InActorReference, const FVector& InForce)
+void FPhysInterface_Chaos::AddImpulse_AssumesLocked(const FPhysicsActorHandle& InActorReference, const FVector& InForce)
 {
 	// #todo : Implement
     //InActorReference.GetScene()->AddForce(InForce, InActorReference.GetId());
 }
 
-void FPhysInterface_Chaos::AddAngularImpulseInRadians_AssumesLocked(const FPhysicsActorReference_Chaos& InActorReference, const FVector& InTorque)
+void FPhysInterface_Chaos::AddAngularImpulseInRadians_AssumesLocked(const FPhysicsActorHandle& InActorReference, const FVector& InTorque)
 {
 	// #todo : Implement
     //InActorReference.GetScene()->AddTorque(InTorque, InActorReference.GetId());
 }
 
-void FPhysInterface_Chaos::AddVelocity_AssumesLocked(const FPhysicsActorReference_Chaos& InActorReference, const FVector& InForce)
+void FPhysInterface_Chaos::AddVelocity_AssumesLocked(const FPhysicsActorHandle& InActorReference, const FVector& InForce)
 {	
 	// #todo : Implement
     //InActorReference.GetScene()->AddForce(InForce * InActorReference.GetScene()->Scene.GetSolver()->GetRigidParticles().M(InActorReference.GetScene()->GetIndexFromId(InActorReference.GetId())), InActorReference.GetId());
 }
 
-void FPhysInterface_Chaos::AddAngularVelocityInRadians_AssumesLocked(const FPhysicsActorReference_Chaos& InActorReference, const FVector& InTorque)
+void FPhysInterface_Chaos::AddAngularVelocityInRadians_AssumesLocked(const FPhysicsActorHandle& InActorReference, const FVector& InTorque)
 {
 	// #todo : Implement
     //InActorReference.GetScene()->AddTorque(InActorReference.GetScene()->Scene.GetSolver()->GetRigidParticles().I(InActorReference.GetScene()->GetIndexFromId(InActorReference.GetId())) * Chaos::TVector<float, 3>(InTorque), InActorReference.GetId());
 }
 
-void FPhysInterface_Chaos::AddImpulseAtLocation_AssumesLocked(const FPhysicsActorReference_Chaos& InActorReference, const FVector& InImpulse, const FVector& InLocation)
+void FPhysInterface_Chaos::AddImpulseAtLocation_AssumesLocked(const FPhysicsActorHandle& InActorReference, const FVector& InImpulse, const FVector& InLocation)
 {
     // @todo(mlentine): We don't currently have a way to apply an instantaneous force. Do we need this?
 }
 
-void FPhysInterface_Chaos::AddRadialImpulse_AssumesLocked(const FPhysicsActorReference_Chaos& InActorReference, const FVector& InOrigin, float InRadius, float InStrength, ERadialImpulseFalloff InFalloff, bool bInVelChange)
+void FPhysInterface_Chaos::AddRadialImpulse_AssumesLocked(const FPhysicsActorHandle& InActorReference, const FVector& InOrigin, float InRadius, float InStrength, ERadialImpulseFalloff InFalloff, bool bInVelChange)
 {
     // @todo(mlentine): We don't currently have a way to apply an instantaneous force. Do we need this?
 }
 
-bool FPhysInterface_Chaos::IsGravityEnabled_AssumesLocked(const FPhysicsActorReference_Chaos& InActorReference)
+bool FPhysInterface_Chaos::IsGravityEnabled_AssumesLocked(const FPhysicsActorHandle& InActorReference)
 {
     // @todo(mlentine): Gravity is system wide currently. This should change.
     return true;
 }
-void FPhysInterface_Chaos::SetGravityEnabled_AssumesLocked(const FPhysicsActorReference_Chaos& InActorReference, bool bEnabled)
+void FPhysInterface_Chaos::SetGravityEnabled_AssumesLocked(const FPhysicsActorHandle& InActorReference, bool bEnabled)
 {
     // @todo(mlentine): Gravity is system wide currently. This should change.
 }
 
-float FPhysInterface_Chaos::GetSleepEnergyThreshold_AssumesLocked(const FPhysicsActorReference_Chaos& InActorReference)
+float FPhysInterface_Chaos::GetSleepEnergyThreshold_AssumesLocked(const FPhysicsActorHandle& InActorReference)
 {
     return 0;
 }
-void FPhysInterface_Chaos::SetSleepEnergyThreshold_AssumesLocked(const FPhysicsActorReference_Chaos& InActorReference, float InEnergyThreshold)
+void FPhysInterface_Chaos::SetSleepEnergyThreshold_AssumesLocked(const FPhysicsActorHandle& InActorReference, float InEnergyThreshold)
 {
 }
 
-void FPhysInterface_Chaos::SetMass_AssumesLocked(const FPhysicsActorReference_Chaos& InHandle, float InMass)
+void FPhysInterface_Chaos::SetMass_AssumesLocked(const FPhysicsActorHandle& InHandle, float InMass)
 {
 	// #todo : Implement
 	//LocalParticles.M(Index) = InMass;
 }
 
-void FPhysInterface_Chaos::SetMassSpaceInertiaTensor_AssumesLocked(const FPhysicsActorReference_Chaos& InHandle, const FVector& InTensor)
+void FPhysInterface_Chaos::SetMassSpaceInertiaTensor_AssumesLocked(const FPhysicsActorHandle& InHandle, const FVector& InTensor)
 {
 	// #todo : Implement
 	//LocalParticles.I(Index).M[0][0] = InTensor[0];
@@ -482,62 +458,62 @@ void FPhysInterface_Chaos::SetMassSpaceInertiaTensor_AssumesLocked(const FPhysic
     //LocalParticles.I(Index).M[2][2] = InTensor[2];
 }
 
-void FPhysInterface_Chaos::SetComLocalPose_AssumesLocked(const FPhysicsActorReference_Chaos& InHandle, const FTransform& InComLocalPose)
+void FPhysInterface_Chaos::SetComLocalPose_AssumesLocked(const FPhysicsActorHandle& InHandle, const FTransform& InComLocalPose)
 {
     //@todo(mlentine): What is InComLocalPose? If the center of an object is not the local pose then many things break including the three vector represtnation of inertia.
 }
 
-float FPhysInterface_Chaos::GetStabilizationEnergyThreshold_AssumesLocked(const FPhysicsActorReference_Chaos& InHandle)
+float FPhysInterface_Chaos::GetStabilizationEnergyThreshold_AssumesLocked(const FPhysicsActorHandle& InHandle)
 {
 	// #todo : Implement
 	return 0.0f;
 }
 
-void FPhysInterface_Chaos::SetStabilizationEnergyThreshold_AssumesLocked(const FPhysicsActorReference_Chaos& InHandle, float InThreshold)
+void FPhysInterface_Chaos::SetStabilizationEnergyThreshold_AssumesLocked(const FPhysicsActorHandle& InHandle, float InThreshold)
 {
 	// #todo : Implement
 }
 
-uint32 FPhysInterface_Chaos::GetSolverPositionIterationCount_AssumesLocked(const FPhysicsActorReference_Chaos& InHandle)
+uint32 FPhysInterface_Chaos::GetSolverPositionIterationCount_AssumesLocked(const FPhysicsActorHandle& InHandle)
 {
 	// #todo : Implement
 	return 0;
 }
 
-void FPhysInterface_Chaos::SetSolverPositionIterationCount_AssumesLocked(const FPhysicsActorReference_Chaos& InHandle, uint32 InSolverIterationCount)
+void FPhysInterface_Chaos::SetSolverPositionIterationCount_AssumesLocked(const FPhysicsActorHandle& InHandle, uint32 InSolverIterationCount)
 {
 	// #todo : Implement
 }
 
-uint32 FPhysInterface_Chaos::GetSolverVelocityIterationCount_AssumesLocked(const FPhysicsActorReference_Chaos& InHandle)
+uint32 FPhysInterface_Chaos::GetSolverVelocityIterationCount_AssumesLocked(const FPhysicsActorHandle& InHandle)
 {
 	// #todo : Implement
 	return 0;
 }
 
-void FPhysInterface_Chaos::SetSolverVelocityIterationCount_AssumesLocked(const FPhysicsActorReference_Chaos& InHandle, uint32 InSolverIterationCount)
+void FPhysInterface_Chaos::SetSolverVelocityIterationCount_AssumesLocked(const FPhysicsActorHandle& InHandle, uint32 InSolverIterationCount)
 {
 	// #todo : Implement
 }
 
-float FPhysInterface_Chaos::GetWakeCounter_AssumesLocked(const FPhysicsActorReference_Chaos& InHandle)
+float FPhysInterface_Chaos::GetWakeCounter_AssumesLocked(const FPhysicsActorHandle& InHandle)
 {
 	// #todo : Implement
 	return 0.0f;
 }
 
-void FPhysInterface_Chaos::SetWakeCounter_AssumesLocked(const FPhysicsActorReference_Chaos& InHandle, float InWakeCounter)
+void FPhysInterface_Chaos::SetWakeCounter_AssumesLocked(const FPhysicsActorHandle& InHandle, float InWakeCounter)
 {
 	// #todo : Implement
 }
 
-SIZE_T FPhysInterface_Chaos::GetResourceSizeEx(const FPhysicsActorReference_Chaos& InActorRef)
+SIZE_T FPhysInterface_Chaos::GetResourceSizeEx(const FPhysicsActorHandle& InActorRef)
 {
-    return sizeof(FPhysicsActorReference_Chaos);
+    return sizeof(FPhysicsActorHandle);
 }
 	
 // Constraints
-FPhysicsConstraintReference_Chaos FPhysInterface_Chaos::CreateConstraint(const FPhysicsActorReference_Chaos& InActorRef1, const FPhysicsActorReference_Chaos& InActorRef2, const FTransform& InLocalFrame1, const FTransform& InLocalFrame2)
+FPhysicsConstraintReference_Chaos FPhysInterface_Chaos::CreateConstraint(const FPhysicsActorHandle& InActorRef1, const FPhysicsActorHandle& InActorRef2, const FTransform& InLocalFrame1, const FTransform& InLocalFrame2)
 {
 	// #todo : Implement
 	FPhysicsConstraintReference_Chaos ConstraintRef;
@@ -848,109 +824,183 @@ FPhysicsShapeHandle FPhysInterface_Chaos::CreateShape(physx::PxGeometry* InGeom,
 {
 	// #todo : Implement
 	// @todo(mlentine): Should we be doing anything with the InGeom here?
-    FPhysicsActorHandle NewActor;
-	return {nullptr,bSimulation ,bQuery,NewActor };
+    FPhysicsActorHandle NewActor = nullptr;
+	return { nullptr, bSimulation, bQuery, NewActor };
+}
+
+const FBodyInstance* FPhysInterface_Chaos::ShapeToOriginalBodyInstance(const FBodyInstance* InCurrentInstance, const Chaos::TPerShapeData<float, 3>* InShape)
+{
+	//question: this is identical to physx version, should it be in body instance?
+	check(InCurrentInstance);
+	check(InShape);
+
+	const FBodyInstance* TargetInstance = InCurrentInstance->WeldParent ? InCurrentInstance->WeldParent : InCurrentInstance;
+	const FBodyInstance* OutInstance = TargetInstance;
+
+	if (const TMap<FPhysicsShapeHandle, FBodyInstance::FWeldInfo>* WeldInfo = InCurrentInstance->GetCurrentWeldInfo())
+	{
+		for (const TPair<FPhysicsShapeHandle, FBodyInstance::FWeldInfo>& Pair : *WeldInfo)
+		{
+			if (Pair.Key.Shape == InShape)
+			{
+				TargetInstance = Pair.Value.ChildBI;
+			}
+		}
+	}
+
+	return TargetInstance;
 }
+
 	
 void FPhysInterface_Chaos::AddGeometry(FPhysicsActorHandle& InActor, const FGeometryAddParams& InParams, TArray<FPhysicsShapeHandle>* OutOptShapes)
 {
 	const FVector& Scale = InParams.Scale;
-	FBodyInstancePhysicsObject * PhysicsObject = InActor.GetPhysicsObject();
-	if (InParams.Geometry!=nullptr && ( PhysicsObject!=nullptr || OutOptShapes!=nullptr ) ) 
+	TArray<TUniquePtr<Chaos::TImplicitObject<float, 3>>> Geoms;
+	Chaos::TShapesArray<float, 3> Shapes;
+
+	auto NewShapeHelper = [&InParams](const Chaos::TImplicitObject<float, 3>* InGeom, bool bComplexShape = false)
+	{
+		auto NewShape = MakeUnique<Chaos::TPerShapeData<float, 3>>();
+		NewShape->Geometry = InGeom;
+		NewShape->QueryData = bComplexShape ? InParams.CollisionData.CollisionFilterData.QueryComplexFilter : InParams.CollisionData.CollisionFilterData.QuerySimpleFilter;
+		NewShape->SimData = InParams.CollisionData.CollisionFilterData.SimFilter;
+		return NewShape;
+	};
+
+	if (InActor) 
 	{
 		for (uint32 i = 0; i < static_cast<uint32>(InParams.Geometry->SphereElems.Num()); ++i)
 		{
-			check(Scale[0] == Scale[1] && Scale[1] == Scale[2]);
+			ensure(FMath::IsNearlyEqual(Scale[0],Scale[1]) && FMath::IsNearlyEqual(Scale[1],Scale[2]) );
 			const auto& CollisionSphere = InParams.Geometry->SphereElems[i];
-			Chaos::TImplicitObject<float, 3>* ImplicitSphere = new Chaos::TSphere<float, 3>(Chaos::TVector<float, 3>(0.f, 0.f, 0.f), CollisionSphere.Radius * Scale[0]);
-			if (PhysicsObject) PhysicsObject->ImplicitObjects_GameThread.Add(ImplicitSphere);
-			else if (OutOptShapes) OutOptShapes->Add({ImplicitSphere,true,true,InActor});
+			auto ImplicitSphere = MakeUnique<Chaos::TSphere<float, 3>>(Chaos::TVector<float, 3>(0.f, 0.f, 0.f), CollisionSphere.Radius * Scale[0]);
+			auto NewShape = NewShapeHelper(ImplicitSphere.Get());
+			Shapes.Emplace(MoveTemp(NewShape));
+
+			if (OutOptShapes) OutOptShapes->Add({NewShape.Get(),true,true,InActor});
+			Geoms.Add(MoveTemp(ImplicitSphere));
+
 		}
 		for (uint32 i = 0; i < static_cast<uint32>(InParams.Geometry->BoxElems.Num()); ++i)
 		{
 			const auto& Box = InParams.Geometry->BoxElems[i];
 			Chaos::TVector<float, 3> half_extents = Scale * Chaos::TVector<float, 3>(Box.X / 2.f, Box.Y / 2.f, Box.Z / 2.f);
-			Chaos::TImplicitObject<float, 3>* ImplicitBox = new Chaos::TBox<float, 3>(-half_extents, half_extents);
-			if (PhysicsObject) PhysicsObject->ImplicitObjects_GameThread.Add(ImplicitBox);
-			else if (OutOptShapes) OutOptShapes->Add({ImplicitBox,true,true,InActor});
+			auto ImplicitBox = MakeUnique<Chaos::TBox<float, 3>>(-half_extents, half_extents);
+			auto NewShape = NewShapeHelper(ImplicitBox.Get());
+			Shapes.Emplace(MoveTemp(NewShape));
+
+			if (OutOptShapes) OutOptShapes->Add({NewShape.Get(),true,true,InActor});
+			Geoms.Add(MoveTemp(ImplicitBox));
+
 		}
 		for (uint32 i = 0; i < static_cast<uint32>(InParams.Geometry->SphylElems.Num()); ++i)
 		{
-			check(Scale[0] == Scale[1] && Scale[1] == Scale[2]);
-			const auto& TCapsule = InParams.Geometry->SphylElems[i];
-			if (TCapsule.Length == 0)
+			const FKSphylElem& UnscaledSphyl = InParams.Geometry->SphylElems[i];
+
+			const FKSphylElem ScaledSphylElem = UnscaledSphyl.GetFinalScaled(Scale, InParams.LocalTransform);
+			float HalfHeight = FMath::Max(ScaledSphylElem.Length * 0.5f, KINDA_SMALL_NUMBER);
+			const float Radius = FMath::Max(ScaledSphylElem.Radius, KINDA_SMALL_NUMBER);
+
+			if (HalfHeight < KINDA_SMALL_NUMBER)
 			{
-				Chaos::TSphere<float, 3> * ImplicitSphere = new Chaos::TSphere<float, 3>(Chaos::TVector<float, 3>(0), TCapsule.Radius * Scale[0]);
-				if (PhysicsObject) PhysicsObject->ImplicitObjects_GameThread.Add(ImplicitSphere);
-				else if (OutOptShapes) OutOptShapes->Add({ImplicitSphere,true,true,InActor});
+				//not a capsule just use a sphere
+				auto ImplicitSphere = MakeUnique<Chaos::TSphere<float, 3>>(Chaos::TVector<float, 3>(0), Radius);
+				auto NewShape = NewShapeHelper(ImplicitSphere.Get());
+				Shapes.Emplace(MoveTemp(NewShape));
+
+				if (OutOptShapes) OutOptShapes->Add({NewShape.Get(),true,true,InActor});
+				Geoms.Add(MoveTemp(ImplicitSphere));
+
 			}
 			else
 			{
-				Chaos::TVector<float, 3> half_extents(0, 0, TCapsule.Length / 2 * Scale[0]);
+				Chaos::TVector<float, 3> HalfExtents(0, 0, HalfHeight);
 
-				Chaos::TCylinder<float> * ImplicitCylinder = new Chaos::TCylinder<float>(-half_extents, half_extents, TCapsule.Radius * Scale[0]);
-				if (PhysicsObject) PhysicsObject->ImplicitObjects_GameThread.Add(ImplicitCylinder);
-				else if (OutOptShapes) OutOptShapes->Add({ImplicitCylinder,true,true,InActor});
+				auto ImplicitCapsule = MakeUnique<Chaos::TCapsule<float>>(-HalfExtents, HalfExtents, Radius);
+				auto NewShape = NewShapeHelper(ImplicitCapsule.Get());
+				Shapes.Emplace(MoveTemp(NewShape));
 
-				Chaos::TSphere<float, 3> * ImplicitSphereA = new Chaos::TSphere<float, 3>(-half_extents, TCapsule.Radius * Scale[0]);
-				if (PhysicsObject) PhysicsObject->ImplicitObjects_GameThread.Add(ImplicitSphereA);
-				else if (OutOptShapes) OutOptShapes->Add({ImplicitSphereA,true,true,InActor});
-
-				Chaos::TSphere<float, 3> * ImplicitSphereB = new Chaos::TSphere<float, 3>(half_extents, TCapsule.Radius * Scale[0]);
-				if (PhysicsObject) PhysicsObject->ImplicitObjects_GameThread.Add(ImplicitSphereB);
-				else if (OutOptShapes) OutOptShapes->Add({ImplicitSphereB,true,true,InActor});
+				if (OutOptShapes) OutOptShapes->Add({ NewShape.Get(),true,true,InActor });
+				Geoms.Add(MoveTemp(ImplicitCapsule));
 			}
 		}
 #if 0
 		for (uint32 i = 0; i < static_cast<uint32>(InParams.Geometry->TaperedCapsuleElems.Num()); ++i)
 		{
-			check(Scale[0] == Scale[1] && Scale[1] == Scale[2]);
+			ensure(FMath::IsNearlyEqual(Scale[0], Scale[1]) && FMath::IsNearlyEqual(Scale[1], Scale[2]));
 			const auto& TCapsule = InParams.Geometry->TaperedCapsuleElems[i];
 			if (TCapsule.Length == 0)
 			{
 				Chaos::TSphere<float, 3> * ImplicitSphere = new Chaos::TSphere<float, 3>(-half_extents, TCapsule.Radius * Scale[0]);
-				if (PhysicsObject) PhysicsObject->ImplicitObjects_GameThread.Add(ImplicitSphere);
+				if (PhysicsProxy) PhysicsProxy->ImplicitObjects_GameThread.Add(ImplicitSphere);
 				else if (OutOptShapes) OutOptShapes->Add({ImplicitSphere,true,true,InActor});
 			}
 			else
 			{
 				Chaos::TVector<float, 3> half_extents(0, 0, TCapsule.Length / 2 * Scale[0]);
-				Chaos::TCylinder<float> * ImplicitCylinder = new Chaos::TCylinder<float>(-half_extents, half_extents, TCapsule.Radius * Scale[0]);
-				if (PhysicsObject) PhysicsObject->ImplicitObjects_GameThread.Add(ImplicitSphere);
+				auto ImplicitCylinder = MakeUnique<Chaos::TCylinder<float>>(-half_extents, half_extents, TCapsule.Radius * Scale[0]);
+				if (PhysicsProxy) PhysicsProxy->ImplicitObjects_GameThread.Add(MoveTemp(ImplicitSphere));
 				else if (OutOptShapes) OutOptShapes->Add({ImplicitSphere,true,true,InActor});
 
-				Chaos::TSphere<float, 3> * ImplicitSphereA = new Chaos::TSphere<float, 3>(-half_extents, TCapsule.Radius * Scale[0]);
-				if (PhysicsObject) PhysicsObject->ImplicitObjects_GameThread.Add(ImplicitSphereA);
+				auto ImplicitSphereA = MakeUnique<Chaos::TSphere<float, 3>>(-half_extents, TCapsule.Radius * Scale[0]);
+				if (PhysicsProxy) PhysicsProxy->ImplicitObjects_GameThread.Add(MoveTemp(ImplicitSphereA));
 				else if (OutOptShapes) OutOptShapes->Add({ImplicitSphereA,true,true,InActor});
 
-				Chaos::TSphere<float, 3> * ImplicitSphereB = new Chaos::TSphere<float, 3>(half_extents, TCapsule.Radius * Scale[0]);
-				if (PhysicsObject) PhysicsObject->ImplicitObjects_GameThread.Add(ImplicitSphereB);
+				auto ImplicitSphereB = MakeUnique<Chaos::TSphere<float, 3>>(half_extents, TCapsule.Radius * Scale[0]);
+				if (PhysicsProxy) PhysicsProxy->ImplicitObjects_GameThread.Add(MoveTemp(ImplicitSphereB));
 				else if (OutOptShapes) OutOptShapes->Add({ImplicitSphereB,true,true,InActor});
 			}
 		}
 #endif
 
-#if WITH_PHYSX
 		for (uint32 i = 0; i < static_cast<uint32>(InParams.Geometry->ConvexElems.Num()); ++i)
 		{
-			const auto& CollisionBody = InParams.Geometry->ConvexElems[i];
-			Chaos::TImplicitObject<float, 3>* Implicit = ConvertPhysXMeshToLevelset(CollisionBody.GetConvexMesh(), Scale).Release();
-			if (PhysicsObject) PhysicsObject->ImplicitObjects_GameThread.Add(Implicit);
-			else if (OutOptShapes) OutOptShapes->Add({Implicit,true,true,InActor});
+			const FKConvexElem& CollisionBody = InParams.Geometry->ConvexElems[i];
+			if (const auto& ConvexImplicit = CollisionBody.GetChaosConvexMesh())
+			{
+				auto Implicit = MakeUnique<Chaos::TImplicitObjectScaled<float, 3>>(MakeSerializable(ConvexImplicit), Scale);
+				auto NewShape = NewShapeHelper(Implicit.Get());
+				Shapes.Emplace(MoveTemp(NewShape));
+
+				if (OutOptShapes) OutOptShapes->Add({ NewShape.Get(),true,true,InActor });
+				Geoms.Add(MoveTemp(Implicit));
+			}
+		}
+
+		for (const auto& ChaosTriMesh : InParams.ChaosTriMeshes)
+		{
+			auto Implicit = MakeUnique<Chaos::TImplicitObjectScaled<float, 3>>(MakeSerializable(ChaosTriMesh), Scale);
+			auto NewShape = NewShapeHelper(Implicit.Get(), true);
+			Shapes.Emplace(MoveTemp(NewShape));
+
+			if (OutOptShapes) OutOptShapes->Add({ NewShape.Get(),true,true,InActor });
+			Geoms.Add(MoveTemp(Implicit));
 		}
-#endif
 	}
 	else
 	{
 #if WITH_PHYSX
 		for (const auto& PhysXMesh : InParams.TriMeshes)
 		{
-			Chaos::TImplicitObject<float, 3>* Implicit = ConvertPhysXMeshToLevelset(PhysXMesh, Scale).Release();
-			if (PhysicsObject) PhysicsObject->ImplicitObjects_GameThread.Add(Implicit);
-			else if (OutOptShapes) OutOptShapes->Add({Implicit,true,true,InActor});
+			auto Implicit = ConvertPhysXMeshToLevelset(PhysXMesh, Scale);
+			auto NewShape = NewShapeHelper(Implicit.Get(), true);
+			Shapes.Emplace(MoveTemp(NewShape));
+
+			if (OutOptShapes) OutOptShapes->Add({NewShape.Get(),true,true,InActor});
+			Geoms.Add(MoveTemp(Implicit));
+
 		}
 #endif
 	}
+
+#if WITH_CHAOS
+	if (InActor)
+	{
+		//todo: we should not be creating unique geometry per actor
+		InActor->SetGeometry(MakeUnique<Chaos::TImplicitObjectUnion<float, 3>>(MoveTemp(Geoms)));
+		InActor->SetShapesArray(MoveTemp(Shapes));
+	}
+#endif
 }
 
 
@@ -959,8 +1009,56 @@ void FPhysInterface_Chaos::AddGeometry(FPhysicsActorHandle& InActor, const FGeom
 // @todo(mlentine,brice): We probably need to actually duplicate the data here, add virtual TImplicitObject::NewCopy()
 FPhysicsShapeHandle FPhysInterface_Chaos::CloneShape(const FPhysicsShapeHandle& InShape)
 {
-	FPhysicsActorHandle NewActor; // why zero and not the default INDEX_NONE?
-	return { InShape.Object,InShape.bSimulation, InShape.bQuery, NewActor };
+	FPhysicsActorHandle NewActor = nullptr; // why zero and not the default INDEX_NONE?
+	return { InShape.Shape, InShape.bSimulation, InShape.bQuery, NewActor };
+}
+
+FPhysicsGeometryCollection_Chaos FPhysInterface_Chaos::GetGeometryCollection(const FPhysicsShapeHandle& InShape)
+{
+	FPhysicsGeometryCollection_Chaos NewCollection(InShape);
+	return NewCollection;
+}
+
+FPhysicsGeometryCollection_Chaos::~FPhysicsGeometryCollection_Chaos() = default;
+FPhysicsGeometryCollection_Chaos::FPhysicsGeometryCollection_Chaos(FPhysicsGeometryCollection_Chaos&& Steal) = default;
+
+ECollisionShapeType FPhysicsGeometryCollection_Chaos::GetType() const
+{
+	return GetImplicitType(Geom);
+}
+
+const Chaos::TImplicitObject<float, 3>& FPhysicsGeometryCollection_Chaos::GetGeometry() const
+{
+	return Geom;
+}
+
+const Chaos::TBox<float, 3>& FPhysicsGeometryCollection_Chaos::GetBoxGeometry() const
+{
+	return Geom.GetObjectChecked<Chaos::TBox<float, 3>>();
+}
+
+const Chaos::TSphere<float, 3>&  FPhysicsGeometryCollection_Chaos::GetSphereGeometry() const
+{
+	return Geom.GetObjectChecked<Chaos::TSphere<float, 3>>();
+}
+const Chaos::TCapsule<float>&  FPhysicsGeometryCollection_Chaos::GetCapsuleGeometry() const
+{
+	return Geom.GetObjectChecked<Chaos::TCapsule<float>>();
+}
+
+const Chaos::TConvex<float, 3>& FPhysicsGeometryCollection_Chaos::GetConvexGeometry() const
+{
+	return Geom.GetObjectChecked<Chaos::TConvex<float, 3>>();
+}
+
+const Chaos::TTriangleMeshImplicitObject<float>& FPhysicsGeometryCollection_Chaos::GetTriMeshGeometry() const
+{
+	return Geom.GetObjectChecked<Chaos::TTriangleMeshImplicitObject<float>>();
+}
+
+FPhysicsGeometryCollection_Chaos::FPhysicsGeometryCollection_Chaos(const FPhysicsShapeReference_Chaos& InShape)
+	: Geom(InShape.GetGeometry())
+{
 }
 
 FCollisionFilterData FPhysInterface_Chaos::GetSimulationFilter(const FPhysicsShapeHandle& InShape)
@@ -973,6 +1071,16 @@ FCollisionFilterData FPhysInterface_Chaos::GetQueryFilter(const FPhysicsShapeHan
 	return FCollisionFilterData();
 }
 
+void FPhysInterface_Chaos::SetQueryFilter(const FPhysicsShapeReference_Chaos& InShapeRef, const FCollisionFilterData& InFilter)
+{
+	InShapeRef.Shape->QueryData = InFilter;
+}
+
+void FPhysInterface_Chaos::SetSimulationFilter(const FPhysicsShapeReference_Chaos& InShapeRef, const FCollisionFilterData& InFilter)
+{
+	InShapeRef.Shape->SimData = InFilter;
+}
+
 bool FPhysInterface_Chaos::IsSimulationShape(const FPhysicsShapeHandle& InShape)
 {
     return InShape.bSimulation;
@@ -983,13 +1091,13 @@ bool FPhysInterface_Chaos::IsQueryShape(const FPhysicsShapeHandle& InShape)
     return InShape.bQuery;
 }
 
-bool FPhysInterface_Chaos::IsShapeType(const FPhysicsShapeHandle& InShape, ECollisionShapeType InType)
+bool FPhysInterface_Chaos::IsShapeType(const FPhysicsShapeReference_Chaos& InShapeRef, ECollisionShapeType InType)
 {
-    if (InType == ECollisionShapeType::Box && InShape.Object->GetType() == Chaos::ImplicitObjectType::Box)
+    if (InType == ECollisionShapeType::Box && InShapeRef.Shape->Geometry->GetType() == Chaos::ImplicitObjectType::Box)
     {
         return true;
     }
-    if (InType == ECollisionShapeType::Sphere && InShape.Object->GetType() == Chaos::ImplicitObjectType::Sphere)
+    if (InType == ECollisionShapeType::Sphere && InShapeRef.Shape->Geometry->GetType() == Chaos::ImplicitObjectType::Sphere)
     {
         return true;
     }
@@ -997,25 +1105,25 @@ bool FPhysInterface_Chaos::IsShapeType(const FPhysicsShapeHandle& InShape, EColl
     return false;
 }
 
-ECollisionShapeType FPhysInterface_Chaos::GetShapeType(const FPhysicsShapeHandle& InShape)
+ECollisionShapeType FPhysInterface_Chaos::GetShapeType(const FPhysicsShapeReference_Chaos& InShapeRef)
 {
-    if (InShape.Object->GetType() == Chaos::ImplicitObjectType::Box)
+    if (InShapeRef.Shape->Geometry->GetType() == Chaos::ImplicitObjectType::Box)
     {
         return ECollisionShapeType::Box;
     }
-    if (InShape.Object->GetType() == Chaos::ImplicitObjectType::Sphere)
+    if (InShapeRef.Shape->Geometry->GetType() == Chaos::ImplicitObjectType::Sphere)
     {
         return ECollisionShapeType::Sphere;
     }
     return ECollisionShapeType::None;
 }
 
-FTransform FPhysInterface_Chaos::GetLocalTransform(const FPhysicsShapeHandle& InShape)
+FTransform FPhysInterface_Chaos::GetLocalTransform(const FPhysicsShapeReference_Chaos& InShapeRef)
 {
     // Transforms are baked into the object so there is never a local transform
-    if (InShape.Object->GetType() == Chaos::ImplicitObjectType::Transformed && InShape.ActorRef.IsValid())
+    if (InShapeRef.Shape->Geometry->GetType() == Chaos::ImplicitObjectType::Transformed && FPhysicsInterface::IsValid(InShapeRef.ActorRef))
     {
-        return InShape.Object->GetObject<Chaos::TImplicitObjectTransformed<float, 3>>()->GetTransform();
+        return InShapeRef.Shape->Geometry->GetObject<Chaos::TImplicitObjectTransformed<float, 3>>()->GetTransform();
     }
     else
     {
@@ -1071,80 +1179,440 @@ void FPhysInterface_Chaos::CalculateMassPropertiesFromShapeCollection(physx::PxM
 
 #endif
 
-bool FPhysInterface_Chaos::LineTrace_Geom(FHitResult& OutHit, const FBodyInstance* InInstance, const FVector& InStart, const FVector& InEnd, bool bTraceComplex, bool bExtractPhysMaterial)
+bool FPhysInterface_Chaos::LineTrace_Geom(FHitResult& OutHit, const FBodyInstance* InInstance, const FVector& WorldStart, const FVector& WorldEnd, bool bTraceComplex, bool bExtractPhysMaterial)
 {
-	OutHit.TraceStart = InStart;
-	OutHit.TraceEnd = InEnd;
-    
-    TArray<FPhysicsShapeReference_Chaos> OutShapes;
-    InInstance->GetAllShapes_AssumesLocked(OutShapes);
-    check(OutShapes.Num() == 1);
-    const auto& Result = OutShapes[0].Object->FindClosestIntersection(InStart, InEnd, 0);
-    if (Result.Second)
-    {
-        // @todo(mlentine): What should we fill here?
-        OutHit.ImpactPoint = Result.First;
-        OutHit.ImpactNormal = OutShapes[0].Object->Normal(OutHit.ImpactPoint);
-    }
-    return Result.Second;
+	// Need an instandce to trace against
+	check(InInstance);
+
+	OutHit.TraceStart = WorldStart;
+	OutHit.TraceEnd = WorldEnd;
+
+	bool bHitSomething = false;
+
+	const FVector Delta = WorldEnd - WorldStart;
+	const float DeltaMag = Delta.Size();
+	if (DeltaMag > KINDA_SMALL_NUMBER)
+	{
+		{
+			// #PHYS2 Really need a concept for "multi" locks here - as we're locking ActorRef but not TargetInstance->ActorRef
+			FPhysicsCommand::ExecuteRead(InInstance->ActorHandle, [&](const FPhysicsActorHandle& Actor)
+			{
+				// If we're welded then the target instance is actually our parent
+				const FBodyInstance* TargetInstance = InInstance->WeldParent ? InInstance->WeldParent : InInstance;
+				if(const Chaos::TGeometryParticle<float, 3>* RigidBody = TargetInstance->ActorHandle)
+				{
+					FRaycastHit BestHit;
+					BestHit.Distance = FLT_MAX;
+
+					// Get all the shapes from the actor
+					PhysicsInterfaceTypes::FInlineShapeArray Shapes;
+					const int32 NumShapes = FillInlineShapeArray_AssumesLocked(Shapes, Actor);
+
+					const FTransform WorldTM(RigidBody->R(), RigidBody->X());
+					const FVector LocalStart = WorldTM.InverseTransformPositionNoScale(WorldStart);
+					const FVector LocalDelta = WorldTM.InverseTransformVectorNoScale(Delta);
+
+					// Iterate over each shape
+					for (int32 ShapeIdx = 0; ShapeIdx < NumShapes; ShapeIdx++)
+					{
+						// #PHYS2 - SHAPES - Resolve this single cast case
+						FPhysicsShapeReference_Chaos& ShapeRef = Shapes[ShapeIdx];
+						Chaos::TPerShapeData<float, 3>* Shape = ShapeRef.Shape;
+						check(Shape);
+
+						if (TargetInstance->IsShapeBoundToBody(ShapeRef) == false)
+						{
+							continue;
+						}
+
+						// Filter so we trace against the right kind of collision
+						FCollisionFilterData ShapeFilter = Shape->QueryData;
+						const bool bShapeIsComplex = (ShapeFilter.Word3 & EPDF_ComplexCollision) != 0;
+						const bool bShapeIsSimple = (ShapeFilter.Word3 & EPDF_SimpleCollision) != 0;
+						if ((bTraceComplex && bShapeIsComplex) || (!bTraceComplex && bShapeIsSimple))
+						{
+
+							float Distance;
+							Chaos::TVector<float, 3> LocalPosition;
+							Chaos::TVector<float, 3> LocalNormal;
+
+							int32 FaceIndex;
+							if (Shape->Geometry->Raycast(LocalStart, LocalDelta / DeltaMag, DeltaMag, 0, Distance, LocalPosition, LocalNormal, FaceIndex))
+							{
+								if (Distance < BestHit.Distance)
+								{
+									BestHit.Distance = Distance;
+									BestHit.WorldNormal = LocalNormal;	//will convert to world when best is chosen
+									BestHit.WorldPosition = LocalPosition;
+									BestHit.Shape = Shape;
+									BestHit.Actor = Actor;
+								}
+							}
+						}
+					}
+
+					if (BestHit.Distance < FLT_MAX)
+					{
+						BestHit.WorldNormal = WorldTM.TransformVectorNoScale(BestHit.WorldNormal);
+						BestHit.WorldPosition = WorldTM.TransformPositionNoScale(BestHit.WorldPosition);
+
+						// we just like to make sure if the hit is made, set to test touch
+						FCollisionFilterData QueryFilter;
+						QueryFilter.Word2 = 0xFFFFF;
+
+						FTransform StartTM(WorldStart);
+						const UPrimitiveComponent* OwnerComponentInst = InInstance->OwnerComponent.Get();
+						ConvertQueryImpactHit(OwnerComponentInst ? OwnerComponentInst->GetWorld() : nullptr, BestHit, OutHit, DeltaMag, QueryFilter, WorldStart, WorldEnd, nullptr, StartTM, true, bExtractPhysMaterial);
+						bHitSomething = true;
+					}
+				}
+			});
+		}
+	}
+
+	return bHitSomething;
 }
 
 bool FPhysInterface_Chaos::Sweep_Geom(FHitResult& OutHit, const FBodyInstance* InInstance, const FVector& InStart, const FVector& InEnd, const FQuat& InShapeRotation, const FCollisionShape& InShape, bool bSweepComplex)
 {
-    // @todo(mlentine): Need to implement this
-    return false;
+	bool bSweepHit = false;
+
+	if (InShape.IsNearlyZero())
+	{
+		bSweepHit = LineTrace_Geom(OutHit, InInstance, InStart, InEnd, bSweepComplex);
+	}
+	else
+	{
+		OutHit.TraceStart = InStart;
+		OutHit.TraceEnd = InEnd;
+
+		const FBodyInstance* TargetInstance = InInstance->WeldParent ? InInstance->WeldParent : InInstance;
+
+		FPhysicsCommand::ExecuteRead(TargetInstance->ActorHandle, [&](const FPhysicsActorHandle& Actor)
+		{
+			const Chaos::TGeometryParticle<float, 3>* RigidBody = Actor;
+
+			if (RigidBody && InInstance->OwnerComponent.Get())
+			{
+				FPhysicsShapeAdapter ShapeAdapter(InShapeRotation, InShape);
+
+				const FVector Delta = InEnd - InStart;
+				const float DeltaMag = Delta.Size();
+				if (DeltaMag > KINDA_SMALL_NUMBER)
+				{
+					const FTransform ActorTM(RigidBody->R(), RigidBody->X());
+
+					UPrimitiveComponent* OwnerComponentInst = InInstance->OwnerComponent.Get();
+					FTransform StartTM(ShapeAdapter.GetGeomOrientation(), InStart);
+					FTransform CompTM(OwnerComponentInst->GetComponentTransform());
+
+					Chaos::TVector<float,3> Dir = Delta / DeltaMag;
+
+					FSweepHit Hit;
+
+					// Get all the shapes from the actor
+					PhysicsInterfaceTypes::FInlineShapeArray Shapes;
+					// #PHYS2 - SHAPES - Resolve this function to not use px stuff
+					const int32 NumShapes = FillInlineShapeArray_AssumesLocked(Shapes, Actor); // #PHYS2 - Need a lock/execute here?
+
+					// Iterate over each shape
+					for (int32 ShapeIdx = 0; ShapeIdx < NumShapes; ShapeIdx++)
+					{
+						FPhysicsShapeReference_Chaos& ShapeRef = Shapes[ShapeIdx];
+						Chaos::TPerShapeData<float, 3>* Shape = ShapeRef.Shape;
+						check(Shape);
+
+						// Skip shapes not bound to this instance
+						if (!TargetInstance->IsShapeBoundToBody(ShapeRef))
+						{
+							continue;
+						}
+
+						// Filter so we trace against the right kind of collision
+						FCollisionFilterData ShapeFilter = Shape->QueryData;;
+						const bool bShapeIsComplex = (ShapeFilter.Word3 & EPDF_ComplexCollision) != 0;
+						const bool bShapeIsSimple = (ShapeFilter.Word3 & EPDF_SimpleCollision) != 0;
+						if ((bSweepComplex && bShapeIsComplex) || (!bSweepComplex && bShapeIsSimple))
+						{
+							//question: this is returning first result, is that valid? Keeping it the same as physx for now
+							Chaos::TVector<float, 3> WorldPosition;
+							Chaos::TVector<float, 3> WorldNormal;
+							int32 FaceIdx;
+							if(Chaos::SweepQuery<float, 3>(*Shape->Geometry, ActorTM, ShapeAdapter.GetGeometry(), StartTM, Dir, DeltaMag, Hit.Distance, WorldPosition, WorldNormal, FaceIdx))
+							{
+								// we just like to make sure if the hit is made
+								FCollisionFilterData QueryFilter;
+								QueryFilter.Word2 = 0xFFFFF;
+
+								// we don't get Shape information when we access via PShape, so I filled it up
+								Hit.Shape = Shape;
+								Hit.Actor = ShapeRef.ActorRef;
+								Hit.WorldPosition = WorldPosition;
+								Hit.WorldNormal = WorldNormal;
+
+								FTransform StartTransform(InStart);
+								Hit.FaceIndex = FindFaceIndex(Hit, Dir);
+								ConvertQueryImpactHit(OwnerComponentInst->GetWorld(), Hit, OutHit, DeltaMag, QueryFilter, InStart, InEnd, nullptr, StartTransform, false, false);
+								bSweepHit = true;
+							}
+						}
+					}
+				}
+			}
+		});
+	}
+
+	return bSweepHit;
+}
+
+bool Overlap_GeomInternal(const FBodyInstance* InInstance, const Chaos::TImplicitObject<float, 3>& InGeom, const FTransform& GeomTransform, FMTDResult* OutOptResult)
+{
+	const FBodyInstance* TargetInstance = InInstance->WeldParent ? InInstance->WeldParent : InInstance;
+	Chaos::TGeometryParticle<float, 3>* RigidBody = TargetInstance->ActorHandle;
+
+	if (RigidBody == nullptr)
+	{
+		return false;
+	}
+
+	// Get all the shapes from the actor
+	PhysicsInterfaceTypes::FInlineShapeArray Shapes;
+	const int32 NumShapes = FillInlineShapeArray_AssumesLocked(Shapes, RigidBody);
+
+	const FTransform ActorTM(RigidBody->R(), RigidBody->X());
+
+	// Iterate over each shape
+	for (int32 ShapeIdx = 0; ShapeIdx < NumShapes; ++ShapeIdx)
+	{
+		FPhysicsShapeReference_Chaos& ShapeRef = Shapes[ShapeIdx];
+		const Chaos::TPerShapeData<float,3>* Shape = ShapeRef.Shape;
+		check(Shape);
+
+		if (TargetInstance->IsShapeBoundToBody(ShapeRef))
+		{
+			//Chaos::TVector<float,3> POutDirection;
+			//float OutDistance;
+
+			if (OutOptResult)
+			{
+				ensure(false);	//todo: implement in Chaos
+				/*
+				PxTransform PTransform = U2PTransform(FPhysicsInterface::GetTransform(ShapeRef));
+				if (PxGeometryQuery::computePenetration(POutDirection, OutDistance, InPxGeom, ShapePose, PShape->getGeometry().any(), PTransform))
+				{
+					//TODO: there are some edge cases that give us nan results. In these cases we skip
+					if (!POutDirection.isFinite())
+					{
+						POutDirection.x = 0.f;
+						POutDirection.y = 0.f;
+						POutDirection.z = 0.f;
+					}
+
+					OutOptResult->Direction = P2UVector(POutDirection);
+					OutOptResult->Distance = FMath::Abs(OutDistance);
+
+					return true;
+				}*/
+				OutOptResult->Distance = 0;
+			}
+			//else	//todo: don't bother with this once MTD is implemented
+			{
+				if(Chaos::OverlapQuery<float, 3>(*Shape->Geometry, ActorTM, InGeom, GeomTransform))
+				{
+					return true;
+				}
+			}
+
+		}
+	}
+
+	return false;
 }
 
 bool FPhysInterface_Chaos::Overlap_Geom(const FBodyInstance* InBodyInstance, const FPhysicsGeometryCollection& InGeometry, const FTransform& InShapeTransform, FMTDResult* OutOptResult)
 {
-    // @todo(mlentine): Need to implement this
-    return false;
+	return Overlap_GeomInternal(InBodyInstance, InGeometry.GetGeometry(), InShapeTransform, OutOptResult);
 }
 
 bool FPhysInterface_Chaos::Overlap_Geom(const FBodyInstance* InBodyInstance, const FCollisionShape& InCollisionShape, const FQuat& InShapeRotation, const FTransform& InShapeTransform, FMTDResult* OutOptResult)
 {
-    // @todo(mlentine): Need to implement this
-    return false;
+	FPhysicsShapeAdapter Adaptor(InShapeRotation, InCollisionShape);
+	return Overlap_GeomInternal(InBodyInstance, Adaptor.GetGeometry(), Adaptor.GetGeometryPose(InShapeTransform.GetTranslation()), OutOptResult);
 }
 
 bool FPhysInterface_Chaos::GetSquaredDistanceToBody(const FBodyInstance* InInstance, const FVector& InPoint, float& OutDistanceSquared, FVector* OutOptPointOnBody)
 {
-    // @todo(mlentine): What spaces are in and out point?
+	const FTransform BodyTM = InInstance->GetUnrealWorldTransform();
+	const FVector LocalPoint = BodyTM.InverseTransformPositionNoScale(InPoint);
+
     TArray<FPhysicsShapeReference_Chaos> OutShapes;
     InInstance->GetAllShapes_AssumesLocked(OutShapes);
     check(OutShapes.Num() == 1);
     Chaos::TVector<float, 3> Normal;
-    const auto Phi = OutShapes[0].Object->PhiWithNormal(InPoint, Normal);
+    const auto Phi = OutShapes[0].Shape->Geometry->PhiWithNormal(LocalPoint, Normal);
     OutDistanceSquared = Phi * Phi;
     if (OutOptPointOnBody)
     {
-        *OutOptPointOnBody = InPoint - Phi * Normal;
+		const Chaos::TVector<float, 3> LocalClosestPoint = LocalPoint - Phi * Normal;
+        *OutOptPointOnBody = BodyTM.TransformPositionNoScale(LocalClosestPoint);
     }
     return true;
 }
 
 template<typename AllocatorType>
-int32 GetAllShapesInternal_AssumedLocked(const FPhysicsActorReference_Chaos& InActorHandle, TArray<FPhysicsShapeReference_Chaos, AllocatorType>& OutShapes)
+int32 GetAllShapesInternal_AssumedLocked(const FPhysicsActorHandle& InActorHandle, TArray<FPhysicsShapeReference_Chaos, AllocatorType>& OutShapes)
 {
 	OutShapes.Reset();
-	//FPhysicsShapeHandle NewShape;
-	//NewShape.bSimulation = true;
-	//NewShape.bQuery = true;
-	//NewShape.Object = LocalParticles.Geometry(Index);
-	//OutShapes.Add(NewShape);
+	const Chaos::TShapesArray<float, 3>& ShapesArray = InActorHandle->ShapesArray();
+	//todo: can we avoid this construction?
+	for (const auto& Shape : ShapesArray)
+	{
+		OutShapes.Add(FPhysicsShapeReference_Chaos(Shape.Get(), true, true, InActorHandle));
+	}
 	return OutShapes.Num();
 }
 
 template <>
-int32 FPhysInterface_Chaos::GetAllShapes_AssumedLocked(const FPhysicsActorReference_Chaos& InActorHandle, TArray<FPhysicsShapeReference_Chaos, FDefaultAllocator>& OutShapes)
+int32 FPhysInterface_Chaos::GetAllShapes_AssumedLocked(const FPhysicsActorHandle& InActorHandle, TArray<FPhysicsShapeReference_Chaos, FDefaultAllocator>& OutShapes)
 {
 	return GetAllShapesInternal_AssumedLocked(InActorHandle, OutShapes);
 }
 
 template <>
-int32 FPhysInterface_Chaos::GetAllShapes_AssumedLocked(const FPhysicsActorReference_Chaos& InActorHandle, PhysicsInterfaceTypes::FInlineShapeArray& OutShapes)
+int32 FPhysInterface_Chaos::GetAllShapes_AssumedLocked(const FPhysicsActorHandle& InActorHandle, PhysicsInterfaceTypes::FInlineShapeArray& OutShapes)
 {
 	return GetAllShapesInternal_AssumedLocked(InActorHandle, OutShapes);
 }
 
+
+FPhysicsShapeAdapter_Chaos::FPhysicsShapeAdapter_Chaos(const FQuat& Rot, const FCollisionShape& CollisionShape)
+	: GeometryRotation(Rot)
+{
+	switch (CollisionShape.ShapeType)
+	{
+		case ECollisionShape::Capsule:
+		{
+			const float CapsuleRadius = CollisionShape.GetCapsuleRadius();
+			const float CapsuleHalfHeight = CollisionShape.GetCapsuleHalfHeight();
+			if (CapsuleRadius < CapsuleHalfHeight)
+			{
+				const float UseHalfHeight = FMath::Max(CollisionShape.GetCapsuleAxisHalfLength(), FCollisionShape::MinCapsuleAxisHalfHeight());
+				const FVector Bot = FVector(0.f, 0.f, -UseHalfHeight);
+				const FVector Top = FVector(0.f, 0.f, UseHalfHeight);
+				const float UseRadius = FMath::Max(CapsuleRadius, FCollisionShape::MinCapsuleRadius());
+				Geometry = TUniquePtr<FPhysicsGeometry>(new Chaos::TCapsule<float>(Bot, Top, UseRadius));
+			}
+			else
+			{
+				// Use a sphere instead.
+				const float UseRadius = FMath::Max(CapsuleRadius, FCollisionShape::MinSphereRadius());
+				Geometry = TUniquePtr<FPhysicsGeometry>(new Chaos::TSphere<float,3>(Chaos::TVector<float,3>(0), UseRadius));
+			}
+			break;
+		}
+		case ECollisionShape::Box:
+		{
+			Chaos::TVector<float,3> HalfExtents = CollisionShape.GetBox();
+			HalfExtents.X = FMath::Max(HalfExtents.X, FCollisionShape::MinBoxExtent());
+			HalfExtents.Y = FMath::Max(HalfExtents.Y, FCollisionShape::MinBoxExtent());
+			HalfExtents.Z = FMath::Max(HalfExtents.Z, FCollisionShape::MinBoxExtent());
+
+			Geometry = TUniquePtr<FPhysicsGeometry>(new Chaos::TBox<float, 3>(-HalfExtents, HalfExtents));
+			break;
+		}
+		case ECollisionShape::Sphere:
+		{
+			const float UseRadius = FMath::Max(CollisionShape.GetSphereRadius(), FCollisionShape::MinSphereRadius());
+			Geometry = TUniquePtr<FPhysicsGeometry>(new Chaos::TSphere<float, 3>(Chaos::TVector<float, 3>(0), UseRadius));
+			break;
+		}
+		default:
+			ensure(false);
+			break;
+	}
+}
+
+const FPhysicsGeometry& FPhysicsShapeAdapter_Chaos::GetGeometry() const
+{
+	return *Geometry;
+}
+
+FTransform FPhysicsShapeAdapter_Chaos::GetGeometryPose(const FVector& Pos) const
+{
+	return FTransform(GeometryRotation, Pos);
+}
+
+const FQuat& FPhysicsShapeAdapter_Chaos::GetGeomOrientation() const
+{
+	return GeometryRotation;
+}
+
+FVector FindBoxOpposingNormal(const FLocationHit& Hit, const FVector& TraceDirectionDenorm, const FVector& InNormal)
+{
+	// We require normal info for our algorithm.
+	const bool bNormalData = !!(Hit.Flags & EHitFlags::Normal);
+	if (!bNormalData)
+	{
+		return InNormal;
+	}
+
+	ensure(Hit.Shape->Geometry->GetType() == Chaos::ImplicitObjectType::Box);
+	const FTransform LocalToWorld(Hit.Actor->R(), Hit.Actor->X());
+
+	// Find which faces were included in the contact normal, and for multiple faces, use the one most opposing the sweep direction.
+	const FVector ContactNormalLocal = LocalToWorld.InverseTransformVectorNoScale(Hit.WorldNormal);
+	const float* ContactNormalLocalPtr = &ContactNormalLocal.X;
+	const float* TraceDirDenormWorldPtr = &TraceDirectionDenorm.X;
+	const FVector TraceDirDenormLocal = LocalToWorld.InverseTransformVectorNoScale(TraceDirectionDenorm);
+	const float* TraceDirDenormLocalPtr = &TraceDirDenormLocal.X;
+
+	FVector BestLocalNormal(ContactNormalLocal);
+	float* BestLocalNormalPtr = &BestLocalNormal.X;
+	float BestOpposingDot = FLT_MAX;
+
+	for (int32 i = 0; i < 3; i++)
+	{
+		// Select axis of face to compare to, based on normal.
+		if (ContactNormalLocalPtr[i] > KINDA_SMALL_NUMBER)
+		{
+			const float TraceDotFaceNormal = TraceDirDenormLocalPtr[i]; // TraceDirDenormLocal.dot(BoxFaceNormal)
+			if (TraceDotFaceNormal < BestOpposingDot)
+			{
+				BestOpposingDot = TraceDotFaceNormal;
+				BestLocalNormal = FVector(0.f);
+				BestLocalNormalPtr[i] = 1.f;
+			}
+		}
+		else if (ContactNormalLocalPtr[i] < -KINDA_SMALL_NUMBER)
+		{
+			const float TraceDotFaceNormal = -TraceDirDenormLocalPtr[i]; // TraceDirDenormLocal.dot(BoxFaceNormal)
+			if (TraceDotFaceNormal < BestOpposingDot)
+			{
+				BestOpposingDot = TraceDotFaceNormal;
+				BestLocalNormal = FVector(0.f);
+				BestLocalNormalPtr[i] = -1.f;
+			}
+		}
+	}
+
+	// Fill in result
+	const FVector WorldNormal = LocalToWorld.TransformVectorNoScale(BestLocalNormal);
+	return WorldNormal;
+}
+
+FVector FindHeightFieldOpposingNormal(const FLocationHit& PHit, const FVector& TraceDirectionDenorm, const FVector& InNormal)
+{
+	//todo: implement
+	return FVector(0, 0, 1);
+}
+FVector FindConvexMeshOpposingNormal(const FLocationHit& PHit, const FVector& TraceDirectionDenorm, const FVector& InNormal)
+{
+	//todo: implement
+	return FVector(0, 0, 1);
+}
+FVector FindTriMeshOpposingNormal(const FLocationHit& PHit, const FVector& TraceDirectionDenorm, const FVector& InNormal)
+{
+	//todo: implement
+	return FVector(0, 0, 1);
+}
 #endif
