Copying //UE4/Dev-Physics to Dev-Main (//UE4/Dev-Main) [at] 7703071

diff --git a/Engine/Source/Runtime/Engine/Private/PhysicsEngine/Experimental/PhysScene_Chaos.cpp b/Engine/Source/Runtime/Engine/Private/PhysicsEngine/Experimental/PhysScene_Chaos.cpp
--- a/Engine/Source/Runtime/Engine/Private/PhysicsEngine/Experimental/PhysScene_Chaos.cpp
+++ b/Engine/Source/Runtime/Engine/Private/PhysicsEngine/Experimental/PhysScene_Chaos.cpp
@@ -4,7 +4,7 @@
 
 #include "Physics/Experimental/PhysScene_Chaos.h"
 
-#include "PBDRigidsSolver.h"
+#include "PhysicsSolver.h"
 #include "ChaosSolversModule.h"
 #include "ChaosLog.h"
 #include "ChaosStats.h"
@@ -24,11 +24,11 @@
 #include "PhysicsEngine/PhysicsSettings.h"
 #include "Components/PrimitiveComponent.h"
 
-#include "SolverObjects/FieldSystemPhysicsObject.h"
-#include "SolverObjects/GeometryCollectionPhysicsObject.h"
-#include "SolverObjects/SkeletalMeshPhysicsObject.h"
-#include "SolverObjects/StaticMeshPhysicsObject.h"
-#include "SolverObjects/BodyInstancePhysicsObject.h"
+#include "PhysicsProxy/FieldSystemPhysicsProxy.h"
+#include "PhysicsProxy/GeometryCollectionPhysicsProxy.h"
+#include "PhysicsProxy/SkeletalMeshPhysicsProxy.h"
+#include "PhysicsProxy/StaticMeshPhysicsProxy.h"
+#include "Chaos/Framework/SingleParticlePhysicsProxy.h"
 #include "Chaos/UniformGrid.h"
 #include "Chaos/BoundingVolume.h"
 #include "Chaos/ISpatialAcceleration.h"
@@ -191,17 +191,19 @@ private:
 
 		FSpacialDebugDraw DrawInterface(WorldPtr);
 
-		const TArray<FPBDRigidsSolver*>& Solvers = ChaosModule->GetSolvers();
+		const TArray<FPhysicsSolver*>& Solvers = ChaosModule->GetSolvers();
 
-		for(FPBDRigidsSolver* Solver : Solvers)
+		for(FPhysicsSolver* Solver : Solvers)
 		{
 			if(bDrawHier)
 			{
+#if TODO_REIMPLEMENT_SPATIAL_ACCELERATION_ACCESS
 				if (const ISpatialAcceleration<float, 3>* SpatialAcceleration = Solver->GetSpatialAcceleration())
 				{
 					SpatialAcceleration->DebugDraw(&DrawInterface);
 					Solver->ReleaseSpatialAcceleration();
 				}
+#endif
 				
 #if 0
 				if (const Chaos::TBoundingVolume<TPBDRigidParticles<float, 3>, float, 3>* BV = SpatialAcceleration->Cast<TBoundingVolume<TPBDRigidParticles<float, 3>, float, 3>>())
@@ -316,7 +318,7 @@ static void CopyParticleData(Chaos::TPBDRigidParticles<float, 3>& ToParticles, c
 	ToParticles.InvI(ToIndex) = FromParticles.InvI(FromIndex);
 	ToParticles.SetGeometry(ToIndex, FromParticles.Geometry(FromIndex));	//question: do we need to deal with dynamic geometry?
 	ToParticles.CollisionParticles(ToIndex) = MoveTemp(FromParticles.CollisionParticles(FromIndex));
-	ToParticles.SetDisabledLowLevel(ToIndex, FromParticles.Disabled(FromIndex));
+	ToParticles.DisabledRef(ToIndex) = FromParticles.Disabled(FromIndex);
 	ToParticles.SetSleeping(ToIndex, FromParticles.Sleeping(FromIndex));
 }
 
@@ -373,8 +375,8 @@ FPhysScene_Chaos::FPhysScene_Chaos(AActor* InSolverActor)
 	// once we have a better lifecycle for the scenes.
 	FCoreDelegates::OnPreExit.AddRaw(this, &FPhysScene_Chaos::Shutdown);
 
-	SolverObjectToComponentMap.Reset();
-	ComponentToSolverObjectMap.Reset();
+	PhysicsProxyToComponentMap.Reset();
+	ComponentToPhysicsProxyMap.Reset();
 
 #if WITH_EDITOR
 	FGameDelegates::Get().GetEndPlayMapDelegate().AddRaw(this, &FPhysScene_Chaos::OnWorldEndPlay);
@@ -408,7 +410,11 @@ bool FPhysScene_Chaos::IsTickable() const
 {
 	const bool bDedicatedThread = ChaosModule->IsPersistentTaskRunning();
 
+#if TODO_REIMPLEMENT_SOLVER_ENABLING
 	return !bDedicatedThread && GetSolver()->Enabled();
+#else
+	return false;
+#endif
 }
 
 void FPhysScene_Chaos::Tick(float DeltaTime)
@@ -420,14 +426,16 @@ void FPhysScene_Chaos::Tick(float DeltaTime)
 	if (!ChaosModule->ShouldStepSolver(SingleStepCounter)) { return; }
 #endif
 
-	Chaos::FPBDRigidsSolver* Solver = GetSolver();
+	Chaos::FPhysicsSolver* Solver = GetSolver();
 
 #if CHAOS_WITH_PAUSABLE_SOLVER
 	// Update solver depending on the pause status of the actor's world attached to this scene
 	OnUpdateWorldPause();
 
+#if TODO_REIMPLEMENT_SOLVER_PAUSING
 	// Return now if this solver is paused
 	if (Solver->Paused()) { return; }
+#endif
 #endif
 
 	float SafeDelta = FMath::Clamp(DeltaTime, 0.0f, UPhysicsSettings::Get()->MaxPhysicsDeltaTime);
@@ -435,24 +443,22 @@ void FPhysScene_Chaos::Tick(float DeltaTime)
 	UE_LOG(LogFPhysScene_ChaosSolver, Verbose, TEXT("FPhysScene_Chaos::Tick(%3.5f)"), SafeDelta);
 	Solver->AdvanceSolverBy(SafeDelta);
 
-	// Sync proxies after simulation
-	FSolverObjectStorage& Objects = Solver->GetObjectStorage();
-
-	Objects.ForEachSolverObjectParallel([](auto* Object)
+	Solver->ForEachPhysicsProxyParallel([](auto* Object)
 	{
 		// #BGallagher TODO Just use one side of the buffer for single-thread tick
-		Object->CacheResults();
-		Object->FlipCache();
+		Object->BufferPhysicsResults();
+		Object->FlipBuffer();
 	});
 
-	Objects.ForEachSolverObject([](auto* Object)
+	Solver->ForEachPhysicsProxy([](auto* Object)
 	{
-		Object->SyncToCache();
+		Object->PullFromPhysicsState();
 	});
 
+
 }
 
-Chaos::FPBDRigidsSolver* FPhysScene_Chaos::GetSolver() const
+Chaos::FPhysicsSolver* FPhysScene_Chaos::GetSolver() const
 {
 	return SceneSolver;
 }
@@ -468,7 +474,7 @@ Chaos::IDispatcher* FPhysScene_Chaos::GetDispatcher() const
 }
 
 template<typename ObjectType>
-void AddSolverObject(ObjectType* InObject, Chaos::FPBDRigidsSolver* InSolver, Chaos::IDispatcher* InDispatcher)
+void AddPhysicsProxy(ObjectType* InObject, Chaos::FPhysicsSolver* InSolver, Chaos::IDispatcher* InDispatcher)
 {
 	check(IsInGameThread() && InSolver);
 	//const bool bDedicatedThread = ChaosModule->IsPersistentTaskEnabled();
@@ -479,59 +485,74 @@ void AddSolverObject(ObjectType* InObject, Chaos::FPBDRigidsSolver* InSolver, Ch
 	if(/*bDedicatedThread && */InDispatcher)
 	{
 		// Pass the proxy off to the physics thread
-		InDispatcher->EnqueueCommand([InObject, InSolver](Chaos::FPersistentPhysicsTask* PhysThread)
+		InDispatcher->EnqueueCommandImmediate([InObject, InSolver](Chaos::FPersistentPhysicsTask* PhysThread)
 		{
+#if CHAOS_PARTICLEHANDLE_TODO
 			InSolver->RegisterObject(InObject);
+#endif
 		});
 	}
 }
 
-void FPhysScene_Chaos::AddObject(UPrimitiveComponent* Component, FSkeletalMeshPhysicsObject* InObject)
+void FPhysScene_Chaos::AddObject(UPrimitiveComponent* Component, FSkeletalMeshPhysicsProxy* InObject)
 {
 	AddToComponentMaps(Component, InObject);
 
-	Chaos::FPBDRigidsSolver* Solver = GetSolver();
-	Solver->GetObjectStorage_GameThread().SkeletalMeshObjects.Add(InObject);
+	check(false);
+#if 0
+	Chaos::FPhysicsSolver* Solver = GetSolver();
+	Solver->RegisterObject(InObject);
 
-	AddSolverObject(InObject, Solver, GetDispatcher());
+	AddPhysicsProxy(InObject, Solver, GetDispatcher());
+#endif
 }
 
-void FPhysScene_Chaos::AddObject(UPrimitiveComponent* Component, FStaticMeshPhysicsObject* InObject)
+void FPhysScene_Chaos::AddObject(UPrimitiveComponent* Component, FStaticMeshPhysicsProxy* InObject)
 {
 	AddToComponentMaps(Component, InObject);
 
-	Chaos::FPBDRigidsSolver* Solver = GetSolver();
-	Solver->GetObjectStorage_GameThread().StaticMeshObjects.Add(InObject);
+	check(false);
+#if 0
+	Chaos::FPhysicsSolver* Solver = GetSolver();
+	Solver->RegisterObject(InObject);
 
-	AddSolverObject(InObject, Solver, GetDispatcher());
+	AddPhysicsProxy(InObject, Solver, GetDispatcher());
+#endif
 }
 
-void FPhysScene_Chaos::AddObject(UPrimitiveComponent* Component, FBodyInstancePhysicsObject* InObject)
+void FPhysScene_Chaos::AddObject(UPrimitiveComponent* Component, FGeometryParticlePhysicsProxy* InObject)
 {
 	AddToComponentMaps(Component, InObject);
 
-	Chaos::FPBDRigidsSolver* Solver = GetSolver();
-	Solver->GetObjectStorage_GameThread().BodyInstanceObjects.Add(InObject);
+	check(false);
+#if 0
+	Chaos::FPhysicsSolver* Solver = GetSolver();
+	Solver->RegisterObject(InObject);
 
-	AddSolverObject(InObject, Solver, GetDispatcher());
+	AddPhysicsProxy(InObject, Solver, GetDispatcher());
+#endif
 }
 
-
-void FPhysScene_Chaos::AddObject(UPrimitiveComponent* Component, FGeometryCollectionPhysicsObject* InObject)
+void FPhysScene_Chaos::AddObject(UPrimitiveComponent* Component, FGeometryCollectionPhysicsProxy* InObject)
 {
 	AddToComponentMaps(Component, InObject);
 
-	Chaos::FPBDRigidsSolver* Solver = GetSolver();
-	Solver->GetObjectStorage_GameThread().GeometryCollectionObjects.Add(InObject);
+	check(false);
+#if 0
+	Chaos::FPhysicsSolver* Solver = GetSolver();
+	Solver->RegisterObject(InObject);
 
-	AddSolverObject(InObject, Solver, GetDispatcher());
+	AddPhysicsProxy(InObject, Solver, GetDispatcher());
+#endif
 }
 
-void FPhysScene_Chaos::AddObject(UPrimitiveComponent* Component, FFieldSystemPhysicsObject* InObject)
+void FPhysScene_Chaos::AddObject(UPrimitiveComponent* Component, FFieldSystemPhysicsProxy* InObject)
 {
 	AddToComponentMaps(Component, InObject);
 
-	Chaos::FPBDRigidsSolver* CurrSceneSolver = GetSolver();
+	check(false);
+#if 0
+	Chaos::FPhysicsSolver* CurrSceneSolver = GetSolver();
 
 	check(IsInGameThread() && CurrSceneSolver);
 	//const bool bDedicatedThread = ChaosModule->IsPersistentTaskEnabled();
@@ -540,27 +561,30 @@ void FPhysScene_Chaos::AddObject(UPrimitiveComponent* Component, FFieldSystemPhy
 	
 	Chaos::IDispatcher* Dispatcher = GetDispatcher();
 
-	const TArray<Chaos::FPBDRigidsSolver*>& Solvers = ChaosModule->GetSolvers();
-	for(Chaos::FPBDRigidsSolver* Solver : Solvers)
+	const TArray<Chaos::FPhysicsSolver*>& Solvers = ChaosModule->GetSolvers();
+	for(Chaos::FPhysicsSolver* Solver : Solvers)
 	{
-		if(true || Solver->HasActiveObjects())
+		if(true || Solver->HasActiveParticles())
 		{
-			Solver->GetObjectStorage_GameThread().FieldSystemObjects.Add(InObject);
+			Solver->RegisterObject(InObject);
 
 			if(/*bDedicatedThread && */Dispatcher)
 			{
 				// Pass the proxy off to the physics thread
 				Dispatcher->EnqueueCommand([InObject, InSolver = Solver](Chaos::FPersistentPhysicsTask* PhysThread)
 				{
+#if CHAOS_PARTICLEHANDLE_TODO
 					InSolver->RegisterObject(InObject);
+#endif
 				});
 			}
 		}
 	}
+#endif
 }
 
 template<typename ObjectType>
-void RemoveSolverObject(ObjectType* InObject, Chaos::FPBDRigidsSolver* InSolver, FChaosSolversModule* InModule)
+void RemovePhysicsProxy(ObjectType* InObject, Chaos::FPhysicsSolver* InSolver, FChaosSolversModule* InModule)
 {
 	check(IsInGameThread());
 
@@ -572,9 +596,11 @@ void RemoveSolverObject(ObjectType* InObject, Chaos::FPBDRigidsSolver* InSolver,
 	if(PhysDispatcher)
 	{
 		// Remove the object from the solver
-		PhysDispatcher->EnqueueCommand([InObject, InSolver, bDedicatedThread](Chaos::FPersistentPhysicsTask* PhysThread)
+		PhysDispatcher->EnqueueCommandImmediate([InObject, InSolver, bDedicatedThread](Chaos::FPersistentPhysicsTask* PhysThread)
 		{
+#if CHAOS_PARTICLEHANDLE_TODO
 			InSolver->UnregisterObject(InObject);
+#endif
 			InObject->OnRemoveFromScene();
 
 			if (!bDedicatedThread)
@@ -587,10 +613,12 @@ void RemoveSolverObject(ObjectType* InObject, Chaos::FPBDRigidsSolver* InSolver,
 	}
 }
 
-void FPhysScene_Chaos::RemoveObject(FSkeletalMeshPhysicsObject* InObject)
+void FPhysScene_Chaos::RemoveObject(FSkeletalMeshPhysicsProxy* InObject)
 {
-	Chaos::FPBDRigidsSolver* Solver = InObject->GetSolver();
-	const int32 NumRemoved = Solver->GetObjectStorage_GameThread().SkeletalMeshObjects.Remove(InObject);
+	check(false);
+#if 0
+	Chaos::FPhysicsSolver* Solver = InObject->GetSolver();
+	const int32 NumRemoved = Solver->UnregisterObject(InObject);
 
 	if(NumRemoved == 0)
 	{
@@ -599,13 +627,17 @@ void FPhysScene_Chaos::RemoveObject(FSkeletalMeshPhysicsObject* InObject)
 
 	RemoveFromComponentMaps(InObject);
 
-	RemoveSolverObject(InObject, Solver, ChaosModule);
+	RemovePhysicsProxy(InObject, Solver, ChaosModule);
+#endif
 }
 
-void FPhysScene_Chaos::RemoveObject(FStaticMeshPhysicsObject* InObject)
+void FPhysScene_Chaos::RemoveObject(FStaticMeshPhysicsProxy* InObject)
 {
-	Chaos::FPBDRigidsSolver* Solver = InObject->GetSolver();
-	const int32 NumRemoved = Solver->GetObjectStorage_GameThread().StaticMeshObjects.Remove(InObject);
+	check(false);
+#if 0
+	Chaos::FPhysicsSolver* Solver = InObject->GetSolver();
+
+	const int32 NumRemoved = Solver->UnregisterObject(InObject);
 
 	if(NumRemoved == 0)
 	{
@@ -614,13 +646,17 @@ void FPhysScene_Chaos::RemoveObject(FStaticMeshPhysicsObject* InObject)
 
 	RemoveFromComponentMaps(InObject);
 
-	RemoveSolverObject(InObject, Solver, ChaosModule);
+	RemovePhysicsProxy(InObject, Solver, ChaosModule);
+#endif
 }
 
-void FPhysScene_Chaos::RemoveObject(FBodyInstancePhysicsObject* InObject)
+void FPhysScene_Chaos::RemoveObject(FGeometryParticlePhysicsProxy* InObject)
 {
-	Chaos::FPBDRigidsSolver* Solver = InObject->GetSolver();
-	const int32 NumRemoved = Solver->GetObjectStorage_GameThread().BodyInstanceObjects.Remove(InObject);
+	check(false);
+#if 0
+	Chaos::FPhysicsSolver* Solver = InObject->GetSolver();
+
+	const int32 NumRemoved = Solver->UnregisterObject(InObject);
 
 	if (NumRemoved == 0)
 	{
@@ -629,13 +665,17 @@ void FPhysScene_Chaos::RemoveObject(FBodyInstancePhysicsObject* InObject)
 
 	RemoveFromComponentMaps(InObject);
 
-	RemoveSolverObject(InObject, Solver, ChaosModule);
+	RemovePhysicsProxy(InObject, Solver, ChaosModule);
+#endif
 }
 
-void FPhysScene_Chaos::RemoveObject(FGeometryCollectionPhysicsObject* InObject)
+void FPhysScene_Chaos::RemoveObject(FGeometryCollectionPhysicsProxy* InObject)
 {
-	Chaos::FPBDRigidsSolver* Solver = InObject->GetSolver();
-	const int32 NumRemoved = Solver->GetObjectStorage_GameThread().GeometryCollectionObjects.Remove(InObject);
+	check(false);
+#if 0
+	Chaos::FPhysicsSolver* Solver = InObject->GetSolver();
+
+	const int32 NumRemoved = Solver->UnregisterObject(InObject);
 
 	if(NumRemoved == 0)
 	{
@@ -644,30 +684,35 @@ void FPhysScene_Chaos::RemoveObject(FGeometryCollectionPhysicsObject* InObject)
 
 	RemoveFromComponentMaps(InObject);
 
-	RemoveSolverObject(InObject, Solver, ChaosModule);
+	RemovePhysicsProxy(InObject, Solver, ChaosModule);
+#endif
 }
 
-void FPhysScene_Chaos::RemoveObject(FFieldSystemPhysicsObject* InObject)
+void FPhysScene_Chaos::RemoveObject(FFieldSystemPhysicsProxy* InObject)
 {
+	check(false);
+#if 0
 	if(!InObject)
 	{
 		return;
 	}
 
-	Chaos::FPBDRigidsSolver* CurrSceneSolver = GetSolver();
-	CurrSceneSolver->GetObjectStorage_GameThread().FieldSystemObjects.Remove(InObject);
+	Chaos::FPhysicsSolver* CurrSceneSolver = GetSolver();
+
+	CurrSceneSolver->UnregisterObject(InObject);
 
 	check(IsInGameThread());
 
 	Chaos::IDispatcher* Dispatcher = GetDispatcher();
 	const bool bDedicatedThread = Dispatcher->GetMode() == Chaos::EThreadingMode::DedicatedThread;
 
-	const TArray<Chaos::FPBDRigidsSolver*>& Solvers = ChaosModule->GetSolvers();
-	for(Chaos::FPBDRigidsSolver* Solver : Solvers)
+	const TArray<Chaos::FPhysicsSolver*>& Solvers = ChaosModule->GetSolvers();
+	for(Chaos::FPhysicsSolver* Solver : Solvers)
 	{
-		if(true || Solver->HasActiveObjects())
+		if(true || Solver->HasActiveParticles())
 		{
-			Solver->GetObjectStorage_GameThread().FieldSystemObjects.Add(InObject);
+			// @todo(question) : Is the RegisterObject a bug? Seems like it should be removing. [brice->benn]
+			Solver->RegisterObject(InObject);
 
 			if(/*bDedicatedThread && */Dispatcher)
 			{
@@ -686,7 +731,39 @@ void FPhysScene_Chaos::RemoveObject(FFieldSystemPhysicsObject* InObject)
 	{
 		delete InObject;
 	}
+#endif
+}
+
+void FPhysScene_Chaos::UnregisterEvent(const Chaos::EEventType& EventID)
+{
+	check(IsInGameThread());
+
+	Chaos::IDispatcher* Dispatcher = GetDispatcher();
+	Chaos::FPBDRigidsSolver* Solver = GetSolver();
+
+	if (Dispatcher)
+	{
+		Dispatcher->EnqueueCommandImmediate([EventID, InSolver = Solver](Chaos::FPersistentPhysicsTask* PhysThread)
+		{
+			InSolver->GetEventManager()->UnregisterEvent(EventID);
+		});
+	}
+}
+
+void FPhysScene_Chaos::UnregisterEventHandler(const Chaos::EEventType& EventID, const void* Handler)
+{
+	check(IsInGameThread());
+
+	Chaos::IDispatcher* Dispatcher = GetDispatcher();
+	Chaos::FPBDRigidsSolver* Solver = GetSolver();
 
+	if (Dispatcher)
+	{
+		Dispatcher->EnqueueCommandImmediate([EventID, Handler, InSolver = Solver](Chaos::FPersistentPhysicsTask* PhysThread)
+		{
+			InSolver->GetEventManager()->UnregisterHandler(EventID, Handler);
+		});
+	}
 }
 
 void FPhysScene_Chaos::Shutdown()
@@ -700,8 +777,8 @@ void FPhysScene_Chaos::Shutdown()
 	ChaosModule = nullptr;
 	SceneSolver = nullptr;
 
-	SolverObjectToComponentMap.Reset();
-	ComponentToSolverObjectMap.Reset();
+	PhysicsProxyToComponentMap.Reset();
+	ComponentToPhysicsProxyMap.Reset();
 }
 
 void FPhysScene_Chaos::AddReferencedObjects(FReferenceCollector& Collector)
@@ -712,7 +789,7 @@ void FPhysScene_Chaos::AddReferencedObjects(FReferenceCollector& Collector)
 		Collector.AddReferencedObject(Obj);
 	}
 
-	for (TPair<ISolverObjectBase*, UPrimitiveComponent*>& Pair : SolverObjectToComponentMap)
+	for (TPair<IPhysicsProxyBase*, UPrimitiveComponent*>& Pair : PhysicsProxyToComponentMap)
 	{
 		Collector.AddReferencedObject(Pair.Get<1>());
 	}
@@ -741,6 +818,7 @@ void FPhysScene_Chaos::OnUpdateWorldPause()
 		}
 	}
 
+#if TODO_REIMPLEMENT_SOLVER_PAUSING
 	if (bIsWorldPaused != bIsPaused)
 	{
 		bIsWorldPaused = bIsPaused;
@@ -749,12 +827,13 @@ void FPhysScene_Chaos::OnUpdateWorldPause()
 		if (PhysDispatcher)
 		{
 			UE_LOG(LogFPhysScene_ChaosSolver, Verbose, TEXT("FPhysScene_Chaos::OnUpdateWorldPause() pause status changed for actor %s, bIsPaused = %d"), Actor ? *Actor->GetName() : TEXT("None"), bIsPaused);
-			PhysDispatcher->EnqueueCommand(SceneSolver, [bIsPaused](Chaos::FPBDRigidsSolver* Solver)
+			PhysDispatcher->EnqueueCommandImmediate(SceneSolver, [bIsPaused](Chaos::FPhysicsSolver* Solver)
 			{
 				Solver->SetPaused(bIsPaused);
 			});
 		}
 	}
+#endif
 }
 #endif  // #if CHAOS_WITH_PAUSABLE_SOLVER
 
@@ -780,24 +859,24 @@ void FPhysScene_Chaos::AddPieModifiedObject(UObject* InObj)
 }
 #endif
 
-void FPhysScene_Chaos::AddToComponentMaps(UPrimitiveComponent* Component, ISolverObjectBase* InObject)
+void FPhysScene_Chaos::AddToComponentMaps(UPrimitiveComponent* Component, IPhysicsProxyBase* InObject)
 {
 	if (Component != nullptr && InObject != nullptr)
 	{
-		SolverObjectToComponentMap.Add(InObject, Component);
-		ComponentToSolverObjectMap.Add(Component, InObject);
+		PhysicsProxyToComponentMap.Add(InObject, Component);
+		ComponentToPhysicsProxyMap.Add(Component, InObject);
 	}
 }
 
-void FPhysScene_Chaos::RemoveFromComponentMaps(ISolverObjectBase* InObject)
+void FPhysScene_Chaos::RemoveFromComponentMaps(IPhysicsProxyBase* InObject)
 {
-	UPrimitiveComponent** const Component = SolverObjectToComponentMap.Find(InObject);
+	UPrimitiveComponent** const Component = PhysicsProxyToComponentMap.Find(InObject);
 	if (Component)
 	{
-		ComponentToSolverObjectMap.Remove(*Component);
+		ComponentToPhysicsProxyMap.Remove(*Component);
 	}
 
-	SolverObjectToComponentMap.Remove(InObject);
+	PhysicsProxyToComponentMap.Remove(InObject);
 }
 
 #if WITH_CHAOS
@@ -805,14 +884,19 @@ void FPhysScene_Chaos::RemoveFromComponentMaps(ISolverObjectBase* InObject)
 FPhysScene_ChaosInterface::FPhysScene_ChaosInterface(const AWorldSettings* InSettings /*= nullptr*/)
 {
 	//Initialize unique ptrs that are just here to allow forward declare. This should be reworked todo(ocohen)
-	MGravity = MakeUnique<Chaos::PerParticleGravity<float, 3>>();
-	MSpringConstraints = MakeUnique<Chaos::TPBDSpringConstraints<float, 3>>();
+#if TODO_FIX_REFERENCES_TO_ADDARRAY
 	Scene.GetSolver()->GetEvolution()->GetParticles().AddArray(&BodyInstances);
+#endif
+	// @todo: What do we do with MGravity?
+	Scene.GetSolver()->GetEvolution()->AddForceFunction(Chaos::Utilities::GetRigidsGravityFunction(Chaos::TVector<float, 3>(0, 0, -1), 980.f));
+
+	Scene.GetSolver()->PhysSceneHack = this;
+
 }
 
 void FPhysScene_ChaosInterface::OnWorldBeginPlay()
 {
-	Chaos::FPBDRigidsSolver* Solver = Scene.GetSolver();
+	Chaos::FPhysicsSolver* Solver = Scene.GetSolver();
 	if (Solver)
 	{
 		Solver->SetEnabled(true);
@@ -821,18 +905,20 @@ void FPhysScene_ChaosInterface::OnWorldBeginPlay()
 
 void FPhysScene_ChaosInterface::OnWorldEndPlay()
 {
-	Chaos::FPBDRigidsSolver* Solver = Scene.GetSolver();
+	Chaos::FPhysicsSolver* Solver = Scene.GetSolver();
 	if (Solver)
 	{
 		Solver->SetEnabled(false);
 	}
 }
 
-void FPhysScene_ChaosInterface::AddActorsToScene_AssumesLocked(const TArray<FPhysicsActorHandle>& InActors)
+void FPhysScene_ChaosInterface::AddActorsToScene_AssumesLocked(TArray<FPhysicsActorHandle>& InHandles)
 {
-	for(const FPhysicsActorHandle& Actor : InActors)
+	Chaos::FPhysicsSolver* Solver = Scene.GetSolver();
+	Chaos::IDispatcher* Dispatcher = Scene.GetDispatcher();
+	for (FPhysicsActorHandle& Handle : InHandles)
 	{
-		check(Actor.GetScene());
+		FPhysicsInterface::AddActorToSolver(Handle, Solver, Dispatcher);
 	}
 }
 
@@ -856,12 +942,12 @@ const UWorld* FPhysScene_ChaosInterface::GetOwningWorld() const
 	return MOwningWorld;
 }
 
-Chaos::FPBDRigidsSolver* FPhysScene_ChaosInterface::GetSolver()
+Chaos::FPhysicsSolver* FPhysScene_ChaosInterface::GetSolver()
 {
 	return Scene.GetSolver();
 }
 
-const Chaos::FPBDRigidsSolver* FPhysScene_ChaosInterface::GetSolver() const
+const Chaos::FPhysicsSolver* FPhysScene_ChaosInterface::GetSolver() const
 {
 	return Scene.GetSolver();
 }
@@ -888,8 +974,6 @@ void FPhysScene_ChaosInterface::AddForce_AssumesLocked(FBodyInstance* BodyInstan
 
 void FPhysScene_ChaosInterface::AddForceAtPosition_AssumesLocked(FBodyInstance* BodyInstance, const FVector& Force, const FVector& Position, bool bAllowSubstepping, bool bIsLocalForce /*= false*/)
 {
-	check(!bIsLocalForce);
-	check(BodyInstance->ActorHandle.GetScene() == this);
 	// #todo : Implement
 	//AddTorque(Chaos::TVector<float, 3>::CrossProduct(Position - Scene.GetSolver()->GetRigidParticles().X(GetIndexFromId(Id)), Force), Id);
 	//AddForce(Force, Id);
@@ -897,8 +981,6 @@ void FPhysScene_ChaosInterface::AddForceAtPosition_AssumesLocked(FBodyInstance*
 
 void FPhysScene_ChaosInterface::AddRadialForceToBody_AssumesLocked(FBodyInstance* BodyInstance, const FVector& Origin, const float Radius, const float Strength, const uint8 Falloff, bool bAccelChange, bool bAllowSubstepping)
 {
-	check(BodyInstance->ActorHandle.GetScene() == this);
-
 	// #todo : Implement
 #if 0
 	Chaos::TVector<float, 3> Direction = (static_cast<FVector>(Scene.GetSolver()->GetRigidParticles().X(Index)) - Origin);
@@ -930,12 +1012,10 @@ void FPhysScene_ChaosInterface::AddRadialForceToBody_AssumesLocked(FBodyInstance
 
 void FPhysScene_ChaosInterface::ClearForces_AssumesLocked(FBodyInstance* BodyInstance, bool bAllowSubstepping)
 {
-	check(BodyInstance->ActorHandle.GetScene() == this);
 }
 
 void FPhysScene_ChaosInterface::AddTorque_AssumesLocked(FBodyInstance* BodyInstance, const FVector& Torque, bool bAllowSubstepping, bool bAccelChange)
 {
-	check(BodyInstance->ActorHandle.GetScene() == this);
 	// #todo : Implement
 	// AddTorque(Torque, BodyInstance->ActorHandle.GetId());
 }
@@ -948,6 +1028,8 @@ void FPhysScene_ChaosInterface::ClearTorques_AssumesLocked(FBodyInstance* BodyIn
 void FPhysScene_ChaosInterface::SetKinematicTarget_AssumesLocked(FBodyInstance* BodyInstance, const FTransform& TargetTM, bool bAllowSubstepping)
 {
 	// #todo : Implement
+	//for now just pass it into actor directly
+	FPhysInterface_Chaos::SetKinematicTarget_AssumesLocked(BodyInstance->GetPhysicsActorHandle(), TargetTM);
 }
 
 bool FPhysScene_ChaosInterface::GetKinematicTarget_AssumesLocked(const FBodyInstance* BodyInstance, FTransform& OutTM) const
@@ -1009,6 +1091,8 @@ void FPhysScene_ChaosInterface::SetUpForFrame(const FVector* NewGrav, float InDe
 
 void FPhysScene_ChaosInterface::StartFrame()
 {
+	using namespace Chaos;
+
 	SCOPE_CYCLE_COUNTER(STAT_Scene_StartFrame);
 
 	FChaosSolversModule* SolverModule = FChaosSolversModule::GetModule();
@@ -1029,6 +1113,29 @@ void FPhysScene_ChaosInterface::StartFrame()
 	{
 		check(!CompletionEvent.GetReference())
 
+		{
+			const TArray<FPhysicsSolver*>& SolverList = SolverModule->GetSolvers();
+			TArray<FPhysicsSolver*> ActiveSolvers;
+			ActiveSolvers.Reserve(SolverList.Num());
+
+			// #BG calculate active solver list once as we dispatch our first task
+			for (FPhysicsSolver* Solver : SolverList)
+			{
+				if (Solver->HasActiveParticles())
+				{
+					ActiveSolvers.Add(Solver);
+				}
+			}
+
+			const int32 NumActiveSolvers = ActiveSolvers.Num();
+
+			PhysicsParallelFor(NumActiveSolvers, [&](int32 Index)
+			{
+				FPhysicsSolver* Solver = ActiveSolvers[Index];
+				Solver->UpdatePhysicsThreadStructures();
+			});
+		}
+
 		FGraphEventRef SimulationCompleteEvent = FGraphEvent::CreateGraphEvent();
 
 		// Need to fire off a parallel task to handle running physics commands and
@@ -1074,6 +1181,7 @@ void FPhysScene_ChaosInterface::EndFrame(ULineBatchComponent* InLineBatcher)
 	case EChaosThreadingMode::SingleThread:
 	{
 		SyncBodies();
+		Scene.GetSolver()->SyncEvents_GameThread();
 	}
 	break;
 	case EChaosThreadingMode::TaskGraph:
@@ -1087,14 +1195,14 @@ void FPhysScene_ChaosInterface::EndFrame(ULineBatchComponent* InLineBatcher)
 		{
 			SCOPE_CYCLE_COUNTER(STAT_FlipResults);
 
-			const TArray<FPBDRigidsSolver*>& SolverList = SolverModule->GetSolvers();
-			TArray<FPBDRigidsSolver*> ActiveSolvers;
+			const TArray<FPhysicsSolver*>& SolverList = SolverModule->GetSolvers();
+			TArray<FPhysicsSolver*> ActiveSolvers;
 			ActiveSolvers.Reserve(SolverList.Num());
 
 			// #BG calculate active solver list once as we dispatch our first task
-			for(FPBDRigidsSolver* Solver : SolverList)
+			for(FPhysicsSolver* Solver : SolverList)
 			{
-				if(Solver->HasActiveObjects() && Solver->GetObjectStorage().GetNumObjects() > 0)
+				if(Solver->HasActiveParticles())
 				{
 					ActiveSolvers.Add(Solver);
 				}
@@ -1104,21 +1212,25 @@ void FPhysScene_ChaosInterface::EndFrame(ULineBatchComponent* InLineBatcher)
 
 			PhysicsParallelFor(NumActiveSolvers, [&](int32 Index)
 			{
-				FPBDRigidsSolver* Solver = ActiveSolvers[Index];
-				FSolverObjectStorage& Objects = Solver->GetObjectStorage();
+				FPhysicsSolver* Solver = ActiveSolvers[Index];
 
-				Objects.ForEachSolverObject([](auto* Object)
+				Solver->ForEachPhysicsProxy([](auto* Object)
 				{
-					Object->FlipCache();
+					Object->FlipBuffer();
 				});
 
-				Objects.ForEachSolverObject([](auto* Object)
+				Solver->ForEachPhysicsProxy([](auto* Object)
 				{
-					Object->SyncToCache();
+					Object->PullFromPhysicsState();
 				});
 
-				Solver->SyncEvents_GameThread();
 			});
+
+			for (FPhysicsSolver* Solver : ActiveSolvers)
+			{
+				SyncBodies(Solver);
+				Solver->SyncEvents_GameThread();
+			}
 		}
 	}
 		break;
@@ -1153,8 +1265,9 @@ void FPhysScene_ChaosInterface::ListAwakeRigidBodies(bool bIncludeKinematic)
 
 int32 FPhysScene_ChaosInterface::GetNumAwakeBodies() const
 {
-	Chaos::FPBDRigidsSolver* Solver = Scene.GetSolver();
+	Chaos::FPhysicsSolver* Solver = Scene.GetSolver();
 	int32 Count = 0;
+#if TODO_REIMPLEMENT_GET_RIGID_PARTICLES
 	uint32 ParticlesSize = Solver->GetRigidParticles().Size();
 	for(uint32 ParticleIndex = 0; ParticleIndex < ParticlesSize; ++ParticleIndex)
 	{
@@ -1163,6 +1276,7 @@ int32 FPhysScene_ChaosInterface::GetNumAwakeBodies() const
 			Count++;
 		}
 	}
+#endif
 	return Count;
 }
 
@@ -1206,6 +1320,7 @@ void FPhysScene_ChaosInterface::SyncBodies()
 {
 	TArray<FPhysScenePendingComponentTransform_Chaos> PendingTransforms;
 
+#if TODO_REIMPLEMENT_GET_RIGID_PARTICLES
 	for(uint32 Index = 0; Index < Scene.GetSolver()->GetRigidParticles().Size(); ++Index)
 	{
 		if(BodyInstances[Index])
@@ -1215,6 +1330,7 @@ void FPhysScene_ChaosInterface::SyncBodies()
 			PendingTransforms.Add(NewEntry);
 		}
 	}
+#endif
 
 	for(FPhysScenePendingComponentTransform_Chaos& Entry : PendingTransforms)
 	{
@@ -1239,9 +1355,51 @@ void FPhysScene_ChaosInterface::SyncBodies()
 	}
 }
 
+void FPhysScene_ChaosInterface::SyncBodies(Chaos::FPhysicsSolver* Solver)
+{
+	TArray<FPhysScenePendingComponentTransform_Chaos> PendingTransforms;
+
+	Solver->ForEachPhysicsProxy([this, Solver, &PendingTransforms](auto* Object)
+	{
+		FBodyInstance* BodyInstance = FPhysicsUserData::Get<FBodyInstance>(Object->GetUserData());
+		if (BodyInstance)
+		{
+			if (BodyInstance->InstanceBodyIndex == INDEX_NONE && BodyInstance->OwnerComponent.IsValid())
+			{
+				Chaos::TRigidTransform<float, 3> NewTransform(Object->GetTransform());
+				FPhysScenePendingComponentTransform_Chaos NewEntry(BodyInstance->OwnerComponent.Get(), NewTransform);
+				PendingTransforms.Add(NewEntry);
+			}
+		}
+	});
+
+	for(FPhysScenePendingComponentTransform_Chaos& Entry : PendingTransforms)
+	{
+		UPrimitiveComponent* OwnerComponent = Entry.OwningComp.Get();
+		if(OwnerComponent != nullptr)
+		{
+			AActor* Owner = OwnerComponent->GetOwner();
+
+			if(!Entry.NewTransform.EqualsNoScale(OwnerComponent->GetComponentTransform()))
+			{
+				const FVector MoveBy = Entry.NewTransform.GetLocation() - OwnerComponent->GetComponentTransform().GetLocation();
+				const FQuat NewRotation = Entry.NewTransform.GetRotation();
+
+				OwnerComponent->MoveComponent(MoveBy, NewRotation, false, NULL, MOVECOMP_SkipPhysicsMove);
+			}
+
+			if(Owner != NULL && !Owner->IsPendingKill())
+			{
+				Owner->CheckStillInWorld();
+			}
+		}
+	}
+}
+
+
 
 FPhysicsConstraintReference_Chaos 
-FPhysScene_ChaosInterface::AddSpringConstraint(const TArray< TPair<FPhysicsActorReference_Chaos, FPhysicsActorReference_Chaos> >& Constraint)
+FPhysScene_ChaosInterface::AddSpringConstraint(const TArray< TPair<FPhysicsActorHandle, FPhysicsActorHandle> >& Constraint)
 {
 	// #todo : Implement
 	return FPhysicsConstraintReference_Chaos();
@@ -1258,20 +1416,20 @@ void FPhysScene_ChaosInterface::CompleteSceneSimulation(ENamedThreads::Type Curr
 
 	// Cache our results to the threaded buffer.
 	{
-		SCOPE_CYCLE_COUNTER(STAT_CacheResults);
+		SCOPE_CYCLE_COUNTER(STAT_BufferPhysicsResults);
 
 		FChaosSolversModule* Module = FChaosSolversModule::GetModule();
 
 		check(Module);
 
-		const TArray<FPBDRigidsSolver*>& SolverList = Module->GetSolvers();
-		TArray<FPBDRigidsSolver*> ActiveSolvers;
+		const TArray<FPhysicsSolver*>& SolverList = Module->GetSolvers();
+		TArray<FPhysicsSolver*> ActiveSolvers;
 		ActiveSolvers.Reserve(SolverList.Num());
 
 		// #BG calculate active solver list once as we dispatch our first task
-		for(FPBDRigidsSolver* Solver : SolverList)
+		for(FPhysicsSolver* Solver : SolverList)
 		{
-			if(Solver->HasActiveObjects() && Solver->GetObjectStorage().GetNumObjects() > 0)
+			if(Solver->HasActiveParticles())
 			{
 				ActiveSolvers.Add(Solver);
 			}
@@ -1281,12 +1439,11 @@ void FPhysScene_ChaosInterface::CompleteSceneSimulation(ENamedThreads::Type Curr
 
 		PhysicsParallelFor(NumActiveSolvers, [&](int32 Index)
 		{
-			FPBDRigidsSolver* Solver = ActiveSolvers[Index];
-			FSolverObjectStorage& Objects = Solver->GetObjectStorage();
+			FPhysicsSolver* Solver = ActiveSolvers[Index];
 
-			Objects.ForEachSolverObjectParallel([](auto* Object)
+			Solver->ForEachPhysicsProxyParallel([](auto* Object)
 			{
-				Object->CacheResults();
+				Object->BufferPhysicsResults();
 			});
 		});
 	}
diff --git a/Engine/Source/Runtime/Engine/Private/PhysicsEngine/Experimental/PhysicsInterfaceImmediatePhysX.cpp b/Engine/Source/Runtime/Engine/Private/PhysicsEngine/Experimental/PhysicsInterfaceImmediatePhysX.cpp
--- a/Engine/Source/Runtime/Engine/Private/PhysicsEngine/Experimental/PhysicsInterfaceImmediatePhysX.cpp
+++ b/Engine/Source/Runtime/Engine/Private/PhysicsEngine/Experimental/PhysicsInterfaceImmediatePhysX.cpp
@@ -21,7 +21,7 @@
 
 #include "PhysicsEngine/ConstraintDrives.h"
 #include "PhysicsEngine/AggregateGeom.h"
-#include "Physics/PhysicsGeometryPhysX.h"
+#include "Physics/PhysicsGeometry.h"
 #include "Engine/Engine.h"
 #include "PhysicalMaterials/PhysicalMaterial.h"
 
@@ -1662,7 +1662,7 @@ bool FPhysicsInterface_ImmediatePhysX::Sweep_Geom(FHitResult& OutHit, const FBod
 		{
 			if(Actor.IsValid() && InInstance->OwnerComponent != nullptr)
 			{
-				FPhysXShapeAdaptor ShapeAdaptor(InShapeRotation, InShape);
+				FPhysXShapeAdapter ShapeAdaptor(InShapeRotation, InShape);
 				
 				const FVector Delta = InEnd - InStart;
 				const float DeltaMag = Delta.Size();
@@ -1776,7 +1776,7 @@ bool FPhysicsInterface_ImmediatePhysX::Overlap_Geom(const FBodyInstance* InBodyI
 
 bool FPhysicsInterface_ImmediatePhysX::Overlap_Geom(const FBodyInstance* InBodyInstance, const FCollisionShape& InCollisionShape, const FQuat& InShapeRotation, const FTransform& InShapeTransform, FMTDResult* OutOptResult)
 {
-	FPhysXShapeAdaptor Adaptor(InShapeRotation, InCollisionShape);
+	FPhysXShapeAdapter Adaptor(InShapeRotation, InCollisionShape);
 
 	return Overlap_GeomInternal(InBodyInstance, Adaptor.GetGeometry(), InShapeTransform, OutOptResult);
 }
diff --git a/Engine/Source/Runtime/Engine/Private/PhysicsEngine/ImmediatePhysics/ImmediatePhysicsChaos/ImmediatePhysicsActorHandle_Chaos.cpp b/Engine/Source/Runtime/Engine/Private/PhysicsEngine/ImmediatePhysics/ImmediatePhysicsChaos/ImmediatePhysicsActorHandle_Chaos.cpp
new file mode 100644
--- /dev/null
+++ b/Engine/Source/Runtime/Engine/Private/PhysicsEngine/ImmediatePhysics/ImmediatePhysicsChaos/ImmediatePhysicsActorHandle_Chaos.cpp
@@ -0,0 +1,418 @@
+// Copyright 1998-2019 Epic Games, Inc. All Rights Reserved.
+
+#include "Physics/ImmediatePhysics/ImmediatePhysicsChaos/ImmediatePhysicsActorHandle_Chaos.h"
+
+#if INCLUDE_CHAOS
+
+#include "Chaos/Box.h"
+#include "Chaos/Capsule.h"
+#include "Chaos/GeometryParticles.h"
+#include "Chaos/ParticleHandle.h"
+#include "Chaos/PBDRigidParticles.h"
+#include "Chaos/PBDRigidsEvolutionGBF.h"
+#include "Chaos/Sphere.h"
+
+#include "PhysicsEngine/BodyInstance.h"
+#include "PhysicsEngine/BodySetup.h"
+
+namespace ImmediatePhysics_Chaos
+{
+	//
+	// Utils
+	//
+
+	bool CreateGeometry(FBodyInstance* BodyInstance, const FVector& Scale, TUniquePtr<Chaos::TImplicitObject<FReal, Dimensions>>& OutGeom, TArray<TUniquePtr<Chaos::TPerShapeData<float, 3>>>& OutShapes)
+	{
+		// @todo(ccaulfield): copied from FPhysInterface_Chaos::AddGeometry. Merge implementations.
+		TArray<TUniquePtr<Chaos::TImplicitObject<float, 3>>> Geoms;
+		TArray<TUniquePtr<Chaos::TPerShapeData<float, 3>>> Shapes;
+
+		FBodyCollisionData CollisionData;
+		BodyInstance->BuildBodyFilterData(CollisionData.CollisionFilterData);
+		BodyInstance->BuildBodyCollisionFlags(CollisionData.CollisionFlags, BodyInstance->GetCollisionEnabled(), BodyInstance->BodySetup->GetCollisionTraceFlag() == CTF_UseComplexAsSimple);
+
+		FKAggregateGeom* Geometry = &BodyInstance->BodySetup->AggGeom;
+		FVector ActorToCoMTranslation = BodyInstance->GetMassSpaceLocal().GetTranslation();
+
+		auto NewShapeHelper = [CollisionData](const Chaos::TImplicitObject<float, 3>* InGeom, bool bComplexShape = false)
+		{
+			auto NewShape = MakeUnique<Chaos::TPerShapeData<float, 3>>();
+			NewShape->Geometry = InGeom;
+			NewShape->QueryData = bComplexShape ? CollisionData.CollisionFilterData.QueryComplexFilter : CollisionData.CollisionFilterData.QuerySimpleFilter;
+			NewShape->SimData = CollisionData.CollisionFilterData.SimFilter;
+			return NewShape;
+		};
+
+		for (uint32 i = 0; i < static_cast<uint32>(Geometry->SphereElems.Num()); ++i)
+		{
+			ensure(FMath::IsNearlyEqual(Scale[0], Scale[1]) && FMath::IsNearlyEqual(Scale[1], Scale[2]));
+			const auto& CollisionSphere = Geometry->SphereElems[i];
+			auto ImplicitSphere = MakeUnique<Chaos::TSphere<float, 3>>(Chaos::TVector<float, 3>(0.f, 0.f, 0.f), CollisionSphere.Radius * Scale[0]);
+			auto NewShape = NewShapeHelper(ImplicitSphere.Get());
+			Shapes.Emplace(MoveTemp(NewShape));
+			Geoms.Add(MoveTemp(ImplicitSphere));
+
+		}
+		for (uint32 i = 0; i < static_cast<uint32>(Geometry->BoxElems.Num()); ++i)
+		{
+			const auto& Box = Geometry->BoxElems[i];
+			Chaos::TVector<float, 3> half_extents = Scale * Chaos::TVector<float, 3>(Box.X / 2.f, Box.Y / 2.f, Box.Z / 2.f);
+			auto ImplicitBox = MakeUnique<Chaos::TBox<float, 3>>(-half_extents, half_extents);
+			auto NewShape = NewShapeHelper(ImplicitBox.Get());
+			Shapes.Emplace(MoveTemp(NewShape));
+			Geoms.Add(MoveTemp(ImplicitBox));
+
+		}
+		for (uint32 i = 0; i < static_cast<uint32>(Geometry->SphylElems.Num()); ++i)
+		{
+			ensure(FMath::IsNearlyEqual(Scale[0], Scale[1]) && FMath::IsNearlyEqual(Scale[1], Scale[2]));
+			const auto& Sphyl = Geometry->SphylElems[i];
+			if (Sphyl.Length == 0)
+			{
+				auto ImplicitSphere = MakeUnique<Chaos::TSphere<float, 3>>(Chaos::TVector<float, 3>(0), Sphyl.Radius * Scale[0]);
+				auto NewShape = NewShapeHelper(ImplicitSphere.Get());
+				Shapes.Emplace(MoveTemp(NewShape));
+				Geoms.Add(MoveTemp(ImplicitSphere));
+
+			}
+			else
+			{
+				Chaos::TVector<float, 3> HalfExtents(0, 0, Sphyl.Length / 2 * Scale[0]);
+				FTransform SphylTransform = Sphyl.GetTransform();
+				SphylTransform.SetTranslation(SphylTransform.GetTranslation() - ActorToCoMTranslation);
+
+				auto ImplicitCapsule = MakeUnique<Chaos::TCapsule<float>>(SphylTransform.TransformPosition(-HalfExtents), SphylTransform.TransformPosition(HalfExtents), Sphyl.Radius * Scale[0]);
+				auto NewShape = NewShapeHelper(ImplicitCapsule.Get());
+				Shapes.Emplace(MoveTemp(NewShape));
+				Geoms.Add(MoveTemp(ImplicitCapsule));
+			}
+		}
+
+		if (Geoms.Num() == 0)
+		{
+			return false;
+		}
+
+		if (Geoms.Num() == 1)
+		{
+			OutGeom = MoveTemp(Geoms[0]);
+		}
+		else
+		{
+			OutGeom = MakeUnique<Chaos::TImplicitObjectUnion<float, 3>>(MoveTemp(Geoms));
+		}
+
+		OutShapes = MoveTemp(Shapes);
+
+		return true;
+	}
+
+	//
+	// Actor Handle
+	//
+
+	FActorHandle::FActorHandle(Chaos::TPBDRigidsEvolutionGBF<FReal, Dimensions>* InEvolution, EActorType ActorType, FBodyInstance* BodyInstance, const FTransform& Transform)
+		: Evolution(InEvolution)
+		, ParticleHandle(nullptr)
+	{
+		using namespace Chaos;
+
+		// @todo(ccaulfield): Scale
+		if (CreateGeometry(BodyInstance, FVector::OneVector, Geometry, Shapes))
+		{
+			switch (ActorType)
+			{
+			case EActorType::StaticActor:
+				ParticleHandle = Evolution->CreateStaticParticles(1, TGeometryParticleParameters<FReal, Dimensions>())[0];
+				break;
+			case EActorType::KinematicActor:
+				ParticleHandle = Evolution->CreateKinematicParticles(1, TKinematicGeometryParticleParameters<FReal, Dimensions>())[0];
+				break;
+			case EActorType::DynamicActor:
+				ParticleHandle = Evolution->CreateDynamicParticles(1, TPBDRigidParticleParameters<FReal, Dimensions>())[0];
+				break;
+			}
+
+			if (ParticleHandle != nullptr)
+			{
+				ActorToCoMTranslation = BodyInstance->GetMassSpaceLocal().GetTranslation();
+
+				SetWorldTransform(Transform);
+
+				ParticleHandle->SetGeometry(MakeSerializable(Geometry));
+
+				if (auto* Kinematic = ParticleHandle->AsKinematic())
+				{
+					Kinematic->SetV(FVector::ZeroVector);
+					Kinematic->SetW(FVector::ZeroVector);
+				}
+
+				if (auto* Dynamic = ParticleHandle->AsDynamic())
+				{
+					float Mass = BodyInstance->GetBodyMass();
+					float MassInv = (Mass > 0.0f) ? 1.0f / Mass : 0.0f;
+					FVector Inertia = BodyInstance->GetBodyInertiaTensor();
+					FVector InertiaInv = (Mass > 0.0f) ? Inertia.Reciprocal() : FVector::ZeroVector;
+					Dynamic->SetM(Mass);
+					Dynamic->SetInvM(MassInv);
+					Dynamic->SetI({ Inertia.X, Inertia.Y, Inertia.Z });
+					Dynamic->SetInvI({ InertiaInv.X, InertiaInv.Y, InertiaInv.Z });
+					Dynamic->Disabled() = true;
+				}
+			}
+		}
+	}
+
+	FActorHandle::~FActorHandle()
+	{
+		if (ParticleHandle != nullptr)
+		{
+			Evolution->DestroyParticle(ParticleHandle);
+			ParticleHandle = nullptr;
+			Geometry = nullptr;
+		}
+	}
+
+	Chaos::TGenericParticleHandle<FReal, Dimensions> FActorHandle::Handle() const
+	{
+		return { ParticleHandle };
+	}
+
+	void FActorHandle::SetEnabled(bool bEnabled)
+	{
+		if (auto* Dynamic = ParticleHandle->AsDynamic())
+		{
+			Dynamic->Disabled() = !bEnabled;
+		}
+	}
+
+	void FActorHandle::SetWorldTransform(const FTransform& WorldTM)
+	{
+		FTransform ParticleTransform = FTransform(WorldTM.GetRotation(), WorldTM.TransformPosition(ActorToCoMTranslation));
+
+		ParticleHandle->SetX(ParticleTransform.GetTranslation());
+		ParticleHandle->SetR(ParticleTransform.GetRotation());
+
+		if (auto* Dynamic = ParticleHandle->AsDynamic())
+		{
+			Dynamic->SetP(Dynamic->X());
+			Dynamic->SetQ(Dynamic->R());
+		}
+	}
+
+	void FActorHandle::SetIsKinematic(bool bKinematic)
+	{
+#if IMMEDIATEPHYSICS_CHAOS_TODO
+#endif
+	}
+
+	bool FActorHandle::GetIsKinematic() const
+	{
+		return Handle()->IsKinematic();
+	}
+
+
+#if IMMEDIATEPHYSICS_CHAOS_TODO
+	FImmediateKinematicTarget& FActorHandle::GetKinematicTarget()
+	{
+	}
+#endif
+
+	void FActorHandle::SetKinematicTarget(const FTransform& WorldTM)
+	{
+		using namespace Chaos;
+
+#if IMMEDIATEPHYSICS_CHAOS_TODO
+		// Proper kinematic targets
+#endif
+		if (TKinematicGeometryParticleHandle<FReal, Dimensions>* KinematicParticleHandle = ParticleHandle->AsKinematic())
+		{
+			ParticleHandle->SetX(WorldTM.GetTranslation());
+			ParticleHandle->SetR(WorldTM.GetRotation());
+		}
+	}
+
+	bool FActorHandle::HasKinematicTarget() const
+	{
+#if IMMEDIATEPHYSICS_CHAOS_TODO
+#endif
+		return false;
+	}
+
+	bool FActorHandle::IsSimulated() const
+	{
+#if IMMEDIATEPHYSICS_CHAOS_TODO
+#endif
+		return true;
+	}
+
+	FTransform FActorHandle::GetWorldTransform() const
+	{
+		FTransform ParticleTransform = FTransform(Handle()->R(), Handle()->X());
+		return FTransform(ParticleTransform.GetRotation(), ParticleTransform.GetTranslation() - ParticleTransform.TransformVector(ActorToCoMTranslation));
+	}
+
+	void FActorHandle::SetLinearVelocity(const FVector& NewLinearVelocity)
+	{
+		using namespace Chaos;
+
+		if (TKinematicGeometryParticleHandle<FReal, Dimensions>* KinematicParticleHandle = ParticleHandle->AsKinematic())
+		{
+			KinematicParticleHandle->SetV(NewLinearVelocity);
+		}
+	}
+
+	FVector FActorHandle::GetLinearVelocity() const
+	{
+		return Handle()->V();
+	}
+
+	void FActorHandle::SetAngularVelocity(const FVector& NewAngularVelocity)
+	{
+		using namespace Chaos;
+
+		if (TKinematicGeometryParticleHandle<FReal, Dimensions>* KinematicParticleHandle = ParticleHandle->AsKinematic())
+		{
+			KinematicParticleHandle->SetW(NewAngularVelocity);
+		}
+	}
+
+	FVector FActorHandle::GetAngularVelocity() const
+	{
+		return Handle()->W();
+	}
+
+	void FActorHandle::AddForce(const FVector& Force)
+	{
+#if IMMEDIATEPHYSICS_CHAOS_TODO
+#endif
+	}
+
+	void FActorHandle::AddRadialForce(const FVector& Origin, float Strength, float Radius, ERadialImpulseFalloff Falloff, EForceType ForceType)
+	{
+#if IMMEDIATEPHYSICS_CHAOS_TODO
+#endif
+	}
+
+	void FActorHandle::SetLinearDamping(float NewLinearDamping)
+	{
+#if IMMEDIATEPHYSICS_CHAOS_TODO
+#endif
+	}
+
+	float FActorHandle::GetLinearDamping() const
+	{
+#if IMMEDIATEPHYSICS_CHAOS_TODO
+#endif
+		return 0.0f;
+	}
+
+	void FActorHandle::SetAngularDamping(float NewAngularDamping)
+	{
+#if IMMEDIATEPHYSICS_CHAOS_TODO
+#endif
+	}
+
+	float FActorHandle::GetAngularDamping() const
+	{
+#if IMMEDIATEPHYSICS_CHAOS_TODO
+#endif
+		return 0.0f;
+	}
+
+	void FActorHandle::SetMaxLinearVelocitySquared(float NewMaxLinearVelocitySquared)
+	{
+#if IMMEDIATEPHYSICS_CHAOS_TODO
+#endif
+	}
+
+	float FActorHandle::GetMaxLinearVelocitySquared() const
+	{
+#if IMMEDIATEPHYSICS_CHAOS_TODO
+#endif
+		return FLT_MAX;
+	}
+
+	void FActorHandle::SetMaxAngularVelocitySquared(float NewMaxAngularVelocitySquared)
+	{
+#if IMMEDIATEPHYSICS_CHAOS_TODO
+#endif
+	}
+
+	float FActorHandle::GetMaxAngularVelocitySquared() const
+	{
+#if IMMEDIATEPHYSICS_CHAOS_TODO
+#endif
+		return FLT_MAX;
+	}
+
+	void FActorHandle::SetInverseMass(float NewInverseMass)
+	{
+		using namespace Chaos;
+
+		if (TPBDRigidParticleHandle<FReal, Dimensions>* Dynamic = ParticleHandle->AsDynamic())
+		{
+			float NewMass = (NewInverseMass > SMALL_NUMBER) ? 1.0f / NewInverseMass : 0.0f;
+			Dynamic->SetM(NewMass);
+			Dynamic->SetInvM(NewInverseMass);
+		}
+	}
+
+	float FActorHandle::GetInverseMass() const
+	{
+		return Handle()->InvM();
+	}
+
+	void FActorHandle::SetInverseInertia(const FVector& NewInverseInertia)
+	{
+		using namespace Chaos;
+
+		if (TPBDRigidParticleHandle<FReal, Dimensions>* Dynamic = ParticleHandle->AsDynamic())
+		{
+			FVector NewInertia = FVector::ZeroVector;
+			if ((NewInverseInertia.X > SMALL_NUMBER) && (NewInverseInertia.Y > SMALL_NUMBER) && (NewInverseInertia.Z > SMALL_NUMBER))
+			{
+				NewInertia = { 1.0f / NewInverseInertia.X , 1.0f / NewInverseInertia.Y, 1.0f / NewInverseInertia.Z };
+			}
+			Dynamic->SetI({ NewInertia.X, NewInertia.Y, NewInertia.Z });
+			Dynamic->SetInvI({ NewInverseInertia.X, NewInverseInertia.Y, NewInverseInertia.Z });
+		}
+	}
+
+	FVector FActorHandle::GetInverseInertia() const
+	{
+		using namespace Chaos;
+
+		PMatrix<float, 3, 3> InvI = Handle()->InvI();
+		return { InvI.M[0][0], InvI.M[1][1], InvI.M[2][2] };
+	}
+
+	void FActorHandle::SetMaxDepenetrationVelocity(float NewMaxDepenetrationVelocity)
+	{
+#if IMMEDIATEPHYSICS_CHAOS_TODO
+#endif
+	}
+
+	float FActorHandle::GetMaxDepenetrationVelocity(float NewMaxDepenetrationVelocity) const
+	{
+#if IMMEDIATEPHYSICS_CHAOS_TODO
+#endif
+		return FLT_MAX;
+	}
+
+	void FActorHandle::SetMaxContactImpulse(float NewMaxContactImpulse)
+	{
+#if IMMEDIATEPHYSICS_CHAOS_TODO
+#endif
+	}
+
+	float FActorHandle::GetMaxContactImpulse() const
+	{
+#if IMMEDIATEPHYSICS_CHAOS_TODO
+#endif
+		return FLT_MAX;
+	}
+
+}
+
+#endif // INCLUDE_CHAOS
diff --git a/Engine/Source/Runtime/Engine/Private/PhysicsEngine/ImmediatePhysics/ImmediatePhysicsChaos/ImmediatePhysicsJointHandle_Chaos.cpp b/Engine/Source/Runtime/Engine/Private/PhysicsEngine/ImmediatePhysics/ImmediatePhysicsChaos/ImmediatePhysicsJointHandle_Chaos.cpp
new file mode 100644
--- /dev/null
+++ b/Engine/Source/Runtime/Engine/Private/PhysicsEngine/ImmediatePhysics/ImmediatePhysicsChaos/ImmediatePhysicsJointHandle_Chaos.cpp
@@ -0,0 +1,40 @@
+// Copyright 1998-2019 Epic Games, Inc. All Rights Reserved.
+
+#include "Physics/ImmediatePhysics/ImmediatePhysicsChaos/ImmediatePhysicsJointHandle_Chaos.h"
+
+#if INCLUDE_CHAOS
+
+#include "Physics/ImmediatePhysics/ImmediatePhysicsChaos/ImmediatePhysicsActorHandle_Chaos.h"
+
+#include "Chaos/PBDJointConstraints.h"
+
+#include "PhysicsEngine/ConstraintInstance.h"
+
+
+namespace ImmediatePhysics_Chaos
+{
+
+	FJointHandle::FJointHandle(Chaos::TPBDJointConstraints<FReal, Dimensions>* InConstraints, FConstraintInstance* ConstraintInstance, FActorHandle* Actor1, FActorHandle* Actor2)
+		: Constraints(InConstraints)
+		, ConstraintIndex(INDEX_NONE)
+	{
+		FVector JointLocationLocal = ConstraintInstance->GetRefFrame(EConstraintFrame::Frame1).GetLocation();
+		FVector JointLocationWorld = Actor1->GetWorldTransform().TransformPosition(JointLocationLocal);
+		ConstraintIndex = Constraints->AddConstraint({ Actor1->ParticleHandle, Actor2->ParticleHandle }, JointLocationWorld);
+
+#if IMMEDIATEPHYSICS_CHAOS_TODO
+		// Chaos constraint handles
+		// Proper joint constraint
+#endif
+	}
+
+	FJointHandle::~FJointHandle()
+	{
+#if IMMEDIATEPHYSICS_CHAOS_TODO
+		// Remove joints
+#endif
+	}
+
+}
+
+#endif
diff --git a/Engine/Source/Runtime/Engine/Private/PhysicsEngine/ImmediatePhysics/ImmediatePhysicsChaos/ImmediatePhysicsSimulation_Chaos.cpp b/Engine/Source/Runtime/Engine/Private/PhysicsEngine/ImmediatePhysics/ImmediatePhysicsChaos/ImmediatePhysicsSimulation_Chaos.cpp
new file mode 100644
--- /dev/null
+++ b/Engine/Source/Runtime/Engine/Private/PhysicsEngine/ImmediatePhysics/ImmediatePhysicsChaos/ImmediatePhysicsSimulation_Chaos.cpp
@@ -0,0 +1,167 @@
+// Copyright 1998-2019 Epic Games, Inc. All Rights Reserved.
+
+#include "Physics/ImmediatePhysics/ImmediatePhysicsChaos/ImmediatePhysicsSimulation_Chaos.h"
+
+#if INCLUDE_CHAOS
+
+#include "Physics/ImmediatePhysics/ImmediatePhysicsChaos/ImmediatePhysicsActorHandle_Chaos.h"
+#include "Physics/ImmediatePhysics/ImmediatePhysicsChaos/ImmediatePhysicsJointHandle_Chaos.h"
+#include "Physics/ImmediatePhysics/ImmediatePhysicsStats.h"
+
+#include "Chaos/ParticleHandle.h"
+#include "Chaos/PBDJointConstraints.h"
+#include "Chaos/PBDRigidParticles.h"
+#include "Chaos/PBDRigidsEvolutionGBF.h"
+
+#include "PhysicsEngine/BodyInstance.h"
+#include "PhysicsEngine/ConstraintInstance.h"
+
+namespace ImmediatePhysics_Chaos
+{
+
+	//
+	//
+	//
+
+	FSimulation::FSimulation()
+		: NumActiveActorHandles(0)
+		, Gravity(FVector::ZeroVector)
+	{
+		using namespace Chaos;
+
+		const int DefaultSolverIterations = 5;
+		Particles = MakeUnique<TPBDRigidsSOAs<FReal, Dimensions>>();
+		Joints = MakeUnique<TPBDJointConstraints<FReal, Dimensions>>();
+		JointsRule = MakeUnique<TPBDConstraintIslandRule<TPBDJointConstraints<FReal, Dimensions>, FReal, Dimensions>>(*Joints);
+		Evolution = MakeUnique<TPBDRigidsEvolutionGBF<FReal, Dimensions>>(*Particles.Get(), DefaultSolverIterations);
+
+		Evolution->AddConstraintRule(JointsRule.Get());
+
+		Evolution->AddForceFunction([this](TTransientPBDRigidParticleHandle<FReal, Dimensions>& Particle, const FReal Dt)
+			{
+				Particle.F() += Gravity * Particle.M();
+			});
+	}
+
+	FSimulation::~FSimulation()
+	{
+		for (FActorHandle* ActorHandle : ActorHandles)
+		{
+			delete ActorHandle;
+		}
+		ActorHandles.Empty();
+
+		for (FJointHandle* JointHandle : JointHandles)
+		{
+			delete JointHandle;
+		}
+		JointHandles.Empty();
+
+		Evolution.Reset();
+		JointsRule.Reset();
+		Joints.Reset();
+		Particles.Reset();
+	}
+
+	FActorHandle* FSimulation::CreateStaticActor(FBodyInstance* BodyInstance)
+	{
+		return CreateActor(EActorType::StaticActor, BodyInstance, BodyInstance->GetUnrealWorldTransform());
+	}
+
+	FActorHandle* FSimulation::CreateKinematicActor(FBodyInstance* BodyInstance, const FTransform& Transform)
+	{
+		return CreateActor(EActorType::KinematicActor, BodyInstance, Transform);
+	}
+
+	FActorHandle* FSimulation::CreateDynamicActor(FBodyInstance* BodyInstance, const FTransform& Transform)
+	{
+		return CreateActor(EActorType::DynamicActor, BodyInstance, Transform);
+	}
+
+	FActorHandle* FSimulation::CreateActor(EActorType ActorType, FBodyInstance* BodyInstance, const FTransform& Transform)
+	{
+		FActorHandle* ActorHandle = new FActorHandle(Evolution.Get(), ActorType, BodyInstance, Transform);
+		int ActorIndex = ActorHandles.Add(ActorHandle);
+		return ActorHandle;
+	}
+
+	void FSimulation::DestroyActor(FActorHandle* ActorHandle)
+	{
+		// @todo(ccaulfield): FActorHandle could remember its index to optimize this
+		ActorHandles.Remove(ActorHandle);
+		delete ActorHandle;
+	}
+
+	FJointHandle* FSimulation::CreateJoint(FConstraintInstance* ConstraintInstance, FActorHandle* Body1, FActorHandle* Body2)
+	{
+		FJointHandle* JointHandle = new FJointHandle(Joints.Get(), ConstraintInstance, Body1, Body2);
+		JointHandles.Add(JointHandle);
+		return JointHandle;
+	}
+
+	void FSimulation::DestroyJoint(FJointHandle* JointHandle)
+	{
+		// @todo(ccaulfield): FJointHandle could remember its index to optimize this
+		JointHandles.Remove(JointHandle);
+		delete JointHandle;
+	}
+
+	void FSimulation::SetNumActiveBodies(int32 InNumActiveActorHandles)
+	{
+		if (InNumActiveActorHandles < NumActiveActorHandles)
+		{
+			for (int ActorHandleIndex = InNumActiveActorHandles; ActorHandleIndex < NumActiveActorHandles; ++ActorHandleIndex)
+			{
+				GetActorHandle(ActorHandleIndex)->SetEnabled(false);
+			}
+		}
+		else if (InNumActiveActorHandles > NumActiveActorHandles)
+		{
+			check(InNumActiveActorHandles <= ActorHandles.Num());
+			for (int ActorHandleIndex = NumActiveActorHandles; ActorHandleIndex < InNumActiveActorHandles; ++ActorHandleIndex)
+			{
+				GetActorHandle(ActorHandleIndex)->SetEnabled(true);
+			}
+		}
+	
+		NumActiveActorHandles = InNumActiveActorHandles;
+	}
+
+	void FSimulation::SetIgnoreCollisionPairTable(const TArray<FIgnorePair>& InIgnoreCollisionPairTable)
+	{
+#if IMMEDIATEPHYSICS_CHAOS_TODO
+		IgnoreCollisionPairTable.Reset();
+		for (const FIgnorePair& IgnorePair : InIgnoreTable)
+		{
+			TSet<FActorHandle*>& IgnoreActorsA = IgnoreCollisionPairTable.FindOrAdd(IgnorePair.A);
+			IgnoreActorsA.Add(IgnorePair.B);
+
+			TSet<FActorHandle*>& IgnoreActorsB = IgnoreCollisionPairTable.FindOrAdd(IgnorePair.B);
+			IgnoreActorsB.Add(IgnorePair.A);
+		}
+		bRecreateIterationCache = true;
+#endif
+	}
+
+	void FSimulation::SetIgnoreCollisionActors(const TArray<FActorHandle*>& InIgnoreCollisionActors)
+	{
+#if IMMEDIATEPHYSICS_CHAOS_TODO
+		IgnoreCollisionActors.Reset();
+		IgnoreCollisionActors.Append(InIgnoreCollisionActors);
+		bRecreateIterationCache = true;
+#endif
+	}
+
+	DECLARE_CYCLE_STAT(TEXT("FSimulation::Simulate_Chaos"), STAT_ImmediateSimulate_Chaos, STATGROUP_ImmediatePhysics);
+
+	void FSimulation::Simulate(float DeltaTime, const FVector& InGravity)
+	{
+		SCOPE_CYCLE_COUNTER(STAT_ImmediateSimulate_Chaos);
+
+		Gravity = InGravity;
+
+		Evolution->AdvanceOneTimeStep(DeltaTime);
+	}
+}
+
+#endif // INCLUDE_CHAOS
\ No newline at end of file
diff --git a/Engine/Source/Runtime/Engine/Private/PhysicsEngine/ImmediatePhysics/ImmediatePhysicsPhysX/ImmediatePhysicsActor_Physx.cpp b/Engine/Source/Runtime/Engine/Private/PhysicsEngine/ImmediatePhysics/ImmediatePhysicsPhysX/ImmediatePhysicsActor_Physx.cpp
new file mode 100644
--- /dev/null
+++ b/Engine/Source/Runtime/Engine/Private/PhysicsEngine/ImmediatePhysics/ImmediatePhysicsPhysX/ImmediatePhysicsActor_Physx.cpp
@@ -0,0 +1,123 @@
+// Copyright 1998-2019 Epic Games, Inc. All Rights Reserved.
+
+#include "Physics/ImmediatePhysics/ImmediatePhysicsPhysX/ImmediatePhysicsActor_PhysX.h"
+
+#if WITH_PHYSX && PHYSICS_INTERFACE_PHYSX
+
+#include "PhysicsPublic.h"
+#include "PhysicsEngine/BodySetup.h"
+#include "PhysicalMaterials/PhysicalMaterial.h"
+
+namespace ImmediatePhysics_PhysX
+{
+
+	FMaterial* GetBaseMaterialFromUnrealMaterial(UPhysicalMaterial* InMaterial)
+	{
+		return &FMaterial::Default;
+	}
+
+	void FActor::CreateGeometry(PxRigidActor* RigidActor, const PxTransform& ActorToBodyTM)
+	{
+		const uint32 NumShapes = RigidActor->getNbShapes();
+		TArray<PxShape*> ActorShapes;
+		ActorShapes.AddUninitialized(NumShapes);
+		RigidActor->getShapes(ActorShapes.GetData(), sizeof(PxShape*)*NumShapes);
+		PxTransform BodyToActorTM = ActorToBodyTM.getInverse();
+
+		Shapes.Empty(NumShapes);
+		TArray<PxMaterial*, TInlineAllocator<1>> Materials;
+
+		for (PxShape* Shape : ActorShapes)
+		{
+			if (!(Shape->getFlags() & PxShapeFlag::eSIMULATION_SHAPE))
+			{
+				continue;
+			}
+
+			const PxTransform LocalPose = Shape->getLocalPose();
+			const PxTransform BodyLocalShape = BodyToActorTM * Shape->getLocalPose();
+			const PxGeometryHolder& GeomHolder = Shape->getGeometry();
+			const PxBounds3 Bounds = PxGeometryQuery::getWorldBounds(GeomHolder.any(), PxTransform(PxIdentity), /*inflation=*/1.f);
+			const float BoundsMag = Bounds.getExtents().magnitude();
+			const PxVec3 BoundsCenter = Bounds.getCenter();
+			int32 NumMaterials = Shape->getNbMaterials();
+
+			Materials.SetNumUninitialized(NumMaterials);
+			Shape->getMaterials(Materials.GetData(), sizeof(Materials[0]) * NumMaterials);
+
+			FMaterial NewMaterial;
+			if (NumMaterials > 0)
+			{
+				NewMaterial = FMaterial(Materials[0]);
+			}
+
+			switch (GeomHolder.getType())
+			{
+			case PxGeometryType::eSPHERE:		Shapes.Emplace(BodyLocalShape, BoundsCenter, BoundsMag, new PxSphereGeometry(GeomHolder.sphere().radius), NewMaterial); break;
+			case PxGeometryType::eCAPSULE:		Shapes.Emplace(BodyLocalShape, BoundsCenter, BoundsMag, new PxCapsuleGeometry(GeomHolder.capsule().radius, GeomHolder.capsule().halfHeight), NewMaterial); break;
+			case PxGeometryType::eBOX:			Shapes.Emplace(BodyLocalShape, BoundsCenter, BoundsMag, new PxBoxGeometry(GeomHolder.box().halfExtents), NewMaterial); break;
+			case PxGeometryType::eCONVEXMESH:	Shapes.Emplace(BodyLocalShape, BoundsCenter, BoundsMag, new PxConvexMeshGeometry(GeomHolder.convexMesh().convexMesh, GeomHolder.convexMesh().scale, GeomHolder.convexMesh().meshFlags), NewMaterial); break;
+			case PxGeometryType::eHEIGHTFIELD:	Shapes.Emplace(BodyLocalShape, BoundsCenter, BoundsMag, new PxHeightFieldGeometry(GeomHolder.heightField().heightField, GeomHolder.heightField().heightFieldFlags, GeomHolder.heightField().heightScale, GeomHolder.heightField().rowScale, GeomHolder.heightField().columnScale), NewMaterial); break;
+			case PxGeometryType::eTRIANGLEMESH: Shapes.Emplace(BodyLocalShape, BoundsCenter, BoundsMag, new PxTriangleMeshGeometry(GeomHolder.triangleMesh().triangleMesh, GeomHolder.triangleMesh().scale, GeomHolder.triangleMesh().meshFlags), NewMaterial); break;
+			default: continue;	//we don't bother with other types
+			}
+		}
+	}
+
+	bool FActor::AddShape(PxShape* InShape)
+	{
+		if (!(InShape->getFlags() & PxShapeFlag::eSIMULATION_SHAPE))
+		{
+			return false;
+		}
+
+		const PxTransform LocalPose = InShape->getLocalPose();
+		const PxGeometryHolder& GeomHolder = InShape->getGeometry();
+		const PxBounds3 ShapeBounds = PxGeometryQuery::getWorldBounds(GeomHolder.any(), PxTransform(PxIdentity), 1.0f);
+		const float BoundsMagnitude = ShapeBounds.getExtents().magnitude();
+		const PxVec3 BoundsCenter = ShapeBounds.getCenter();
+		const int32 NumMaterials = InShape->getNbMaterials();
+		TArray<PxMaterial*, TInlineAllocator<1>> Materials;
+
+		// #PHYS2 Add multi-material support
+		Materials.SetNumUninitialized(NumMaterials);
+		InShape->getMaterials(Materials.GetData(), sizeof(Materials[0]) * NumMaterials);
+
+		FMaterial NewMaterial;
+		if (NumMaterials > 0)
+		{
+			NewMaterial = FMaterial(Materials[0]);
+		}
+
+		switch (GeomHolder.getType())
+		{
+		case PxGeometryType::eSPHERE:		Shapes.Emplace(LocalPose, BoundsCenter, BoundsMagnitude, new PxSphereGeometry(GeomHolder.sphere().radius), NewMaterial); break;
+		case PxGeometryType::eCAPSULE:		Shapes.Emplace(LocalPose, BoundsCenter, BoundsMagnitude, new PxCapsuleGeometry(GeomHolder.capsule().radius, GeomHolder.capsule().halfHeight), NewMaterial); break;
+		case PxGeometryType::eBOX:			Shapes.Emplace(LocalPose, BoundsCenter, BoundsMagnitude, new PxBoxGeometry(GeomHolder.box().halfExtents), NewMaterial); break;
+		case PxGeometryType::eCONVEXMESH:	Shapes.Emplace(LocalPose, BoundsCenter, BoundsMagnitude, new PxConvexMeshGeometry(GeomHolder.convexMesh().convexMesh, GeomHolder.convexMesh().scale, GeomHolder.convexMesh().meshFlags), NewMaterial); break;
+		case PxGeometryType::eHEIGHTFIELD:	Shapes.Emplace(LocalPose, BoundsCenter, BoundsMagnitude, new PxHeightFieldGeometry(GeomHolder.heightField().heightField, GeomHolder.heightField().heightFieldFlags, GeomHolder.heightField().heightScale, GeomHolder.heightField().rowScale, GeomHolder.heightField().columnScale), NewMaterial); break;
+		case PxGeometryType::eTRIANGLEMESH: Shapes.Emplace(LocalPose, BoundsCenter, BoundsMagnitude, new PxTriangleMeshGeometry(GeomHolder.triangleMesh().triangleMesh, GeomHolder.triangleMesh().scale, GeomHolder.triangleMesh().meshFlags), NewMaterial); break;
+		default: return false;	//we don't bother with other types
+		}
+
+		// If we get here then we succeeded in creating a shape.
+		Shapes.Last().UserData = InShape->userData;
+
+		return true;
+	}
+
+	void FActor::TerminateGeometry()
+	{
+		for (FShape& Shape : Shapes)
+		{
+			delete Shape.Geometry;
+		}
+
+		Shapes.Empty();
+	}
+
+}
+
+#endif // WITH_PHYSX
+
+
diff --git a/Engine/Source/Runtime/Engine/Private/PhysicsEngine/ImmediatePhysics/ImmediatePhysicsPhysX/ImmediatePhysicsContactPointRecorder.cpp b/Engine/Source/Runtime/Engine/Private/PhysicsEngine/ImmediatePhysics/ImmediatePhysicsPhysX/ImmediatePhysicsContactPointRecorder.cpp
new file mode 100644
--- /dev/null
+++ b/Engine/Source/Runtime/Engine/Private/PhysicsEngine/ImmediatePhysics/ImmediatePhysicsPhysX/ImmediatePhysicsContactPointRecorder.cpp
@@ -0,0 +1,66 @@
+// Copyright 1998-2019 Epic Games, Inc. All Rights Reserved.
+
+#include "Physics/ImmediatePhysics/ImmediatePhysicsPhysX/ImmediatePhysicsContactPointRecorder_PhysX.h"
+
+#if WITH_PHYSX && PHYSICS_INTERFACE_PHYSX
+
+#include "Physics/ImmediatePhysics/ImmediatePhysicsPhysX/ImmediatePhysicsSimulation_PhysX.h"
+
+namespace ImmediatePhysics_PhysX
+{
+
+	EFrictionCombineMode::Type GetCombineMode(EFrictionCombineMode::Type A, EFrictionCombineMode::Type B)
+	{
+		return FMath::Max(A, B);
+	}
+
+	float UseCombineMode(EFrictionCombineMode::Type Mode, float A, float B)
+	{
+		switch (Mode)
+		{
+		case EFrictionCombineMode::Average: return (A * 0.5f) + (B * 0.5f);
+		case EFrictionCombineMode::Multiply: return A * B;
+		case EFrictionCombineMode::Min: return FMath::Min(A, B);
+		case EFrictionCombineMode::Max: return FMath::Max(A, B);
+		default: return 0.f;	//error
+		}
+	}
+
+	bool FContactPointRecorder::recordContacts(const Gu::ContactPoint* ContactPoints, const PxU32 NumContacts, const PxU32 Index)
+	{
+		FContactPair ContactPair;
+		ContactPair.DynamicActorDataIndex = DynamicActorDataIndex;
+		ContactPair.OtherActorDataIndex = OtherActorDataIndex;
+		ContactPair.StartContactIndex = Simulation.ContactPoints.Num();
+		ContactPair.NumContacts = NumContacts;
+		ContactPair.PairIdx = PairIdx;
+
+		//TODO: can probably optimize this for the default material case
+		EFrictionCombineMode::Type FrictionCombine = GetCombineMode(SimulatedShapeMaterial.FrictionCombineMode, OtherShapeMaterial.FrictionCombineMode);
+		EFrictionCombineMode::Type RestitutionCombine = GetCombineMode(SimulatedShapeMaterial.FrictionCombineMode, OtherShapeMaterial.FrictionCombineMode);
+
+		const float StaticFriction = UseCombineMode(FrictionCombine, SimulatedShapeMaterial.StaticFriction, OtherShapeMaterial.StaticFriction);
+		const float DynamicFriction = UseCombineMode(FrictionCombine, SimulatedShapeMaterial.DynamicFriction, OtherShapeMaterial.DynamicFriction);
+		const float Restitution = UseCombineMode(RestitutionCombine, SimulatedShapeMaterial.Restitution, OtherShapeMaterial.Restitution);
+
+		for (PxU32 ContactIndex = 0; ContactIndex < NumContacts; ++ContactIndex)
+		{
+			//Fill in solver-specific data that our contact gen does not produce...
+
+			Gu::ContactPoint NewPoint = ContactPoints[ContactIndex];
+			NewPoint.maxImpulse = PX_MAX_F32;
+			NewPoint.targetVel = PxVec3(0.f);
+			NewPoint.staticFriction = StaticFriction;
+			NewPoint.dynamicFriction = DynamicFriction;
+			NewPoint.restitution = Restitution;
+			NewPoint.materialFlags = 0;
+			Simulation.ContactPoints.Add(NewPoint);
+		}
+
+		Simulation.ContactPairs.Add(ContactPair);
+		return true;
+	}
+
+}
+
+#endif // WITH_PHYSX
diff --git a/Engine/Source/Runtime/Engine/Private/PhysicsEngine/ImmediatePhysics/ImmediatePhysicsPhysX/ImmediatePhysicsD6JointData_PhysX.cpp b/Engine/Source/Runtime/Engine/Private/PhysicsEngine/ImmediatePhysics/ImmediatePhysicsPhysX/ImmediatePhysicsD6JointData_PhysX.cpp
new file mode 100644
--- /dev/null
+++ b/Engine/Source/Runtime/Engine/Private/PhysicsEngine/ImmediatePhysics/ImmediatePhysicsPhysX/ImmediatePhysicsD6JointData_PhysX.cpp
@@ -0,0 +1,751 @@
+// Copyright 1998-2019 Epic Games, Inc. All Rights Reserved.
+
+#include "Physics/ImmediatePhysics/ImmediatePhysicsPhysX/ImmediatePhysicsD6JointData_PhysX.h"
+
+#if WITH_PHYSX && PHYSICS_INTERFACE_PHYSX
+
+#include "PhysicsPublic.h"
+
+/****** 
+ ************ :(
+///CODE COPIED DIRECTLY OUT OF PHYSX (until refactor)
+ ************
+ ******/
+
+
+namespace ImmediatePhysics_PhysX
+{
+
+	D6JointData::D6JointData(PxD6Joint* Joint)
+		: linearLimit(GPhysXSDK->getTolerancesScale(), 0.f)
+		, twistLimit(0.f, 0.f)
+		, swingLimit(0.f, 0.f)
+	{
+		PxRigidActor* Actor0 = nullptr;
+		PxRigidActor* Actor1 = nullptr;
+		Joint->getActors(Actor0, Actor1);
+
+		auto GetCOM = [](PxRigidActor* Actor) -> PxTransform
+		{
+			if (Actor)
+			{
+				if (Actor->getType() == PxActorType::eRIGID_DYNAMIC || Actor->getType() == PxActorType::eARTICULATION_LINK)
+				{
+					return static_cast<PxRigidBody*>(Actor)->getCMassLocalPose();
+				}
+				else
+				{
+					return static_cast<PxRigidStatic*>(Actor)->getGlobalPose().getInverse();
+				}
+			}
+
+			return PxTransform(PxIdentity);
+		};
+
+		c2b[0] = GetCOM(Actor0).transformInv(Joint->getLocalPose(PxJointActorIndex::eACTOR0));
+		c2b[1] = GetCOM(Actor1).transformInv(Joint->getLocalPose(PxJointActorIndex::eACTOR1));
+
+		linearLimit = Joint->getLinearLimit();
+		swingLimit = Joint->getSwingLimit();
+		twistLimit = Joint->getTwistLimit();
+		linearMinDist = 0.001f;
+
+		for (uint32 DriveIdx = 0; DriveIdx < PxD6Drive::eCOUNT; ++DriveIdx)
+		{
+			drive[DriveIdx] = Joint->getDrive((PxD6Drive::Enum)DriveIdx);
+		}
+
+		Joint->getDriveVelocity(driveLinearVelocity, driveAngularVelocity);
+		drivePosition = Joint->getDrivePosition();
+
+		for (uint32 MotionIdx = 0; MotionIdx < PxD6Axis::eCOUNT; ++MotionIdx)
+		{
+			motion[MotionIdx] = Joint->getMotion((PxD6Axis::Enum)MotionIdx);
+		}
+
+		invMassScale.linear0 = Joint->getInvMassScale0();
+		invMassScale.linear1 = Joint->getInvMassScale1();
+		invMassScale.angular0 = Joint->getInvInertiaScale0();
+		invMassScale.angular1 = Joint->getInvInertiaScale1();
+	}
+
+	bool IsActive(const D6JointData& JointData, const PxD6Drive::Enum Index)
+	{
+		const PxD6JointDrive& DriveData = JointData.drive[Index];
+		return DriveData.stiffness != 0 || DriveData.damping != 0;
+	}
+
+	void PrepareJointData(D6JointData& JointData)
+	{
+		//angular
+		JointData.thSwingY = PxTan(JointData.swingLimit.yAngle / 2);
+		JointData.thSwingZ = PxTan(JointData.swingLimit.zAngle / 2);
+		JointData.thSwingPad = PxTan(JointData.swingLimit.contactDistance / 2);
+
+		JointData.tqSwingY = PxTan(JointData.swingLimit.yAngle / 4);
+		JointData.tqSwingZ = PxTan(JointData.swingLimit.zAngle / 4);
+		JointData.tqSwingPad = PxTan(JointData.swingLimit.contactDistance / 4);
+
+		JointData.tqTwistLow = PxTan(JointData.twistLimit.lower / 4);
+		JointData.tqTwistHigh = PxTan(JointData.twistLimit.upper / 4);
+		JointData.tqTwistPad = PxTan(JointData.twistLimit.contactDistance / 4);
+
+		//linear
+		JointData.driving = 0;
+		JointData.limited = 0;
+		JointData.locked = 0;
+
+		for (PxU32 i = 0; i < PxD6Axis::eCOUNT; i++)
+		{
+			if (JointData.motion[i] == PxD6Motion::eLIMITED)
+				JointData.limited |= 1 << i;
+			else if (JointData.motion[i] == PxD6Motion::eLOCKED)
+				JointData.locked |= 1 << i;
+		}
+
+		// a linear direction isn't driven if it's locked
+		if (IsActive(JointData, PxD6Drive::eX) && JointData.motion[PxD6Axis::eX] != PxD6Motion::eLOCKED) JointData.driving |= 1 << PxD6Drive::eX;
+		if (IsActive(JointData, PxD6Drive::eY) && JointData.motion[PxD6Axis::eY] != PxD6Motion::eLOCKED) JointData.driving |= 1 << PxD6Drive::eY;
+		if (IsActive(JointData, PxD6Drive::eZ) && JointData.motion[PxD6Axis::eZ] != PxD6Motion::eLOCKED) JointData.driving |= 1 << PxD6Drive::eZ;
+
+		// SLERP drive requires all angular dofs unlocked, and inhibits swing/twist
+
+		bool swing1Locked = JointData.motion[PxD6Axis::eSWING1] == PxD6Motion::eLOCKED;
+		bool swing2Locked = JointData.motion[PxD6Axis::eSWING2] == PxD6Motion::eLOCKED;
+		bool twistLocked = JointData.motion[PxD6Axis::eTWIST] == PxD6Motion::eLOCKED;
+
+		if (IsActive(JointData, PxD6Drive::eSLERP) && !swing1Locked && !swing2Locked && !twistLocked)
+			JointData.driving |= 1 << PxD6Drive::eSLERP;
+		else
+		{
+			if (IsActive(JointData, PxD6Drive::eTWIST) && !twistLocked)
+				JointData.driving |= 1 << PxD6Drive::eTWIST;
+			if (IsActive(JointData, PxD6Drive::eSWING) && (!swing1Locked || !swing2Locked))
+				JointData.driving |= 1 << PxD6Drive::eSWING;
+		}
+	}
+
+	void computeJacobianAxes(PxVec3 row[3], const PxQuat& qa, const PxQuat& qb)
+	{
+		// Compute jacobian matrix for (qa* qb)  [[* means conjugate in this expr]]
+		// d/dt (qa* qb) = 1/2 L(qa*) R(qb) (omega_b - omega_a)
+		// result is L(qa*) R(qb), where L(q) and R(q) are left/right q multiply matrix
+
+		PxReal wa = qa.w, wb = qb.w;
+		const PxVec3 va(qa.x, qa.y, qa.z), vb(qb.x, qb.y, qb.z);
+
+		const PxVec3 c = vb * wa + va * wb;
+		const PxReal d0 = wa * wb;
+		const PxReal d1 = va.dot(vb);
+		const PxReal d = d0 - d1;
+
+		row[0] = (va * vb.x + vb * va.x + PxVec3(d, c.z, -c.y)) * 0.5f;
+		row[1] = (va * vb.y + vb * va.y + PxVec3(-c.z, d, c.x)) * 0.5f;
+		row[2] = (va * vb.z + vb * va.z + PxVec3(c.y, -c.x, d)) * 0.5f;
+
+		if ((d0 + d1) != 0.0f)  // check if relative rotation is 180 degrees which can lead to singular matrix
+			return;
+		else
+		{
+			row[0].x += PX_EPS_F32;
+			row[1].y += PX_EPS_F32;
+			row[2].z += PX_EPS_F32;
+		}
+	}
+
+
+	PxReal tanAdd(PxReal tan1, PxReal tan2)
+	{
+		PX_ASSERT(PxAbs(1 - tan1 * tan2) > 1e-6f);
+		return (tan1 + tan2) / (1 - tan1 * tan2);
+	}
+
+	PX_CUDA_CALLABLE PxF32 sqr(const PxF32 a)
+	{
+		return a * a;
+	}
+
+	PX_CUDA_CALLABLE PxReal tanHalf(PxReal sin, PxReal cos)
+	{
+		return sin / (1 + cos);
+	}
+
+
+	class ConstraintHelper
+	{
+		Px1DConstraint* mConstraints;
+		Px1DConstraint* mCurrent;
+		PxVec3 mRa, mRb;
+
+	public:
+		ConstraintHelper(Px1DConstraint* c, const PxVec3& ra, const PxVec3& rb)
+			: mConstraints(c), mCurrent(c), mRa(ra), mRb(rb) {}
+
+
+		// hard linear & angular
+		void linearHard(const PxVec3& axis, PxReal posErr)
+		{
+			Px1DConstraint *c = linear(axis, posErr, PxConstraintSolveHint::eEQUALITY);
+			c->flags |= Px1DConstraintFlag::eOUTPUT_FORCE;
+		}
+
+		void angularHard(const PxVec3& axis, PxReal posErr)
+		{
+			Px1DConstraint *c = angular(axis, posErr, PxConstraintSolveHint::eEQUALITY);
+			c->flags |= Px1DConstraintFlag::eOUTPUT_FORCE;
+		}
+
+		// limited linear & angular
+		void linearLimit(const PxVec3& axis, PxReal ordinate, PxReal limitValue, const PxJointLimitParameters& limit)
+		{
+			PxReal pad = limit.isSoft() ? 0 : limit.contactDistance;
+
+			if (ordinate + pad > limitValue)
+				addLimit(linear(axis, limitValue - ordinate, PxConstraintSolveHint::eNONE), limit);
+		}
+
+		void angularLimit(const PxVec3& axis, PxReal ordinate, PxReal limitValue, PxReal pad, const PxJointLimitParameters& limit)
+		{
+			if (limit.isSoft())
+				pad = 0;
+
+			if (ordinate + pad > limitValue)
+				addLimit(angular(axis, limitValue - ordinate, PxConstraintSolveHint::eNONE), limit);
+		}
+
+
+		void angularLimit(const PxVec3& axis, PxReal error, const PxJointLimitParameters& limit)
+		{
+			addLimit(angular(axis, error, PxConstraintSolveHint::eNONE), limit);
+		}
+
+		void halfAnglePair(PxReal halfAngle, PxReal lower, PxReal upper, PxReal pad, const PxVec3& axis, const PxJointLimitParameters& limit)
+		{
+			PX_ASSERT(lower < upper);
+			if (limit.isSoft())
+				pad = 0;
+
+			if (halfAngle < lower + pad)
+				angularLimit(-axis, -(lower - halfAngle) * 2, limit);
+			if (halfAngle > upper - pad)
+				angularLimit(axis, (upper - halfAngle) * 2, limit);
+		}
+
+		void quarterAnglePair(PxReal quarterAngle, PxReal lower, PxReal upper, PxReal pad, const PxVec3& axis, const PxJointLimitParameters& limit)
+		{
+			if (limit.isSoft())
+				pad = 0;
+
+			PX_ASSERT(lower < upper);
+			if (quarterAngle < lower + pad)
+				angularLimit(-axis, -(lower - quarterAngle) * 4, limit);
+			if (quarterAngle > upper - pad)
+				angularLimit(axis, (upper - quarterAngle) * 4, limit);
+		}
+
+		// driven linear & angular
+
+		void linear(const PxVec3& axis, PxReal velTarget, PxReal error, const PxD6JointDrive& drive)
+		{
+			addDrive(linear(axis, error, PxConstraintSolveHint::eNONE), velTarget, drive);
+		}
+
+		void angular(const PxVec3& axis, PxReal velTarget, PxReal error, const PxD6JointDrive& drive, PxConstraintSolveHint::Enum hint = PxConstraintSolveHint::eNONE)
+		{
+			addDrive(angular(axis, error, hint), velTarget, drive);
+		}
+
+
+		PxU32 getCount() { return PxU32(mCurrent - mConstraints); }
+
+		void prepareLockedAxes(const PxQuat& qA, const PxQuat& qB, const PxVec3& cB2cAp, PxU32 lin, PxU32 ang)
+		{
+			Px1DConstraint* current = mCurrent;
+			if (ang)
+			{
+				PxQuat qB2qA = qA.getConjugate() * qB;
+
+				PxVec3 row[3];
+				computeJacobianAxes(row, qA, qB);
+				PxVec3 imp = qB2qA.getImaginaryPart();
+				if (ang & 1) angular(row[0], -imp.x, PxConstraintSolveHint::eEQUALITY, current++);
+				if (ang & 2) angular(row[1], -imp.y, PxConstraintSolveHint::eEQUALITY, current++);
+				if (ang & 4) angular(row[2], -imp.z, PxConstraintSolveHint::eEQUALITY, current++);
+			}
+
+			if (lin)
+			{
+				PxMat33 axes(qA);
+				if (lin & 1) linear(axes[0], -cB2cAp[0], PxConstraintSolveHint::eEQUALITY, current++);
+				if (lin & 2) linear(axes[1], -cB2cAp[1], PxConstraintSolveHint::eEQUALITY, current++);
+				if (lin & 4) linear(axes[2], -cB2cAp[2], PxConstraintSolveHint::eEQUALITY, current++);
+			}
+
+			for (Px1DConstraint* front = mCurrent; front < current; front++)
+				front->flags = Px1DConstraintFlag::eOUTPUT_FORCE;
+
+			mCurrent = current;
+		}
+
+		Px1DConstraint *getConstraintRow()
+		{
+			return mCurrent++;
+		}
+
+	private:
+		Px1DConstraint* linear(const PxVec3& axis, PxReal posErr, PxConstraintSolveHint::Enum hint)
+		{
+			Px1DConstraint* c = mCurrent++;
+
+			c->solveHint = PxU16(hint);
+			c->linear0 = axis;					c->angular0 = mRa.cross(axis);
+			c->linear1 = axis;					c->angular1 = mRb.cross(axis);
+			PX_ASSERT(c->linear0.isFinite());
+			PX_ASSERT(c->linear1.isFinite());
+			PX_ASSERT(c->angular0.isFinite());
+			PX_ASSERT(c->angular1.isFinite());
+
+			c->geometricError = posErr;
+
+			return c;
+		}
+
+		Px1DConstraint* angular(const PxVec3& axis, PxReal posErr, PxConstraintSolveHint::Enum hint)
+		{
+			Px1DConstraint* c = mCurrent++;
+
+			c->solveHint = PxU16(hint);
+			c->linear0 = PxVec3(0);		c->angular0 = axis;
+			c->linear1 = PxVec3(0);		c->angular1 = axis;
+
+			c->geometricError = posErr;
+			return c;
+		}
+
+		Px1DConstraint* linear(const PxVec3& axis, PxReal posErr, PxConstraintSolveHint::Enum hint, Px1DConstraint* c)
+		{
+			c->solveHint = PxU16(hint);
+			c->linear0 = axis;					c->angular0 = mRa.cross(axis);
+			c->linear1 = axis;					c->angular1 = mRb.cross(axis);
+			PX_ASSERT(c->linear0.isFinite());
+			PX_ASSERT(c->linear1.isFinite());
+			PX_ASSERT(c->angular0.isFinite());
+			PX_ASSERT(c->angular1.isFinite());
+
+			c->geometricError = posErr;
+
+			return c;
+		}
+
+		Px1DConstraint* angular(const PxVec3& axis, PxReal posErr, PxConstraintSolveHint::Enum hint, Px1DConstraint* c)
+		{
+			c->solveHint = PxU16(hint);
+			c->linear0 = PxVec3(0.f);		c->angular0 = axis;
+			c->linear1 = PxVec3(0.f);		c->angular1 = axis;
+
+			c->geometricError = posErr;
+
+			return c;
+		}
+
+		void addLimit(Px1DConstraint* c, const PxJointLimitParameters& limit)
+		{
+			PxU16 flags = PxU16(c->flags | Px1DConstraintFlag::eOUTPUT_FORCE);
+
+			if (limit.isSoft())
+			{
+				flags |= Px1DConstraintFlag::eSPRING;
+				c->mods.spring.stiffness = limit.stiffness;
+				c->mods.spring.damping = limit.damping;
+			}
+			else
+			{
+				c->solveHint = PxConstraintSolveHint::eINEQUALITY;
+				c->mods.bounce.restitution = limit.restitution;
+				c->mods.bounce.velocityThreshold = limit.bounceThreshold;
+				if (c->geometricError > 0)
+					flags |= Px1DConstraintFlag::eKEEPBIAS;
+				if (limit.restitution > 0)
+					flags |= Px1DConstraintFlag::eRESTITUTION;
+			}
+
+			c->flags = flags;
+			c->minImpulse = 0;
+		}
+
+		void addDrive(Px1DConstraint* c, PxReal velTarget, const PxD6JointDrive& drive)
+		{
+			c->velocityTarget = velTarget;
+
+			PxU16 flags = PxU16(c->flags | Px1DConstraintFlag::eSPRING | Px1DConstraintFlag::eHAS_DRIVE_LIMIT);
+			if (drive.flags & PxD6JointDriveFlag::eACCELERATION)
+				flags |= Px1DConstraintFlag::eACCELERATION_SPRING;
+			c->flags = flags;
+			c->mods.spring.stiffness = drive.stiffness;
+			c->mods.spring.damping = drive.damping;
+
+			c->minImpulse = -drive.forceLimit;
+			c->maxImpulse = drive.forceLimit;
+
+			PX_ASSERT(c->linear0.isFinite());
+			PX_ASSERT(c->angular0.isFinite());
+		}
+	};
+
+	void separateSwingTwist(const PxQuat& q, PxQuat& swing, PxQuat& twist)
+	{
+		twist = q.x != 0.0f ? PxQuat(q.x, 0, 0, q.w).getNormalized() : PxQuat(PxIdentity);
+		swing = q * twist.getConjugate();
+	}
+
+	PX_CUDA_CALLABLE PxVec3 ellipseClamp(const PxVec3& point, const PxVec3& radii)
+	{
+		// This function need to be implemented in the header file because
+		// it is included in a spu shader program.
+
+		// finds the closest point on the ellipse to a given point
+
+		// (p.y, p.z) is the input point
+		// (e.y, e.z) are the radii of the ellipse
+
+		// lagrange multiplier method with Newton/Halley hybrid root-finder.
+		// see http://www.geometrictools.com/Documentation/DistancePointToEllipse2.pdf
+		// for proof of Newton step robustness and initial estimate.
+		// Halley converges much faster but sometimes overshoots - when that happens we take
+		// a newton step instead
+
+		// converges in 1-2 iterations where D&C works well, and it's good with 4 iterations
+		// with any ellipse that isn't completely crazy
+
+		const PxU32 MAX_ITERATIONS = 20;
+		const PxReal convergenceThreshold = 1e-4f;
+
+		// iteration requires first quadrant but we recover generality later
+
+		PxVec3 q(0, PxAbs(point.y), PxAbs(point.z));
+		const PxReal tinyEps = 1e-6f; // very close to minor axis is numerically problematic but trivial
+		if (radii.y >= radii.z)
+		{
+			if (q.z < tinyEps)
+				return PxVec3(0, point.y > 0 ? radii.y : -radii.y, 0);
+		}
+		else
+		{
+			if (q.y < tinyEps)
+				return PxVec3(0, 0, point.z > 0 ? radii.z : -radii.z);
+		}
+
+		PxVec3 denom, e2 = radii.multiply(radii), eq = radii.multiply(q);
+
+		// we can use any initial guess which is > maximum(-e.y^2,-e.z^2) and for which f(t) is > 0.
+		// this guess works well near the axes, but is weak along the diagonals.
+
+		PxReal t = PxMax(eq.y - e2.y, eq.z - e2.z);
+
+		for (PxU32 i = 0; i < MAX_ITERATIONS; i++)
+		{
+			denom = PxVec3(0, 1 / (t + e2.y), 1 / (t + e2.z));
+			PxVec3 denom2 = eq.multiply(denom);
+
+			PxVec3 fv = denom2.multiply(denom2);
+			PxReal f = fv.y + fv.z - 1;
+
+			// although in exact arithmetic we are guaranteed f>0, we can get here
+			// on the first iteration via catastrophic cancellation if the point is
+			// very close to the origin. In that case we just behave as if f=0
+
+			if (f < convergenceThreshold)
+				return e2.multiply(point).multiply(denom);
+
+			PxReal df = fv.dot(denom) * -2.0f;
+			t = t - f / df;
+		}
+
+		// we didn't converge, so clamp what we have
+		PxVec3 r = e2.multiply(point).multiply(denom);
+		return r * PxRecipSqrt(sqr(r.y / radii.y) + sqr(r.z / radii.z));
+	}
+
+
+	class ConeLimitHelper
+	{
+	public:
+		ConeLimitHelper(PxReal tanQSwingY, PxReal tanQSwingZ, PxReal tanQPadding)
+			: mTanQYMax(tanQSwingY), mTanQZMax(tanQSwingZ), mTanQPadding(tanQPadding) {}
+
+		// whether the point is inside the (inwardly) padded cone - if it is, there's no limit
+		// constraint
+
+		bool contains(const PxVec3& tanQSwing)
+		{
+			PxReal tanQSwingYPadded = tanAdd(PxAbs(tanQSwing.y), mTanQPadding);
+			PxReal tanQSwingZPadded = tanAdd(PxAbs(tanQSwing.z), mTanQPadding);
+			return sqr(tanQSwingYPadded / mTanQYMax) + sqr(tanQSwingZPadded / mTanQZMax) <= 1;
+		}
+
+		PxVec3 clamp(const PxVec3& tanQSwing,
+			PxVec3& normal)
+		{
+			PxVec3 p = ellipseClamp(tanQSwing, PxVec3(0, mTanQYMax, mTanQZMax));
+			normal = PxVec3(0, p.y / sqr(mTanQYMax), p.z / sqr(mTanQZMax));
+#ifdef PX_PARANOIA_ELLIPSE_CHECK
+			PxReal err = PxAbs(sqr(p.y / mTanQYMax) + sqr(p.z / mTanQZMax) - 1);
+			PX_ASSERT(err < 1e-3);
+#endif
+
+			return p;
+		}
+
+
+		// input is a swing quat, such that swing.x = twist.y = twist.z = 0, q = swing * twist
+		// The routine is agnostic to the sign of q.w (i.e. we don't need the minimal-rotation swing)
+
+		// output is an axis such that positive rotation increases the angle outward from the
+		// limit (i.e. the image of the x axis), the error is the sine of the angular difference,
+		// positive if the twist axis is inside the cone 
+
+		bool getLimit(const PxQuat& swing, PxVec3& axis, PxReal& error)
+		{
+			PX_ASSERT(swing.w > 0);
+			PxVec3 twistAxis = swing.getBasisVector0();
+			PxVec3 tanQSwing = PxVec3(0, tanHalf(swing.z, swing.w), -tanHalf(swing.y, swing.w));
+			if (contains(tanQSwing))
+				return false;
+
+			PxVec3 normal, clamped = clamp(tanQSwing, normal);
+
+			// rotation vector and ellipse normal
+			PxVec3 r(0, -clamped.z, clamped.y), d(0, -normal.z, normal.y);
+
+			// the point on the cone defined by the tanQ swing vector r
+			PxVec3 p(1.f, 0, 0);
+			PxReal r2 = r.dot(r), a = 1 - r2, b = 1 / (1 + r2), b2 = b * b;
+			PxReal v1 = 2 * a*b2;
+			PxVec3 v2(a, 2 * r.z, -2 * r.y);		// a*p + 2*r.cross(p);
+			PxVec3 coneLine = v1 * v2 - p;		// already normalized
+
+												// the derivative of coneLine in the direction d	
+			PxReal rd = r.dot(d);
+			PxReal dv1 = -4 * rd*(3 - r2)*b2*b;
+			PxVec3 dv2(-2 * rd, 2 * d.z, -2 * d.y);
+
+			PxVec3 coneNormal = v1 * dv2 + dv1 * v2;
+
+			axis = coneLine.cross(coneNormal) / coneNormal.magnitude();
+			error = coneLine.cross(axis).dot(twistAxis);
+
+			PX_ASSERT(PxAbs(axis.magnitude() - 1) < 1e-5f);
+
+#ifdef PX_PARANOIA_ELLIPSE_CHECK
+			bool inside = sqr(tanQSwing.y / mTanQYMax) + sqr(tanQSwing.z / mTanQZMax) <= 1;
+			PX_ASSERT(inside && error > -1e-4f || !inside && error < 1e-4f);
+#endif
+
+			return true;
+		}
+
+	private:
+
+
+		PxReal mTanQYMax, mTanQZMax, mTanQPadding;
+	};
+
+
+	PxU32 D6JointSolverPrep(Px1DConstraint* constraints,
+		PxVec3& body0WorldOffset,
+		PxU32 maxConstraints,
+		PxConstraintInvMassScale& invMassScale,
+		const void* constantBlock,
+		const PxTransform& bA2w,
+		const PxTransform& bB2w)
+	{
+		PX_UNUSED(maxConstraints);
+
+		PX_UNUSED(maxConstraints);
+		const D6JointData& data =
+			*reinterpret_cast<const D6JointData*>(constantBlock);
+
+		invMassScale = data.invMassScale;
+
+		const PxU32 SWING1_FLAG = 1 << PxD6Axis::eSWING1,
+			SWING2_FLAG = 1 << PxD6Axis::eSWING2,
+			TWIST_FLAG = 1 << PxD6Axis::eTWIST;
+
+		const PxU32 ANGULAR_MASK = SWING1_FLAG | SWING2_FLAG | TWIST_FLAG;
+		const PxU32 LINEAR_MASK = 1 << PxD6Axis::eX | 1 << PxD6Axis::eY | 1 << PxD6Axis::eZ;
+
+		const PxD6JointDrive* drives = data.drive;
+		PxU32 locked = data.locked, limited = data.limited, driving = data.driving;
+
+		PxTransform cA2w = bA2w.transform(data.c2b[0]);
+		PxTransform cB2w = bB2w.transform(data.c2b[1]);
+
+		body0WorldOffset = cB2w.p - bA2w.p;
+		ConstraintHelper g(constraints, cB2w.p - bA2w.p, cB2w.p - bB2w.p);
+
+		if (cA2w.q.dot(cB2w.q) < 0)	// minimum dist quat (equiv to flipping cB2bB.q, which we don't use anywhere)
+			cB2w.q = -cB2w.q;
+
+		PxTransform cB2cA = cA2w.transformInv(cB2w);
+
+		PX_ASSERT(data.c2b[0].isValid());
+		PX_ASSERT(data.c2b[1].isValid());
+		PX_ASSERT(cA2w.isValid());
+		PX_ASSERT(cB2w.isValid());
+		PX_ASSERT(cB2cA.isValid());
+
+		PxMat33 cA2w_m(cA2w.q), cB2w_m(cB2w.q);
+
+		// handy for swing computation
+		PxVec3 bX = cB2w_m[0], aY = cA2w_m[1], aZ = cA2w_m[2];
+
+		if (driving & ((1 << PxD6Drive::eX) | (1 << PxD6Drive::eY) | (1 << PxD6Drive::eZ)))
+		{
+			// TODO: make drive unilateral if we are outside the limit
+			PxVec3 posErr = data.drivePosition.p - cB2cA.p;
+			for (PxU32 i = 0; i < 3; i++)
+			{
+				// -driveVelocity because velTarget is child (body1) - parent (body0) and Jacobian is 1 for body0 and -1 for parent
+				if (driving & (1 << (PxD6Drive::eX + i)))
+					g.linear(cA2w_m[i], -data.driveLinearVelocity[i], posErr[i], drives[PxD6Drive::eX + i]);
+			}
+		}
+
+		if (driving & ((1 << PxD6Drive::eSLERP) | (1 << PxD6Drive::eSWING) | (1 << PxD6Drive::eTWIST)))
+		{
+			PxQuat d2cA_q = cB2cA.q.dot(data.drivePosition.q) > 0 ? data.drivePosition.q : -data.drivePosition.q;
+
+			const PxVec3& v = data.driveAngularVelocity;
+			PxQuat delta = d2cA_q.getConjugate() * cB2cA.q;
+
+			if (driving & (1 << PxD6Drive::eSLERP))
+			{
+				PxVec3 velTarget = -cA2w.rotate(data.driveAngularVelocity);
+
+				PxVec3 axis[3] = { PxVec3(1.f,0,0), PxVec3(0,1.f,0), PxVec3(0,0,1.f) };
+
+				if (drives[PxD6Drive::eSLERP].stiffness != 0)
+					computeJacobianAxes(axis, cA2w.q * d2cA_q, cB2w.q);	// converges faster if there is only velocity drive
+
+				for (PxU32 i = 0; i < 3; i++)
+					g.angular(axis[i], axis[i].dot(velTarget), -delta.getImaginaryPart()[i], drives[PxD6Drive::eSLERP], PxConstraintSolveHint::eSLERP_SPRING);
+			}
+			else
+			{
+				if (driving & (1 << PxD6Drive::eTWIST))
+					g.angular(bX, v.x, -2.0f * delta.x, drives[PxD6Drive::eTWIST]);
+
+				if (driving & (1 << PxD6Drive::eSWING))
+				{
+					PxVec3 err = delta.rotate(PxVec3(1.f, 0, 0));
+
+					if (!(locked & SWING1_FLAG))
+						g.angular(cB2w_m[1], v.y, err.z, drives[PxD6Drive::eSWING]);
+
+					if (!(locked & SWING2_FLAG))
+						g.angular(cB2w_m[2], v.z, -err.y, drives[PxD6Drive::eSWING]);
+				}
+			}
+		}
+
+		if (limited & ANGULAR_MASK)
+		{
+			PxQuat swing, twist;
+			separateSwingTwist(cB2cA.q, swing, twist);
+
+			// swing limits: if just one is limited: if the other is free, we support 
+			// (-pi/2, +pi/2) limit, using tan of the half-angle as the error measure parameter. 
+			// If the other is locked, we support (-pi, +pi) limits using the tan of the quarter-angle
+			// Notation: th == Ps::tanHalf, tq = tanQuarter
+
+			if (limited & SWING1_FLAG && limited & SWING2_FLAG)
+			{
+				ConeLimitHelper coneHelper(data.tqSwingZ, data.tqSwingY, data.tqSwingPad);
+
+				PxVec3 axis;
+				PxReal error;
+				if (coneHelper.getLimit(swing, axis, error))
+					g.angularLimit(cA2w.rotate(axis), error, data.swingLimit);
+			}
+			else
+			{
+				const PxJointLimitCone &limit = data.swingLimit;
+
+				PxReal tqPad = data.tqSwingPad, thPad = data.thSwingPad;
+
+				if (limited & SWING1_FLAG)
+				{
+					if (locked & SWING2_FLAG)
+					{
+						g.quarterAnglePair(tanHalf(swing.y, swing.w), -data.tqSwingY, data.tqSwingY, tqPad, aY, limit);
+					}
+					else
+					{
+						PxReal dot = -aZ.dot(bX);
+						g.halfAnglePair(tanHalf(dot, 1 - dot * dot), -data.thSwingY, data.thSwingY, thPad, aZ.cross(bX), limit);
+					}
+				}
+				if (limited & SWING2_FLAG)
+				{
+					if (locked & SWING1_FLAG)
+					{
+						g.quarterAnglePair(tanHalf(swing.z, swing.w), -data.tqSwingZ, data.tqSwingZ, tqPad, aZ, limit);
+					}
+					else
+					{
+						PxReal dot = aY.dot(bX);
+						g.halfAnglePair(tanHalf(dot, 1 - dot * dot), -data.thSwingZ, data.thSwingZ, thPad, -aY.cross(bX), limit);
+					}
+				}
+			}
+
+			if (limited & TWIST_FLAG)
+			{
+				g.quarterAnglePair(tanHalf(twist.x, twist.w), data.tqTwistLow, data.tqTwistHigh, data.tqTwistPad,
+					cB2w_m[0], data.twistLimit);
+			}
+		}
+
+		if (limited & LINEAR_MASK)
+		{
+			PxVec3 limitDir = PxVec3(0);
+
+			for (PxU32 i = 0; i < 3; i++)
+			{
+				if (limited & (1 << (PxD6Axis::eX + i)))
+					limitDir += cA2w_m[i] * cB2cA.p[i];
+			}
+
+			PxReal distance = limitDir.magnitude();
+			if (distance > data.linearMinDist)
+				g.linearLimit(limitDir * (1.0f / distance), distance, data.linearLimit.value, data.linearLimit);
+		}
+
+		// we handle specially the case of just one swing dof locked
+
+		PxU32 angularLocked = locked & ANGULAR_MASK;
+
+		if (angularLocked == SWING1_FLAG)
+		{
+			g.angularHard(bX.cross(aZ), -bX.dot(aZ));
+			locked &= ~SWING1_FLAG;
+		}
+		else if (angularLocked == SWING2_FLAG)
+		{
+			locked &= ~SWING2_FLAG;
+			g.angularHard(bX.cross(aY), -bX.dot(aY));
+		}
+
+		g.prepareLockedAxes(cA2w.q, cB2w.q, cB2cA.p, locked & 7, locked >> 3);
+
+		return g.getCount();
+	}
+
+	/*PxMassProperties FLowLevelPhysicsEntity::ComputeMassProperties() const
+	{
+		//return PxRigidBodyExt::computeMassPropertiesFromGeometries(Geometries.GetData(), LocalTMs.GetData(), Geometries.Num());
+	}*/
+
+}
+
+#endif // WITH_PHYSX
diff --git a/Engine/Source/Runtime/Engine/Private/PhysicsEngine/ImmediatePhysics/ImmediatePhysicsPhysX/ImmediatePhysicsMaterial_PhysX.cpp b/Engine/Source/Runtime/Engine/Private/PhysicsEngine/ImmediatePhysics/ImmediatePhysicsPhysX/ImmediatePhysicsMaterial_PhysX.cpp
new file mode 100644
--- /dev/null
+++ b/Engine/Source/Runtime/Engine/Private/PhysicsEngine/ImmediatePhysics/ImmediatePhysicsPhysX/ImmediatePhysicsMaterial_PhysX.cpp
@@ -0,0 +1,23 @@
+// Copyright 1998-2019 Epic Games, Inc. All Rights Reserved.
+
+#include "Physics/ImmediatePhysics/ImmediatePhysicsPhysX/ImmediatePhysicsMaterial_PhysX.h"
+
+#if WITH_PHYSX && PHYSICS_INTERFACE_PHYSX
+#include "PhysXPublic.h"
+
+namespace ImmediatePhysics_PhysX
+{
+	FMaterial::FMaterial(physx::PxMaterial* InPxMaterial)
+		: StaticFriction(InPxMaterial->getStaticFriction())
+		, DynamicFriction(InPxMaterial->getDynamicFriction())
+		, Restitution(InPxMaterial->getRestitution())
+		, FrictionCombineMode((EFrictionCombineMode::Type)InPxMaterial->getFrictionCombineMode())
+		, RestitutionCombineMode((EFrictionCombineMode::Type)InPxMaterial->getRestitutionCombineMode())
+	{
+	}
+
+	/** Default shape material. Created from the CDO of UPhysicalMaterial */
+	FMaterial FMaterial::Default;
+}
+
+#endif // WITH_PHYSX
